[Insights]
  - Timing view: Added "Cpu Core" tracks when context switches are available. By default the tracks are hidden and can be enabled from Tracks /  Context Switches sub-menu.
  - Timing view: Added "Alt + C" shortcut key to toggle visibility of "Cpu Core" tracks.
  - Timing view: Added "filter by event type" for Cpu Core tracks. Double click an event in a Cpu Core track highlights all events of same type (same thread) in all Cpu Core tracks. Also fixed the "filter by event type" to fade out / highlight only tracks of same type.

diff --git a/Engine/Source/Developer/TraceInsights/Private/Insights/ContextSwitches/ContextSwitchesProfilerManager.cpp b/Engine/Source/Developer/TraceInsights/Private/Insights/ContextSwitches/ContextSwitchesProfilerManager.cpp
--- a/Engine/Source/Developer/TraceInsights/Private/Insights/ContextSwitches/ContextSwitchesProfilerManager.cpp
+++ b/Engine/Source/Developer/TraceInsights/Private/Insights/ContextSwitches/ContextSwitchesProfilerManager.cpp
@@ -3,15 +3,12 @@
 #include "ContextSwitchesProfilerManager.h"
 
 #include "Features/IModularFeatures.h"
+#include "Framework/Commands/UICommandList.h"
 #include "Modules/ModuleManager.h"
-#include "TraceServices/AnalysisService.h"
 #include "TraceServices/Model/ContextSwitches.h"
 
-#include "Insights/ContextSwitches/ViewModels/ContextSwitchesTimingTrack.h"
-#include "Insights/InsightsStyle.h"
+#include "Insights/ContextSwitches/ViewModels/ContextSwitchesSharedState.h"
 #include "Insights/TimingProfilerManager.h"
-#include "Insights/ViewModels/ThreadTimingTrack.h"
-#include "Insights/ViewModels/ThreadTrackEvent.h"
 #include "Insights/Widgets/STimingProfilerWindow.h"
 #include "Insights/Widgets/STimingView.h"
 
@@ -137,7 +134,7 @@ bool FContextSwitchesProfilerManager::Tick(float DeltaTime)
 		{
 			TraceServices::FAnalysisSessionReadScope SessionReadScope(*Session.Get());
 
-			const TraceServices::IContextSwitchProvider* ContextSwitchesProvider = TraceServices::ReadContextSwitchProvider(*Session.Get());
+			const TraceServices::IContextSwitchesProvider* ContextSwitchesProvider = TraceServices::ReadContextSwitchesProvider(*Session.Get());
 			if (ContextSwitchesProvider && ContextSwitchesProvider->HasData())
 			{
 				TSharedPtr<STimingProfilerWindow> Window = FTimingProfilerManager::Get()->GetProfilerWindow();
diff --git a/Engine/Source/Developer/TraceInsights/Private/Insights/ContextSwitches/ContextSwitchesProfilerManager.h b/Engine/Source/Developer/TraceInsights/Private/Insights/ContextSwitches/ContextSwitchesProfilerManager.h
--- a/Engine/Source/Developer/TraceInsights/Private/Insights/ContextSwitches/ContextSwitchesProfilerManager.h
+++ b/Engine/Source/Developer/TraceInsights/Private/Insights/ContextSwitches/ContextSwitchesProfilerManager.h
@@ -4,13 +4,15 @@
 
 #include "Containers/Ticker.h"
 #include "CoreMinimal.h"
-#include "Framework/Commands/UICommandList.h"
 
 // Insights
 #include "Insights/InsightsManager.h"
 #include "Insights/IUnrealInsightsModule.h"
 
+////////////////////////////////////////////////////////////////////////////////////////////////////
+
 class FThreadTrackEvent;
+class FUICommandList;
 
 namespace Insights
 {
@@ -75,4 +77,6 @@ private:
 	TSharedPtr<FContextSwitchesSharedState> ContextSwitchesSharedState;
 };
 
+////////////////////////////////////////////////////////////////////////////////////////////////////
+
 } // namespace Insights
diff --git a/Engine/Source/Developer/TraceInsights/Private/Insights/ContextSwitches/ViewModels/ContextSwitchTimingEvent.cpp b/Engine/Source/Developer/TraceInsights/Private/Insights/ContextSwitches/ViewModels/ContextSwitchTimingEvent.cpp
--- a/Engine/Source/Developer/TraceInsights/Private/Insights/ContextSwitches/ViewModels/ContextSwitchTimingEvent.cpp
+++ b/Engine/Source/Developer/TraceInsights/Private/Insights/ContextSwitches/ViewModels/ContextSwitchTimingEvent.cpp
@@ -2,6 +2,8 @@
 
 #include "Insights/ContextSwitches/ViewModels/ContextSwitchTimingEvent.h"
 
+////////////////////////////////////////////////////////////////////////////////////////////////////
+// FContextSwitchTimingEvent
 ////////////////////////////////////////////////////////////////////////////////////////////////////
 
 INSIGHTS_IMPLEMENT_RTTI(FContextSwitchTimingEvent)
@@ -12,5 +14,17 @@ FContextSwitchTimingEvent::FContextSwitchTimingEvent(const TSharedRef<const FBas
 	: FTimingEvent(InTrack, InStartTime, InEndTime, InDepth)
 {}
 
+////////////////////////////////////////////////////////////////////////////////////////////////////
+// FCpuCoreTimingEvent
+////////////////////////////////////////////////////////////////////////////////////////////////////
+
+INSIGHTS_IMPLEMENT_RTTI(FCpuCoreTimingEvent)
+
+////////////////////////////////////////////////////////////////////////////////////////////////////
+
+FCpuCoreTimingEvent::FCpuCoreTimingEvent(const TSharedRef<const FBaseTimingTrack> InTrack, double InStartTime, double InEndTime, uint32 InDepth)
+	: FTimingEvent(InTrack, InStartTime, InEndTime, InDepth)
+{}
+
 ////////////////////////////////////////////////////////////////////////////////////////////////////
 ////////////////////////////////////////////////////////////////////////////////////////////////////
diff --git a/Engine/Source/Developer/TraceInsights/Private/Insights/ContextSwitches/ViewModels/ContextSwitchTimingEvent.h b/Engine/Source/Developer/TraceInsights/Private/Insights/ContextSwitches/ViewModels/ContextSwitchTimingEvent.h
--- a/Engine/Source/Developer/TraceInsights/Private/Insights/ContextSwitches/ViewModels/ContextSwitchTimingEvent.h
+++ b/Engine/Source/Developer/TraceInsights/Private/Insights/ContextSwitches/ViewModels/ContextSwitchTimingEvent.h
@@ -17,11 +17,22 @@ public:
 	FContextSwitchTimingEvent(const TSharedRef<const FBaseTimingTrack> InTrack, double InStartTime, double InEndTime, uint32 InDepth);
 	virtual ~FContextSwitchTimingEvent() {}
 
-	void SetCoreNumber(uint32 InCoreNumber) { CoreNumber = InCoreNumber; }
-	uint32 GetCoreNumber() const { return CoreNumber; }
+	void SetCoreNumber(uint32 InCoreNumber) { SetType((uint64)InCoreNumber); }
+	uint32 GetCoreNumber() const { return (uint32)GetType(); }
+};
+
+////////////////////////////////////////////////////////////////////////////////////////////////////
+
+class FCpuCoreTimingEvent : public FTimingEvent
+{
+	INSIGHTS_DECLARE_RTTI(FCpuCoreTimingEvent, FTimingEvent)
+
+public:
+	FCpuCoreTimingEvent(const TSharedRef<const FBaseTimingTrack> InTrack, double InStartTime, double InEndTime, uint32 InDepth);
+	virtual ~FCpuCoreTimingEvent() {}
 
-private:
-	uint32 CoreNumber = 0;
+	void SetSystemThreadId(uint32 InSystemThreadId) { SetType((uint64)InSystemThreadId); }
+	uint32 GetSystemThreadId() const { return (uint32)GetType(); }
 };
 
 ////////////////////////////////////////////////////////////////////////////////////////////////////
diff --git a/Engine/Source/Developer/TraceInsights/Private/Insights/ContextSwitches/ViewModels/ContextSwitchesSharedState.cpp b/Engine/Source/Developer/TraceInsights/Private/Insights/ContextSwitches/ViewModels/ContextSwitchesSharedState.cpp
new file mode 100644
--- /dev/null
+++ b/Engine/Source/Developer/TraceInsights/Private/Insights/ContextSwitches/ViewModels/ContextSwitchesSharedState.cpp
@@ -0,0 +1,418 @@
+// Copyright Epic Games, Inc. All Rights Reserved.
+
+#include "ContextSwitchesSharedState.h"
+
+#include "EditorStyleSet.h"
+#include "Framework/Commands/Commands.h"
+#include "Framework/Commands/UICommandList.h"
+#include "Framework/MultiBox/MultiBoxBuilder.h"
+#include "TraceServices/Model/ContextSwitches.h"
+
+// Insights
+#include "Insights/ContextSwitches/ContextSwitchesProfilerManager.h"
+#include "Insights/ContextSwitches/ViewModels/ContextSwitchesTimingTrack.h"
+#include "Insights/ContextSwitches/ViewModels/CpuCoreTimingTrack.h"
+#include "Insights/ITimingViewSession.h"
+#include "Insights/InsightsManager.h"
+#include "Insights/ViewModels/ThreadTimingTrack.h"
+#include "Insights/Widgets/STimingView.h"
+
+#define LOCTEXT_NAMESPACE "ContextSwitches"
+
+namespace Insights
+{
+
+////////////////////////////////////////////////////////////////////////////////////////////////////
+// FContextSwitchesStateCommands
+////////////////////////////////////////////////////////////////////////////////////////////////////
+
+class FContextSwitchesStateCommands : public TCommands<FContextSwitchesStateCommands>
+{
+public:
+	FContextSwitchesStateCommands()
+		: TCommands<FContextSwitchesStateCommands>(TEXT("FContextSwitchesStateCommands"), NSLOCTEXT("FContextSwitchesStateCommands", "Context Switches State Commands", "Context Switches Commands"), NAME_None, FEditorStyle::Get().GetStyleSetName())
+	{
+	}
+
+	virtual ~FContextSwitchesStateCommands()
+	{
+	}
+
+	// UI_COMMAND takes long for the compiler to optimize
+	PRAGMA_DISABLE_OPTIMIZATION
+	virtual void RegisterCommands() override
+	{
+		UI_COMMAND(Command_ShowCoreTracks, "Core Tracks", "Show/hide the Cpu Core tracks.", EUserInterfaceActionType::ToggleButton, FInputChord(EModifierKey::Alt, EKeys::C));
+		UI_COMMAND(Command_ShowContextSwitches, "Context Switches", "Show/hide context switches on top of cpu timing tracks.", EUserInterfaceActionType::ToggleButton, FInputChord(EModifierKey::Shift, EKeys::C));
+		UI_COMMAND(Command_ShowOverlays, "Overlays", "Extend the visualisation of context switches over the cpu timing tracks.", EUserInterfaceActionType::ToggleButton, FInputChord(EModifierKey::Shift, EKeys::O));
+		UI_COMMAND(Command_ShowExtendedLines, "Extended Lines", "Show/hide the extended vertical lines at edges of each context switch event.", EUserInterfaceActionType::ToggleButton, FInputChord(EModifierKey::Shift, EKeys::L));
+	}
+	PRAGMA_ENABLE_OPTIMIZATION
+
+	TSharedPtr<FUICommandInfo> Command_ShowCoreTracks;
+	TSharedPtr<FUICommandInfo> Command_ShowContextSwitches;
+	TSharedPtr<FUICommandInfo> Command_ShowOverlays;
+	TSharedPtr<FUICommandInfo> Command_ShowExtendedLines;
+};
+
+////////////////////////////////////////////////////////////////////////////////////////////////////
+// FContextSwitchesSharedState
+////////////////////////////////////////////////////////////////////////////////////////////////////
+
+FContextSwitchesSharedState::FContextSwitchesSharedState(STimingView* InTimingView) 
+	: TimingView(InTimingView)
+	, ThreadsSerial(0)
+	, CpuCoresSerial(0)
+	, bAreCoreTracksVisible(false)
+	, bAreContextSwitchesVisible(true)
+	, bAreOverlaysVisible(true)
+	, bAreExtendedLinesVisible(true)
+	, bSyncWithProviders(true)
+{
+}
+
+////////////////////////////////////////////////////////////////////////////////////////////////////
+
+void FContextSwitchesSharedState::OnBeginSession(Insights::ITimingViewSession& InSession)
+{
+	if (&InSession != TimingView)
+	{
+		return;
+	}
+
+	ThreadsSerial = 0;
+	CpuCoresSerial = 0;
+
+	bAreCoreTracksVisible = true;
+	bAreContextSwitchesVisible = true;
+	bAreOverlaysVisible = true;
+	bAreExtendedLinesVisible = true;
+
+	bSyncWithProviders = true;
+}
+
+////////////////////////////////////////////////////////////////////////////////////////////////////
+
+void FContextSwitchesSharedState::OnEndSession(Insights::ITimingViewSession& InSession)
+{
+	if (&InSession != TimingView)
+	{
+		return;
+	}
+
+	ThreadsSerial = 0;
+	CpuCoresSerial = 0;
+
+	bAreCoreTracksVisible = true;
+	bAreContextSwitchesVisible = true;
+	bAreOverlaysVisible = true;
+	bAreExtendedLinesVisible = true;
+
+	bSyncWithProviders = false;
+}
+
+////////////////////////////////////////////////////////////////////////////////////////////////////
+
+void FContextSwitchesSharedState::Tick(Insights::ITimingViewSession& InSession, const TraceServices::IAnalysisSession& InAnalysisSession)
+{
+	if (&InSession != TimingView)
+	{
+		return;
+	}
+
+	if (bSyncWithProviders && AreContextSwitchesAvailable())
+	{
+		if (bAreCoreTracksVisible)
+		{
+			AddCoreTracks();
+		}
+
+		if (bAreContextSwitchesVisible)
+		{
+			AddContextSwitchesChildTracks();
+		}
+
+		if (FInsightsManager::Get()->IsAnalysisComplete())
+		{
+			// No need to sync anymore when analysis is completed.
+			bSyncWithProviders = false;
+		}
+	}
+}
+
+////////////////////////////////////////////////////////////////////////////////////////////////////
+
+void FContextSwitchesSharedState::ExtendFilterMenu(Insights::ITimingViewSession& InSession, FMenuBuilder& InMenuBuilder)
+{
+	if (&InSession != TimingView)
+	{
+		return;
+	}
+
+	InMenuBuilder.BeginSection("ContextSwitches");
+	{
+		InMenuBuilder.AddSubMenu(
+			LOCTEXT("ContextSwitches_SubMenu", "Context Switches"),
+			LOCTEXT("ContextSwitches_SubMenu_Desc", "Context Switch track options"),
+			FNewMenuDelegate::CreateSP(this, &FContextSwitchesSharedState::BuildSubMenu),
+			false,
+			FSlateIcon()
+		);
+	}
+	InMenuBuilder.EndSection();
+}
+
+////////////////////////////////////////////////////////////////////////////////////////////////////
+
+void FContextSwitchesSharedState::BuildSubMenu(FMenuBuilder& InMenuBuilder)
+{
+	InMenuBuilder.BeginSection("ContextSwitches", LOCTEXT("ContextSwitchesHeading", "Context Switches"));
+	{
+		InMenuBuilder.AddMenuEntry(FContextSwitchesStateCommands::Get().Command_ShowCoreTracks);
+		InMenuBuilder.AddMenuEntry(FContextSwitchesStateCommands::Get().Command_ShowContextSwitches);
+		InMenuBuilder.AddMenuEntry(FContextSwitchesStateCommands::Get().Command_ShowOverlays);
+		InMenuBuilder.AddMenuEntry(FContextSwitchesStateCommands::Get().Command_ShowExtendedLines);
+	}
+	InMenuBuilder.EndSection();
+}
+
+////////////////////////////////////////////////////////////////////////////////////////////////////
+
+void FContextSwitchesSharedState::AddCommands()
+{
+	FContextSwitchesStateCommands::Register();
+
+	TSharedPtr<FUICommandList> CommandList = TimingView->GetCommandList();
+	ensure(CommandList.IsValid());
+
+	CommandList->MapAction(
+		FContextSwitchesStateCommands::Get().Command_ShowCoreTracks,
+		FExecuteAction::CreateSP(this, &FContextSwitchesSharedState::ContextMenu_ShowCoreTracks_Execute),
+		FCanExecuteAction::CreateSP(this, &FContextSwitchesSharedState::ContextMenu_ShowCoreTracks_CanExecute),
+		FIsActionChecked::CreateSP(this, &FContextSwitchesSharedState::ContextMenu_ShowCoreTracks_IsChecked));
+
+	CommandList->MapAction(
+		FContextSwitchesStateCommands::Get().Command_ShowContextSwitches,
+		FExecuteAction::CreateSP(this, &FContextSwitchesSharedState::ContextMenu_ShowContextSwitches_Execute),
+		FCanExecuteAction::CreateSP(this, &FContextSwitchesSharedState::ContextMenu_ShowContextSwitches_CanExecute),
+		FIsActionChecked::CreateSP(this, &FContextSwitchesSharedState::ContextMenu_ShowContextSwitches_IsChecked));
+
+	CommandList->MapAction(
+		FContextSwitchesStateCommands::Get().Command_ShowOverlays,
+		FExecuteAction::CreateSP(this, &FContextSwitchesSharedState::ContextMenu_ShowOverlays_Execute),
+		FCanExecuteAction::CreateSP(this, &FContextSwitchesSharedState::ContextMenu_ShowOverlays_CanExecute),
+		FIsActionChecked::CreateSP(this, &FContextSwitchesSharedState::ContextMenu_ShowOverlays_IsChecked));
+
+	CommandList->MapAction(
+		FContextSwitchesStateCommands::Get().Command_ShowExtendedLines,
+		FExecuteAction::CreateSP(this, &FContextSwitchesSharedState::ContextMenu_ShowExtendedLines_Execute),
+		FCanExecuteAction::CreateSP(this, &FContextSwitchesSharedState::ContextMenu_ShowExtendedLines_CanExecute),
+		FIsActionChecked::CreateSP(this, &FContextSwitchesSharedState::ContextMenu_ShowExtendedLines_IsChecked));
+}
+
+////////////////////////////////////////////////////////////////////////////////////////////////////
+
+bool FContextSwitchesSharedState::AreContextSwitchesAvailable() const
+{
+	return FContextSwitchesProfilerManager::Get()->IsAvailable();
+}
+
+////////////////////////////////////////////////////////////////////////////////////////////////////
+
+void FContextSwitchesSharedState::AddCoreTracks()
+{
+	if (!TimingView)
+	{
+		return;
+	}
+
+	TSharedPtr<const TraceServices::IAnalysisSession> Session = FInsightsManager::Get()->GetSession();
+	if (!Session.IsValid())
+	{
+		return;
+	}
+
+	TraceServices::FAnalysisSessionReadScope SessionReadScope(*Session.Get());
+
+	const TraceServices::IContextSwitchesProvider* ContextSwitchesProvider = TraceServices::ReadContextSwitchesProvider(*Session.Get());
+	if (ContextSwitchesProvider == nullptr)
+	{
+		return;
+	}
+
+	const uint64 NewCpuCoresSerial = ContextSwitchesProvider->GetCpuCoresSerial();
+	if (NewCpuCoresSerial == CpuCoresSerial)
+	{
+		return;
+	}
+
+	CpuCoresSerial = NewCpuCoresSerial;
+
+	ContextSwitchesProvider->EnumerateCpuCores([this](const TraceServices::FCpuCoreInfo& CpuCoreInfo)
+		{
+			TSharedPtr<FCpuCoreTimingTrack>* TrackPtrPtr = CpuCoreTimingTracks.Find(CpuCoreInfo.CoreNumber);
+			if (TrackPtrPtr == nullptr)
+			{
+				const FString TrackName = FString::Printf(TEXT("Core %u"), CpuCoreInfo.CoreNumber);
+				TSharedPtr<FCpuCoreTimingTrack> Track = MakeShared<FCpuCoreTimingTrack>(*this, TrackName, CpuCoreInfo.CoreNumber);
+
+				const int32 Order = FTimingTrackOrder::Cpu - 1024 + CpuCoreInfo.CoreNumber;
+				Track->SetOrder(Order);
+
+				CpuCoreTimingTracks.Add(CpuCoreInfo.CoreNumber, Track);
+				TimingView->AddScrollableTrack(Track);
+			}
+		});
+}
+
+////////////////////////////////////////////////////////////////////////////////////////////////////
+
+void FContextSwitchesSharedState::RemoveCoreTracks()
+{
+	CpuCoresSerial = 0;
+
+	if (TimingView)
+	{
+		for (const auto& KV : CpuCoreTimingTracks)
+		{
+			TimingView->RemoveTrack(KV.Value);
+		}
+	}
+
+	CpuCoreTimingTracks.Reset();
+}
+
+////////////////////////////////////////////////////////////////////////////////////////////////////
+
+void FContextSwitchesSharedState::AddContextSwitchesChildTracks()
+{
+#if 0
+	TSharedPtr<const TraceServices::IAnalysisSession> Session = FInsightsManager::Get()->GetSession();
+	if (!Session.IsValid())
+	{
+		return;
+	}
+
+	TraceServices::FAnalysisSessionReadScope SessionReadScope(*Session.Get());
+
+	const TraceServices::IContextSwitchesProvider* ContextSwitchesProvider = TraceServices::ReadContextSwitchesProvider(*Session.Get());
+
+	if (ContextSwitchesProvider == nullptr)
+	{
+		return;
+	}
+
+	const uint64 NewThreadsSerial = ContextSwitchesProvider->GetThreadsSerial();
+	if (NewThreadsSerial == ThreadsSerial)
+	{
+		return;
+	}
+
+	ThreadsSerial = NewThreadsSerial;
+
+	//TODO: Create "Cpu Thread" timing tracks also for threads without cpu timing events (i.e. only with context switch events).
+#endif
+
+	TSharedPtr<FThreadTimingSharedState> TimingSharedState = TimingView->GetThreadTimingSharedState();
+
+	if (!TimingSharedState.IsValid())
+	{
+		return;
+	}
+
+	const TMap<uint32, TSharedPtr<FCpuTimingTrack>>& CpuTracks = TimingSharedState->GetAllCpuTracks();
+
+	for (const TPair<uint32, TSharedPtr<FCpuTimingTrack>>& MapEntry : CpuTracks)
+	{
+		const TSharedPtr<FCpuTimingTrack>& CpuTrack = MapEntry.Value;
+		if (CpuTrack.IsValid() && !CpuTrack->GetChildTrack().IsValid())
+		{
+			TSharedPtr<FContextSwitchesTimingTrack> ContextSwitchesTrack = MakeShared<FContextSwitchesTimingTrack>(*this, TEXT("Context Switches"), CpuTrack->GetTimelineIndex(), CpuTrack->GetThreadId());
+			ContextSwitchesTrack->SetParentTrack(CpuTrack);
+			CpuTrack->SetChildTrack(ContextSwitchesTrack);
+		}
+	}
+}
+
+////////////////////////////////////////////////////////////////////////////////////////////////////
+
+void FContextSwitchesSharedState::RemoveContextSwitchesChildTracks()
+{
+	ThreadsSerial = 0;
+
+	TSharedPtr<FThreadTimingSharedState> TimingSharedState = TimingView->GetThreadTimingSharedState();
+
+	if (!TimingSharedState.IsValid())
+	{
+		return;
+	}
+
+	const TMap<uint32, TSharedPtr<FCpuTimingTrack>>& CpuTracks = TimingSharedState->GetAllCpuTracks();
+
+	for (const TPair<uint32, TSharedPtr<FCpuTimingTrack>>& MapEntry : CpuTracks)
+	{
+		const TSharedPtr<FCpuTimingTrack>& CpuTrack = MapEntry.Value;
+		if (CpuTrack.IsValid() && CpuTrack->GetChildTrack().IsValid() && CpuTrack->GetChildTrack()->Is<FContextSwitchesTimingTrack>())
+		{
+			CpuTrack->SetChildTrack(nullptr);
+		}
+	}
+}
+
+////////////////////////////////////////////////////////////////////////////////////////////////////
+
+void FContextSwitchesSharedState::SetCoreTracksVisible(bool bOnOff)
+{
+	if (bAreCoreTracksVisible != bOnOff)
+	{
+		bAreCoreTracksVisible = bOnOff;
+
+		if (bAreCoreTracksVisible)
+		{
+			AddCoreTracks();
+		}
+		else
+		{
+			RemoveCoreTracks();
+		}
+	}
+}
+
+////////////////////////////////////////////////////////////////////////////////////////////////////
+
+void FContextSwitchesSharedState::SetContextSwitchesVisible(bool bOnOff)
+{
+	if (bAreContextSwitchesVisible != bOnOff)
+	{
+		bAreContextSwitchesVisible = bOnOff;
+
+		if (bAreContextSwitchesVisible)
+		{
+			AddContextSwitchesChildTracks();
+		}
+		else
+		{
+			RemoveContextSwitchesChildTracks();
+		}
+	}
+}
+
+////////////////////////////////////////////////////////////////////////////////////////////////////
+
+void FContextSwitchesSharedState::SetOverlaysVisible(bool bOnOff)
+{
+	bAreOverlaysVisible = bOnOff;
+}
+
+////////////////////////////////////////////////////////////////////////////////////////////////////
+
+void FContextSwitchesSharedState::SetExtendedLinesVisible(bool bOnOff)
+{
+	bAreExtendedLinesVisible = bOnOff;
+}
+
+////////////////////////////////////////////////////////////////////////////////////////////////////
+
+} // namespace Insights
+
+#undef LOCTEXT_NAMESPACE
diff --git a/Engine/Source/Developer/TraceInsights/Private/Insights/ContextSwitches/ViewModels/ContextSwitchesSharedState.h b/Engine/Source/Developer/TraceInsights/Private/Insights/ContextSwitches/ViewModels/ContextSwitchesSharedState.h
new file mode 100644
--- /dev/null
+++ b/Engine/Source/Developer/TraceInsights/Private/Insights/ContextSwitches/ViewModels/ContextSwitchesSharedState.h
@@ -0,0 +1,119 @@
+// Copyright Epic Games, Inc. All Rights Reserved.
+
+#pragma once
+
+#include "CoreMinimal.h"
+
+// Insights
+#include "Insights/ITimingViewExtender.h"
+
+////////////////////////////////////////////////////////////////////////////////////////////////////
+
+class FMenuBuilder;
+class STimingView;
+
+namespace TraceServices
+{
+	class IAnalysisSession;
+}
+
+namespace Insights
+{
+
+class FCpuCoreTimingTrack;
+class FContextSwitchesTimingTrack;
+class ITimingViewSession;
+
+////////////////////////////////////////////////////////////////////////////////////////////////////
+
+class FContextSwitchesSharedState : public Insights::ITimingViewExtender, public TSharedFromThis<FContextSwitchesSharedState>
+{
+public:
+	FContextSwitchesSharedState(STimingView* InTimingView);
+	virtual ~FContextSwitchesSharedState() = default;
+
+	//////////////////////////////////////////////////
+	// ITimingViewExtender interface
+
+	virtual void OnBeginSession(Insights::ITimingViewSession& InSession) override;
+	virtual void OnEndSession(Insights::ITimingViewSession& InSession) override;
+	virtual void Tick(Insights::ITimingViewSession& InSession, const TraceServices::IAnalysisSession& InAnalysisSession) override;
+	virtual void ExtendFilterMenu(Insights::ITimingViewSession& InSession, FMenuBuilder& InMenuBuilder) override;
+
+	//////////////////////////////////////////////////
+
+	STimingView* GetTimingView() { return TimingView; }
+
+	bool AreContextSwitchesAvailable() const;
+
+	bool AreCoreTracksVisible() const { return bAreCoreTracksVisible; }
+	void ShowCoreTracks() { SetCoreTracksVisible(true); }
+	void HideCoreTracks() { SetCoreTracksVisible(false); }
+	void ToggleCoreTracks() { SetCoreTracksVisible(!bAreCoreTracksVisible); }
+	void SetCoreTracksVisible(bool bOnOff);
+
+	bool AreContextSwitchesVisible() const { return bAreContextSwitchesVisible; }
+	void ShowContextSwitches() { SetContextSwitchesVisible(true); }
+	void HideContextSwitches() { SetContextSwitchesVisible(false); }
+	void ToggleContextSwitches() { SetContextSwitchesVisible(!bAreContextSwitchesVisible); }
+	void SetContextSwitchesVisible(bool bOnOff);
+
+	bool AreOverlaysVisible() const { return bAreOverlaysVisible; }
+	void ShowOverlays() { SetOverlaysVisible(true); }
+	void HideOverlays() { SetOverlaysVisible(false); }
+	void ToggleOverlays() { SetOverlaysVisible(!bAreOverlaysVisible); }
+	void SetOverlaysVisible(bool bOnOff);
+
+	bool AreExtendedLinesVisible() const { return bAreExtendedLinesVisible; }
+	void ShowExtendedLines() { SetExtendedLinesVisible(true); }
+	void HideExtendedLines() { SetExtendedLinesVisible(false); }
+	void ToggleExtendedLines() { SetExtendedLinesVisible(!bAreExtendedLinesVisible); }
+	void SetExtendedLinesVisible(bool bOnOff);
+
+	void AddCommands();
+
+private:
+	void AddCoreTracks();
+	void RemoveCoreTracks();
+
+	void AddContextSwitchesChildTracks();
+	void RemoveContextSwitchesChildTracks();
+
+	void BuildSubMenu(FMenuBuilder& InMenuBuilder);
+
+	void ContextMenu_ShowCoreTracks_Execute() { ToggleCoreTracks(); }
+	bool ContextMenu_ShowCoreTracks_CanExecute() const { return AreContextSwitchesAvailable(); }
+	bool ContextMenu_ShowCoreTracks_IsChecked() const { return AreCoreTracksVisible(); }
+
+	void ContextMenu_ShowContextSwitches_Execute() { ToggleContextSwitches(); }
+	bool ContextMenu_ShowContextSwitches_CanExecute() const { return AreContextSwitchesAvailable(); }
+	bool ContextMenu_ShowContextSwitches_IsChecked() const { return AreContextSwitchesVisible(); }
+
+	void ContextMenu_ShowOverlays_Execute() { ToggleOverlays(); }
+	bool ContextMenu_ShowOverlays_CanExecute() const { return AreContextSwitchesAvailable() && AreContextSwitchesVisible(); }
+	bool ContextMenu_ShowOverlays_IsChecked() const { return AreOverlaysVisible(); }
+
+	void ContextMenu_ShowExtendedLines_Execute() { ToggleExtendedLines(); }
+	bool ContextMenu_ShowExtendedLines_CanExecute() const { return AreContextSwitchesAvailable() && AreContextSwitchesVisible() && AreOverlaysVisible(); }
+	bool ContextMenu_ShowExtendedLines_IsChecked() const { return AreExtendedLinesVisible(); }
+
+private:
+	STimingView* TimingView;
+
+	TMap<uint32, TSharedPtr<FCpuCoreTimingTrack>> CpuCoreTimingTracks;
+	TMap<uint32, TSharedPtr<FContextSwitchesTimingTrack>> ContextSwitchesTimingTracks;
+
+	uint64 ThreadsSerial;
+	uint64 CpuCoresSerial;
+
+	bool bAreCoreTracksVisible;
+	bool bAreContextSwitchesVisible;
+	bool bAreOverlaysVisible;
+	bool bAreExtendedLinesVisible;
+
+	bool bSyncWithProviders;
+};
+
+////////////////////////////////////////////////////////////////////////////////////////////////////
+
+} // namespace Insights
diff --git a/Engine/Source/Developer/TraceInsights/Private/Insights/ContextSwitches/ViewModels/ContextSwitchesTimingTrack.cpp b/Engine/Source/Developer/TraceInsights/Private/Insights/ContextSwitches/ViewModels/ContextSwitchesTimingTrack.cpp
--- a/Engine/Source/Developer/TraceInsights/Private/Insights/ContextSwitches/ViewModels/ContextSwitchesTimingTrack.cpp
+++ b/Engine/Source/Developer/TraceInsights/Private/Insights/ContextSwitches/ViewModels/ContextSwitchesTimingTrack.cpp
@@ -2,336 +2,24 @@
 
 #include "ContextSwitchesTimingTrack.h"
 
-#include "EditorStyleSet.h"
-#include "Framework/Commands/Commands.h"
-#include "Framework/Commands/UICommandList.h"
-#include "Framework/MultiBox/MultiBoxBuilder.h"
 #include "TraceServices/Model/ContextSwitches.h"
 
 // Insights
 #include "Insights/Common/TimeUtils.h"
-#include "Insights/ContextSwitches/ContextSwitchesProfilerManager.h"
+#include "Insights/ContextSwitches/ViewModels/ContextSwitchesSharedState.h"
 #include "Insights/ContextSwitches/ViewModels/ContextSwitchTimingEvent.h"
-#include "Insights/ITimingViewSession.h"
 #include "Insights/InsightsManager.h"
-#include "Insights/ViewModels/ThreadTimingTrack.h"
-#include "Insights/ViewModels/ThreadTrackEvent.h"
 #include "Insights/ViewModels/TimingTrackViewport.h"
 #include "Insights/ViewModels/TimingViewDrawHelper.h"
 #include "Insights/ViewModels/TooltipDrawState.h"
-#include "Insights/Widgets/STimingView.h"
 
-using namespace TraceServices;
-
-#define LOCTEXT_NAMESPACE "ContextSwitchesTrack"
+#define LOCTEXT_NAMESPACE "FContextSwitchesTimingTrack"
 
 namespace Insights
 {
-////////////////////////////////////////////////////////////////////////////////////////////////////
-// FContextSwitchesStateCommands
-////////////////////////////////////////////////////////////////////////////////////////////////////
-
-class FContextSwitchesStateCommands : public TCommands<FContextSwitchesStateCommands>
-{
-public:
-	FContextSwitchesStateCommands()
-		: TCommands<FContextSwitchesStateCommands>(TEXT("FContextSwitchesStateCommands"), NSLOCTEXT("FContextSwitchesStateCommands", "Context Switches State Commands", "Context Switches Commands"), NAME_None, FEditorStyle::Get().GetStyleSetName())
-	{
-	}
-
-	virtual ~FContextSwitchesStateCommands()
-	{
-	}
-
-	// UI_COMMAND takes long for the compiler to optimize
-	PRAGMA_DISABLE_OPTIMIZATION
-	virtual void RegisterCommands() override
-	{
-		UI_COMMAND(Command_ShowCoreTracks, "Core Tracks", "Show/hide the Cpu Core tracks.", EUserInterfaceActionType::ToggleButton, FInputChord());
-		UI_COMMAND(Command_ShowContextSwitches, "Context Switches", "Show/hide context switches on top of cpu timing tracks.", EUserInterfaceActionType::ToggleButton, FInputChord(EModifierKey::Shift, EKeys::C));
-		UI_COMMAND(Command_ShowOverlays, "Overlays", "Extend the visualisation of context switches over the cpu timing tracks.", EUserInterfaceActionType::ToggleButton, FInputChord(EModifierKey::Shift, EKeys::O));
-		UI_COMMAND(Command_ShowExtendedLines, "Extended Lines", "Show/hide the extended vertical lines at edges of each context switch event.", EUserInterfaceActionType::ToggleButton, FInputChord(EModifierKey::Shift, EKeys::L));
-	}
-	PRAGMA_ENABLE_OPTIMIZATION
-
-	TSharedPtr<FUICommandInfo> Command_ShowCoreTracks;
-	TSharedPtr<FUICommandInfo> Command_ShowContextSwitches;
-	TSharedPtr<FUICommandInfo> Command_ShowOverlays;
-	TSharedPtr<FUICommandInfo> Command_ShowExtendedLines;
-};
-
-////////////////////////////////////////////////////////////////////////////////////////////////////
-// FContextSwitchesSharedState
-////////////////////////////////////////////////////////////////////////////////////////////////////
-
-FContextSwitchesSharedState::FContextSwitchesSharedState(STimingView* InTimingView) 
-	: TimingView(InTimingView)
-	, bAreContextSwitchesVisible(true)
-	, bAreCoreTracksVisible(true)
-	, bAreOverlaysVisible(true)
-	, bAreExtendedLinesVisible(true)
-	, bSyncWithProviders(true)
-{
-}
-
-////////////////////////////////////////////////////////////////////////////////////////////////////
-
-void FContextSwitchesSharedState::OnBeginSession(Insights::ITimingViewSession& InSession)
-{
-	if (&InSession != TimingView)
-	{
-		return;
-	}
-
-	bAreCoreTracksVisible = true;
-	bAreContextSwitchesVisible = true;
-	bAreOverlaysVisible = true;
-	bAreExtendedLinesVisible = true;
-
-	bSyncWithProviders = true;
-}
-
-////////////////////////////////////////////////////////////////////////////////////////////////////
-
-void FContextSwitchesSharedState::OnEndSession(Insights::ITimingViewSession& InSession)
-{
-	if (&InSession != TimingView)
-	{
-		return;
-	}
-
-	bAreCoreTracksVisible = true;
-	bAreContextSwitchesVisible = true;
-	bAreOverlaysVisible = true;
-	bAreExtendedLinesVisible = true;
-
-	bSyncWithProviders = false;
-}
-
-////////////////////////////////////////////////////////////////////////////////////////////////////
-
-void FContextSwitchesSharedState::Tick(Insights::ITimingViewSession& InSession, const TraceServices::IAnalysisSession& InAnalysisSession)
-{
-	if (&InSession != TimingView)
-	{
-		return;
-	}
-
-	if (bSyncWithProviders && AreContextSwitchesAvailable())
-	{
-		if (bAreCoreTracksVisible)
-		{
-			AddCoreTracks();
-		}
-
-		if (bAreContextSwitchesVisible)
-		{
-			AddContextSwitchesChildTracks();
-		}
-
-		if (FInsightsManager::Get()->IsAnalysisComplete())
-		{
-			// No need to sync anymore when analysis is completed.
-			bSyncWithProviders = false;
-		}
-	}
-}
-
-////////////////////////////////////////////////////////////////////////////////////////////////////
-
-void FContextSwitchesSharedState::ExtendFilterMenu(Insights::ITimingViewSession& InSession, FMenuBuilder& InMenuBuilder)
-{
-	if (&InSession != TimingView)
-	{
-		return;
-	}
-
-	InMenuBuilder.BeginSection("ContextSwitches");
-	{
-		InMenuBuilder.AddSubMenu(
-			LOCTEXT("ContextSwitches_SubMenu", "Context Switches"),
-			LOCTEXT("ContextSwitches_SubMenu_Desc", "Context Switch track options"),
-			FNewMenuDelegate::CreateSP(this, &FContextSwitchesSharedState::BuildSubMenu),
-			false,
-			FSlateIcon()
-		);
-	}
-	InMenuBuilder.EndSection();
-}
-
-////////////////////////////////////////////////////////////////////////////////////////////////////
-
-void FContextSwitchesSharedState::BuildSubMenu(FMenuBuilder& InMenuBuilder)
-{
-	InMenuBuilder.BeginSection("ContextSwitches", LOCTEXT("ContextSwitchesHeading", "Context Switches"));
-	{
-		//TODO: InMenuBuilder.AddMenuEntry(FContextSwitchesStateCommands::Get().Command_ShowCoreTracks);
-		InMenuBuilder.AddMenuEntry(FContextSwitchesStateCommands::Get().Command_ShowContextSwitches);
-		InMenuBuilder.AddMenuEntry(FContextSwitchesStateCommands::Get().Command_ShowOverlays);
-		InMenuBuilder.AddMenuEntry(FContextSwitchesStateCommands::Get().Command_ShowExtendedLines);
-	}
-	InMenuBuilder.EndSection();
-}
-
-////////////////////////////////////////////////////////////////////////////////////////////////////
-
-void FContextSwitchesSharedState::AddCommands()
-{
-	FContextSwitchesStateCommands::Register();
-
-	TSharedPtr<FUICommandList> CommandList = TimingView->GetCommandList();
-	ensure(CommandList.IsValid());
-
-	CommandList->MapAction(
-		FContextSwitchesStateCommands::Get().Command_ShowCoreTracks,
-		FExecuteAction::CreateSP(this, &FContextSwitchesSharedState::ContextMenu_ShowCoreTracks_Execute),
-		FCanExecuteAction::CreateSP(this, &FContextSwitchesSharedState::ContextMenu_ShowCoreTracks_CanExecute),
-		FIsActionChecked::CreateSP(this, &FContextSwitchesSharedState::ContextMenu_ShowCoreTracks_IsChecked));
-
-	CommandList->MapAction(
-		FContextSwitchesStateCommands::Get().Command_ShowContextSwitches,
-		FExecuteAction::CreateSP(this, &FContextSwitchesSharedState::ContextMenu_ShowContextSwitches_Execute),
-		FCanExecuteAction::CreateSP(this, &FContextSwitchesSharedState::ContextMenu_ShowContextSwitches_CanExecute),
-		FIsActionChecked::CreateSP(this, &FContextSwitchesSharedState::ContextMenu_ShowContextSwitches_IsChecked));
-
-	CommandList->MapAction(
-		FContextSwitchesStateCommands::Get().Command_ShowOverlays,
-		FExecuteAction::CreateSP(this, &FContextSwitchesSharedState::ContextMenu_ShowOverlays_Execute),
-		FCanExecuteAction::CreateSP(this, &FContextSwitchesSharedState::ContextMenu_ShowOverlays_CanExecute),
-		FIsActionChecked::CreateSP(this, &FContextSwitchesSharedState::ContextMenu_ShowOverlays_IsChecked));
-
-	CommandList->MapAction(
-		FContextSwitchesStateCommands::Get().Command_ShowExtendedLines,
-		FExecuteAction::CreateSP(this, &FContextSwitchesSharedState::ContextMenu_ShowExtendedLines_Execute),
-		FCanExecuteAction::CreateSP(this, &FContextSwitchesSharedState::ContextMenu_ShowExtendedLines_CanExecute),
-		FIsActionChecked::CreateSP(this, &FContextSwitchesSharedState::ContextMenu_ShowExtendedLines_IsChecked));
-}
-
-////////////////////////////////////////////////////////////////////////////////////////////////////
-
-bool FContextSwitchesSharedState::AreContextSwitchesAvailable() const
-{
-	return FContextSwitchesProfilerManager::Get()->IsAvailable();
-}
-
-////////////////////////////////////////////////////////////////////////////////////////////////////
-
-void FContextSwitchesSharedState::AddCoreTracks()
-{
-	//TODO
-}
-
-////////////////////////////////////////////////////////////////////////////////////////////////////
-
-void FContextSwitchesSharedState::RemoveCoreTracks()
-{
-	//TODO
-}
-
-////////////////////////////////////////////////////////////////////////////////////////////////////
-
-void FContextSwitchesSharedState::AddContextSwitchesChildTracks()
-{
-	TSharedPtr<FThreadTimingSharedState> TimingSharedState = TimingView->GetThreadTimingSharedState();
-
-	if (!TimingSharedState.IsValid())
-	{
-		return;
-	}
-
-	const TMap<uint32, TSharedPtr<FCpuTimingTrack>>& CpuTracks = TimingSharedState->GetAllCpuTracks();
-
-	for (const TPair<uint32, TSharedPtr<FCpuTimingTrack>>& MapEntry : CpuTracks)
-	{
-		const TSharedPtr<FCpuTimingTrack>& CpuTrack = MapEntry.Value;
-		if (CpuTrack.IsValid() && !CpuTrack->GetChildTrack().IsValid())
-		{
-			TSharedPtr<FContextSwitchesTimingTrack> ContextSwitchesTrack = MakeShared<FContextSwitchesTimingTrack>(*this, TEXT("Context Switches"), CpuTrack->GetTimelineIndex(), CpuTrack->GetThreadId());
-			ContextSwitchesTrack->SetParentTrack(CpuTrack);
-			CpuTrack->SetChildTrack(ContextSwitchesTrack);
-		}
-	}
-}
-
-////////////////////////////////////////////////////////////////////////////////////////////////////
-
-void FContextSwitchesSharedState::RemoveContextSwitchesChildTracks()
-{
-	TSharedPtr<FThreadTimingSharedState> TimingSharedState = TimingView->GetThreadTimingSharedState();
-
-	if (!TimingSharedState.IsValid())
-	{
-		return;
-	}
-
-	const TMap<uint32, TSharedPtr<FCpuTimingTrack>>& CpuTracks = TimingSharedState->GetAllCpuTracks();
-
-	for (const TPair<uint32, TSharedPtr<FCpuTimingTrack>>& MapEntry : CpuTracks)
-	{
-		const TSharedPtr<FCpuTimingTrack>& CpuTrack = MapEntry.Value;
-		if (CpuTrack.IsValid() && CpuTrack->GetChildTrack().IsValid() && CpuTrack->GetChildTrack()->Is<FContextSwitchesTimingTrack>())
-		{
-			CpuTrack->SetChildTrack(nullptr);
-		}
-	}
-}
-
-////////////////////////////////////////////////////////////////////////////////////////////////////
-
-void FContextSwitchesSharedState::SetCoreTracksVisible(bool bOnOff)
-{
-	if (bAreCoreTracksVisible != bOnOff)
-	{
-		bAreCoreTracksVisible = bOnOff;
-
-		if (bAreCoreTracksVisible)
-		{
-			AddCoreTracks();
-		}
-		else
-		{
-			RemoveCoreTracks();
-		}
-	}
-}
-
-////////////////////////////////////////////////////////////////////////////////////////////////////
-
-void FContextSwitchesSharedState::SetContextSwitchesVisible(bool bOnOff)
-{
-	if (bAreContextSwitchesVisible != bOnOff)
-	{
-		bAreContextSwitchesVisible = bOnOff;
-
-		if (bAreContextSwitchesVisible)
-		{
-			AddContextSwitchesChildTracks();
-		}
-		else
-		{
-			RemoveContextSwitchesChildTracks();
-		}
-	}
-}
 
 ////////////////////////////////////////////////////////////////////////////////////////////////////
 
-void FContextSwitchesSharedState::SetOverlaysVisible(bool bOnOff)
-{
-	bAreOverlaysVisible = bOnOff;
-}
-
-////////////////////////////////////////////////////////////////////////////////////////////////////
-
-void FContextSwitchesSharedState::SetExtendedLinesVisible(bool bOnOff)
-{
-	bAreExtendedLinesVisible = bOnOff;
-}
-
-////////////////////////////////////////////////////////////////////////////////////////////////////
-// FContextSwitchesTimingTrack
-////////////////////////////////////////////////////////////////////////////////////////////////////
-
 INSIGHTS_IMPLEMENT_RTTI(FContextSwitchesTimingTrack)
 
 ////////////////////////////////////////////////////////////////////////////////////////////////////
@@ -346,7 +34,7 @@ void FContextSwitchesTimingTrack::BuildDrawState(ITimingEventsTrackDrawStateBuil
 	
 	TraceServices::FAnalysisSessionReadScope SessionReadScope(*Session.Get());
 
-	const IContextSwitchProvider* ContextSwitchesProvider = TraceServices::ReadContextSwitchProvider(*Session.Get());
+	const TraceServices::IContextSwitchesProvider* ContextSwitchesProvider = TraceServices::ReadContextSwitchesProvider(*Session.Get());
 
 	if (ContextSwitchesProvider == nullptr)
 	{
@@ -355,10 +43,10 @@ void FContextSwitchesTimingTrack::BuildDrawState(ITimingEventsTrackDrawStateBuil
 
 	const FTimingTrackViewport& Viewport = Context.GetViewport();
 
-	ContextSwitchesProvider->EnumerateContextSwitches(ThreadId, Viewport.GetStartTime(), Viewport.GetEndTime(), [&Builder](const FContextSwitch& ContextSwitch)
+	ContextSwitchesProvider->EnumerateContextSwitches(ThreadId, Viewport.GetStartTime(), Viewport.GetEndTime(), [&Builder](const TraceServices::FContextSwitch& ContextSwitch)
 		{
 			Builder.AddEvent(ContextSwitch.Start, ContextSwitch.End, 0, *FString::Printf(TEXT("Core %d"), ContextSwitch.CoreNumber), 0, 0);
-			return EContextSwitchEnumerationResult::Continue;
+			return TraceServices::EContextSwitchEnumerationResult::Continue;
 		});
 }
 
@@ -382,7 +70,7 @@ void FContextSwitchesTimingTrack::DrawLineEvents(const ITimingTrackDrawContext&
 {
 	const FTimingViewDrawHelper& Helper = *static_cast<const FTimingViewDrawHelper*>(&Context.GetHelper());
 
-	if (Context.GetEventFilter().IsValid() || HasCustomFilter())
+	if ((Context.GetEventFilter().IsValid() && Context.GetEventFilter()->FilterTrack(*this)) || HasCustomFilter())
 	{
 		Helper.DrawFadedLineEvents(GetDrawState(), *this, OffsetY, 0.1f);
 
@@ -482,107 +170,60 @@ void FContextSwitchesTimingTrack::PostDraw(const ITimingTrackDrawContext& Contex
 
 ////////////////////////////////////////////////////////////////////////////////////////////////////
 
-const TSharedPtr<const ITimingEvent> FContextSwitchesTimingTrack::GetEvent(float InPosX, float InPosY, const FTimingTrackViewport& Viewport) const
+const TSharedPtr<const ITimingEvent> FContextSwitchesTimingTrack::GetEvent(double InTime, double SecondsPerPixel, int32 Depth) const
 {
-	TSharedPtr<FContextSwitchTimingEvent> TimingEvent;
-
-	const FTimingViewLayout& Layout = Viewport.GetLayout();
-
-	float TopLaneY = 0.0f;
-	float TrackLanesHeight = 0.0f;
-	if (ChildTrack.IsValid())
-	{
-		const float HeaderDY = InPosY - ChildTrack->GetPosY();
-		if (HeaderDY >= 0 && HeaderDY < ChildTrack->GetHeight())
-		{
-			return ChildTrack->GetEvent(InPosX, InPosY, Viewport);
-		}
-		TopLaneY = GetPosY() + 1.0f + Layout.TimelineDY + ChildTrack->GetHeight() + Layout.ChildTimelineDY;
-		TrackLanesHeight = GetHeight() - ChildTrack->GetHeight() - 1.0f - 2 * Layout.TimelineDY - Layout.ChildTimelineDY;
-	}
-	else
+	TSharedPtr<const TraceServices::IAnalysisSession> Session = FInsightsManager::Get()->GetSession();
+	if (!Session.IsValid())
 	{
-		if (IsChildTrack())
-		{
-			TopLaneY = GetPosY();
-			TrackLanesHeight = GetHeight();
-		}
-		else
-		{
-			TopLaneY = GetPosY() + 1.0f + Layout.TimelineDY;
-			TrackLanesHeight = GetHeight() - 1.0f - 2 * Layout.TimelineDY;
-		}
+		return nullptr;
 	}
 
-	const float DY = InPosY - TopLaneY;
+	TraceServices::FAnalysisSessionReadScope SessionReadScope(*Session.Get());
 
-	// If mouse is not above first sub-track or below last sub-track...
-	if (DY >= 0 && DY < TrackLanesHeight)
+	const TraceServices::IContextSwitchesProvider* ContextSwitchesProvider = TraceServices::ReadContextSwitchesProvider(*Session.Get());
+	if (ContextSwitchesProvider == nullptr)
 	{
-		const int32 Depth = DY / (Layout.EventH + Layout.EventDY);
-
-		const double SecondsPerPixel = 1.0 / Viewport.GetScaleX();
-
-		const double EventTime = Viewport.SlateUnitsToTime(InPosX);
-
-		TSharedPtr<const TraceServices::IAnalysisSession> Session = FInsightsManager::Get()->GetSession();
-
-		if (!Session.IsValid())
-		{
-			return TimingEvent;
-		}
-
-		TraceServices::FAnalysisSessionReadScope SessionReadScope(*Session.Get());
+		return nullptr;
+	}
 
-		const IContextSwitchProvider* ContextSwitchesProvider = TraceServices::ReadContextSwitchProvider(*Session.Get());
+	TraceServices::FContextSwitch BestMatchContextSwitch;
+	double Delta = 2 * SecondsPerPixel;
 
-		if (ContextSwitchesProvider == nullptr)
+	ContextSwitchesProvider->EnumerateContextSwitches(ThreadId, InTime - Delta, InTime + Delta,
+		[InTime, &BestMatchContextSwitch, &Delta](const TraceServices::FContextSwitch& ContextSwitch)
 		{
-			return TimingEvent;
-		}
-
-		FContextSwitch BestMatchContextSwitch;
-		double Delta = 2 * SecondsPerPixel;
-
-		ContextSwitchesProvider->EnumerateContextSwitches(ThreadId, EventTime - 2 * SecondsPerPixel, EventTime + 2 * SecondsPerPixel,
-			[EventTime, &BestMatchContextSwitch, &Delta](const FContextSwitch& ContextSwitch)
+			if (ContextSwitch.Start <= InTime && ContextSwitch.End >= InTime)
 			{
-				if (ContextSwitch.Start <= EventTime && ContextSwitch.End >= EventTime)
-				{
-					BestMatchContextSwitch = ContextSwitch;
-					Delta = 0.0f;
-					return EContextSwitchEnumerationResult::Stop;
-				}
+				BestMatchContextSwitch = ContextSwitch;
+				Delta = 0.0f;
+				return TraceServices::EContextSwitchEnumerationResult::Stop;
+			}
 
-				if (ContextSwitch.End <= EventTime)
-				{
-					if (EventTime - ContextSwitch.End < Delta)
-					{
-						Delta = EventTime - ContextSwitch.End;
-						BestMatchContextSwitch = ContextSwitch;
-					}
-				}
+			double DeltaLeft = InTime - ContextSwitch.End;
+			if (DeltaLeft >= 0 && DeltaLeft < Delta)
+			{
+				Delta = DeltaLeft;
+				BestMatchContextSwitch = ContextSwitch;
+			}
 
-				if (ContextSwitch.Start >= EventTime)
-				{
-					if (ContextSwitch.Start - EventTime < Delta)
-					{
-						Delta = ContextSwitch.Start - EventTime;
-						BestMatchContextSwitch = ContextSwitch;
-					}
-				}
+			double DeltaRight = ContextSwitch.Start - InTime;
+			if (DeltaRight >= 0 && DeltaRight < Delta)
+			{
+				Delta = DeltaRight;
+				BestMatchContextSwitch = ContextSwitch;
+			}
 
-				return EContextSwitchEnumerationResult::Continue;
-			});
+			return TraceServices::EContextSwitchEnumerationResult::Continue;
+		});
 
-		if (Delta < 2 * SecondsPerPixel)
-		{
-			TimingEvent = MakeShared<FContextSwitchTimingEvent>(SharedThis(this), BestMatchContextSwitch.Start, BestMatchContextSwitch.End, 0);
-			TimingEvent->SetCoreNumber(BestMatchContextSwitch.CoreNumber);
-		}
+	if (Delta < 2 * SecondsPerPixel)
+	{
+		TSharedPtr<FContextSwitchTimingEvent> TimingEvent = MakeShared<FContextSwitchTimingEvent>(SharedThis(this), BestMatchContextSwitch.Start, BestMatchContextSwitch.End, 0);
+		TimingEvent->SetCoreNumber(BestMatchContextSwitch.CoreNumber);
+		return TimingEvent;
 	}
 
-	return TimingEvent;
+	return nullptr;
 }
 
 ////////////////////////////////////////////////////////////////////////////////////////////////////
@@ -601,7 +242,7 @@ void FContextSwitchesTimingTrack::InitTooltip(FTooltipDrawState& InOutTooltip, c
 
 	InOutTooltip.AddNameValueTextLine(TEXT("Start Time:"), TimeUtils::FormatTimeAuto(InTooltipEvent.GetStartTime(), 6));
 	InOutTooltip.AddNameValueTextLine(TEXT("End Time:"), TimeUtils::FormatTimeAuto(InTooltipEvent.GetEndTime(), 6));
-	InOutTooltip.AddNameValueTextLine(TEXT("Duration:"), TimeUtils::FormatTimeAuto(InTooltipEvent.GetDuration(), 6));
+	InOutTooltip.AddNameValueTextLine(TEXT("Duration:"), TimeUtils::FormatTimeAuto(InTooltipEvent.GetDuration()));
 
 	InOutTooltip.UpdateLayout();
 }
diff --git a/Engine/Source/Developer/TraceInsights/Private/Insights/ContextSwitches/ViewModels/ContextSwitchesTimingTrack.h b/Engine/Source/Developer/TraceInsights/Private/Insights/ContextSwitches/ViewModels/ContextSwitchesTimingTrack.h
--- a/Engine/Source/Developer/TraceInsights/Private/Insights/ContextSwitches/ViewModels/ContextSwitchesTimingTrack.h
+++ b/Engine/Source/Developer/TraceInsights/Private/Insights/ContextSwitches/ViewModels/ContextSwitchesTimingTrack.h
@@ -5,100 +5,16 @@
 #include "CoreMinimal.h"
 
 // Insights
-#include "Insights/ITimingViewExtender.h"
 #include "Insights/ViewModels/TimingEventsTrack.h"
 
-class FThreadTrackEvent;
-class STimingView;
-
 ////////////////////////////////////////////////////////////////////////////////////////////////////
 
-namespace Insights
-{
-class FContextSwitchesTimingTrack;
+class FThreadTrackEvent;
 
-class FContextSwitchesSharedState : public Insights::ITimingViewExtender, public TSharedFromThis<FContextSwitchesSharedState>
+namespace Insights
 {
 
-public:
-	FContextSwitchesSharedState(STimingView* InTimingView);
-	virtual ~FContextSwitchesSharedState() = default;
-
-	//////////////////////////////////////////////////
-	// ITimingViewExtender interface
-
-	virtual void OnBeginSession(Insights::ITimingViewSession& InSession) override;
-	virtual void OnEndSession(Insights::ITimingViewSession& InSession) override;
-	virtual void Tick(Insights::ITimingViewSession& InSession, const TraceServices::IAnalysisSession& InAnalysisSession) override;
-	virtual void ExtendFilterMenu(Insights::ITimingViewSession& InSession, FMenuBuilder& InMenuBuilder) override;
-
-	//////////////////////////////////////////////////
-
-	STimingView* GetTimingView() { return TimingView; }
-
-	bool AreContextSwitchesAvailable() const;
-
-	bool AreCoreTracksVisible() const { return bAreCoreTracksVisible; }
-	void ShowCoreTracks() { SetCoreTracksVisible(true); }
-	void HideCoreTracks() { SetCoreTracksVisible(false); }
-	void ToggleCoreTracks() { SetCoreTracksVisible(!bAreCoreTracksVisible); }
-	void SetCoreTracksVisible(bool bOnOff);
-
-	bool AreContextSwitchesVisible() const { return bAreContextSwitchesVisible; }
-	void ShowContextSwitches() { SetContextSwitchesVisible(true); }
-	void HideContextSwitches() { SetContextSwitchesVisible(false); }
-	void ToggleContextSwitches() { SetContextSwitchesVisible(!bAreContextSwitchesVisible); }
-	void SetContextSwitchesVisible(bool bOnOff);
-
-	bool AreOverlaysVisible() const { return bAreOverlaysVisible; }
-	void ShowOverlays() { SetOverlaysVisible(true); }
-	void HideOverlays() { SetOverlaysVisible(false); }
-	void ToggleOverlays() { SetOverlaysVisible(!bAreOverlaysVisible); }
-	void SetOverlaysVisible(bool bOnOff);
-
-	bool AreExtendedLinesVisible() const { return bAreExtendedLinesVisible; }
-	void ShowExtendedLines() { SetExtendedLinesVisible(true); }
-	void HideExtendedLines() { SetExtendedLinesVisible(false); }
-	void ToggleExtendedLines() { SetExtendedLinesVisible(!bAreExtendedLinesVisible); }
-	void SetExtendedLinesVisible(bool bOnOff);
-
-	void AddCommands();
-
-private:
-	void AddCoreTracks();
-	void RemoveCoreTracks();
-
-	void AddContextSwitchesChildTracks();
-	void RemoveContextSwitchesChildTracks();
-
-	void BuildSubMenu(FMenuBuilder& InMenuBuilder);
-
-	void ContextMenu_ShowCoreTracks_Execute() { ToggleCoreTracks(); }
-	bool ContextMenu_ShowCoreTracks_CanExecute() const { return AreContextSwitchesAvailable(); }
-	bool ContextMenu_ShowCoreTracks_IsChecked() const { return AreCoreTracksVisible(); }
-
-	void ContextMenu_ShowContextSwitches_Execute() { ToggleContextSwitches(); }
-	bool ContextMenu_ShowContextSwitches_CanExecute() const { return AreContextSwitchesAvailable(); }
-	bool ContextMenu_ShowContextSwitches_IsChecked() const { return AreContextSwitchesVisible(); }
-
-	void ContextMenu_ShowOverlays_Execute() { ToggleOverlays(); }
-	bool ContextMenu_ShowOverlays_CanExecute() const { return AreContextSwitchesAvailable() && AreContextSwitchesVisible(); }
-	bool ContextMenu_ShowOverlays_IsChecked() const { return AreOverlaysVisible(); }
-
-	void ContextMenu_ShowExtendedLines_Execute() { ToggleExtendedLines(); }
-	bool ContextMenu_ShowExtendedLines_CanExecute() const { return AreContextSwitchesAvailable() && AreContextSwitchesVisible() && AreOverlaysVisible(); }
-	bool ContextMenu_ShowExtendedLines_IsChecked() const { return AreExtendedLinesVisible(); }
-
-private:
-	STimingView* TimingView;
-
-	bool bAreContextSwitchesVisible;
-	bool bAreCoreTracksVisible;
-	bool bAreOverlaysVisible;
-	bool bAreExtendedLinesVisible;
-
-	bool bSyncWithProviders;
-};
+class FContextSwitchesSharedState;
 
 ////////////////////////////////////////////////////////////////////////////////////////////////////
 
@@ -107,17 +23,18 @@ class FContextSwitchesTimingTrack : public FTimingEventsTrack
 	INSIGHTS_DECLARE_RTTI(FContextSwitchesTimingTrack, FTimingEventsTrack)
 
 public:
-	explicit FContextSwitchesTimingTrack(FContextSwitchesSharedState& InSharedState, const FString& InName, uint32 InTimelineIndex, uint32 ThreadId)
+	explicit FContextSwitchesTimingTrack(FContextSwitchesSharedState& InSharedState, const FString& InName, uint32 InTimelineIndex, uint32 InThreadId)
 		: FTimingEventsTrack(InName)
-		, TimelineIndex(InTimelineIndex)
-		, ThreadId(ThreadId)
 		, SharedState(InSharedState)
+		, TimelineIndex(InTimelineIndex)
+		, ThreadId(InThreadId)
 	{
 	}
 
 	virtual ~FContextSwitchesTimingTrack() {}
 
 	uint32 GetTimelineIndex() const { return TimelineIndex; }
+	uint32 GetThreadId() const { return ThreadId; }
 
 	virtual void BuildDrawState(ITimingEventsTrackDrawStateBuilder& Builder, const ITimingTrackUpdateContext& Context) override;
 	virtual void BuildFilteredDrawState(ITimingEventsTrackDrawStateBuilder& Builder, const ITimingTrackUpdateContext& Context) override;
@@ -127,16 +44,16 @@ public:
 
 	virtual void InitTooltip(FTooltipDrawState& InOutTooltip, const ITimingEvent& InTooltipEvent) const override;
 
-	virtual const TSharedPtr<const ITimingEvent> GetEvent(float InPosX, float InPosY, const FTimingTrackViewport& Viewport) const override;
-
 protected:
+	virtual const TSharedPtr<const ITimingEvent> GetEvent(double InTime, double SecondsPerPixel, int32 Depth) const override;
+
 	void DrawLineEvents(const ITimingTrackDrawContext& Context, const float OffsetY = 1.0f) const;
 
 private:
+	FContextSwitchesSharedState& SharedState;
+
 	uint32 TimelineIndex;
 	uint32 ThreadId;
-
-	FContextSwitchesSharedState& SharedState;
 };
 
 ////////////////////////////////////////////////////////////////////////////////////////////////////
diff --git a/Engine/Source/Developer/TraceInsights/Private/Insights/ContextSwitches/ViewModels/CpuCoreTimingTrack.cpp b/Engine/Source/Developer/TraceInsights/Private/Insights/ContextSwitches/ViewModels/CpuCoreTimingTrack.cpp
new file mode 100644
--- /dev/null
+++ b/Engine/Source/Developer/TraceInsights/Private/Insights/ContextSwitches/ViewModels/CpuCoreTimingTrack.cpp
@@ -0,0 +1,290 @@
+// Copyright Epic Games, Inc. All Rights Reserved.
+
+#include "CpuCoreTimingTrack.h"
+
+#include "TraceServices/Model/ContextSwitches.h"
+#include "TraceServices/Model/Threads.h"
+
+// Insights
+#include "Insights/Common/TimeUtils.h"
+#include "Insights/ContextSwitches/ViewModels/ContextSwitchesSharedState.h"
+#include "Insights/ContextSwitches/ViewModels/ContextSwitchTimingEvent.h"
+#include "Insights/InsightsManager.h"
+#include "Insights/ViewModels/TimingTrackViewport.h"
+#include "Insights/ViewModels/TimingViewDrawHelper.h"
+#include "Insights/ViewModels/TooltipDrawState.h"
+
+#define LOCTEXT_NAMESPACE "FCpuCoreTimingTrack"
+
+namespace Insights
+{
+
+////////////////////////////////////////////////////////////////////////////////////////////////////
+
+INSIGHTS_IMPLEMENT_RTTI(FCpuCoreTimingTrack)
+
+////////////////////////////////////////////////////////////////////////////////////////////////////
+
+void FCpuCoreTimingTrack::BuildDrawState(ITimingEventsTrackDrawStateBuilder& Builder, const ITimingTrackUpdateContext& Context)
+{
+	TSharedPtr<const TraceServices::IAnalysisSession> Session = FInsightsManager::Get()->GetSession();
+	if (!Session.IsValid())
+	{
+		return;
+	}
+
+	TraceServices::FAnalysisSessionReadScope SessionReadScope(*Session.Get());
+
+	const TraceServices::IContextSwitchesProvider* ContextSwitchesProvider = TraceServices::ReadContextSwitchesProvider(*Session.Get());
+	if (ContextSwitchesProvider == nullptr)
+	{
+		return;
+	}
+
+	const FTimingTrackViewport& Viewport = Context.GetViewport();
+
+	ContextSwitchesProvider->EnumerateCpuCoreEvents(CoreNumber, Viewport.GetStartTime(), Viewport.GetEndTime(), [&Builder, ContextSwitchesProvider](const TraceServices::FCpuCoreEvent& CpuCoreEvent)
+		{
+			constexpr uint32 EventDepth = 0;
+			const uint32 EventColor = FTimingEvent::ComputeEventColor(CpuCoreEvent.SystemThreadId);
+			const uint32 SystemThreadId = CpuCoreEvent.SystemThreadId;
+
+			Builder.AddEvent(CpuCoreEvent.Start, CpuCoreEvent.End, EventDepth, EventColor,
+				[SystemThreadId](float Width) -> const FString
+				{
+					const TCHAR* ThreadName = nullptr;
+
+					TSharedPtr<const TraceServices::IAnalysisSession> Session = FInsightsManager::Get()->GetSession();
+					if (Session.IsValid())
+					{
+						TraceServices::FAnalysisSessionReadScope SessionReadScope(*Session.Get());
+						const TraceServices::IContextSwitchesProvider* ContextSwitchesProvider = TraceServices::ReadContextSwitchesProvider(*Session.Get());
+						if (ContextSwitchesProvider)
+						{
+							uint32 ThreadId;
+							if (ContextSwitchesProvider->GetThreadId(SystemThreadId, ThreadId))
+							{
+								const TraceServices::IThreadProvider& ThreadProvider = TraceServices::ReadThreadProvider(*Session.Get());
+								ThreadName = ThreadProvider.GetThreadName(ThreadId);
+							}
+						}
+					}
+
+					if (ThreadName)
+					{
+						return FString(ThreadName);
+					}
+					else
+					{
+						return FString::Printf(TEXT("Unknown %d"), SystemThreadId);
+					}
+				});
+
+			return TraceServices::EContextSwitchEnumerationResult::Continue;
+		});
+}
+
+////////////////////////////////////////////////////////////////////////////////////////////////////
+
+void FCpuCoreTimingTrack::BuildFilteredDrawState(ITimingEventsTrackDrawStateBuilder& Builder, const ITimingTrackUpdateContext& Context)
+{
+	const TSharedPtr<ITimingEventFilter> EventFilterPtr = Context.GetEventFilter();
+	if (EventFilterPtr.IsValid() && EventFilterPtr->FilterTrack(*this))
+	{
+		bool bFilterOnlyByEventType = false; // this is the most often use case, so the below code tries to optimize it
+		uint64 FilterEventType = 0;
+		if (EventFilterPtr->Is<FTimingEventFilterByEventType>())
+		{
+			bFilterOnlyByEventType = true;
+			const FTimingEventFilterByEventType& EventFilter = EventFilterPtr->As<FTimingEventFilterByEventType>();
+			FilterEventType = EventFilter.GetEventType();
+		}
+
+		TSharedPtr<const TraceServices::IAnalysisSession> Session = FInsightsManager::Get()->GetSession();
+		if (Session.IsValid() && TraceServices::ReadContextSwitchesProvider(*Session.Get()))
+		{
+			TraceServices::FAnalysisSessionReadScope SessionReadScope(*Session.Get());
+
+			const TraceServices::IContextSwitchesProvider* ContextSwitchesProvider = TraceServices::ReadContextSwitchesProvider(*Session.Get());
+
+			const FTimingTrackViewport& Viewport = Context.GetViewport();
+
+			if (bFilterOnlyByEventType)
+			{
+				ContextSwitchesProvider->EnumerateCpuCoreEvents(CoreNumber, Viewport.GetStartTime(), Viewport.GetEndTime(), [&Builder, ContextSwitchesProvider, FilterEventType](const TraceServices::FCpuCoreEvent& CpuCoreEvent)
+					{
+						if (CpuCoreEvent.SystemThreadId != FilterEventType)
+						{
+							return TraceServices::EContextSwitchEnumerationResult::Continue;
+						}
+
+						constexpr uint32 EventDepth = 0;
+						const uint32 EventColor = FTimingEvent::ComputeEventColor(CpuCoreEvent.SystemThreadId);
+						const uint32 SystemThreadId = CpuCoreEvent.SystemThreadId;
+
+						Builder.AddEvent(CpuCoreEvent.Start, CpuCoreEvent.End, EventDepth, EventColor,
+							[SystemThreadId](float Width) -> const FString
+							{
+								const TCHAR* ThreadName = nullptr;
+
+								TSharedPtr<const TraceServices::IAnalysisSession> Session = FInsightsManager::Get()->GetSession();
+								if (Session.IsValid())
+								{
+									TraceServices::FAnalysisSessionReadScope SessionReadScope(*Session.Get());
+									const TraceServices::IContextSwitchesProvider* ContextSwitchesProvider = TraceServices::ReadContextSwitchesProvider(*Session.Get());
+									if (ContextSwitchesProvider)
+									{
+										uint32 ThreadId;
+										if (ContextSwitchesProvider->GetThreadId(SystemThreadId, ThreadId))
+										{
+											const TraceServices::IThreadProvider& ThreadProvider = TraceServices::ReadThreadProvider(*Session.Get());
+											ThreadName = ThreadProvider.GetThreadName(ThreadId);
+										}
+									}
+								}
+
+								if (ThreadName)
+								{
+									return FString(ThreadName);
+								}
+								else
+								{
+									return FString::Printf(TEXT("Unknown %d"), SystemThreadId);
+								}
+							});
+
+						return TraceServices::EContextSwitchEnumerationResult::Continue;
+					});
+			}
+			else // generic filter
+			{
+				//TODO: if (EventFilterPtr->FilterEvent(TimingEvent))
+			}
+		}
+	}
+}
+
+////////////////////////////////////////////////////////////////////////////////////////////////////
+
+void FCpuCoreTimingTrack::Draw(const ITimingTrackDrawContext& Context) const
+{
+	FTimingEventsTrack::Draw(Context);
+}
+
+////////////////////////////////////////////////////////////////////////////////////////////////////
+
+const TSharedPtr<const ITimingEvent> FCpuCoreTimingTrack::GetEvent(double InTime, double SecondsPerPixel, int32 Depth) const
+{
+	TSharedPtr<const TraceServices::IAnalysisSession> Session = FInsightsManager::Get()->GetSession();
+	if (!Session.IsValid())
+	{
+		return nullptr;
+	}
+
+	TraceServices::FAnalysisSessionReadScope SessionReadScope(*Session.Get());
+
+	const TraceServices::IContextSwitchesProvider* ContextSwitchesProvider = TraceServices::ReadContextSwitchesProvider(*Session.Get());
+	if (ContextSwitchesProvider == nullptr)
+	{
+		return nullptr;
+	}
+
+	TraceServices::FCpuCoreEvent BestMatchEvent;
+	double Delta = 2 * SecondsPerPixel;
+
+	ContextSwitchesProvider->EnumerateCpuCoreEvents(CoreNumber, InTime - Delta, InTime + Delta,
+		[InTime, &BestMatchEvent, &Delta](const TraceServices::FCpuCoreEvent& CpuCoreEvent)
+		{
+			if (CpuCoreEvent.Start <= InTime && CpuCoreEvent.End >= InTime)
+			{
+				BestMatchEvent = CpuCoreEvent;
+				Delta = 0.0f;
+				return TraceServices::EContextSwitchEnumerationResult::Stop;
+			}
+
+			double DeltaLeft = InTime - CpuCoreEvent.End;
+			if (DeltaLeft >= 0 && DeltaLeft < Delta)
+			{
+				Delta = DeltaLeft;
+				BestMatchEvent = CpuCoreEvent;
+			}
+
+			double DeltaRight = CpuCoreEvent.Start - InTime;
+			if (DeltaRight >= 0 && DeltaRight < Delta)
+			{
+				Delta = DeltaRight;
+				BestMatchEvent = CpuCoreEvent;
+			}
+
+			return TraceServices::EContextSwitchEnumerationResult::Continue;
+		});
+
+	if (Delta < 2 * SecondsPerPixel)
+	{
+		TSharedPtr<FCpuCoreTimingEvent> TimingEvent = MakeShared<FCpuCoreTimingEvent>(SharedThis(this), BestMatchEvent.Start, BestMatchEvent.End, 0);
+		TimingEvent->SetSystemThreadId(BestMatchEvent.SystemThreadId);
+		return TimingEvent;
+	}
+
+	return nullptr;
+}
+
+////////////////////////////////////////////////////////////////////////////////////////////////////
+
+void FCpuCoreTimingTrack::InitTooltip(FTooltipDrawState& InOutTooltip, const ITimingEvent& InTooltipEvent) const
+{
+	InOutTooltip.ResetContent();
+
+	if (!InTooltipEvent.CheckTrack(this) || !InTooltipEvent.Is<FCpuCoreTimingEvent>())
+	{
+		return;
+	}
+
+	const FCpuCoreTimingEvent& CpuCoreEvent = InTooltipEvent.As<FCpuCoreTimingEvent>();
+
+	const uint32 SystemThreadId = CpuCoreEvent.GetSystemThreadId();
+	uint32 ThreadId = ~0;
+	const TCHAR* ThreadName = nullptr;
+	{
+		TSharedPtr<const TraceServices::IAnalysisSession> Session = FInsightsManager::Get()->GetSession();
+		if (Session.IsValid())
+		{
+			TraceServices::FAnalysisSessionReadScope SessionReadScope(*Session.Get());
+			const TraceServices::IContextSwitchesProvider* ContextSwitchesProvider = TraceServices::ReadContextSwitchesProvider(*Session.Get());
+			if (ContextSwitchesProvider)
+			{
+				if (ContextSwitchesProvider->GetThreadId(SystemThreadId, ThreadId))
+				{
+					const TraceServices::IThreadProvider& ThreadProvider = TraceServices::ReadThreadProvider(*Session.Get());
+					ThreadName = ThreadProvider.GetThreadName(ThreadId);
+				}
+			}
+		}
+	}
+	if (ThreadName)
+	{
+		InOutTooltip.AddTitle(ThreadName);
+	}
+	else
+	{
+		InOutTooltip.AddTitle(TEXT("Unknown Thread"));
+	}
+
+	InOutTooltip.AddNameValueTextLine(TEXT("System Thread Id:"), FString::Printf(TEXT("%d"), SystemThreadId));
+	if (ThreadId != ~0)
+	{
+		InOutTooltip.AddNameValueTextLine(TEXT("Thread Id:"), FString::Printf(TEXT("%d"), ThreadId));
+	}
+
+	InOutTooltip.AddNameValueTextLine(TEXT("Start Time:"), TimeUtils::FormatTimeAuto(InTooltipEvent.GetStartTime(), 6));
+	InOutTooltip.AddNameValueTextLine(TEXT("End Time:"), TimeUtils::FormatTimeAuto(InTooltipEvent.GetEndTime(), 6));
+	InOutTooltip.AddNameValueTextLine(TEXT("Duration:"), TimeUtils::FormatTimeAuto(InTooltipEvent.GetDuration()));
+
+	InOutTooltip.UpdateLayout();
+}
+
+////////////////////////////////////////////////////////////////////////////////////////////////////
+
+} // namespace Insights
+
+#undef LOCTEXT_NAMESPACE
diff --git a/Engine/Source/Developer/TraceInsights/Private/Insights/ContextSwitches/ViewModels/CpuCoreTimingTrack.h b/Engine/Source/Developer/TraceInsights/Private/Insights/ContextSwitches/ViewModels/CpuCoreTimingTrack.h
new file mode 100644
--- /dev/null
+++ b/Engine/Source/Developer/TraceInsights/Private/Insights/ContextSwitches/ViewModels/CpuCoreTimingTrack.h
@@ -0,0 +1,55 @@
+// Copyright Epic Games, Inc. All Rights Reserved.
+
+#pragma once
+
+#include "CoreMinimal.h"
+
+// Insights
+#include "Insights/ViewModels/TimingEventsTrack.h"
+
+////////////////////////////////////////////////////////////////////////////////////////////////////
+
+class FThreadTrackEvent;
+
+namespace Insights
+{
+
+class FContextSwitchesSharedState;
+
+////////////////////////////////////////////////////////////////////////////////////////////////////
+
+class FCpuCoreTimingTrack : public FTimingEventsTrack
+{
+	INSIGHTS_DECLARE_RTTI(FCpuCoreTimingTrack, FTimingEventsTrack)
+
+public:
+	explicit FCpuCoreTimingTrack(FContextSwitchesSharedState& InSharedState, const FString& InName, uint32 InCoreNumber)
+		: FTimingEventsTrack(InName)
+		, SharedState(InSharedState)
+		, CoreNumber(InCoreNumber)
+	{
+	}
+
+	virtual ~FCpuCoreTimingTrack() {}
+
+	uint32 GetCoreNumber() const { return CoreNumber; }
+
+	virtual void BuildDrawState(ITimingEventsTrackDrawStateBuilder& Builder, const ITimingTrackUpdateContext& Context) override;
+	virtual void BuildFilteredDrawState(ITimingEventsTrackDrawStateBuilder& Builder, const ITimingTrackUpdateContext& Context) override;
+
+	virtual void Draw(const ITimingTrackDrawContext& Context) const override;
+
+	virtual void InitTooltip(FTooltipDrawState& InOutTooltip, const ITimingEvent& InTooltipEvent) const override;
+
+protected:
+	virtual const TSharedPtr<const ITimingEvent> GetEvent(double InTime, double SecondsPerPixel, int32 Depth) const override;
+
+private:
+	FContextSwitchesSharedState& SharedState;
+
+	uint32 CoreNumber;
+};
+
+////////////////////////////////////////////////////////////////////////////////////////////////////
+
+} // namespace Insights
diff --git a/Engine/Source/Developer/TraceInsights/Private/Insights/ViewModels/TimingEventsTrack.cpp b/Engine/Source/Developer/TraceInsights/Private/Insights/ViewModels/TimingEventsTrack.cpp
--- a/Engine/Source/Developer/TraceInsights/Private/Insights/ViewModels/TimingEventsTrack.cpp
+++ b/Engine/Source/Developer/TraceInsights/Private/Insights/ViewModels/TimingEventsTrack.cpp
@@ -270,7 +270,7 @@ void FTimingEventsTrack::DrawEvents(const ITimingTrackDrawContext& Context, cons
 {
 	const FTimingViewDrawHelper& Helper = *static_cast<const FTimingViewDrawHelper*>(&Context.GetHelper());
 
-	if (Context.GetEventFilter().IsValid() || HasCustomFilter())
+	if ((Context.GetEventFilter().IsValid() && Context.GetEventFilter()->FilterTrack(*this)) || HasCustomFilter())
 	{
 		Helper.DrawFadedEvents(GetDrawState(), *this, OffsetY, 0.1f);
 
@@ -399,36 +399,44 @@ const TSharedPtr<const ITimingEvent> FTimingEventsTrack::GetEvent(float InPosX,
 	if (DY >= 0 && DY < TrackLanesHeight)
 	{
 		const int32 Depth = DY / (Layout.EventH + Layout.EventDY);
+		const double SecondsPerPixel = 1.0 / Viewport.GetScaleX();
+		const double TimeAtPosX = Viewport.SlateUnitsToTime(InPosX);
 
-		auto EventFilter = [Depth](double, double, uint32 EventDepth)
-		{
-			return EventDepth == Depth;
-		};
+		return GetEvent(TimeAtPosX, SecondsPerPixel, Depth);
+	}
 
-		const double SecondsPerPixel = 1.0 / Viewport.GetScaleX();
+	return nullptr;
+}
 
-		const double StartTime0 = Viewport.SlateUnitsToTime(InPosX);
-		const double EndTime0 = StartTime0;
-		TSharedPtr<const ITimingEvent> FoundEvent = SearchEvent(FTimingEventSearchParameters(StartTime0, EndTime0, ETimingEventSearchFlags::StopAtFirstMatch, EventFilter));
+////////////////////////////////////////////////////////////////////////////////////////////////////
 
-		if (!FoundEvent.IsValid())
-		{
-			const double StartTime = StartTime0;
-			const double EndTime = StartTime0 + SecondsPerPixel; // +1px
-			FoundEvent = SearchEvent(FTimingEventSearchParameters(StartTime, EndTime, ETimingEventSearchFlags::StopAtFirstMatch, EventFilter));
-		}
+const TSharedPtr<const ITimingEvent> FTimingEventsTrack::GetEvent(double InTime, double SecondsPerPixel, int32 Depth) const
+{
+	const double StartTime0 = InTime;
+	const double EndTime0 = InTime;
 
-		if (!FoundEvent.IsValid())
-		{
-			const double StartTime = StartTime0 - SecondsPerPixel; // -1px
-			const double EndTime = StartTime0 + 2.0 * SecondsPerPixel; // +2px
-			FoundEvent = SearchEvent(FTimingEventSearchParameters(StartTime, EndTime, ETimingEventSearchFlags::StopAtFirstMatch, EventFilter));
-		}
+	auto EventFilter = [Depth](double, double, uint32 EventDepth)
+	{
+		return EventDepth == Depth;
+	};
 
-		return FoundEvent;
+	TSharedPtr<const ITimingEvent> FoundEvent = SearchEvent(FTimingEventSearchParameters(InTime, InTime, ETimingEventSearchFlags::StopAtFirstMatch, EventFilter));
+
+	if (!FoundEvent.IsValid())
+	{
+		const double StartTime = InTime;
+		const double EndTime = InTime + SecondsPerPixel; // +1px
+		FoundEvent = SearchEvent(FTimingEventSearchParameters(StartTime, EndTime, ETimingEventSearchFlags::StopAtFirstMatch, EventFilter));
 	}
 
-	return nullptr;
+	if (!FoundEvent.IsValid())
+	{
+		const double StartTime = InTime - SecondsPerPixel; // -1px
+		const double EndTime = InTime + 2.0 * SecondsPerPixel; // +2px
+		FoundEvent = SearchEvent(FTimingEventSearchParameters(StartTime, EndTime, ETimingEventSearchFlags::StopAtFirstMatch, EventFilter));
+	}
+
+	return FoundEvent;
 }
 
 ////////////////////////////////////////////////////////////////////////////////////////////////////
@@ -439,6 +447,8 @@ TSharedPtr<ITimingEventFilter> FTimingEventsTrack::GetFilterByEvent(const TShare
 	{
 		const FTimingEvent& Event = InTimingEvent->As<FTimingEvent>();
 		TSharedRef<FTimingEventFilter> EventFilterRef = MakeShared<FTimingEventFilterByEventType>(Event.GetType());
+		EventFilterRef->SetFilterByTrackTypeName(true);
+		EventFilterRef->SetTrackTypeName(GetTypeName());
 		return EventFilterRef;
 	}
 	return nullptr;
diff --git a/Engine/Source/Developer/TraceInsights/Private/Insights/Widgets/STimingView.cpp b/Engine/Source/Developer/TraceInsights/Private/Insights/Widgets/STimingView.cpp
--- a/Engine/Source/Developer/TraceInsights/Private/Insights/Widgets/STimingView.cpp
+++ b/Engine/Source/Developer/TraceInsights/Private/Insights/Widgets/STimingView.cpp
@@ -3824,6 +3824,8 @@ void STimingView::ToggleEventFilterByEventType(const uint64 EventType)
 	else
 	{
 		TSharedRef<FTimingEventFilterByEventType> NewEventFilter = MakeShared<FTimingEventFilterByEventType>(EventType);
+		NewEventFilter->SetFilterByTrackTypeName(true);
+		NewEventFilter->SetTrackTypeName(FThreadTimingTrack::GetStaticTypeName());
 		SetEventFilter(NewEventFilter); // set new filter
 	}
 }
diff --git a/Engine/Source/Developer/TraceInsights/Public/Insights/ViewModels/TimingEventsTrack.h b/Engine/Source/Developer/TraceInsights/Public/Insights/ViewModels/TimingEventsTrack.h
--- a/Engine/Source/Developer/TraceInsights/Public/Insights/ViewModels/TimingEventsTrack.h
+++ b/Engine/Source/Developer/TraceInsights/Public/Insights/ViewModels/TimingEventsTrack.h
@@ -84,6 +84,8 @@ protected:
 	int32 GetHeaderBackgroundLayerId(const ITimingTrackDrawContext& Context) const;
 	int32 GetHeaderTextLayerId(const ITimingTrackDrawContext& Context) const;
 
+	virtual const TSharedPtr<const ITimingEvent> GetEvent(double InTime, double SecondsPerPixel, int32 Depth) const;
+
 	virtual bool HasCustomFilter() const { return false; }
 
 private:
diff --git a/Engine/Source/Developer/TraceServices/Private/Analyzers/PlatformEventTraceAnalysis.cpp b/Engine/Source/Developer/TraceServices/Private/Analyzers/PlatformEventTraceAnalysis.cpp
--- a/Engine/Source/Developer/TraceServices/Private/Analyzers/PlatformEventTraceAnalysis.cpp
+++ b/Engine/Source/Developer/TraceServices/Private/Analyzers/PlatformEventTraceAnalysis.cpp
@@ -8,11 +8,11 @@ namespace TraceServices
 {
 
 FPlatformEventTraceAnalyzer::FPlatformEventTraceAnalyzer(IAnalysisSession& InSession,
-														 FContextSwitchProvider& ContextSwitchProvider,
-														 FStackSampleProvider& StackSampleProvider)
+														 FContextSwitchesProvider& InContextSwitchesProvider,
+														 FStackSamplesProvider& InStackSamplesProvider)
 	: Session(InSession)
-	, ContextSwitchProvider(ContextSwitchProvider)
-	, StackSampleProvider(StackSampleProvider)
+	, ContextSwitchesProvider(InContextSwitchesProvider)
+	, StackSamplesProvider(InStackSamplesProvider)
 {
 }
 
@@ -38,7 +38,7 @@ bool FPlatformEventTraceAnalyzer::OnEvent(uint16 RouteId, EStyle Style, const FO
 		double End = Context.EventTime.AsSeconds(EventData.GetValue<uint64>("EndTime"));
 		uint32 ThreadId = EventData.GetValue<uint32>("ThreadId");
 		uint32 CoreNumber = EventData.GetValue<uint8>("CoreNumber");
-		ContextSwitchProvider.Add(ThreadId, Start, End, CoreNumber);
+		ContextSwitchesProvider.Add(ThreadId, Start, End, CoreNumber);
 
 		Session.UpdateDurationSeconds(End);
 
@@ -50,7 +50,7 @@ bool FPlatformEventTraceAnalyzer::OnEvent(uint16 RouteId, EStyle Style, const FO
 		double Time = Context.EventTime.AsSeconds(EventData.GetValue<uint64>("Time"));
 		uint32 ThreadId = EventData.GetValue<uint32>("ThreadId");
 		const TArrayReader<uint64>& Addresses = EventData.GetArray<uint64>("Addresses");
-		StackSampleProvider.Add(ThreadId, Time, Addresses.Num(), Addresses.GetData());
+		StackSamplesProvider.Add(ThreadId, Time, Addresses.Num(), Addresses.GetData());
 
 		Session.UpdateDurationSeconds(Time);
 
@@ -65,7 +65,7 @@ bool FPlatformEventTraceAnalyzer::OnEvent(uint16 RouteId, EStyle Style, const FO
 void FPlatformEventTraceAnalyzer::OnThreadInfo(const FThreadInfo& ThreadInfo)
 {
 	FAnalysisSessionEditScope _(Session);
-	ContextSwitchProvider.AddThreadInfo(ThreadInfo.GetId(), ThreadInfo.GetSystemId());
+	ContextSwitchesProvider.AddThreadInfo(ThreadInfo.GetId(), ThreadInfo.GetSystemId());
 }
 
 } // namespace TraceServices
diff --git a/Engine/Source/Developer/TraceServices/Private/Analyzers/PlatformEventTraceAnalysis.h b/Engine/Source/Developer/TraceServices/Private/Analyzers/PlatformEventTraceAnalysis.h
--- a/Engine/Source/Developer/TraceServices/Private/Analyzers/PlatformEventTraceAnalysis.h
+++ b/Engine/Source/Developer/TraceServices/Private/Analyzers/PlatformEventTraceAnalysis.h
@@ -9,16 +9,16 @@ namespace TraceServices
 {
 
 class IAnalysisSession;
-class FContextSwitchProvider;
-class FStackSampleProvider;
+class FContextSwitchesProvider;
+class FStackSamplesProvider;
 
 class FPlatformEventTraceAnalyzer
 	: public UE::Trace::IAnalyzer
 {
 public:
 	FPlatformEventTraceAnalyzer(IAnalysisSession& Session,
-								FContextSwitchProvider& ContextSwitchProvider,
-								FStackSampleProvider& StackSampleProvider);
+								FContextSwitchesProvider& ContextSwitchesProvider,
+								FStackSamplesProvider& StackSamplesProvider);
 	virtual void OnAnalysisBegin(const FOnAnalysisContext& Context) override;
 	virtual bool OnEvent(uint16 RouteId, EStyle Style, const FOnEventContext& Context) override;
 	virtual void OnThreadInfo(const FThreadInfo& ThreadInfo) override;
@@ -31,8 +31,8 @@ private:
 	};
 
 	IAnalysisSession& Session;
-	FContextSwitchProvider& ContextSwitchProvider;
-	FStackSampleProvider& StackSampleProvider;
+	FContextSwitchesProvider& ContextSwitchesProvider;
+	FStackSamplesProvider& StackSamplesProvider;
 };
 
 } // namespace TraceServices
diff --git a/Engine/Source/Developer/TraceServices/Private/Model/ContextSwitches.cpp b/Engine/Source/Developer/TraceServices/Private/Model/ContextSwitches.cpp
--- a/Engine/Source/Developer/TraceServices/Private/Model/ContextSwitches.cpp
+++ b/Engine/Source/Developer/TraceServices/Private/Model/ContextSwitches.cpp
@@ -8,92 +8,248 @@
 namespace TraceServices
 {
 
-const FName FContextSwitchProvider::ProviderName = "ContextSwitchProvider";
+const FName FContextSwitchesProvider::ProviderName = "ContextSwitchesProvider";
 
-FContextSwitchProvider::FContextSwitchProvider(IAnalysisSession& InSession)
+FContextSwitchesProvider::FContextSwitchesProvider(IAnalysisSession& InSession)
 	: Session(InSession)
+	, NumCpuCores(0)
 {
 }
 
-FContextSwitchProvider::~FContextSwitchProvider()
+FContextSwitchesProvider::~FContextSwitchesProvider()
 {
 	for (const auto& ContextSwitches : Threads)
 	{
 		delete ContextSwitches.Value;
 	}
+	Threads.Reset();
+
+	for (const auto& CpuCoreEvents : CpuCores)
+	{
+		if (CpuCoreEvents)
+		{
+			delete CpuCoreEvents;
+		}
+	}
+	CpuCores.Reset();
+}
+
+bool FContextSwitchesProvider::HasData() const
+{
+	Session.ReadAccessCheck();
+
+	return Threads.Num() > 0;
+}
+
+bool FContextSwitchesProvider::GetSystemThreadId(uint32 ThreadId, uint32& OutSystemThreadId) const
+{
+	Session.ReadAccessCheck();
+
+	const uint32* SystemThreadIdPtr = TraceToSystemThreadIdMap.Find(ThreadId);
+	if (SystemThreadIdPtr)
+	{
+		OutSystemThreadId = *SystemThreadIdPtr;
+		return true;
+	}
+	return false;
+}
+
+bool FContextSwitchesProvider::GetThreadId(uint32 SystemThreadId, uint32& OutThreadId) const
+{
+	Session.ReadAccessCheck();
+
+	const uint32* ThreadIdPtr = SystemToTraceThreadIdMap.Find(SystemThreadId);
+	if (ThreadIdPtr)
+	{
+		OutThreadId = *ThreadIdPtr;
+		return true;
+	}
+	return false;
+}
+
+bool FContextSwitchesProvider::GetSystemThreadId(uint32 CoreNumber, double Time, uint32& OutSystemThreadId) const
+{
+	Session.ReadAccessCheck();
+
+	if (CoreNumber >= (uint32)CpuCores.Num())
+	{
+		return false;
+	}
+
+	const TPagedArray<FCpuCoreEvent>* CpuCoreEvents = CpuCores[CoreNumber];
+	if (!CpuCoreEvents)
+	{
+		return false;
+	}
+
+	uint64 PageIndex = Algo::LowerBoundBy(*CpuCoreEvents, Time, [](const TPagedArrayPage<FCpuCoreEvent>& In) -> double
+		{
+			return GetFirstItem(In)->Start;
+		});
+
+	if (PageIndex < GetNum(*CpuCoreEvents))
+	{
+		const TPagedArrayPage<FCpuCoreEvent>& Page = GetData(*CpuCoreEvents)[PageIndex];
+		uint64 Index = Algo::LowerBoundBy(Page, Time, [](const FCpuCoreEvent& In) -> double
+			{
+				return In.Start;
+			});
+
+		if (Index < GetNum(Page))
+		{
+			const FCpuCoreEvent& CpuCoreEvent = GetData(Page)[Index];
+			if (CpuCoreEvent.Start >= Time && CpuCoreEvent.End < Time)
+			{
+				OutSystemThreadId = CpuCoreEvent.SystemThreadId;
+				return true;
+			}
+		}
+	}
+
+	return false;
 }
 
-int32 FContextSwitchProvider::GetCoreNumber(uint32 ThreadId, double Time) const
+bool FContextSwitchesProvider::GetThreadId(uint32 CoreNumber, double Time, uint32& OutThreadId) const
+{
+	Session.ReadAccessCheck();
+
+	uint32 SystemThreadId;
+	if (GetSystemThreadId(CoreNumber, Time, SystemThreadId))
+	{
+		return GetThreadId(SystemThreadId, OutThreadId);
+	}
+
+	return false;
+}
+
+bool FContextSwitchesProvider::GetCoreNumber(uint32 ThreadId, double Time, uint32& OutCoreNumber) const
 {
 	Session.ReadAccessCheck();
 
 	const TPagedArray<FContextSwitch>* ContextSwitches = GetContextSwitches(ThreadId);
 	if (!ContextSwitches)
 	{
-		return -1;
+		return false;
 	}
 
-	uint64 PageIndex = Algo::LowerBoundBy(*ContextSwitches, Time, [](const TPagedArrayPage<FContextSwitch>& In) -> uint64
-	{
-		return GetFirstItem(In)->Start;
-	});
+	uint64 PageIndex = Algo::LowerBoundBy(*ContextSwitches, Time, [](const TPagedArrayPage<FContextSwitch>& In) -> double
+		{
+			return GetFirstItem(In)->Start;
+		});
 
 	if (PageIndex < GetNum(*ContextSwitches))
 	{
 		const TPagedArrayPage<FContextSwitch>& Page = GetData(*ContextSwitches)[PageIndex];
-		uint64 Index = Algo::LowerBoundBy(Page, Time, [](const FContextSwitch& In) -> uint64
-		{
-			return In.Start;
-		});
+		uint64 Index = Algo::LowerBoundBy(Page, Time, [](const FContextSwitch& In) -> double
+			{
+				return In.Start;
+			});
 
 		if (Index < GetNum(Page))
 		{
 			const FContextSwitch& ContextSwitch = GetData(Page)[Index];
 			if (ContextSwitch.Start >= Time && ContextSwitch.End < Time)
 			{
-				return ContextSwitch.CoreNumber;
+				OutCoreNumber = ContextSwitch.CoreNumber;
+				return true;
 			}
 		}
 	}
 
-	return -1;
+	return false;
 }
 
-const TPagedArray<FContextSwitch>* FContextSwitchProvider::GetContextSwitches(uint32 ThreadId) const
+void FContextSwitchesProvider::EnumerateCpuCores(CpuCoreCallback Callback) const
 {
 	Session.ReadAccessCheck();
 
-	if (ThreadIdMap.Contains(ThreadId))
+	const uint32 CoreCount = (uint32)CpuCores.Num();
+	for (uint32 CoreNumber = 0; CoreNumber < CoreCount; ++CoreNumber)
+	{
+		if (CpuCores[CoreNumber] != nullptr)
+		{
+			FCpuCoreInfo CpuCoreInfo;
+			CpuCoreInfo.CoreNumber = CoreNumber;
+			Callback(CpuCoreInfo);
+		}
+	}
+}
+
+void FContextSwitchesProvider::EnumerateCpuCoreEvents(uint32 CoreNumber, double StartTime, double EndTime, CpuCoreEventCallback Callback) const
+{
+	Session.ReadAccessCheck();
+
+	if (CoreNumber >= (uint32)CpuCores.Num())
+	{
+		return;
+	}
+
+	const TPagedArray<FCpuCoreEvent>* CpuCoreEvents = CpuCores[CoreNumber];
+	if (CpuCoreEvents == nullptr)
+	{
+		return;
+	}
+
+	uint64 PageIndex = Algo::UpperBoundBy(*CpuCoreEvents, StartTime, [](const TPagedArrayPage<FCpuCoreEvent>& Page)
+		{
+			return Page.Items[0].Start;
+		});
+
+	if (PageIndex > 0)
 	{
-		auto ContextSwitches = Threads.Find(ThreadIdMap[ThreadId]);
+		--PageIndex;
+	}
+
+	auto Iterator = CpuCoreEvents->GetIteratorFromPage(PageIndex);
+	const FCpuCoreEvent* CurrentCpuCoreEvent = Iterator.NextItem();
+	while (CurrentCpuCoreEvent && CurrentCpuCoreEvent->Start < EndTime)
+	{
+		if (CurrentCpuCoreEvent->End > StartTime)
+		{
+			if (Callback(*CurrentCpuCoreEvent) == EContextSwitchEnumerationResult::Stop)
+			{
+				break;
+			}
+		}
+
+		CurrentCpuCoreEvent = Iterator.NextItem();
+	}
+}
+
+const TPagedArray<FContextSwitch>* FContextSwitchesProvider::GetContextSwitches(uint32 ThreadId) const
+{
+	Session.ReadAccessCheck();
+
+	if (TraceToSystemThreadIdMap.Contains(ThreadId))
+	{
+		auto ContextSwitches = Threads.Find(TraceToSystemThreadIdMap[ThreadId]);
 		return ContextSwitches ? *ContextSwitches : nullptr;
 	}
 
 	return nullptr;
 }
 
-const void FContextSwitchProvider::EnumerateContextSwitches(uint32 ThreadId, double StartTime, double EndTime, ContextSwitchCallback Callback) const
+void FContextSwitchesProvider::EnumerateContextSwitches(uint32 ThreadId, double StartTime, double EndTime, ContextSwitchCallback Callback) const
 {
 	Session.ReadAccessCheck();
 
-	const TPagedArray<FContextSwitch>* PagedArray = GetContextSwitches(ThreadId);
-
-	if (PagedArray == nullptr)
+	const TPagedArray<FContextSwitch>* ContextSwitches = GetContextSwitches(ThreadId);
+	if (ContextSwitches == nullptr)
 	{
 		return;
 	}
 
-	uint64 ContextSwitchPageIndex = Algo::UpperBoundBy(*PagedArray, StartTime, [](const TPagedArrayPage<FContextSwitch>& ContextSwitchPage)
+	uint64 PageIndex = Algo::UpperBoundBy(*ContextSwitches, StartTime, [](const TPagedArrayPage<FContextSwitch>& Page)
 		{
-			return ContextSwitchPage.Items[0].Start;
+			return Page.Items[0].Start;
 		});
 
-	if (ContextSwitchPageIndex > 0)
+	if (PageIndex > 0)
 	{
-		--ContextSwitchPageIndex;
+		--PageIndex;
 	}
 
-	auto Iterator = PagedArray->GetIteratorFromPage(ContextSwitchPageIndex);
+	auto Iterator = ContextSwitches->GetIteratorFromPage(PageIndex);
 	const FContextSwitch* CurrentContextSwitch = Iterator.NextItem();
 	while (CurrentContextSwitch && CurrentContextSwitch->Start < EndTime)
 	{
@@ -109,38 +265,73 @@ const void FContextSwitchProvider::EnumerateContextSwitches(uint32 ThreadId, dou
 	}
 }
 
-void FContextSwitchProvider::Add(uint32 ThreadId, double Start, double End, uint32 CoreNumber)
+void FContextSwitchesProvider::Add(uint32 SystemThreadId, double Start, double End, uint32 CoreNumber)
 {
 	Session.WriteAccessCheck();
 
-	TPagedArray<FContextSwitch>** ContextSwitches = Threads.Find(ThreadId);
+	//////////////////////////////////////////////////
+	// Context Switch events
+
+	TPagedArray<FContextSwitch>** ContextSwitches = Threads.Find(SystemThreadId);
 	if (!ContextSwitches)
 	{
-		ContextSwitches = &Threads.Add(ThreadId, new TPagedArray<FContextSwitch>(Session.GetLinearAllocator(), 4096));
+		ContextSwitches = &Threads.Add(SystemThreadId, new TPagedArray<FContextSwitch>(Session.GetLinearAllocator(), 4096));
 	}
 
 	FContextSwitch& ContextSwitch = (*ContextSwitches)->PushBack();
 	ContextSwitch.Start = Start;
 	ContextSwitch.End = End;
 	ContextSwitch.CoreNumber = CoreNumber;
-}
 
-bool FContextSwitchProvider::HasData() const
-{
-	return Threads.Num() > 0;
+	//////////////////////////////////////////////////
+	// Cpu Core events
+
+	if (CoreNumber >= (uint32)CpuCores.Num())
+	{
+		CpuCores.AddDefaulted(CoreNumber - CpuCores.Num() + 1);
+	}
+	TPagedArray<FCpuCoreEvent>* CpuCoreEvents = CpuCores[CoreNumber];
+	if (!CpuCoreEvents)
+	{
+		CpuCoreEvents = new TPagedArray<FCpuCoreEvent>(Session.GetLinearAllocator(), 4096);
+		++NumCpuCores;
+		CpuCores[CoreNumber] = CpuCoreEvents;
+	}
+
+#if 0 // for debugging
+	auto Page = CpuCoreEvents->GetLastPage();
+	if (Page)
+	{
+		const FCpuCoreEvent* LastCpuCoreEvent = TraceServices::GetLastItem(*Page);
+		if (LastCpuCoreEvent)
+		{
+			check(LastCpuCoreEvent->End <= Start);
+		}
+	}
+#endif
+
+	FCpuCoreEvent& CpuCoreEvent = CpuCoreEvents->PushBack();
+	CpuCoreEvent.Start = Start;
+	CpuCoreEvent.End = End;
+	CpuCoreEvent.SystemThreadId = SystemThreadId;
 }
 
-void FContextSwitchProvider::AddThreadInfo(uint32 TraceThreadId, uint32 SystemThreadId)
+void FContextSwitchesProvider::AddThreadInfo(uint32 ThreadId, uint32 SystemThreadId)
 {
-	if (!ThreadIdMap.Contains(TraceThreadId))
+	if (!TraceToSystemThreadIdMap.Contains(ThreadId))
+	{
+		TraceToSystemThreadIdMap.Add(ThreadId, SystemThreadId);
+	}
+
+	if (!SystemToTraceThreadIdMap.Contains(SystemThreadId))
 	{
-		ThreadIdMap.Add(TraceThreadId, SystemThreadId);
+		SystemToTraceThreadIdMap.Add(SystemThreadId, ThreadId);
 	}
 }
 
-const IContextSwitchProvider* ReadContextSwitchProvider(const IAnalysisSession& Session)
+const IContextSwitchesProvider* ReadContextSwitchesProvider(const IAnalysisSession& Session)
 {
-	return Session.ReadProvider<IContextSwitchProvider>(FContextSwitchProvider::ProviderName);
+	return Session.ReadProvider<IContextSwitchesProvider>(FContextSwitchesProvider::ProviderName);
 }
 
 } // namespace TraceServices
diff --git a/Engine/Source/Developer/TraceServices/Private/Model/ContextSwitchesPrivate.h b/Engine/Source/Developer/TraceServices/Private/Model/ContextSwitchesPrivate.h
--- a/Engine/Source/Developer/TraceServices/Private/Model/ContextSwitchesPrivate.h
+++ b/Engine/Source/Developer/TraceServices/Private/Model/ContextSwitchesPrivate.h
@@ -9,31 +9,49 @@
 namespace TraceServices
 {
 
-class FContextSwitchProvider
-	: public IContextSwitchProvider
+class FContextSwitchesProvider
+	: public IContextSwitchesProvider
 {
 public:
 	static const FName ProviderName;
 
-	explicit FContextSwitchProvider(IAnalysisSession& Session);
-	virtual ~FContextSwitchProvider();
-
-	int32 GetCoreNumber(uint32 ThreadId, double Time) const override;
-	const void EnumerateContextSwitches(uint32 ThreadId, double StartTime, double EndTime, ContextSwitchCallback Callback) const override;
-
-	void Add(uint32 ThreadId, double Start, double End, uint32 CoreNumber);
-	void AddThreadInfo(uint32 TraceThreadId, uint32 SystemThreadId);
+	explicit FContextSwitchesProvider(IAnalysisSession& Session);
+	virtual ~FContextSwitchesProvider();
 
 	virtual bool HasData() const override;
+	bool GetSystemThreadId(uint32 ThreadId, uint32& OutSystemThreadId) const override;
+	bool GetThreadId(uint32 SystemThreadId, uint32& OutThreadId) const override;
+	bool GetSystemThreadId(uint32 CoreNumber, double Time, uint32& OutSystemThreadId) const override;
+	bool GetThreadId(uint32 CoreNumber, double Time, uint32& OutThreadId) const override;
+	bool GetCoreNumber(uint32 ThreadId, double Time, uint32& OutCoreNumber) const override;
+	uint64 GetThreadsSerial() const override { return uint64(Threads.Num()); }
+	uint64 GetCpuCoresSerial() const override { return uint64(NumCpuCores); }
+	uint32 GetNumCpuCoresWithEvents() const { return NumCpuCores; }
+	uint32 GetExclusiveMaxCpuCoreNumber() const { return CpuCores.Num(); }
+	void EnumerateCpuCores(CpuCoreCallback Callback) const override;
+	void EnumerateContextSwitches(uint32 ThreadId, double StartTime, double EndTime, ContextSwitchCallback Callback) const override;
+	void EnumerateCpuCoreEvents(uint32 CoreNumber, double StartTime, double EndTime, CpuCoreEventCallback Callback) const override;
+
+	void Add(uint32 SystemThreadId, double Start, double End, uint32 CoreNumber);
+	void AddThreadInfo(uint32 ThreadId, uint32 SystemThreadId);
 
 private:
 	const TPagedArray<FContextSwitch>* GetContextSwitches(uint32 ThreadId) const;
 
 	IAnalysisSession& Session;
+
+	// (Trace) Thread Id -> System Thread Id
+	TMap<uint32, uint32> TraceToSystemThreadIdMap;
+
+	// System Thread Id -> (Trace) Thread Id
+	TMap<uint32, uint32> SystemToTraceThreadIdMap;
+
 	// System Thread Id -> PagedArray
 	TMap<uint32, TPagedArray<FContextSwitch>*> Threads;
-	// Trace Thread Id -> System Thread Id
-	TMap<uint32, uint32> ThreadIdMap;
+
+	// [Core Number] -> PagedArray; some can be nullptr
+	TArray<TPagedArray<FCpuCoreEvent>*> CpuCores;
+	uint32 NumCpuCores;
 };
 
 } // namespace TraceServices
diff --git a/Engine/Source/Developer/TraceServices/Private/Model/StackSamples.cpp b/Engine/Source/Developer/TraceServices/Private/Model/StackSamples.cpp
--- a/Engine/Source/Developer/TraceServices/Private/Model/StackSamples.cpp
+++ b/Engine/Source/Developer/TraceServices/Private/Model/StackSamples.cpp
@@ -8,15 +8,15 @@
 namespace TraceServices
 {
 
-const FName FStackSampleProvider::ProviderName = "StackSampleProvider";
+const FName FStackSamplesProvider::ProviderName = "StackSamplesProvider";
 
-FStackSampleProvider::FStackSampleProvider(IAnalysisSession& InSession)
+FStackSamplesProvider::FStackSamplesProvider(IAnalysisSession& InSession)
 	: Session(InSession)
 	, AddressValues(Session.GetLinearAllocator(), 4096)
 {
 }
 
-FStackSampleProvider::~FStackSampleProvider()
+FStackSamplesProvider::~FStackSamplesProvider()
 {
 	for (const auto& ThreadSamples : Threads)
 	{
@@ -24,7 +24,7 @@ FStackSampleProvider::~FStackSampleProvider()
 	}
 }
 
-const TPagedArray<FStackSample>* FStackSampleProvider::GetStackSamples(uint32 ThreadId) const
+const TPagedArray<FStackSample>* FStackSamplesProvider::GetStackSamples(uint32 ThreadId) const
 {
 	Session.ReadAccessCheck();
 
@@ -32,7 +32,7 @@ const TPagedArray<FStackSample>* FStackSampleProvider::GetStackSamples(uint32 Th
 	return StackSamples ? *StackSamples : nullptr;
 }
 
-void FStackSampleProvider::Add(uint32 ThreadId, double Time, uint32 Count, const uint64* Addresses)
+void FStackSamplesProvider::Add(uint32 ThreadId, double Time, uint32 Count, const uint64* Addresses)
 {
 	Session.WriteAccessCheck();
 
@@ -72,9 +72,9 @@ void FStackSampleProvider::Add(uint32 ThreadId, double Time, uint32 Count, const
 	StackSample.Addresses = &AddressValues.Last() - Count;
 }
 
-const IStackSampleProvider& ReadStackSampleProvider(const IAnalysisSession& Session)
+const IStackSamplesProvider& ReadStackSamplesProvider(const IAnalysisSession& Session)
 {
-	return *Session.ReadProvider<IStackSampleProvider>(FStackSampleProvider::ProviderName);
+	return *Session.ReadProvider<IStackSamplesProvider>(FStackSamplesProvider::ProviderName);
 }
 
 } // namespace TraceServices
diff --git a/Engine/Source/Developer/TraceServices/Private/Model/StackSamplesPrivate.h b/Engine/Source/Developer/TraceServices/Private/Model/StackSamplesPrivate.h
--- a/Engine/Source/Developer/TraceServices/Private/Model/StackSamplesPrivate.h
+++ b/Engine/Source/Developer/TraceServices/Private/Model/StackSamplesPrivate.h
@@ -9,14 +9,14 @@
 namespace TraceServices
 {
 
-class FStackSampleProvider
-	: public IStackSampleProvider
+class FStackSamplesProvider
+	: public IStackSamplesProvider
 {
 public:
 	static const FName ProviderName;
 
-	explicit FStackSampleProvider(IAnalysisSession& Session);
-	virtual ~FStackSampleProvider();
+	explicit FStackSamplesProvider(IAnalysisSession& Session);
+	virtual ~FStackSamplesProvider();
 
 	const TPagedArray<FStackSample>* GetStackSamples(uint32 ThreadId) const override;
 
diff --git a/Engine/Source/Developer/TraceServices/Private/Modules/PlatformEventsModule.cpp b/Engine/Source/Developer/TraceServices/Private/Modules/PlatformEventsModule.cpp
--- a/Engine/Source/Developer/TraceServices/Private/Modules/PlatformEventsModule.cpp
+++ b/Engine/Source/Developer/TraceServices/Private/Modules/PlatformEventsModule.cpp
@@ -21,12 +21,12 @@ void FPlatformEventsModule::OnAnalysisBegin(IAnalysisSession& InSession)
 {
 	FAnalysisSession& Session = static_cast<FAnalysisSession&>(InSession);
 
-	FContextSwitchProvider* ContextSwitchProvider = new FContextSwitchProvider(Session);
-	FStackSampleProvider* StackSampleProvider = new FStackSampleProvider(Session);
+	FContextSwitchesProvider* ContextSwitchesProvider = new FContextSwitchesProvider(Session);
+	FStackSamplesProvider* StackSamplesProvider = new FStackSamplesProvider(Session);
 	
-	Session.AddProvider(FContextSwitchProvider::ProviderName, ContextSwitchProvider);
-	Session.AddProvider(FStackSampleProvider::ProviderName, StackSampleProvider);
-	Session.AddAnalyzer(new FPlatformEventTraceAnalyzer(Session, *ContextSwitchProvider, *StackSampleProvider));
+	Session.AddProvider(FContextSwitchesProvider::ProviderName, ContextSwitchesProvider);
+	Session.AddProvider(FStackSamplesProvider::ProviderName, StackSamplesProvider);
+	Session.AddAnalyzer(new FPlatformEventTraceAnalyzer(Session, *ContextSwitchesProvider, *StackSamplesProvider));
 }
 
 } // namespace TraceServices
diff --git a/Engine/Source/Developer/TraceServices/Public/TraceServices/Model/ContextSwitches.h b/Engine/Source/Developer/TraceServices/Public/TraceServices/Model/ContextSwitches.h
--- a/Engine/Source/Developer/TraceServices/Public/TraceServices/Model/ContextSwitches.h
+++ b/Engine/Source/Developer/TraceServices/Public/TraceServices/Model/ContextSwitches.h
@@ -15,24 +15,46 @@ struct FContextSwitch
 	uint32 CoreNumber;
 };
 
+struct FCpuCoreInfo
+{
+	uint32 CoreNumber;
+};
+
+struct FCpuCoreEvent
+{
+	double Start;
+	double End;
+	uint32 SystemThreadId;
+};
+
 enum class EContextSwitchEnumerationResult
 {
 	Continue,
 	Stop,
 };
 
-class IContextSwitchProvider
+class IContextSwitchesProvider
 	: public IProvider
 {
 public:
+	typedef TFunctionRef<void(const FCpuCoreInfo&)> CpuCoreCallback;
+	typedef TFunctionRef<EContextSwitchEnumerationResult(const FCpuCoreEvent&)> CpuCoreEventCallback;
 	typedef TFunctionRef<EContextSwitchEnumerationResult(const FContextSwitch&)> ContextSwitchCallback;
 
-	virtual ~IContextSwitchProvider() = default;
-	virtual int32 GetCoreNumber(uint32 ThreadId, double Time) const = 0;
-	virtual const void EnumerateContextSwitches(uint32 ThreadId, double StartTime, double EndTime, ContextSwitchCallback Callback) const = 0;
+	virtual ~IContextSwitchesProvider() = default;
 	virtual bool HasData() const = 0;
+	virtual bool GetSystemThreadId(uint32 ThreadId, uint32& OutSystemThreadId) const = 0;
+	virtual bool GetThreadId(uint32 SystemThreadId, uint32& OutThreadId) const = 0;
+	virtual bool GetSystemThreadId(uint32 CoreNumber, double Time, uint32& OutSystemThreadId) const = 0;
+	virtual bool GetThreadId(uint32 CoreNumber, double Time, uint32& OutThreadId) const = 0;
+	virtual bool GetCoreNumber(uint32 ThreadId, double Time, uint32& OutCoreNumber) const = 0;
+	virtual uint64 GetThreadsSerial() const = 0;
+	virtual uint64 GetCpuCoresSerial() const = 0;
+	virtual void EnumerateCpuCores(CpuCoreCallback Callback) const = 0;
+	virtual void EnumerateCpuCoreEvents(uint32 CoreNumber, double StartTime, double EndTime, CpuCoreEventCallback Callback) const = 0;
+	virtual void EnumerateContextSwitches(uint32 ThreadId, double StartTime, double EndTime, ContextSwitchCallback Callback) const = 0;
 };
 
-TRACESERVICES_API const IContextSwitchProvider* ReadContextSwitchProvider(const IAnalysisSession& Session);
+TRACESERVICES_API const IContextSwitchesProvider* ReadContextSwitchesProvider(const IAnalysisSession& Session);
 
 } // namespace TraceServices
diff --git a/Engine/Source/Developer/TraceServices/Public/TraceServices/Model/StackSamples.h b/Engine/Source/Developer/TraceServices/Public/TraceServices/Model/StackSamples.h
--- a/Engine/Source/Developer/TraceServices/Public/TraceServices/Model/StackSamples.h
+++ b/Engine/Source/Developer/TraceServices/Public/TraceServices/Model/StackSamples.h
@@ -15,14 +15,14 @@ struct FStackSample
 	uint64* Addresses;
 };
 
-class IStackSampleProvider
+class IStackSamplesProvider
 	: public IProvider
 {
 public:
-	virtual ~IStackSampleProvider() = default;
+	virtual ~IStackSamplesProvider() = default;
 	virtual const TPagedArray<FStackSample>* GetStackSamples(uint32 ThreadId) const = 0;
 };
 
-TRACESERVICES_API const IStackSampleProvider& ReadStackSampleProvider(const IAnalysisSession& Session);
+TRACESERVICES_API const IStackSamplesProvider& ReadStackSamplesProvider(const IAnalysisSession& Session);
 
 } // namespace TraceServices
