[Insights] Improved struct size for FTableTreeNode (allocates space for aggregated values only when needed; i.e. for group nodes).

diff --git a/Engine/Source/Developer/TraceInsights/Private/Insights/Table/ViewModels/TableTreeNode.h b/Engine/Source/Developer/TraceInsights/Private/Insights/Table/ViewModels/TableTreeNode.h
--- a/Engine/Source/Developer/TraceInsights/Private/Insights/Table/ViewModels/TableTreeNode.h
+++ b/Engine/Source/Developer/TraceInsights/Private/Insights/Table/ViewModels/TableTreeNode.h
@@ -57,6 +57,7 @@ public:
 		: FBaseTreeNode(InName, false)
 		, ParentTable(InParentTable)
 		, RowId(InRowIndex)
+		, AggregatedValues(nullptr)
 	{
 	}
 
@@ -65,22 +66,46 @@ public:
 		: FBaseTreeNode(InGroupName, true)
 		, ParentTable(InParentTable)
 		, RowId(FTableRowId::InvalidRowIndex)
+		, AggregatedValues(nullptr)
 	{
 	}
 
+	virtual ~FTableTreeNode()
+	{
+		CleanupAggregatedValues();
+	}
+
 	virtual const FName& GetTypeName() const override { return TypeName; }
 
 	const TWeakPtr<FTable>& GetParentTable() const { return ParentTable; }
 	FTableRowId GetRowId() const { return RowId; }
 	int32 GetRowIndex() const { return RowId.RowIndex; }
 
-	void ResetAggregatedValues() { AggregatedValues.Reset(); }
-	void ResetAggregatedValues(const FName& ColumnId) { AggregatedValues.Remove(ColumnId); }
-	bool HasAggregatedValue(const FName& ColumnId) const { return AggregatedValues.Contains(ColumnId); }
-	const FTableCellValue* FindAggregatedValue(const FName& ColumnId) const { return AggregatedValues.Find(ColumnId); }
-	const FTableCellValue& GetAggregatedValue(const FName& ColumnId) const { return AggregatedValues.FindChecked(ColumnId); }
-	void AddAggregatedValue(const FName& ColumnId, const FTableCellValue& Value) { AggregatedValues.Add(ColumnId, Value); }
-	void SetAggregatedValue(const FName& ColumnId, const FTableCellValue& Value) { AggregatedValues[ColumnId] = Value; }
+	void InitAggregatedValues()
+	{
+		if (!AggregatedValues)
+		{
+			AggregatedValues = new TMap<FName, FTableCellValue>();
+		}
+	}
+
+	void CleanupAggregatedValues()
+	{
+		if (AggregatedValues)
+		{
+			delete AggregatedValues;
+			AggregatedValues = nullptr;
+		}
+	}
+
+	void ResetAggregatedValues() { CleanupAggregatedValues(); }
+	void ResetAggregatedValues(const FName& ColumnId) { if (AggregatedValues) { AggregatedValues->Remove(ColumnId); } }
+	bool HasAggregatedValue(const FName& ColumnId) const { return AggregatedValues && AggregatedValues->Contains(ColumnId); }
+	const FTableCellValue* FindAggregatedValue(const FName& ColumnId) const { return AggregatedValues ? AggregatedValues->Find(ColumnId) : nullptr; }
+	const FTableCellValue& GetAggregatedValue(const FName& ColumnId) const { return AggregatedValues->FindChecked(ColumnId); }
+	void AddAggregatedValue(const FName& ColumnId, const FTableCellValue& Value) { InitAggregatedValues(); AggregatedValues->Add(ColumnId, Value); }
+	void SetAggregatedValue(const FName& ColumnId, const FTableCellValue& Value) { InitAggregatedValues(); (*AggregatedValues)[ColumnId] = Value; }
+
 	virtual bool IsFiltered() const override { return bIsFiltered; }
 	virtual void SetIsFiltered(bool InValue) { bIsFiltered = InValue; }
 
@@ -88,7 +113,7 @@ protected:
 	TWeakPtr<FTable> ParentTable;
 	FTableRowId RowId;
 
-	TMap<FName, FTableCellValue> AggregatedValues;
+	TMap<FName, FTableCellValue>* AggregatedValues;
 
 private:
 	bool bIsFiltered = false;
