[Insights] Minor fix re precision of "duration" time values in tooltips of Tasks.

diff --git a/Engine/Source/Developer/TraceInsights/Private/Insights/TaskGraphProfiler/ViewModels/TaskTimingTrack.cpp b/Engine/Source/Developer/TraceInsights/Private/Insights/TaskGraphProfiler/ViewModels/TaskTimingTrack.cpp
--- a/Engine/Source/Developer/TraceInsights/Private/Insights/TaskGraphProfiler/ViewModels/TaskTimingTrack.cpp
+++ b/Engine/Source/Developer/TraceInsights/Private/Insights/TaskGraphProfiler/ViewModels/TaskTimingTrack.cpp
@@ -590,7 +590,7 @@ void FTaskTimingTrack::InitTooltip(FTooltipDrawState& InOutTooltip, const ITimin
 
 	InOutTooltip.AddNameValueTextLine(TaskTrackEvent.GetStartLabel(), TimeUtils::FormatTimeAuto(TaskTrackEvent.GetStartTime(), 6));
 	InOutTooltip.AddNameValueTextLine(TaskTrackEvent.GetEndLabel(), TimeUtils::FormatTimeAuto(TaskTrackEvent.GetEndTime(), 6));
-	InOutTooltip.AddNameValueTextLine(TEXT("Duration:"), TimeUtils::FormatTimeAuto(TaskTrackEvent.GetEndTime() - TaskTrackEvent.GetStartTime(), 6));
+	InOutTooltip.AddNameValueTextLine(TEXT("Duration:"), TimeUtils::FormatTimeAuto(TaskTrackEvent.GetEndTime() - TaskTrackEvent.GetStartTime()));
 
 	TSharedPtr<const TraceServices::IAnalysisSession> Session = FInsightsManager::Get()->GetSession();
 	if (!Session.IsValid())
