[Insights] Memory: Added "FMallocPoisonProxy::" (in addition to FMemory:: and FMallocWrapper::) as a filter prefix to ignore functions on top of callstacks.

diff --git a/Engine/Source/Developer/TraceInsights/Private/Insights/MemoryProfiler/ViewModels/MemAllocTable.cpp b/Engine/Source/Developer/TraceInsights/Private/Insights/MemoryProfiler/ViewModels/MemAllocTable.cpp
--- a/Engine/Source/Developer/TraceInsights/Private/Insights/MemoryProfiler/ViewModels/MemAllocTable.cpp
+++ b/Engine/Source/Developer/TraceInsights/Private/Insights/MemoryProfiler/ViewModels/MemAllocTable.cpp
@@ -597,7 +597,12 @@ void FMemAllocTable::AddDefaultColumns()
 									break;
 								}
 
-								FStringView IgnoreSymbols[] = { TEXT("FMemory::"_SV), TEXT("FMallocWrapper::"_SV) };
+								FStringView IgnoreSymbols[] =
+								{
+									TEXT("FMemory::"_SV),
+									TEXT("FMallocWrapper::"_SV),
+									TEXT("FMallocPoisonProxy::"_SV),
+								};
 
 								bool bMatchIgnoreSymbol = false;
 								for (uint32 IgnoreSymbolsIndex = 0; IgnoreSymbolsIndex < UE_ARRAY_COUNT(IgnoreSymbols); ++IgnoreSymbolsIndex)
