[Insights]
  - TraceServices: Changed ReadNetProfilerProvider, ReadMemoryProvider, ReadDiagnosticsProvider to return a pointer to the respective provider (instead of a reference). Can return nullptr. The UI should not assume the provider exists.
  - TraceServices: Added GetNetProfilerProviderName, GetMemoryProviderName, GetDiagnosticsProviderName, GetAllocationsProviderName to return name of respective provider.
  - TraceServices: Added ReadModuleProvider + GetModuleProviderName, ReadCallstacksProvider + GetCallstacksProviderName to access respective providers.
  - TraceServices: Added a default empty implementation to IModule for GenerateReports, GetLoggers and GetCommandLineArgument.
  - TraceServices: Changed FAnalysisService::StartAnalysis to not create default providers for Memory (LLM Stats), NetProfiler and Diagnostics, but to allow the respective modules (FMemoryModule, FNetProfilerModule, FDiagnosticsModule) to create the providers.
  - TraceInsights: Fixed UI code to check if the Memory, NetProfiler and Diagnostics provider are available.

diff --git a/Engine/Source/Developer/TraceInsights/Private/Insights/MemoryProfiler/MemoryProfilerManager.cpp b/Engine/Source/Developer/TraceInsights/Private/Insights/MemoryProfiler/MemoryProfilerManager.cpp
--- a/Engine/Source/Developer/TraceInsights/Private/Insights/MemoryProfiler/MemoryProfilerManager.cpp
+++ b/Engine/Source/Developer/TraceInsights/Private/Insights/MemoryProfiler/MemoryProfilerManager.cpp
@@ -258,11 +258,14 @@ bool FMemoryProfilerManager::Tick(float DeltaTime)
 				AvailabilityCheck.Disable();
 			}
 
-			const TraceServices::IMemoryProvider& MemoryProvider = TraceServices::ReadMemoryProvider(*Session.Get());
-			uint32 TagCount = MemoryProvider.GetTagCount();
-			if (TagCount > 0)
+			const TraceServices::IMemoryProvider* MemoryProvider = TraceServices::ReadMemoryProvider(*Session.Get());
+			if (MemoryProvider)
 			{
-				bShouldBeAvailable = true;
+				uint32 TagCount = MemoryProvider->GetTagCount();
+				if (TagCount > 0)
+				{
+					bShouldBeAvailable = true;
+				}
 			}
 
 			const TraceServices::IAllocationsProvider* AllocationsProvider = TraceServices::ReadAllocationsProvider(*Session.Get());
diff --git a/Engine/Source/Developer/TraceInsights/Private/Insights/MemoryProfiler/ViewModels/MemoryGraphTrack.cpp b/Engine/Source/Developer/TraceInsights/Private/Insights/MemoryProfiler/ViewModels/MemoryGraphTrack.cpp
--- a/Engine/Source/Developer/TraceInsights/Private/Insights/MemoryProfiler/ViewModels/MemoryGraphTrack.cpp
+++ b/Engine/Source/Developer/TraceInsights/Private/Insights/MemoryProfiler/ViewModels/MemoryGraphTrack.cpp
@@ -308,30 +308,32 @@ void FMemoryGraphTrack::PreUpdateMemTagSeries(FMemoryGraphSeries& Series, const
 	if (Session.IsValid())
 	{
 		TraceServices::FAnalysisSessionReadScope SessionReadScope(*Session.Get());
-
-		// Compute Min/Max values.
-		const TraceServices::IMemoryProvider& MemoryProvider = TraceServices::ReadMemoryProvider(*Session.Get());
-		const uint64 TotalSampleCount = MemoryProvider.GetTagSampleCount(Series.GetTrackerId(), Series.GetTagId());
-		if (TotalSampleCount > 0)
+		const TraceServices::IMemoryProvider* MemoryProvider = TraceServices::ReadMemoryProvider(*Session.Get());
+		if (MemoryProvider)
 		{
-			double MinValue = +std::numeric_limits<double>::infinity();
-			double MaxValue = -std::numeric_limits<double>::infinity();
-
-			MemoryProvider.EnumerateTagSamples(Series.GetTrackerId(), Series.GetTagId(), Viewport.GetStartTime(), Viewport.GetEndTime(), true,
-				[&MinValue, &MaxValue](double Time, double Duration, const TraceServices::FMemoryTagSample& Sample)
+			const uint64 TotalSampleCount = MemoryProvider->GetTagSampleCount(Series.GetTrackerId(), Series.GetTagId());
+			if (TotalSampleCount > 0)
 			{
-				const double Value = static_cast<double>(Sample.Value);
-				if (Value < MinValue)
-				{
-					MinValue = Value;
-				}
-				if (Value > MaxValue)
+				double MinValue = +std::numeric_limits<double>::infinity();
+				double MaxValue = -std::numeric_limits<double>::infinity();
+
+				// Compute Min/Max values.
+				MemoryProvider->EnumerateTagSamples(Series.GetTrackerId(), Series.GetTagId(), Viewport.GetStartTime(), Viewport.GetEndTime(), true,
+					[&MinValue, &MaxValue](double Time, double Duration, const TraceServices::FMemoryTagSample& Sample)
 				{
-					MaxValue = Value;
-				}
-			});
+					const double Value = static_cast<double>(Sample.Value);
+					if (Value < MinValue)
+					{
+						MinValue = Value;
+					}
+					if (Value > MaxValue)
+					{
+						MaxValue = Value;
+					}
+				});
 
-			Series.SetValueRange(MinValue, MaxValue);
+				Series.SetValueRange(MinValue, MaxValue);
+			}
 		}
 	}
 }
@@ -347,25 +349,27 @@ void FMemoryGraphTrack::UpdateMemTagSeries(FMemoryGraphSeries& Series, const FTi
 	{
 		TraceServices::FAnalysisSessionReadScope SessionReadScope(*Session.Get());
 
-		const TraceServices::IMemoryProvider& MemoryProvider = TraceServices::ReadMemoryProvider(*Session.Get());
-
-		const uint64 TotalSampleCount = MemoryProvider.GetTagSampleCount(Series.GetTrackerId(), Series.GetTagId());
-
-		if (TotalSampleCount > 0)
+		const TraceServices::IMemoryProvider* MemoryProvider = TraceServices::ReadMemoryProvider(*Session.Get());
+		if (MemoryProvider)
 		{
-			const float TopY = 4.0f;
-			const float BottomY = GetHeight() - 4.0f;
+			const uint64 TotalSampleCount = MemoryProvider->GetTagSampleCount(Series.GetTrackerId(), Series.GetTagId());
 
-			if (Series.IsAutoZoomEnabled() && TopY < BottomY)
+			if (TotalSampleCount > 0)
 			{
-				Series.UpdateAutoZoom(TopY, BottomY, Series.GetMinValue(), Series.GetMaxValue());
-			}
+				const float TopY = 4.0f;
+				const float BottomY = GetHeight() - 4.0f;
 
-			MemoryProvider.EnumerateTagSamples(Series.GetTrackerId(), Series.GetTagId(), Viewport.GetStartTime(), Viewport.GetEndTime(), true,
-				[this, &Builder](double Time, double Duration, const TraceServices::FMemoryTagSample& Sample)
-			{
-				Builder.AddEvent(Time, Duration, static_cast<double>(Sample.Value));
-			});
+				if (Series.IsAutoZoomEnabled() && TopY < BottomY)
+				{
+					Series.UpdateAutoZoom(TopY, BottomY, Series.GetMinValue(), Series.GetMaxValue());
+				}
+
+				MemoryProvider->EnumerateTagSamples(Series.GetTrackerId(), Series.GetTagId(), Viewport.GetStartTime(), Viewport.GetEndTime(), true,
+					[this, &Builder](double Time, double Duration, const TraceServices::FMemoryTagSample& Sample)
+				{
+					Builder.AddEvent(Time, Duration, static_cast<double>(Sample.Value));
+				});
+			}
 		}
 	}
 }
diff --git a/Engine/Source/Developer/TraceInsights/Private/Insights/MemoryProfiler/ViewModels/MemorySharedState.cpp b/Engine/Source/Developer/TraceInsights/Private/Insights/MemoryProfiler/ViewModels/MemorySharedState.cpp
--- a/Engine/Source/Developer/TraceInsights/Private/Insights/MemoryProfiler/ViewModels/MemorySharedState.cpp
+++ b/Engine/Source/Developer/TraceInsights/Private/Insights/MemoryProfiler/ViewModels/MemorySharedState.cpp
@@ -301,14 +301,16 @@ void FMemorySharedState::SyncTrackers()
 	if (Session.IsValid())
 	{
 		TraceServices::FAnalysisSessionReadScope SessionReadScope(*Session.Get());
-		const TraceServices::IMemoryProvider& MemoryProvider = TraceServices::ReadMemoryProvider(*Session.Get());
-
-		MemoryProvider.EnumerateTrackers([this](const TraceServices::FMemoryTrackerInfo& Tracker)
+		const TraceServices::IMemoryProvider* MemoryProvider = TraceServices::ReadMemoryProvider(*Session.Get());
+		if (MemoryProvider)
 		{
-			Trackers.Add(MakeShared<Insights::FMemoryTracker>(Tracker.Id, Tracker.Name));
-		});
+			MemoryProvider->EnumerateTrackers([this](const TraceServices::FMemoryTrackerInfo& Tracker)
+			{
+				Trackers.Add(MakeShared<Insights::FMemoryTracker>(Tracker.Id, Tracker.Name));
+			});
 
-		Trackers.Sort([](const TSharedPtr<Insights::FMemoryTracker>& A, const TSharedPtr<Insights::FMemoryTracker>& B) { return A->GetId() < B->GetId(); });
+			Trackers.Sort([](const TSharedPtr<Insights::FMemoryTracker>& A, const TSharedPtr<Insights::FMemoryTracker>& B) { return A->GetId() < B->GetId(); });
+		}
 	}
 
 	if (Trackers.Num() > 0)
diff --git a/Engine/Source/Developer/TraceInsights/Private/Insights/MemoryProfiler/ViewModels/MemoryTag.cpp b/Engine/Source/Developer/TraceInsights/Private/Insights/MemoryProfiler/ViewModels/MemoryTag.cpp
--- a/Engine/Source/Developer/TraceInsights/Private/Insights/MemoryProfiler/ViewModels/MemoryTag.cpp
+++ b/Engine/Source/Developer/TraceInsights/Private/Insights/MemoryProfiler/ViewModels/MemoryTag.cpp
@@ -129,44 +129,46 @@ void FMemoryTagList::UpdateInternal()
 	if (Session.IsValid())
 	{
 		TraceServices::FAnalysisSessionReadScope SessionReadScope(*Session.Get());
-		const TraceServices::IMemoryProvider& MemoryProvider = TraceServices::ReadMemoryProvider(*Session.Get());
-
-		const int32 TraceSerialNumber = MemoryProvider.GetTagSerial();
-		if (LastTraceSerialNumber != TraceSerialNumber)
+		const TraceServices::IMemoryProvider* MemoryProvider = TraceServices::ReadMemoryProvider(*Session.Get());
+		if (MemoryProvider)
 		{
-			LastTraceSerialNumber = TraceSerialNumber;
-			++SerialNumber;
-			bUpdated = true;
-
-			MemoryProvider.EnumerateTags([this](const TraceServices::FMemoryTagInfo& TraceTag)
+			const int32 TraceSerialNumber = MemoryProvider->GetTagSerial();
+			if (LastTraceSerialNumber != TraceSerialNumber)
 			{
-				static_assert(FMemoryTag::InvalidTagId == static_cast<FMemoryTagId>(TraceServices::FMemoryTagInfo::InvalidTagId), "Memory TagId type mismatch!");
-				FMemoryTagId TagId = static_cast<FMemoryTagId>(TraceTag.Id);
+				LastTraceSerialNumber = TraceSerialNumber;
+				++SerialNumber;
+				bUpdated = true;
 
-				FMemoryTag* TagPtr = TagIdMap.FindRef(TagId);
-				if (!TagPtr)
-				{
-					TagPtr = new FMemoryTag();
-
-					FMemoryTag& Tag = *TagPtr;
-					Tag.Index = Tags.Num();
-					Tag.Id = TagId;
-					Tag.ParentId = static_cast<FMemoryTagId>(TraceTag.ParentId);
-					Tag.StatName = TraceTag.Name;
-					Tag.StatFullName = Tag.StatName;
-					Tag.Trackers = TraceTag.Trackers;
-					Tag.SetColorAuto();
-
-					Tags.Add(TagPtr);
-					TagIdMap.Add(TagId, TagPtr);
-				}
-				else
+				MemoryProvider->EnumerateTags([this](const TraceServices::FMemoryTagInfo& TraceTag)
 				{
-					TagPtr->Trackers = TraceTag.Trackers;
-				}
-			});
+					static_assert(FMemoryTag::InvalidTagId == static_cast<FMemoryTagId>(TraceServices::FMemoryTagInfo::InvalidTagId), "Memory TagId type mismatch!");
+					FMemoryTagId TagId = static_cast<FMemoryTagId>(TraceTag.Id);
 
-			ensure(Tags.Num() == MemoryProvider.GetTagCount());
+					FMemoryTag* TagPtr = TagIdMap.FindRef(TagId);
+					if (!TagPtr)
+					{
+						TagPtr = new FMemoryTag();
+
+						FMemoryTag& Tag = *TagPtr;
+						Tag.Index = Tags.Num();
+						Tag.Id = TagId;
+						Tag.ParentId = static_cast<FMemoryTagId>(TraceTag.ParentId);
+						Tag.StatName = TraceTag.Name;
+						Tag.StatFullName = Tag.StatName;
+						Tag.Trackers = TraceTag.Trackers;
+						Tag.SetColorAuto();
+
+						Tags.Add(TagPtr);
+						TagIdMap.Add(TagId, TagPtr);
+					}
+					else
+					{
+						TagPtr->Trackers = TraceTag.Trackers;
+					}
+				});
+
+				ensure(Tags.Num() == MemoryProvider->GetTagCount());
+			}
 		}
 	}
 
diff --git a/Engine/Source/Developer/TraceInsights/Private/Insights/NetworkingProfiler/NetworkingProfilerManager.cpp b/Engine/Source/Developer/TraceInsights/Private/Insights/NetworkingProfiler/NetworkingProfilerManager.cpp
--- a/Engine/Source/Developer/TraceInsights/Private/Insights/NetworkingProfiler/NetworkingProfilerManager.cpp
+++ b/Engine/Source/Developer/TraceInsights/Private/Insights/NetworkingProfiler/NetworkingProfilerManager.cpp
@@ -235,8 +235,11 @@ bool FNetworkingProfilerManager::Tick(float DeltaTime)
 				AvailabilityCheck.Disable();
 			}
 
-			const TraceServices::INetProfilerProvider& NetProfilerProvider = TraceServices::ReadNetProfilerProvider(*Session.Get());
-			NetTraceVersion = NetProfilerProvider.GetNetTraceVersion();
+			const TraceServices::INetProfilerProvider* NetProfilerProvider = TraceServices::ReadNetProfilerProvider(*Session.Get());
+			if (NetProfilerProvider)
+			{
+				NetTraceVersion = NetProfilerProvider->GetNetTraceVersion();
+			}
 		}
 		else
 		{
diff --git a/Engine/Source/Developer/TraceInsights/Private/Insights/NetworkingProfiler/Widgets/SNetStatsView.cpp b/Engine/Source/Developer/TraceInsights/Private/Insights/NetworkingProfiler/Widgets/SNetStatsView.cpp
--- a/Engine/Source/Developer/TraceInsights/Private/Insights/NetworkingProfiler/Widgets/SNetStatsView.cpp
+++ b/Engine/Source/Developer/TraceInsights/Private/Insights/NetworkingProfiler/Widgets/SNetStatsView.cpp
@@ -1574,32 +1574,33 @@ void SNetStatsView::RebuildTree(bool bResync)
 	if (Session.IsValid())
 	{
 		TraceServices::FAnalysisSessionReadScope SessionReadScope(*Session.Get());
-
-		const TraceServices::INetProfilerProvider& NetProfilerProvider = TraceServices::ReadNetProfilerProvider(*Session.Get());
-
-		NetProfilerProvider.ReadEventTypes([this, &bResync, &NetProfilerProvider](const TraceServices::FNetProfilerEventType* NetEvents, uint64 InNetEventCount)
+		const TraceServices::INetProfilerProvider* NetProfilerProvider = TraceServices::ReadNetProfilerProvider(*Session.Get());
+		if (NetProfilerProvider)
 		{
-			const uint32 NetEventCount = static_cast<uint32>(InNetEventCount);
-			if (NetEventCount != NetEventNodes.Num())
+			NetProfilerProvider->ReadEventTypes([this, &bResync, NetProfilerProvider](const TraceServices::FNetProfilerEventType* NetEvents, uint64 InNetEventCount)
 			{
-				NetEventNodes.Empty(NetEventCount);
-				for (uint32 Index = 0; Index < NetEventCount; ++Index)
+				const uint32 NetEventCount = static_cast<uint32>(InNetEventCount);
+				if (NetEventCount != NetEventNodes.Num())
 				{
-					const TraceServices::FNetProfilerEventType& NetEvent = NetEvents[Index];
-					ensure(NetEvent.EventTypeIndex == Index);
-					const TCHAR* NamePtr = nullptr;
-					NetProfilerProvider.ReadName(NetEvent.NameIndex, [&NamePtr](const TraceServices::FNetProfilerName& InName)
+					NetEventNodes.Empty(NetEventCount);
+					for (uint32 Index = 0; Index < NetEventCount; ++Index)
 					{
-						NamePtr = InName.Name;
-					});
-					const FName Name(NamePtr);
-					const ENetEventNodeType Type = ENetEventNodeType::NetEvent;
-					FNetEventNodePtr NetEventNodePtr = MakeShared<FNetEventNode>(NetEvent.EventTypeIndex, Name, Type, NetEvent.Level);
-					NetEventNodes.Add(NetEventNodePtr);
+						const TraceServices::FNetProfilerEventType& NetEvent = NetEvents[Index];
+						ensure(NetEvent.EventTypeIndex == Index);
+						const TCHAR* NamePtr = nullptr;
+						NetProfilerProvider->ReadName(NetEvent.NameIndex, [&NamePtr](const TraceServices::FNetProfilerName& InName)
+						{
+							NamePtr = InName.Name;
+						});
+						const FName Name(NamePtr);
+						const ENetEventNodeType Type = ENetEventNodeType::NetEvent;
+						FNetEventNodePtr NetEventNodePtr = MakeShared<FNetEventNode>(NetEvent.EventTypeIndex, Name, Type, NetEvent.Level);
+						NetEventNodes.Add(NetEventNodePtr);
+					}
+					ensure(NetEventNodes.Num() == NetEventCount);
 				}
-				ensure(NetEventNodes.Num() == NetEventCount);
-			}
-		});
+			});
+		}
 	}
 	SyncStopwatch.Stop();
 
@@ -1700,9 +1701,12 @@ void SNetStatsView::UpdateStatsInternal()
 		AggregationStopwatch.Start();
 		{
 			TraceServices::FAnalysisSessionReadScope SessionReadScope(*Session.Get());
-			const TraceServices::INetProfilerProvider& NetProfilerProvider = TraceServices::ReadNetProfilerProvider(*Session.Get());
-			// CreateAggregation requires [PacketStartIndex, PacketEndIndex] as inclusive interval and [StartPos, EndPos) as exclusive interval.
-			AggregationResultTable.Reset(NetProfilerProvider.CreateAggregation(ConnectionIndex, ConnectionMode, StatsPacketStartIndex, StatsPacketEndIndex - 1, StatsStartPosition, StatsEndPosition));
+			const TraceServices::INetProfilerProvider* NetProfilerProvider = TraceServices::ReadNetProfilerProvider(*Session.Get());
+			if (NetProfilerProvider)
+			{
+				// CreateAggregation requires [PacketStartIndex, PacketEndIndex] as inclusive interval and [StartPos, EndPos) as exclusive interval.
+				AggregationResultTable.Reset(NetProfilerProvider->CreateAggregation(ConnectionIndex, ConnectionMode, StatsPacketStartIndex, StatsPacketEndIndex - 1, StatsStartPosition, StatsEndPosition));
+			}
 		}
 		AggregationStopwatch.Stop();
 
diff --git a/Engine/Source/Developer/TraceInsights/Private/Insights/NetworkingProfiler/Widgets/SNetworkingProfilerWindow.cpp b/Engine/Source/Developer/TraceInsights/Private/Insights/NetworkingProfiler/Widgets/SNetworkingProfilerWindow.cpp
--- a/Engine/Source/Developer/TraceInsights/Private/Insights/NetworkingProfiler/Widgets/SNetworkingProfilerWindow.cpp
+++ b/Engine/Source/Developer/TraceInsights/Private/Insights/NetworkingProfiler/Widgets/SNetworkingProfilerWindow.cpp
@@ -494,8 +494,11 @@ void SNetworkingProfilerWindow::Tick(const FGeometry& AllottedGeometry, const do
 	if (Session.IsValid())
 	{
 		TraceServices::FAnalysisSessionReadScope SessionReadScope(*Session.Get());
-		const TraceServices::INetProfilerProvider& NetProfilerProvider = TraceServices::ReadNetProfilerProvider(*Session.Get());
-		GameInstanceCount = NetProfilerProvider.GetGameInstanceCount();
+		const TraceServices::INetProfilerProvider* NetProfilerProvider = TraceServices::ReadNetProfilerProvider(*Session.Get());
+		if (NetProfilerProvider)
+		{
+			GameInstanceCount = NetProfilerProvider->GetGameInstanceCount();
+		}
 	}
 
 	if (GameInstanceCount != AvailableGameInstances.Num())
@@ -507,8 +510,11 @@ void SNetworkingProfilerWindow::Tick(const FGeometry& AllottedGeometry, const do
 	if (Session.IsValid() && SelectedGameInstance.IsValid())
 	{
 		TraceServices::FAnalysisSessionReadScope SessionReadScope(*Session.Get());
-		const TraceServices::INetProfilerProvider& NetProfilerProvider = TraceServices::ReadNetProfilerProvider(*Session.Get());
-		ConnectionCount = NetProfilerProvider.GetConnectionCount(SelectedGameInstance->GetIndex());
+		const TraceServices::INetProfilerProvider* NetProfilerProvider = TraceServices::ReadNetProfilerProvider(*Session.Get());
+		if (NetProfilerProvider)
+		{
+			ConnectionCount = NetProfilerProvider->GetConnectionCount(SelectedGameInstance->GetIndex());
+		}
 	}
 
 	if (ConnectionCount != AvailableConnections.Num())
@@ -527,11 +533,14 @@ void SNetworkingProfilerWindow::UpdateAvailableGameInstances()
 	if (Session.IsValid())
 	{
 		TraceServices::FAnalysisSessionReadScope SessionReadScope(*Session.Get());
-		const TraceServices::INetProfilerProvider& NetProfilerProvider = TraceServices::ReadNetProfilerProvider(*Session.Get());
-		NetProfilerProvider.ReadGameInstances([this](const TraceServices::FNetProfilerGameInstance& GameInstance)
+		const TraceServices::INetProfilerProvider* NetProfilerProvider = TraceServices::ReadNetProfilerProvider(*Session.Get());
+		if (NetProfilerProvider)
 		{
-			AvailableGameInstances.Add(MakeShared<FGameInstanceItem>(GameInstance));
-		});
+			NetProfilerProvider->ReadGameInstances([this](const TraceServices::FNetProfilerGameInstance& GameInstance)
+			{
+				AvailableGameInstances.Add(MakeShared<FGameInstanceItem>(GameInstance));
+			});
+		}
 	}
 
 	if (GameInstanceComboBox.IsValid())
@@ -552,11 +561,14 @@ void SNetworkingProfilerWindow::UpdateAvailableConnections()
 	if (Session.IsValid() && SelectedGameInstance.IsValid())
 	{
 		TraceServices::FAnalysisSessionReadScope SessionReadScope(*Session.Get());
-		const TraceServices::INetProfilerProvider& NetProfilerProvider = TraceServices::ReadNetProfilerProvider(*Session.Get());
-		NetProfilerProvider.ReadConnections(SelectedGameInstance->GetIndex(), [this](const TraceServices::FNetProfilerConnection& Connection)
+		const TraceServices::INetProfilerProvider* NetProfilerProvider = TraceServices::ReadNetProfilerProvider(*Session.Get());
+		if (NetProfilerProvider)
 		{
-			AvailableConnections.Add(MakeShared<FConnectionItem>(Connection));
-		});
+			NetProfilerProvider->ReadConnections(SelectedGameInstance->GetIndex(), [this](const TraceServices::FNetProfilerConnection& Connection)
+			{
+				AvailableConnections.Add(MakeShared<FConnectionItem>(Connection));
+			});
+		}
 	}
 
 	if (ConnectionComboBox.IsValid())
diff --git a/Engine/Source/Developer/TraceInsights/Private/Insights/NetworkingProfiler/Widgets/SPacketContentView.cpp b/Engine/Source/Developer/TraceInsights/Private/Insights/NetworkingProfiler/Widgets/SPacketContentView.cpp
--- a/Engine/Source/Developer/TraceInsights/Private/Insights/NetworkingProfiler/Widgets/SPacketContentView.cpp
+++ b/Engine/Source/Developer/TraceInsights/Private/Insights/NetworkingProfiler/Widgets/SPacketContentView.cpp
@@ -705,12 +705,14 @@ void SPacketContentView::EnableFilterEventType(const uint32 InEventTypeIndex)
 	if (Session.IsValid())
 	{
 		TraceServices::FAnalysisSessionReadScope SessionReadScope(*Session.Get());
-		const TraceServices::INetProfilerProvider& NetProfilerProvider = TraceServices::ReadNetProfilerProvider(*Session.Get());
-
-		NetProfilerProvider.ReadEventType(InEventTypeIndex, [&EventName](const TraceServices::FNetProfilerEventType& EventType)
+		const TraceServices::INetProfilerProvider* NetProfilerProvider = TraceServices::ReadNetProfilerProvider(*Session.Get());
+		if (NetProfilerProvider)
 		{
-			EventName = FText::FromString(EventType.Name);
-		});
+			NetProfilerProvider->ReadEventType(InEventTypeIndex, [&EventName](const TraceServices::FNetProfilerEventType& EventType)
+			{
+				EventName = FText::FromString(EventType.Name);
+			});
+		}
 	}
 
 	bFilterByEventType = true;
@@ -729,12 +731,14 @@ uint32 SPacketContentView::GetPacketSequence(int32 InPacketIndex) const
 	if (Session.IsValid())
 	{
 		TraceServices::FAnalysisSessionReadScope SessionReadScope(*Session.Get());
-		const TraceServices::INetProfilerProvider& NetProfilerProvider = TraceServices::ReadNetProfilerProvider(*Session.Get());
-
-		NetProfilerProvider.EnumeratePackets(ConnectionIndex, ConnectionMode, InPacketIndex, InPacketIndex, [&NewSequenceNumber](const TraceServices::FNetProfilerPacket& Packet)
+		const TraceServices::INetProfilerProvider* NetProfilerProvider = TraceServices::ReadNetProfilerProvider(*Session.Get());
+		if (NetProfilerProvider)
 		{
-			NewSequenceNumber = Packet.SequenceNumber;
-		});
+			NetProfilerProvider->EnumeratePackets(ConnectionIndex, ConnectionMode, InPacketIndex, InPacketIndex, [&NewSequenceNumber](const TraceServices::FNetProfilerPacket& Packet)
+			{
+				NewSequenceNumber = Packet.SequenceNumber;
+			});
+		}
 	}
 
 	return NewSequenceNumber;
@@ -765,39 +769,41 @@ void SPacketContentView::UpdateState()
 		if (Session.IsValid())
 		{
 			TraceServices::FAnalysisSessionReadScope SessionReadScope(*Session.Get());
-			const TraceServices::INetProfilerProvider& NetProfilerProvider = TraceServices::ReadNetProfilerProvider(*Session.Get());
-
-			const FAxisViewportDouble& ViewportX = Viewport.GetHorizontalAxisViewport();
-
-			//const int64 StartPos = static_cast<int64>(FMath::FloorToDouble(ViewportX.GetValueAtOffset(0.0f)));
-			//const int64 EndPos = static_cast<int64>(FMath::CeilToDouble(ViewportX.GetValueAtOffset(ViewportX.GetSize())));
-			const uint32 StartPos = 0;
-			const uint32 EndPos = PacketBitSize;
-			NetProfilerProvider.EnumeratePacketContentEventsByPosition(ConnectionIndex, ConnectionMode, PacketIndex, StartPos, EndPos, [this, &Builder, &FilteredDrawStateBuilder, &NetProfilerProvider](const TraceServices::FNetProfilerContentEvent& Event)
+			const TraceServices::INetProfilerProvider* NetProfilerProvider = TraceServices::ReadNetProfilerProvider(*Session.Get());
+			if (NetProfilerProvider)
 			{
-				const TCHAR* Name = nullptr;
-				NetProfilerProvider.ReadName(Event.NameIndex, [&Name](const TraceServices::FNetProfilerName& NetProfilerName)
-				{
-					Name = NetProfilerName.Name;
-				});
+				const FAxisViewportDouble& ViewportX = Viewport.GetHorizontalAxisViewport();
 
-				uint32 NetId = 0;
-				if (Event.ObjectInstanceIndex != 0)
+				//const int64 StartPos = static_cast<int64>(FMath::FloorToDouble(ViewportX.GetValueAtOffset(0.0f)));
+				//const int64 EndPos = static_cast<int64>(FMath::CeilToDouble(ViewportX.GetValueAtOffset(ViewportX.GetSize())));
+				const uint32 StartPos = 0;
+				const uint32 EndPos = PacketBitSize;
+				NetProfilerProvider->EnumeratePacketContentEventsByPosition(ConnectionIndex, ConnectionMode, PacketIndex, StartPos, EndPos, [this, &Builder, &FilteredDrawStateBuilder, NetProfilerProvider](const TraceServices::FNetProfilerContentEvent& Event)
 				{
-					NetProfilerProvider.ReadObject(GameInstanceIndex, Event.ObjectInstanceIndex, [&NetId](const TraceServices::FNetProfilerObjectInstance& ObjectInstance)
+					const TCHAR* Name = nullptr;
+					NetProfilerProvider->ReadName(Event.NameIndex, [&Name](const TraceServices::FNetProfilerName& NetProfilerName)
 					{
-						NetId = ObjectInstance.NetId;
+						Name = NetProfilerName.Name;
 					});
-				}
 
-				Builder.AddEvent(Event, Name, NetId);
+					uint32 NetId = 0;
+					if (Event.ObjectInstanceIndex != 0)
+					{
+						NetProfilerProvider->ReadObject(GameInstanceIndex, Event.ObjectInstanceIndex, [&NetId](const TraceServices::FNetProfilerObjectInstance& ObjectInstance)
+						{
+							NetId = ObjectInstance.NetId;
+						});
+					}
 
-				if ((!bFilterByEventType || FilterEventTypeIndex == Event.EventTypeIndex) &&
-					(!bFilterByNetId || (Event.ObjectInstanceIndex != 0 && FilterNetId == NetId)))
-				{
-					FilteredDrawStateBuilder.AddEvent(Event, Name, NetId);
-				}
-			});
+					Builder.AddEvent(Event, Name, NetId);
+
+					if ((!bFilterByEventType || FilterEventTypeIndex == Event.EventTypeIndex) &&
+						(!bFilterByNetId || (Event.ObjectInstanceIndex != 0 && FilterNetId == NetId)))
+					{
+						FilteredDrawStateBuilder.AddEvent(Event, Name, NetId);
+					}
+				});
+			}
 		}
 
 		Builder.Flush();
@@ -836,24 +842,26 @@ void SPacketContentView::UpdateHoveredEvent()
 		if (Session.IsValid())
 		{
 			TraceServices::FAnalysisSessionReadScope SessionReadScope(*Session.Get());
-			const TraceServices::INetProfilerProvider& NetProfilerProvider = TraceServices::ReadNetProfilerProvider(*Session.Get());
-
-			NetProfilerProvider.ReadEventType(Event.EventTypeIndex, [&EventType](const TraceServices::FNetProfilerEventType& InEventType)
+			const TraceServices::INetProfilerProvider* NetProfilerProvider = TraceServices::ReadNetProfilerProvider(*Session.Get());
+			if (NetProfilerProvider)
 			{
-				EventType = InEventType;
-			});
-
-			NetProfilerProvider.ReadName(EventType.NameIndex, [&Name](const TraceServices::FNetProfilerName& NetProfilerName)
-			{
-				Name = NetProfilerName.Name;
-			});
+				NetProfilerProvider->ReadEventType(Event.EventTypeIndex, [&EventType](const TraceServices::FNetProfilerEventType& InEventType)
+				{
+					EventType = InEventType;
+				});
 
-			if (Event.ObjectInstanceIndex != 0)
-			{
-				NetProfilerProvider.ReadObject(GameInstanceIndex, Event.ObjectInstanceIndex, [&ObjectInstance](const TraceServices::FNetProfilerObjectInstance& InObjectInstance)
+				NetProfilerProvider->ReadName(EventType.NameIndex, [&Name](const TraceServices::FNetProfilerName& NetProfilerName)
 				{
-					ObjectInstance = InObjectInstance;
+					Name = NetProfilerName.Name;
 				});
+
+				if (Event.ObjectInstanceIndex != 0)
+				{
+					NetProfilerProvider->ReadObject(GameInstanceIndex, Event.ObjectInstanceIndex, [&ObjectInstance](const TraceServices::FNetProfilerObjectInstance& InObjectInstance)
+					{
+						ObjectInstance = InObjectInstance;
+					});
+				}
 			}
 		}
 
diff --git a/Engine/Source/Developer/TraceInsights/Private/Insights/NetworkingProfiler/Widgets/SPacketView.cpp b/Engine/Source/Developer/TraceInsights/Private/Insights/NetworkingProfiler/Widgets/SPacketView.cpp
--- a/Engine/Source/Developer/TraceInsights/Private/Insights/NetworkingProfiler/Widgets/SPacketView.cpp
+++ b/Engine/Source/Developer/TraceInsights/Private/Insights/NetworkingProfiler/Widgets/SPacketView.cpp
@@ -208,18 +208,19 @@ void SPacketView::Tick(const FGeometry& AllottedGeometry, const double InCurrent
 		if (Session.IsValid())
 		{
 			TraceServices::FAnalysisSessionReadScope SessionReadScope(*Session.Get());
-
-			const TraceServices::INetProfilerProvider& NetProfilerProvider = TraceServices::ReadNetProfilerProvider(*Session.Get());
-
-			const uint32 NewConnectionChangeCount = NetProfilerProvider.GetConnectionChangeCount();
-			if (NewConnectionChangeCount != ConnectionChangeCount)
+			const TraceServices::INetProfilerProvider* NetProfilerProvider = TraceServices::ReadNetProfilerProvider(*Session.Get());
+			if (NetProfilerProvider)
 			{
-				ConnectionChangeCount = NewConnectionChangeCount;
-				bIsStateDirty = true;
-
-				if (bIsAutoZoomEnabled)
+				const uint32 NewConnectionChangeCount = NetProfilerProvider->GetConnectionChangeCount();
+				if (NewConnectionChangeCount != ConnectionChangeCount)
 				{
-					bAutoZoom = true;
+					ConnectionChangeCount = NewConnectionChangeCount;
+					bIsStateDirty = true;
+
+					if (bIsAutoZoomEnabled)
+					{
+						bAutoZoom = true;
+					}
 				}
 			}
 		}
@@ -313,11 +314,10 @@ void SPacketView::UpdateState()
 	if (Session.IsValid())
 	{
 		TraceServices::FAnalysisSessionReadScope SessionReadScope(*Session.Get());
-		const TraceServices::INetProfilerProvider& NetProfilerProvider = TraceServices::ReadNetProfilerProvider(*Session.Get());
-
-		if (IsConnectionValid(NetProfilerProvider, GameInstanceIndex, ConnectionIndex, ConnectionMode))
+		const TraceServices::INetProfilerProvider* NetProfilerProvider = TraceServices::ReadNetProfilerProvider(*Session.Get());
+		if (NetProfilerProvider && IsConnectionValid(*NetProfilerProvider, GameInstanceIndex, ConnectionIndex, ConnectionMode))
 		{
-			const uint32 NumPackets = NetProfilerProvider.GetPacketCount(ConnectionIndex, ConnectionMode);
+			const uint32 NumPackets = NetProfilerProvider->GetPacketCount(ConnectionIndex, ConnectionMode);
 
 			ViewportX.SetMinMaxInterval(0, NumPackets);
 
@@ -332,7 +332,7 @@ void SPacketView::UpdateState()
 				if (PacketStartIndex <= PacketEndIndex)
 				{
 					int32 PacketIndex = PacketStartIndex;
-					NetProfilerProvider.EnumeratePackets(ConnectionIndex, ConnectionMode, PacketStartIndex, PacketEndIndex, [this, &Builder, &PacketIndex, &Filter, &NetProfilerProvider](const TraceServices::FNetProfilerPacket& Packet)
+					NetProfilerProvider->EnumeratePackets(ConnectionIndex, ConnectionMode, PacketStartIndex, PacketEndIndex, [this, &Builder, &PacketIndex, &Filter, NetProfilerProvider](const TraceServices::FNetProfilerPacket& Packet)
 					{
 						FNetworkPacketAggregatedSample* SamplePtr = Builder.AddPacket(PacketIndex++, Packet);
 
@@ -349,7 +349,7 @@ void SPacketView::UpdateState()
 
 								const uint32 StartPos = 0;
 								const uint32 EndPos = Packet.ContentSizeInBits;
-								NetProfilerProvider.EnumeratePacketContentEventsByPosition(ConnectionIndex, ConnectionMode, PacketIndex - 1, StartPos, EndPos, [this, &bFilterMatch, &Filter, &NetProfilerProvider](const TraceServices::FNetProfilerContentEvent& Event)
+								NetProfilerProvider->EnumeratePacketContentEventsByPosition(ConnectionIndex, ConnectionMode, PacketIndex - 1, StartPos, EndPos, [this, &bFilterMatch, &Filter, NetProfilerProvider](const TraceServices::FNetProfilerContentEvent& Event)
 								{
 									if (!bFilterMatch)
 									{
@@ -363,7 +363,7 @@ void SPacketView::UpdateState()
 											uint32 NetId = uint32(-1);
 											if (Event.ObjectInstanceIndex != 0)
 											{
-												NetProfilerProvider.ReadObject(GameInstanceIndex, Event.ObjectInstanceIndex, [&NetId](const TraceServices::FNetProfilerObjectInstance& ObjectInstance)
+												NetProfilerProvider->ReadObject(GameInstanceIndex, Event.ObjectInstanceIndex, [&NetId](const TraceServices::FNetProfilerObjectInstance& ObjectInstance)
 												{
 													NetId = ObjectInstance.NetId;
 												});
@@ -453,14 +453,13 @@ FNetworkPacketSampleRef SPacketView::GetSample(const int32 InPacketIndex)
 	if (Session.IsValid())
 	{
 		TraceServices::FAnalysisSessionReadScope SessionReadScope(*Session.Get());
-		const TraceServices::INetProfilerProvider& NetProfilerProvider = TraceServices::ReadNetProfilerProvider(*Session.Get());
-
-		if (IsConnectionValid(NetProfilerProvider, GameInstanceIndex, ConnectionIndex, ConnectionMode))
+		const TraceServices::INetProfilerProvider* NetProfilerProvider = TraceServices::ReadNetProfilerProvider(*Session.Get());
+		if (NetProfilerProvider && IsConnectionValid(*NetProfilerProvider, GameInstanceIndex, ConnectionIndex, ConnectionMode))
 		{
-			const uint32 NumPackets = NetProfilerProvider.GetPacketCount(ConnectionIndex, ConnectionMode);
+			const uint32 NumPackets = NetProfilerProvider->GetPacketCount(ConnectionIndex, ConnectionMode);
 			if (InPacketIndex >= 0 && InPacketIndex < static_cast<int32>(NumPackets))
 			{
-				NetProfilerProvider.EnumeratePackets(ConnectionIndex, ConnectionMode, InPacketIndex, InPacketIndex, [InPacketIndex, &SampleRef](const TraceServices::FNetProfilerPacket& Packet)
+				NetProfilerProvider->EnumeratePackets(ConnectionIndex, ConnectionMode, InPacketIndex, InPacketIndex, [InPacketIndex, &SampleRef](const TraceServices::FNetProfilerPacket& Packet)
 				{
 					SampleRef.Sample = MakeShared<FNetworkPacketAggregatedSample>();
 					SampleRef.Sample->AddPacket(InPacketIndex, Packet);
@@ -1361,12 +1360,14 @@ void SPacketView::SelectPacketBySequenceNumber(const uint32 InSequenceNumber)
 	if (Session.IsValid())
 	{
 		TraceServices::FAnalysisSessionReadScope SessionReadScope(*Session.Get());
-		const TraceServices::INetProfilerProvider& NetProfilerProvider = TraceServices::ReadNetProfilerProvider(*Session.Get());
-
-		const int32 PacketId = NetProfilerProvider.FindPacketIndexFromPacketSequence(ConnectionIndex, ConnectionMode, InSequenceNumber);
-		if (PacketId != -1)
+		const TraceServices::INetProfilerProvider* NetProfilerProvider = TraceServices::ReadNetProfilerProvider(*Session.Get());
+		if (NetProfilerProvider)
 		{
-			SetSelectedPacket(PacketId);
+			const int32 PacketId = NetProfilerProvider->FindPacketIndexFromPacketSequence(ConnectionIndex, ConnectionMode, InSequenceNumber);
+			if (PacketId != -1)
+			{
+				SetSelectedPacket(PacketId);
+			}
 		}
 	}
 }
diff --git a/Engine/Source/Developer/TraceInsights/Private/Insights/Widgets/SSessionInfoWindow.cpp b/Engine/Source/Developer/TraceInsights/Private/Insights/Widgets/SSessionInfoWindow.cpp
--- a/Engine/Source/Developer/TraceInsights/Private/Insights/Widgets/SSessionInfoWindow.cpp
+++ b/Engine/Source/Developer/TraceInsights/Private/Insights/Widgets/SSessionInfoWindow.cpp
@@ -285,11 +285,11 @@ void SSessionInfoWindow::Tick(const FGeometry& AllottedGeometry, const double In
 	if (Session.IsValid())
 	{
 		TraceServices::FAnalysisSessionReadScope SessionReadScope(*Session.Get());
-		const TraceServices::IDiagnosticsProvider& DiagnosticsProvider = TraceServices::ReadDiagnosticsProvider(*Session.Get());
+		const TraceServices::IDiagnosticsProvider* DiagnosticsProvider = TraceServices::ReadDiagnosticsProvider(*Session.Get());
 
-		if (DiagnosticsProvider.IsSessionInfoAvailable())
+		if (DiagnosticsProvider && DiagnosticsProvider->IsSessionInfoAvailable())
 		{
-			TraceServices::FSessionInfo SessionInfo = DiagnosticsProvider.GetSessionInfo();
+			TraceServices::FSessionInfo SessionInfo = DiagnosticsProvider->GetSessionInfo();
 			PlatformText = FText::FromString(SessionInfo.Platform);
 			AppNameText = FText::FromString(SessionInfo.AppName);
 			BranchText = FText::FromString(SessionInfo.Branch);
diff --git a/Engine/Source/Developer/TraceServices/Private/AnalysisService.cpp b/Engine/Source/Developer/TraceServices/Private/AnalysisService.cpp
--- a/Engine/Source/Developer/TraceServices/Private/AnalysisService.cpp
+++ b/Engine/Source/Developer/TraceServices/Private/AnalysisService.cpp
@@ -236,18 +236,9 @@ TSharedPtr<const IAnalysisSession> FAnalysisService::StartAnalysis(const TCHAR*
 	FCounterProvider* CounterProvider = new FCounterProvider(*Session, *FrameProvider);
 	Session->AddProvider(FCounterProvider::ProviderName, CounterProvider);
 
-	FNetProfilerProvider* NetProfilerProvider = new FNetProfilerProvider(*Session);
-	Session->AddProvider(FNetProfilerProvider::ProviderName, NetProfilerProvider);
-
 	FChannelProvider* ChannelProvider = new FChannelProvider();
 	Session->AddProvider(FChannelProvider::ProviderName, ChannelProvider);
 
-	FMemoryProvider* MemoryProvider = new FMemoryProvider(*Session);
-	Session->AddProvider(FMemoryProvider::ProviderName, MemoryProvider);
-
-	FDiagnosticsProvider* DiagnosticsProvider = new FDiagnosticsProvider(*Session);
-	Session->AddProvider(FDiagnosticsProvider::ProviderName, DiagnosticsProvider);
-
 	Session->AddAnalyzer(new FMiscTraceAnalyzer(*Session, *ThreadProvider, *BookmarkProvider, *LogProvider, *FrameProvider, *ChannelProvider));
 	Session->AddAnalyzer(new FLogTraceAnalyzer(*Session, *LogProvider));
 
diff --git a/Engine/Source/Developer/TraceServices/Private/Analyzers/CallstacksAnalysis.cpp b/Engine/Source/Developer/TraceServices/Private/Analyzers/CallstacksAnalysis.cpp
--- a/Engine/Source/Developer/TraceServices/Private/Analyzers/CallstacksAnalysis.cpp
+++ b/Engine/Source/Developer/TraceServices/Private/Analyzers/CallstacksAnalysis.cpp
@@ -8,9 +8,11 @@ namespace TraceServices
 {
 
 ////////////////////////////////////////////////////////////////////////////////
-FCallstacksAnalyzer::FCallstacksAnalyzer(IAnalysisSession& Session, FCallstacksProvider* InProvider)
-	: Provider(InProvider)
+FCallstacksAnalyzer::FCallstacksAnalyzer(IAnalysisSession& InSession, FCallstacksProvider* InProvider)
+	: Session(InSession)
+	, Provider(InProvider)
 {
+	check(Provider != nullptr);
 }
 
 ////////////////////////////////////////////////////////////////////////////////
@@ -20,7 +22,6 @@ void FCallstacksAnalyzer::OnAnalysisBegin(const FOnAnalysisContext& Context)
 	Builder.RouteEvent(RouteId_Callstack, "Memory", "CallstackSpec");
 }
 
-
 ////////////////////////////////////////////////////////////////////////////////
 bool FCallstacksAnalyzer::OnEvent(uint16 RouteId, EStyle Style, const FOnEventContext& Context)
 {
diff --git a/Engine/Source/Developer/TraceServices/Private/Analyzers/CallstacksAnalysis.h b/Engine/Source/Developer/TraceServices/Private/Analyzers/CallstacksAnalysis.h
--- a/Engine/Source/Developer/TraceServices/Private/Analyzers/CallstacksAnalysis.h
+++ b/Engine/Source/Developer/TraceServices/Private/Analyzers/CallstacksAnalysis.h
@@ -16,16 +16,17 @@ class FCallstacksAnalyzer
 	: public UE::Trace::IAnalyzer
 {
 public:
-					FCallstacksAnalyzer(IAnalysisSession& Session, FCallstacksProvider* Provider);
-	virtual void	OnAnalysisBegin(const FOnAnalysisContext& Context) override;
-	virtual bool 	OnEvent(uint16 RouteId, EStyle Style, const FOnEventContext& Context) override;
+	FCallstacksAnalyzer(IAnalysisSession& Session, FCallstacksProvider* Provider);
+	virtual void OnAnalysisBegin(const FOnAnalysisContext& Context) override;
+	virtual bool OnEvent(uint16 RouteId, EStyle Style, const FOnEventContext& Context) override;
 
 private:
-	enum Routes 
+	enum Routes
 	{
 		RouteId_Callstack,
 	};
 
+	IAnalysisSession& Session;
 	FCallstacksProvider* Provider;
 };
 
diff --git a/Engine/Source/Developer/TraceServices/Private/Analyzers/DiagnosticsAnalysis.cpp b/Engine/Source/Developer/TraceServices/Private/Analyzers/DiagnosticsAnalysis.cpp
--- a/Engine/Source/Developer/TraceServices/Private/Analyzers/DiagnosticsAnalysis.cpp
+++ b/Engine/Source/Developer/TraceServices/Private/Analyzers/DiagnosticsAnalysis.cpp
@@ -6,10 +6,11 @@
 namespace TraceServices
 {
 
-FDiagnosticsAnalyzer::FDiagnosticsAnalyzer(IAnalysisSession& InSession)
+FDiagnosticsAnalyzer::FDiagnosticsAnalyzer(IAnalysisSession& InSession, FDiagnosticsProvider* InProvider)
 	: Session(InSession)
+	, Provider(InProvider)
 {
-	Provider = Session.EditProvider<FDiagnosticsProvider>(FDiagnosticsProvider::ProviderName);
+	check(Provider != nullptr);
 }
 
 FDiagnosticsAnalyzer::~FDiagnosticsAnalyzer()
diff --git a/Engine/Source/Developer/TraceServices/Private/Analyzers/DiagnosticsAnalysis.h b/Engine/Source/Developer/TraceServices/Private/Analyzers/DiagnosticsAnalysis.h
--- a/Engine/Source/Developer/TraceServices/Private/Analyzers/DiagnosticsAnalysis.h
+++ b/Engine/Source/Developer/TraceServices/Private/Analyzers/DiagnosticsAnalysis.h
@@ -16,7 +16,7 @@ class FDiagnosticsAnalyzer
 	: public UE::Trace::IAnalyzer
 {
 public:
-	FDiagnosticsAnalyzer(IAnalysisSession& Session);
+	FDiagnosticsAnalyzer(IAnalysisSession& Session, FDiagnosticsProvider* InProvider);
 	~FDiagnosticsAnalyzer();
 	virtual void OnAnalysisBegin(const FOnAnalysisContext& Context) override;
 	virtual bool OnEvent(uint16 RouteId, EStyle Style, const FOnEventContext& Context) override;
@@ -27,8 +27,8 @@ private:
 		RouteId_Session,
 		RouteId_Session2,
 	};
-	FDiagnosticsProvider* Provider;
 	IAnalysisSession& Session;
+	FDiagnosticsProvider* Provider;
 };
 
 } // namespace TraceServices
diff --git a/Engine/Source/Developer/TraceServices/Private/Analyzers/MemoryAnalysis.cpp b/Engine/Source/Developer/TraceServices/Private/Analyzers/MemoryAnalysis.cpp
--- a/Engine/Source/Developer/TraceServices/Private/Analyzers/MemoryAnalysis.cpp
+++ b/Engine/Source/Developer/TraceServices/Private/Analyzers/MemoryAnalysis.cpp
@@ -6,11 +6,11 @@
 namespace TraceServices
 {
 
-FMemoryAnalyzer::FMemoryAnalyzer(IAnalysisSession& InSession)
-	: Provider(nullptr)
-	, Session(InSession)
+FMemoryAnalyzer::FMemoryAnalyzer(IAnalysisSession& InSession, FMemoryProvider* InProvider)
+	: Session(InSession)
+	, Provider(InProvider)
 {
-	Provider = Session.EditProvider<FMemoryProvider>(FMemoryProvider::ProviderName);
+	check(Provider != nullptr);
 }
 
 FMemoryAnalyzer::~FMemoryAnalyzer()
diff --git a/Engine/Source/Developer/TraceServices/Private/Analyzers/MemoryAnalysis.h b/Engine/Source/Developer/TraceServices/Private/Analyzers/MemoryAnalysis.h
--- a/Engine/Source/Developer/TraceServices/Private/Analyzers/MemoryAnalysis.h
+++ b/Engine/Source/Developer/TraceServices/Private/Analyzers/MemoryAnalysis.h
@@ -17,7 +17,7 @@ class FMemoryAnalyzer
 	: public UE::Trace::IAnalyzer
 {
 public:
-	FMemoryAnalyzer(IAnalysisSession& Session);
+	FMemoryAnalyzer(IAnalysisSession& Session, FMemoryProvider* InProvider);
 	~FMemoryAnalyzer();
 	virtual void OnAnalysisBegin(const FOnAnalysisContext& Context) override;
 	virtual bool OnEvent(uint16 RouteId, EStyle Style, const FOnEventContext& Context) override;
@@ -30,8 +30,8 @@ private:
 		RouteId_TagValue,
 	};
 
-	FMemoryProvider* Provider;
 	IAnalysisSession& Session;
+	FMemoryProvider* Provider;
 	uint64 Sample = 0;
 };
 
diff --git a/Engine/Source/Developer/TraceServices/Private/Analyzers/ModuleAnalysis.cpp b/Engine/Source/Developer/TraceServices/Private/Analyzers/ModuleAnalysis.cpp
--- a/Engine/Source/Developer/TraceServices/Private/Analyzers/ModuleAnalysis.cpp
+++ b/Engine/Source/Developer/TraceServices/Private/Analyzers/ModuleAnalysis.cpp
@@ -57,7 +57,7 @@ bool FModuleAnalyzer::OnEvent(uint16 RouteId, EStyle Style, const FOnEventContex
 					Provider = CreateModuleProvider(Session, SymbolFormat);
 					if (Provider)
 					{
-						Session.AddProvider(TEXT("ModuleProvider"), Provider);
+						Session.AddProvider(GetModuleProviderName(), Provider);
 					}
 				}
 			}
diff --git a/Engine/Source/Developer/TraceServices/Private/Model/AllocationsProvider.cpp b/Engine/Source/Developer/TraceServices/Private/Model/AllocationsProvider.cpp
--- a/Engine/Source/Developer/TraceServices/Private/Model/AllocationsProvider.cpp
+++ b/Engine/Source/Developer/TraceServices/Private/Model/AllocationsProvider.cpp
@@ -25,7 +25,7 @@
 
 // Use optimized path for the last alloc.
 // If enabled, caches the last added alloc, as there is a high chance to be freed in the next "free" event.
-// ~10% to ~30% of all allocs are expected to have an "event distance" == 1 event ("free" event follows the "alloc" event imediatelly)
+// ~10% to ~30% of all allocs are expected to have an "event distance" == 1 event ("free" event follows the "alloc" event immediately)
 #define INSIGHTS_USE_LAST_ALLOC 1
 
 #define INSIGHTS_VALIDATE_ALLOC_EVENTS 0
@@ -677,14 +677,6 @@ FAllocationsProvider::~FAllocationsProvider()
 
 ////////////////////////////////////////////////////////////////////////////////////////////////////
 
-FName FAllocationsProvider::GetName()
-{
-	static FName Name(TEXT("AllocationsProvider"));
-	return Name;
-}
-
-////////////////////////////////////////////////////////////////////////////////////////////////////
-
 void FAllocationsProvider::EditInit(double InTime, uint8 InMinAlignment, uint8 InSizeShift, uint8 InSummarySizeShift)
 {
 	Lock.WriteAccessCheck();
@@ -1335,9 +1327,17 @@ const IAllocationsProvider::FQueryStatus FAllocationsProvider::PollQuery(FQueryH
 
 ////////////////////////////////////////////////////////////////////////////////////////////////////
 
+FName GetAllocationsProviderName()
+{
+	static FName Name(TEXT("AllocationsProvider"));
+	return Name;
+}
+
+////////////////////////////////////////////////////////////////////////////////////////////////////
+
 const IAllocationsProvider* ReadAllocationsProvider(const IAnalysisSession& Session)
 {
-	return Session.ReadProvider<IAllocationsProvider>(FAllocationsProvider::GetName());
+	return Session.ReadProvider<IAllocationsProvider>(GetAllocationsProviderName());
 }
 
 ////////////////////////////////////////////////////////////////////////////////////////////////////
diff --git a/Engine/Source/Developer/TraceServices/Private/Model/AllocationsProvider.h b/Engine/Source/Developer/TraceServices/Private/Model/AllocationsProvider.h
--- a/Engine/Source/Developer/TraceServices/Private/Model/AllocationsProvider.h
+++ b/Engine/Source/Developer/TraceServices/Private/Model/AllocationsProvider.h
@@ -202,8 +202,6 @@ public:
 	explicit FAllocationsProvider(IAnalysisSession& InSession);
 	virtual ~FAllocationsProvider();
 
-	static FName GetName();
-
 	virtual void BeginEdit() const override { Lock.BeginWrite(); }
 	virtual void EndEdit() const override { Lock.EndWrite(); }
 	void EditAccessCheck() const { return Lock.WriteAccessCheck(); }
diff --git a/Engine/Source/Developer/TraceServices/Private/Model/CallstacksProvider.cpp b/Engine/Source/Developer/TraceServices/Private/Model/CallstacksProvider.cpp
--- a/Engine/Source/Developer/TraceServices/Private/Model/CallstacksProvider.cpp
+++ b/Engine/Source/Developer/TraceServices/Private/Model/CallstacksProvider.cpp
@@ -55,7 +55,7 @@ void FCallstacksProvider::AddCallstack(uint64 InCallstackId, const uint64* InFra
 	// created so EditProvider() may return a null pointer.
 	if (!ModuleProvider)
 	{
-		ModuleProvider = Session.EditProvider<IModuleProvider>(FName("ModuleProvider"));
+		ModuleProvider = Session.EditProvider<IModuleProvider>(GetModuleProviderName());
 	}
 
 	// Make sure all the frames fit on one page by appending dummy entries.
@@ -117,10 +117,15 @@ void FCallstacksProvider::GetCallstacks(const TArrayView<uint64>& CallstackIds,
 }
 
 /////////////////////////////////////////////////////////////////////
-FName FCallstacksProvider::GetName() const
+FName GetCallstacksProviderName()
 {
 	static FName Name(TEXT("CallstacksProvider"));
 	return Name;
 }
 
+const ICallstacksProvider* ReadCallstacksProvider(const IAnalysisSession& Session)
+{
+	return Session.ReadProvider<ICallstacksProvider>(GetCallstacksProviderName());
+}
+
 } // namespace TraceServices
diff --git a/Engine/Source/Developer/TraceServices/Private/Model/CallstacksProvider.h b/Engine/Source/Developer/TraceServices/Private/Model/CallstacksProvider.h
--- a/Engine/Source/Developer/TraceServices/Private/Model/CallstacksProvider.h
+++ b/Engine/Source/Developer/TraceServices/Private/Model/CallstacksProvider.h
@@ -23,7 +23,6 @@ public:
 	const FCallstack*	GetCallstack(uint64 CallstackId) const override;
 	void				GetCallstacks(const TArrayView<uint64>& CallstackIds, FCallstack const** OutCallstacks) const override;
 	void				AddCallstack(uint64 CallstackId, const uint64* Frames, uint8 FrameCount);
-	FName				GetName() const;
 
 private:
 	enum
diff --git a/Engine/Source/Developer/TraceServices/Private/Model/Diagnostics.cpp b/Engine/Source/Developer/TraceServices/Private/Model/Diagnostics.cpp
--- a/Engine/Source/Developer/TraceServices/Private/Model/Diagnostics.cpp
+++ b/Engine/Source/Developer/TraceServices/Private/Model/Diagnostics.cpp
@@ -5,8 +5,6 @@
 namespace TraceServices
 {
 
-FName FDiagnosticsProvider::ProviderName(TEXT("DiagnosticsProvider"));
-
 FDiagnosticsProvider::FDiagnosticsProvider(IAnalysisSession& InSession)
 	: Session(InSession)
 {
@@ -31,9 +29,15 @@ const FSessionInfo& FDiagnosticsProvider::GetSessionInfo() const
 	return SessionInfo;
 }
 
-const IDiagnosticsProvider& ReadDiagnosticsProvider(const IAnalysisSession& Session)
+FName GetDiagnosticsProviderName()
+{
+	static FName Name(TEXT("DiagnosticsProvider"));
+	return Name;
+}
+
+const IDiagnosticsProvider* ReadDiagnosticsProvider(const IAnalysisSession& Session)
 {
-	return *Session.ReadProvider<IDiagnosticsProvider>(FDiagnosticsProvider::ProviderName);
+	return Session.ReadProvider<IDiagnosticsProvider>(GetDiagnosticsProviderName());
 }
 
 } // namespace TraceServices
diff --git a/Engine/Source/Developer/TraceServices/Private/Model/DiagnosticsPrivate.h b/Engine/Source/Developer/TraceServices/Private/Model/DiagnosticsPrivate.h
--- a/Engine/Source/Developer/TraceServices/Private/Model/DiagnosticsPrivate.h
+++ b/Engine/Source/Developer/TraceServices/Private/Model/DiagnosticsPrivate.h
@@ -14,8 +14,6 @@ public:
 	virtual ~FDiagnosticsProvider() {}
 
 public:
-	static FName ProviderName;
-
 	void SetSessionInfo(const FSessionInfo& InSessionInfo);
 	virtual const FSessionInfo& GetSessionInfo() const override;
 	virtual bool IsSessionInfoAvailable() const override;
diff --git a/Engine/Source/Developer/TraceServices/Private/Model/Memory.cpp b/Engine/Source/Developer/TraceServices/Private/Model/Memory.cpp
--- a/Engine/Source/Developer/TraceServices/Private/Model/Memory.cpp
+++ b/Engine/Source/Developer/TraceServices/Private/Model/Memory.cpp
@@ -5,8 +5,6 @@
 namespace TraceServices
 {
 
-FName FMemoryProvider::ProviderName(TEXT("MemoryProvider"));
-
 FMemoryProvider::FMemoryProvider(IAnalysisSession& InSession)
 	: Session(InSession)
 	, TagsSerial(0)
@@ -215,9 +213,15 @@ void FMemoryProvider::EnumerateTagSamples(FMemoryTrackerId TrackerId, FMemoryTag
 	}
 }
 
-const IMemoryProvider& ReadMemoryProvider(const IAnalysisSession& Session)
+FName GetMemoryProviderName()
+{
+	static FName Name(TEXT("MemoryProvider"));
+	return Name;
+}
+
+const IMemoryProvider* ReadMemoryProvider(const IAnalysisSession& Session)
 {
-	return *Session.ReadProvider<IMemoryProvider>(FMemoryProvider::ProviderName);
+	return Session.ReadProvider<IMemoryProvider>(GetMemoryProviderName());
 }
 
 } // namespace TraceServices
diff --git a/Engine/Source/Developer/TraceServices/Private/Model/ModuleProvider.cpp b/Engine/Source/Developer/TraceServices/Private/Model/ModuleProvider.cpp
--- a/Engine/Source/Developer/TraceServices/Private/Model/ModuleProvider.cpp
+++ b/Engine/Source/Developer/TraceServices/Private/Model/ModuleProvider.cpp
@@ -263,6 +263,21 @@ IModuleAnalysisProvider* CreateModuleProvider(IAnalysisSession& InSession, const
 	return Provider;
 }
 
+/////////////////////////////////////////////////////////////////////
+FName GetModuleProviderName()
+{
+	static FName Name(TEXT("ModuleProvider"));
+	return Name;
+}
+
+/////////////////////////////////////////////////////////////////////
+const IModuleProvider* ReadModuleProvider(const IAnalysisSession& Session)
+{
+	return Session.ReadProvider<IModuleProvider>(GetModuleProviderName());
+}
+
+/////////////////////////////////////////////////////////////////////
+
 } // namespace TraceServices
 
 #undef USE_SYMSLIB
diff --git a/Engine/Source/Developer/TraceServices/Private/Model/ModuleProvider.h b/Engine/Source/Developer/TraceServices/Private/Model/ModuleProvider.h
--- a/Engine/Source/Developer/TraceServices/Private/Model/ModuleProvider.h
+++ b/Engine/Source/Developer/TraceServices/Private/Model/ModuleProvider.h
@@ -10,17 +10,15 @@ namespace TraceServices
 
 class IAnalysisSession;
 
-/////////////////////////////////////////////////////////////////////
 class IModuleAnalysisProvider : public IModuleProvider
 {
 public:
-	virtual void	OnModuleLoad(const FStringView& Module, uint64 Base, uint32 Size, const uint8* ImageId, uint32 ImageIdSize) = 0;
-	virtual void	OnModuleUnload(uint64 Base) = 0;
-	virtual void	OnAnalysisComplete() = 0;
+	virtual void OnModuleLoad(const FStringView& Module, uint64 Base, uint32 Size, const uint8* ImageId, uint32 ImageIdSize) = 0;
+	virtual void OnModuleUnload(uint64 Base) = 0;
+	virtual void OnAnalysisComplete() = 0;
 };
 
-/** Create a module provider with the given symbol format */
+/** Create a module provider with the given symbol format. */
 IModuleAnalysisProvider* CreateModuleProvider(IAnalysisSession& Session, const FAnsiStringView& SymbolFormat);
 
-
 } // namespace TraceServices
diff --git a/Engine/Source/Developer/TraceServices/Private/Model/NetProfilerProvider.cpp b/Engine/Source/Developer/TraceServices/Private/Model/NetProfilerProvider.cpp
--- a/Engine/Source/Developer/TraceServices/Private/Model/NetProfilerProvider.cpp
+++ b/Engine/Source/Developer/TraceServices/Private/Model/NetProfilerProvider.cpp
@@ -8,8 +8,6 @@
 namespace TraceServices
 {
 
-const FName FNetProfilerProvider::ProviderName("NetProfilerProvider");
-
 const TCHAR* LexToString(const ENetProfilerChannelCloseReason Value)
 {
 	switch (Value)
@@ -371,13 +369,6 @@ uint32 FNetProfilerProvider::GetObjectsChangeCount(uint32 GameInstanceIndex) con
 	return GameInstance.ObjectsChangeCount;
 }
 
-const INetProfilerProvider& ReadNetProfilerProvider(const IAnalysisSession& Session)
-{
-	Session.ReadAccessCheck();
-
-	return *Session.ReadProvider<INetProfilerProvider>(FNetProfilerProvider::ProviderName);
-}
-
 int32 FNetProfilerProvider::FindPacketIndexFromPacketSequence(uint32 ConnectionIndex, ENetProfilerConnectionMode Mode, uint32 SequenceNumber) const
 {
 	Session.ReadAccessCheck();
@@ -678,4 +669,15 @@ ITable<FNetProfilerAggregatedStats>* FNetProfilerProvider::CreateAggregation(uin
 	return Table;
 }
 
+FName GetNetProfilerProviderName()
+{
+	static FName Name(TEXT("NetProfilerProvider"));
+	return Name;
+}
+
+const INetProfilerProvider* ReadNetProfilerProvider(const IAnalysisSession& Session)
+{
+	return Session.ReadProvider<INetProfilerProvider>(GetNetProfilerProviderName());
+}
+
 } // namespace TraceServices
diff --git a/Engine/Source/Developer/TraceServices/Private/Model/NetProfilerProvider.h b/Engine/Source/Developer/TraceServices/Private/Model/NetProfilerProvider.h
--- a/Engine/Source/Developer/TraceServices/Private/Model/NetProfilerProvider.h
+++ b/Engine/Source/Developer/TraceServices/Private/Model/NetProfilerProvider.h
@@ -48,8 +48,6 @@ class FNetProfilerProvider
 	: public INetProfilerProvider
 {
 public:
-	static const FName ProviderName;
-
 	explicit FNetProfilerProvider(IAnalysisSession& InSession);
 	virtual ~FNetProfilerProvider();
 
diff --git a/Engine/Source/Developer/TraceServices/Private/Modules/DiagnosticsModule.cpp b/Engine/Source/Developer/TraceServices/Private/Modules/DiagnosticsModule.cpp
--- a/Engine/Source/Developer/TraceServices/Private/Modules/DiagnosticsModule.cpp
+++ b/Engine/Source/Developer/TraceServices/Private/Modules/DiagnosticsModule.cpp
@@ -17,7 +17,10 @@ void FDiagnosticsModule::GetModuleInfo(FModuleInfo& OutModuleInfo)
 
 void FDiagnosticsModule::OnAnalysisBegin(IAnalysisSession& Session)
 {
-	Session.AddAnalyzer(new FDiagnosticsAnalyzer(Session));
+	FDiagnosticsProvider* DiagnosticsProvider = new FDiagnosticsProvider(Session);
+	Session.AddProvider(GetDiagnosticsProviderName(), DiagnosticsProvider);
+
+	Session.AddAnalyzer(new FDiagnosticsAnalyzer(Session, DiagnosticsProvider));
 }
 
 void FDiagnosticsModule::GetLoggers(TArray<const TCHAR *>& OutLoggers)
diff --git a/Engine/Source/Developer/TraceServices/Private/Modules/DiagnosticsModule.h b/Engine/Source/Developer/TraceServices/Private/Modules/DiagnosticsModule.h
--- a/Engine/Source/Developer/TraceServices/Private/Modules/DiagnosticsModule.h
+++ b/Engine/Source/Developer/TraceServices/Private/Modules/DiagnosticsModule.h
@@ -14,11 +14,6 @@ public:
 	virtual void GetModuleInfo(FModuleInfo& OutModuleInfo) override;
 	virtual void OnAnalysisBegin(IAnalysisSession& Session) override;
 	virtual void GetLoggers(TArray<const TCHAR *>& OutLoggers) override;
-	virtual const TCHAR* GetCommandLineArgument() override
-	{
-		return nullptr;
-	}
-	virtual void GenerateReports(const IAnalysisSession& Session, const TCHAR* CmdLine, const TCHAR* OutputDirectory) override {}
 };
 
 } // namespace TraceServices
diff --git a/Engine/Source/Developer/TraceServices/Private/Modules/LoadTimeProfilerModule.h b/Engine/Source/Developer/TraceServices/Private/Modules/LoadTimeProfilerModule.h
--- a/Engine/Source/Developer/TraceServices/Private/Modules/LoadTimeProfilerModule.h
+++ b/Engine/Source/Developer/TraceServices/Private/Modules/LoadTimeProfilerModule.h
@@ -11,20 +11,17 @@ class FLoadTimeProfilerModule
 	: public IModule
 {
 public:
+	#if WITH_EDITOR
+	virtual bool ShouldBeEnabledByDefault() const override { return false; }
+	#endif
 	virtual void GetModuleInfo(FModuleInfo& OutModuleInfo) override;
 	virtual void OnAnalysisBegin(IAnalysisSession& Session) override;
-	virtual void GetLoggers(TArray<const TCHAR *>& OutLoggers) override;
+	virtual void GenerateReports(const IAnalysisSession& Session, const TCHAR* CmdLine, const TCHAR* OutputDirectory) override;
+	virtual void GetLoggers(TArray<const TCHAR*>& OutLoggers) override;
 	virtual const TCHAR* GetCommandLineArgument() override
 	{
 		return TEXT("loadtimetrace");
 	}
-	virtual void GenerateReports(const IAnalysisSession& Session, const TCHAR* CmdLine, const TCHAR* OutputDirectory) override;
-#if WITH_EDITOR
-	virtual bool ShouldBeEnabledByDefault() const override
-	{
-		return false;
-	}
-#endif
 };
 
 } // namespace TraceServices
diff --git a/Engine/Source/Developer/TraceServices/Private/Modules/MemoryModule.cpp b/Engine/Source/Developer/TraceServices/Private/Modules/MemoryModule.cpp
--- a/Engine/Source/Developer/TraceServices/Private/Modules/MemoryModule.cpp
+++ b/Engine/Source/Developer/TraceServices/Private/Modules/MemoryModule.cpp
@@ -22,20 +22,26 @@ void FMemoryModule::GetModuleInfo(FModuleInfo& OutModuleInfo)
 
 void FMemoryModule::OnAnalysisBegin(IAnalysisSession& Session)
 {
-	FAllocationsProvider* AllocationsProvider = new FAllocationsProvider(Session);
-	Session.AddProvider(AllocationsProvider->GetName(), AllocationsProvider);
-	Session.AddAnalyzer(new FAllocationsAnalyzer(Session, *AllocationsProvider));
+	// LLM Tag Stats
+	FMemoryProvider* MemoryProvider = new FMemoryProvider(Session);
+	Session.AddProvider(GetMemoryProviderName(), MemoryProvider);
+	Session.AddAnalyzer(new FMemoryAnalyzer(Session, MemoryProvider));
 
+	// Module
 	Session.AddAnalyzer(new FModuleAnalyzer(Session));
-	
+
+	// Callstack
 	FCallstacksProvider* CallstacksProvider = new FCallstacksProvider(Session);
-	Session.AddProvider(CallstacksProvider->GetName(), CallstacksProvider);
+	Session.AddProvider(GetCallstacksProviderName(), CallstacksProvider);
 	Session.AddAnalyzer(new FCallstacksAnalyzer(Session, CallstacksProvider));
 
-	Session.AddAnalyzer(new FMemoryAnalyzer(Session));
+	// Allocations
+	FAllocationsProvider* AllocationsProvider = new FAllocationsProvider(Session);
+	Session.AddProvider(GetAllocationsProviderName(), AllocationsProvider);
+	Session.AddAnalyzer(new FAllocationsAnalyzer(Session, *AllocationsProvider));
 }
 
-void FMemoryModule::GetLoggers(TArray<const TCHAR *>& OutLoggers)
+void FMemoryModule::GetLoggers(TArray<const TCHAR*>& OutLoggers)
 {
 	OutLoggers.Add(TEXT("Memory"));
 }
diff --git a/Engine/Source/Developer/TraceServices/Private/Modules/MemoryModule.h b/Engine/Source/Developer/TraceServices/Private/Modules/MemoryModule.h
--- a/Engine/Source/Developer/TraceServices/Private/Modules/MemoryModule.h
+++ b/Engine/Source/Developer/TraceServices/Private/Modules/MemoryModule.h
@@ -11,22 +11,12 @@ class FMemoryModule
 	: public IModule
 {
 public:
+	#if WITH_EDITOR
+	virtual bool ShouldBeEnabledByDefault() const override { return false; }
+	#endif
 	virtual void GetModuleInfo(FModuleInfo& OutModuleInfo) override;
 	virtual void OnAnalysisBegin(IAnalysisSession& Session) override;
-	virtual void GetLoggers(TArray<const TCHAR *>& OutLoggers) override;
-	virtual const TCHAR* GetCommandLineArgument() override
-	{
-		return nullptr;
-	}
-	virtual void GenerateReports(const IAnalysisSession& Session, const TCHAR* CmdLine, const TCHAR* OutputDirectory) override
-	{
-	}
-#if WITH_EDITOR
-	virtual bool ShouldBeEnabledByDefault() const override
-	{
-		return false;
-	}
-#endif
+	virtual void GetLoggers(TArray<const TCHAR*>& OutLoggers) override;
 };
 
 } // namespace TraceServices
diff --git a/Engine/Source/Developer/TraceServices/Private/Modules/NetProfilerModule.cpp b/Engine/Source/Developer/TraceServices/Private/Modules/NetProfilerModule.cpp
--- a/Engine/Source/Developer/TraceServices/Private/Modules/NetProfilerModule.cpp
+++ b/Engine/Source/Developer/TraceServices/Private/Modules/NetProfilerModule.cpp
@@ -9,26 +9,23 @@ namespace TraceServices
 {
 
 FName FNetProfilerModule::ModuleName("TraceModule_NetProfiler");
-static const FName NetProfilerProviderName("NetProfilerProvider");
 
 void FNetProfilerModule::GetModuleInfo(FModuleInfo& OutModuleInfo)
 {
 	OutModuleInfo.Name = ModuleName;
 	OutModuleInfo.DisplayName = TEXT("NetProfiler");
 }
-	
-void FNetProfilerModule::OnAnalysisBegin(IAnalysisSession& InSession)
-{
-	FAnalysisSession& Session = static_cast<FAnalysisSession&>(InSession);
 
+void FNetProfilerModule::OnAnalysisBegin(IAnalysisSession& Session)
+{
 	FNetProfilerProvider* NetProfilerProvider = new FNetProfilerProvider(Session);
-	Session.AddProvider(NetProfilerProviderName, NetProfilerProvider);
-	InSession.AddAnalyzer(new FNetTraceAnalyzer(InSession, *NetProfilerProvider));
+	Session.AddProvider(GetNetProfilerProviderName(), NetProfilerProvider);
+
+	Session.AddAnalyzer(new FNetTraceAnalyzer(Session, *NetProfilerProvider));
 }
 
-void FNetProfilerModule::GetLoggers(TArray<const TCHAR *>& OutLoggers)
+void FNetProfilerModule::GetLoggers(TArray<const TCHAR*>& OutLoggers)
 {
-	//OutLoggers.Add(TEXT("NetProfiler"));
 }
 
 } // namespace TraceServices
diff --git a/Engine/Source/Developer/TraceServices/Public/TraceServices/Model/AllocationsProvider.h b/Engine/Source/Developer/TraceServices/Public/TraceServices/Model/AllocationsProvider.h
--- a/Engine/Source/Developer/TraceServices/Public/TraceServices/Model/AllocationsProvider.h
+++ b/Engine/Source/Developer/TraceServices/Public/TraceServices/Model/AllocationsProvider.h
@@ -166,6 +166,7 @@ public:
 	virtual const TCHAR* GetTagName(int32 Tag) const = 0;
 };
 
+TRACESERVICES_API FName GetAllocationsProviderName();
 TRACESERVICES_API const IAllocationsProvider* ReadAllocationsProvider(const IAnalysisSession& Session);
 
 } // namespace TraceServices
diff --git a/Engine/Source/Developer/TraceServices/Public/TraceServices/Model/Callstack.h b/Engine/Source/Developer/TraceServices/Public/TraceServices/Model/Callstack.h
--- a/Engine/Source/Developer/TraceServices/Public/TraceServices/Model/Callstack.h
+++ b/Engine/Source/Developer/TraceServices/Public/TraceServices/Model/Callstack.h
@@ -12,8 +12,8 @@ namespace TraceServices
 /////////////////////////////////////////////////////////////////////
 struct TRACESERVICES_API FStackFrame
 {
-	uint64					Addr;
-	const FResolvedSymbol*	Symbol;
+	uint64 Addr;
+	const FResolvedSymbol* Symbol;
 };
 
 /////////////////////////////////////////////////////////////////////
@@ -21,11 +21,11 @@ struct TRACESERVICES_API FCallstack
 {
 	FCallstack(const FStackFrame* FirstEntry, uint8 FrameCount);
 	/** Get the number of stack frames in callstack. */
-	uint32			Num() const;
+	uint32 Num() const;
 	/** Gets the address at a given stack depth. */
-	uint64			Addr(uint8 Depth) const;
+	uint64 Addr(uint8 Depth) const;
 	/** Gets the cached symbol name at a given stack depth. */
-	const TCHAR*	Name(uint8 Depth) const;
+	const TCHAR* Name(uint8 Depth) const;
 	/** Gets the entire frame at given depth */
 	const FStackFrame* Frame(uint8 Depth) const;
 
@@ -46,23 +46,28 @@ class ICallstacksProvider
 	: public IProvider
 {
 public:
-	virtual				~ICallstacksProvider() = default;
+	virtual ~ICallstacksProvider() = default;
 
 	/**
 	  * Queries a callstack id.
-	  * @param CallstackId		Callstack id to query
-	  * @return					Callstack information. If id is not found a callstack with zero stack depth is returned.
+	  * @param CallstackId Callstack id to query
+	  * @return Callstack information. If id is not found a callstack with zero stack depth is returned.
 	  */
-	virtual const FCallstack*	GetCallstack(uint64 CallstackId) const = 0;
+	virtual const FCallstack* GetCallstack(uint64 CallstackId) const = 0;
 
 	/**
 	  * Queries a set of callstack ids.
-	  * @param CallstackIds		List of callstack ids to query
-	  * @param OutCallstacks	Output list of callstacks. Caller is responsible for allocating space according to CallstackIds length.
+	  * @param CallstackIds List of callstack ids to query
+	  * @param OutCallstacks Output list of callstacks. Caller is responsible for allocating space according to CallstackIds length.
 	  */
-	virtual void				GetCallstacks(const TArrayView<uint64>& CallstackIds, FCallstack const** OutCallstacks) const = 0;
+	virtual void GetCallstacks(const TArrayView<uint64>& CallstackIds, FCallstack const** OutCallstacks) const = 0;
 };
 
+/////////////////////////////////////////////////////////////////////
+
+TRACESERVICES_API FName GetCallstacksProviderName();
+TRACESERVICES_API const ICallstacksProvider* ReadCallstacksProvider(const IAnalysisSession& Session);
+
 /////////////////////////////////////////////////////////////////////
 inline FCallstack::FCallstack(const FStackFrame* InFirstFrame, uint8 InFrameCount)
 {
diff --git a/Engine/Source/Developer/TraceServices/Public/TraceServices/Model/Diagnostics.h b/Engine/Source/Developer/TraceServices/Public/TraceServices/Model/Diagnostics.h
--- a/Engine/Source/Developer/TraceServices/Public/TraceServices/Model/Diagnostics.h
+++ b/Engine/Source/Developer/TraceServices/Public/TraceServices/Model/Diagnostics.h
@@ -30,6 +30,7 @@ public:
 	virtual const FSessionInfo& GetSessionInfo() const = 0;
 };
 
-TRACESERVICES_API const IDiagnosticsProvider& ReadDiagnosticsProvider(const IAnalysisSession& Session);
+TRACESERVICES_API FName GetDiagnosticsProviderName();
+TRACESERVICES_API const IDiagnosticsProvider* ReadDiagnosticsProvider(const IAnalysisSession& Session);
 
 } // namespace TraceServices
diff --git a/Engine/Source/Developer/TraceServices/Public/TraceServices/Model/Memory.h b/Engine/Source/Developer/TraceServices/Public/TraceServices/Model/Memory.h
--- a/Engine/Source/Developer/TraceServices/Public/TraceServices/Model/Memory.h
+++ b/Engine/Source/Developer/TraceServices/Public/TraceServices/Model/Memory.h
@@ -111,6 +111,7 @@ public:
 	virtual void EnumerateTagSamples(FMemoryTrackerId Tracker, FMemoryTagId Tag, double StartTime, double EndTime, bool bIncludeRangeNeighbours, TFunctionRef<void(double Time, double Duration, const FMemoryTagSample&)> Callback) const = 0;
 };
 
-TRACESERVICES_API const IMemoryProvider& ReadMemoryProvider(const IAnalysisSession& Session);
+TRACESERVICES_API FName GetMemoryProviderName();
+TRACESERVICES_API const IMemoryProvider* ReadMemoryProvider(const IAnalysisSession& Session);
 
 } // namespace TraceServices
diff --git a/Engine/Source/Developer/TraceServices/Public/TraceServices/Model/Modules.h b/Engine/Source/Developer/TraceServices/Public/TraceServices/Model/Modules.h
--- a/Engine/Source/Developer/TraceServices/Public/TraceServices/Model/Modules.h
+++ b/Engine/Source/Developer/TraceServices/Public/TraceServices/Model/Modules.h
@@ -95,4 +95,7 @@ public:
 
 ////////////////////////////////////////////////////////////////////////////////
 
+TRACESERVICES_API FName GetModuleProviderName();
+TRACESERVICES_API const IModuleProvider* ReadModuleProvider(const IAnalysisSession& Session);
+
 } // namespace TraceServices
diff --git a/Engine/Source/Developer/TraceServices/Public/TraceServices/Model/NetProfiler.h b/Engine/Source/Developer/TraceServices/Public/TraceServices/Model/NetProfiler.h
--- a/Engine/Source/Developer/TraceServices/Public/TraceServices/Model/NetProfiler.h
+++ b/Engine/Source/Developer/TraceServices/Public/TraceServices/Model/NetProfiler.h
@@ -220,6 +220,7 @@ public:
 	virtual ITable<FNetProfilerAggregatedStats>* CreateAggregation(uint32 ConnectionIndex, ENetProfilerConnectionMode Mode, uint32 PacketIndexIntervalStart, uint32 PacketIndexIntervalEnd, uint32 StartPosition, uint32 EndPosition) const = 0;
 };
 
-TRACESERVICES_API const INetProfilerProvider& ReadNetProfilerProvider(const IAnalysisSession& Session);
+TRACESERVICES_API FName GetNetProfilerProviderName();
+TRACESERVICES_API const INetProfilerProvider* ReadNetProfilerProvider(const IAnalysisSession& Session);
 
 } // namespace TraceServices
diff --git a/Engine/Source/Developer/TraceServices/Public/TraceServices/ModuleService.h b/Engine/Source/Developer/TraceServices/Public/TraceServices/ModuleService.h
--- a/Engine/Source/Developer/TraceServices/Public/TraceServices/ModuleService.h
+++ b/Engine/Source/Developer/TraceServices/Public/TraceServices/ModuleService.h
@@ -30,12 +30,13 @@ class IModule
 	: public IModularFeature
 {
 public:
+	virtual bool ShouldBeEnabledByDefault() const { return true; }
 	virtual void GetModuleInfo(FModuleInfo& OutModuleInfo) = 0;
 	virtual void OnAnalysisBegin(IAnalysisSession& Session) = 0;
-	virtual void GetLoggers(TArray<const TCHAR*>& OutLoggers) = 0;
-	virtual void GenerateReports(const IAnalysisSession& Session, const TCHAR* CmdLine, const TCHAR* OutputDirectory) = 0;
-	virtual const TCHAR* GetCommandLineArgument() = 0;
-	virtual bool ShouldBeEnabledByDefault() const { return true; }
+	virtual void GenerateReports(const IAnalysisSession& Session, const TCHAR* CmdLine, const TCHAR* OutputDirectory) {}
+
+	virtual void GetLoggers(TArray<const TCHAR*>& OutLoggers) {}
+	virtual const TCHAR* GetCommandLineArgument() { return nullptr; }
 };
 
 class IModuleService
