[Insights]
  - Log view: Added Copy to clipboard for selected log/message, for logs in selected time range and for all logs.
  - Log view: Added Save to file (Tab-Separated Values or Comma-Separated Values) for logs in selected time range and for all logs.

diff --git a/Engine/Source/Developer/TraceInsights/Private/Insights/Widgets/SLogView.cpp b/Engine/Source/Developer/TraceInsights/Private/Insights/Widgets/SLogView.cpp
--- a/Engine/Source/Developer/TraceInsights/Private/Insights/Widgets/SLogView.cpp
+++ b/Engine/Source/Developer/TraceInsights/Private/Insights/Widgets/SLogView.cpp
@@ -4,10 +4,16 @@
 
 #include "Algo/BinarySearch.h"
 #include "Async/AsyncWork.h"
+#include "DesktopPlatformModule.h"
 #include "EditorFontGlyphs.h"
 #include "EditorStyleSet.h"
 #include "Framework/Application/SlateApplication.h"
+#include "Framework/Commands/Commands.h"
+#include "Framework/Commands/UICommandList.h"
 #include "Framework/MultiBox/MultiBoxBuilder.h"
+#include "HAL/PlatformApplicationMisc.h"
+#include "HAL/PlatformFileManager.h"
+#include "Logging/MessageLog.h"
 #include "SlateOptMacros.h"
 #include "Widgets/Images/SImage.h"
 #include "Widgets/Input/SComboButton.h"
@@ -29,6 +35,49 @@
 
 #define LOCTEXT_NAMESPACE "SLogView"
 
+////////////////////////////////////////////////////////////////////////////////////////////////////
+// FTableTreeViewCommands
+////////////////////////////////////////////////////////////////////////////////////////////////////
+
+class FLogViewCommands : public TCommands<FLogViewCommands>
+{
+public:
+	FLogViewCommands()
+	: TCommands<FLogViewCommands>(TEXT("FLogViewCommands"), NSLOCTEXT("FLogViewCommands", "Log View Commands", "Log View Commands"), NAME_None, FEditorStyle::Get().GetStyleSetName())
+	{
+	}
+
+	virtual ~FLogViewCommands()
+	{
+	}
+
+	// UI_COMMAND takes long for the compiler to optimize
+	PRAGMA_DISABLE_OPTIMIZATION
+	virtual void RegisterCommands() override
+	{
+		UI_COMMAND(Command_HideSelectedCategory, "Hide Category", "Hides the selected log category.", EUserInterfaceActionType::Button, FInputChord());
+		UI_COMMAND(Command_ShowOnlySelectedCategory, "Show Only Selected Category", "Shows only the selected log category (hides all other log categories).", EUserInterfaceActionType::Button, FInputChord());
+		UI_COMMAND(Command_ShowAllCategories, "Show All Categories", "Resets the category filter (shows all log categories).", EUserInterfaceActionType::Button, FInputChord());
+		UI_COMMAND(Command_CopySelected, "Copy", "Copies the selected log (with all its properties) to clipboard.", EUserInterfaceActionType::Button, FInputChord(EModifierKey::Control, EKeys::C));
+		UI_COMMAND(Command_CopyMessage, "Copy Message", "Copies the message text of selected log to clipboard.", EUserInterfaceActionType::Button, FInputChord(EModifierKey::Shift, EKeys::C));
+		UI_COMMAND(Command_CopyRange, "Copy Range", "Copies all the logs in the selected time range (highlighted in blue) to clipboard.", EUserInterfaceActionType::Button, FInputChord(EModifierKey::Control | EModifierKey::Shift, EKeys::C));
+		UI_COMMAND(Command_CopyAll, "Copy All", "Copies all the (filtered) logs to clipboard.", EUserInterfaceActionType::Button, FInputChord());
+		UI_COMMAND(Command_SaveRange, "Save Range As...", "Saves all the logs in the selected time range (highlighted in blue) to a text file (tab-separated values or comma-separated values).", EUserInterfaceActionType::Button, FInputChord(EModifierKey::Control, EKeys::S));
+		UI_COMMAND(Command_SaveAll, "Save All As...", "Saves all the (filtered) logs to a text file (tab-separated values or comma-separated values).", EUserInterfaceActionType::Button, FInputChord());
+	}
+	PRAGMA_ENABLE_OPTIMIZATION
+
+	TSharedPtr<FUICommandInfo> Command_HideSelectedCategory;
+	TSharedPtr<FUICommandInfo> Command_ShowOnlySelectedCategory;
+	TSharedPtr<FUICommandInfo> Command_ShowAllCategories;
+	TSharedPtr<FUICommandInfo> Command_CopySelected;
+	TSharedPtr<FUICommandInfo> Command_CopyMessage;
+	TSharedPtr<FUICommandInfo> Command_CopyRange;
+	TSharedPtr<FUICommandInfo> Command_CopyAll;
+	TSharedPtr<FUICommandInfo> Command_SaveRange;
+	TSharedPtr<FUICommandInfo> Command_SaveAll;
+};
+
 ////////////////////////////////////////////////////////////////////////////////////////////////////
 // SLogMessageRow
 ////////////////////////////////////////////////////////////////////////////////////////////////////
@@ -685,6 +734,8 @@ void SLogView::Construct(const FArguments& InArgs)
 		]
 	];
 
+	InitCommandList();
+
 	// Register ourselves with the profiler manager.
 	//TODO: FTimingProfilerManager::Get()->OnRequestLogViewUpdate().AddSP(this, &SLogView::ProfilerManager_OnRequestLogViewUpdate);
 }
@@ -700,6 +751,23 @@ TSharedRef<ITableRow> SLogView::OnGenerateRow(TSharedPtr<FLogMessage> InLogMessa
 
 ////////////////////////////////////////////////////////////////////////////////////////////////////
 
+void SLogView::InitCommandList()
+{
+	FLogViewCommands::Register();
+	CommandList = MakeShared<FUICommandList>();
+	CommandList->MapAction(FLogViewCommands::Get().Command_HideSelectedCategory, FExecuteAction::CreateSP(this, &SLogView::HideSelectedCategory), FCanExecuteAction::CreateSP(this, &SLogView::CanHideSelectedCategory));
+	CommandList->MapAction(FLogViewCommands::Get().Command_ShowOnlySelectedCategory, FExecuteAction::CreateSP(this, &SLogView::ShowOnlySelectedCategory), FCanExecuteAction::CreateSP(this, &SLogView::CanShowOnlySelectedCategory));
+	CommandList->MapAction(FLogViewCommands::Get().Command_ShowAllCategories, FExecuteAction::CreateSP(this, &SLogView::ShowAllCategories), FCanExecuteAction::CreateSP(this, &SLogView::CanShowAllCategories));
+	CommandList->MapAction(FLogViewCommands::Get().Command_CopySelected, FExecuteAction::CreateSP(this, &SLogView::CopySelected), FCanExecuteAction::CreateSP(this, &SLogView::CanCopySelected));
+	CommandList->MapAction(FLogViewCommands::Get().Command_CopyMessage, FExecuteAction::CreateSP(this, &SLogView::CopyMessage), FCanExecuteAction::CreateSP(this, &SLogView::CanCopyMessage));
+	CommandList->MapAction(FLogViewCommands::Get().Command_CopyRange, FExecuteAction::CreateSP(this, &SLogView::CopyRange), FCanExecuteAction::CreateSP(this, &SLogView::CanCopyRange));
+	CommandList->MapAction(FLogViewCommands::Get().Command_CopyAll, FExecuteAction::CreateSP(this, &SLogView::CopyAll), FCanExecuteAction::CreateSP(this, &SLogView::CanCopyAll));
+	CommandList->MapAction(FLogViewCommands::Get().Command_SaveRange, FExecuteAction::CreateSP(this, &SLogView::SaveRange), FCanExecuteAction::CreateSP(this, &SLogView::CanSaveRange));
+	CommandList->MapAction(FLogViewCommands::Get().Command_SaveAll, FExecuteAction::CreateSP(this, &SLogView::SaveAll), FCanExecuteAction::CreateSP(this, &SLogView::CanSaveAll));
+}
+
+////////////////////////////////////////////////////////////////////////////////////////////////////
+
 void SLogView::Tick(const FGeometry& AllottedGeometry, const double InCurrentTime, const float InDeltaTime)
 {
 	int32 NewMessageCount = 0;
@@ -1049,7 +1117,8 @@ TSharedPtr<SWidget> SLogView::ListView_GetContextMenu()
 {
 	TSharedPtr<FLogMessage> SelectedLogMessage = GetSelectedLogMessage();
 
-	FMenuBuilder MenuBuilder(/*bInShouldCloseWindowAfterMenuSelection=*/true, nullptr);
+	const bool bShouldCloseWindowAfterMenuSelection = true;
+	FMenuBuilder MenuBuilder(bShouldCloseWindowAfterMenuSelection, CommandList);
 
 	MenuBuilder.BeginSection("LogViewContextMenu");
 	{
@@ -1059,31 +1128,74 @@ TSharedPtr<SWidget> SLogView::ListView_GetContextMenu()
 			FName CategoryName(Record.GetCategory());
 
 			MenuBuilder.AddMenuEntry(
+				FLogViewCommands::Get().Command_HideSelectedCategory,
+				NAME_None,
 				FText::Format(LOCTEXT("HideCategory", "Hide \"{0}\" Category"), Record.GetCategoryAsText()),
 				FText::Format(LOCTEXT("HideCategory_Tooltip", "Hide the \"{0}\" log category."), Record.GetCategoryAsText()),
-				FSlateIcon(),
-				FUIAction(FExecuteAction::CreateSP(this, &SLogView::ToggleCategory_Execute, CategoryName)),
-				NAME_None,
-				EUserInterfaceActionType::Button
+				FSlateIcon()
 			);
 
 			MenuBuilder.AddMenuEntry(
+				FLogViewCommands::Get().Command_ShowOnlySelectedCategory,
+				NAME_None,
 				FText::Format(LOCTEXT("ShowOnlyCategory", "Show Only \"{0}\" Category"), Record.GetCategoryAsText()),
 				FText::Format(LOCTEXT("ShowOnlyCategory_Tooltip", "Show only the \"{0}\" log category (hide all other log categories)."), Record.GetCategoryAsText()),
-				FSlateIcon(),
-				FUIAction(FExecuteAction::CreateSP(this, &SLogView::ShowOnlyCategory_Execute, CategoryName)),
-				NAME_None,
-				EUserInterfaceActionType::Button
+				FSlateIcon()
 			);
 		}
 
+		MenuBuilder.AddMenuEntry(FLogViewCommands::Get().Command_ShowAllCategories);
+
+		MenuBuilder.AddSeparator();
+
 		MenuBuilder.AddMenuEntry(
-			LOCTEXT("ShowAllCategoriesCtxMenu", "Show All Categories"),
-			LOCTEXT("ShowAllCategoriesCtxMenu_Tooltip", "Reset category filter (show all log categories)."),
-			FSlateIcon(),
-			FUIAction(FExecuteAction::CreateSP(this, &SLogView::ShowAllCategories_Execute)),
+			FLogViewCommands::Get().Command_CopySelected,
 			NAME_None,
-			EUserInterfaceActionType::Button
+			TAttribute<FText>(),
+			TAttribute<FText>(),
+			FSlateIcon(FCoreStyle::Get().GetStyleSetName(), "GenericCommands.Copy")
+		);
+
+		MenuBuilder.AddMenuEntry(
+			FLogViewCommands::Get().Command_CopyMessage,
+			NAME_None,
+			TAttribute<FText>(),
+			TAttribute<FText>(),
+			FSlateIcon(FCoreStyle::Get().GetStyleSetName(), "GenericCommands.Copy")
+		);
+
+		MenuBuilder.AddMenuEntry(
+			FLogViewCommands::Get().Command_CopyRange,
+			NAME_None,
+			TAttribute<FText>(),
+			TAttribute<FText>(),
+			FSlateIcon(FCoreStyle::Get().GetStyleSetName(), "GenericCommands.Copy")
+		);
+
+		MenuBuilder.AddMenuEntry(
+			FLogViewCommands::Get().Command_CopyAll,
+			NAME_None,
+			TAttribute<FText>(),
+			TAttribute<FText>(),
+			FSlateIcon(FCoreStyle::Get().GetStyleSetName(), "GenericCommands.Copy")
+		);
+
+		MenuBuilder.AddSeparator();
+
+		MenuBuilder.AddMenuEntry(
+			FLogViewCommands::Get().Command_SaveRange,
+			NAME_None,
+			TAttribute<FText>(),
+			TAttribute<FText>(),
+			FSlateIcon(FCoreStyle::Get().GetStyleSetName(), "Icons.Save")
+		);
+
+		MenuBuilder.AddMenuEntry(
+			FLogViewCommands::Get().Command_SaveAll,
+			NAME_None,
+			TAttribute<FText>(),
+			TAttribute<FText>(),
+			FSlateIcon(FCoreStyle::Get().GetStyleSetName(), "Icons.Save")
 		);
 	}
 	MenuBuilder.EndSection();
@@ -1207,9 +1319,9 @@ void SLogView::CreateCategoriesFilterMenuSection(FMenuBuilder& MenuBuilder)
 			.ColorAndOpacity(FSlateColor(FTimeMarkerTrackBuilder::GetColorByCategory(*CategoryString)));
 
 		MenuBuilder.AddMenuEntry(
-			FUIAction(FExecuteAction::CreateSP(this, &SLogView::ToggleCategory_Execute, CategoryName),
-				FCanExecuteAction(),
-				FIsActionChecked::CreateSP(this, &SLogView::ToggleCategory_IsChecked, CategoryName)),
+			FUIAction(FExecuteAction::CreateSP(this, &SLogView::ToggleCategory, CategoryName),
+			FCanExecuteAction(),
+			FIsActionChecked::CreateSP(this, &SLogView::IsLogCategoryEnabled, CategoryName)),
 			TextBlock,
 			NAME_None,
 			FText::Format(LOCTEXT("Category_Tooltip", "Filter the Log View to show/hide category: {0}"), CategoryText),
@@ -1266,7 +1378,71 @@ void SLogView::ShowHideAllCategories_Execute()
 
 ////////////////////////////////////////////////////////////////////////////////////////////////////
 
-void SLogView::ShowAllCategories_Execute()
+bool SLogView::IsLogCategoryEnabled(FName InName) const
+{
+	return Filter.IsLogCategoryEnabled(InName);
+}
+
+////////////////////////////////////////////////////////////////////////////////////////////////////
+
+void SLogView::ToggleCategory(FName InName)
+{
+	Filter.ToggleLogCategory(InName);
+	OnFilterChanged();
+}
+
+////////////////////////////////////////////////////////////////////////////////////////////////////
+
+bool SLogView::CanHideSelectedCategory() const
+{
+	return ListView->GetSelectedItems().Num() == 1;
+}
+
+////////////////////////////////////////////////////////////////////////////////////////////////////
+
+void SLogView::HideSelectedCategory()
+{
+	TSharedPtr<FLogMessage> SelectedLogMessage = GetSelectedLogMessage();
+	if (SelectedLogMessage.IsValid())
+	{
+		FLogMessageRecord& Record = Cache.Get(SelectedLogMessage->GetIndex());
+		FName CategoryName(Record.GetCategory());
+		Filter.DisableLogCategory(CategoryName);
+		OnFilterChanged();
+	}
+}
+
+////////////////////////////////////////////////////////////////////////////////////////////////////
+
+bool SLogView::CanShowOnlySelectedCategory() const
+{
+	return ListView->GetSelectedItems().Num() == 1;
+}
+
+////////////////////////////////////////////////////////////////////////////////////////////////////
+
+void SLogView::ShowOnlySelectedCategory()
+{
+	TSharedPtr<FLogMessage> SelectedLogMessage = GetSelectedLogMessage();
+	if (SelectedLogMessage.IsValid())
+	{
+		FLogMessageRecord& Record = Cache.Get(SelectedLogMessage->GetIndex());
+		FName CategoryName(Record.GetCategory());
+		Filter.EnableOnlyCategory(CategoryName);
+		OnFilterChanged();
+	}
+}
+
+////////////////////////////////////////////////////////////////////////////////////////////////////
+
+bool SLogView::CanShowAllCategories() const
+{
+	return true;
+}
+
+////////////////////////////////////////////////////////////////////////////////////////////////////
+
+void SLogView::ShowAllCategories()
 {
 	if (Filter.IsShowAllCategoriesEnabled())
 	{
@@ -1282,25 +1458,386 @@ void SLogView::ShowAllCategories_Execute()
 
 ////////////////////////////////////////////////////////////////////////////////////////////////////
 
-void SLogView::ShowOnlyCategory_Execute(FName InName)
+FReply SLogView::OnKeyDown(const FGeometry& MyGeometry, const FKeyEvent& InKeyEvent)
 {
-	Filter.EnableOnlyCategory(InName);
-	OnFilterChanged();
+	return CommandList->ProcessCommandBindings(InKeyEvent) == true ? FReply::Handled() : FReply::Unhandled();
 }
 
 ////////////////////////////////////////////////////////////////////////////////////////////////////
 
-bool SLogView::ToggleCategory_IsChecked(FName InName) const
+void SLogView::AppendFormatMessageDetailed(const FLogMessageRecord& Log, TStringBuilderBase<TCHAR>& InOutStringBuilder) const
 {
-	return Filter.IsLogCategoryEnabled(InName);
+	InOutStringBuilder.Appendf(TEXT("Index=%d"), Log.GetIndex());
+	InOutStringBuilder.Append(TEXT("\nTime="));
+	InOutStringBuilder.Append(TimeUtils::FormatTimeHMS(Log.GetTime(), TimeUtils::Microsecond));
+	InOutStringBuilder.Append(TEXT("\nVerbosity="));
+	InOutStringBuilder.Append(::ToString(Log.GetVerbosity()));
+	InOutStringBuilder.Append(TEXT("\nCategory="));
+	InOutStringBuilder.Append(Log.GetCategoryAsString());
+	InOutStringBuilder.Append(TEXT("\nMessage="));
+	InOutStringBuilder.Append(Log.GetMessageAsString());
+	InOutStringBuilder.Append(TEXT("\nFile="));
+	InOutStringBuilder.Append(Log.GetFile());
+	InOutStringBuilder.Appendf(TEXT("\nLine=%d\n"), Log.GetLine());
 }
 
 ////////////////////////////////////////////////////////////////////////////////////////////////////
 
-void SLogView::ToggleCategory_Execute(FName InName)
+void SLogView::AppendFormatMessageDelimitedHeader(TStringBuilderBase<TCHAR>& InOutStringBuilder, TCHAR Separator) const
 {
-	Filter.ToggleLogCategory(InName);
-	OnFilterChanged();
+	InOutStringBuilder.Append(TEXT("Index"));
+	InOutStringBuilder.Append(Separator);
+	InOutStringBuilder.Append(TEXT("Time"));
+	InOutStringBuilder.Append(Separator);
+	InOutStringBuilder.Append(TEXT("Verbosity"));
+	InOutStringBuilder.Append(Separator);
+	InOutStringBuilder.Append(TEXT("Category"));
+	InOutStringBuilder.Append(Separator);
+	InOutStringBuilder.Append(TEXT("Message"));
+	InOutStringBuilder.Append(Separator);
+	InOutStringBuilder.Append(TEXT("File"));
+	InOutStringBuilder.Append(Separator);
+	InOutStringBuilder.Append(TEXT("Line"));
+}
+
+////////////////////////////////////////////////////////////////////////////////////////////////////
+
+void SLogView::AppendFormatMessageDelimited(const FLogMessageRecord& Log, TStringBuilderBase<TCHAR>& InOutStringBuilder, TCHAR Separator) const
+{
+	InOutStringBuilder.Appendf(TEXT("%d"), Log.GetIndex());
+	InOutStringBuilder.Append(Separator);
+	InOutStringBuilder.Append(TimeUtils::FormatTimeHMS(Log.GetTime(), TimeUtils::Microsecond));
+	InOutStringBuilder.Append(Separator);
+	InOutStringBuilder.Append(::ToString(Log.GetVerbosity()));
+	InOutStringBuilder.Append(Separator);
+	InOutStringBuilder.Append(Log.GetCategoryAsString());
+	InOutStringBuilder.Append(Separator);
+	FString Message = Log.GetMessageAsString();
+	if (Separator == TEXT('\t'))
+	{
+		Message.ReplaceCharInline(TEXT('\t'), TEXT(' '), ESearchCase::CaseSensitive);
+		InOutStringBuilder.Append(Message);
+	}
+	else if (Separator == TEXT(','))
+	{
+		InOutStringBuilder.Append(TEXT('\"'));
+		FString EscapedMessage = Message.Replace(TEXT("\""), TEXT("\"\""), ESearchCase::CaseSensitive);
+		InOutStringBuilder.Append(EscapedMessage);
+		InOutStringBuilder.Append(TEXT('\"'));
+	}
+	else
+	{
+		InOutStringBuilder.Append(Message);
+	}
+	InOutStringBuilder.Append(Separator);
+	InOutStringBuilder.Append(Log.GetFile());
+	InOutStringBuilder.Append(Separator);
+	InOutStringBuilder.Appendf(TEXT("%d"), Log.GetLine());
+}
+
+////////////////////////////////////////////////////////////////////////////////////////////////////
+
+bool SLogView::CanCopySelected() const
+{
+	return ListView->GetSelectedItems().Num() == 1;
+}
+
+////////////////////////////////////////////////////////////////////////////////////////////////////
+
+void SLogView::CopySelected() const
+{
+	TSharedPtr<FLogMessage> SelectedLogMessage = GetSelectedLogMessage();
+	if (SelectedLogMessage)
+	{
+		const FLogMessageRecord& Log = Cache.Get(SelectedLogMessage->GetIndex());
+		TStringBuilder<1024> StringBuilder;
+		AppendFormatMessageDetailed(Log, StringBuilder);
+		FPlatformApplicationMisc::ClipboardCopy(StringBuilder.ToString());
+	}
+}
+
+////////////////////////////////////////////////////////////////////////////////////////////////////
+
+bool SLogView::CanCopyMessage() const
+{
+	return ListView->GetSelectedItems().Num() == 1;
+}
+
+////////////////////////////////////////////////////////////////////////////////////////////////////
+
+void SLogView::CopyMessage() const
+{
+	TSharedPtr<FLogMessage> SelectedLogMessage = GetSelectedLogMessage();
+	if (SelectedLogMessage)
+	{
+		const FLogMessageRecord& Log = Cache.Get(SelectedLogMessage->GetIndex());
+		FPlatformApplicationMisc::ClipboardCopy(*Log.GetMessageAsString());
+	}
+}
+
+////////////////////////////////////////////////////////////////////////////////////////////////////
+
+bool SLogView::CanCopyRange() const
+{
+	if (Messages.Num() == 0)
+	{
+		return false;
+	}
+
+	TSharedPtr<STimingProfilerWindow> Window = FTimingProfilerManager::Get()->GetProfilerWindow();
+	TSharedPtr<STimingView> TimingView = Window ? Window->GetTimingView() : nullptr;
+	if (!TimingView)
+	{
+		return false;
+	}
+
+	const double SelectionStartTime = TimingView->GetSelectionStartTime();
+	const double SelectionEndTime = TimingView->GetSelectionEndTime();
+	return SelectionStartTime < SelectionEndTime;
+}
+
+////////////////////////////////////////////////////////////////////////////////////////////////////
+
+void SLogView::CopyRange() const
+{
+	if (Messages.Num() == 0)
+	{
+		return;
+	}
+
+	TSharedPtr<STimingProfilerWindow> Window = FTimingProfilerManager::Get()->GetProfilerWindow();
+	TSharedPtr<STimingView> TimingView = Window ? Window->GetTimingView() : nullptr;
+	if (!TimingView)
+	{
+		return;
+	}
+
+	const double SelectionStartTime = TimingView->GetSelectionStartTime();
+	const double SelectionEndTime = TimingView->GetSelectionEndTime();
+	if (SelectionStartTime >= SelectionEndTime)
+	{
+		return;
+	}
+
+	TStringBuilder<1024> StringBuilder;
+	AppendFormatMessageDelimitedHeader(StringBuilder, TEXT('\t'));
+	StringBuilder.Append(TEXT('\n'));
+
+	int32 NumMessagesInSelectedRange = 0;
+
+	TSharedPtr<const TraceServices::IAnalysisSession> Session = FInsightsManager::Get()->GetSession();
+	if (Session.IsValid())
+	{
+		TraceServices::FAnalysisSessionReadScope SessionReadScope(*Session.Get());
+		const TraceServices::ILogProvider& LogProvider = TraceServices::ReadLogProvider(*Session.Get());
+
+		for (const TSharedPtr<FLogMessage>& Message : Messages)
+		{
+			LogProvider.ReadMessage(Message->GetIndex(), [this, SelectionStartTime, SelectionEndTime, &StringBuilder, &NumMessagesInSelectedRange](const TraceServices::FLogMessageInfo& MessageInfo)
+			{
+				if (MessageInfo.Time >= SelectionStartTime && MessageInfo.Time <= SelectionEndTime)
+				{
+					FLogMessageRecord Log(MessageInfo);
+					AppendFormatMessageDelimited(Log, StringBuilder, TEXT('\t'));
+					StringBuilder.Append(TEXT('\n'));
+					++NumMessagesInSelectedRange;
+				}
+			});
+		}
+	}
+
+	if (NumMessagesInSelectedRange > 0)
+	{
+		FPlatformApplicationMisc::ClipboardCopy(StringBuilder.ToString());
+		UE_LOG(TraceInsights, Log, TEXT("Copied %d logs to clipboard."), NumMessagesInSelectedRange);
+	}
+}
+
+////////////////////////////////////////////////////////////////////////////////////////////////////
+
+bool SLogView::CanCopyAll() const
+{
+	return Messages.Num() > 0;
+}
+
+////////////////////////////////////////////////////////////////////////////////////////////////////
+
+void SLogView::CopyAll() const
+{
+	if (Messages.Num() == 0)
+	{
+		return;
+	}
+
+	TStringBuilder<1024> StringBuilder;
+	AppendFormatMessageDelimitedHeader(StringBuilder, TEXT('\t'));
+	StringBuilder.Append(TEXT('\n'));
+
+	TSharedPtr<const TraceServices::IAnalysisSession> Session = FInsightsManager::Get()->GetSession();
+	if (Session.IsValid())
+	{
+		TraceServices::FAnalysisSessionReadScope SessionReadScope(*Session.Get());
+		const TraceServices::ILogProvider& LogProvider = TraceServices::ReadLogProvider(*Session.Get());
+
+		for (const TSharedPtr<FLogMessage>& Message : Messages)
+		{
+			LogProvider.ReadMessage(Message->GetIndex(), [this, &StringBuilder](const TraceServices::FLogMessageInfo& MessageInfo)
+			{
+				FLogMessageRecord Log(MessageInfo);
+				AppendFormatMessageDelimited(Log, StringBuilder, TEXT('\t'));
+				StringBuilder.Append(TEXT('\n'));
+			});
+		}
+	}
+
+	FPlatformApplicationMisc::ClipboardCopy(StringBuilder.ToString());
+	UE_LOG(TraceInsights, Log, TEXT("Copied %d logs to clipboard."), Messages.Num());
+}
+
+////////////////////////////////////////////////////////////////////////////////////////////////////
+
+bool SLogView::CanSaveRange() const
+{
+	return CanCopyRange();
+}
+
+////////////////////////////////////////////////////////////////////////////////////////////////////
+
+void SLogView::SaveRange() const
+{
+	constexpr bool bSaveLogsInSelectedRangeOnly = true;
+	SaveLogsToFile(bSaveLogsInSelectedRangeOnly);
+}
+
+////////////////////////////////////////////////////////////////////////////////////////////////////
+
+bool SLogView::CanSaveAll() const
+{
+	return Messages.Num() > 0;
+}
+
+////////////////////////////////////////////////////////////////////////////////////////////////////
+
+void SLogView::SaveAll() const
+{
+	constexpr bool bSaveLogsInSelectedRangeOnly = false;
+	SaveLogsToFile(bSaveLogsInSelectedRangeOnly);
+}
+
+////////////////////////////////////////////////////////////////////////////////////////////////////
+
+void SLogView::SaveLogsToFile(bool bSaveLogsInSelectedRangeOnly) const
+{
+	if (Messages.Num() == 0)
+	{
+		return;
+	}
+
+	double SelectionStartTime = 0.0;
+	double SelectionEndTime = 0.0;
+	if (bSaveLogsInSelectedRangeOnly)
+	{
+		TSharedPtr<STimingProfilerWindow> Window = FTimingProfilerManager::Get()->GetProfilerWindow();
+		TSharedPtr<STimingView> TimingView = Window ? Window->GetTimingView() : nullptr;
+		if (!TimingView)
+		{
+			return;
+		}
+
+		SelectionStartTime = TimingView->GetSelectionStartTime();
+		SelectionEndTime = TimingView->GetSelectionEndTime();
+		if (SelectionStartTime >= SelectionEndTime)
+		{
+			return;
+		}
+	}
+
+	TArray<FString> SaveFilenames;
+	IDesktopPlatform* DesktopPlatform = FDesktopPlatformModule::Get();
+	bool bDialogResult = false;
+	if (DesktopPlatform)
+	{
+		const FString DefaultBrowsePath = FPaths::ProjectLogDir();
+		bDialogResult = DesktopPlatform->SaveFileDialog(
+			FSlateApplication::Get().FindBestParentWindowHandleForDialogs(nullptr),
+			LOCTEXT("SaveLogsToFileTitle", "Save Logs").ToString(),
+			DefaultBrowsePath,
+			TEXT(""),
+			TEXT("Tab-Separated Values (*.tsv)|*.tsv|Text Files (*.txt)|*.txt|Comma-Separated Values (*.csv)|*.csv|All Files (*.*)|*.*"),
+			EFileDialogFlags::None,
+			SaveFilenames
+		);
+	}
+
+	if (!bDialogResult || SaveFilenames.Num() == 0)
+	{
+		return;
+	}
+
+	FString& Path = SaveFilenames[0];
+	TCHAR Separator = TEXT('\t');
+	if (Path.EndsWith(TEXT(".csv")))
+	{
+		Separator = TEXT(',');
+	}
+
+	IFileHandle* ExportFileHandle = FPlatformFileManager::Get().GetPlatformFile().OpenWrite(*Path);
+	if (ExportFileHandle == nullptr)
+	{
+		FMessageLog ReportMessageLog(FInsightsManager::Get()->GetLogListingName());
+		ReportMessageLog.Error(LOCTEXT("FailedToOpenFile", "Save logs failed. Failed to open file for write."));
+		ReportMessageLog.Notify();
+		return;
+	}
+
+	FStopwatch Stopwatch;
+	Stopwatch.Start();
+
+	UTF16CHAR BOM = UNICODE_BOM;
+	ExportFileHandle->Write((uint8*)&BOM, sizeof(UTF16CHAR));
+
+	// Write header.
+	{
+		TStringBuilder<1024> StringBuilder;
+		AppendFormatMessageDelimitedHeader(StringBuilder, Separator);
+		StringBuilder.Append(TEXT('\n'));
+		FTCHARToUTF16 UTF16String(StringBuilder.ToString(), StringBuilder.Len());
+		ExportFileHandle->Write((const uint8*)UTF16String.Get(), UTF16String.Length() * sizeof(UTF16CHAR));
+	}
+
+	int32 NumMessagesInSelectedRange = 0;
+
+	TSharedPtr<const TraceServices::IAnalysisSession> Session = FInsightsManager::Get()->GetSession();
+	if (Session.IsValid())
+	{
+		TraceServices::FAnalysisSessionReadScope SessionReadScope(*Session.Get());
+		const TraceServices::ILogProvider& LogProvider = TraceServices::ReadLogProvider(*Session.Get());
+
+		for (const TSharedPtr<FLogMessage>& Message : Messages)
+		{
+			LogProvider.ReadMessage(Message->GetIndex(), [this, bSaveLogsInSelectedRangeOnly, SelectionStartTime, SelectionEndTime, Separator, ExportFileHandle, &NumMessagesInSelectedRange](const TraceServices::FLogMessageInfo& MessageInfo)
+			{
+				if (!bSaveLogsInSelectedRangeOnly || (MessageInfo.Time >= SelectionStartTime && MessageInfo.Time <= SelectionEndTime))
+				{
+					TStringBuilder<1024> StringBuilder;
+					FLogMessageRecord Log(MessageInfo);
+					AppendFormatMessageDelimited(Log, StringBuilder, Separator);
+					StringBuilder.Append(TEXT('\n'));
+					FTCHARToUTF16 UTF16String(StringBuilder.ToString(), StringBuilder.Len());
+					ExportFileHandle->Write((const uint8*)UTF16String.Get(), UTF16String.Length() * sizeof(UTF16CHAR));
+					++NumMessagesInSelectedRange;
+				}
+			});
+		}
+	}
+
+	ExportFileHandle->Flush();
+	delete ExportFileHandle;
+	ExportFileHandle = nullptr;
+
+	Stopwatch.Stop();
+	const double TotalTime = Stopwatch.GetAccumulatedTime();
+	UE_LOG(TraceInsights, Log, TEXT("Saved %d logs to file in %.3fs."), NumMessagesInSelectedRange, TotalTime);
 }
 
 ////////////////////////////////////////////////////////////////////////////////////////////////////
diff --git a/Engine/Source/Developer/TraceInsights/Private/Insights/Widgets/SLogView.h b/Engine/Source/Developer/TraceInsights/Private/Insights/Widgets/SLogView.h
--- a/Engine/Source/Developer/TraceInsights/Private/Insights/Widgets/SLogView.h
+++ b/Engine/Source/Developer/TraceInsights/Private/Insights/Widgets/SLogView.h
@@ -66,6 +66,8 @@ public:
 	 */
 	virtual void Tick(const FGeometry& AllottedGeometry, const double InCurrentTime, const float InDeltaTime) override;
 
+	virtual FReply OnKeyDown(const FGeometry& MyGeometry, const FKeyEvent& InKeyEvent) override;
+
 	FLogMessageCache& GetCache() { return Cache; }
 
 	TSharedPtr<FLogMessage> GetSelectedLogMessage() const;
@@ -79,6 +81,8 @@ protected:
 	/** Generate a new list view row. */
 	TSharedRef<ITableRow> OnGenerateRow(TSharedPtr<FLogMessage> InLogMessage, const TSharedRef<STableViewBase>& OwnerTable);
 
+	void InitCommandList();
+
 	void SelectLogMessage(TSharedPtr<FLogMessage> LogMessage);
 	void OnMouseButtonClick(TSharedPtr<FLogMessage> LogMessage);
 	void OnSelectionChanged(TSharedPtr<FLogMessage> LogMessage, ESelectInfo::Type SelectInfo);
@@ -104,14 +108,45 @@ protected:
 	bool ShowHideAllCategories_IsChecked() const;
 	void ShowHideAllCategories_Execute();
 	
-	void ShowAllCategories_Execute();
-	
-	void ShowOnlyCategory_Execute(FName InName);
+	bool IsLogCategoryEnabled(FName InName) const;
+	void ToggleCategory(FName InName);
+
+	bool CanHideSelectedCategory() const;
+	void HideSelectedCategory();
+
+	bool CanShowOnlySelectedCategory() const;
+	void ShowOnlySelectedCategory();
+
+	bool CanShowAllCategories() const;
+	void ShowAllCategories();
+
+	void AppendFormatMessageDetailed(const FLogMessageRecord& Log, TStringBuilderBase<TCHAR>& InOutStringBuilder) const;
+	void AppendFormatMessageDelimitedHeader(TStringBuilderBase<TCHAR>& InOutStringBuilder, TCHAR Separator = TEXT('\t')) const;
+	void AppendFormatMessageDelimited(const FLogMessageRecord& Log, TStringBuilderBase<TCHAR>& InOutStringBuilder, TCHAR Separator = TEXT('\t')) const;
 
-	bool ToggleCategory_IsChecked(FName InName) const;
-	void ToggleCategory_Execute(FName InName);
+	bool CanCopySelected() const;
+	void CopySelected() const;
+
+	bool CanCopyMessage() const;
+	void CopyMessage() const;
+
+	bool CanCopyRange() const;
+	void CopyRange() const;
+
+	bool CanCopyAll() const;
+	void CopyAll() const;
+
+	bool CanSaveRange() const;
+	void SaveRange() const;
+
+	bool CanSaveAll() const;
+	void SaveAll() const;
+
+	void SaveLogsToFile(bool bSaveLogsInSelectedRangeOnly) const;
 
 protected:
+	TSharedPtr<FUICommandList> CommandList;
+
 	/** The list view widget. */
 	TSharedPtr<SListView<TSharedPtr<FLogMessage>>> ListView;
 
@@ -148,7 +183,7 @@ protected:
 	FText StatsText;
 
 	/** Cached log messages. */
-	FLogMessageCache Cache;
+	mutable FLogMessageCache Cache;
 
 	/** List of trace log messages to show in list view (i.e. filtered). */
 	TArray<TSharedPtr<FLogMessage>> Messages; // TODO: this needs virtualisation (an a new SListView)
