[Insights]
  - Added "-NoUI" command line argument to be able to run analysis for a session without the UI to be created (useful when used together with -AutoQuit).
  - Added a log message (emitted when analysis completes) with session duration, analysis duration and analysis speed.

diff --git a/Engine/Source/Developer/TraceInsights/Private/Insights/InsightsManager.cpp b/Engine/Source/Developer/TraceInsights/Private/Insights/InsightsManager.cpp
--- a/Engine/Source/Developer/TraceInsights/Private/Insights/InsightsManager.cpp
+++ b/Engine/Source/Developer/TraceInsights/Private/Insights/InsightsManager.cpp
@@ -17,6 +17,7 @@
 
 // Insights
 #include "Insights/Common/InsightsMenuBuilder.h"
+#include "Insights/Common/TimeUtils.h"
 #include "Insights/InsightsStyle.h"
 #include "Insights/IUnrealInsightsModule.h"
 #include "Insights/LoadingProfiler/LoadingProfilerManager.h"
@@ -401,18 +402,11 @@ void FInsightsManager::UpdateSessionDuration()
 {
 	if (Session.IsValid())
 	{
+		bool bLocalIsAnalysisComplete = false;
 		double LocalSessionDuration = 0.0;
 		{
 			TraceServices::FAnalysisSessionReadScope SessionReadScope(*Session.Get());
-			if (bIsAnalysisComplete != Session->IsAnalysisComplete())
-			{
-				bIsAnalysisComplete = Session->IsAnalysisComplete();
-				if (bIsAnalysisComplete)
-				{
-					SessionAnalysisCompletedEvent.Broadcast();
-				}
-			}
-
+			bLocalIsAnalysisComplete = Session->IsAnalysisComplete();
 			LocalSessionDuration = Session->GetDurationSeconds();
 		}
 
@@ -422,6 +416,22 @@ void FInsightsManager::UpdateSessionDuration()
 			AnalysisStopwatch.Update();
 			AnalysisDuration = AnalysisStopwatch.GetAccumulatedTime();
 			AnalysisSpeedFactor = SessionDuration / AnalysisDuration;
+			if (bIsAnalysisComplete)
+			{
+				UE_LOG(TraceInsights, Warning, TEXT("The session duration was updated (%s) after the analysis has been completed."),
+					*TimeUtils::FormatTimeAuto(GetSessionDuration(), 2));
+			}
+		}
+
+		if (bLocalIsAnalysisComplete && !bIsAnalysisComplete)
+		{
+			bIsAnalysisComplete = true;
+			SessionAnalysisCompletedEvent.Broadcast();
+
+			UE_LOG(TraceInsights, Log, TEXT("Analysis has completed in %s (%.1fX speed; session duration: %s)."),
+				*TimeUtils::FormatTimeAuto(AnalysisDuration, 2),
+				AnalysisSpeedFactor,
+				*TimeUtils::FormatTimeAuto(SessionDuration, 2));
 		}
 	}
 }
diff --git a/Engine/Source/Developer/TraceInsights/Private/Insights/TraceInsightsModule.cpp b/Engine/Source/Developer/TraceInsights/Private/Insights/TraceInsightsModule.cpp
--- a/Engine/Source/Developer/TraceInsights/Private/Insights/TraceInsightsModule.cpp
+++ b/Engine/Source/Developer/TraceInsights/Private/Insights/TraceInsightsModule.cpp
@@ -378,11 +378,12 @@ void FTraceInsightsModule::AddAreaForWidgetReflector(TSharedRef<FTabManager::FLa
 
 void FTraceInsightsModule::ShutdownUserInterface()
 {
-	check(PersistentLayout.IsValid());
-
-	// Save application layout.
-	FLayoutSaveRestore::SaveToConfig(UnrealInsightsLayoutIni, PersistentLayout.ToSharedRef());
-	GConfig->Flush(false, UnrealInsightsLayoutIni);
+	if (PersistentLayout.IsValid())
+	{
+		// Save application layout.
+		FLayoutSaveRestore::SaveToConfig(UnrealInsightsLayoutIni, PersistentLayout.ToSharedRef());
+		GConfig->Flush(false, UnrealInsightsLayoutIni);
+	}
 }
 
 ////////////////////////////////////////////////////////////////////////////////////////////////////
diff --git a/Engine/Source/Programs/UnrealInsights/Private/UserInterfaceCommand.cpp b/Engine/Source/Programs/UnrealInsights/Private/UserInterfaceCommand.cpp
--- a/Engine/Source/Programs/UnrealInsights/Private/UserInterfaceCommand.cpp
+++ b/Engine/Source/Programs/UnrealInsights/Private/UserInterfaceCommand.cpp
@@ -303,6 +303,8 @@ void FUserInterfaceCommand::InitializeSlateApplication(bool bOpenTraceFile, cons
 	//This parameter will cause the application to close when analysis fails to start or completes succesfully
 	const bool bAutoQuit = FParse::Param(FCommandLine::Get(), TEXT("AutoQuit"));
 
+	const bool bNoUI = FParse::Param(FCommandLine::Get(), TEXT("NoUI"));
+
 	const bool bInitializeTesting = FParse::Param(FCommandLine::Get(), TEXT("InsightsTest"));
 
 	if (bUseTraceId)
@@ -317,7 +319,11 @@ void FUserInterfaceCommand::InitializeSlateApplication(bool bOpenTraceFile, cons
 			}
 		}
 
-		TraceInsightsModule.CreateSessionViewer(bAllowDebugTools);
+		if (!bNoUI)
+		{
+			TraceInsightsModule.CreateSessionViewer(bAllowDebugTools);
+		}
+
 		TraceInsightsModule.ConnectToStore(StoreHost, StorePort);
 		TraceInsightsModule.StartAnalysisForTrace(TraceId, bAutoQuit);
 	}
@@ -335,7 +341,11 @@ void FUserInterfaceCommand::InitializeSlateApplication(bool bOpenTraceFile, cons
 				}
 			}
 
-			TraceInsightsModule.CreateSessionViewer(bAllowDebugTools);
+			if (!bNoUI)
+			{
+				TraceInsightsModule.CreateSessionViewer(bAllowDebugTools);
+			}
+
 			TraceInsightsModule.StartAnalysisForTraceFile(TraceFile, bAutoQuit);
 		}
 		else
