Add "Text Files" and "All Files" as options when exporting from a STableTreeView.cpp

diff --git a/Engine/Source/Developer/TraceInsights/Private/Insights/Table/Widgets/STableTreeView.cpp b/Engine/Source/Developer/TraceInsights/Private/Insights/Table/Widgets/STableTreeView.cpp
--- a/Engine/Source/Developer/TraceInsights/Private/Insights/Table/Widgets/STableTreeView.cpp
+++ b/Engine/Source/Developer/TraceInsights/Private/Insights/Table/Widgets/STableTreeView.cpp
@@ -3089,7 +3089,7 @@ void STableTreeView::ContextMenu_ExportToFile_Execute(bool bInExportCollapsed, b
 			LOCTEXT("ExportFileTitle", "Export Table").ToString(),
 			DefaultBrowsePath,
 			TEXT(""),
-			TEXT("Tab-separated values file (*.tsv)|*.tsv|Comma-separated values file (*.csv)|*.csv"),
+			TEXT("Tab-Separated Values (*.tsv)|*.tsv|Text Files (*.txt)|*.txt|Comma-Separated Values (*.csv)|*.csv|All Files (*.*)|*.*"),
 			EFileDialogFlags::None,
 			SaveFilenames
 		);
@@ -3133,7 +3133,7 @@ void STableTreeView::ContextMenu_ExportToFile_Execute(bool bInExportCollapsed, b
 
 		// Note: This is a no-op on platforms that are using a 16-bit TCHAR
 		FTCHARToUTF16 UTF16String(*Text, Text.Len());
-		ExportFileHandle->Write((uint8*)*Text, Text.Len() * sizeof(UTF16CHAR));
+		ExportFileHandle->Write((const uint8*)UTF16String.Get(), UTF16String.Length() * sizeof(UTF16CHAR));
 
 		InNodes.Empty(InNodes.Num());
 	};
