[Insights] Updated icon for CopyToClipboard actions (context menu of several tree views). Also added icons for log category quick filter actions (context menu of the Log View).

diff --git a/Engine/Source/Developer/TraceInsights/Private/Insights/MemoryProfiler/Widgets/SMemTagTreeView.cpp b/Engine/Source/Developer/TraceInsights/Private/Insights/MemoryProfiler/Widgets/SMemTagTreeView.cpp
--- a/Engine/Source/Developer/TraceInsights/Private/Insights/MemoryProfiler/Widgets/SMemTagTreeView.cpp
+++ b/Engine/Source/Developer/TraceInsights/Private/Insights/MemoryProfiler/Widgets/SMemTagTreeView.cpp
@@ -502,7 +502,10 @@ TSharedPtr<SWidget> SMemTagTreeView::TreeView_GetMenuContent()
 		(
 			LOCTEXT("ContextMenu_Header_Misc_CopySelectedToClipboard", "Copy To Clipboard"),
 			LOCTEXT("ContextMenu_Header_Misc_CopySelectedToClipboard_Desc", "Copies selection to clipboard"),
-			FSlateIcon(FEditorStyle::GetStyleSetName(), "Profiler.Misc.CopyToClipboard"), Action_CopySelectedToClipboard, NAME_None, EUserInterfaceActionType::Button
+			FSlateIcon(FCoreStyle::Get().GetStyleSetName(), "GenericCommands.Copy"),
+			Action_CopySelectedToClipboard,
+			NAME_None,
+			EUserInterfaceActionType::Button
 		);
 		*/
 
diff --git a/Engine/Source/Developer/TraceInsights/Private/Insights/NetworkingProfiler/Widgets/SNetStatsView.cpp b/Engine/Source/Developer/TraceInsights/Private/Insights/NetworkingProfiler/Widgets/SNetStatsView.cpp
--- a/Engine/Source/Developer/TraceInsights/Private/Insights/NetworkingProfiler/Widgets/SNetStatsView.cpp
+++ b/Engine/Source/Developer/TraceInsights/Private/Insights/NetworkingProfiler/Widgets/SNetStatsView.cpp
@@ -332,7 +332,10 @@ TSharedPtr<SWidget> SNetStatsView::TreeView_GetMenuContent()
 		(
 			LOCTEXT("ContextMenu_Header_Misc_CopySelectedToClipboard", "Copy To Clipboard"),
 			LOCTEXT("ContextMenu_Header_Misc_CopySelectedToClipboard_Desc", "Copies selection to clipboard"),
-			FSlateIcon(FEditorStyle::GetStyleSetName(), "Profiler.Misc.CopyToClipboard"), Action_CopySelectedToClipboard, NAME_None, EUserInterfaceActionType::Button
+			FSlateIcon(FCoreStyle::Get().GetStyleSetName(), "GenericCommands.Copy"),
+			Action_CopySelectedToClipboard,
+			NAME_None,
+			EUserInterfaceActionType::Button
 		);
 		*/
 
diff --git a/Engine/Source/Developer/TraceInsights/Private/Insights/Table/Widgets/STableTreeView.cpp b/Engine/Source/Developer/TraceInsights/Private/Insights/Table/Widgets/STableTreeView.cpp
--- a/Engine/Source/Developer/TraceInsights/Private/Insights/Table/Widgets/STableTreeView.cpp
+++ b/Engine/Source/Developer/TraceInsights/Private/Insights/Table/Widgets/STableTreeView.cpp
@@ -51,7 +51,7 @@ class FTableTreeViewCommands : public TCommands<FTableTreeViewCommands>
 {
 public:
 	FTableTreeViewCommands()
-		: TCommands<FTableTreeViewCommands>(TEXT("FTableTreeViewCommands"), NSLOCTEXT("FTableTreeViewCommands", "Table Tree View Commands", "Table Tree View Commands"), NAME_None, FEditorStyle::Get().GetStyleSetName())
+	: TCommands<FTableTreeViewCommands>(TEXT("FTableTreeViewCommands"), NSLOCTEXT("FTableTreeViewCommands", "Table Tree View Commands", "Table Tree View Commands"), NAME_None, FEditorStyle::Get().GetStyleSetName())
 	{
 	}
 
@@ -86,6 +86,10 @@ public:
 	TSharedPtr<FUICommandInfo> Command_ExportEntireTreeToFileNoLeafs;
 };
 
+////////////////////////////////////////////////////////////////////////////////////////////////////
+// STableTreeView
+////////////////////////////////////////////////////////////////////////////////////////////////////
+
 const FName STableTreeView::RootNodeName(TEXT("Root"));
 
 ////////////////////////////////////////////////////////////////////////////////////////////////////
@@ -460,7 +464,7 @@ TSharedPtr<SWidget> STableTreeView::TreeView_GetMenuContent()
 			NAME_None,
 			TAttribute<FText>(),
 			TAttribute<FText>(),
-			FSlateIcon(FEditorStyle::GetStyleSetName(), "Profiler.Misc.CopyToClipboard")
+			FSlateIcon(FCoreStyle::Get().GetStyleSetName(), "GenericCommands.Copy")
 		);
 
 		MenuBuilder.AddMenuEntry
@@ -469,7 +473,7 @@ TSharedPtr<SWidget> STableTreeView::TreeView_GetMenuContent()
 			NAME_None,
 			TAttribute<FText>(),
 			TAttribute<FText>(),
-			FSlateIcon(FEditorStyle::GetStyleSetName(), "Profiler.Misc.CopyToClipboard")
+			FSlateIcon(FCoreStyle::Get().GetStyleSetName(), "GenericCommands.Copy")
 		);
 
 		MenuBuilder.AddMenuEntry
@@ -478,7 +482,7 @@ TSharedPtr<SWidget> STableTreeView::TreeView_GetMenuContent()
 			NAME_None,
 			TAttribute<FText>(),
 			TAttribute<FText>(),
-			FSlateIcon(FEditorStyle::GetStyleSetName(), "Profiler.Misc.CopyToClipboard")
+			FSlateIcon(FCoreStyle::Get().GetStyleSetName(), "GenericCommands.Copy")
 		);
 
 		MenuBuilder.AddSubMenu
@@ -496,7 +500,7 @@ TSharedPtr<SWidget> STableTreeView::TreeView_GetMenuContent()
 			LOCTEXT("ContextMenu_Header_Misc_Sort_Desc", "Export to file"),
 			FNewMenuDelegate::CreateSP(this, &STableTreeView::TreeView_BuildExportMenu),
 			false,
-			FSlateIcon(FEditorStyle::GetStyleSetName(), "ProfilerCommand.ProfilerManager_Save")
+			FSlateIcon(FCoreStyle::Get().GetStyleSetName(), "Icons.Save")
 		);
 	}
 	MenuBuilder.EndSection();
@@ -643,7 +647,7 @@ void STableTreeView::TreeView_BuildExportMenu(FMenuBuilder& MenuBuilder)
 		NAME_None,
 		TAttribute<FText>(),
 		TAttribute<FText>(),
-		FSlateIcon(FEditorStyle::GetStyleSetName(), "ProfilerCommand.ProfilerManager_Save")
+		FSlateIcon(FCoreStyle::Get().GetStyleSetName(), "Icons.Save")
 	);
 
 	MenuBuilder.AddMenuEntry
@@ -652,7 +656,7 @@ void STableTreeView::TreeView_BuildExportMenu(FMenuBuilder& MenuBuilder)
 		NAME_None,
 		TAttribute<FText>(),
 		TAttribute<FText>(),
-		FSlateIcon(FEditorStyle::GetStyleSetName(), "ProfilerCommand.ProfilerManager_Save")
+		FSlateIcon(FCoreStyle::Get().GetStyleSetName(), "Icons.Save")
 	);
 
 	MenuBuilder.AddMenuEntry
@@ -661,7 +665,7 @@ void STableTreeView::TreeView_BuildExportMenu(FMenuBuilder& MenuBuilder)
 		NAME_None,
 		TAttribute<FText>(),
 		TAttribute<FText>(),
-		FSlateIcon(FEditorStyle::GetStyleSetName(), "ProfilerCommand.ProfilerManager_Save")
+		FSlateIcon(FCoreStyle::Get().GetStyleSetName(), "Icons.Save")
 	);
 
 	MenuBuilder.EndSection();
diff --git a/Engine/Source/Developer/TraceInsights/Private/Insights/Widgets/SLogView.cpp b/Engine/Source/Developer/TraceInsights/Private/Insights/Widgets/SLogView.cpp
--- a/Engine/Source/Developer/TraceInsights/Private/Insights/Widgets/SLogView.cpp
+++ b/Engine/Source/Developer/TraceInsights/Private/Insights/Widgets/SLogView.cpp
@@ -1132,7 +1132,7 @@ TSharedPtr<SWidget> SLogView::ListView_GetContextMenu()
 				NAME_None,
 				FText::Format(LOCTEXT("HideCategory", "Hide \"{0}\" Category"), Record.GetCategoryAsText()),
 				FText::Format(LOCTEXT("HideCategory_Tooltip", "Hide the \"{0}\" log category."), Record.GetCategoryAsText()),
-				FSlateIcon()
+				FSlateIcon(FCoreStyle::Get().GetStyleSetName(), "Icons.Hidden")
 			);
 
 			MenuBuilder.AddMenuEntry(
@@ -1140,11 +1140,17 @@ TSharedPtr<SWidget> SLogView::ListView_GetContextMenu()
 				NAME_None,
 				FText::Format(LOCTEXT("ShowOnlyCategory", "Show Only \"{0}\" Category"), Record.GetCategoryAsText()),
 				FText::Format(LOCTEXT("ShowOnlyCategory_Tooltip", "Show only the \"{0}\" log category (hide all other log categories)."), Record.GetCategoryAsText()),
-				FSlateIcon()
+				FSlateIcon(FCoreStyle::Get().GetStyleSetName(), "Icons.Visible")
 			);
 		}
 
-		MenuBuilder.AddMenuEntry(FLogViewCommands::Get().Command_ShowAllCategories);
+		MenuBuilder.AddMenuEntry(
+			FLogViewCommands::Get().Command_ShowAllCategories,
+			NAME_None,
+			TAttribute<FText>(),
+			TAttribute<FText>(),
+			FSlateIcon(FCoreStyle::Get().GetStyleSetName(), "Icons.Visible")
+		);
 
 		MenuBuilder.AddSeparator();
 
diff --git a/Engine/Source/Developer/TraceInsights/Private/Insights/Widgets/SStatsView.cpp b/Engine/Source/Developer/TraceInsights/Private/Insights/Widgets/SStatsView.cpp
--- a/Engine/Source/Developer/TraceInsights/Private/Insights/Widgets/SStatsView.cpp
+++ b/Engine/Source/Developer/TraceInsights/Private/Insights/Widgets/SStatsView.cpp
@@ -438,7 +438,7 @@ TSharedPtr<SWidget> SStatsView::TreeView_GetMenuContent()
 			NAME_None,
 			TAttribute<FText>(),
 			TAttribute<FText>(),
-			FSlateIcon(FEditorStyle::GetStyleSetName(), "Profiler.Misc.CopyToClipboard")
+			FSlateIcon(FCoreStyle::Get().GetStyleSetName(), "GenericCommands.Copy")
 		);
 
 		MenuBuilder.AddSubMenu
diff --git a/Engine/Source/Developer/TraceInsights/Private/Insights/Widgets/STimerTreeView.cpp b/Engine/Source/Developer/TraceInsights/Private/Insights/Widgets/STimerTreeView.cpp
--- a/Engine/Source/Developer/TraceInsights/Private/Insights/Widgets/STimerTreeView.cpp
+++ b/Engine/Source/Developer/TraceInsights/Private/Insights/Widgets/STimerTreeView.cpp
@@ -237,7 +237,7 @@ TSharedPtr<SWidget> STimerTreeView::TreeView_GetMenuContent()
 			NAME_None,
 			TAttribute<FText>(),
 			TAttribute<FText>(),
-			FSlateIcon(FEditorStyle::GetStyleSetName(), "Profiler.Misc.CopyToClipboard")
+			FSlateIcon(FCoreStyle::Get().GetStyleSetName(), "GenericCommands.Copy")
 		);
 
 		MenuBuilder.AddSubMenu
@@ -271,7 +271,10 @@ TSharedPtr<SWidget> STimerTreeView::TreeView_GetMenuContent()
 		(
 			LOCTEXT("ContextMenu_Header_Columns_ShowAllColumns", "Show All Columns"),
 			LOCTEXT("ContextMenu_Header_Columns_ShowAllColumns_Desc", "Resets tree view to show all columns"),
-			FSlateIcon(FEditorStyle::GetStyleSetName(), "Profiler.EventGraph.ResetColumn"), Action_ShowAllColumns, NAME_None, EUserInterfaceActionType::Button
+			FSlateIcon(FEditorStyle::GetStyleSetName(), "Profiler.EventGraph.ResetColumn"),
+			Action_ShowAllColumns,
+			NAME_None,
+			EUserInterfaceActionType::Button
 		);
 
 		FUIAction Action_ResetColumns
@@ -283,7 +286,10 @@ TSharedPtr<SWidget> STimerTreeView::TreeView_GetMenuContent()
 		(
 			LOCTEXT("ContextMenu_Header_Columns_ResetColumns", "Reset Columns to Default"),
 			LOCTEXT("ContextMenu_Header_Columns_ResetColumns_Desc", "Resets columns to default"),
-			FSlateIcon(FEditorStyle::GetStyleSetName(), "Profiler.EventGraph.ResetColumn"), Action_ResetColumns, NAME_None, EUserInterfaceActionType::Button
+			FSlateIcon(FEditorStyle::GetStyleSetName(), "Profiler.EventGraph.ResetColumn"),
+			Action_ResetColumns,
+			NAME_None,
+			EUserInterfaceActionType::Button
 		);
 	}
 	MenuBuilder.EndSection();
diff --git a/Engine/Source/Developer/TraceInsights/Private/Insights/Widgets/STimersView.cpp b/Engine/Source/Developer/TraceInsights/Private/Insights/Widgets/STimersView.cpp
--- a/Engine/Source/Developer/TraceInsights/Private/Insights/Widgets/STimersView.cpp
+++ b/Engine/Source/Developer/TraceInsights/Private/Insights/Widgets/STimersView.cpp
@@ -479,7 +479,7 @@ TSharedPtr<SWidget> STimersView::TreeView_GetMenuContent()
 			NAME_None,
 			TAttribute<FText>(),
 			TAttribute<FText>(),
-			FSlateIcon(FEditorStyle::GetStyleSetName(), "Profiler.Misc.CopyToClipboard")
+			FSlateIcon(FCoreStyle::Get().GetStyleSetName(), "GenericCommands.Copy")
 		);
 
 		MenuBuilder.AddSubMenu
