[Insights]
  - Timing view: Fixed the drawing of borders for timing events when scale factor is greater than 100%.
  - Timing view: Fixed the drawing of borders for hovered and selected  timing event (2px).

diff --git a/Engine/Source/Developer/TraceInsights/Private/Insights/ViewModels/TimingViewDrawHelper.cpp b/Engine/Source/Developer/TraceInsights/Private/Insights/ViewModels/TimingViewDrawHelper.cpp
--- a/Engine/Source/Developer/TraceInsights/Private/Insights/ViewModels/TimingViewDrawHelper.cpp
+++ b/Engine/Source/Developer/TraceInsights/Private/Insights/ViewModels/TimingViewDrawHelper.cpp
@@ -18,6 +18,8 @@
 #include "Insights/ViewModels/TimingEventsTrack.h"
 #include "Insights/ViewModels/TimingTrackViewport.h"
 
+#define INSIGHTS_USE_LEGACY_BORDER 0
+
 ////////////////////////////////////////////////////////////////////////////////////////////////////
 // FTimingEventsTrackDrawStateBuilder
 ////////////////////////////////////////////////////////////////////////////////////////////////////
@@ -355,7 +357,11 @@ void FTimingViewDrawHelper::DrawEvents(const FTimingEventsTrackDrawState& DrawSt
 			for (const FTimingEventsTrackDrawState::FBoxPrimitive& Box : DrawState.Borders)
 			{
 				const float Y = TopLaneY + (Layout.EventH + Layout.EventDY) * Box.Depth;
+#if INSIGHTS_USE_LEGACY_BORDER
 				DrawContext.DrawBox(EventBorderLayerId, Box.X, Y, Box.W, EventBorderH, EventBorderBrush, Box.Color);
+#else
+				DrawContext.DrawBox(EventBorderLayerId, Box.X, Y, Box.W, EventBorderH, WhiteBrush, Box.Color);
+#endif
 			}
 			NumDrawBorders += DrawState.Borders.Num();
 		}
@@ -430,7 +436,12 @@ void FTimingViewDrawHelper::DrawFadedEvents(const FTimingEventsTrackDrawState& D
 			for (const FTimingEventsTrackDrawState::FBoxPrimitive& Box : DrawState.Borders)
 			{
 				const float Y = TopLaneY + (Layout.EventH + Layout.EventDY) * Box.Depth;
+#if INSIGHTS_USE_LEGACY_BORDER
 				DrawContext.DrawBox(EventBorderLayerId, Box.X, Y, Box.W, EventBorderH, EventBorderBrush, Box.Color.CopyWithNewOpacity(Opacity));
+#else
+				FSlateRoundedBoxBrush Brush(FLinearColor(0.0f, 0.0f, 0.0f, 0.0f), 0.0f, Box.Color.CopyWithNewOpacity(Opacity), 1.0f);
+				DrawContext.DrawBox(EventBorderLayerId, Box.X, Y, Box.W, EventBorderH, &Brush, FLinearColor(0.0f, 0.0f, 0.0f, 0.0f));
+#endif
 			}
 			NumDrawBorders += DrawState.Borders.Num();
 		}
@@ -788,14 +799,23 @@ void FTimingViewDrawHelper::DrawTimingEventHighlight(double StartTime, double En
 
 	const int32 LayerId = ReservedLayerId + ToInt32(EDrawLayer::EventHighlight);
 
+#if INSIGHTS_USE_LEGACY_BORDER
 	constexpr float BorderOffset = 1.0f;
+#else
+	constexpr float BorderOffset = 2.0f;
+#endif
 
 	if (Mode == EDrawEventMode::Hovered)
 	{
 		const FLinearColor Color(1.0f, 1.0f, 0.0f, 1.0f); // yellow
 
 		// Draw border around the timing event box.
+#if INSIGHTS_USE_LEGACY_BORDER
 		DrawContext.DrawBox(LayerId, EventX1 - BorderOffset, Y - BorderOffset, EventW + 2 * BorderOffset, Layout.EventH + 2 * BorderOffset, HoveredEventBorderBrush, Color);
+#else
+		FSlateRoundedBoxBrush Brush(FLinearColor(0.0f, 0.0f, 0.0f, 0.0f), 2.0f, Color, 2.0f);
+		DrawContext.DrawBox(LayerId, EventX1 - BorderOffset, Y - BorderOffset, EventW + 2 * BorderOffset, Layout.EventH + 2 * BorderOffset, &Brush, FLinearColor(0.0f, 0.0f, 0.0f, 0.0f));
+#endif
 	}
 	else // EDrawEventMode::Selected or EDrawEventMode::SelectedAndHovered
 	{
@@ -807,7 +827,12 @@ void FTimingViewDrawHelper::DrawTimingEventHighlight(double StartTime, double En
 		const FLinearColor Color(S, S, Blue, 1.0f);
 
 		// Draw border around the timing event box.
+#if INSIGHTS_USE_LEGACY_BORDER
 		DrawContext.DrawBox(LayerId, EventX1 - BorderOffset, Y - BorderOffset, EventW + 2 * BorderOffset, Layout.EventH + 2 * BorderOffset, SelectedEventBorderBrush, Color);
+#else
+		FSlateRoundedBoxBrush Brush(FLinearColor(0.0f, 0.0f, 0.0f, 0.0f), 2.0f, Color, 2.0f);
+		DrawContext.DrawBox(LayerId, EventX1 - BorderOffset, Y - BorderOffset, EventW + 2 * BorderOffset, Layout.EventH + 2 * BorderOffset, &Brush, FLinearColor(0.0f, 0.0f, 0.0f, 0.0f));
+#endif
 	}
 }
 
@@ -831,3 +856,5 @@ void FTimingViewDrawHelper::DrawRelations(const TArray<TUniquePtr<ITimingEventRe
 }
 
 ////////////////////////////////////////////////////////////////////////////////////////////////////
+
+#undef INSIGHTS_USE_LEGACY_BORDER
