[Insights] Move Task and Slate Insights to the Timing Insights workspace group

diff --git a/Engine/Plugins/Slate/SlateInsights/Source/SlateInsights/Private/SlateInsightsModule.cpp b/Engine/Plugins/Slate/SlateInsights/Source/SlateInsights/Private/SlateInsightsModule.cpp
--- a/Engine/Plugins/Slate/SlateInsights/Source/SlateInsights/Private/SlateInsightsModule.cpp
+++ b/Engine/Plugins/Slate/SlateInsights/Source/SlateInsights/Private/SlateInsightsModule.cpp
@@ -80,7 +80,7 @@ void FSlateInsightsModule::RegisterTimingProfilerLayoutExtensions(FInsightsMajor
 	MinorTabConfig.TabLabel = LOCTEXT("SlateTabTitle", "Slate Frame View");
 	MinorTabConfig.TabTooltip = LOCTEXT("SlateTabTitleTooltip", "Opens the Slate Frame View tab, allows for diagnostics  of a Slate frame.");
 	MinorTabConfig.TabIcon = FSlateIcon(FSlateInsightsStyle::Get().GetStyleSetName(), "SlateProfiler.Icon.Small");
-	//MinorTabConfig.WorkspaceGroup = InOutExtender.GetMinorTabs();
+	MinorTabConfig.WorkspaceGroup = InOutExtender.GetWorkspaceGroup();
 	MinorTabConfig.OnSpawnTab = FOnSpawnTab::CreateLambda([this](const FSpawnTabArgs& Args)
 		{
 			TSharedRef<SSlateFrameSchematicView> Content = SNew(SSlateFrameSchematicView);
diff --git a/Engine/Source/Developer/TraceInsights/Private/Insights/TaskGraphProfiler/TaskGraphProfilerManager.cpp b/Engine/Source/Developer/TraceInsights/Private/Insights/TaskGraphProfiler/TaskGraphProfilerManager.cpp
--- a/Engine/Source/Developer/TraceInsights/Private/Insights/TaskGraphProfiler/TaskGraphProfilerManager.cpp
+++ b/Engine/Source/Developer/TraceInsights/Private/Insights/TaskGraphProfiler/TaskGraphProfilerManager.cpp
@@ -215,6 +215,7 @@ void FTaskGraphProfilerManager::RegisterTimingProfilerLayoutExtensions(FInsights
 	MinorTabConfig.TabIcon = FSlateIcon(FSlateIcon(FInsightsStyle::GetStyleSetName(), "TimersView.Icon.Small"));
 	MinorTabConfig.OnSpawnTab = FOnSpawnTab::CreateRaw(this, &FTaskGraphProfilerManager::SpawnTab_TaskTableTreeView);
 	MinorTabConfig.CanSpawnTab = FCanSpawnTab::CreateRaw(this, &FTaskGraphProfilerManager::CanSpawnTab_TaskTableTreeView);
+	MinorTabConfig.WorkspaceGroup = InOutExtender.GetWorkspaceGroup();
 
 	InOutExtender.GetLayoutExtender().ExtendLayout(FTimingProfilerTabs::StatsCountersID
 		, ELayoutExtensionPosition::After
diff --git a/Engine/Source/Developer/TraceInsights/Private/Insights/Widgets/STimingProfilerWindow.cpp b/Engine/Source/Developer/TraceInsights/Private/Insights/Widgets/STimingProfilerWindow.cpp
--- a/Engine/Source/Developer/TraceInsights/Private/Insights/Widgets/STimingProfilerWindow.cpp
+++ b/Engine/Source/Developer/TraceInsights/Private/Insights/Widgets/STimingProfilerWindow.cpp
@@ -380,7 +380,7 @@ void STimingProfilerWindow::Construct(const FArguments& InArgs, const TSharedRef
 
 	TSharedRef<FWorkspaceItem> AppMenuGroup = TabManager->AddLocalWorkspaceMenuCategory(LOCTEXT("TimingProfilerMenuGroupName", "Timing Insights"));
 
-	Extension = MakeShared<FInsightsMajorTabExtender>(TabManager);
+	Extension = MakeShared<FInsightsMajorTabExtender>(TabManager, AppMenuGroup);
 
 	TabManager->RegisterTabSpawner(FTimingProfilerTabs::ToolbarID, FOnSpawnTab::CreateRaw(this, &STimingProfilerWindow::SpawnTab_Toolbar))
 		.SetDisplayName(LOCTEXT("DeviceToolbarTabTitle", "Toolbar"))
diff --git a/Engine/Source/Developer/TraceInsights/Public/Insights/IUnrealInsightsModule.h b/Engine/Source/Developer/TraceInsights/Public/Insights/IUnrealInsightsModule.h
--- a/Engine/Source/Developer/TraceInsights/Public/Insights/IUnrealInsightsModule.h
+++ b/Engine/Source/Developer/TraceInsights/Public/Insights/IUnrealInsightsModule.h
@@ -9,6 +9,7 @@
 #include "Framework/MultiBox/MultiBoxExtender.h"
 
 class FExtender;
+class FWorkspaceItem;
 
 namespace UE
 {
@@ -124,13 +125,18 @@ struct TRACEINSIGHTS_API FInsightsMajorTabConfig
 /** Combination of extenders applied to the individual major tabs within Insights */
 struct TRACEINSIGHTS_API FInsightsMajorTabExtender
 {
-	FInsightsMajorTabExtender(TSharedPtr<FTabManager>& InTabManager) : MenuExtender(MakeShared<FExtender>()), TabManager(InTabManager) {}
+	FInsightsMajorTabExtender(TSharedPtr<FTabManager>& InTabManager, TSharedRef<FWorkspaceItem> InWorkspaceGroup)
+		: MenuExtender(MakeShared<FExtender>())
+		, TabManager(InTabManager)
+		, WorkspaceGroup(InWorkspaceGroup)
+	{}
 
 	TSharedPtr<FExtender>& GetMenuExtender() { return MenuExtender; }
 	FLayoutExtender& GetLayoutExtender() { return LayoutExtender; }
 	FInsightsMinorTabConfig& AddMinorTabConfig() { return MinorTabs.AddDefaulted_GetRef(); }
 	TSharedPtr<FTabManager> GetTabManager() const { return TabManager; }
 	const TArray<FInsightsMinorTabConfig>& GetMinorTabs() const { return MinorTabs; }
+	TSharedRef<FWorkspaceItem> GetWorkspaceGroup() { return WorkspaceGroup; }
 
 protected:
 	/** Extender used to add to the menu for this tab */
@@ -144,6 +150,8 @@ protected:
 
 	/** Tab manager for this major tab*/
 	TSharedPtr<FTabManager> TabManager;
+
+	TSharedRef<FWorkspaceItem> WorkspaceGroup;
 };
 
 ////////////////////////////////////////////////////////////////////////////////////////////////////
