[Insights] Timing View: Fixed track names: removed "Track" from "Main Graph Track" and "Tasks Overview Track".

diff --git a/Engine/Source/Developer/TraceInsights/Private/Insights/TaskGraphProfiler/ViewModels/TaskTimingTrack.cpp b/Engine/Source/Developer/TraceInsights/Private/Insights/TaskGraphProfiler/ViewModels/TaskTimingTrack.cpp
--- a/Engine/Source/Developer/TraceInsights/Private/Insights/TaskGraphProfiler/ViewModels/TaskTimingTrack.cpp
+++ b/Engine/Source/Developer/TraceInsights/Private/Insights/TaskGraphProfiler/ViewModels/TaskTimingTrack.cpp
@@ -105,7 +105,7 @@ void FTaskTimingSharedState::Tick(Insights::ITimingViewSession& InSession, const
 	{
 		InitCommandList();
 
-		TaskTrack = MakeShared<FTaskTimingTrack>(*this, TEXT("Task Overview Track"), 0);
+		TaskTrack = MakeShared<FTaskTimingTrack>(*this, TEXT("Task Overview"), 0);
 		TaskTrack->SetVisibilityFlag(true);
 		TaskTrack->SetOrder(FTimingTrackOrder::Task);
 
diff --git a/Engine/Source/Developer/TraceInsights/Private/Insights/Widgets/STimingView.cpp b/Engine/Source/Developer/TraceInsights/Private/Insights/Widgets/STimingView.cpp
--- a/Engine/Source/Developer/TraceInsights/Private/Insights/Widgets/STimingView.cpp
+++ b/Engine/Source/Developer/TraceInsights/Private/Insights/Widgets/STimingView.cpp
@@ -96,7 +96,7 @@ STimingView::STimingView()
 	DefaultTimeMarker->SetName(TEXT(""));
 	DefaultTimeMarker->SetColor(FLinearColor(0.85f, 0.5f, 0.03f, 0.5f));
 
-	GraphTrack->SetName(TEXT("Main Graph Track"));
+	GraphTrack->SetName(TEXT("Main Graph"));
 
 	IModularFeatures::Get().RegisterModularFeature(Insights::TimingViewExtenderFeatureName, FrameSharedState.Get());
 	IModularFeatures::Get().RegisterModularFeature(Insights::TimingViewExtenderFeatureName, ThreadTimingSharedState.Get());
