[Insights] Initial work for quick find wigdet.

diff --git a/Engine/Source/Developer/TraceInsights/Private/Insights/TimingProfilerCommands.cpp b/Engine/Source/Developer/TraceInsights/Private/Insights/TimingProfilerCommands.cpp
--- a/Engine/Source/Developer/TraceInsights/Private/Insights/TimingProfilerCommands.cpp
+++ b/Engine/Source/Developer/TraceInsights/Private/Insights/TimingProfilerCommands.cpp
@@ -53,7 +53,7 @@ FTimingProfilerCommands::FTimingProfilerCommands()
 PRAGMA_DISABLE_OPTIMIZATION
 void FTimingProfilerCommands::RegisterCommands()
 {
-	UI_COMMAND(ToggleFramesTrackVisibility, "Frames", "Toggles the visibility of the Frames track", EUserInterfaceActionType::ToggleButton, FInputChord(EModifierKey::Control, EKeys::F));
+	UI_COMMAND(ToggleFramesTrackVisibility, "Frames", "Toggles the visibility of the Frames track", EUserInterfaceActionType::ToggleButton, FInputChord());
 	UI_COMMAND(ToggleGraphTrackVisibility, "Graph", "Toggles the visibility of the Overview Graph track", EUserInterfaceActionType::ToggleButton, FInputChord(EModifierKey::Control, EKeys::G));
 	UI_COMMAND(ToggleTimingViewVisibility, "Timing", "Toggles the visibility of the main Timing view", EUserInterfaceActionType::ToggleButton, FInputChord(EModifierKey::Control, EKeys::T));
 	UI_COMMAND(ToggleTimersViewVisibility, "Timers", "Toggles the visibility of the Timers view", EUserInterfaceActionType::ToggleButton, FInputChord());
@@ -102,6 +102,12 @@ void FTimingViewCommands::RegisterCommands()
 		EUserInterfaceActionType::ToggleButton,
 		FInputChord(EKeys::G));
 
+	UI_COMMAND(QuickFind,
+		"Quick Find",
+		"Quick find or filter events in the timing view.",
+		EUserInterfaceActionType::Button,
+		FInputChord(EModifierKey::Control, EKeys::F));
+
 	//UI_COMMAND(ShowAllGpuTracks, "GPU Track(s)", "Show/hide the GPU track(s).", EUserInterfaceActionType::ToggleButton, FInputChord(EKeys::Y));
 	//UI_COMMAND(ShowAllCpuTracks, "CPU Thread Tracks", "Show/hide all CPU tracks (and all CPU thread groups).", EUserInterfaceActionType::ToggleButton, FInputChord(EKeys::U));
 }
diff --git a/Engine/Source/Developer/TraceInsights/Private/Insights/TimingProfilerCommands.h b/Engine/Source/Developer/TraceInsights/Private/Insights/TimingProfilerCommands.h
--- a/Engine/Source/Developer/TraceInsights/Private/Insights/TimingProfilerCommands.h
+++ b/Engine/Source/Developer/TraceInsights/Private/Insights/TimingProfilerCommands.h
@@ -75,6 +75,9 @@ public:
 	/** Toggles visibility for Main Graph track. */
 	TSharedPtr<FUICommandInfo> ShowMainGraphTrack;
 
+	/** Opens the quick find widget */
+	TSharedPtr<FUICommandInfo> QuickFind;
+
 	/** Toggles visibility for GPU thread track. */
 	//TSharedPtr<FUICommandInfo> ShowAllGpuTracks;
 
diff --git a/Engine/Source/Developer/TraceInsights/Private/Insights/ViewModels/FilterConfigurator.h b/Engine/Source/Developer/TraceInsights/Private/Insights/ViewModels/FilterConfigurator.h
--- a/Engine/Source/Developer/TraceInsights/Private/Insights/ViewModels/FilterConfigurator.h
+++ b/Engine/Source/Developer/TraceInsights/Private/Insights/ViewModels/FilterConfigurator.h
@@ -3,6 +3,7 @@
 #pragma once
 
 #include "CoreMinimal.h"
+#include "Insights/ViewModels/IFilterExecutor.h"
 #include "Insights/ViewModels/Filters.h"
 #include "Insights/ViewModels/FilterConfiguratorNode.h"
 
@@ -11,7 +12,7 @@ namespace Insights
  
 ////////////////////////////////////////////////////////////////////////////////////////////////////
 
-class FFilterConfigurator
+class FFilterConfigurator : public IFilterExecutor
 {
 public:
 	FFilterConfigurator();
@@ -23,11 +24,11 @@ public:
 
 	bool operator!=(const FFilterConfigurator& Other) { return !(*this == Other); }
 
-	~FFilterConfigurator();
+	virtual ~FFilterConfigurator();
 
 	FFilterConfiguratorNodePtr GetRootNode() { return RootNode; }
 
-	bool ApplyFilters(const FFilterContext& Context) const;
+	virtual bool ApplyFilters(const FFilterContext& Context) const override;
 
 	TSharedPtr<TArray<TSharedPtr<struct FFilter>>>& GetAvailableFilters() { return AvailableFilters; }
 
diff --git a/Engine/Source/Developer/TraceInsights/Private/Insights/ViewModels/Filters.cpp b/Engine/Source/Developer/TraceInsights/Private/Insights/ViewModels/Filters.cpp
--- a/Engine/Source/Developer/TraceInsights/Private/Insights/ViewModels/Filters.cpp
+++ b/Engine/Source/Developer/TraceInsights/Private/Insights/ViewModels/Filters.cpp
@@ -4,7 +4,7 @@
 
 #include "Widgets/Docking/SDockTab.h"
 
-#include "Insights/Widgets/SFilterConfigurator.h"
+#include "Insights/Widgets/SAdvancedFilter.h"
 #include "Insights/ViewModels/FilterConfigurator.h"
 
 #define LOCTEXT_NAMESPACE "SFilterService"
@@ -36,7 +36,7 @@ FFilterService::~FFilterService()
 
 TSharedPtr<SWidget> FFilterService::CreateFilterConfiguratorWidget(TSharedPtr<FFilterConfigurator> FilterConfiguratorViewModel)
 {
-	SAssignNew(PendingWidget, SFilterConfigurator, FilterConfiguratorViewModel);
+	SAssignNew(PendingWidget, SAdvancedFilter, FilterConfiguratorViewModel);
 
 	if (FGlobalTabmanager::Get()->HasTabSpawner(FilterConfiguratorTabId))
 	{
diff --git a/Engine/Source/Developer/TraceInsights/Private/Insights/ViewModels/Filters.h b/Engine/Source/Developer/TraceInsights/Private/Insights/ViewModels/Filters.h
--- a/Engine/Source/Developer/TraceInsights/Private/Insights/ViewModels/Filters.h
+++ b/Engine/Source/Developer/TraceInsights/Private/Insights/ViewModels/Filters.h
@@ -201,7 +201,7 @@ private:
 	static TSharedPtr<FFilterService> Instance;
 	FFilterStorage FilterStorage;
 
-	TSharedPtr<class SFilterConfigurator> PendingWidget;
+	TSharedPtr<class SAdvancedFilter> PendingWidget;
 };
 
 ////////////////////////////////////////////////////////////////////////////////////////////////////
diff --git a/Engine/Source/Developer/TraceInsights/Private/Insights/ViewModels/QuickFind.cpp b/Engine/Source/Developer/TraceInsights/Private/Insights/ViewModels/QuickFind.cpp
new file mode 100644
--- /dev/null
+++ b/Engine/Source/Developer/TraceInsights/Private/Insights/ViewModels/QuickFind.cpp
@@ -0,0 +1,24 @@
+// Copyright Epic Games, Inc. All Rights Reserved.
+
+#include "QuickFind.h"
+
+namespace Insights
+{
+
+////////////////////////////////////////////////////////////////////////////////////////////////////
+
+FQuickFind::FQuickFind(TSharedPtr<FFilterConfigurator> InFilterConfiguratorViewModel)
+{
+	FilterConfigurator = InFilterConfiguratorViewModel;
+}
+
+////////////////////////////////////////////////////////////////////////////////////////////////////
+
+FQuickFind::~FQuickFind()
+{
+	OnDestroyedEvent.Broadcast();
+}
+
+////////////////////////////////////////////////////////////////////////////////////////////////////
+
+} // namespace Insights
\ No newline at end of file
diff --git a/Engine/Source/Developer/TraceInsights/Private/Insights/ViewModels/QuickFind.h b/Engine/Source/Developer/TraceInsights/Private/Insights/ViewModels/QuickFind.h
new file mode 100644
--- /dev/null
+++ b/Engine/Source/Developer/TraceInsights/Private/Insights/ViewModels/QuickFind.h
@@ -0,0 +1,51 @@
+// Copyright Epic Games, Inc. All Rights Reserved.
+
+#pragma once
+
+#include "CoreMinimal.h"
+#include "Insights/ViewModels/Filters.h"
+#include "Insights/ViewModels/FilterConfigurator.h"
+
+namespace Insights
+{
+ 
+////////////////////////////////////////////////////////////////////////////////////////////////////
+
+class FQuickFind
+{
+public:
+	FQuickFind(TSharedPtr<FFilterConfigurator> InFilterConfiguratorViewModel);
+
+	~FQuickFind();
+
+	TSharedPtr<FFilterConfigurator> GetFilterConfigurator() { return FilterConfigurator; }
+
+	////////////////////////////////////////////////////////////////////////////////////////////////////
+	// OnDestroyedEvent
+public:
+	/** The event to execute when an instance is destroyed. */
+	DECLARE_MULTICAST_DELEGATE(FOnDestroyedEvent);
+	FOnDestroyedEvent& GetOnDestroyedEvent() { return OnDestroyedEvent; }
+
+private:
+	FOnDestroyedEvent OnDestroyedEvent;
+
+	////////////////////////////////////////////////////////////////////////////////////////////////////
+	// OnFindNextEvent
+public:
+	/** The event to execute when the user selects the "Find Next" option. */
+	DECLARE_MULTICAST_DELEGATE(FOnFindNextEvent);
+	FOnFindNextEvent& GetOnFindNextEvent() { return OnFindNextEvent; }
+
+private:
+	FOnFindNextEvent OnFindNextEvent;
+
+	////////////////////////////////////////////////////////////////////////////////////////////////////
+
+private:
+	TSharedPtr<FFilterConfigurator> FilterConfigurator;
+};
+
+////////////////////////////////////////////////////////////////////////////////////////////////////
+
+} // namespace Insights
\ No newline at end of file
diff --git a/Engine/Source/Developer/TraceInsights/Private/Insights/ViewModels/ThreadTimingTrack.cpp b/Engine/Source/Developer/TraceInsights/Private/Insights/ViewModels/ThreadTimingTrack.cpp
--- a/Engine/Source/Developer/TraceInsights/Private/Insights/ViewModels/ThreadTimingTrack.cpp
+++ b/Engine/Source/Developer/TraceInsights/Private/Insights/ViewModels/ThreadTimingTrack.cpp
@@ -1448,6 +1448,12 @@ bool FThreadTimingTrack::FindTimingProfilerEvent(const FThreadTrackEvent& InTimi
 
 bool FThreadTimingTrack::FindTimingProfilerEvent(const FTimingEventSearchParameters& InParameters, TFunctionRef<void(double, double, uint32, const TraceServices::FTimingProfilerEvent&)> InFoundPredicate) const
 {
+	FFilterContext FilterConfiguratorContext;
+	FilterConfiguratorContext.AddFilterData<double>(static_cast<int32>(EFilterField::StartTime), 0.0f);
+	FilterConfiguratorContext.AddFilterData<double>(static_cast<int32>(EFilterField::EndTime), 0.0f);
+	FilterConfiguratorContext.AddFilterData<double>(static_cast<int32>(EFilterField::Duration), 0.0f);
+	FilterConfiguratorContext.AddFilterData<int64>(static_cast<int32>(EFilterField::EventType), 0);
+
 	return TTimingEventSearch<TraceServices::FTimingProfilerEvent>::Search(
 		InParameters,
 
@@ -1474,12 +1480,47 @@ bool FThreadTimingTrack::FindTimingProfilerEvent(const FTimingEventSearchParamet
 			}
 		},
 
+		[&FilterConfiguratorContext, &InParameters](double EventStartTime, double EventEndTime, uint32 EventDepth, const TraceServices::FTimingProfilerEvent& Event)
+		{
+			if (!InParameters.FilterExecutor.IsValid())
+			{
+				return true;
+			}
+
+			TSharedPtr<const TraceServices::IAnalysisSession> Session = FInsightsManager::Get()->GetSession();
+			if (Session.IsValid())
+			{
+				TraceServices::FAnalysisSessionReadScope SessionReadScope(*Session.Get());
+
+				if (TraceServices::ReadTimingProfilerProvider(*Session.Get()))
+				{
+					const TraceServices::ITimingProfilerProvider& TimingProfilerProvider = *TraceServices::ReadTimingProfilerProvider(*Session.Get());
+					const TraceServices::ITimingProfilerTimerReader* TimerReader;
+					TimingProfilerProvider.ReadTimers([&TimerReader](const TraceServices::ITimingProfilerTimerReader& Out) { TimerReader = &Out; });
+
+					const TraceServices::FTimingProfilerTimer* Timer = TimerReader->GetTimer(Event.TimerIndex);
+					if (ensure(Timer != nullptr))
+					{
+						FilterConfiguratorContext.SetFilterData<double>(static_cast<int32>(EFilterField::StartTime), EventStartTime);
+						FilterConfiguratorContext.SetFilterData<double>(static_cast<int32>(EFilterField::EndTime), EventEndTime);
+						FilterConfiguratorContext.SetFilterData<double>(static_cast<int32>(EFilterField::Duration), EventEndTime - EventStartTime);
+						FilterConfiguratorContext.SetFilterData<int64>(static_cast<int32>(EFilterField::EventType), Timer->Id);
+					}
+					return InParameters.FilterExecutor->ApplyFilters(FilterConfiguratorContext);
+				}
+			}
+
+			return false;
+		},
+
 		[&InFoundPredicate](double InFoundStartTime, double InFoundEndTime, uint32 InFoundDepth, const TraceServices::FTimingProfilerEvent& InEvent)
 		{
 			InFoundPredicate(InFoundStartTime, InFoundEndTime, InFoundDepth, InEvent);
 		},
 
-		SearchCache);
+		TTimingEventSearch<TraceServices::FTimingProfilerEvent>::NoMatch,
+
+		&SearchCache);
 }
 
 ////////////////////////////////////////////////////////////////////////////////////////////////////
diff --git a/Engine/Source/Developer/TraceInsights/Private/Insights/Widgets/SAdvancedFilter.cpp b/Engine/Source/Developer/TraceInsights/Private/Insights/Widgets/SAdvancedFilter.cpp
new file mode 100644
--- /dev/null
+++ b/Engine/Source/Developer/TraceInsights/Private/Insights/Widgets/SAdvancedFilter.cpp
@@ -0,0 +1,172 @@
+// Copyright Epic Games, Inc. All Rights Reserved.
+
+#include "SAdvancedFilter.h"
+
+#include "EditorStyleSet.h"
+#include "SlateOptMacros.h"
+#include "Widgets/DeclarativeSyntaxSupport.h"
+#include "Widgets/Docking/SDockTab.h"
+#include "Widgets/Input/SButton.h"
+#include "Widgets/Layout/SScrollBox.h"
+#include "Widgets/Images/SImage.h"
+#include "Widgets/SBoxPanel.h"
+#include "Widgets/SOverlay.h"
+#include "Widgets/Text/STextBlock.h"
+#include "Widgets/Views/STableViewBase.h"
+
+// Insights
+#include "Insights/Widgets/SFilterConfigurator.h"
+#include "Insights/ViewModels/FilterConfigurator.h"
+
+#define LOCTEXT_NAMESPACE "SAdvancedFilter"
+
+namespace Insights
+{
+
+////////////////////////////////////////////////////////////////////////////////////////////////////
+// SAdvancedFilter
+////////////////////////////////////////////////////////////////////////////////////////////////////
+
+SAdvancedFilter::SAdvancedFilter()
+{
+}
+
+////////////////////////////////////////////////////////////////////////////////////////////////////
+
+SAdvancedFilter::~SAdvancedFilter()
+{
+}
+
+////////////////////////////////////////////////////////////////////////////////////////////////////
+
+void SAdvancedFilter::InitCommandList()
+{
+}
+
+////////////////////////////////////////////////////////////////////////////////////////////////////
+
+BEGIN_SLATE_FUNCTION_BUILD_OPTIMIZATION
+
+void SAdvancedFilter::Construct(const FArguments& InArgs, TSharedPtr<FFilterConfigurator> InFilterConfiguratorViewModel)
+{
+	FilterConfiguratorViewModel = MakeShared<FFilterConfigurator>(*InFilterConfiguratorViewModel);
+
+	ChildSlot
+	[
+		SNew(SVerticalBox)
+
+		// Tree view
+		+ SVerticalBox::Slot()
+		.FillHeight(1.0f)
+		.Padding(0.0f, 6.0f, 0.0f, 0.0f)
+		[
+			SAssignNew(FilterConfigurator, SFilterConfigurator, FilterConfiguratorViewModel)
+		]
+
+		+ SVerticalBox::Slot()
+		.AutoHeight()
+		.Padding(0.0f, 5.0f, 0.0f, 5.0f)
+		.HAlign(EHorizontalAlignment::HAlign_Right)
+		.VAlign(EVerticalAlignment::VAlign_Bottom)
+		[
+			SNew(SHorizontalBox)
+
+			+ SHorizontalBox::Slot()
+			.HAlign(HAlign_Left)
+			.VAlign(VAlign_Center)
+			.AutoWidth()
+			.Padding(1.0f)
+			[
+				SNew(SButton)
+				.Text(LOCTEXT("OK", "OK"))
+				.ToolTipText(LOCTEXT("OKDESC", "Apply the changes to the filters."))
+				.OnClicked(this, &SAdvancedFilter::OK_OnClicked)
+			]
+
+			+ SHorizontalBox::Slot()
+			.HAlign(HAlign_Left)
+			.VAlign(VAlign_Center)
+			.AutoWidth()
+			.Padding(1.0f)
+			[
+				SNew(SButton)
+				.Text(LOCTEXT("Cancel", "Cancel"))
+				.ToolTipText(LOCTEXT("CancelDesc", "Cancel the changes to the filter."))
+				.OnClicked(this, &SAdvancedFilter::Cancel_OnClicked)
+			]
+		]
+	];
+
+	SHeaderRow::FColumn::FArguments ColumnArgs;
+	ColumnArgs
+		.ColumnId(TEXT("Filter"))
+		.DefaultLabel(LOCTEXT("FilterColumnHeader", "Filter"))
+		.HAlignHeader(HAlign_Fill)
+		.VAlignHeader(VAlign_Fill)
+		.HeaderContentPadding(FMargin(2.0f))
+		.HAlignCell(HAlign_Fill)
+		.VAlignCell(VAlign_Fill)
+		.HeaderContent()
+		[
+			SNew(SBox)
+			.HAlign(HAlign_Fill)
+			.VAlign(VAlign_Center)
+			[
+				SNew(STextBlock)
+				.Text(LOCTEXT("FilterColumnHeader", "Filter"))
+			]
+		];
+
+	OnViewModelDestroyedHandle = InFilterConfiguratorViewModel->GetOnDestroyedEvent().AddSP(this, &SAdvancedFilter::RequestClose);
+	OriginalFilterConfiguratorViewModel = InFilterConfiguratorViewModel;
+}
+
+END_SLATE_FUNCTION_BUILD_OPTIMIZATION
+
+////////////////////////////////////////////////////////////////////////////////////////////////////
+
+FReply SAdvancedFilter::OK_OnClicked()
+{
+	TSharedPtr<FFilterConfigurator> OriginalFilterVM = OriginalFilterConfiguratorViewModel.Pin();
+	if (OriginalFilterVM.IsValid())
+	{
+		*OriginalFilterVM = *FilterConfiguratorViewModel;
+	}
+
+	RequestClose();
+
+	return FReply::Handled();
+}
+
+////////////////////////////////////////////////////////////////////////////////////////////////////
+
+FReply SAdvancedFilter::Cancel_OnClicked()
+{
+	RequestClose();
+
+	return FReply::Handled();
+}
+
+////////////////////////////////////////////////////////////////////////////////////////////////////
+
+void SAdvancedFilter::RequestClose()
+{
+	TSharedPtr<SDockTab> ParentTabSharedPtr = ParentTab.Pin();
+	if (ParentTabSharedPtr.IsValid())
+	{
+		ParentTabSharedPtr->RequestCloseTab();
+		ParentTabSharedPtr.Reset();
+	}
+
+	TSharedPtr<FFilterConfigurator> OriginalFilterVM = OriginalFilterConfiguratorViewModel.Pin();
+	if (OriginalFilterVM.IsValid())
+	{
+		OriginalFilterVM->GetOnDestroyedEvent().Remove(OnViewModelDestroyedHandle);
+	}
+}
+
+////////////////////////////////////////////////////////////////////////////////////////////////////
+
+} // namespace Insights
+
+#undef LOCTEXT_NAMESPACE
diff --git a/Engine/Source/Developer/TraceInsights/Private/Insights/Widgets/SAdvancedFilter.h b/Engine/Source/Developer/TraceInsights/Private/Insights/Widgets/SAdvancedFilter.h
new file mode 100644
--- /dev/null
+++ b/Engine/Source/Developer/TraceInsights/Private/Insights/Widgets/SAdvancedFilter.h
@@ -0,0 +1,71 @@
+// Copyright Epic Games, Inc. All Rights Reserved.
+
+#pragma once
+
+#include "CoreMinimal.h"
+#include "Widgets/SCompoundWidget.h"
+#include "Widgets/SWidget.h"
+#include "Widgets/Views/STreeView.h"
+
+// Insights
+#include "Insights/ViewModels/Filters.h"
+#include "Insights/ViewModels/FilterConfiguratorNode.h"
+
+class SDockTab;
+
+namespace Insights
+{
+
+class SFilterConfigurator;
+////////////////////////////////////////////////////////////////////////////////////////////////////
+
+/**
+ * A custom widget used to configure custom filters.
+ */
+class SAdvancedFilter: public SCompoundWidget
+{
+public:
+	/** Default constructor. */
+	SAdvancedFilter();
+
+	/** Virtual destructor. */
+	virtual ~SAdvancedFilter();
+
+	SLATE_BEGIN_ARGS(SAdvancedFilter) {}
+	SLATE_END_ARGS()
+
+	/**
+	 * Construct this widget
+	 * @param InArgs - The declaration data for this widget
+	 */
+	void Construct(const FArguments& InArgs, TSharedPtr<class FFilterConfigurator> InFilterConfiguratorViewModel);
+
+	void Reset();
+
+	void SetParentTab(const TSharedPtr<SDockTab> InTab) { ParentTab = InTab; }
+	const TWeakPtr<SDockTab> GetParentTab() { return ParentTab; };
+
+	void RequestClose();
+
+private:
+	void InitCommandList();
+
+	FReply OK_OnClicked();
+
+	FReply Cancel_OnClicked();
+
+private:
+	TSharedPtr<SFilterConfigurator> FilterConfigurator;
+
+	TWeakPtr<class FFilterConfigurator> OriginalFilterConfiguratorViewModel;
+
+	TSharedPtr<class FFilterConfigurator> FilterConfiguratorViewModel;
+
+	TWeakPtr<SDockTab> ParentTab;
+
+	FDelegateHandle OnViewModelDestroyedHandle;
+};
+
+////////////////////////////////////////////////////////////////////////////////////////////////////
+
+} // namespace Insights
diff --git a/Engine/Source/Developer/TraceInsights/Private/Insights/Widgets/SFilterConfigurator.cpp b/Engine/Source/Developer/TraceInsights/Private/Insights/Widgets/SFilterConfigurator.cpp
--- a/Engine/Source/Developer/TraceInsights/Private/Insights/Widgets/SFilterConfigurator.cpp
+++ b/Engine/Source/Developer/TraceInsights/Private/Insights/Widgets/SFilterConfigurator.cpp
@@ -105,38 +105,6 @@ void SFilterConfigurator::Construct(const FArguments& InArgs, TSharedPtr<FFilter
 				]
 			]
 		]
-			+ SVerticalBox::Slot()
-			.AutoHeight()
-			.Padding(0.0f, 5.0f, 0.0f, 5.0f)
-			.HAlign(EHorizontalAlignment::HAlign_Right)
-			.VAlign(EVerticalAlignment::VAlign_Bottom)
-			[
-				SNew(SHorizontalBox)
-
-				+ SHorizontalBox::Slot()
-				.HAlign(HAlign_Left)
-				.VAlign(VAlign_Center)
-				.AutoWidth()
-				.Padding(1.0f)
-				[
-					SNew(SButton)
-					.Text(LOCTEXT("OK", "OK"))
-					.ToolTipText(LOCTEXT("OKDESC", "Apply the changes to the filters."))
-					.OnClicked(this, &SFilterConfigurator::OK_OnClicked)
-				]
-
-				+ SHorizontalBox::Slot()
-				.HAlign(HAlign_Left)
-				.VAlign(VAlign_Center)
-				.AutoWidth()
-				.Padding(1.0f)
-				[
-					SNew(SButton)
-					.Text(LOCTEXT("Cancel", "Cancel"))
-					.ToolTipText(LOCTEXT("CancelDesc", "Cancel the changes to the filter."))
-					.OnClicked(this, &SFilterConfigurator::Cancel_OnClicked)
-				]
-			]
 	];
 
 	SHeaderRow::FColumn::FArguments ColumnArgs;
@@ -161,9 +129,7 @@ void SFilterConfigurator::Construct(const FArguments& InArgs, TSharedPtr<FFilter
 
 	TreeViewHeaderRow->InsertColumn(ColumnArgs, 0);
 
-	OnViewModelDestroyedHandle = InFilterConfiguratorViewModel->GetOnDestroyedEvent().AddSP(this, &SFilterConfigurator::RequestClose);
-	OriginalFilterConfiguratorViewModel = InFilterConfiguratorViewModel;
-	FilterConfiguratorViewModel = MakeShared<FFilterConfigurator>(*InFilterConfiguratorViewModel);
+	FilterConfiguratorViewModel = InFilterConfiguratorViewModel;
 
 	GroupNodes.Add(FilterConfiguratorViewModel->GetRootNode());
 }
@@ -195,48 +161,6 @@ TSharedRef<ITableRow> SFilterConfigurator::TreeView_OnGenerateRow(FFilterConfigu
 
 ////////////////////////////////////////////////////////////////////////////////////////////////////
 
-FReply SFilterConfigurator::OK_OnClicked()
-{
-	TSharedPtr<FFilterConfigurator> OriginalFilterVM = OriginalFilterConfiguratorViewModel.Pin();
-	if (OriginalFilterVM.IsValid())
-	{
-		*OriginalFilterVM = *FilterConfiguratorViewModel;
-	}
-
-	RequestClose();
-
-	return FReply::Handled();
-}
-
-////////////////////////////////////////////////////////////////////////////////////////////////////
-
-FReply SFilterConfigurator::Cancel_OnClicked()
-{
-	RequestClose();
-
-	return FReply::Handled();
-}
-
-////////////////////////////////////////////////////////////////////////////////////////////////////
-
-void SFilterConfigurator::RequestClose()
-{
-	TSharedPtr<SDockTab> ParentTabSharedPtr = ParentTab.Pin();
-	if (ParentTabSharedPtr.IsValid())
-	{
-		ParentTabSharedPtr->RequestCloseTab();
-		ParentTabSharedPtr.Reset();
-	}
-
-	TSharedPtr<FFilterConfigurator> OriginalFilterVM = OriginalFilterConfiguratorViewModel.Pin();
-	if (OriginalFilterVM.IsValid())
-	{
-		OriginalFilterVM->GetOnDestroyedEvent().Remove(OnViewModelDestroyedHandle);
-	}
-}
-
-////////////////////////////////////////////////////////////////////////////////////////////////////
-
 } // namespace Insights
 
 #undef LOCTEXT_NAMESPACE
diff --git a/Engine/Source/Developer/TraceInsights/Private/Insights/Widgets/SFilterConfigurator.h b/Engine/Source/Developer/TraceInsights/Private/Insights/Widgets/SFilterConfigurator.h
--- a/Engine/Source/Developer/TraceInsights/Private/Insights/Widgets/SFilterConfigurator.h
+++ b/Engine/Source/Developer/TraceInsights/Private/Insights/Widgets/SFilterConfigurator.h
@@ -59,14 +59,8 @@ public:
 private:
 	void InitCommandList();
 
-	FReply OK_OnClicked();
-
-	FReply Cancel_OnClicked();
-
 private:
 
-	TWeakPtr<class FFilterConfigurator> OriginalFilterConfiguratorViewModel;
-
 	TSharedPtr<class FFilterConfigurator> FilterConfiguratorViewModel;
 
 	/** The tree widget which holds the filters. */
@@ -82,8 +76,6 @@ private:
 	TArray<FFilterConfiguratorNodePtr> GroupNodes;
 
 	TWeakPtr<SDockTab> ParentTab;
-
-	FDelegateHandle OnViewModelDestroyedHandle;
 };
 
 ////////////////////////////////////////////////////////////////////////////////////////////////////
diff --git a/Engine/Source/Developer/TraceInsights/Private/Insights/Widgets/SQuickFind.cpp b/Engine/Source/Developer/TraceInsights/Private/Insights/Widgets/SQuickFind.cpp
new file mode 100644
--- /dev/null
+++ b/Engine/Source/Developer/TraceInsights/Private/Insights/Widgets/SQuickFind.cpp
@@ -0,0 +1,198 @@
+// Copyright Epic Games, Inc. All Rights Reserved.
+
+#include "SQuickFind.h"
+
+#include "EditorStyleSet.h"
+#include "SlateOptMacros.h"
+#include "Widgets/DeclarativeSyntaxSupport.h"
+#include "Widgets/Docking/SDockTab.h"
+#include "Widgets/Input/SButton.h"
+#include "Widgets/Layout/SScrollBox.h"
+#include "Widgets/Images/SImage.h"
+#include "Widgets/SBoxPanel.h"
+#include "Widgets/SOverlay.h"
+#include "Widgets/Text/STextBlock.h"
+#include "Widgets/Views/STableViewBase.h"
+
+// Insights
+#include "Insights/Widgets/SFilterConfigurator.h"
+#include "Insights/ViewModels/FilterConfigurator.h"
+#include "Insights/ViewModels/QuickFind.h"
+
+#define LOCTEXT_NAMESPACE "SQuickFind"
+
+namespace Insights
+{
+
+////////////////////////////////////////////////////////////////////////////////////////////////////
+// SQuickFind
+////////////////////////////////////////////////////////////////////////////////////////////////////
+
+FName const SQuickFind::QuickFindTabId(TEXT("QuickFilter"));
+TSharedPtr<SQuickFind> SQuickFind::PendingWidget;
+bool SQuickFind::bIsTabRegistered = false;
+
+
+SQuickFind::SQuickFind()
+{
+}
+
+////////////////////////////////////////////////////////////////////////////////////////////////////
+
+SQuickFind::~SQuickFind()
+{
+}
+
+////////////////////////////////////////////////////////////////////////////////////////////////////
+
+void SQuickFind::InitCommandList()
+{
+}
+
+////////////////////////////////////////////////////////////////////////////////////////////////////
+
+BEGIN_SLATE_FUNCTION_BUILD_OPTIMIZATION
+
+void SQuickFind::Construct(const FArguments& InArgs, TSharedPtr<FQuickFind> InQuickFindViewModel)
+{
+	ChildSlot
+	[
+		SNew(SVerticalBox)
+
+		// Tree view
+		+ SVerticalBox::Slot()
+		.FillHeight(1.0f)
+		.Padding(0.0f, 6.0f, 0.0f, 0.0f)
+		[
+			SAssignNew(FilterConfigurator, SFilterConfigurator, InQuickFindViewModel->GetFilterConfigurator())
+		]
+
+		+ SVerticalBox::Slot()
+		.AutoHeight()
+		.Padding(0.0f, 5.0f, 0.0f, 5.0f)
+		.HAlign(EHorizontalAlignment::HAlign_Right)
+		.VAlign(EVerticalAlignment::VAlign_Bottom)
+		[
+			SNew(SHorizontalBox)
+
+			+ SHorizontalBox::Slot()
+			.HAlign(HAlign_Left)
+			.VAlign(VAlign_Center)
+			.AutoWidth()
+			.Padding(1.0f)
+			[
+				SNew(SButton)
+				.Text(LOCTEXT("FindNext", "Find Next"))
+				.ToolTipText(LOCTEXT("FindNextDesc", "Find the next occurence that matches the search criteria."))
+				.OnClicked(this, &SQuickFind::FindNext_Onclicked)
+			]
+		]
+	];
+
+	SHeaderRow::FColumn::FArguments ColumnArgs;
+	ColumnArgs
+		.ColumnId(TEXT("Filter"))
+		.DefaultLabel(LOCTEXT("FilterColumnHeader", "Filter"))
+		.HAlignHeader(HAlign_Fill)
+		.VAlignHeader(VAlign_Fill)
+		.HeaderContentPadding(FMargin(2.0f))
+		.HAlignCell(HAlign_Fill)
+		.VAlignCell(VAlign_Fill)
+		.HeaderContent()
+		[
+			SNew(SBox)
+			.HAlign(HAlign_Fill)
+			.VAlign(VAlign_Center)
+			[
+				SNew(STextBlock)
+				.Text(LOCTEXT("FilterColumnHeader", "Filter"))
+			]
+		];
+
+	QuickFindViewModel = InQuickFindViewModel;
+	OnViewModelDestroyedHandle = InQuickFindViewModel->GetFilterConfigurator()->GetOnDestroyedEvent().AddSP(this, &SQuickFind::RequestClose);
+}
+
+END_SLATE_FUNCTION_BUILD_OPTIMIZATION
+
+////////////////////////////////////////////////////////////////////////////////////////////////////
+
+FReply SQuickFind::FindNext_Onclicked()
+{
+	QuickFindViewModel->GetFilterConfigurator()->GetRootNode()->ProcessFilter();
+	QuickFindViewModel->GetOnFindNextEvent().Broadcast();
+
+	return FReply::Handled();
+}
+
+////////////////////////////////////////////////////////////////////////////////////////////////////
+
+FReply SQuickFind::Cancel_OnClicked()
+{
+	RequestClose();
+
+	return FReply::Handled();
+}
+
+////////////////////////////////////////////////////////////////////////////////////////////////////
+
+void SQuickFind::RequestClose()
+{
+	TSharedPtr<SDockTab> ParentTabSharedPtr = ParentTab.Pin();
+	if (ParentTabSharedPtr.IsValid())
+	{
+		ParentTabSharedPtr->RequestCloseTab();
+		ParentTabSharedPtr.Reset();
+	}
+}
+
+////////////////////////////////////////////////////////////////////////////////////////////////////
+
+void SQuickFind::RegisterQuickFindTab()
+{
+	if (bIsTabRegistered)
+	{
+		return;
+	}
+
+	FTabSpawnerEntry& TabSpawnerEntry = FGlobalTabmanager::Get()->RegisterNomadTabSpawner(QuickFindTabId,
+		FOnSpawnTab::CreateStatic(&SQuickFind::SpawnTab))
+		.SetDisplayName(LOCTEXT("FilterConfiguratorTabTitle", "Filter Configurator"))
+		.SetMenuType(ETabSpawnerMenuType::Hidden);
+
+	bIsTabRegistered = true;
+}
+
+////////////////////////////////////////////////////////////////////////////////////////////////////
+
+TSharedRef<SDockTab> SQuickFind::SpawnTab(const FSpawnTabArgs& Args)
+{
+	const TSharedRef<SDockTab> DockTab = SNew(SDockTab)
+		.TabRole(ETabRole::NomadTab);
+
+	DockTab->SetContent(PendingWidget.ToSharedRef());
+	PendingWidget->SetParentTab(DockTab);
+
+	PendingWidget = nullptr;
+	return DockTab;
+}
+
+////////////////////////////////////////////////////////////////////////////////////////////////////
+
+TSharedPtr<SQuickFind> SQuickFind::CreateAndOpenQuickFilterWidget(TSharedPtr<FQuickFind> InQuickFindViewModel)
+{
+	SAssignNew(PendingWidget, SQuickFind, InQuickFindViewModel);
+
+	if (FGlobalTabmanager::Get()->HasTabSpawner(QuickFindTabId))
+	{
+		FGlobalTabmanager::Get()->TryInvokeTab(QuickFindTabId);
+	}
+
+	return PendingWidget;
+}
+
+////////////////////////////////////////////////////////////////////////////////////////////////////
+
+} // namespace Insights
+
+#undef LOCTEXT_NAMESPACE
diff --git a/Engine/Source/Developer/TraceInsights/Private/Insights/Widgets/SQuickFind.h b/Engine/Source/Developer/TraceInsights/Private/Insights/Widgets/SQuickFind.h
new file mode 100644
--- /dev/null
+++ b/Engine/Source/Developer/TraceInsights/Private/Insights/Widgets/SQuickFind.h
@@ -0,0 +1,80 @@
+// Copyright Epic Games, Inc. All Rights Reserved.
+
+#pragma once
+
+#include "CoreMinimal.h"
+#include "Widgets/SCompoundWidget.h"
+#include "Widgets/SWidget.h"
+#include "Widgets/Views/STreeView.h"
+
+// Insights
+#include "Insights/ViewModels/Filters.h"
+#include "Insights/ViewModels/FilterConfiguratorNode.h"
+
+class SDockTab;
+
+namespace Insights
+{
+
+class SFilterConfigurator;
+class FQuickFind;
+////////////////////////////////////////////////////////////////////////////////////////////////////
+
+/**
+ * A custom widget used to configure custom filters.
+ */
+class SQuickFind: public SCompoundWidget
+{
+public:
+	SQuickFind();
+
+	virtual ~SQuickFind();
+
+	SLATE_BEGIN_ARGS(SQuickFind) {}
+	SLATE_END_ARGS()
+
+	/**
+	 * Construct this widget
+	 * @param InArgs - The declaration data for this widget
+	 */
+	void Construct(const FArguments& InArgs, TSharedPtr<FQuickFind> InFilterConfiguratorViewModel);
+
+	void Reset();
+
+	void SetParentTab(const TSharedPtr<SDockTab> InTab) { ParentTab = InTab; }
+	const TWeakPtr<SDockTab> GetParentTab() { return ParentTab; };
+
+	void RequestClose();
+
+	static void RegisterQuickFindTab();
+
+	static TSharedPtr<SQuickFind> CreateAndOpenQuickFilterWidget(TSharedPtr<FQuickFind> InQuickFindViewModel);
+
+private:
+	void InitCommandList();
+
+	FReply FindNext_Onclicked();
+
+	FReply Cancel_OnClicked();
+
+	static TSharedRef<SDockTab> SpawnTab(const FSpawnTabArgs& Args);
+
+private:
+	static const FName QuickFindTabId;
+
+	static TSharedPtr<SQuickFind> PendingWidget;
+
+	static bool bIsTabRegistered;
+
+	TSharedPtr<SFilterConfigurator> FilterConfigurator;
+
+	TSharedPtr<FQuickFind> QuickFindViewModel;
+
+	TWeakPtr<SDockTab> ParentTab;
+
+	FDelegateHandle OnViewModelDestroyedHandle;
+};
+
+////////////////////////////////////////////////////////////////////////////////////////////////////
+
+} // namespace Insights
diff --git a/Engine/Source/Developer/TraceInsights/Private/Insights/Widgets/STimingView.cpp b/Engine/Source/Developer/TraceInsights/Private/Insights/Widgets/STimingView.cpp
--- a/Engine/Source/Developer/TraceInsights/Private/Insights/Widgets/STimingView.cpp
+++ b/Engine/Source/Developer/TraceInsights/Private/Insights/Widgets/STimingView.cpp
@@ -46,6 +46,7 @@
 #include "Insights/ViewModels/BaseTimingTrack.h"
 #include "Insights/ViewModels/DrawHelpers.h"
 #include "Insights/ViewModels/FileActivityTimingTrack.h"
+#include "Insights/ViewModels/FilterConfigurator.h"
 #include "Insights/ViewModels/FrameTimingTrack.h"
 #include "Insights/ViewModels/GraphSeries.h"
 #include "Insights/ViewModels/GraphTrack.h"
@@ -56,10 +57,12 @@
 #include "Insights/ViewModels/TimingEventSearch.h"
 #include "Insights/ViewModels/TimingGraphTrack.h"
 #include "Insights/ViewModels/TimingViewDrawHelper.h"
+#include "Insights/ViewModels/QuickFind.h"
 #include "Insights/Widgets/SStatsView.h"
 #include "Insights/Widgets/STimersView.h"
 #include "Insights/Widgets/STimingProfilerWindow.h"
 #include "Insights/Widgets/STimingViewTrackList.h"
+#include "Insights/Widgets/SQuickFind.h"
 
 #include <limits>
 
@@ -246,6 +249,7 @@ void STimingView::Construct(const FArguments& InArgs)
 		]
 	];
 
+	Insights::SQuickFind::RegisterQuickFindTab();
 	UpdateHorizontalScrollBar();
 	UpdateVerticalScrollBar();
 
@@ -2794,6 +2798,8 @@ void STimingView::ShowContextMenu(const FPointerEvent& MouseEvent)
 		bHasAnyActions = true;
 	}
 
+	MenuBuilder.AddMenuEntry(Commands.QuickFind);
+
 	for (Insights::ITimingViewExtender* Extender : GetExtenders())
 	{
 		bHasAnyActions |= Extender->ExtendGlobalContextMenu(*this, MenuBuilder);
@@ -3019,7 +3025,10 @@ void STimingView::BindCommands()
 		FCanExecuteAction::CreateLambda([] { return true; }),
 		FIsActionChecked::CreateSP(this, &STimingView::ShowHideGraphTrack_IsChecked));
 
-	//TODO: let extenders to bind other commands
+	CommandList->MapAction(
+		Commands.QuickFind,
+		FExecuteAction::CreateSP(this, &STimingView::QuickFind_Execute),
+		FCanExecuteAction::CreateSP(this, &STimingView::QuickFind_CanExecute));
 
 	//CommandList->MapAction(
 	//	Commands.ShowAllGpuTracks,
@@ -4167,6 +4176,36 @@ void STimingView::ToggleTrackVisibility_Execute(uint64 InTrackId)
 
 ////////////////////////////////////////////////////////////////////////////////////////////////////
 
+bool STimingView::QuickFind_CanExecute() const
+{
+	return true;
+}
+
+////////////////////////////////////////////////////////////////////////////////////////////////////
+
+void STimingView::QuickFind_Execute()
+{
+	using namespace Insights;
+
+	if (!QuickFindVm.IsValid())
+	{
+		TSharedPtr<FFilterConfigurator> FilterConfigurator = MakeShared<FFilterConfigurator>();
+		TSharedPtr<TArray<TSharedPtr<struct FFilter>>>& AvailableFilters = FilterConfigurator->GetAvailableFilters();
+
+		AvailableFilters->Add(MakeShared<FFilter>(static_cast<int32>(EFilterField::StartTime), LOCTEXT("StartTime", "Start Time"), LOCTEXT("StartTime", "Start Time"), EFilterDataType::Double, FFilterService::Get()->GetDoubleOperators()));
+		AvailableFilters->Add(MakeShared<FFilter>(static_cast<int32>(EFilterField::EndTime), LOCTEXT("EndTime", "End Time"), LOCTEXT("EndTime", "End Time"), EFilterDataType::Double, FFilterService::Get()->GetDoubleOperators()));
+		AvailableFilters->Add(MakeShared<FFilter>(static_cast<int32>(EFilterField::Duration), LOCTEXT("Duration", "Duration"), LOCTEXT("Duration", "Duration"), EFilterDataType::Double, FFilterService::Get()->GetDoubleOperators()));
+		AvailableFilters->Add(MakeShared<FFilter>(static_cast<int32>(EFilterField::EventType), LOCTEXT("Type", "Type"), LOCTEXT("Type", "Type"), EFilterDataType::Int64, FFilterService::Get()->GetIntegerOperators()));
+
+		QuickFindVm = MakeShared<FQuickFind>(FilterConfigurator);
+		QuickFindVm->GetOnFindNextEvent().AddSP(this, &STimingView::FindNextEvent);
+	}
+
+	SQuickFind::CreateAndOpenQuickFilterWidget(QuickFindVm);
+}
+
+////////////////////////////////////////////////////////////////////////////////////////////////////
+
 void STimingView::OnTrackVisibilityChanged()
 {
 	if (HoveredTrack.IsValid())
@@ -4231,5 +4270,47 @@ void STimingView::ClearRelations()
 
 ////////////////////////////////////////////////////////////////////////////////////////////////////
 
+void STimingView::FindNextEvent()
+{
+	TSharedPtr<const ITimingEvent> BestMatchEvent;
+	double StartTime = SelectedEvent.IsValid() ? SelectedEvent->GetStartTime() : std::numeric_limits<double>::lowest();
+
+	auto EventFilter = [StartTime](double EventStartTime, double EventEndTime, uint32 EventDepth)
+	{
+		return EventStartTime > StartTime;
+	};
+	FTimingEventSearchParameters Params(StartTime, std::numeric_limits<double>::max(), ETimingEventSearchFlags::StopAtFirstMatch, EventFilter);
+	Params.FilterExecutor = QuickFindVm->GetFilterConfigurator();
+
+	for (auto& Entry : AllTracks)
+	{
+		if (!Entry.Value->IsVisible())
+		{
+			continue;
+		}
+
+		Params.EndTime = BestMatchEvent.IsValid() ? BestMatchEvent->GetStartTime() : std::numeric_limits<double>::max();
+
+		TSharedPtr<const ITimingEvent> FoundEvent = Entry.Value->SearchEvent(Params);
+		if (FoundEvent.IsValid())
+		{
+			if (BestMatchEvent.IsValid())
+			{
+				ensure(FoundEvent->GetStartTime() < BestMatchEvent->GetStartTime());
+			}
+			BestMatchEvent = FoundEvent;
+		}
+	}
+
+	if (BestMatchEvent)
+	{
+		SelectedEvent = BestMatchEvent;
+		BringIntoView(SelectedEvent->GetStartTime(), SelectedEvent->GetEndTime());
+		OnSelectedTimingEventChanged();
+	}
+}
+
+////////////////////////////////////////////////////////////////////////////////////////////////////
+
 #undef ACTIVATE_BENCHMARK
 #undef LOCTEXT_NAMESPACE
diff --git a/Engine/Source/Developer/TraceInsights/Private/Insights/Widgets/STimingView.h b/Engine/Source/Developer/TraceInsights/Private/Insights/Widgets/STimingView.h
--- a/Engine/Source/Developer/TraceInsights/Private/Insights/Widgets/STimingView.h
+++ b/Engine/Source/Developer/TraceInsights/Private/Insights/Widgets/STimingView.h
@@ -41,6 +41,7 @@ namespace Insights
 {
 	class ITimingViewExtender;
 	class FTimeMarker;
+	class FQuickFind;
 }
 
 /** A custom widget used to display timing events. */
@@ -80,6 +81,9 @@ public:
 	bool IsAutoHideEmptyTracksEnabled() const;
 	void ToggleAutoHideEmptyTracks();
 
+	bool QuickFind_CanExecute() const;
+	void QuickFind_Execute();
+
 	bool ToggleTrackVisibility_IsChecked(uint64 InTrackId) const;
 	void ToggleTrackVisibility_Execute(uint64 InTrackId);
 
@@ -440,6 +444,8 @@ protected:
 	bool ContextMenu_ShowTaskDependecies_IsChecked();
 
 	void SetTrackPosY(TSharedPtr<FBaseTimingTrack>& TrackPtr, float TopOffsetY) const;
+	
+	void FindNextEvent();
 
 protected:
 	/** The track's viewport. Encapsulates info about position and scale. */
@@ -634,4 +640,6 @@ protected:
 	TSharedPtr<FUICommandList> CommandList;
 
 	TArray<TUniquePtr<ITimingEventRelation>> CurrentRelations;
+
+	TSharedPtr<Insights::FQuickFind> QuickFindVm;
 };
\ No newline at end of file
diff --git a/Engine/Source/Developer/TraceInsights/Public/Insights/ViewModels/IFilterExecutor.h b/Engine/Source/Developer/TraceInsights/Public/Insights/ViewModels/IFilterExecutor.h
new file mode 100644
--- /dev/null
+++ b/Engine/Source/Developer/TraceInsights/Public/Insights/ViewModels/IFilterExecutor.h
@@ -0,0 +1,20 @@
+// Copyright Epic Games, Inc. All Rights Reserved.
+
+#pragma once
+
+#include "CoreMinimal.h"
+
+namespace Insights
+{
+ 
+////////////////////////////////////////////////////////////////////////////////////////////////////
+
+class IFilterExecutor
+{
+public:
+	virtual bool ApplyFilters(const class FFilterContext& Context) const = 0;
+};
+
+////////////////////////////////////////////////////////////////////////////////////////////////////
+
+} // namespace Insights
\ No newline at end of file
diff --git a/Engine/Source/Developer/TraceInsights/Public/Insights/ViewModels/TimingEventSearch.h b/Engine/Source/Developer/TraceInsights/Public/Insights/ViewModels/TimingEventSearch.h
--- a/Engine/Source/Developer/TraceInsights/Public/Insights/ViewModels/TimingEventSearch.h
+++ b/Engine/Source/Developer/TraceInsights/Public/Insights/ViewModels/TimingEventSearch.h
@@ -5,6 +5,11 @@
 #include "CoreMinimal.h"
 #include "Containers/Map.h"
 
+namespace Insights
+{
+	class IFilterExecutor;
+}
+
 ////////////////////////////////////////////////////////////////////////////////////////////////////
 // Search behavior flags
 enum class ETimingEventSearchFlags : int32
@@ -88,6 +93,8 @@ public:
 
 	// Search behavior flags
 	ETimingEventSearchFlags Flags;
+
+	TSharedPtr<Insights::IFilterExecutor> FilterExecutor;
 };
 
 ////////////////////////////////////////////////////////////////////////////////////////////////////
