[Insights] Fix context switches not being displayed after the timing view is reopened. #rb none

diff --git a/Engine/Source/Developer/TraceInsights/Private/Insights/ContextSwitches/ViewModels/ContextSwitchesSharedState.cpp b/Engine/Source/Developer/TraceInsights/Private/Insights/ContextSwitches/ViewModels/ContextSwitchesSharedState.cpp
--- a/Engine/Source/Developer/TraceInsights/Private/Insights/ContextSwitches/ViewModels/ContextSwitchesSharedState.cpp
+++ b/Engine/Source/Developer/TraceInsights/Private/Insights/ContextSwitches/ViewModels/ContextSwitchesSharedState.cpp
@@ -14,7 +14,9 @@
 #include "Insights/ContextSwitches/ViewModels/CpuCoreTimingTrack.h"
 #include "Insights/ITimingViewSession.h"
 #include "Insights/InsightsManager.h"
+#include "Insights/TimingProfilerManager.h"
 #include "Insights/ViewModels/ThreadTimingTrack.h"
+#include "Insights/Widgets/STimingProfilerWindow.h"
 #include "Insights/Widgets/STimingView.h"
 
 #define LOCTEXT_NAMESPACE "ContextSwitches"
@@ -77,13 +79,31 @@ void FContextSwitchesSharedState::OnBeginSession(Insights::ITimingViewSession& I
 {
 	if (&InSession != TimingView)
 	{
-		return;
+		TSharedPtr<STimingProfilerWindow> Window = FTimingProfilerManager::Get()->GetProfilerWindow();
+		if (Window.IsValid())
+		{
+			TSharedPtr<STimingView> WindowTimingView = Window->GetTimingView();
+			if (WindowTimingView.IsValid() && WindowTimingView.Get() == &InSession)
+			{
+				TimingView = WindowTimingView.Get();
+				AddCommands();
+			}
+		}
+		else if(TimingView == nullptr && FTimingProfilerManager::Get()->IsTimingViewVisible())
+		{
+			TimingView = (STimingView *)&InSession;
+			AddCommands();
+		}
+		else
+		{
+			return;
+		}
 	}
 
 	ThreadsSerial = 0;
 	CpuCoresSerial = 0;
 
-	bAreCoreTracksVisible = true;
+	bAreCoreTracksVisible = false;
 	bAreContextSwitchesVisible = true;
 	bAreOverlaysVisible = true;
 	bAreExtendedLinesVisible = true;
@@ -103,12 +123,13 @@ void FContextSwitchesSharedState::OnEndSession(Insights::ITimingViewSession& InS
 	ThreadsSerial = 0;
 	CpuCoresSerial = 0;
 
-	bAreCoreTracksVisible = true;
+	bAreCoreTracksVisible = false;
 	bAreContextSwitchesVisible = true;
 	bAreOverlaysVisible = true;
 	bAreExtendedLinesVisible = true;
 
 	bSyncWithProviders = false;
+	TimingView = nullptr;
 }
 
 ////////////////////////////////////////////////////////////////////////////////////////////////////
diff --git a/Engine/Source/Developer/TraceInsights/Private/Insights/TaskGraphProfiler/ViewModels/TaskTimingTrack.cpp b/Engine/Source/Developer/TraceInsights/Private/Insights/TaskGraphProfiler/ViewModels/TaskTimingTrack.cpp
--- a/Engine/Source/Developer/TraceInsights/Private/Insights/TaskGraphProfiler/ViewModels/TaskTimingTrack.cpp
+++ b/Engine/Source/Developer/TraceInsights/Private/Insights/TaskGraphProfiler/ViewModels/TaskTimingTrack.cpp
@@ -17,8 +17,10 @@
 #include "Insights/ViewModels/ThreadTimingTrack.h"
 #include "Insights/ViewModels/ThreadTrackEvent.h"
 #include "Insights/ViewModels/TimingTrackViewport.h"
+#include "Insights/TimingProfilerManager.h"
 #include "Insights/ViewModels/TimingViewDrawHelper.h"
 #include "Insights/ViewModels/TooltipDrawState.h"
+#include "Insights/Widgets/STimingProfilerWindow.h"
 #include "Insights/Widgets/STimingView.h"
 
 #define LOCTEXT_NAMESPACE "TaskTimingTrack"
@@ -73,7 +75,23 @@ void FTaskTimingSharedState::OnBeginSession(Insights::ITimingViewSession& InSess
 {
 	if (&InSession != TimingView)
 	{
-		return;
+		TSharedPtr<STimingProfilerWindow> Window = FTimingProfilerManager::Get()->GetProfilerWindow();
+		if (Window.IsValid())
+		{
+			TSharedPtr<STimingView> WindowTimingView = Window->GetTimingView();
+			if (WindowTimingView.IsValid() && WindowTimingView.Get() == &InSession)
+			{
+				TimingView = WindowTimingView.Get();
+			}
+		}
+		else if (TimingView == nullptr && FTimingProfilerManager::Get()->IsTimingViewVisible())
+		{
+			TimingView = (STimingView*)&InSession;
+		}
+		else
+		{
+			return;
+		}
 	}
 
 	TaskTrack = nullptr;
@@ -90,6 +108,7 @@ void FTaskTimingSharedState::OnEndSession(Insights::ITimingViewSession& InSessio
 
 	FTaskGraphProfilerManager::Get()->ClearTaskRelations();
 	TaskTrack = nullptr;
+	TimingView = nullptr;
 }
 
 ////////////////////////////////////////////////////////////////////////////////////////////////////
