Merge Changelist

NetStat Group sorting is no longer by name.
Each group defines how it is sorted.

CL 16605730

diff --git a/Engine/Source/Developer/TraceInsights/Private/Insights/NetworkingProfiler/Widgets/SNetStatsView.cpp b/Engine/Source/Developer/TraceInsights/Private/Insights/NetworkingProfiler/Widgets/SNetStatsView.cpp
--- a/Engine/Source/Developer/TraceInsights/Private/Insights/NetworkingProfiler/Widgets/SNetStatsView.cpp
+++ b/Engine/Source/Developer/TraceInsights/Private/Insights/NetworkingProfiler/Widgets/SNetStatsView.cpp
@@ -627,12 +627,12 @@ void SNetStatsView::UpdateTree()
 	FStopwatch Stopwatch;
 	Stopwatch.Start();
 
-	CreateGroups();
+	CreateSortedGroups();
 
 	Stopwatch.Update();
 	const double Time1 = Stopwatch.GetAccumulatedTime();
 
-	SortTreeNodes();
+	SortTreeChildNodes();
 
 	Stopwatch.Update();
 	const double Time2 = Stopwatch.GetAccumulatedTime();
@@ -962,7 +962,7 @@ bool SNetStatsView::SearchBox_IsEnabled() const
 // Grouping
 ////////////////////////////////////////////////////////////////////////////////////////////////////
 
-void SNetStatsView::CreateGroups()
+void SNetStatsView::CreateSortedGroups()
 {
 	if (GroupingMode == ENetEventGroupingMode::Flat)
 	{
@@ -1066,8 +1066,8 @@ void SNetStatsView::GroupBy_OnSelectionChanged(TSharedPtr<ENetEventGroupingMode>
 	{
 		GroupingMode = *NewGroupingMode;
 
-		CreateGroups();
-		SortTreeNodes();
+		CreateSortedGroups();
+		SortTreeChildNodes();
 		ApplyFiltering();
 	}
 }
@@ -1143,30 +1143,10 @@ void SNetStatsView::UpdateCurrentSortingByColumn()
 
 ////////////////////////////////////////////////////////////////////////////////////////////////////
 
-void SNetStatsView::SortTreeNodes()
+void SNetStatsView::SortTreeChildNodes()
 {
 	if (CurrentSorter.IsValid())
 	{
-		// Sort groups (always by name).
-		TArray<Insights::FBaseTreeNodePtr> SortedGroupNodes;
-		for (const FNetEventNodePtr& NodePtr : GroupNodes)
-		{
-			SortedGroupNodes.Add(NodePtr);
-		}
-		TSharedPtr<Insights::ITableCellValueSorter> Sorter = CurrentSorter;
-		Insights::ESortMode SortMode = (ColumnSortMode == EColumnSortMode::Type::Descending) ? Insights::ESortMode::Descending : Insights::ESortMode::Ascending;
-		if (CurrentSorter->GetName() != FName(TEXT("ByName")))
-		{
-			Sorter = MakeShared<Insights::FSorterByName>(Table->GetColumns()[0]);
-			SortMode = Insights::ESortMode::Ascending;
-		}
-		Sorter->Sort(SortedGroupNodes, SortMode);
-		GroupNodes.Reset();
-		for (const Insights::FBaseTreeNodePtr& NodePtr : SortedGroupNodes)
-		{
-			GroupNodes.Add(StaticCastSharedPtr<FNetEventNode>(NodePtr));
-		}
-
 		// Sort nodes in each group.
 		for (FNetEventNodePtr& Root : GroupNodes)
 		{
@@ -1218,7 +1198,7 @@ void SNetStatsView::SetSortModeForColumn(const FName& ColumnId, const EColumnSor
 	ColumnSortMode = SortMode;
 	UpdateCurrentSortingByColumn();
 
-	SortTreeNodes();
+	SortTreeChildNodes();
 	ApplyFiltering();
 }
 
diff --git a/Engine/Source/Developer/TraceInsights/Private/Insights/NetworkingProfiler/Widgets/SNetStatsView.h b/Engine/Source/Developer/TraceInsights/Private/Insights/NetworkingProfiler/Widgets/SNetStatsView.h
--- a/Engine/Source/Developer/TraceInsights/Private/Insights/NetworkingProfiler/Widgets/SNetStatsView.h
+++ b/Engine/Source/Developer/TraceInsights/Private/Insights/NetworkingProfiler/Widgets/SNetStatsView.h
@@ -175,7 +175,7 @@ protected:
 	////////////////////////////////////////////////////////////////////////////////////////////////////
 	// Grouping
 
-	void CreateGroups();
+	void CreateSortedGroups();
 	void CreateGroupByOptionsSources();
 
 	void GroupBy_OnSelectionChanged(TSharedPtr<ENetEventGroupingMode> NewGroupingMode, ESelectInfo::Type SelectInfo);
@@ -195,7 +195,7 @@ protected:
 	void CreateSortings();
 
 	void UpdateCurrentSortingByColumn();
-	void SortTreeNodes();
+	void SortTreeChildNodes();
 	void SortTreeNodesRec(FNetEventNode& Node, const Insights::ITableCellValueSorter& Sorter);
 
 	EColumnSortMode::Type GetSortModeForColumn(const FName ColumnId) const;
