Merge Changelists

UE Insights instance, connection, and connection mode selections are preserved after new instances are added.

CL 16704028

diff --git a/Engine/Source/Developer/TraceInsights/Private/Insights/NetworkingProfiler/Widgets/SNetworkingProfilerWindow.cpp b/Engine/Source/Developer/TraceInsights/Private/Insights/NetworkingProfiler/Widgets/SNetworkingProfilerWindow.cpp
--- a/Engine/Source/Developer/TraceInsights/Private/Insights/NetworkingProfiler/Widgets/SNetworkingProfilerWindow.cpp
+++ b/Engine/Source/Developer/TraceInsights/Private/Insights/NetworkingProfiler/Widgets/SNetworkingProfilerWindow.cpp
@@ -527,6 +527,18 @@ void SNetworkingProfilerWindow::Tick(const FGeometry& AllottedGeometry, const do
 
 void SNetworkingProfilerWindow::UpdateAvailableGameInstances()
 {
+	uint32 CurrentSelectedGameInstanceIndex = 0;
+	bool bIsInstanceSelected = false;
+	if (GameInstanceComboBox.IsValid())
+	{
+		TSharedPtr<FGameInstanceItem> CurrentSelectedGameInstance = GameInstanceComboBox->GetSelectedItem();
+		if (CurrentSelectedGameInstance.IsValid())
+		{
+			CurrentSelectedGameInstanceIndex = CurrentSelectedGameInstance->GetIndex();
+			bIsInstanceSelected = true;
+		}
+	}
+	
 	AvailableGameInstances.Reset();
 
 	TSharedPtr<const TraceServices::IAnalysisSession> Session = FInsightsManager::Get()->GetSession();
@@ -545,7 +557,15 @@ void SNetworkingProfilerWindow::UpdateAvailableGameInstances()
 
 	if (GameInstanceComboBox.IsValid())
 	{
-		GameInstanceComboBox->SetSelectedItem(AvailableGameInstances.Num() > 0 ? AvailableGameInstances[0] : nullptr);
+		const uint32 AvailableGameInstanceCount = AvailableGameInstances.Num();
+		if (bIsInstanceSelected && AvailableGameInstanceCount > CurrentSelectedGameInstanceIndex)
+		{
+			GameInstanceComboBox->SetSelectedItem(AvailableGameInstances[CurrentSelectedGameInstanceIndex]);
+		}
+		else
+		{
+			GameInstanceComboBox->SetSelectedItem(AvailableGameInstances.Num() > 0 ? AvailableGameInstances[0] : nullptr);
+		}
 	}
 
 	UpdateAvailableConnections();
@@ -555,6 +575,18 @@ void SNetworkingProfilerWindow::UpdateAvailableGameInstances()
 
 void SNetworkingProfilerWindow::UpdateAvailableConnections()
 {
+	uint32 CurrentSelectedConnectionIndex = 0;
+	bool bIsConnectionSelected = false;
+	if (ConnectionComboBox.IsValid())
+	{
+		TSharedPtr<FConnectionItem> CurrentSelectedConnection = ConnectionComboBox->GetSelectedItem();
+		if (CurrentSelectedConnection.IsValid())
+		{
+			CurrentSelectedConnectionIndex = CurrentSelectedConnection->GetIndex();
+			bIsConnectionSelected = true;
+		}
+	}
+
 	AvailableConnections.Reset();
 
 	TSharedPtr<const TraceServices::IAnalysisSession> Session = FInsightsManager::Get()->GetSession();
@@ -573,7 +605,15 @@ void SNetworkingProfilerWindow::UpdateAvailableConnections()
 
 	if (ConnectionComboBox.IsValid())
 	{
-		ConnectionComboBox->SetSelectedItem(AvailableConnections.Num() > 0 ? AvailableConnections[0] : nullptr);
+		const uint32 AvailableConnectionCount = AvailableConnections.Num();
+		if (bIsConnectionSelected && AvailableConnectionCount > CurrentSelectedConnectionIndex)
+		{
+			ConnectionComboBox->SetSelectedItem(AvailableConnections[CurrentSelectedConnectionIndex]);
+		}
+		else
+		{
+			ConnectionComboBox->SetSelectedItem(AvailableConnections.Num() > 0 ? AvailableConnections[0] : nullptr);
+		}
 	}
 
 	UpdateAvailableConnectionModes();
@@ -583,6 +623,18 @@ void SNetworkingProfilerWindow::UpdateAvailableConnections()
 
 void SNetworkingProfilerWindow::UpdateAvailableConnectionModes()
 {
+	bool bIsConnectionModeSelected = false;
+	TraceServices::ENetProfilerConnectionMode CurrentSelectedConnectionModeChoice = TraceServices::ENetProfilerConnectionMode::Incoming;
+	if (ConnectionModeComboBox.IsValid())
+	{
+		TSharedPtr<FConnectionModeItem> CurrentSelectedConnectionMode = ConnectionModeComboBox->GetSelectedItem();
+		if (CurrentSelectedConnectionMode.IsValid())
+		{
+			CurrentSelectedConnectionModeChoice = CurrentSelectedConnectionMode->Mode;
+			bIsConnectionModeSelected = true;
+		}
+	}
+
 	AvailableConnectionModes.Reset();
 
 	if (SelectedConnection.IsValid())
@@ -599,7 +651,22 @@ void SNetworkingProfilerWindow::UpdateAvailableConnectionModes()
 
 	if (ConnectionModeComboBox.IsValid())
 	{
-		ConnectionModeComboBox->SetSelectedItem(AvailableConnectionModes.Num() > 0 ? AvailableConnectionModes[0] : nullptr);
+		const uint32 AvailableConnectionModeCount = AvailableConnectionModes.Num();
+		if (bIsConnectionModeSelected && AvailableConnectionModeCount == 2)
+		{
+			if (CurrentSelectedConnectionModeChoice == TraceServices::ENetProfilerConnectionMode::Outgoing)
+			{
+				ConnectionModeComboBox->SetSelectedItem(AvailableConnectionModes[1]);
+			}
+			else
+			{
+				ConnectionModeComboBox->SetSelectedItem(AvailableConnectionModes[0]);
+			}
+		}
+		else
+		{
+			ConnectionModeComboBox->SetSelectedItem(AvailableConnectionModes.Num() > 0 ? AvailableConnectionModes[0] : nullptr);
+		}
 	}
 }
 
