Merge changelist

In the Packet Content View, arrow keys have the following functionality:
If an event is selected pressing the left or right arrow keys moves to the previous or next event on the same level.
Pressing the up or down arrow keys moves up and down between levels.
Pressing Ctrl/Command with the left or right arrow keys moves between packets.

Packets with split contents zoom appropriately when zooming is initiated.

CL 16745207, 16774443

diff --git a/Engine/Source/Developer/TraceInsights/Private/Insights/NetworkingProfiler/Widgets/SPacketContentView.cpp b/Engine/Source/Developer/TraceInsights/Private/Insights/NetworkingProfiler/Widgets/SPacketContentView.cpp
--- a/Engine/Source/Developer/TraceInsights/Private/Insights/NetworkingProfiler/Widgets/SPacketContentView.cpp
+++ b/Engine/Source/Developer/TraceInsights/Private/Insights/NetworkingProfiler/Widgets/SPacketContentView.cpp
@@ -207,7 +207,7 @@ void SPacketContentView::Construct(const FArguments& InArgs, TSharedPtr<SNetwork
 				SNew(SButton)
 				.ToolTipText(LOCTEXT("FindPreviousEventToolTip", "Previous Event"))
 				.ContentPadding(0.0f)
-				.OnClicked(this, &SPacketContentView::FindPreviousEvent)
+				.OnClicked(this, &SPacketContentView::FindPreviousEvent, EEventNavigationType::AnyLevel)
 				.Content()
 				[
 					SNew(SImage)
@@ -222,7 +222,7 @@ void SPacketContentView::Construct(const FArguments& InArgs, TSharedPtr<SNetwork
 				SNew(SButton)
 				.ToolTipText(LOCTEXT("FindNextEventToolTip", "Next Event"))
 				.ContentPadding(0.0f)
-				.OnClicked(this, &SPacketContentView::FindNextEvent)
+				.OnClicked(this, &SPacketContentView::FindNextEvent, EEventNavigationType::AnyLevel)
 				.Content()
 				[
 					SNew(SImage)
@@ -409,7 +409,7 @@ FReply SPacketContentView::FindFirstEvent()
 
 ////////////////////////////////////////////////////////////////////////////////////////////////////
 
-FReply SPacketContentView::FindPreviousEvent()
+FReply SPacketContentView::FindPreviousEvent(EEventNavigationType NavigationType)
 {
 	if (!SelectedEvent.IsValid())
 	{
@@ -424,8 +424,27 @@ FReply SPacketContentView::FindPreviousEvent()
 		{
 			if (EventIndex > 0)
 			{
-				SelectedEvent.Set(FilteredDrawState->Events[EventIndex - 1]);
-				OnSelectedEventChanged();
+				switch (NavigationType)
+				{
+				case EEventNavigationType::AnyLevel:
+					SelectedEvent.Set(FilteredDrawState->Events[EventIndex - 1]);
+					OnSelectedEventChanged();
+					break;
+				case EEventNavigationType::SameLevel:
+					for (int32 PrevEventIndex = EventIndex - 1; PrevEventIndex >= 0; --PrevEventIndex)
+					{
+						const FNetworkPacketEvent& PrevEvent = FilteredDrawState->Events[PrevEventIndex];
+						if (Event.Level == PrevEvent.Level)
+						{
+							SelectedEvent.Set(FilteredDrawState->Events[PrevEventIndex]);
+							OnSelectedEventChanged();
+							break;
+						}
+					}
+					break;
+				default:
+					break;
+				}
 				break;
 			}
 		}
@@ -446,7 +465,7 @@ FReply SPacketContentView::FindPreviousEvent()
 
 ////////////////////////////////////////////////////////////////////////////////////////////////////
 
-FReply SPacketContentView::FindNextEvent()
+FReply SPacketContentView::FindNextEvent(EEventNavigationType NavigationType)
 {
 	if (!SelectedEvent.IsValid())
 	{
@@ -461,8 +480,27 @@ FReply SPacketContentView::FindNextEvent()
 		{
 			if (EventIndex < EventCount - 1)
 			{
-				SelectedEvent.Set(FilteredDrawState->Events[EventIndex + 1]);
-				OnSelectedEventChanged();
+				switch (NavigationType)
+				{
+				case EEventNavigationType::AnyLevel:
+					SelectedEvent.Set(FilteredDrawState->Events[EventIndex + 1]);
+					OnSelectedEventChanged();
+					break;
+				case EEventNavigationType::SameLevel:
+					for (int32 NextEventIndex = EventIndex + 1; NextEventIndex <= EventCount - 1; ++NextEventIndex)
+					{
+						const FNetworkPacketEvent& NextEvent = FilteredDrawState->Events[NextEventIndex];
+						if (Event.Level == NextEvent.Level)
+						{
+							SelectedEvent.Set(FilteredDrawState->Events[NextEventIndex]);
+							OnSelectedEventChanged();
+							break;
+						}
+					}
+					break;
+				default:
+					break;
+				}
 				break;
 			}
 		}
@@ -497,6 +535,76 @@ FReply SPacketContentView::FindLastEvent()
 
 ////////////////////////////////////////////////////////////////////////////////////////////////////
 
+FReply SPacketContentView::FindPreviousLevel()
+{
+	if (!SelectedEvent.IsValid())
+	{
+		return FindFirstEvent();
+	}
+
+	const int32 EventCount = FilteredDrawState->Events.Num();
+	for (int32 EventIndex = 0; EventIndex < EventCount; ++EventIndex)
+	{
+		const FNetworkPacketEvent& Event = FilteredDrawState->Events[EventIndex];
+		if (Event.Equals(SelectedEvent.Event))
+		{
+			for (int32 PrevEventIndex = EventIndex - 1; PrevEventIndex > 0; --PrevEventIndex)
+			{
+				const FNetworkPacketEvent& PrevEvent = FilteredDrawState->Events[PrevEventIndex];
+				if (PrevEvent.Level < Event.Level &&
+					PrevEvent.BitOffset <= Event.BitOffset &&
+					PrevEvent.BitSize + PrevEvent.BitOffset >= Event.BitSize + Event.BitOffset)
+				{
+					SelectedEvent.Set(FilteredDrawState->Events[PrevEventIndex]);
+					OnSelectedEventChanged();
+					break;
+				}
+			}
+			break;
+		}
+	}
+
+	BringEventIntoView(SelectedEvent);
+	return FReply::Handled();
+}
+
+////////////////////////////////////////////////////////////////////////////////////////////////////
+
+FReply SPacketContentView::FindNextLevel()
+{
+	if (!SelectedEvent.IsValid())
+	{
+		return FindLastEvent();
+	}
+
+	const int32 EventCount = FilteredDrawState->Events.Num();
+	for (int32 EventIndex = 0; EventIndex < EventCount; ++EventIndex)
+	{
+		const FNetworkPacketEvent& Event = FilteredDrawState->Events[EventIndex];
+		if (Event.Equals(SelectedEvent.Event))
+		{
+			for (int32 NextEventIndex = EventIndex + 1; NextEventIndex < EventCount; ++NextEventIndex)
+			{
+				const FNetworkPacketEvent& NextEvent = FilteredDrawState->Events[NextEventIndex];
+				if (NextEvent.Level > Event.Level &&
+					NextEvent.BitOffset >= Event.BitOffset &&
+					NextEvent.BitSize + NextEvent.BitOffset <= Event.BitSize + Event.BitOffset)
+				{
+					SelectedEvent.Set(FilteredDrawState->Events[NextEventIndex]);
+					OnSelectedEventChanged();
+					break;
+				}
+			}
+			break;
+		}
+	}
+
+	BringEventIntoView(SelectedEvent);
+	return FReply::Handled();
+}
+
+////////////////////////////////////////////////////////////////////////////////////////////////////
+
 ECheckBoxState SPacketContentView::FilterByNetId_IsChecked() const
 {
 	return bFilterByNetId ? ECheckBoxState::Checked : ECheckBoxState::Unchecked;
@@ -601,6 +709,7 @@ void SPacketContentView::Tick(const FGeometry& AllottedGeometry, const double In
 	{
 		bIsStateDirty = false;
 		UpdateState();
+		AdjustForSplitContent();
 	}
 
 	Tooltip.Update();
@@ -1332,9 +1441,21 @@ FReply SPacketContentView::OnKeyDown(const FGeometry& MyGeometry, const FKeyEven
 		{
 			FindFirstEvent();
 		}
+		else if (InKeyEvent.GetModifierKeys().IsControlDown() ||
+				 InKeyEvent.GetModifierKeys().IsCommandDown())
+		{
+			if (ProfilerWindow.IsValid())
+			{
+				const TSharedPtr<SPacketView> PacketView = ProfilerWindow->GetPacketView();
+				if (PacketView.IsValid())
+				{
+					PacketView->SelectPreviousPacket();
+				}
+			}
+		}
 		else
 		{
-			FindPreviousEvent();
+			FindPreviousEvent(EEventNavigationType::SameLevel);
 		}
 		return FReply::Handled();
 	}
@@ -1344,12 +1465,46 @@ FReply SPacketContentView::OnKeyDown(const FGeometry& MyGeometry, const FKeyEven
 		{
 			FindLastEvent();
 		}
+		else if (InKeyEvent.GetModifierKeys().IsControlDown() ||
+				 InKeyEvent.GetModifierKeys().IsCommandDown())
+		{
+			if (ProfilerWindow.IsValid())
+			{
+				const TSharedPtr<SPacketView> PacketView = ProfilerWindow->GetPacketView();
+				if (PacketView.IsValid())
+				{
+					PacketView->SelectNextPacket();
+				}
+			}
+		}
 		else
 		{
-			FindNextEvent();
+			FindNextEvent(EEventNavigationType::SameLevel);
 		}
 		return FReply::Handled();
 	}
+	else if (InKeyEvent.GetKey() == EKeys::Up)
+	{
+		FindPreviousLevel();
+		return FReply::Handled();
+	}
+	else if (InKeyEvent.GetKey() == EKeys::Down)
+	{
+		FindNextLevel();
+		return FReply::Handled();
+	}
+	else if (InKeyEvent.GetKey() == EKeys::Equals ||
+			 InKeyEvent.GetKey() == EKeys::Add)
+	{
+		ZoomHorizontally(1.0f, MousePosition.X);
+		return FReply::Handled();
+	}
+	else if (InKeyEvent.GetKey() == EKeys::Hyphen ||
+			 InKeyEvent.GetKey() == EKeys::Subtract)
+	{
+		ZoomHorizontally(-1.0f, MousePosition.X);
+		return FReply::Handled();
+	}
 
 	return SCompoundWidget::OnKeyDown(MyGeometry, InKeyEvent);
 }
@@ -1433,4 +1588,23 @@ void SPacketContentView::BringEventIntoView(const FNetworkPacketEventRef& EventR
 
 ////////////////////////////////////////////////////////////////////////////////////////////////////
 
+void SPacketContentView::AdjustForSplitContent()
+{
+	if (FilteredDrawState->Events.Num() > 0)
+	{
+		FAxisViewportDouble& ViewportX = Viewport.GetHorizontalAxisViewport();
+		const FNetworkPacketEvent& LastEvent = FilteredDrawState->Events.Last();
+		const uint32 LastBit = LastEvent.BitOffset + LastEvent.BitSize;
+		if (LastBit > PacketBitSize)
+		{
+			PacketBitSize = LastBit;
+			ViewportX.SetMinMaxValueInterval(0.0, static_cast<double>(PacketBitSize));
+			ViewportX.CenterOnValueInterval(0.0, static_cast<double>(PacketBitSize));
+			UpdateHorizontalScrollBar();
+		}
+	}
+}
+
+////////////////////////////////////////////////////////////////////////////////////////////////////
+
 #undef LOCTEXT_NAMESPACE
diff --git a/Engine/Source/Developer/TraceInsights/Private/Insights/NetworkingProfiler/Widgets/SPacketContentView.h b/Engine/Source/Developer/TraceInsights/Private/Insights/NetworkingProfiler/Widgets/SPacketContentView.h
--- a/Engine/Source/Developer/TraceInsights/Private/Insights/NetworkingProfiler/Widgets/SPacketContentView.h
+++ b/Engine/Source/Developer/TraceInsights/Private/Insights/NetworkingProfiler/Widgets/SPacketContentView.h
@@ -97,6 +97,8 @@ public:
 		Hand,
 	};
 
+	enum class EEventNavigationType { AnyLevel, SameLevel };
+
 public:
 	/** Default constructor. */
 	SPacketContentView();
@@ -151,9 +153,11 @@ public:
 	void DisableFilterEventType();
 
 	FReply FindFirstEvent();
-	FReply FindPreviousEvent();
-	FReply FindNextEvent();
+	FReply FindPreviousEvent(EEventNavigationType NavigationType);
+	FReply FindNextEvent(EEventNavigationType NavigationType);
 	FReply FindLastEvent();
+	FReply FindPreviousLevel();
+	FReply FindNextLevel();
 
 private:
 	FReply FindPreviousPacket_OnClicked();
@@ -204,6 +208,8 @@ private:
 	void OnSelectedEventChanged();
 	void SelectHoveredEvent();
 
+	void AdjustForSplitContent();
+
 private:
 	TSharedPtr<SNetworkingProfilerWindow> ProfilerWindow;
 
