Perf improvement for memory graph track.

diff --git a/Engine/Source/Developer/TraceInsights/Private/Insights/MemoryProfiler/ViewModels/MemoryGraphTrack.cpp b/Engine/Source/Developer/TraceInsights/Private/Insights/MemoryProfiler/ViewModels/MemoryGraphTrack.cpp
--- a/Engine/Source/Developer/TraceInsights/Private/Insights/MemoryProfiler/ViewModels/MemoryGraphTrack.cpp
+++ b/Engine/Source/Developer/TraceInsights/Private/Insights/MemoryProfiler/ViewModels/MemoryGraphTrack.cpp
@@ -562,59 +562,84 @@ void FMemoryGraphTrack::UpdateAllocationsTimelineSeries(FMemoryGraphSeries& Seri
 		const TraceServices::IAllocationsProvider* AllocationsProvider = TraceServices::ReadAllocationsProvider(*Session.Get());
 		if (AllocationsProvider)
 		{
-			TraceServices::IAllocationsProvider::FReadScopeLock ProviderReadScope(*AllocationsProvider);
+			struct FTimelineEvent
+			{
+				FTimelineEvent(double InTime, double InDuration, double InValue) :
+					Time(InTime),
+					Duration(InDuration),
+					Value(InValue)
+				{
+
+				}
+				double Time;
+				double Duration;
+				double Value;
+			};
 
 			int32 StartIndex = -1;
 			int32 EndIndex = -1;
-			AllocationsProvider->GetTimelineIndexRange(Viewport.GetStartTime(), Viewport.GetEndTime(), StartIndex, EndIndex);
-			if (EndIndex >= 0)
+
+			TArray<FTimelineEvent> TimelineEvents;
 			{
-				--StartIndex; // include one more point on the left side
-				++EndIndex; // include one more point on the right side
-			}
+				TraceServices::IAllocationsProvider::FReadScopeLock ProviderReadScope(*AllocationsProvider);
+				AllocationsProvider->GetTimelineIndexRange(Viewport.GetStartTime(), Viewport.GetEndTime(), StartIndex, EndIndex);
 
-			const float TopY = 4.0f;
-			const float BottomY = GetHeight() - 4.0f;
+				if (EndIndex >= 0)
+				{
+					--StartIndex; // include one more point on the left side
+					++EndIndex; // include one more point on the right side
+				}
 
-			if (Series.IsAutoZoomEnabled() && TopY < BottomY)
-			{
-				Series.UpdateAutoZoom(TopY, BottomY, Series.GetMinValue(), Series.GetMaxValue());
-			}
+				TimelineEvents.Reserve((EndIndex - StartIndex) + 1);
 
-			auto Callback64 = [&Builder](double Time, double Duration, uint64 Value)
-			{
-				Builder.AddEvent(Time, Duration, static_cast<double>(Value));
-			};
-			auto Callback32 = [&Builder](double Time, double Duration, uint32 Value)
-			{
-				Builder.AddEvent(Time, Duration, static_cast<double>(Value));
-			};
-			auto Callback32Negative = [&Builder](double Time, double Duration, uint32 Value)
-			{
-				// Shows FreeEvents as negative values in order to be displayed on same graph as AllocEvents.
-				Builder.AddEvent(Time, Duration, -static_cast<double>(Value));
-			};
+				const float TopY = 4.0f;
+				const float BottomY = GetHeight() - 4.0f;
 
-			switch (Series.GetTimelineType())
+				if (Series.IsAutoZoomEnabled() && TopY < BottomY)
+				{
+					Series.UpdateAutoZoom(TopY, BottomY, Series.GetMinValue(), Series.GetMaxValue());
+				}
+
+				auto Callback64 = [&TimelineEvents](double Time, double Duration, uint64 Value)
+				{
+					TimelineEvents.Emplace(Time, Duration, static_cast<double>(Value));
+				};
+				auto Callback32 = [&TimelineEvents](double Time, double Duration, uint32 Value)
+				{
+					TimelineEvents.Emplace(Time, Duration, static_cast<double>(Value));
+				};
+				auto Callback32Negative = [&TimelineEvents](double Time, double Duration, uint32 Value)
+				{
+					// Shows FreeEvents as negative values in order to be displayed on same graph as AllocEvents.
+					TimelineEvents.Emplace(Time, Duration, -static_cast<double>(Value));
+				};
+
+				switch (Series.GetTimelineType())
+				{
+				case FMemoryGraphSeries::ETimelineType::MinTotalMem:
+					AllocationsProvider->EnumerateMinTotalAllocatedMemoryTimeline(StartIndex, EndIndex, Callback64);
+					break;
+				case FMemoryGraphSeries::ETimelineType::MaxTotalMem:
+					AllocationsProvider->EnumerateMaxTotalAllocatedMemoryTimeline(StartIndex, EndIndex, Callback64);
+					break;
+				case FMemoryGraphSeries::ETimelineType::MinLiveAllocs:
+					AllocationsProvider->EnumerateMinLiveAllocationsTimeline(StartIndex, EndIndex, Callback32);
+					break;
+				case FMemoryGraphSeries::ETimelineType::MaxLiveAllocs:
+					AllocationsProvider->EnumerateMaxLiveAllocationsTimeline(StartIndex, EndIndex, Callback32);
+					break;
+				case FMemoryGraphSeries::ETimelineType::AllocEvents:
+					AllocationsProvider->EnumerateAllocEventsTimeline(StartIndex, EndIndex, Callback32);
+					break;
+				case FMemoryGraphSeries::ETimelineType::FreeEvents:
+					AllocationsProvider->EnumerateFreeEventsTimeline(StartIndex, EndIndex, Callback32Negative);
+					break;
+				}
+			}
+
+			for (const FTimelineEvent& TimelineEvent : TimelineEvents)
 			{
-			case FMemoryGraphSeries::ETimelineType::MinTotalMem:
-				AllocationsProvider->EnumerateMinTotalAllocatedMemoryTimeline(StartIndex, EndIndex, Callback64);
-				break;
-			case FMemoryGraphSeries::ETimelineType::MaxTotalMem:
-				AllocationsProvider->EnumerateMaxTotalAllocatedMemoryTimeline(StartIndex, EndIndex, Callback64);
-				break;
-			case FMemoryGraphSeries::ETimelineType::MinLiveAllocs:
-				AllocationsProvider->EnumerateMinLiveAllocationsTimeline(StartIndex, EndIndex, Callback32);
-				break;
-			case FMemoryGraphSeries::ETimelineType::MaxLiveAllocs:
-				AllocationsProvider->EnumerateMaxLiveAllocationsTimeline(StartIndex, EndIndex, Callback32);
-				break;
-			case FMemoryGraphSeries::ETimelineType::AllocEvents:
-				AllocationsProvider->EnumerateAllocEventsTimeline(StartIndex, EndIndex, Callback32);
-				break;
-			case FMemoryGraphSeries::ETimelineType::FreeEvents:
-				AllocationsProvider->EnumerateFreeEventsTimeline(StartIndex, EndIndex, Callback32Negative);
-				break;
+				Builder.AddEvent(TimelineEvent.Time, TimelineEvent.Duration, TimelineEvent.Value);
 			}
 		}
 	}
