[Insights] Add "Filter All" functionality to the Quick Find widget

diff --git a/Engine/Source/Developer/TraceInsights/Private/Insights/ViewModels/QuickFind.h b/Engine/Source/Developer/TraceInsights/Private/Insights/ViewModels/QuickFind.h
--- a/Engine/Source/Developer/TraceInsights/Private/Insights/ViewModels/QuickFind.h
+++ b/Engine/Source/Developer/TraceInsights/Private/Insights/ViewModels/QuickFind.h
@@ -51,6 +51,26 @@ private:
 	FOnFindPreviousEvent OnFindPreviousEvent;
 
 	////////////////////////////////////////////////////////////////////////////////////////////////////
+	// OnFilterAllEvent
+public:
+	/** The event to execute when the user selects the "Filter All" option. */
+	DECLARE_MULTICAST_DELEGATE(FOnFilterAllEvent);
+	FOnFilterAllEvent& GetOnFilterAllEvent() { return OnFilterAllEvent; }
+
+private:
+	FOnFilterAllEvent OnFilterAllEvent;
+
+	////////////////////////////////////////////////////////////////////////////////////////////////////
+	// OnClearAllFilters
+public:
+	/** The event to execute when the user selects the "Clear Filters" option. */
+	DECLARE_MULTICAST_DELEGATE(FOnClearFiltersEvent);
+	FOnClearFiltersEvent& GetOnClearFiltersEvent() { return OnClearFiltersEvent; }
+
+private:
+	FOnClearFiltersEvent OnClearFiltersEvent;
+
+	////////////////////////////////////////////////////////////////////////////////////////////////////
 
 private:
 	TSharedPtr<FFilterConfigurator> FilterConfigurator;
diff --git a/Engine/Source/Developer/TraceInsights/Private/Insights/ViewModels/ThreadTimingTrack.cpp b/Engine/Source/Developer/TraceInsights/Private/Insights/ViewModels/ThreadTimingTrack.cpp
--- a/Engine/Source/Developer/TraceInsights/Private/Insights/ViewModels/ThreadTimingTrack.cpp
+++ b/Engine/Source/Developer/TraceInsights/Private/Insights/ViewModels/ThreadTimingTrack.cpp
@@ -1637,4 +1637,20 @@ int32 FThreadTimingTrack::GetDepthAt(double Time) const
 
 ////////////////////////////////////////////////////////////////////////////////////////////////////
 
+void FThreadTimingTrack::SetFilterConfigurator(TSharedPtr<Insights::FFilterConfigurator> InFilterConfigurator)
+{
+	if (InFilterConfigurator.IsValid())
+	{
+		FilterConfigurator = MakeShared<Insights::FFilterConfigurator>(*InFilterConfigurator);
+	}
+	else
+	{
+		FilterConfigurator.Reset();
+	}
+
+	SetDirtyFlag();
+}
+
+////////////////////////////////////////////////////////////////////////////////////////////////////
+
 #undef LOCTEXT_NAMESPACE
diff --git a/Engine/Source/Developer/TraceInsights/Private/Insights/ViewModels/ThreadTimingTrack.h b/Engine/Source/Developer/TraceInsights/Private/Insights/ViewModels/ThreadTimingTrack.h
--- a/Engine/Source/Developer/TraceInsights/Private/Insights/ViewModels/ThreadTimingTrack.h
+++ b/Engine/Source/Developer/TraceInsights/Private/Insights/ViewModels/ThreadTimingTrack.h
@@ -162,6 +162,8 @@ public:
 
 	int32 GetDepthAt(double Time) const;
 
+	virtual void SetFilterConfigurator(TSharedPtr<Insights::FFilterConfigurator> InFilterConfigurator);
+
 protected:
 	virtual bool HasCustomFilter() const override;
 
diff --git a/Engine/Source/Developer/TraceInsights/Private/Insights/Widgets/SQuickFind.cpp b/Engine/Source/Developer/TraceInsights/Private/Insights/Widgets/SQuickFind.cpp
--- a/Engine/Source/Developer/TraceInsights/Private/Insights/Widgets/SQuickFind.cpp
+++ b/Engine/Source/Developer/TraceInsights/Private/Insights/Widgets/SQuickFind.cpp
@@ -98,6 +98,30 @@ void SQuickFind::Construct(const FArguments& InArgs, TSharedPtr<FQuickFind> InQu
 				.ToolTipText(LOCTEXT("FindNextDesc", "Find the next occurence that matches the search criteria."))
 				.OnClicked(this, &SQuickFind::FindNext_OnClicked)
 			]
+
+			+ SHorizontalBox::Slot()
+				.HAlign(HAlign_Left)
+				.VAlign(VAlign_Center)
+				.AutoWidth()
+				.Padding(1.0f)
+				[
+					SNew(SButton)
+					.Text(LOCTEXT("FilterAll", "Filter All"))
+					.ToolTipText(LOCTEXT("FilterAllDesc", "Filter all the tracks using the the search criteria."))
+					.OnClicked(this, &SQuickFind::FilterAll_OnClicked)
+				]
+
+			+ SHorizontalBox::Slot()
+				.HAlign(HAlign_Left)
+				.VAlign(VAlign_Center)
+				.AutoWidth()
+				.Padding(1.0f)
+				[
+					SNew(SButton)
+					.Text(LOCTEXT("ClearFilters", "Clear Filters"))
+					.ToolTipText(LOCTEXT("ClearFiltersDesc", "Clear all filters applied to the tracks."))
+					.OnClicked(this, &SQuickFind::ClearFilters_OnClicked)
+				]
 		]
 	];
 
@@ -122,7 +146,6 @@ void SQuickFind::Construct(const FArguments& InArgs, TSharedPtr<FQuickFind> InQu
 		];
 
 	QuickFindViewModel = InQuickFindViewModel;
-	OnViewModelDestroyedHandle = InQuickFindViewModel->GetFilterConfigurator()->GetOnDestroyedEvent().AddSP(this, &SQuickFind::RequestClose);
 }
 
 END_SLATE_FUNCTION_BUILD_OPTIMIZATION
@@ -147,23 +170,21 @@ FReply SQuickFind::FindPrevious_OnClicked()
 
 ////////////////////////////////////////////////////////////////////////////////////////////////////
 
-FReply SQuickFind::Cancel_OnClicked()
+FReply SQuickFind::FilterAll_OnClicked()
 {
-	RequestClose();
+	QuickFindViewModel->GetFilterConfigurator()->GetRootNode()->ProcessFilter();
+	QuickFindViewModel->GetOnFilterAllEvent().Broadcast();
 
 	return FReply::Handled();
 }
 
 ////////////////////////////////////////////////////////////////////////////////////////////////////
 
-void SQuickFind::RequestClose()
+FReply SQuickFind::ClearFilters_OnClicked()
 {
-	TSharedPtr<SDockTab> ParentTabSharedPtr = ParentTab.Pin();
-	if (ParentTabSharedPtr.IsValid())
-	{
-		ParentTabSharedPtr->RequestCloseTab();
-		ParentTabSharedPtr.Reset();
-	}
+	QuickFindViewModel->GetOnClearFiltersEvent().Broadcast();
+
+	return FReply::Handled();
 }
 
 ////////////////////////////////////////////////////////////////////////////////////////////////////
diff --git a/Engine/Source/Developer/TraceInsights/Private/Insights/Widgets/SQuickFind.h b/Engine/Source/Developer/TraceInsights/Private/Insights/Widgets/SQuickFind.h
--- a/Engine/Source/Developer/TraceInsights/Private/Insights/Widgets/SQuickFind.h
+++ b/Engine/Source/Developer/TraceInsights/Private/Insights/Widgets/SQuickFind.h
@@ -44,8 +44,6 @@ public:
 	void SetParentTab(const TSharedPtr<SDockTab> InTab) { ParentTab = InTab; }
 	const TWeakPtr<SDockTab> GetParentTab() { return ParentTab; };
 
-	void RequestClose();
-
 	static void RegisterQuickFindTab();
 
 	static TSharedPtr<SQuickFind> CreateAndOpenQuickFilterWidget(TSharedPtr<FQuickFind> InQuickFindViewModel);
@@ -57,7 +55,9 @@ private:
 
 	FReply FindPrevious_OnClicked();
 
-	FReply Cancel_OnClicked();
+	FReply FilterAll_OnClicked();
+
+	FReply ClearFilters_OnClicked();
 
 	static TSharedRef<SDockTab> SpawnTab(const FSpawnTabArgs& Args);
 
diff --git a/Engine/Source/Developer/TraceInsights/Private/Insights/Widgets/STimingView.cpp b/Engine/Source/Developer/TraceInsights/Private/Insights/Widgets/STimingView.cpp
--- a/Engine/Source/Developer/TraceInsights/Private/Insights/Widgets/STimingView.cpp
+++ b/Engine/Source/Developer/TraceInsights/Private/Insights/Widgets/STimingView.cpp
@@ -4200,6 +4200,8 @@ void STimingView::QuickFind_Execute()
 		QuickFindVm = MakeShared<FQuickFind>(FilterConfigurator);
 		QuickFindVm->GetOnFindNextEvent().AddSP(this, &STimingView::FindNextEvent);
 		QuickFindVm->GetOnFindPreviousEvent().AddSP(this, &STimingView::FindPrevEvent);
+		QuickFindVm->GetOnFilterAllEvent().AddSP(this, &STimingView::FilterAllTracks);
+		QuickFindVm->GetOnClearFiltersEvent().AddSP(this, &STimingView::ClearFilters);
 	}
 
 	SQuickFind::CreateAndOpenQuickFilterWidget(QuickFindVm);
@@ -4355,5 +4357,25 @@ void STimingView::FindPrevEvent()
 
 ////////////////////////////////////////////////////////////////////////////////////////////////////
 
+void STimingView::FilterAllTracks()
+{
+	for (auto& Entry : AllTracks)
+	{
+		Entry.Value->SetFilterConfigurator(QuickFindVm->GetFilterConfigurator());
+	}
+}
+
+////////////////////////////////////////////////////////////////////////////////////////////////////
+
+void STimingView::ClearFilters()
+{
+	for (auto& Entry : AllTracks)
+	{
+		Entry.Value->SetFilterConfigurator(nullptr);
+	}
+}
+
+////////////////////////////////////////////////////////////////////////////////////////////////////
+
 #undef ACTIVATE_BENCHMARK
 #undef LOCTEXT_NAMESPACE
diff --git a/Engine/Source/Developer/TraceInsights/Private/Insights/Widgets/STimingView.h b/Engine/Source/Developer/TraceInsights/Private/Insights/Widgets/STimingView.h
--- a/Engine/Source/Developer/TraceInsights/Private/Insights/Widgets/STimingView.h
+++ b/Engine/Source/Developer/TraceInsights/Private/Insights/Widgets/STimingView.h
@@ -447,6 +447,8 @@ protected:
 	
 	void FindNextEvent();
 	void FindPrevEvent();
+	void FilterAllTracks();
+	void ClearFilters();
 
 protected:
 	/** The track's viewport. Encapsulates info about position and scale. */
diff --git a/Engine/Source/Developer/TraceInsights/Public/Insights/ViewModels/BaseTimingTrack.h b/Engine/Source/Developer/TraceInsights/Public/Insights/ViewModels/BaseTimingTrack.h
--- a/Engine/Source/Developer/TraceInsights/Public/Insights/ViewModels/BaseTimingTrack.h
+++ b/Engine/Source/Developer/TraceInsights/Public/Insights/ViewModels/BaseTimingTrack.h
@@ -13,6 +13,7 @@
 
 class FMenuBuilder;
 namespace TraceServices { class IAnalysisSession; };
+namespace Insights { class FFilterConfigurator; }
 
 class FTimingEventSearchParameters;
 class FTimingTrackViewport;
@@ -262,6 +263,8 @@ public:
 	TWeakPtr<FBaseTimingTrack> GetParentTrack() const { return ParentTrack; }
 	bool IsChildTrack() const { return ParentTrack.IsValid(); }
 
+	virtual void SetFilterConfigurator(TSharedPtr<Insights::FFilterConfigurator> InFilterConfigurator) {}
+
 	// Returns number of text lines needed to display the debug string.
 	//TODO: virtual int GetDebugStringLineCount() const { return 0; }
 
