[Insights] Remove dead code

diff --git a/Engine/Source/Developer/TraceInsights/Private/Insights/Widgets/STimingView.h b/Engine/Source/Developer/TraceInsights/Private/Insights/Widgets/STimingView.h
--- a/Engine/Source/Developer/TraceInsights/Private/Insights/Widgets/STimingView.h
+++ b/Engine/Source/Developer/TraceInsights/Private/Insights/Widgets/STimingView.h
@@ -439,10 +439,6 @@ protected:
 	FReply AllowTracksToProcessOnMouseButtonUp(const FGeometry& MyGeometry, const FPointerEvent& MouseEvent);
 	FReply AllowTracksToProcessOnMouseButtonDoubleClick(const FGeometry& MyGeometry, const FPointerEvent& MouseEvent);
 
-	void ContextMenu_ShowTaskDependecies_Execute();
-	bool ContextMenu_ShowTaskDependecies_CanExecute();
-	bool ContextMenu_ShowTaskDependecies_IsChecked();
-
 	void SetTrackPosY(TSharedPtr<FBaseTimingTrack>& TrackPtr, float TopOffsetY) const;
 	
 	void FindNextEvent();
