[Insights]
Networking Insights: Fixed the drawing of borders for packet content events when scale factor is greater than 100%.
Networking Insights: Fixed the drawing of borders for hovered and selected packet content event (2px).
This fix is similar with one for timing events (#jira UE-75607).

diff --git a/Engine/Source/Developer/TraceInsights/Private/Insights/NetworkingProfiler/ViewModels/PacketContentViewDrawHelper.cpp b/Engine/Source/Developer/TraceInsights/Private/Insights/NetworkingProfiler/ViewModels/PacketContentViewDrawHelper.cpp
--- a/Engine/Source/Developer/TraceInsights/Private/Insights/NetworkingProfiler/ViewModels/PacketContentViewDrawHelper.cpp
+++ b/Engine/Source/Developer/TraceInsights/Private/Insights/NetworkingProfiler/ViewModels/PacketContentViewDrawHelper.cpp
@@ -17,6 +17,8 @@
 #include <limits>
 #include "Misc/StringBuilder.h"
 
+#define INSIGHTS_USE_LEGACY_BORDER 0
+
 ////////////////////////////////////////////////////////////////////////////////////////////////////
 // FPacketContentViewDrawStateBuilder
 ////////////////////////////////////////////////////////////////////////////////////////////////////
@@ -397,6 +399,7 @@ void FPacketContentViewDrawHelper::Draw(const FPacketContentViewDrawState& DrawS
 	// Draw borders.
 	//if (LayoutEventH > 0.0f)
 	{
+#if INSIGHTS_USE_LEGACY_BORDER
 		const float EventBorderH = LayoutEventH;
 		for (const FPacketContentViewDrawState::FBoxPrimitive& Box : DrawState.Borders)
 		{
@@ -404,6 +407,29 @@ void FPacketContentViewDrawHelper::Draw(const FPacketContentViewDrawState& DrawS
 			DrawContext.DrawBox(Box.X, Y, Box.W, EventBorderH, EventBorderBrush, Box.Color.CopyWithNewOpacity(Opacity));
 		}
 		DrawContext.LayerId++;
+#else
+		if (Opacity == 1.0f)
+		{
+			const int32 EventBorderLayerId = DrawContext.LayerId - 2; // draw under the filled boxes
+			const float EventBorderH = LayoutEventH;
+			for (const FPacketContentViewDrawState::FBoxPrimitive& Box : DrawState.Borders)
+			{
+				const float Y = LayoutPosY + (LayoutEventH + LayoutEventDY) * Box.Depth;
+				DrawContext.DrawBox(EventBorderLayerId, Box.X, Y, Box.W, EventBorderH, WhiteBrush, Box.Color);
+			}
+		}
+		else
+		{
+			const float EventBorderH = LayoutEventH;
+			for (const FPacketContentViewDrawState::FBoxPrimitive& Box : DrawState.Borders)
+			{
+				const float Y = LayoutPosY + (LayoutEventH + LayoutEventDY) * Box.Depth;
+				FSlateRoundedBoxBrush Brush(FLinearColor(0.0f, 0.0f, 0.0f, 0.0f), 0.0f, Box.Color.CopyWithNewOpacity(Opacity), 1.0f);
+				DrawContext.DrawBox(Box.X, Y, Box.W, EventBorderH, &Brush, FLinearColor(0.0f, 0.0f, 0.0f, 0.0f));
+			}
+			DrawContext.LayerId++;
+		}
+#endif
 	}
 
 	// Draw texts.
@@ -462,12 +488,23 @@ void FPacketContentViewDrawHelper::DrawEventHighlight(const FNetworkPacketEvent&
 	const float EventW = EventX2 - EventX1;
 	const float EventY = LayoutPosY + (LayoutEventH + LayoutEventDY) * Event.Level;
 
+#if INSIGHTS_USE_LEGACY_BORDER
+	constexpr float BorderOffset = 1.0f;
+#else
+	constexpr float BorderOffset = 2.0f;
+#endif
+
 	if (Mode == EHighlightMode::Hovered)
 	{
 		const FLinearColor Color(1.0f, 1.0f, 0.0f, 1.0f); // yellow
 
 		// Draw border around the timing event box.
-		DrawContext.DrawBox(EventX1 - 2.0f, EventY - 2.0f, EventW + 4.0f, LayoutEventH + 4.0f, HoveredEventBorderBrush, Color);
+#if INSIGHTS_USE_LEGACY_BORDER
+		DrawContext.DrawBox(EventX1 - BorderOffset, EventY - BorderOffset, EventW + 2 * BorderOffset, LayoutEventH + 2 * BorderOffset, HoveredEventBorderBrush, Color);
+#else
+		FSlateRoundedBoxBrush Brush(FLinearColor(0.0f, 0.0f, 0.0f, 0.0f), 2.0f, Color, 2.0f);
+		DrawContext.DrawBox(EventX1 - BorderOffset, EventY - BorderOffset, EventW + 2 * BorderOffset, LayoutEventH + 2 * BorderOffset, &Brush, FLinearColor(0.0f, 0.0f, 0.0f, 0.0f));
+#endif
 	}
 	else // EHighlightMode::Selected or EHighlightMode::SelectedAndHovered
 	{
@@ -479,9 +516,16 @@ void FPacketContentViewDrawHelper::DrawEventHighlight(const FNetworkPacketEvent&
 		const FLinearColor Color(S, S, Blue, 1.0f);
 
 		// Draw border around the timing event box.
-		DrawContext.DrawBox(EventX1 - 2.0f, EventY - 2.0f, EventW + 4.0f, LayoutEventH + 4.0f, SelectedEventBorderBrush, Color);
+#if INSIGHTS_USE_LEGACY_BORDER
+		DrawContext.DrawBox(EventX1 - BorderOffset, EventY - BorderOffset, EventW + 2 * BorderOffset, LayoutEventH + 2 * BorderOffset, SelectedEventBorderBrush, Color);
+#else
+		FSlateRoundedBoxBrush Brush(FLinearColor(0.0f, 0.0f, 0.0f, 0.0f), 2.0f, Color, 2.0f);
+		DrawContext.DrawBox(EventX1 - BorderOffset, EventY - BorderOffset, EventW + 2 * BorderOffset, LayoutEventH + 2 * BorderOffset, &Brush, FLinearColor(0.0f, 0.0f, 0.0f, 0.0f));
+#endif
 	}
 	DrawContext.LayerId++;
 }
 
 ////////////////////////////////////////////////////////////////////////////////////////////////////
+
+#undef INSIGHTS_USE_LEGACY_BORDER
