[Insights] Added a debug "Starship Test Suite" option in menu of Timing Insights and in the Open menu from Session Browser.

diff --git a/Engine/Source/Developer/TraceInsights/Private/Insights/Common/InsightsMenuBuilder.cpp b/Engine/Source/Developer/TraceInsights/Private/Insights/Common/InsightsMenuBuilder.cpp
--- a/Engine/Source/Developer/TraceInsights/Private/Insights/Common/InsightsMenuBuilder.cpp
+++ b/Engine/Source/Developer/TraceInsights/Private/Insights/Common/InsightsMenuBuilder.cpp
@@ -4,9 +4,12 @@
 
 #include "Framework/Docking/TabManager.h"
 #include "Framework/MultiBox/MultiBoxBuilder.h"
+#include "Widgets/Testing/SStarshipSuite.h"
 #include "WorkspaceMenuStructure.h"
 #include "WorkspaceMenuStructureModule.h"
 
+#include "Insights/InsightsStyle.h"
+
 #define LOCTEXT_NAMESPACE "InsightsMenuBuilder"
 
 ////////////////////////////////////////////////////////////////////////////////////////////////////
@@ -60,7 +63,27 @@ void FInsightsMenuBuilder::PopulateMenu(FMenuBuilder& MenuBuilder)
 			FUIAction(FExecuteAction::CreateLambda([=] { FGlobalTabmanager::Get()->TryInvokeTab(WidgetReflectorTabId); })));
 		MenuBuilder.EndSection();
 	}
-#endif
+
+#if !UE_BUILD_SHIPPING
+	// Open Starship Test Suite
+	{
+		FUIAction OpenStarshipSuiteAction;
+		OpenStarshipSuiteAction.ExecuteAction = FExecuteAction::CreateLambda([this]()
+			{
+				RestoreStarshipSuite();
+			});
+
+		MenuBuilder.AddMenuEntry(
+			LOCTEXT("OpenStarshipSuite", "Starship Test Suite"),
+			LOCTEXT("OpenStarshipSuiteDesc", "Opens the Starship UX test suite."),
+			FSlateIcon(FInsightsStyle::GetStyleSetName(), "Icon.Bug"),
+			OpenStarshipSuiteAction,
+			NAME_None,
+			EUserInterfaceActionType::Button
+		);
+	}
+#endif // !UE_BUILD_SHIPPING
+#endif // !WITH_EDITOR
 }
 
 ////////////////////////////////////////////////////////////////////////////////////////////////////
diff --git a/Engine/Source/Developer/TraceInsights/Private/Insights/InsightsStyle.cpp b/Engine/Source/Developer/TraceInsights/Private/Insights/InsightsStyle.cpp
--- a/Engine/Source/Developer/TraceInsights/Private/Insights/InsightsStyle.cpp
+++ b/Engine/Source/Developer/TraceInsights/Private/Insights/InsightsStyle.cpp
@@ -162,6 +162,8 @@ TSharedRef<FSlateStyleSet> FInsightsStyle::Create()
 	Style.Set("FindNext", new IMAGE_BRUSH("Animation/forward", Icon20x20));
 	Style.Set("FindLast", new IMAGE_BRUSH("Animation/forward_end", Icon20x20));
 
+	Style.Set("Icon.Bug", new IMAGE_BRUSH_SVG("../../Editor/Slate/Starship/Common/bug", Icon24x24));
+
 	return StyleRef;
 }
 
diff --git a/Engine/Source/Developer/TraceInsights/Private/Insights/Widgets/SStartPageWindow.cpp b/Engine/Source/Developer/TraceInsights/Private/Insights/Widgets/SStartPageWindow.cpp
--- a/Engine/Source/Developer/TraceInsights/Private/Insights/Widgets/SStartPageWindow.cpp
+++ b/Engine/Source/Developer/TraceInsights/Private/Insights/Widgets/SStartPageWindow.cpp
@@ -36,8 +36,7 @@
 #include "Widgets/Input/SButton.h"
 #include "Widgets/Input/SEditableTextBox.h"
 #include "Widgets/Input/SComboButton.h"
-//#include "WorkspaceMenuStructure.h"
-//#include "WorkspaceMenuStructureModule.h"
+#include "Widgets/Testing/SStarshipSuite.h"
 
 #if WITH_EDITOR
 	#include "EngineAnalytics.h"
@@ -1838,12 +1837,9 @@ TSharedRef<SWidget> SStartPageWindow::MakeTraceListMenu()
 
 	MenuBuilder.BeginSection("DebugOptions", LOCTEXT("OptionsHeading", "Debug Options"));
 
-	auto CanExecute = []() {return true; };
-
 	// Enable Automation Tests Option.
 	{
 		FUIAction ToogleAutomationTestsAction;
-		ToogleAutomationTestsAction.CanExecuteAction = FCanExecuteAction::CreateLambda(CanExecute);
 		ToogleAutomationTestsAction.ExecuteAction = FExecuteAction::CreateLambda([this]()
 			{
 				this->SetEnableAutomaticTesting(!this->GetEnableAutomaticTesting());
@@ -1855,7 +1851,7 @@ TSharedRef<SWidget> SStartPageWindow::MakeTraceListMenu()
 
 		MenuBuilder.AddMenuEntry(
 			LOCTEXT("EnableAutomatedTesting", "Enable Automation Testing"),
-			LOCTEXT("EnableAutomatedTestingDesc", "Activates the automatic test system for sessions opened from this window."),
+			LOCTEXT("EnableAutomatedTestingDesc", "Activates the automatic test system for new sessions opened from this window."),
 			FSlateIcon(),
 			ToogleAutomationTestsAction,
 			NAME_None,
@@ -1866,7 +1862,6 @@ TSharedRef<SWidget> SStartPageWindow::MakeTraceListMenu()
 	// Enable Debug Tools Option.
 	{
 		FUIAction ToogleDebugToolsAction;
-		ToogleDebugToolsAction.CanExecuteAction = FCanExecuteAction::CreateLambda(CanExecute);
 		ToogleDebugToolsAction.ExecuteAction = FExecuteAction::CreateLambda([this]()
 			{
 				this->SetEnableDebugTools(!this->GetEnableDebugTools());
@@ -1878,7 +1873,7 @@ TSharedRef<SWidget> SStartPageWindow::MakeTraceListMenu()
 
 		MenuBuilder.AddMenuEntry(
 			LOCTEXT("EnableDebugTools", "Enable Debug Tools"),
-			LOCTEXT("EnableDebugToolsDesc", "Enables debug tools for sessions opened from this window."),
+			LOCTEXT("EnableDebugToolsDesc", "Enables debug tools for new sessions opened from this window."),
 			FSlateIcon(),
 			ToogleDebugToolsAction,
 			NAME_None,
@@ -1886,6 +1881,26 @@ TSharedRef<SWidget> SStartPageWindow::MakeTraceListMenu()
 		);
 	}
 
+#if !UE_BUILD_SHIPPING
+	// Open Starship Test Suite
+	{
+		FUIAction OpenStarshipSuiteAction;
+		OpenStarshipSuiteAction.ExecuteAction = FExecuteAction::CreateLambda([this]()
+			{
+				RestoreStarshipSuite();
+			});
+
+		MenuBuilder.AddMenuEntry(
+			LOCTEXT("OpenStarshipSuite", "Starship Test Suite"),
+			LOCTEXT("OpenStarshipSuiteDesc", "Opens the Starship UX test suite."),
+			FSlateIcon(FInsightsStyle::GetStyleSetName(), "Icon.Bug"),
+			OpenStarshipSuiteAction,
+			NAME_None,
+			EUserInterfaceActionType::Button
+		);
+	}
+#endif // !UE_BUILD_SHIPPING
+
 	MenuBuilder.EndSection();
 
 	return MenuBuilder.MakeWidget();
