[Insights] Enable each Timing View to create its own SQuickFindTab. Prevent SQuickFindTab from beign saved in the "open" state in the layout.

diff --git a/Engine/Source/Developer/TraceInsights/Private/Insights/LoadingProfiler/LoadingProfilerManager.cpp b/Engine/Source/Developer/TraceInsights/Private/Insights/LoadingProfiler/LoadingProfilerManager.cpp
--- a/Engine/Source/Developer/TraceInsights/Private/Insights/LoadingProfiler/LoadingProfilerManager.cpp
+++ b/Engine/Source/Developer/TraceInsights/Private/Insights/LoadingProfiler/LoadingProfilerManager.cpp
@@ -379,4 +379,19 @@ void FLoadingProfilerManager::ShowHideRequestsTreeView(const bool bIsVisible)
 
 ////////////////////////////////////////////////////////////////////////////////////////////////////
 
+void FLoadingProfilerManager::OnWindowClosedEvent()
+{
+	TSharedPtr<SLoadingProfilerWindow> Wnd = GetProfilerWindow();
+	if (Wnd)
+	{
+		TSharedPtr<STimingView> TimingView = Wnd->GetTimingView();
+		if (TimingView.IsValid())
+		{
+			TimingView->CloseQuickFindTab();
+		}
+	}
+}
+
+////////////////////////////////////////////////////////////////////////////////////////////////////
+
 #undef LOCTEXT_NAMESPACE
diff --git a/Engine/Source/Developer/TraceInsights/Private/Insights/LoadingProfiler/LoadingProfilerManager.h b/Engine/Source/Developer/TraceInsights/Private/Insights/LoadingProfiler/LoadingProfilerManager.h
--- a/Engine/Source/Developer/TraceInsights/Private/Insights/LoadingProfiler/LoadingProfilerManager.h
+++ b/Engine/Source/Developer/TraceInsights/Private/Insights/LoadingProfiler/LoadingProfilerManager.h
@@ -50,6 +50,7 @@ public:
 	virtual void Shutdown() override;
 	virtual void RegisterMajorTabs(IUnrealInsightsModule& InsightsModule) override;
 	virtual void UnregisterMajorTabs() override;
+	virtual void OnWindowClosedEvent() override;
 
 	//////////////////////////////////////////////////
 
diff --git a/Engine/Source/Developer/TraceInsights/Private/Insights/MemoryProfiler/MemoryProfilerManager.cpp b/Engine/Source/Developer/TraceInsights/Private/Insights/MemoryProfiler/MemoryProfilerManager.cpp
--- a/Engine/Source/Developer/TraceInsights/Private/Insights/MemoryProfiler/MemoryProfilerManager.cpp
+++ b/Engine/Source/Developer/TraceInsights/Private/Insights/MemoryProfiler/MemoryProfilerManager.cpp
@@ -377,6 +377,12 @@ void FMemoryProfilerManager::OnWindowClosedEvent()
 	if (Wnd.IsValid())
 	{
 		Wnd->CloseMemAllocTableTreeTabs();
+
+		TSharedPtr<STimingView> TimingView = Wnd->GetTimingView();
+		if (TimingView.IsValid())
+		{
+			TimingView->CloseQuickFindTab();
+		}
 	}
 }
 
diff --git a/Engine/Source/Developer/TraceInsights/Private/Insights/TimingProfilerManager.cpp b/Engine/Source/Developer/TraceInsights/Private/Insights/TimingProfilerManager.cpp
--- a/Engine/Source/Developer/TraceInsights/Private/Insights/TimingProfilerManager.cpp
+++ b/Engine/Source/Developer/TraceInsights/Private/Insights/TimingProfilerManager.cpp
@@ -611,4 +611,19 @@ void FTimingProfilerManager::UpdateAggregatedCounterStats()
 
 ////////////////////////////////////////////////////////////////////////////////////////////////////
 
+void FTimingProfilerManager::OnWindowClosedEvent()
+{
+	TSharedPtr<STimingProfilerWindow> Wnd = GetProfilerWindow();
+	if (Wnd)
+	{
+		TSharedPtr<STimingView> TimingView = Wnd->GetTimingView();
+		if (TimingView.IsValid())
+		{
+			TimingView->CloseQuickFindTab();
+		}
+	}
+}
+
+////////////////////////////////////////////////////////////////////////////////////////////////////
+
 #undef LOCTEXT_NAMESPACE
diff --git a/Engine/Source/Developer/TraceInsights/Private/Insights/TimingProfilerManager.h b/Engine/Source/Developer/TraceInsights/Private/Insights/TimingProfilerManager.h
--- a/Engine/Source/Developer/TraceInsights/Private/Insights/TimingProfilerManager.h
+++ b/Engine/Source/Developer/TraceInsights/Private/Insights/TimingProfilerManager.h
@@ -62,6 +62,7 @@ public:
 	virtual void Shutdown() override;
 	virtual void RegisterMajorTabs(IUnrealInsightsModule& InsightsModule) override;
 	virtual void UnregisterMajorTabs() override;
+	virtual void OnWindowClosedEvent() override;
 
 	//////////////////////////////////////////////////
 
diff --git a/Engine/Source/Developer/TraceInsights/Private/Insights/Widgets/SQuickFind.cpp b/Engine/Source/Developer/TraceInsights/Private/Insights/Widgets/SQuickFind.cpp
--- a/Engine/Source/Developer/TraceInsights/Private/Insights/Widgets/SQuickFind.cpp
+++ b/Engine/Source/Developer/TraceInsights/Private/Insights/Widgets/SQuickFind.cpp
@@ -28,11 +28,6 @@ namespace Insights
 // SQuickFind
 ////////////////////////////////////////////////////////////////////////////////////////////////////
 
-FName const SQuickFind::QuickFindTabId(TEXT("QuickFilter"));
-TSharedPtr<SQuickFind> SQuickFind::PendingWidget;
-bool SQuickFind::bIsTabRegistered = false;
-
-
 SQuickFind::SQuickFind()
 {
 }
@@ -189,56 +184,6 @@ FReply SQuickFind::ClearFilters_OnClicked()
 
 ////////////////////////////////////////////////////////////////////////////////////////////////////
 
-void SQuickFind::RegisterQuickFindTab()
-{
-	if (bIsTabRegistered)
-	{
-		return;
-	}
-
-	FTabSpawnerEntry& TabSpawnerEntry = FGlobalTabmanager::Get()->RegisterNomadTabSpawner(QuickFindTabId,
-		FOnSpawnTab::CreateStatic(&SQuickFind::SpawnTab))
-		.SetDisplayName(LOCTEXT("FilterConfiguratorTabTitle", "Filter Configurator"))
-		.SetMenuType(ETabSpawnerMenuType::Hidden);
-
-	bIsTabRegistered = true;
-}
-
-////////////////////////////////////////////////////////////////////////////////////////////////////
-
-TSharedRef<SDockTab> SQuickFind::SpawnTab(const FSpawnTabArgs& Args)
-{
-	const TSharedRef<SDockTab> DockTab = SNew(SDockTab)
-		.TabRole(ETabRole::NomadTab);
-
-	if (!PendingWidget.IsValid())
-	{
-		return DockTab;
-	}
-
-	DockTab->SetContent(PendingWidget.ToSharedRef());
-	PendingWidget->SetParentTab(DockTab);
-
-	PendingWidget = nullptr;
-	return DockTab;
-}
-
-////////////////////////////////////////////////////////////////////////////////////////////////////
-
-TSharedPtr<SQuickFind> SQuickFind::CreateAndOpenQuickFilterWidget(TSharedPtr<FQuickFind> InQuickFindViewModel)
-{
-	SAssignNew(PendingWidget, SQuickFind, InQuickFindViewModel);
-
-	if (FGlobalTabmanager::Get()->HasTabSpawner(QuickFindTabId))
-	{
-		FGlobalTabmanager::Get()->TryInvokeTab(QuickFindTabId);
-	}
-
-	return PendingWidget;
-}
-
-////////////////////////////////////////////////////////////////////////////////////////////////////
-
 } // namespace Insights
 
 #undef LOCTEXT_NAMESPACE
diff --git a/Engine/Source/Developer/TraceInsights/Private/Insights/Widgets/SQuickFind.h b/Engine/Source/Developer/TraceInsights/Private/Insights/Widgets/SQuickFind.h
--- a/Engine/Source/Developer/TraceInsights/Private/Insights/Widgets/SQuickFind.h
+++ b/Engine/Source/Developer/TraceInsights/Private/Insights/Widgets/SQuickFind.h
@@ -44,10 +44,6 @@ public:
 	void SetParentTab(const TSharedPtr<SDockTab> InTab) { ParentTab = InTab; }
 	const TWeakPtr<SDockTab> GetParentTab() { return ParentTab; };
 
-	static void RegisterQuickFindTab();
-
-	static TSharedPtr<SQuickFind> CreateAndOpenQuickFilterWidget(TSharedPtr<FQuickFind> InQuickFindViewModel);
-
 private:
 	void InitCommandList();
 
@@ -59,15 +55,7 @@ private:
 
 	FReply ClearFilters_OnClicked();
 
-	static TSharedRef<SDockTab> SpawnTab(const FSpawnTabArgs& Args);
-
 private:
-	static const FName QuickFindTabId;
-
-	static TSharedPtr<SQuickFind> PendingWidget;
-
-	static bool bIsTabRegistered;
-
 	TSharedPtr<SFilterConfigurator> FilterConfigurator;
 
 	TSharedPtr<FQuickFind> QuickFindViewModel;
diff --git a/Engine/Source/Developer/TraceInsights/Private/Insights/Widgets/STimingView.cpp b/Engine/Source/Developer/TraceInsights/Private/Insights/Widgets/STimingView.cpp
--- a/Engine/Source/Developer/TraceInsights/Private/Insights/Widgets/STimingView.cpp
+++ b/Engine/Source/Developer/TraceInsights/Private/Insights/Widgets/STimingView.cpp
@@ -21,6 +21,7 @@
 #include "Styling/CoreStyle.h"
 #include "Styling/SlateBrush.h"
 #include "TraceServices/AnalysisService.h"
+#include "Widgets/Docking/SDockTab.h"
 #include "Widgets/Input/SButton.h"
 #include "Widgets/Input/SCheckBox.h"
 #include "Widgets/Input/SComboButton.h"
@@ -75,6 +76,8 @@
 // start auto generated ids from a big number (MSB set to 1) to avoid collisions with ids for gpu/cpu tracks based on 32bit timeline index
 uint64 FBaseTimingTrack::IdGenerator = (1ULL << 63);
 
+uint32 STimingView::TimingViewId = 0;
+
 const TCHAR* GetFileActivityTypeName(TraceServices::EFileActivityType Type);
 uint32 GetFileActivityTypeColor(TraceServices::EFileActivityType Type);
 
@@ -95,6 +98,7 @@ STimingView::STimingView()
 	, GraphTrack(MakeShared<FTimingGraphTrack>())
 	, WhiteBrush(FInsightsStyle::Get().GetBrush("WhiteBrush"))
 	, MainFont(FCoreStyle::GetDefaultFontStyle("Regular", 8))
+	, QuickFindTabId(TEXT("QuickFind"), TimingViewId++)
 {
 	DefaultTimeMarker->SetName(TEXT(""));
 	DefaultTimeMarker->SetColor(FLinearColor(0.85f, 0.5f, 0.03f, 0.5f));
@@ -249,11 +253,15 @@ void STimingView::Construct(const FArguments& InArgs)
 		]
 	];
 
-	Insights::SQuickFind::RegisterQuickFindTab();
 	UpdateHorizontalScrollBar();
 	UpdateVerticalScrollBar();
 
 	BindCommands();
+
+	FTabSpawnerEntry& TabSpawnerEntry = FGlobalTabmanager::Get()->RegisterNomadTabSpawner(QuickFindTabId,
+		FOnSpawnTab::CreateSP(this, &STimingView::SpawnQuickFindTab))
+		.SetDisplayName(LOCTEXT("QuickFindTabTitle", "Quick Find"))
+		.SetMenuType(ETabSpawnerMenuType::Hidden);
 }
 
 ////////////////////////////////////////////////////////////////////////////////////////////////////
@@ -4204,7 +4212,46 @@ void STimingView::QuickFind_Execute()
 		QuickFindVm->GetOnClearFiltersEvent().AddSP(this, &STimingView::ClearFilters);
 	}
 
-	SQuickFind::CreateAndOpenQuickFilterWidget(QuickFindVm);
+	SAssignNew(QuickFindWidgetSharedPtr, SQuickFind, QuickFindVm);
+
+	if (FGlobalTabmanager::Get()->HasTabSpawner(QuickFindTabId))
+	{
+		FGlobalTabmanager::Get()->TryInvokeTab(QuickFindTabId);
+	}
+}
+
+////////////////////////////////////////////////////////////////////////////////////////////////////
+
+TSharedRef<SDockTab> STimingView::SpawnQuickFindTab(const FSpawnTabArgs& Args)
+{
+	const TSharedRef<SDockTab> DockTab = SNew(SDockTab)
+		.TabRole(ETabRole::NomadTab);
+
+	if (!QuickFindWidgetSharedPtr.IsValid())
+	{
+		return DockTab;
+	}
+
+	DockTab->SetContent(QuickFindWidgetSharedPtr.ToSharedRef());
+	QuickFindWidgetSharedPtr->SetParentTab(DockTab);
+	QuickFindWidgetWeakPtr = QuickFindWidgetSharedPtr;
+	QuickFindWidgetSharedPtr.Reset();
+	return DockTab;
+}
+
+////////////////////////////////////////////////////////////////////////////////////////////////////
+
+void STimingView::CloseQuickFindTab()
+{
+	TSharedPtr<Insights::SQuickFind> QuickFindWidget = QuickFindWidgetWeakPtr.Pin();
+	if (QuickFindWidget)
+	{
+		TSharedPtr<SDockTab> QuickFindTab = QuickFindWidget->GetParentTab().Pin();
+		if (QuickFindTab.IsValid())
+		{
+			QuickFindTab->RequestCloseTab();
+		}
+	}
 }
 
 ////////////////////////////////////////////////////////////////////////////////////////////////////
diff --git a/Engine/Source/Developer/TraceInsights/Private/Insights/Widgets/STimingView.h b/Engine/Source/Developer/TraceInsights/Private/Insights/Widgets/STimingView.h
--- a/Engine/Source/Developer/TraceInsights/Private/Insights/Widgets/STimingView.h
+++ b/Engine/Source/Developer/TraceInsights/Private/Insights/Widgets/STimingView.h
@@ -29,11 +29,13 @@ class FFrameSharedState;
 class FLoadingSharedState;
 class FMarkersTimingTrack;
 class FMenuBuilder;
+class FSpawnTabArgs;
 class FThreadTimingSharedState;
 class FTimeRulerTrack;
 class FTimingGraphTrack;
 class FTimingViewDrawHelper;
 class FUICommandList;
+class SDockTab;
 class SOverlay;
 class SScrollBar;
 
@@ -42,6 +44,7 @@ namespace Insights
 	class ITimingViewExtender;
 	class FTimeMarker;
 	class FQuickFind;
+	class SQuickFind;
 }
 
 /** A custom widget used to display timing events. */
@@ -355,6 +358,8 @@ public:
 
 	TSharedPtr<FUICommandList> GetCommandList() { return CommandList; }
 
+	void CloseQuickFindTab();
+
 protected:
 	virtual FVector2D ComputeDesiredSize(float) const override
 	{
@@ -446,6 +451,8 @@ protected:
 	void FilterAllTracks();
 	void ClearFilters();
 
+	TSharedRef<SDockTab> SpawnQuickFindTab(const FSpawnTabArgs& Args);
+
 protected:
 	/** The track's viewport. Encapsulates info about position and scale. */
 	FTimingTrackViewport Viewport;
@@ -641,4 +648,10 @@ protected:
 	TArray<TUniquePtr<ITimingEventRelation>> CurrentRelations;
 
 	TSharedPtr<Insights::FQuickFind> QuickFindVm;
+	static uint32 TimingViewId;
+	const FName QuickFindTabId;
+	
+	// Used only between the creation of the widget and the spawning of the owning tab. When the tab is spawned, we relinquish ownership.
+	TSharedPtr<Insights::SQuickFind> QuickFindWidgetSharedPtr;
+	TWeakPtr<Insights::SQuickFind> QuickFindWidgetWeakPtr;
 };
\ No newline at end of file
