[Insights] Changed the tree views for Timers, Counters, Callers, Callees and NetStats panels to fill the horizontal space.

diff --git a/Engine/Source/Developer/TraceInsights/Private/Insights/NetworkingProfiler/Widgets/SNetStatsView.cpp b/Engine/Source/Developer/TraceInsights/Private/Insights/NetworkingProfiler/Widgets/SNetStatsView.cpp
--- a/Engine/Source/Developer/TraceInsights/Private/Insights/NetworkingProfiler/Widgets/SNetStatsView.cpp
+++ b/Engine/Source/Developer/TraceInsights/Private/Insights/NetworkingProfiler/Widgets/SNetStatsView.cpp
@@ -202,27 +202,21 @@ void SNetStatsView::Construct(const FArguments& InArgs, TSharedPtr<SNetworkingPr
 			.FillWidth(1.0f)
 			.Padding(0.0f)
 			[
-				SNew(SScrollBox)
-				.Orientation(Orient_Horizontal)
-
-				+ SScrollBox::Slot()
-				[
-					SAssignNew(TreeView, STreeView<FNetEventNodePtr>)
-					.ExternalScrollbar(ExternalScrollbar)
-					.SelectionMode(ESelectionMode::Multi)
-					.TreeItemsSource(&FilteredGroupNodes)
-					.OnGetChildren(this, &SNetStatsView::TreeView_OnGetChildren)
-					.OnGenerateRow(this, &SNetStatsView::TreeView_OnGenerateRow)
-					.OnSelectionChanged(this, &SNetStatsView::TreeView_OnSelectionChanged)
-					.OnMouseButtonDoubleClick(this, &SNetStatsView::TreeView_OnMouseButtonDoubleClick)
-					.OnContextMenuOpening(FOnContextMenuOpening::CreateSP(this, &SNetStatsView::TreeView_GetMenuContent))
-					.ItemHeight(12.0f)
-					.HeaderRow
-					(
-						SAssignNew(TreeViewHeaderRow, SHeaderRow)
-						.Visibility(EVisibility::Visible)
-					)
-				]
+				SAssignNew(TreeView, STreeView<FNetEventNodePtr>)
+				.ExternalScrollbar(ExternalScrollbar)
+				.SelectionMode(ESelectionMode::Multi)
+				.TreeItemsSource(&FilteredGroupNodes)
+				.OnGetChildren(this, &SNetStatsView::TreeView_OnGetChildren)
+				.OnGenerateRow(this, &SNetStatsView::TreeView_OnGenerateRow)
+				.OnSelectionChanged(this, &SNetStatsView::TreeView_OnSelectionChanged)
+				.OnMouseButtonDoubleClick(this, &SNetStatsView::TreeView_OnMouseButtonDoubleClick)
+				.OnContextMenuOpening(FOnContextMenuOpening::CreateSP(this, &SNetStatsView::TreeView_GetMenuContent))
+				.ItemHeight(12.0f)
+				.HeaderRow
+				(
+					SAssignNew(TreeViewHeaderRow, SHeaderRow)
+					.Visibility(EVisibility::Visible)
+				)
 			]
 
 			+ SHorizontalBox::Slot()
@@ -1310,8 +1304,8 @@ void SNetStatsView::ShowColumn(const FName ColumnId)
 		.VAlignCell(VAlign_Fill)
 		.SortMode(this, &SNetStatsView::GetSortModeForColumn, Column.GetId())
 		.OnSort(this, &SNetStatsView::OnSortModeChanged)
-		.ManualWidth(Column.GetInitialWidth())
-		.FixedWidth(Column.IsFixedWidth() ? Column.GetInitialWidth() : TOptional<float>())
+		.FillWidth(Column.GetInitialWidth())
+		//.FixedWidth(Column.IsFixedWidth() ? Column.GetInitialWidth() : TOptional<float>())
 		.HeaderContent()
 		[
 			SNew(SBox)
diff --git a/Engine/Source/Developer/TraceInsights/Private/Insights/Widgets/SStatsView.cpp b/Engine/Source/Developer/TraceInsights/Private/Insights/Widgets/SStatsView.cpp
--- a/Engine/Source/Developer/TraceInsights/Private/Insights/Widgets/SStatsView.cpp
+++ b/Engine/Source/Developer/TraceInsights/Private/Insights/Widgets/SStatsView.cpp
@@ -235,41 +235,35 @@ void SStatsView::Construct(const FArguments& InArgs)
 			.FillWidth(1.0f)
 			.Padding(0.0f)
 			[
-				SNew(SScrollBox)
-				.Orientation(Orient_Horizontal)
+				SNew(SOverlay)
 
-				+ SScrollBox::Slot()
+				+ SOverlay::Slot()
+				.HAlign(HAlign_Fill)
+				.VAlign(VAlign_Fill)
 				[
-					SNew(SOverlay)
-
-					+ SOverlay::Slot()
-					.HAlign(HAlign_Fill)
-					.VAlign(VAlign_Fill)
-					[
-						SAssignNew(TreeView, STreeView<FStatsNodePtr>)
-						.ExternalScrollbar(ExternalScrollbar)
-						.SelectionMode(ESelectionMode::Multi)
-						.TreeItemsSource(&FilteredGroupNodes)
-						.OnGetChildren(this, &SStatsView::TreeView_OnGetChildren)
-						.OnGenerateRow(this, &SStatsView::TreeView_OnGenerateRow)
-						.OnSelectionChanged(this, &SStatsView::TreeView_OnSelectionChanged)
-						.OnMouseButtonDoubleClick(this, &SStatsView::TreeView_OnMouseButtonDoubleClick)
-						.OnContextMenuOpening(FOnContextMenuOpening::CreateSP(this, &SStatsView::TreeView_GetMenuContent))
-						.ItemHeight(12.0f)
-						.HeaderRow
-						(
-							SAssignNew(TreeViewHeaderRow, SHeaderRow)
-							.Visibility(EVisibility::Visible)
-						)
-					]
+					SAssignNew(TreeView, STreeView<FStatsNodePtr>)
+					.ExternalScrollbar(ExternalScrollbar)
+					.SelectionMode(ESelectionMode::Multi)
+					.TreeItemsSource(&FilteredGroupNodes)
+					.OnGetChildren(this, &SStatsView::TreeView_OnGetChildren)
+					.OnGenerateRow(this, &SStatsView::TreeView_OnGenerateRow)
+					.OnSelectionChanged(this, &SStatsView::TreeView_OnSelectionChanged)
+					.OnMouseButtonDoubleClick(this, &SStatsView::TreeView_OnMouseButtonDoubleClick)
+					.OnContextMenuOpening(FOnContextMenuOpening::CreateSP(this, &SStatsView::TreeView_GetMenuContent))
+					.ItemHeight(12.0f)
+					.HeaderRow
+					(
+						SAssignNew(TreeViewHeaderRow, SHeaderRow)
+						.Visibility(EVisibility::Visible)
+					)
+				]
 
-					+ SOverlay::Slot()
-					.HAlign(HAlign_Right)
-					.VAlign(VAlign_Bottom)
-					.Padding(16.0f)
-					[
-						SAssignNew(AsyncOperationStatus, Insights::SAsyncOperationStatus, Aggregator)
-					]
+				+ SOverlay::Slot()
+				.HAlign(HAlign_Right)
+				.VAlign(VAlign_Bottom)
+				.Padding(16.0f)
+				[
+					SAssignNew(AsyncOperationStatus, Insights::SAsyncOperationStatus, Aggregator)
 				]
 			]
 
@@ -1586,8 +1580,8 @@ void SStatsView::ShowColumn(const FName ColumnId)
 		.VAlignCell(VAlign_Fill)
 		.SortMode(this, &SStatsView::GetSortModeForColumn, Column.GetId())
 		.OnSort(this, &SStatsView::OnSortModeChanged)
-		.ManualWidth(Column.GetInitialWidth())
-		.FixedWidth(Column.IsFixedWidth() ? Column.GetInitialWidth() : TOptional<float>())
+		.FillWidth(Column.GetInitialWidth())
+		//.FixedWidth(Column.IsFixedWidth() ? Column.GetInitialWidth() : TOptional<float>())
 		.HeaderContent()
 		[
 			SNew(SBox)
diff --git a/Engine/Source/Developer/TraceInsights/Private/Insights/Widgets/STimerTreeView.cpp b/Engine/Source/Developer/TraceInsights/Private/Insights/Widgets/STimerTreeView.cpp
--- a/Engine/Source/Developer/TraceInsights/Private/Insights/Widgets/STimerTreeView.cpp
+++ b/Engine/Source/Developer/TraceInsights/Private/Insights/Widgets/STimerTreeView.cpp
@@ -108,41 +108,35 @@ void STimerTreeView::Construct(const FArguments& InArgs, const FText& InViewName
 		.FillWidth(1.0f)
 		.Padding(0.0f)
 		[
-			SNew(SScrollBox)
-			.Orientation(Orient_Horizontal)
+			SNew(SOverlay)
 
-			+ SScrollBox::Slot()
+			+ SOverlay::Slot()
+			.HAlign(HAlign_Fill)
+			.VAlign(VAlign_Fill)
 			[
-				SNew(SOverlay)
-
-				+ SOverlay::Slot()
-				.HAlign(HAlign_Fill)
-				.VAlign(VAlign_Fill)
-				[
-					SAssignNew(TreeView, STreeView<FTimerNodePtr>)
-					.ExternalScrollbar(ExternalScrollbar)
-					.SelectionMode(ESelectionMode::Multi)
-					.TreeItemsSource(&TreeNodes)
-					.OnGetChildren(this, &STimerTreeView::TreeView_OnGetChildren)
-					.OnGenerateRow(this, &STimerTreeView::TreeView_OnGenerateRow)
-					//.OnSelectionChanged(this, &STimerTreeView::TreeView_OnSelectionChanged)
-					//.OnMouseButtonDoubleClick(this, &STimerTreeView::TreeView_OnMouseButtonDoubleClick)
-					.OnContextMenuOpening(FOnContextMenuOpening::CreateSP(this, &STimerTreeView::TreeView_GetMenuContent))
-					.ItemHeight(12.0f)
-					.HeaderRow
-					(
-						SAssignNew(TreeViewHeaderRow, SHeaderRow)
-						.Visibility(EVisibility::Visible)
-					)
-				]
-
-				+ SOverlay::Slot()
-				.HAlign(HAlign_Right)
-				.VAlign(VAlign_Bottom)
-				.Padding(16.0f)
-				[
-					SAssignNew(AsyncOperationStatus, Insights::SAsyncOperationStatus, TimerButterflyAggregator)
-				]
+				SAssignNew(TreeView, STreeView<FTimerNodePtr>)
+				.ExternalScrollbar(ExternalScrollbar)
+				.SelectionMode(ESelectionMode::Multi)
+				.TreeItemsSource(&TreeNodes)
+				.OnGetChildren(this, &STimerTreeView::TreeView_OnGetChildren)
+				.OnGenerateRow(this, &STimerTreeView::TreeView_OnGenerateRow)
+				//.OnSelectionChanged(this, &STimerTreeView::TreeView_OnSelectionChanged)
+				//.OnMouseButtonDoubleClick(this, &STimerTreeView::TreeView_OnMouseButtonDoubleClick)
+				.OnContextMenuOpening(FOnContextMenuOpening::CreateSP(this, &STimerTreeView::TreeView_GetMenuContent))
+				.ItemHeight(12.0f)
+				.HeaderRow
+				(
+					SAssignNew(TreeViewHeaderRow, SHeaderRow)
+					.Visibility(EVisibility::Visible)
+				)
+			]
+
+			+ SOverlay::Slot()
+			.HAlign(HAlign_Right)
+			.VAlign(VAlign_Bottom)
+			.Padding(16.0f)
+			[
+				SAssignNew(AsyncOperationStatus, Insights::SAsyncOperationStatus, TimerButterflyAggregator)
 			]
 		]
 
@@ -831,8 +825,8 @@ void STimerTreeView::ShowColumn(const FName ColumnId)
 		.VAlignCell(VAlign_Fill)
 		.SortMode(this, &STimerTreeView::GetSortModeForColumn, Column.GetId())
 		.OnSort(this, &STimerTreeView::OnSortModeChanged)
-		.ManualWidth(Column.GetInitialWidth())
-		.FixedWidth(Column.IsFixedWidth() ? Column.GetInitialWidth() : TOptional<float>())
+		.FillWidth(Column.GetInitialWidth())
+		//.FixedWidth(Column.IsFixedWidth() ? Column.GetInitialWidth() : TOptional<float>())
 		.HeaderContent()
 		[
 			SNew(SBox)
diff --git a/Engine/Source/Developer/TraceInsights/Private/Insights/Widgets/STimersView.cpp b/Engine/Source/Developer/TraceInsights/Private/Insights/Widgets/STimersView.cpp
--- a/Engine/Source/Developer/TraceInsights/Private/Insights/Widgets/STimersView.cpp
+++ b/Engine/Source/Developer/TraceInsights/Private/Insights/Widgets/STimersView.cpp
@@ -244,41 +244,35 @@ void STimersView::Construct(const FArguments& InArgs)
 			.FillWidth(1.0f)
 			.Padding(0.0f)
 			[
-				SNew(SScrollBox)
-				.Orientation(Orient_Horizontal)
+				SNew(SOverlay)
 
-				+ SScrollBox::Slot()
+				+ SOverlay::Slot()
+				.HAlign(HAlign_Fill)
+				.VAlign(VAlign_Fill)
 				[
-					SNew(SOverlay)
-
-					+ SOverlay::Slot()
-					.HAlign(HAlign_Fill)
-					.VAlign(VAlign_Fill)
-					[
-						SAssignNew(TreeView, STreeView<FTimerNodePtr>)
-						.ExternalScrollbar(ExternalScrollbar)
-						.SelectionMode(ESelectionMode::Multi)
-						.TreeItemsSource(&FilteredGroupNodes)
-						.OnGetChildren(this, &STimersView::TreeView_OnGetChildren)
-						.OnGenerateRow(this, &STimersView::TreeView_OnGenerateRow)
-						.OnSelectionChanged(this, &STimersView::TreeView_OnSelectionChanged)
-						.OnMouseButtonDoubleClick(this, &STimersView::TreeView_OnMouseButtonDoubleClick)
-						.OnContextMenuOpening(FOnContextMenuOpening::CreateSP(this, &STimersView::TreeView_GetMenuContent))
-						.ItemHeight(12.0f)
-						.HeaderRow
-						(
-							SAssignNew(TreeViewHeaderRow, SHeaderRow)
-							.Visibility(EVisibility::Visible)
-						)
-					]
+					SAssignNew(TreeView, STreeView<FTimerNodePtr>)
+					.ExternalScrollbar(ExternalScrollbar)
+					.SelectionMode(ESelectionMode::Multi)
+					.TreeItemsSource(&FilteredGroupNodes)
+					.OnGetChildren(this, &STimersView::TreeView_OnGetChildren)
+					.OnGenerateRow(this, &STimersView::TreeView_OnGenerateRow)
+					.OnSelectionChanged(this, &STimersView::TreeView_OnSelectionChanged)
+					.OnMouseButtonDoubleClick(this, &STimersView::TreeView_OnMouseButtonDoubleClick)
+					.OnContextMenuOpening(FOnContextMenuOpening::CreateSP(this, &STimersView::TreeView_GetMenuContent))
+					.ItemHeight(12.0f)
+					.HeaderRow
+					(
+						SAssignNew(TreeViewHeaderRow, SHeaderRow)
+						.Visibility(EVisibility::Visible)
+					)
+				]
 
-					+ SOverlay::Slot()
-					.HAlign(HAlign_Right)
-					.VAlign(VAlign_Bottom)
-					.Padding(16.0f)
-					[
-						SAssignNew(AsyncOperationStatus, Insights::SAsyncOperationStatus, Aggregator)
-					]
+				+ SOverlay::Slot()
+				.HAlign(HAlign_Right)
+				.VAlign(VAlign_Bottom)
+				.Padding(16.0f)
+				[
+					SAssignNew(AsyncOperationStatus, Insights::SAsyncOperationStatus, Aggregator)
 				]
 			]
 
@@ -1537,8 +1531,8 @@ void STimersView::ShowColumn(const FName ColumnId)
 		.VAlignCell(VAlign_Fill)
 		.SortMode(this, &STimersView::GetSortModeForColumn, Column.GetId())
 		.OnSort(this, &STimersView::OnSortModeChanged)
-		.ManualWidth(Column.GetInitialWidth())
-		.FixedWidth(Column.IsFixedWidth() ? Column.GetInitialWidth() : TOptional<float>())
+		.FillWidth(Column.GetInitialWidth())
+		//.FixedWidth(Column.IsFixedWidth() ? Column.GetInitialWidth() : TOptional<float>())
 		.HeaderContent()
 		[
 			SNew(SBox)
