[Insights]
  - Tasks: Improved drawing of task graph relations (added outline, tuned the size of arrow and curvature of splines).
  - Tasks: Fixed layer order of task graph relations (to draw on top of the track separators, which are now moved to background).
  - Tasks: Fixed crashes when a task graph relation is draw outside view or with extreme zoom level.
  - Timing view: Fixed height of tracks (when it has a child track; ex. with context switch events).
  - Timing view: Decreased the space on top and bottom of each timing track, so now the tracks are displayed in a more compact form. This is especially useful for the Cpu Core tracks.
  - Timing view: Added fade in / fade out animation for text when a timing track switches between Compact and Normal modes.

diff --git a/Engine/Source/Developer/TraceInsights/Private/Insights/ViewModels/TaskGraphRelation.cpp b/Engine/Source/Developer/TraceInsights/Private/Insights/ViewModels/TaskGraphRelation.cpp
--- a/Engine/Source/Developer/TraceInsights/Private/Insights/ViewModels/TaskGraphRelation.cpp
+++ b/Engine/Source/Developer/TraceInsights/Private/Insights/ViewModels/TaskGraphRelation.cpp
@@ -33,7 +33,8 @@ void FTaskGraphRelation::Draw(const FDrawContext& DrawContext, const FTimingTrac
 
 	if (Filter == ITimingEventRelation::EDrawFilter::BetweenScrollableTracks)
 	{
-		if (SourceTrackShared->GetLocation() > ETimingTrackLocation::Scrollable || TargetTrackShared->GetLocation() > ETimingTrackLocation::Scrollable)
+		if (SourceTrackShared->GetLocation() != ETimingTrackLocation::Scrollable ||
+			TargetTrackShared->GetLocation() != ETimingTrackLocation::Scrollable)
 		{
 			return;
 		}
@@ -41,7 +42,8 @@ void FTaskGraphRelation::Draw(const FDrawContext& DrawContext, const FTimingTrac
 
 	if (Filter == ITimingEventRelation::EDrawFilter::BetweenDockedTracks)
 	{
-		if (SourceTrackShared->GetLocation() <= ETimingTrackLocation::Scrollable && TargetTrackShared->GetLocation() <= ETimingTrackLocation::Scrollable)
+		if (SourceTrackShared->GetLocation() == ETimingTrackLocation::Scrollable &&
+			TargetTrackShared->GetLocation() == ETimingTrackLocation::Scrollable)
 		{
 			return;
 		}
@@ -49,60 +51,83 @@ void FTaskGraphRelation::Draw(const FDrawContext& DrawContext, const FTimingTrac
 		LayerId = DrawContext.LayerId;
 	}
 
+	const int32 OutlineLayerId = LayerId - 1;
+
 	float X1 = Viewport.TimeToSlateUnitsRounded(SourceTime);
+	float X2 = Viewport.TimeToSlateUnitsRounded(TargetTime);
+	if (FMath::Max(X1, X2) < 0.0f || FMath::Min(X1, X2) > Viewport.GetWidth())
+	{
+		return;
+	}
+
 	float Y1 = SourceTrackShared->GetPosY();
 	Y1 += Viewport.GetLayout().GetLaneY(SourceDepth) + Viewport.GetLayout().EventH / 2.0f;
-	if (SourceTrackShared->GetChildTrack())
+	if (SourceTrackShared->GetChildTrack() && SourceTrackShared->GetChildTrack()->GetHeight() > 0.0f)
 	{
 		Y1 += SourceTrackShared->GetChildTrack()->GetHeight() + Viewport.GetLayout().ChildTimelineDY;
 	}
 
-	float X2 = Viewport.TimeToSlateUnitsRounded(TargetTime);
 	float Y2 = TargetTrackShared->GetPosY();
 	Y2 += Viewport.GetLayout().GetLaneY(TargetDepth) + Viewport.GetLayout().EventH / 2.0f;
-	if (TargetTrackShared->GetChildTrack())
+	if (TargetTrackShared->GetChildTrack() && TargetTrackShared->GetChildTrack()->GetHeight() > 0.0f)
 	{
 		Y2 += TargetTrackShared->GetChildTrack()->GetHeight() + Viewport.GetLayout().ChildTimelineDY;
 	}
 
-	FVector2D StartPoint = FVector2D(X1, Y1);
-	FVector2D EndPoint = FVector2D(X2, Y2);
-	float Distance = FVector2D::Distance(StartPoint, EndPoint);
+	const FVector2D StartPoint = FVector2D(X1, Y1);
+	const FVector2D EndPoint = FVector2D(X2, Y2);
+	const float Distance = FVector2D::Distance(StartPoint, EndPoint);
+
+	constexpr float LineHeightAtStart = 4.0f;
+	constexpr float LineLengthAtStart = 4.0f;
+	constexpr float LineLengthAtEnd = 12.0f;
 
-	constexpr float DirectionFactor = 2.0f;
-	constexpr float LineLengthAtEnds = 20.0f;
-	FVector2D StartDir((X2 - X1) / DirectionFactor, 0.0f);
+	const FVector2D StartDir(FMath::Max(X2 - X1, 4.0f * (LineLengthAtStart + LineLengthAtEnd)), 0.0f);
 
-	FLinearColor Color = FTaskGraphProfilerManager::Get()->GetColorForTaskEvent(Type);
-	TArray<FVector2D> ArrowPoints;
+	constexpr float OutlineThickness = 5.0f;
+	constexpr float LineThickness = 3.0f;
 
-	constexpr float ArrowDirectionLen = 15.0f;
+	constexpr float ArrowDirectionLen = 10.0f;
 	constexpr float ArrowRotationAngle = 20.0f;
 	FVector2D ArrowDirection(-ArrowDirectionLen, 0.0f);
 
-	if (Distance > LineLengthAtEnds && !FMath::IsNearlyEqual(StartPoint.Y, EndPoint.Y))
-	{
-		FVector2D SplineStart(StartPoint.X + LineLengthAtEnds, StartPoint.Y);
-		FVector2D SplineEnd(EndPoint.X - LineLengthAtEnds, EndPoint.Y);
-
-		DrawContext.DrawSpline(LayerId, 0.0f, 0.0f, SplineStart, StartDir, SplineEnd, StartDir, /*Thickness=*/ 2.0f, Color);
+	const FLinearColor OutlineColor(0.0f, 0.0f, 0.0f, 1.0f);
+	const FLinearColor Color = FTaskGraphProfilerManager::Get()->GetColorForTaskEvent(Type);
 
-		ArrowPoints.Add(StartPoint);
-		ArrowPoints.Add(SplineStart);
+	TArray<FVector2D> LinePoints;
+	LinePoints.Add(StartPoint + FVector2D(0.0f, -LineHeightAtStart / 2.0f));
+	LinePoints.Add(StartPoint + FVector2D(0.0f, +LineHeightAtStart / 2.0f));
+	DrawContext.DrawLines(OutlineLayerId, 0.0f, 0.0f, LinePoints, ESlateDrawEffect::None, OutlineColor, /*bAntialias=*/ true, OutlineThickness);
+	DrawContext.DrawLines(LayerId, 0.0f, 0.0f, LinePoints, ESlateDrawEffect::None, Color, /*bAntialias=*/ true, LineThickness);
 
-		DrawContext.DrawLines(LayerId, 0.0f, 0.0f, ArrowPoints, ESlateDrawEffect::None, Color, /*bAntialias=*/ true, /*Thickness=*/ 2.0f);
-
-		ArrowPoints.Empty();
-		ArrowPoints.Add(SplineEnd);
-		ArrowPoints.Add(EndPoint);
-		DrawContext.DrawLines(LayerId, 0.0f, 0.0f, ArrowPoints, ESlateDrawEffect::None, Color, /*bAntialias=*/ true, /*Thickness=*/ 2.0f);
+	constexpr float MinDistance = 1.5f * (LineLengthAtStart + LineLengthAtEnd);
+	constexpr float MaxDistance = 10000.0f; // arbitrary limit to avoid stack overflow in recursive FLineBuilder::Subdivide when rendering splines
+	if (Distance > MinDistance && Distance < MaxDistance && !FMath::IsNearlyEqual(StartPoint.Y, EndPoint.Y))
+	{
+		FVector2D SplineStart(StartPoint.X + LineLengthAtStart, StartPoint.Y);
+		FVector2D SplineEnd(EndPoint.X - LineLengthAtEnd, EndPoint.Y);
+		DrawContext.DrawSpline(OutlineLayerId, 0.0f, 0.0f, SplineStart, StartDir, SplineEnd, StartDir, OutlineThickness, OutlineColor);
+		DrawContext.DrawSpline(LayerId, 0.0f, 0.0f, SplineStart, StartDir, SplineEnd, StartDir, LineThickness, Color);
+
+		LinePoints.Empty();
+		LinePoints.Add(StartPoint);
+		LinePoints.Add(SplineStart);
+		DrawContext.DrawLines(OutlineLayerId, 0.0f, 0.0f, LinePoints, ESlateDrawEffect::None, OutlineColor, /*bAntialias=*/ true, OutlineThickness);
+		DrawContext.DrawLines(LayerId, 0.0f, 0.0f, LinePoints, ESlateDrawEffect::None, Color, /*bAntialias=*/ true, LineThickness);
+
+		LinePoints.Empty();
+		LinePoints.Add(SplineEnd);
+		LinePoints.Add(EndPoint);
+		DrawContext.DrawLines(OutlineLayerId, 0.0f, 0.0f, LinePoints, ESlateDrawEffect::None, OutlineColor, /*bAntialias=*/ true, OutlineThickness);
+		DrawContext.DrawLines(LayerId, 0.0f, 0.0f, LinePoints, ESlateDrawEffect::None, Color, /*bAntialias=*/ true, LineThickness);
 	}
 	else
 	{
-		ArrowPoints.Empty();
-		ArrowPoints.Add(StartPoint);
-		ArrowPoints.Add(EndPoint);
-		DrawContext.DrawLines(LayerId, 0.0f, 0.0f, ArrowPoints, ESlateDrawEffect::None, Color, /*bAntialias=*/ true, /*Thickness=*/ 2.0f);
+		LinePoints.Empty();
+		LinePoints.Add(StartPoint);
+		LinePoints.Add(EndPoint);
+		DrawContext.DrawLines(OutlineLayerId, 0.0f, 0.0f, LinePoints, ESlateDrawEffect::None, OutlineColor, /*bAntialias=*/ true, OutlineThickness);
+		DrawContext.DrawLines(LayerId, 0.0f, 0.0f, LinePoints, ESlateDrawEffect::None, Color, /*bAntialias=*/ true, LineThickness);
 
 		ArrowDirection = StartPoint - EndPoint;
 		ArrowDirection.Normalize();
@@ -111,17 +136,17 @@ void FTaskGraphRelation::Draw(const FDrawContext& DrawContext, const FTimingTrac
 
 	FVector2D ArrowOrigin = EndPoint;
 
-	ArrowPoints.Empty();
-	ArrowPoints.Add(ArrowOrigin);
-	ArrowPoints.Add(ArrowOrigin + ArrowDirection.GetRotated(-ArrowRotationAngle));
-
-	DrawContext.DrawLines(LayerId, 0.0f, 0.0f, ArrowPoints, ESlateDrawEffect::None, Color, /*bAntialias=*/ true, /*Thickness=*/ 2.0f);
-
-	ArrowPoints.Empty();
-	ArrowPoints.Add(ArrowOrigin);
-	ArrowPoints.Add(ArrowOrigin + ArrowDirection.GetRotated(ArrowRotationAngle));
-
-	DrawContext.DrawLines(LayerId, 0.0f, 0.0f, ArrowPoints, ESlateDrawEffect::None, Color, /*bAntialias=*/ true, /*Thickness=*/ 2.0f);
+	LinePoints.Empty();
+	LinePoints.Add(ArrowOrigin);
+	LinePoints.Add(ArrowOrigin + ArrowDirection.GetRotated(-ArrowRotationAngle));
+	DrawContext.DrawLines(OutlineLayerId, 0.0f, 0.0f, LinePoints, ESlateDrawEffect::None, OutlineColor, /*bAntialias=*/ true, OutlineThickness);
+	DrawContext.DrawLines(LayerId, 0.0f, 0.0f, LinePoints, ESlateDrawEffect::None, Color, /*bAntialias=*/ true, LineThickness);
+
+	LinePoints.Empty();
+	LinePoints.Add(ArrowOrigin);
+	LinePoints.Add(ArrowOrigin + ArrowDirection.GetRotated(ArrowRotationAngle));
+	DrawContext.DrawLines(OutlineLayerId, 0.0f, 0.0f, LinePoints, ESlateDrawEffect::None, OutlineColor, /*bAntialias=*/ true, OutlineThickness);
+	DrawContext.DrawLines(LayerId, 0.0f, 0.0f, LinePoints, ESlateDrawEffect::None, Color, /*bAntialias=*/ true, LineThickness);
 }
 
-} // namespace Insights
\ No newline at end of file
+} // namespace Insights
diff --git a/Engine/Source/Developer/TraceInsights/Private/Insights/ViewModels/TimingEventsTrack.cpp b/Engine/Source/Developer/TraceInsights/Private/Insights/ViewModels/TimingEventsTrack.cpp
--- a/Engine/Source/Developer/TraceInsights/Private/Insights/ViewModels/TimingEventsTrack.cpp
+++ b/Engine/Source/Developer/TraceInsights/Private/Insights/ViewModels/TimingEventsTrack.cpp
@@ -171,11 +171,11 @@ void FTimingEventsTrack::UpdateTrackHeight(const ITimingTrackUpdateContext& Cont
 	else
 	{
 		DesiredTrackHeight = Viewport.GetLayout().ComputeTrackHeight(NumLanes);
-		if (ChildTrack.IsValid() && ChildTrack->GetHeight() > Viewport.GetLayout().MinTimelineH)
+		if (ChildTrack.IsValid() && ChildTrack->GetHeight() > 0.0f)
 		{
-			if (DesiredTrackHeight == Viewport.GetLayout().MinTimelineH)
+			if (DesiredTrackHeight == 0.0f)
 			{
-				DesiredTrackHeight += 1.0f + 2.0 * Viewport.GetLayout().TimelineDY;
+				DesiredTrackHeight += 1.0f + 2.0f * Viewport.GetLayout().TimelineDY;
 			}
 			DesiredTrackHeight += ChildTrack->GetHeight() + Viewport.GetLayout().ChildTimelineDY;
 		}
@@ -343,7 +343,7 @@ void FTimingEventsTrack::DrawEvent(const ITimingTrackDrawContext& Context, const
 		const FTimingEvent& TrackEvent = InTimingEvent.As<FTimingEvent>();
 		const FTimingViewLayout& Layout = Context.GetViewport().GetLayout();
 		float Y = TrackEvent.GetTrack()->GetPosY();
-		if (ChildTrack.IsValid())
+		if (ChildTrack.IsValid() && ChildTrack->GetHeight() > 0.0f)
 		{
 			Y += ChildTrack->GetHeight() + Layout.ChildTimelineDY;
 		}
@@ -369,7 +369,7 @@ const TSharedPtr<const ITimingEvent> FTimingEventsTrack::GetEvent(float InPosX,
 
 	float TopLaneY = 0.0f;
 	float TrackLanesHeight = 0.0f;
-	if (ChildTrack.IsValid())
+	if (ChildTrack.IsValid() && ChildTrack->GetHeight() > 0.0f)
 	{
 		const float HeaderDY = InPosY - ChildTrack->GetPosY();
 		if (HeaderDY >= 0 && HeaderDY < ChildTrack->GetHeight())
diff --git a/Engine/Source/Developer/TraceInsights/Private/Insights/ViewModels/TimingViewDrawHelper.cpp b/Engine/Source/Developer/TraceInsights/Private/Insights/ViewModels/TimingViewDrawHelper.cpp
--- a/Engine/Source/Developer/TraceInsights/Private/Insights/ViewModels/TimingViewDrawHelper.cpp
+++ b/Engine/Source/Developer/TraceInsights/Private/Insights/ViewModels/TimingViewDrawHelper.cpp
@@ -318,7 +318,7 @@ void FTimingViewDrawHelper::DrawEvents(const FTimingEventsTrackDrawState& DrawSt
 		{
 			TopLaneY += OffsetY + Layout.TimelineDY;
 		}
-		if (Track.GetChildTrack().IsValid())
+		if (Track.GetChildTrack().IsValid() && Track.ChildTrack->GetHeight() > 0.0f)
 		{
 			TopLaneY += Track.ChildTrack->GetHeight() + Layout.ChildTimelineDY;
 		}
@@ -367,13 +367,25 @@ void FTimingViewDrawHelper::DrawEvents(const FTimingEventsTrackDrawState& DrawSt
 		}
 
 		// Draw texts.
-		if (Layout.EventH > 10.0f)
+		constexpr float TextMinEventH = 7.0f;
+		if (Layout.EventH > TextMinEventH)
 		{
+			float TextOpacity = 1.0f;
+			if (Layout.EventH < FTimingViewLayout::NormalLayoutEventH)
+			{
+				TextOpacity = (Layout.EventH - TextMinEventH + 1.0f) / (FTimingViewLayout::NormalLayoutEventH - TextMinEventH + 1.0f);
+			}
+			const FLinearColor WhiteColor(1.0f, 1.0f, 1.0f, TextOpacity);
+			const FLinearColor BlackColor(0.0f, 0.0f, 0.0f, TextOpacity);
+
+			const float TextY0 = TopLaneY + 1.0f - (FTimingViewLayout::NormalLayoutEventH - Layout.EventH) / 2.0f;
+			const float TextDY = Layout.EventH + Layout.EventDY;
+
 			const int32 EventTextLayerId = ReservedLayerId + ToInt32(EDrawLayer::EventText);
 			for (const FTimingEventsTrackDrawState::FTextPrimitive& Text : DrawState.Texts)
 			{
-				const float Y = TopLaneY + (Layout.EventH + Layout.EventDY) * Text.Depth + 1.0f;
-				DrawContext.DrawText(EventTextLayerId, Text.X, Y, Text.Text, EventFont, Text.bWhite ? FLinearColor::White : FLinearColor::Black);
+				const float Y = TextY0 + TextDY * Text.Depth;
+				DrawContext.DrawText(EventTextLayerId, Text.X, Y, Text.Text, EventFont, Text.bWhite ? WhiteColor : BlackColor);
 			}
 			NumDrawTexts += DrawState.Texts.Num();
 		}
@@ -397,7 +409,7 @@ void FTimingViewDrawHelper::DrawFadedEvents(const FTimingEventsTrackDrawState& D
 		{
 			TopLaneY += OffsetY + Layout.TimelineDY;
 		}
-		if (Track.GetChildTrack().IsValid())
+		if (Track.GetChildTrack().IsValid() && Track.ChildTrack->GetHeight() > 0.0f)
 		{
 			TopLaneY += Track.ChildTrack->GetHeight() + Layout.ChildTimelineDY;
 		}
@@ -447,15 +459,24 @@ void FTimingViewDrawHelper::DrawFadedEvents(const FTimingEventsTrackDrawState& D
 		}
 
 		// Draw texts.
-		if (Layout.EventH > 10.0f)
+		constexpr float TextMinEventH = 7.0f;
+		if (Layout.EventH > TextMinEventH)
 		{
-			const FLinearColor WhiteColor(1.0f, 1.0f, 1.0f, Opacity);
-			const FLinearColor BlackColor(0.0f, 0.0f, 0.0f, Opacity);
+			float TextOpacity = Opacity;
+			if (Layout.EventH < FTimingViewLayout::NormalLayoutEventH)
+			{
+				TextOpacity *= (Layout.EventH - TextMinEventH + 1.0f) / (FTimingViewLayout::NormalLayoutEventH - TextMinEventH + 1.0f);
+			}
+			const FLinearColor WhiteColor(1.0f, 1.0f, 1.0f, TextOpacity);
+			const FLinearColor BlackColor(0.0f, 0.0f, 0.0f, TextOpacity);
+
+			const float TextY0 = TopLaneY + 1.0f - (FTimingViewLayout::NormalLayoutEventH - Layout.EventH) / 2.0f;
+			const float TextDY = Layout.EventH + Layout.EventDY;
 
 			const int32 EventTextLayerId = ReservedLayerId + ToInt32(EDrawLayer::EventText);
 			for (const FTimingEventsTrackDrawState::FTextPrimitive& Text : DrawState.Texts)
 			{
-				const float Y = TopLaneY + (Layout.EventH + Layout.EventDY) * Text.Depth + 1.0f;
+				const float Y = TextY0 + (TextDY) * Text.Depth;
 				DrawContext.DrawText(EventTextLayerId, Text.X, Y, Text.Text, EventFont, Text.bWhite ? WhiteColor : BlackColor);
 			}
 			NumDrawTexts += DrawState.Texts.Num();
@@ -481,7 +502,7 @@ void FTimingViewDrawHelper::DrawLineEvents(const FTimingEventsTrackDrawState& Dr
 		{
 			TopLaneY += OffsetY + Layout.TimelineDY;
 		}
-		if (Track.GetChildTrack().IsValid())
+		if (Track.GetChildTrack().IsValid() && Track.ChildTrack->GetHeight() > 0.0f)
 		{
 			TopLaneY += Track.ChildTrack->GetHeight() + Layout.ChildTimelineDY;
 		}
@@ -528,13 +549,21 @@ void FTimingViewDrawHelper::DrawLineEvents(const FTimingEventsTrackDrawState& Dr
 		}
 
 		// Draw texts.
-		if (Layout.EventH > 10.0f)
+		constexpr float TextMinEventH = 7.0f;
+		if (Layout.EventH > TextMinEventH)
 		{
+			float TextOpacity = 1.0f;
+			if (Layout.EventH < FTimingViewLayout::NormalLayoutEventH)
+			{
+				TextOpacity *= (Layout.EventH - TextMinEventH + 1.0f) / (FTimingViewLayout::NormalLayoutEventH - TextMinEventH + 1.0f);
+			}
+
+			const float TextY = TopLaneY - (FTimingViewLayout::NormalLayoutEventH - Layout.EventH) / 2.0f;
+
 			const int32 EventTextLayerId = ReservedLayerId + ToInt32(EDrawLayer::EventText);
-			const float Y = TopLaneY;
 			for (const FTimingEventsTrackDrawState::FTextPrimitive& Text : DrawState.Texts)
 			{
-				DrawContext.DrawText(EventTextLayerId, Text.X, Y, Text.Text, EventFont, Text.Color);
+				DrawContext.DrawText(EventTextLayerId, Text.X, TextY, Text.Text, EventFont, Text.Color.CopyWithNewOpacity(Text.Color.A * TextOpacity));
 			}
 			NumDrawTexts += DrawState.Texts.Num();
 		}
@@ -663,10 +692,11 @@ void FTimingViewDrawHelper::DrawTrackHeader(const FBaseTimingTrack& Track, const
 
 	if (TrackH > 0.0f)
 	{
-		// Draw a horizontal line between timelines (top line of a track).
-		DrawContext.DrawBox(HeaderLayerId, 0.0f, TrackY, Viewport.GetWidth(), 1.0f, WhiteBrush, EdgeColor);
+		// Draw a horizontal line between tracks (top line of a track).
+		const int32 TrackBackgroundLayerId = HeaderLayerId - ToInt32(EDrawLayer::HeaderBackground) + ToInt32(EDrawLayer::TrackBackground);
+		DrawContext.DrawBox(TrackBackgroundLayerId, 0.0f, TrackY, Viewport.GetWidth(), 1.0f, WhiteBrush, EdgeColor);
 
-		// Draw header with name of timeline.
+		// Draw header with name of tracks.
 		if (TrackH > 4.0f)
 		{
 			const TSharedRef<FSlateFontMeasure> FontMeasureService = FSlateApplication::Get().GetRenderer()->GetFontMeasureService();
@@ -721,6 +751,8 @@ void FTimingViewDrawHelper::EndDrawTracks() const
 {
 	if (Viewport.GetWidth() > 0.0f)
 	{
+		const int32 TrackBackgroundLayerId = ReservedLayerId + ToInt32(EDrawLayer::TrackBackground);
+
 		const float TopY = Viewport.GetTopOffset();
 		const float BottomY = Viewport.GetHeight() - Viewport.GetBottomOffset();
 
@@ -732,7 +764,7 @@ void FTimingViewDrawHelper::EndDrawTracks() const
 			if (Y >= TopY && Y < BottomY)
 			{
 				// Draw a last horizontal line.
-				DrawContext.DrawBox(ReservedLayerId + ToInt32(EDrawLayer::HeaderBackground), 0.0f, Y, Viewport.GetWidth(), 1.0f, WhiteBrush, EdgeColor);
+				DrawContext.DrawBox(TrackBackgroundLayerId, 0.0f, Y, Viewport.GetWidth(), 1.0f, WhiteBrush, EdgeColor);
 			}
 
 			// Note: ValidAreaX and ValidAreaW are computed in DrawBackground.
@@ -742,14 +774,14 @@ void FTimingViewDrawHelper::EndDrawTracks() const
 				if (TopInvalidAreaH > 0.0f)
 				{
 					// Draw invalid area (top).
-					DrawContext.DrawBox(ReservedLayerId + ToInt32(EDrawLayer::HeaderBackground), ValidAreaX, TopY, ValidAreaW, TopInvalidAreaH, BackgroundAreaBrush, InvalidAreaColor);
+					DrawContext.DrawBox(TrackBackgroundLayerId, ValidAreaX, TopY, ValidAreaW, TopInvalidAreaH, BackgroundAreaBrush, InvalidAreaColor);
 				}
 
 				const float BottomInvalidAreaH = FMath::Min(BottomY - Y - 1.0f, Viewport.GetScrollableAreaHeight());
 				if (BottomInvalidAreaH > 0.0f)
 				{
 					// Draw invalid area (bottom).
-					DrawContext.DrawBox(ReservedLayerId + ToInt32(EDrawLayer::HeaderBackground), ValidAreaX, BottomY - BottomInvalidAreaH, ValidAreaW, BottomInvalidAreaH, BackgroundAreaBrush, InvalidAreaColor);
+					DrawContext.DrawBox(TrackBackgroundLayerId, ValidAreaX, BottomY - BottomInvalidAreaH, ValidAreaW, BottomInvalidAreaH, BackgroundAreaBrush, InvalidAreaColor);
 				}
 			}
 		}
diff --git a/Engine/Source/Developer/TraceInsights/Private/Insights/ViewModels/TimingViewDrawHelper.h b/Engine/Source/Developer/TraceInsights/Private/Insights/ViewModels/TimingViewDrawHelper.h
--- a/Engine/Source/Developer/TraceInsights/Private/Insights/ViewModels/TimingViewDrawHelper.h
+++ b/Engine/Source/Developer/TraceInsights/Private/Insights/ViewModels/TimingViewDrawHelper.h
@@ -143,10 +143,12 @@ class FTimingViewDrawHelper final : public ITimingViewDrawHelper
 private:
 	enum class EDrawLayer : int32
 	{
+		TrackBackground,
 		EventBorder,
 		EventFill,
 		EventText,
 		EventHighlight,
+		RelationBackground,
 		Relation,
 		HeaderBackground,
 		HeaderText,
diff --git a/Engine/Source/Developer/TraceInsights/Private/Insights/ViewModels/TrackHeader.cpp b/Engine/Source/Developer/TraceInsights/Private/Insights/ViewModels/TrackHeader.cpp
--- a/Engine/Source/Developer/TraceInsights/Private/Insights/ViewModels/TrackHeader.cpp
+++ b/Engine/Source/Developer/TraceInsights/Private/Insights/ViewModels/TrackHeader.cpp
@@ -122,7 +122,7 @@ void FTrackHeader::Draw(const ITimingTrackDrawContext& Context) const
 	{
 		FDrawContext& DrawContext = Context.GetDrawContext();
 
-		// Draw a horizontal line between timelines (top line of a track).
+		// Draw a horizontal line between tracks (top line of a track).
 		DrawContext.DrawBox(0.0f, ParentTrack.GetPosY(), Context.GetViewport().GetWidth(), 1.0f, WhiteBrush, FLinearColor(0.05f, 0.05f, 0.05f, 1.0f));
 
 		DrawInternal(Context, false);
diff --git a/Engine/Source/Developer/TraceInsights/Public/Insights/ViewModels/TimingViewLayout.h b/Engine/Source/Developer/TraceInsights/Public/Insights/ViewModels/TimingViewLayout.h
--- a/Engine/Source/Developer/TraceInsights/Public/Insights/ViewModels/TimingViewLayout.h
+++ b/Engine/Source/Developer/TraceInsights/Public/Insights/ViewModels/TimingViewLayout.h
@@ -12,8 +12,8 @@ struct TRACEINSIGHTS_API FTimingViewLayout
 
 	static constexpr float NormalLayoutEventH = 14.0f;
 	static constexpr float NormalLayoutEventDY = 2.0f;
-	static constexpr float NormalLayoutTimelineDY = 10.0f;
-	static constexpr float NormalLayoutChildTimelineDY = 5.0f;
+	static constexpr float NormalLayoutTimelineDY = 2.0f;
+	static constexpr float NormalLayoutChildTimelineDY = 4.0f;
 	static constexpr float NormalLayoutMinTimelineH = 0.0f;
 
 	static constexpr float CompactLayoutEventH = 4.0f;
@@ -49,12 +49,12 @@ struct TRACEINSIGHTS_API FTimingViewLayout
 	}
 
 	/* The layout of a track:
-	*	1.0f - Line between tracks
+	*	1.0f // the line between tracks
 	*	TimelineDY
-	*	HeaderLanes
-	*	ChildTimelineDY
+	*	[ChildTrack]
+	*	[ChildTimelineDY] // if ChildTrack is valid and ChildTrack.Height > 0
 	*	TrackLanes
-	*	TimelineDY 
+	*	TimelineDY
 	*/
 
 	float ComputeTrackHeight(int32 NumLanes) const
@@ -67,7 +67,7 @@ struct TRACEINSIGHTS_API FTimingViewLayout
 		{
 			// 1.0f is for horizontal line between timelines
 			float TrackHeight = 1.0f + TimelineDY + (EventH + EventDY) * NumLanes - EventDY + TimelineDY;
-				
+
 			if (TrackHeight < FTimingViewLayout::RealMinTimelineH)
 			{
 				return FTimingViewLayout::RealMinTimelineH;
@@ -85,7 +85,6 @@ struct TRACEINSIGHTS_API FTimingViewLayout
 		}
 		else //if (NumLanes > 0)
 		{
-			// 1.0f is for horizontal line between timelines
 			const float TrackHeight = (EventH + EventDY) * NumLanes - EventDY;
 
 			if (TrackHeight < FTimingViewLayout::RealMinTimelineH)
