[Insights]
  - Updated UX for the Session Info view (work in progress).
  - Also added display for URI of current session.

diff --git a/Engine/Source/Developer/TraceInsights/Private/Insights/Widgets/SSessionInfoWindow.cpp b/Engine/Source/Developer/TraceInsights/Private/Insights/Widgets/SSessionInfoWindow.cpp
--- a/Engine/Source/Developer/TraceInsights/Private/Insights/Widgets/SSessionInfoWindow.cpp
+++ b/Engine/Source/Developer/TraceInsights/Private/Insights/Widgets/SSessionInfoWindow.cpp
@@ -50,7 +50,7 @@ SSessionInfoWindow::~SSessionInfoWindow()
 #if WITH_EDITOR
 	if (DurationActive > 0.0f && FEngineAnalytics::IsAvailable())
 	{
-		FEngineAnalytics::GetProvider().RecordEvent(TEXT("Editor.Usage.Insights.SessionInfo"), FAnalyticsEventAttribute(TEXT("Duration"), DurationActive));
+		FEngineAnalytics::GetProvider().RecordEvent(TEXT("Insights.Usage.SessionInfo"), FAnalyticsEventAttribute(TEXT("Duration"), DurationActive));
 	}
 #endif // WITH_EDITOR
 }
@@ -137,23 +137,37 @@ void SSessionInfoWindow::Construct(const FArguments& InArgs, const TSharedRef<SD
 
 ////////////////////////////////////////////////////////////////////////////////////////////////////
 
-void SSessionInfoWindow::AddInfoLine(TSharedPtr<SVerticalBox> InVerticalBox, const FText& InHeader, const TAttribute<FText>& InValue, bool bMultiLine) const
+void SSessionInfoWindow::AddSectionLine(TSharedPtr<SVerticalBox> InVerticalBox, const FText& InSectionName) const
 {
 	InVerticalBox->AddSlot()
 		.AutoHeight()
-		.Padding(8.0f, 4.0f, 8.0f, 0.0f)
+		.Padding(0.0f, 4.0f, 0.0f, 4.0f)
 		[
-			SNew(STextBlock)
-			.Text(InHeader)
-			.Font(FCoreStyle::GetDefaultFontStyle("Bold", 11))
+			SNew(SBorder)
+			.BorderBackgroundColor(FLinearColor(0.03f, 0.03f, 0.03f, 1.0f))
+			.BorderImage(FCoreStyle::Get().GetBrush("WhiteBrush"))
+			.Padding(FMargin(0.0f, 4.0f, 0.0f, 4.0f))
+			[
+				SNew(SBox)
+				.Padding(FMargin(16.0f, 0.0f, 16.0f, 0.0f))
+				[
+					SNew(STextBlock)
+					.Text(InSectionName)
+					.ColorAndOpacity(FLinearColor(0.7f, 0.7f, 0.7f, 1.0f))
+				]
+			]
 		];
+}
 
+////////////////////////////////////////////////////////////////////////////////////////////////////
+
+TSharedRef<SWidget> SSessionInfoWindow::CreateTextBox(const TAttribute<FText>& InText, bool bMultiLine) const
+{
 	TSharedPtr<SWidget> TextBox;
 	if (bMultiLine)
 	{
 		TextBox = SNew(SMultiLineEditableTextBox)
-			.Text(InValue)
-			.HintText(LOCTEXT("HintText", "N/A"))
+			.Text(InText)
 			.AutoWrapText(true)
 			.BackgroundColor(FLinearColor(0.1f, 0.1f, 0.1f, 1.0f))
 			.IsReadOnly(true);
@@ -161,21 +175,64 @@ void SSessionInfoWindow::AddInfoLine(TSharedPtr<SVerticalBox> InVerticalBox, con
 	else
 	{
 		TextBox = SNew(SEditableTextBox)
-			.Text(InValue)
-			.HintText(LOCTEXT("HintText", "N/A"))
+			.Text(InText)
 			.BackgroundColor(FLinearColor(0.1f, 0.1f, 0.1f, 1.0f))
 			.IsReadOnly(true);
 	}
+	return TextBox.ToSharedRef();
+}
 
+////////////////////////////////////////////////////////////////////////////////////////////////////
+
+void SSessionInfoWindow::AddInfoLine(TSharedPtr<SVerticalBox> InVerticalBox, const FText& InHeader, const TAttribute<FText>& InValue, bool bMultiLine) const
+{
 	InVerticalBox->AddSlot()
 		.AutoHeight()
-		.Padding(8.0f, 0.0f, 8.0f, 4.0f)
+		.Padding(16.0f, 4.0f, 16.0f, 4.0f)
 		[
-			TextBox.ToSharedRef()
+			SNew(SHorizontalBox)
+
+			+ SHorizontalBox::Slot()
+			.HAlign(HAlign_Right)
+			.VAlign(VAlign_Top)
+			.AutoWidth()
+			.Padding(0.0f)
+			[
+				SNew(SBox)
+				.Padding(FMargin(0.0f, 4.0f, 16.0f, 4.0f))
+				.HAlign(HAlign_Right)
+				.MinDesiredWidth(160.0f)
+				[
+					SNew(STextBlock)
+					.Text(InHeader)
+				]
+			]
+
+			+ SHorizontalBox::Slot()
+			.HAlign(HAlign_Fill)
+			.VAlign(VAlign_Top)
+			.FillWidth(1.0f)
+			.Padding(0.0f)
+			[
+				CreateTextBox(InValue, bMultiLine)
+			]
 		];
 }
 
 ////////////////////////////////////////////////////////////////////////////////////////////////////
+
+void SSessionInfoWindow::AddSimpleInfoLine(TSharedPtr<SVerticalBox> InVerticalBox, const TAttribute<FText>& InValue, bool bMultiLine) const
+{
+	InVerticalBox->AddSlot()
+		.AutoHeight()
+		.Padding(16.0f, 4.0f, 16.0f, 4.0f)
+		[
+			CreateTextBox(InValue, bMultiLine)
+		];
+}
+
+////////////////////////////////////////////////////////////////////////////////////////////////////
+
 TSharedRef<SDockTab> SSessionInfoWindow::SpawnTab_SessionInfo(const FSpawnTabArgs& Args)
 {
 	TSharedPtr<SVerticalBox> VerticalBox;
@@ -228,19 +285,24 @@ TSharedRef<SDockTab> SSessionInfoWindow::SpawnTab_SessionInfo(const FSpawnTabArg
 
 	Image->SetImage(new FSlateColorBrush(FLinearColor(0.015f, 0.015f, 0.015f, 1.0f)));
 
-	AddInfoLine(VerticalBox, LOCTEXT("SessionName_HeaderText",	"Session Name:"),		TAttribute<FText>(this, &SSessionInfoWindow::GetSessionNameText));
-	//AddInfoLine(VerticalBox, LOCTEXT("Uri_HeaderText",			"URI:"),				TAttribute<FText>(this, &SSessionInfoWindow::GetUriText));
-	//AddInfoLine(VerticalBox, LOCTEXT("FileSize_HeaderText",		"File Size:"),			TAttribute<FText>(this, &SSessionInfoWindow::GetFileSizeText));
-	AddInfoLine(VerticalBox, LOCTEXT("Platform_HeaderText",		"Platform:"),			TAttribute<FText>(this, &SSessionInfoWindow::GetPlatformText));
-	AddInfoLine(VerticalBox, LOCTEXT("AppName_HeaderText",		"Application Name:"),	TAttribute<FText>(this, &SSessionInfoWindow::GetAppNameText));
-	AddInfoLine(VerticalBox, LOCTEXT("Branch_HeaderText",		"Branch:"),				TAttribute<FText>(this, &SSessionInfoWindow::GetBranchText));
-	AddInfoLine(VerticalBox, LOCTEXT("BuildVersion_HeaderText",	"Build Version:"),		TAttribute<FText>(this, &SSessionInfoWindow::GetBuildVersionText));
-	AddInfoLine(VerticalBox, LOCTEXT("Changelist_HeaderText",	"Changelist:"),			TAttribute<FText>(this, &SSessionInfoWindow::GetChangelistText));
-	AddInfoLine(VerticalBox, LOCTEXT("BuildConfig_HeaderText",	"Build Config:"),		TAttribute<FText>(this, &SSessionInfoWindow::GetBuildConfigText));
-	AddInfoLine(VerticalBox, LOCTEXT("BuildTarget_HeaderText",	"Build Target:"),		TAttribute<FText>(this, &SSessionInfoWindow::GetBuildTargetText));
-	AddInfoLine(VerticalBox, LOCTEXT("CommandLine_HeaderText",	"Command Line:"),		TAttribute<FText>(this, &SSessionInfoWindow::GetCommandLineText), true);
-	AddInfoLine(VerticalBox, LOCTEXT("Status_HeaderText",		"Analysis Status:"),	TAttribute<FText>(this, &SSessionInfoWindow::GetStatusText));
-	AddInfoLine(VerticalBox, LOCTEXT("Modules_HeaderText",		"Analysis Modules:"),	TAttribute<FText>(this, &SSessionInfoWindow::GetModulesText));
+	AddSectionLine(VerticalBox, LOCTEXT("SessionInfo_SectionText", "Session Info"));
+	AddInfoLine(VerticalBox, LOCTEXT("SessionName_HeaderText",	"Session Name"),		TAttribute<FText>(this, &SSessionInfoWindow::GetSessionNameText));
+	AddInfoLine(VerticalBox, LOCTEXT("Uri_HeaderText",			"URI"),					TAttribute<FText>(this, &SSessionInfoWindow::GetUriText));
+	//AddInfoLine(VerticalBox, LOCTEXT("FileSize_HeaderText",		"File Size"),			TAttribute<FText>(this, &SSessionInfoWindow::GetFileSizeText));
+	AddInfoLine(VerticalBox, LOCTEXT("Platform_HeaderText",		"Platform"),			TAttribute<FText>(this, &SSessionInfoWindow::GetPlatformText));
+	AddInfoLine(VerticalBox, LOCTEXT("AppName_HeaderText",		"Application Name"),	TAttribute<FText>(this, &SSessionInfoWindow::GetAppNameText));
+	AddInfoLine(VerticalBox, LOCTEXT("Branch_HeaderText",		"Branch"),				TAttribute<FText>(this, &SSessionInfoWindow::GetBranchText));
+	AddInfoLine(VerticalBox, LOCTEXT("BuildVersion_HeaderText",	"Build Version"),		TAttribute<FText>(this, &SSessionInfoWindow::GetBuildVersionText));
+	AddInfoLine(VerticalBox, LOCTEXT("Changelist_HeaderText",	"Changelist"),			TAttribute<FText>(this, &SSessionInfoWindow::GetChangelistText));
+	AddInfoLine(VerticalBox, LOCTEXT("BuildConfig_HeaderText",	"Build Config"),		TAttribute<FText>(this, &SSessionInfoWindow::GetBuildConfigText));
+	AddInfoLine(VerticalBox, LOCTEXT("BuildTarget_HeaderText",	"Build Target"),		TAttribute<FText>(this, &SSessionInfoWindow::GetBuildTargetText));
+	AddInfoLine(VerticalBox, LOCTEXT("CommandLine_HeaderText",	"Command Line"),		TAttribute<FText>(this, &SSessionInfoWindow::GetCommandLineText), true);
+
+	AddSectionLine(VerticalBox, LOCTEXT("AnalysisStatus_SectionText", "Analysis Status"));
+	AddSimpleInfoLine(VerticalBox, TAttribute<FText>(this, &SSessionInfoWindow::GetStatusText), true);
+
+	AddSectionLine(VerticalBox, LOCTEXT("AnalysisModules_SectionText", "Analysis Modules"));
+	AddSimpleInfoLine(VerticalBox, TAttribute<FText>(this, &SSessionInfoWindow::GetModulesText), true);
 
 	DockTab->SetOnTabClosed(SDockTab::FOnTabClosedCallback::CreateRaw(this, &SSessionInfoWindow::OnSessionInfoTabClosed));
 
@@ -285,6 +347,13 @@ void SSessionInfoWindow::Tick(const FGeometry& AllottedGeometry, const double In
 	if (Session.IsValid())
 	{
 		TraceServices::FAnalysisSessionReadScope SessionReadScope(*Session.Get());
+
+		SessionNameText = FText::FromString(FPaths::GetBaseFilename(Session->GetName()));
+
+		FString Uri(Session->GetName());
+		FPaths::NormalizeFilename(Uri);
+		UriText = FText::FromString(Uri);
+
 		const TraceServices::IDiagnosticsProvider* DiagnosticsProvider = TraceServices::ReadDiagnosticsProvider(*Session.Get());
 
 		if (DiagnosticsProvider && DiagnosticsProvider->IsSessionInfoAvailable())
@@ -361,35 +430,6 @@ FReply SSessionInfoWindow::OnDrop(const FGeometry& MyGeometry, const FDragDropEv
 
 ////////////////////////////////////////////////////////////////////////////////////////////////////
 
-FText SSessionInfoWindow::GetSessionNameText() const
-{
-	FText SessionName;
-	TSharedPtr<const TraceServices::IAnalysisSession> Session = FInsightsManager::Get()->GetSession();
-	if (Session.IsValid())
-	{
-		TraceServices::FAnalysisSessionReadScope SessionReadScope(*Session.Get());
-		SessionName = FText::FromString(FPaths::GetBaseFilename(Session->GetName()));
-	}
-	return SessionName;
-}
-
-////////////////////////////////////////////////////////////////////////////////////////////////////
-
-FText SSessionInfoWindow::GetUriText() const
-{
-	//TODO: update code to use a SessionInfo provider instead
-	FText LocalUri;
-	TSharedPtr<const TraceServices::IAnalysisSession> Session = FInsightsManager::Get()->GetSession();
-	if (Session.IsValid())
-	{
-		TraceServices::FAnalysisSessionReadScope SessionReadScope(*Session.Get());
-		LocalUri = FText::FromString(FString(Session->GetName()));
-	}
-	return LocalUri;
-}
-
-////////////////////////////////////////////////////////////////////////////////////////////////////
-
 FText SSessionInfoWindow::GetFileSizeText() const
 {
 	//TODO
@@ -400,8 +440,6 @@ FText SSessionInfoWindow::GetFileSizeText() const
 
 FText SSessionInfoWindow::GetStatusText() const
 {
-	//TODO: add also info from a SessionInfo provider
-
 	TSharedPtr<FInsightsManager> InsightsManager = FInsightsManager::Get();
 	InsightsManager->UpdateSessionDuration();
 
diff --git a/Engine/Source/Developer/TraceInsights/Private/Insights/Widgets/SSessionInfoWindow.h b/Engine/Source/Developer/TraceInsights/Private/Insights/Widgets/SSessionInfoWindow.h
--- a/Engine/Source/Developer/TraceInsights/Private/Insights/Widgets/SSessionInfoWindow.h
+++ b/Engine/Source/Developer/TraceInsights/Private/Insights/Widgets/SSessionInfoWindow.h
@@ -119,12 +119,16 @@ private:
 	static void FillMenu(FMenuBuilder& MenuBuilder, const TSharedPtr<FTabManager> TabManager);
 
 private:
+	void AddSectionLine(TSharedPtr<SVerticalBox> InVerticalBox, const FText& InSectionName) const;
+	TSharedRef<SWidget> CreateTextBox(const TAttribute<FText>& InText, bool bMultiLine) const;
 	void AddInfoLine(TSharedPtr<SVerticalBox> InVerticalBox, const FText& InHeader, const TAttribute<FText>& InValue, bool bMultiLine = false) const;
+	void AddSimpleInfoLine(TSharedPtr<SVerticalBox> InVerticalBox, const TAttribute<FText>& InValue, bool bMultiLine = false) const;
 
 	TSharedRef<SDockTab> SpawnTab_SessionInfo(const FSpawnTabArgs& Args);
 	void OnSessionInfoTabClosed(TSharedRef<SDockTab> TabBeingClosed);
 
-	FText GetSessionNameText() const;
+	FText GetSessionNameText() const { return SessionNameText; }
+	FText GetUriText() const { return UriText; }
 
 	FText GetPlatformText() const { return PlatformText; }
 	FText GetAppNameText() const { return AppNameText; }
@@ -135,9 +139,9 @@ private:
 	FText GetBuildTargetText() const { return BuildTargetTypeText; }
 	FText GetCommandLineText() const { return CommandLineText; }
 
-	FText GetUriText() const;
 	FText GetFileSizeText() const;
 	FText GetStatusText() const;
+
 	FText GetModulesText() const;
 
 public:
@@ -151,6 +155,9 @@ private:
 	/** Holds the tab manager that manages the front-end's tabs. */
 	TSharedPtr<FTabManager> TabManager;
 
+	FText SessionNameText;
+	FText UriText;
+
 	FText PlatformText;
 	FText AppNameText;
 	FText BranchText;
@@ -160,9 +167,4 @@ private:
 	FText BuildTargetTypeText;
 	FText CommandLineText;
 	bool bIsSessionInfoSet = false;
-
-	//FText Uri;
-	//FDateTime Timestamp;
-	//uint64 Size;
-	//bool bIsLive;
 };
