[Insights]
  - Added context menu actions to navigate between Cpu Core tracks and Cpu Thread tracks.
  - Added context menu actions to dock the corresponding track (Cpu Core track <--> Cpu Thread track) to top/bottom.

diff --git a/Engine/Source/Developer/TraceInsights/Private/Insights/ContextSwitches/ViewModels/ContextSwitchTimingEvent.cpp b/Engine/Source/Developer/TraceInsights/Private/Insights/ContextSwitches/ViewModels/ContextSwitchTimingEvent.cpp
--- a/Engine/Source/Developer/TraceInsights/Private/Insights/ContextSwitches/ViewModels/ContextSwitchTimingEvent.cpp
+++ b/Engine/Source/Developer/TraceInsights/Private/Insights/ContextSwitches/ViewModels/ContextSwitchTimingEvent.cpp
@@ -1,6 +1,6 @@
 // Copyright Epic Games, Inc. All Rights Reserved.
 
-#include "Insights/ContextSwitches/ViewModels/ContextSwitchTimingEvent.h"
+#include "ContextSwitchTimingEvent.h"
 
 ////////////////////////////////////////////////////////////////////////////////////////////////////
 // FContextSwitchTimingEvent
diff --git a/Engine/Source/Developer/TraceInsights/Private/Insights/ContextSwitches/ViewModels/ContextSwitchesSharedState.cpp b/Engine/Source/Developer/TraceInsights/Private/Insights/ContextSwitches/ViewModels/ContextSwitchesSharedState.cpp
--- a/Engine/Source/Developer/TraceInsights/Private/Insights/ContextSwitches/ViewModels/ContextSwitchesSharedState.cpp
+++ b/Engine/Source/Developer/TraceInsights/Private/Insights/ContextSwitches/ViewModels/ContextSwitchesSharedState.cpp
@@ -11,6 +11,7 @@
 // Insights
 #include "Insights/ContextSwitches/ContextSwitchesProfilerManager.h"
 #include "Insights/ContextSwitches/ViewModels/ContextSwitchesTimingTrack.h"
+#include "Insights/ContextSwitches/ViewModels/ContextSwitchTimingEvent.h"
 #include "Insights/ContextSwitches/ViewModels/CpuCoreTimingTrack.h"
 #include "Insights/ITimingViewSession.h"
 #include "Insights/InsightsManager.h"
@@ -28,40 +29,82 @@ namespace Insights
 // FContextSwitchesStateCommands
 ////////////////////////////////////////////////////////////////////////////////////////////////////
 
-class FContextSwitchesStateCommands : public TCommands<FContextSwitchesStateCommands>
+FContextSwitchesStateCommands::FContextSwitchesStateCommands()
+	: TCommands<FContextSwitchesStateCommands>(
+		TEXT("FContextSwitchesStateCommands"),
+		NSLOCTEXT("FContextSwitchesStateCommands", "Context Switches State Commands", "Context Switches Commands"),
+		NAME_None,
+		FEditorStyle::Get().GetStyleSetName())
 {
-public:
-	FContextSwitchesStateCommands()
-		: TCommands<FContextSwitchesStateCommands>(TEXT("FContextSwitchesStateCommands"), NSLOCTEXT("FContextSwitchesStateCommands", "Context Switches State Commands", "Context Switches Commands"), NAME_None, FEditorStyle::Get().GetStyleSetName())
-	{
-	}
+}
 
-	virtual ~FContextSwitchesStateCommands()
-	{
-	}
+////////////////////////////////////////////////////////////////////////////////////////////////////
 
-	// UI_COMMAND takes long for the compiler to optimize
-	PRAGMA_DISABLE_OPTIMIZATION
-	virtual void RegisterCommands() override
-	{
-		UI_COMMAND(Command_ShowCoreTracks, "Core Tracks", "Show/hide the Cpu Core tracks.", EUserInterfaceActionType::ToggleButton, FInputChord(EModifierKey::Alt, EKeys::C));
-		UI_COMMAND(Command_ShowContextSwitches, "Context Switches", "Show/hide context switches on top of cpu timing tracks.", EUserInterfaceActionType::ToggleButton, FInputChord(EModifierKey::Shift, EKeys::C));
-		UI_COMMAND(Command_ShowOverlays, "Overlays", "Extend the visualisation of context switches over the cpu timing tracks.", EUserInterfaceActionType::ToggleButton, FInputChord(EModifierKey::Shift, EKeys::O));
-		UI_COMMAND(Command_ShowExtendedLines, "Extended Lines", "Show/hide the extended vertical lines at edges of each context switch event.", EUserInterfaceActionType::ToggleButton, FInputChord(EModifierKey::Shift, EKeys::L));
-	}
-	PRAGMA_ENABLE_OPTIMIZATION
+FContextSwitchesStateCommands::~FContextSwitchesStateCommands()
+{
+}
+
+////////////////////////////////////////////////////////////////////////////////////////////////////
 
-	TSharedPtr<FUICommandInfo> Command_ShowCoreTracks;
-	TSharedPtr<FUICommandInfo> Command_ShowContextSwitches;
-	TSharedPtr<FUICommandInfo> Command_ShowOverlays;
-	TSharedPtr<FUICommandInfo> Command_ShowExtendedLines;
-};
+// UI_COMMAND takes long for the compiler to optimize
+PRAGMA_DISABLE_OPTIMIZATION
+void FContextSwitchesStateCommands::RegisterCommands()
+{
+	UI_COMMAND(Command_ShowCoreTracks,
+		"Core Tracks",
+		"Shows/hides the Cpu Core tracks.",
+		EUserInterfaceActionType::ToggleButton,
+		FInputChord(EModifierKey::Alt, EKeys::C));
+
+	UI_COMMAND(Command_ShowContextSwitches,
+		"Context Switches",
+		"Shows/hides the context switches on top of cpu timing tracks.",
+		EUserInterfaceActionType::ToggleButton,
+		FInputChord(EModifierKey::Shift, EKeys::C));
+
+	UI_COMMAND(Command_ShowOverlays,
+		"Overlays",
+		"Extends the visualisation of context switches over the cpu timing tracks.",
+		EUserInterfaceActionType::ToggleButton,
+		FInputChord(EModifierKey::Shift, EKeys::O));
+
+	UI_COMMAND(Command_ShowExtendedLines,
+		"Extended Lines",
+		"Shows/hides the extended vertical lines at edges of each context switch event.",
+		EUserInterfaceActionType::ToggleButton,
+		FInputChord(EModifierKey::Shift, EKeys::L));
+
+	UI_COMMAND(Command_NavigateToCpuThreadEvent,
+		"Go To Cpu Thread Track",
+		"Selects the context switch event in the corresponding Cpu Thread track.",
+		EUserInterfaceActionType::Button,
+		FInputChord(/*EModifierKey::Control, EKeys::Tab*/));
+
+	UI_COMMAND(Command_DockCpuThreadTrackToBottom,
+		"Dock Cpu Thread Track To Bottom",
+		"Docks the corresponding Cpu Thread track to the bottom of the Timing view.",
+		EUserInterfaceActionType::Button,
+		FInputChord());
+
+	UI_COMMAND(Command_NavigateToCpuCoreEvent,
+		"Go To Cpu Core Track",
+		"Selects the timing event in the corresponding Cpu Core track.",
+		EUserInterfaceActionType::Button,
+		FInputChord(EModifierKey::Control, EKeys::Tab));
+
+	UI_COMMAND(Command_DockCpuCoreTrackToTop,
+		"Dock Cpu Core Track To Top",
+		"Docks the corresponding Cpu Core track to the top of the Timing view.",
+		EUserInterfaceActionType::Button,
+		FInputChord());
+}
+PRAGMA_ENABLE_OPTIMIZATION
 
 ////////////////////////////////////////////////////////////////////////////////////////////////////
 // FContextSwitchesSharedState
 ////////////////////////////////////////////////////////////////////////////////////////////////////
 
-FContextSwitchesSharedState::FContextSwitchesSharedState(STimingView* InTimingView) 
+FContextSwitchesSharedState::FContextSwitchesSharedState(STimingView* InTimingView)
 	: TimingView(InTimingView)
 	, ThreadsSerial(0)
 	, CpuCoresSerial(0)
@@ -208,27 +251,47 @@ void FContextSwitchesSharedState::AddCommands()
 
 	CommandList->MapAction(
 		FContextSwitchesStateCommands::Get().Command_ShowCoreTracks,
-		FExecuteAction::CreateSP(this, &FContextSwitchesSharedState::ContextMenu_ShowCoreTracks_Execute),
-		FCanExecuteAction::CreateSP(this, &FContextSwitchesSharedState::ContextMenu_ShowCoreTracks_CanExecute),
-		FIsActionChecked::CreateSP(this, &FContextSwitchesSharedState::ContextMenu_ShowCoreTracks_IsChecked));
+		FExecuteAction::CreateSP(this, &FContextSwitchesSharedState::Command_ShowCoreTracks_Execute),
+		FCanExecuteAction::CreateSP(this, &FContextSwitchesSharedState::Command_ShowCoreTracks_CanExecute),
+		FIsActionChecked::CreateSP(this, &FContextSwitchesSharedState::Command_ShowCoreTracks_IsChecked));
 
 	CommandList->MapAction(
 		FContextSwitchesStateCommands::Get().Command_ShowContextSwitches,
-		FExecuteAction::CreateSP(this, &FContextSwitchesSharedState::ContextMenu_ShowContextSwitches_Execute),
-		FCanExecuteAction::CreateSP(this, &FContextSwitchesSharedState::ContextMenu_ShowContextSwitches_CanExecute),
-		FIsActionChecked::CreateSP(this, &FContextSwitchesSharedState::ContextMenu_ShowContextSwitches_IsChecked));
+		FExecuteAction::CreateSP(this, &FContextSwitchesSharedState::Command_ShowContextSwitches_Execute),
+		FCanExecuteAction::CreateSP(this, &FContextSwitchesSharedState::Command_ShowContextSwitches_CanExecute),
+		FIsActionChecked::CreateSP(this, &FContextSwitchesSharedState::Command_ShowContextSwitches_IsChecked));
 
 	CommandList->MapAction(
 		FContextSwitchesStateCommands::Get().Command_ShowOverlays,
-		FExecuteAction::CreateSP(this, &FContextSwitchesSharedState::ContextMenu_ShowOverlays_Execute),
-		FCanExecuteAction::CreateSP(this, &FContextSwitchesSharedState::ContextMenu_ShowOverlays_CanExecute),
-		FIsActionChecked::CreateSP(this, &FContextSwitchesSharedState::ContextMenu_ShowOverlays_IsChecked));
+		FExecuteAction::CreateSP(this, &FContextSwitchesSharedState::Command_ShowOverlays_Execute),
+		FCanExecuteAction::CreateSP(this, &FContextSwitchesSharedState::Command_ShowOverlays_CanExecute),
+		FIsActionChecked::CreateSP(this, &FContextSwitchesSharedState::Command_ShowOverlays_IsChecked));
 
 	CommandList->MapAction(
 		FContextSwitchesStateCommands::Get().Command_ShowExtendedLines,
-		FExecuteAction::CreateSP(this, &FContextSwitchesSharedState::ContextMenu_ShowExtendedLines_Execute),
-		FCanExecuteAction::CreateSP(this, &FContextSwitchesSharedState::ContextMenu_ShowExtendedLines_CanExecute),
-		FIsActionChecked::CreateSP(this, &FContextSwitchesSharedState::ContextMenu_ShowExtendedLines_IsChecked));
+		FExecuteAction::CreateSP(this, &FContextSwitchesSharedState::Command_ShowExtendedLines_Execute),
+		FCanExecuteAction::CreateSP(this, &FContextSwitchesSharedState::Command_ShowExtendedLines_CanExecute),
+		FIsActionChecked::CreateSP(this, &FContextSwitchesSharedState::Command_ShowExtendedLines_IsChecked));
+
+	CommandList->MapAction(
+		FContextSwitchesStateCommands::Get().Command_NavigateToCpuThreadEvent,
+		FExecuteAction::CreateSP(this, &FContextSwitchesSharedState::Command_NavigateToCpuThreadEvent_Execute),
+		FCanExecuteAction::CreateSP(this, &FContextSwitchesSharedState::Command_NavigateToCpuThreadEvent_CanExecute));
+
+	CommandList->MapAction(
+		FContextSwitchesStateCommands::Get().Command_DockCpuThreadTrackToBottom,
+		FExecuteAction::CreateSP(this, &FContextSwitchesSharedState::Command_DockCpuThreadTrackToBottom_Execute),
+		FCanExecuteAction::CreateSP(this, &FContextSwitchesSharedState::Command_DockCpuThreadTrackToBottom_CanExecute));
+
+	CommandList->MapAction(
+		FContextSwitchesStateCommands::Get().Command_NavigateToCpuCoreEvent,
+		FExecuteAction::CreateSP(this, &FContextSwitchesSharedState::Command_NavigateToCpuCoreEvent_Execute),
+		FCanExecuteAction::CreateSP(this, &FContextSwitchesSharedState::Command_NavigateToCpuCoreEvent_CanExecute));
+
+	CommandList->MapAction(
+		FContextSwitchesStateCommands::Get().Command_DockCpuCoreTrackToTop,
+		FExecuteAction::CreateSP(this, &FContextSwitchesSharedState::Command_DockCpuCoreTrackToTop_Execute),
+		FCanExecuteAction::CreateSP(this, &FContextSwitchesSharedState::Command_DockCpuCoreTrackToTop_CanExecute));
 }
 
 ////////////////////////////////////////////////////////////////////////////////////////////////////
@@ -434,6 +497,234 @@ void FContextSwitchesSharedState::SetExtendedLinesVisible(bool bOnOff)
 
 ////////////////////////////////////////////////////////////////////////////////////////////////////
 
+bool FContextSwitchesSharedState::IsValidCpuCoreEventSelected() const
+{
+	if (TimingView == nullptr)
+	{
+		return false;
+	}
+	if (TargetTimingEvent.IsValid())
+	{
+		return TargetTimingEvent->Is<FCpuCoreTimingEvent>();
+	}
+	return TimingView->GetSelectedEvent().IsValid() && TimingView->GetSelectedEvent()->Is<FCpuCoreTimingEvent>();
+}
+
+////////////////////////////////////////////////////////////////////////////////////////////////////
+
+bool FContextSwitchesSharedState::IsValidContextSwitchEventSelected() const
+{
+	if (TimingView == nullptr)
+	{
+		return false;
+	}
+	if (TargetTimingEvent.IsValid())
+	{
+		return TargetTimingEvent->Is<FContextSwitchTimingEvent>();
+	}
+	return TimingView->GetSelectedEvent().IsValid() && TimingView->GetSelectedEvent()->Is<FContextSwitchTimingEvent>();
+}
+
+////////////////////////////////////////////////////////////////////////////////////////////////////
+
+void FContextSwitchesSharedState::Command_NavigateToCpuThreadEvent_Execute()
+{
+	if (TimingView && ensure(IsValidCpuCoreEventSelected()))
+	{
+		if (!TargetTimingEvent.IsValid())
+		{
+			TargetTimingEvent = TimingView->GetSelectedEvent();
+		}
+
+		if (ensure(TargetTimingEvent.IsValid() && TargetTimingEvent->Is<FCpuCoreTimingEvent>()))
+		{
+			const FCpuCoreTimingEvent& CpuCoreEvent = TargetTimingEvent->As<FCpuCoreTimingEvent>();
+			const uint32 SystemThreadId = CpuCoreEvent.GetSystemThreadId();
+			uint32 ThreadId;
+			const TCHAR* ThreadName;
+			GetThreadInfo(SystemThreadId, ThreadId, ThreadName);
+			if (ThreadId != ~0)
+			{
+				TSharedPtr<FThreadTimingTrack> ThreadTimingTrack = GetThreadTimingTrack(ThreadId);
+				if (ThreadTimingTrack.IsValid() && ThreadTimingTrack->IsVisible())
+				{
+					TimingView->SelectTimingTrack(ThreadTimingTrack, true);
+					//TODO: ThreadTimingTrack->SearchEvent(..); + TimingView->SelectTimingEvent(..)
+				}
+			}
+		}
+	}
+}
+
+////////////////////////////////////////////////////////////////////////////////////////////////////
+
+void FContextSwitchesSharedState::Command_DockCpuThreadTrackToBottom_Execute()
+{
+	if (TimingView && ensure(IsValidCpuCoreEventSelected()))
+	{
+		if (!TargetTimingEvent.IsValid())
+		{
+			TargetTimingEvent = TimingView->GetSelectedEvent();
+		}
+
+		if (ensure(TargetTimingEvent.IsValid() && TargetTimingEvent->Is<FCpuCoreTimingEvent>()))
+		{
+			const FCpuCoreTimingEvent& CpuCoreEvent = TargetTimingEvent->As<FCpuCoreTimingEvent>();
+			const uint32 SystemThreadId = CpuCoreEvent.GetSystemThreadId();
+			uint32 ThreadId;
+			const TCHAR* ThreadName;
+			GetThreadInfo(SystemThreadId, ThreadId, ThreadName);
+			if (ThreadId != ~0)
+			{
+				TSharedPtr<FThreadTimingTrack> ThreadTimingTrack = GetThreadTimingTrack(ThreadId);
+				if (ThreadTimingTrack.IsValid())
+				{
+					TimingView->ChangeTrackLocation(ThreadTimingTrack.ToSharedRef(), ETimingTrackLocation::BottomDocked);
+					if (!ThreadTimingTrack->IsVisible())
+					{
+						ThreadTimingTrack->Show();
+						TimingView->OnTrackVisibilityChanged();
+					}
+				}
+			}
+		}
+	}
+}
+
+////////////////////////////////////////////////////////////////////////////////////////////////////
+
+void FContextSwitchesSharedState::Command_NavigateToCpuCoreEvent_Execute()
+{
+	if (TimingView && ensure(IsValidContextSwitchEventSelected()))
+	{
+		if (!TargetTimingEvent.IsValid())
+		{
+			TargetTimingEvent = TimingView->GetSelectedEvent();
+		}
+
+		if (ensure(TargetTimingEvent.IsValid() && TargetTimingEvent->Is<FContextSwitchTimingEvent>()))
+		{
+			const FContextSwitchTimingEvent& ContextSwitchEvent = TargetTimingEvent->As<FContextSwitchTimingEvent>();
+			const uint32 CoreNumber = ContextSwitchEvent.GetCoreNumber();
+			if (CoreNumber != ~0)
+			{
+				TSharedPtr<FCpuCoreTimingTrack> CpuCoreTimingTrack = GetCpuCoreTimingTrack(CoreNumber);
+				if (CpuCoreTimingTrack.IsValid() && CpuCoreTimingTrack->IsVisible())
+				{
+					TimingView->SelectTimingTrack(CpuCoreTimingTrack, true);
+					//TODO: CpuCoreTimingTrack->SearchEvent(..); + TimingView->SelectTimingEvent(..)
+				}
+			}
+		}
+	}
+}
+
+////////////////////////////////////////////////////////////////////////////////////////////////////
+
+void FContextSwitchesSharedState::Command_DockCpuCoreTrackToTop_Execute()
+{
+	if (TimingView && ensure(IsValidContextSwitchEventSelected()))
+	{
+		if (!TargetTimingEvent.IsValid())
+		{
+			TargetTimingEvent = TimingView->GetSelectedEvent();
+		}
+
+		if (ensure(TargetTimingEvent.IsValid() && TargetTimingEvent->Is<FContextSwitchTimingEvent>()))
+		{
+			const FContextSwitchTimingEvent& ContextSwitchEvent = TargetTimingEvent->As<FContextSwitchTimingEvent>();
+			const uint32 CoreNumber = ContextSwitchEvent.GetCoreNumber();
+			if (CoreNumber != ~0)
+			{
+				TSharedPtr<FCpuCoreTimingTrack> CpuCoreTimingTrack = GetCpuCoreTimingTrack(CoreNumber);
+				if (CpuCoreTimingTrack.IsValid())
+				{
+					if (TimingView->CanChangeTrackLocation(CpuCoreTimingTrack.ToSharedRef(), ETimingTrackLocation::TopDocked))
+					{
+						TimingView->ChangeTrackLocation(CpuCoreTimingTrack.ToSharedRef(), ETimingTrackLocation::TopDocked);
+					}
+					if (!CpuCoreTimingTrack->IsVisible())
+					{
+						CpuCoreTimingTrack->Show();
+						TimingView->OnTrackVisibilityChanged();
+					}
+				}
+			}
+		}
+	}
+}
+
+////////////////////////////////////////////////////////////////////////////////////////////////////
+
+void FContextSwitchesSharedState::GetThreadInfo(uint32 InSystemThreadId, uint32& OutThreadId, const TCHAR*& OutThreadName) const
+{
+	OutThreadId = ~0;
+	OutThreadName = TEXT("Unknown Thread");
+
+	TSharedPtr<const TraceServices::IAnalysisSession> Session = FInsightsManager::Get()->GetSession();
+	if (Session.IsValid())
+	{
+		TraceServices::FAnalysisSessionReadScope SessionReadScope(*Session.Get());
+		const TraceServices::IContextSwitchesProvider* ContextSwitchesProvider = TraceServices::ReadContextSwitchesProvider(*Session.Get());
+		if (ContextSwitchesProvider)
+		{
+			if (ContextSwitchesProvider->GetThreadId(InSystemThreadId, OutThreadId))
+			{
+				const TraceServices::IThreadProvider& ThreadProvider = TraceServices::ReadThreadProvider(*Session.Get());
+				OutThreadName = ThreadProvider.GetThreadName(OutThreadId);
+			}
+		}
+	}
+}
+
+////////////////////////////////////////////////////////////////////////////////////////////////////
+
+TSharedPtr<FThreadTimingTrack> FContextSwitchesSharedState::GetThreadTimingTrack(uint32 ThreadId) const
+{
+	TSharedPtr<FThreadTimingTrack> FoundTrack;
+
+	if (TimingView)
+	{
+		TimingView->EnumerateAllTracks([&FoundTrack, ThreadId](TSharedPtr<FBaseTimingTrack>& Track) -> bool
+		{
+			if (Track->Is<FThreadTimingTrack>() &&
+				Track->As<FThreadTimingTrack>().GetThreadId() == ThreadId)
+			{
+				FoundTrack = StaticCastSharedPtr<FThreadTimingTrack>(Track);
+				return false;
+			}
+			return true;
+		});
+	}
+
+	return FoundTrack;
+}
+
+////////////////////////////////////////////////////////////////////////////////////////////////////
+
+TSharedPtr<Insights::FCpuCoreTimingTrack> FContextSwitchesSharedState::GetCpuCoreTimingTrack(uint32 CoreNumber) const
+{
+	TSharedPtr<FCpuCoreTimingTrack> FoundTrack;
+
+	if (TimingView)
+	{
+		TimingView->EnumerateAllTracks([&FoundTrack, CoreNumber](TSharedPtr<FBaseTimingTrack>& Track) -> bool
+		{
+			if (Track->Is<Insights::FCpuCoreTimingTrack>() &&
+				Track->As<Insights::FCpuCoreTimingTrack>().GetCoreNumber() == CoreNumber)
+			{
+				FoundTrack = StaticCastSharedPtr<Insights::FCpuCoreTimingTrack>(Track);
+				return false;
+			}
+			return true;
+		});
+	}
+
+	return FoundTrack;
+}
+
+////////////////////////////////////////////////////////////////////////////////////////////////////
+
 } // namespace Insights
 
 #undef LOCTEXT_NAMESPACE
diff --git a/Engine/Source/Developer/TraceInsights/Private/Insights/ContextSwitches/ViewModels/ContextSwitchesSharedState.h b/Engine/Source/Developer/TraceInsights/Private/Insights/ContextSwitches/ViewModels/ContextSwitchesSharedState.h
--- a/Engine/Source/Developer/TraceInsights/Private/Insights/ContextSwitches/ViewModels/ContextSwitchesSharedState.h
+++ b/Engine/Source/Developer/TraceInsights/Private/Insights/ContextSwitches/ViewModels/ContextSwitchesSharedState.h
@@ -9,7 +9,10 @@
 
 ////////////////////////////////////////////////////////////////////////////////////////////////////
 
+class FBaseTimingTrack;
 class FMenuBuilder;
+class FThreadTimingTrack;
+class ITimingEvent;
 class STimingView;
 
 namespace TraceServices
@@ -26,6 +29,30 @@ class ITimingViewSession;
 
 ////////////////////////////////////////////////////////////////////////////////////////////////////
 
+class FContextSwitchesStateCommands : public TCommands<FContextSwitchesStateCommands>
+{
+public:
+	FContextSwitchesStateCommands();
+	virtual ~FContextSwitchesStateCommands();
+	virtual void RegisterCommands() override;
+
+	// Commands for the Tracks Filter menu.
+	TSharedPtr<FUICommandInfo> Command_ShowCoreTracks;
+	TSharedPtr<FUICommandInfo> Command_ShowContextSwitches;
+	TSharedPtr<FUICommandInfo> Command_ShowOverlays;
+	TSharedPtr<FUICommandInfo> Command_ShowExtendedLines;
+
+	// Commands for a Cpu Core track (context menu).
+	TSharedPtr<FUICommandInfo> Command_NavigateToCpuThreadEvent;
+	TSharedPtr<FUICommandInfo> Command_DockCpuThreadTrackToBottom;
+
+	// Commands for a Cpu Thread track (context menu).
+	TSharedPtr<FUICommandInfo> Command_NavigateToCpuCoreEvent;
+	TSharedPtr<FUICommandInfo> Command_DockCpuCoreTrackToTop;
+};
+
+////////////////////////////////////////////////////////////////////////////////////////////////////
+
 class FContextSwitchesSharedState : public Insights::ITimingViewExtender, public TSharedFromThis<FContextSwitchesSharedState>
 {
 public:
@@ -70,8 +97,17 @@ public:
 	void ToggleExtendedLines() { SetExtendedLinesVisible(!bAreExtendedLinesVisible); }
 	void SetExtendedLinesVisible(bool bOnOff);
 
+	void SetTargetTimingEvent(const TSharedPtr<const ITimingEvent> InEvent)
+	{
+		TargetTimingEvent = InEvent;
+	}
+
 	void AddCommands();
 
+	void GetThreadInfo(uint32 InSystemThreadId, uint32& OutThreadId, const TCHAR*& OutThreadName) const;
+	TSharedPtr<FThreadTimingTrack> GetThreadTimingTrack(uint32 ThreadId) const;
+	TSharedPtr<Insights::FCpuCoreTimingTrack> GetCpuCoreTimingTrack(uint32 CoreNumber) const;
+
 private:
 	void AddCoreTracks();
 	void RemoveCoreTracks();
@@ -81,21 +117,36 @@ private:
 
 	void BuildSubMenu(FMenuBuilder& InMenuBuilder);
 
-	void ContextMenu_ShowCoreTracks_Execute() { ToggleCoreTracks(); }
-	bool ContextMenu_ShowCoreTracks_CanExecute() const { return AreContextSwitchesAvailable(); }
-	bool ContextMenu_ShowCoreTracks_IsChecked() const { return AreCoreTracksVisible(); }
+	void Command_ShowCoreTracks_Execute() { ToggleCoreTracks(); }
+	bool Command_ShowCoreTracks_CanExecute() const { return AreContextSwitchesAvailable(); }
+	bool Command_ShowCoreTracks_IsChecked() const { return AreCoreTracksVisible(); }
+
+	void Command_ShowContextSwitches_Execute() { ToggleContextSwitches(); }
+	bool Command_ShowContextSwitches_CanExecute() const { return AreContextSwitchesAvailable(); }
+	bool Command_ShowContextSwitches_IsChecked() const { return AreContextSwitchesVisible(); }
+
+	void Command_ShowOverlays_Execute() { ToggleOverlays(); }
+	bool Command_ShowOverlays_CanExecute() const { return AreContextSwitchesAvailable() && AreContextSwitchesVisible(); }
+	bool Command_ShowOverlays_IsChecked() const { return AreOverlaysVisible(); }
+
+	void Command_ShowExtendedLines_Execute() { ToggleExtendedLines(); }
+	bool Command_ShowExtendedLines_CanExecute() const { return AreContextSwitchesAvailable() && AreContextSwitchesVisible() && AreOverlaysVisible(); }
+	bool Command_ShowExtendedLines_IsChecked() const { return AreExtendedLinesVisible(); }
 
-	void ContextMenu_ShowContextSwitches_Execute() { ToggleContextSwitches(); }
-	bool ContextMenu_ShowContextSwitches_CanExecute() const { return AreContextSwitchesAvailable(); }
-	bool ContextMenu_ShowContextSwitches_IsChecked() const { return AreContextSwitchesVisible(); }
+	bool IsValidCpuCoreEventSelected() const;
+	bool IsValidContextSwitchEventSelected() const;
 
-	void ContextMenu_ShowOverlays_Execute() { ToggleOverlays(); }
-	bool ContextMenu_ShowOverlays_CanExecute() const { return AreContextSwitchesAvailable() && AreContextSwitchesVisible(); }
-	bool ContextMenu_ShowOverlays_IsChecked() const { return AreOverlaysVisible(); }
+	void Command_NavigateToCpuThreadEvent_Execute();
+	bool Command_NavigateToCpuThreadEvent_CanExecute() const { return AreContextSwitchesAvailable() && AreContextSwitchesVisible() && IsValidCpuCoreEventSelected(); }
 
-	void ContextMenu_ShowExtendedLines_Execute() { ToggleExtendedLines(); }
-	bool ContextMenu_ShowExtendedLines_CanExecute() const { return AreContextSwitchesAvailable() && AreContextSwitchesVisible() && AreOverlaysVisible(); }
-	bool ContextMenu_ShowExtendedLines_IsChecked() const { return AreExtendedLinesVisible(); }
+	void Command_DockCpuThreadTrackToBottom_Execute();
+	bool Command_DockCpuThreadTrackToBottom_CanExecute() const { return AreContextSwitchesAvailable() && AreContextSwitchesVisible() && IsValidCpuCoreEventSelected(); }
+
+	void Command_NavigateToCpuCoreEvent_Execute();
+	bool Command_NavigateToCpuCoreEvent_CanExecute() const { return AreContextSwitchesAvailable() && AreCoreTracksVisible() && IsValidContextSwitchEventSelected(); }
+
+	void Command_DockCpuCoreTrackToTop_Execute();
+	bool Command_DockCpuCoreTrackToTop_CanExecute() const { return AreContextSwitchesAvailable() && AreCoreTracksVisible() && IsValidContextSwitchEventSelected(); }
 
 private:
 	STimingView* TimingView;
@@ -112,6 +163,8 @@ private:
 	bool bAreExtendedLinesVisible;
 
 	bool bSyncWithProviders;
+
+	TSharedPtr<const ITimingEvent> TargetTimingEvent;
 };
 
 ////////////////////////////////////////////////////////////////////////////////////////////////////
diff --git a/Engine/Source/Developer/TraceInsights/Private/Insights/ContextSwitches/ViewModels/ContextSwitchesTimingTrack.cpp b/Engine/Source/Developer/TraceInsights/Private/Insights/ContextSwitches/ViewModels/ContextSwitchesTimingTrack.cpp
--- a/Engine/Source/Developer/TraceInsights/Private/Insights/ContextSwitches/ViewModels/ContextSwitchesTimingTrack.cpp
+++ b/Engine/Source/Developer/TraceInsights/Private/Insights/ContextSwitches/ViewModels/ContextSwitchesTimingTrack.cpp
@@ -12,6 +12,7 @@
 #include "Insights/ViewModels/TimingTrackViewport.h"
 #include "Insights/ViewModels/TimingViewDrawHelper.h"
 #include "Insights/ViewModels/TooltipDrawState.h"
+#include "Insights/Widgets/STimingView.h"
 
 #define LOCTEXT_NAMESPACE "FContextSwitchesTimingTrack"
 
@@ -249,6 +250,31 @@ void FContextSwitchesTimingTrack::InitTooltip(FTooltipDrawState& InOutTooltip, c
 
 ////////////////////////////////////////////////////////////////////////////////////////////////////
 
+void FContextSwitchesTimingTrack::BuildContextMenu(FMenuBuilder& InOutMenuBuilder)
+{
+	if (SharedState.GetTimingView())
+	{
+		const TSharedPtr<const ITimingEvent> HoveredEvent = SharedState.GetTimingView()->GetHoveredEvent();
+		if (HoveredEvent.IsValid() &&
+			&HoveredEvent->GetTrack().Get() == this &&
+			HoveredEvent->Is<FContextSwitchTimingEvent>())
+		{
+			const FContextSwitchTimingEvent& CpuCoreEvent = HoveredEvent->As<FContextSwitchTimingEvent>();
+			const uint32 CoreNumber = CpuCoreEvent.GetCoreNumber();
+
+			SharedState.SetTargetTimingEvent(HoveredEvent);
+
+			FString SectionName = FString::Printf(TEXT("Core %d"), CoreNumber);
+			InOutMenuBuilder.BeginSection("CpuCore", FText::FromString(SectionName));
+			InOutMenuBuilder.AddMenuEntry(FContextSwitchesStateCommands::Get().Command_NavigateToCpuCoreEvent);
+			InOutMenuBuilder.AddMenuEntry(FContextSwitchesStateCommands::Get().Command_DockCpuCoreTrackToTop);
+			InOutMenuBuilder.EndSection();
+		}
+	}
+}
+
+////////////////////////////////////////////////////////////////////////////////////////////////////
+
 } // namespace Insights
 
 #undef LOCTEXT_NAMESPACE
diff --git a/Engine/Source/Developer/TraceInsights/Private/Insights/ContextSwitches/ViewModels/ContextSwitchesTimingTrack.h b/Engine/Source/Developer/TraceInsights/Private/Insights/ContextSwitches/ViewModels/ContextSwitchesTimingTrack.h
--- a/Engine/Source/Developer/TraceInsights/Private/Insights/ContextSwitches/ViewModels/ContextSwitchesTimingTrack.h
+++ b/Engine/Source/Developer/TraceInsights/Private/Insights/ContextSwitches/ViewModels/ContextSwitchesTimingTrack.h
@@ -44,6 +44,8 @@ public:
 
 	virtual void InitTooltip(FTooltipDrawState& InOutTooltip, const ITimingEvent& InTooltipEvent) const override;
 
+	virtual void BuildContextMenu(FMenuBuilder& MenuBuilder) override;
+
 protected:
 	virtual const TSharedPtr<const ITimingEvent> GetEvent(double InTime, double SecondsPerPixel, int32 Depth) const override;
 
diff --git a/Engine/Source/Developer/TraceInsights/Private/Insights/ContextSwitches/ViewModels/CpuCoreTimingTrack.cpp b/Engine/Source/Developer/TraceInsights/Private/Insights/ContextSwitches/ViewModels/CpuCoreTimingTrack.cpp
--- a/Engine/Source/Developer/TraceInsights/Private/Insights/ContextSwitches/ViewModels/CpuCoreTimingTrack.cpp
+++ b/Engine/Source/Developer/TraceInsights/Private/Insights/ContextSwitches/ViewModels/CpuCoreTimingTrack.cpp
@@ -13,6 +13,7 @@
 #include "Insights/ViewModels/TimingTrackViewport.h"
 #include "Insights/ViewModels/TimingViewDrawHelper.h"
 #include "Insights/ViewModels/TooltipDrawState.h"
+#include "Insights/Widgets/STimingView.h"
 
 #define LOCTEXT_NAMESPACE "FCpuCoreTimingTrack"
 
@@ -243,33 +244,11 @@ void FCpuCoreTimingTrack::InitTooltip(FTooltipDrawState& InOutTooltip, const ITi
 	const FCpuCoreTimingEvent& CpuCoreEvent = InTooltipEvent.As<FCpuCoreTimingEvent>();
 
 	const uint32 SystemThreadId = CpuCoreEvent.GetSystemThreadId();
-	uint32 ThreadId = ~0;
-	const TCHAR* ThreadName = nullptr;
-	{
-		TSharedPtr<const TraceServices::IAnalysisSession> Session = FInsightsManager::Get()->GetSession();
-		if (Session.IsValid())
-		{
-			TraceServices::FAnalysisSessionReadScope SessionReadScope(*Session.Get());
-			const TraceServices::IContextSwitchesProvider* ContextSwitchesProvider = TraceServices::ReadContextSwitchesProvider(*Session.Get());
-			if (ContextSwitchesProvider)
-			{
-				if (ContextSwitchesProvider->GetThreadId(SystemThreadId, ThreadId))
-				{
-					const TraceServices::IThreadProvider& ThreadProvider = TraceServices::ReadThreadProvider(*Session.Get());
-					ThreadName = ThreadProvider.GetThreadName(ThreadId);
-				}
-			}
-		}
-	}
-	if (ThreadName)
-	{
-		InOutTooltip.AddTitle(ThreadName);
-	}
-	else
-	{
-		InOutTooltip.AddTitle(TEXT("Unknown Thread"));
-	}
+	uint32 ThreadId;
+	const TCHAR* ThreadName;
+	SharedState.GetThreadInfo(SystemThreadId, ThreadId, ThreadName);
 
+	InOutTooltip.AddTitle(ThreadName);
 	InOutTooltip.AddNameValueTextLine(TEXT("System Thread Id:"), FString::Printf(TEXT("%d"), SystemThreadId));
 	if (ThreadId != ~0)
 	{
@@ -285,6 +264,40 @@ void FCpuCoreTimingTrack::InitTooltip(FTooltipDrawState& InOutTooltip, const ITi
 
 ////////////////////////////////////////////////////////////////////////////////////////////////////
 
+void FCpuCoreTimingTrack::BuildContextMenu(FMenuBuilder& InOutMenuBuilder)
+{
+	if (SharedState.GetTimingView())
+	{
+		const TSharedPtr<const ITimingEvent> HoveredEvent = SharedState.GetTimingView()->GetHoveredEvent();
+		if (HoveredEvent.IsValid() &&
+			&HoveredEvent->GetTrack().Get() == this &&
+			HoveredEvent->Is<FCpuCoreTimingEvent>())
+		{
+			const FCpuCoreTimingEvent& CpuCoreEvent = HoveredEvent->As<FCpuCoreTimingEvent>();
+			const uint32 SystemThreadId = CpuCoreEvent.GetSystemThreadId();
+			uint32 ThreadId;
+			const TCHAR* ThreadName;
+			SharedState.GetThreadInfo(SystemThreadId, ThreadId, ThreadName);
+
+			if (ThreadId != ~0)
+			{
+				SharedState.SetTargetTimingEvent(HoveredEvent);
+			}
+			else
+			{
+				SharedState.SetTargetTimingEvent(nullptr);
+			}
+
+			InOutMenuBuilder.BeginSection("CpuThread", FText::FromString(ThreadName));
+			InOutMenuBuilder.AddMenuEntry(FContextSwitchesStateCommands::Get().Command_NavigateToCpuThreadEvent);
+			InOutMenuBuilder.AddMenuEntry(FContextSwitchesStateCommands::Get().Command_DockCpuThreadTrackToBottom);
+			InOutMenuBuilder.EndSection();
+		}
+	}
+}
+
+////////////////////////////////////////////////////////////////////////////////////////////////////
+
 } // namespace Insights
 
 #undef LOCTEXT_NAMESPACE
diff --git a/Engine/Source/Developer/TraceInsights/Private/Insights/ContextSwitches/ViewModels/CpuCoreTimingTrack.h b/Engine/Source/Developer/TraceInsights/Private/Insights/ContextSwitches/ViewModels/CpuCoreTimingTrack.h
--- a/Engine/Source/Developer/TraceInsights/Private/Insights/ContextSwitches/ViewModels/CpuCoreTimingTrack.h
+++ b/Engine/Source/Developer/TraceInsights/Private/Insights/ContextSwitches/ViewModels/CpuCoreTimingTrack.h
@@ -41,6 +41,8 @@ public:
 
 	virtual void InitTooltip(FTooltipDrawState& InOutTooltip, const ITimingEvent& InTooltipEvent) const override;
 
+	virtual void BuildContextMenu(FMenuBuilder& MenuBuilder) override;
+
 protected:
 	virtual const TSharedPtr<const ITimingEvent> GetEvent(double InTime, double SecondsPerPixel, int32 Depth) const override;
 
diff --git a/Engine/Source/Developer/TraceInsights/Private/Insights/ViewModels/ThreadTimingTrack.cpp b/Engine/Source/Developer/TraceInsights/Private/Insights/ViewModels/ThreadTimingTrack.cpp
--- a/Engine/Source/Developer/TraceInsights/Private/Insights/ViewModels/ThreadTimingTrack.cpp
+++ b/Engine/Source/Developer/TraceInsights/Private/Insights/ViewModels/ThreadTimingTrack.cpp
@@ -15,6 +15,7 @@
 #include "Insights/Common/PaintUtils.h"
 #include "Insights/Common/TimeUtils.h"
 #include "Insights/InsightsManager.h"
+#include "Insights/InsightsStyle.h"
 #include "Insights/ITimingViewSession.h"
 #include "Insights/TimingProfilerManager.h"
 #include "Insights/ViewModels/Filters.h"
@@ -506,8 +507,8 @@ void FThreadTimingSharedState::ExtendFilterMenu(Insights::ITimingViewSession& In
 	{
 		//TODO: MenuBuilder.AddMenuEntry(Commands.ShowAllGpuTracks);
 		InOutMenuBuilder.AddMenuEntry(
-			LOCTEXT("ShowAllGpuTracks", "GPU Track - Y"),
-			LOCTEXT("ShowAllGpuTracks_Tooltip", "Show/hide the GPU track"),
+			LOCTEXT("ShowAllGpuTracks", "GPU Track(s) - Y"),
+			LOCTEXT("ShowAllGpuTracks_Tooltip", "Show/hide the GPU track(s)."),
 			FSlateIcon(),
 			FUIAction(FExecuteAction::CreateSP(this, &FThreadTimingSharedState::ShowHideAllGpuTracks),
 					  FCanExecuteAction(),
@@ -519,7 +520,7 @@ void FThreadTimingSharedState::ExtendFilterMenu(Insights::ITimingViewSession& In
 		//TODO: MenuBuilder.AddMenuEntry(Commands.ShowAllCpuTracks);
 		InOutMenuBuilder.AddMenuEntry(
 			LOCTEXT("ShowAllCpuTracks", "CPU Thread Tracks - U"),
-			LOCTEXT("ShowAllCpuTracks_Tooltip", "Show/hide all CPU tracks (and all CPU thread groups)"),
+			LOCTEXT("ShowAllCpuTracks_Tooltip", "Show/hide all CPU tracks (and all CPU thread groups)."),
 			FSlateIcon(),
 			FUIAction(FExecuteAction::CreateSP(this, &FThreadTimingSharedState::ShowHideAllCpuTracks),
 					  FCanExecuteAction(),
@@ -1232,15 +1233,15 @@ const TSharedPtr<const ITimingEvent> FThreadTimingTrack::GetEvent(float InPosX,
 
 	float TopLaneY = 0.0f;
 	float TrackLanesHeight = 0.0f;
-	if (ChildTrack.IsValid())
+	if (ChildTrack.IsValid() && ChildTrack->GetHeight() > 0.0f)
 	{
 		const float HeaderDY = InPosY - ChildTrack->GetPosY();
-		if (HeaderDY >= 0 && HeaderDY < ChildTrack->GetHeight())
+		if (HeaderDY >= 0.0f && HeaderDY < ChildTrack->GetHeight())
 		{
 			return ChildTrack->GetEvent(InPosX, InPosY, Viewport);
 		}
 		TopLaneY = GetPosY() + 1.0f + Layout.TimelineDY + ChildTrack->GetHeight() + Layout.ChildTimelineDY;
-		TrackLanesHeight = GetHeight() - ChildTrack->GetHeight() - 1.0f - 2 * Layout.TimelineDY - Layout.ChildTimelineDY;
+		TrackLanesHeight = GetHeight() - ChildTrack->GetHeight() - 1.0f - 2.0f * Layout.TimelineDY - Layout.ChildTimelineDY;
 	}
 	else
 	{
@@ -1252,7 +1253,7 @@ const TSharedPtr<const ITimingEvent> FThreadTimingTrack::GetEvent(float InPosX,
 		else
 		{
 			TopLaneY = GetPosY() + 1.0f + Layout.TimelineDY;
-			TrackLanesHeight = GetHeight() - 1.0f - 2 * Layout.TimelineDY;
+			TrackLanesHeight = GetHeight() - 1.0f - 2.0f * Layout.TimelineDY;
 		}
 	}
 
@@ -1386,26 +1387,10 @@ void FThreadTimingTrack::BuildContextMenu(FMenuBuilder& MenuBuilder)
 {
 	if (GetGroupName() != nullptr)
 	{
-		MenuBuilder.BeginSection(TEXT("Options"));
+		MenuBuilder.BeginSection("CpuThread", LOCTEXT("CpuThreadSection", "CPU Thread"));
 		{
-			FExecuteAction FilterTrackAction;
-			FilterTrackAction.BindSP(this, &FThreadTimingTrack::OnFilterTrackClicked);
-
 			MenuBuilder.AddMenuEntry(
-				(LOCTEXT("FilterTrack", "Filter Track")),
-				FText(),
-				FSlateIcon(),
-				FUIAction(FilterTrackAction, FCanExecuteAction::CreateLambda([]() { return true; })),
-				NAME_None,
-				EUserInterfaceActionType::Button
-			);
-		}
-		MenuBuilder.EndSection();
-
-		MenuBuilder.BeginSection(TEXT("Misc"));
-		{
-			MenuBuilder.AddMenuEntry(
-				FText::Format(LOCTEXT("CpuThreadGroupFmt", "CPU Thread Group: {0}"), FText::FromString(GetGroupName())),
+				FText::Format(LOCTEXT("CpuThreadGroupFmt", "Group: {0}"), FText::FromString(GetGroupName())),
 				FText(),
 				FSlateIcon(),
 				FUIAction(FExecuteAction(), FCanExecuteAction::CreateLambda([]() { return false; })),
@@ -1425,6 +1410,29 @@ void FThreadTimingTrack::BuildContextMenu(FMenuBuilder& MenuBuilder)
 		}
 		MenuBuilder.EndSection();
 	}
+
+	if (ChildTrack.IsValid())
+	{
+		ChildTrack->BuildContextMenu(MenuBuilder);
+	}
+
+	MenuBuilder.AddSeparator();
+
+	MenuBuilder.BeginSection("Options");
+	{
+		FExecuteAction FilterTrackAction;
+		FilterTrackAction.BindSP(this, &FThreadTimingTrack::OnFilterTrackClicked);
+
+		MenuBuilder.AddMenuEntry(
+			LOCTEXT("FilterTrack", "Filter Track..."),
+			FText(),
+			FSlateIcon(FInsightsStyle::GetStyleSetName(), "Icon.Filter"),
+			FUIAction(FilterTrackAction, FCanExecuteAction::CreateLambda([]() { return true; })),
+			NAME_None,
+			EUserInterfaceActionType::Button
+		);
+	}
+	MenuBuilder.EndSection();
 }
 
 ////////////////////////////////////////////////////////////////////////////////////////////////////
diff --git a/Engine/Source/Developer/TraceInsights/Private/Insights/Widgets/STimingView.cpp b/Engine/Source/Developer/TraceInsights/Private/Insights/Widgets/STimingView.cpp
--- a/Engine/Source/Developer/TraceInsights/Private/Insights/Widgets/STimingView.cpp
+++ b/Engine/Source/Developer/TraceInsights/Private/Insights/Widgets/STimingView.cpp
@@ -2806,7 +2806,11 @@ void STimingView::ShowContextMenu(const FPointerEvent& MouseEvent)
 		bHasAnyActions = true;
 	}
 
-	MenuBuilder.AddMenuEntry(Commands.QuickFind);
+	MenuBuilder.AddMenuEntry(Commands.QuickFind,
+		NAME_None,
+		TAttribute<FText>(),
+		TAttribute<FText>(),
+		FSlateIcon(FInsightsStyle::GetStyleSetName(), "Icon.Find"));
 
 	for (Insights::ITimingViewExtender* Extender : GetExtenders())
 	{
@@ -2951,6 +2955,40 @@ void STimingView::CreateTrackLocationMenu(FMenuBuilder& MenuBuilder, TSharedRef<
 
 ////////////////////////////////////////////////////////////////////////////////////////////////////
 
+void STimingView::EnumerateAllTracks(TFunctionRef<bool(TSharedPtr<FBaseTimingTrack>&)> Callback)
+{
+	for (TSharedPtr<FBaseTimingTrack>& TrackPtr : TopDockedTracks)
+	{
+		if (!Callback(TrackPtr))
+		{
+			return;
+		}
+	}
+	for (TSharedPtr<FBaseTimingTrack>& TrackPtr : ScrollableTracks)
+	{
+		if (!Callback(TrackPtr))
+		{
+			return;
+		}
+	}
+	for (TSharedPtr<FBaseTimingTrack>& TrackPtr : BottomDockedTracks)
+	{
+		if (!Callback(TrackPtr))
+		{
+			return;
+		}
+	}
+	for (TSharedPtr<FBaseTimingTrack>& TrackPtr : ForegroundTracks)
+	{
+		if (!Callback(TrackPtr))
+		{
+			return;
+		}
+	}
+}
+
+////////////////////////////////////////////////////////////////////////////////////////////////////
+
 void STimingView::ChangeTrackLocation(TSharedRef<FBaseTimingTrack> Track, ETimingTrackLocation NewLocation)
 {
 	if (ensure(CanChangeTrackLocation(Track, NewLocation)))
@@ -3213,6 +3251,33 @@ void STimingView::ScrollAtPosY(float ScrollPosY)
 
 ////////////////////////////////////////////////////////////////////////////////////////////////////
 
+void STimingView::BringIntoViewY(float InTopY, float InBottomY)
+{
+	const float ScrollY = Viewport.GetScrollPosY();
+	const float ScrollH = Viewport.GetScrollableAreaHeight();
+	if (ScrollY > InTopY)
+	{
+		ScrollAtPosY(InTopY);
+	}
+	else if (ScrollY + ScrollH < InBottomY)
+	{
+		ScrollAtPosY(FMath::Min(InTopY, InBottomY - ScrollH));
+	}
+}
+
+////////////////////////////////////////////////////////////////////////////////////////////////////
+
+void STimingView::BringScrollableTrackIntoView(const FBaseTimingTrack& InTrack)
+{
+	if (ensure(InTrack.GetLocation() == ETimingTrackLocation::Scrollable))
+	{
+		const float TopY = Viewport.GetScrollPosY() - Viewport.GetTopOffset() + InTrack.GetPosY();
+		BringIntoViewY(TopY, TopY + InTrack.GetHeight());
+	}
+}
+
+////////////////////////////////////////////////////////////////////////////////////////////////////
+
 void STimingView::ScrollAtTime(double StartTime)
 {
 	if (Viewport.ScrollAtTime(StartTime))
@@ -3601,25 +3666,60 @@ void STimingView::OnSelectedTimingEventChanged()
 
 void STimingView::SelectHoveredTimingTrack()
 {
-	if (SelectedTrack != HoveredTrack)
+	SelectTimingTrack(HoveredTrack, false);
+}
+
+////////////////////////////////////////////////////////////////////////////////////////////////////
+
+void STimingView::SelectHoveredTimingEvent()
+{
+	SelectTimingEvent(HoveredEvent, true);
+}
+
+////////////////////////////////////////////////////////////////////////////////////////////////////
+
+void STimingView::SelectTimingTrack(const TSharedPtr<FBaseTimingTrack> InTrack, bool bBringTrackIntoView)
+{
+	if (SelectedTrack != InTrack)
 	{
-		SelectedTrack = HoveredTrack;
+		SelectedTrack = InTrack;
+
+		if (SelectedTrack.IsValid())
+		{
+			if (bBringTrackIntoView &&
+				SelectedTrack->GetLocation() == ETimingTrackLocation::Scrollable)
+			{
+				TSharedPtr<FBaseTimingTrack> ParentTrack = SelectedTrack->GetParentTrack().Pin();
+				if (ParentTrack.IsValid())
+				{
+					BringScrollableTrackIntoView(*ParentTrack);
+				}
+				else
+				{
+					BringScrollableTrackIntoView(*SelectedTrack);
+				}
+			}
+		}
+
 		OnSelectedTrackChangedDelegate.Broadcast(SelectedTrack);
 	}
 }
 
 ////////////////////////////////////////////////////////////////////////////////////////////////////
 
-void STimingView::SelectHoveredTimingEvent()
+void STimingView::SelectTimingEvent(const TSharedPtr<const ITimingEvent> InEvent, bool bBringEventIntoView)
 {
-	if (SelectedEvent != HoveredEvent)
+	if (SelectedEvent != InEvent)
 	{
-		SelectedEvent = HoveredEvent;
+		SelectedEvent = InEvent;
 
 		if (SelectedEvent.IsValid())
 		{
 			LastSelectionType = ESelectionType::TimingEvent;
-			BringIntoView(SelectedEvent->GetStartTime(), SelectedEvent->GetEndTime());
+			if (bBringEventIntoView)
+			{
+				BringIntoView(SelectedEvent->GetStartTime(), SelectedEvent->GetEndTime());
+			}
 		}
 
 		OnSelectedTimingEventChanged();
@@ -3804,6 +3904,10 @@ void STimingView::FrameSelection()
 		{
 			EndTime += 1.0 / Viewport.GetScaleX(); // +1px
 		}
+		if (SelectedEvent->GetTrack()->GetLocation() == ETimingTrackLocation::Scrollable)
+		{
+			BringScrollableTrackIntoView(*SelectedEvent->GetTrack());
+		}
 	}
 	else
 	{
diff --git a/Engine/Source/Developer/TraceInsights/Private/Insights/Widgets/STimingView.h b/Engine/Source/Developer/TraceInsights/Private/Insights/Widgets/STimingView.h
--- a/Engine/Source/Developer/TraceInsights/Private/Insights/Widgets/STimingView.h
+++ b/Engine/Source/Developer/TraceInsights/Private/Insights/Widgets/STimingView.h
@@ -8,6 +8,7 @@
 #include "Input/Reply.h"
 #include "Layout/Geometry.h"
 #include "Styling/SlateTypes.h"
+#include "Templates/Function.h"
 #include "TraceServices/AnalysisService.h"
 #include "Widgets/DeclarativeSyntaxSupport.h"
 #include "Widgets/SCompoundWidget.h"
@@ -280,7 +281,8 @@ public:
 
 	////////////////////////////////////////////////////////////////////////////////////////////////////
 
-	static const TCHAR* GetLocationName(ETimingTrackLocation Location);
+	// The callback should return 'true' to continue the enumeration.
+	void EnumerateAllTracks(TFunctionRef<bool(TSharedPtr<FBaseTimingTrack>&)> Callback);
 
 	const TArray<TSharedPtr<FBaseTimingTrack>>& GetTrackList(ETimingTrackLocation TrackLocation) const
 	{
@@ -295,6 +297,11 @@ public:
 		}
 	}
 
+	static const TCHAR* GetLocationName(ETimingTrackLocation Location);
+
+	void ChangeTrackLocation(TSharedRef<FBaseTimingTrack> Track, ETimingTrackLocation NewLocation);
+	bool CanChangeTrackLocation(TSharedRef<FBaseTimingTrack> Track, ETimingTrackLocation NewLocation) const;
+
 	void UpdateScrollableTracksOrder();
 	int32 GetFirstScrollableTrackOrder() const;
 	int32 GetLastScrollableTrackOrder() const;
@@ -322,6 +329,8 @@ public:
 	bool IsTimeSelectedInclusive(double Time) const { return Time >= SelectionStartTime && Time <= SelectionEndTime; }
 
 	void ScrollAtPosY(float ScrollPosY);
+	void BringIntoViewY(float InTopY, float InBottomY);
+	void BringScrollableTrackIntoView(const FBaseTimingTrack& Track);
 	void ScrollAtTime(double StartTime);
 	void CenterOnTimeInterval(double IntervalStartTime, double IntervalDuration);
 	void ZoomOnTimeInterval(double IntervalStartTime, double IntervalDuration);
@@ -342,6 +351,9 @@ public:
 	const TSharedPtr<FBaseTimingTrack> GetSelectedTrack() const { return SelectedTrack; }
 	const TSharedPtr<const ITimingEvent> GetSelectedEvent() const { return SelectedEvent; }
 
+	void SelectTimingTrack(const TSharedPtr<FBaseTimingTrack> InTrack, bool bBringTrackIntoView);
+	void SelectTimingEvent(const TSharedPtr<const ITimingEvent> InEvent, bool bBringEventIntoView);
+
 	const TSharedPtr<ITimingEventFilter> GetEventFilter() const { return TimingEventFilter; }
 	void SetEventFilter(const TSharedPtr<ITimingEventFilter> InEventFilter);
 
@@ -369,9 +381,6 @@ protected:
 	void ShowContextMenu(const FPointerEvent& MouseEvent);
 	void CreateTrackLocationMenu(FMenuBuilder& MenuBuilder, TSharedRef<FBaseTimingTrack> Track);
 
-	void ChangeTrackLocation(TSharedRef<FBaseTimingTrack> Track, ETimingTrackLocation NewLocation);
-	bool CanChangeTrackLocation(TSharedRef<FBaseTimingTrack> Track, ETimingTrackLocation NewLocation) const;
-
 	/** Binds our UI commands to delegates. */
 	void BindCommands();
 
