[Insights] Fixed compile error when -DisableUnity -NoPCH (added missing includes).

diff --git a/Engine/Source/Developer/TraceInsights/Private/Insights/ContextSwitches/ViewModels/ContextSwitchesSharedState.h b/Engine/Source/Developer/TraceInsights/Private/Insights/ContextSwitches/ViewModels/ContextSwitchesSharedState.h
--- a/Engine/Source/Developer/TraceInsights/Private/Insights/ContextSwitches/ViewModels/ContextSwitchesSharedState.h
+++ b/Engine/Source/Developer/TraceInsights/Private/Insights/ContextSwitches/ViewModels/ContextSwitchesSharedState.h
@@ -3,6 +3,7 @@
 #pragma once
 
 #include "CoreMinimal.h"
+#include "Framework/Commands/Commands.h"
 
 // Insights
 #include "Insights/ITimingViewExtender.h"
diff --git a/Engine/Source/Developer/TraceInsights/Private/Insights/ContextSwitches/ViewModels/ContextSwitchesTimingTrack.cpp b/Engine/Source/Developer/TraceInsights/Private/Insights/ContextSwitches/ViewModels/ContextSwitchesTimingTrack.cpp
--- a/Engine/Source/Developer/TraceInsights/Private/Insights/ContextSwitches/ViewModels/ContextSwitchesTimingTrack.cpp
+++ b/Engine/Source/Developer/TraceInsights/Private/Insights/ContextSwitches/ViewModels/ContextSwitchesTimingTrack.cpp
@@ -2,6 +2,7 @@
 
 #include "ContextSwitchesTimingTrack.h"
 
+#include "Framework/MultiBox/MultiBoxBuilder.h"
 #include "TraceServices/Model/ContextSwitches.h"
 
 // Insights
diff --git a/Engine/Source/Developer/TraceInsights/Private/Insights/ContextSwitches/ViewModels/CpuCoreTimingTrack.cpp b/Engine/Source/Developer/TraceInsights/Private/Insights/ContextSwitches/ViewModels/CpuCoreTimingTrack.cpp
--- a/Engine/Source/Developer/TraceInsights/Private/Insights/ContextSwitches/ViewModels/CpuCoreTimingTrack.cpp
+++ b/Engine/Source/Developer/TraceInsights/Private/Insights/ContextSwitches/ViewModels/CpuCoreTimingTrack.cpp
@@ -2,6 +2,7 @@
 
 #include "CpuCoreTimingTrack.h"
 
+#include "Framework/MultiBox/MultiBoxBuilder.h"
 #include "TraceServices/Model/ContextSwitches.h"
 #include "TraceServices/Model/Threads.h"
 
