[Insights]
  - Changed UX of the Session Browser page (work in progress).
  - Added palceholder icons (not final) for toolbar of Timing Inisghts and for few panels.
  - Updated toolbar of Log view. Also added rounded background for Verbosity Thresholds for better color contrast.

diff --git a/Engine/Source/Developer/TraceInsights/Private/Insights/InsightsManager.cpp b/Engine/Source/Developer/TraceInsights/Private/Insights/InsightsManager.cpp
--- a/Engine/Source/Developer/TraceInsights/Private/Insights/InsightsManager.cpp
+++ b/Engine/Source/Developer/TraceInsights/Private/Insights/InsightsManager.cpp
@@ -220,8 +220,8 @@ void FInsightsManager::RegisterMajorTabs(IUnrealInsightsModule& InsightsModule)
 		// Register tab spawner for the Start Page.
 		FTabSpawnerEntry& TabSpawnerEntry = FGlobalTabmanager::Get()->RegisterNomadTabSpawner(FInsightsManagerTabs::StartPageTabId,
 			FOnSpawnTab::CreateRaw(this, &FInsightsManager::SpawnStartPageTab))
-			.SetDisplayName(StartPageConfig.TabLabel.IsSet() ? StartPageConfig.TabLabel.GetValue() : LOCTEXT("StartPageTabTitle", "Unreal Insights"))
-			.SetTooltipText(StartPageConfig.TabTooltip.IsSet() ? StartPageConfig.TabTooltip.GetValue() : LOCTEXT("StartPageTooltipText", "Open the start page for Unreal Insights."))
+			.SetDisplayName(StartPageConfig.TabLabel.IsSet() ? StartPageConfig.TabLabel.GetValue() : LOCTEXT("StartPageTabTitle", "Trace Store"))
+			.SetTooltipText(StartPageConfig.TabTooltip.IsSet() ? StartPageConfig.TabTooltip.GetValue() : LOCTEXT("StartPageTooltipText", "Open the Trace Store Browser."))
 			.SetIcon(StartPageConfig.TabIcon.IsSet() ? StartPageConfig.TabIcon.GetValue() : FSlateIcon(FInsightsStyle::GetStyleSetName(), "StartPage.Icon.Small"));
 
 		TSharedRef<FWorkspaceItem> Group = StartPageConfig.WorkspaceGroup.IsValid() ? StartPageConfig.WorkspaceGroup.ToSharedRef() : WorkspaceMenu::GetMenuStructure().GetDeveloperToolsProfilingCategory();
diff --git a/Engine/Source/Developer/TraceInsights/Private/Insights/InsightsStyle.cpp b/Engine/Source/Developer/TraceInsights/Private/Insights/InsightsStyle.cpp
--- a/Engine/Source/Developer/TraceInsights/Private/Insights/InsightsStyle.cpp
+++ b/Engine/Source/Developer/TraceInsights/Private/Insights/InsightsStyle.cpp
@@ -11,6 +11,10 @@
 #define BOX_BRUSH_SVG(RelativePath, ...)    FSlateVectorBoxBrush   (Style.RootToContentDir(RelativePath, TEXT(".svg")), __VA_ARGS__)
 #define BORDER_BRUSH_SVG(RelativePath, ...) FSlateVectorBorderBrush(Style.RootToContentDir(RelativePath, TEXT(".svg")), __VA_ARGS__)
 
+#define CORE_IMAGE_BRUSH(RelativePath, ...)  FSlateVectorImageBrush (Style.RootToCoreContentDir(RelativePath, TEXT(".png")), __VA_ARGS__)
+#define CORE_BOX_BRUSH(RelativePath, ...)    FSlateVectorBoxBrush   (Style.RootToCoreContentDir(RelativePath, TEXT(".png")), __VA_ARGS__)
+#define CORE_BORDER_BRUSH(RelativePath, ...) FSlateVectorBorderBrush(Style.RootToCoreContentDir(RelativePath, TEXT(".png")), __VA_ARGS__)
+
 #define CORE_IMAGE_BRUSH_SVG(RelativePath, ...)  FSlateVectorImageBrush (Style.RootToCoreContentDir(RelativePath, TEXT(".svg")), __VA_ARGS__)
 #define CORE_BOX_BRUSH_SVG(RelativePath, ...)    FSlateVectorBoxBrush   (Style.RootToCoreContentDir(RelativePath, TEXT(".svg")), __VA_ARGS__)
 #define CORE_BORDER_BRUSH_SVG(RelativePath, ...) FSlateVectorBorderBrush(Style.RootToCoreContentDir(RelativePath, TEXT(".svg")), __VA_ARGS__)
@@ -45,6 +49,7 @@ TSharedRef<FSlateStyleSet> FInsightsStyle::Create()
 	const FVector2D Icon20x20(20.0f, 20.0f);
 	const FVector2D Icon24x24(24.0f, 24.0f);
 	const FVector2D Icon32x32(32.0f, 32.0f);
+	const FVector2D Icon40x40(40.0f, 40.0f);
 
 	TSharedRef<FSlateStyleSet> StyleRef = MakeShared<FSlateStyleSet>(FInsightsStyle::GetStyleSetName());
 	StyleRef->SetContentRoot(FPaths::EngineContentDir() / TEXT("Editor/Slate"));
@@ -64,11 +69,13 @@ TSharedRef<FSlateStyleSet> FInsightsStyle::Create()
 	Style.Set("HoveredEventBorder", new FSlateBorderBrush(NAME_None, FMargin(2.0f)));
 	Style.Set("SelectedEventBorder", new FSlateBorderBrush(NAME_None, FMargin(2.0f)));
 
+	Style.Set("RoundedBackground", new FSlateRoundedBoxBrush(FLinearColor(1.0f, 1.0f, 1.0f, 1.0f), Icon24x24));
+
 	//////////////////////////////////////////////////
 	// Icons for major components
 
-	Style.Set("SessionInfo.Icon.Large", new IMAGE_BRUSH("Icons/icon_tab_Tools_16x", Icon32x32));
-	Style.Set("SessionInfo.Icon.Small", new IMAGE_BRUSH("Icons/icon_tab_Tools_16x", Icon16x16));
+	Style.Set("SessionInfo.Icon.Large", new CORE_IMAGE_BRUSH_SVG("Starship/Common/Info", Icon32x32));
+	Style.Set("SessionInfo.Icon.Small", new CORE_IMAGE_BRUSH_SVG("Starship/Common/Info", Icon16x16));
 
 	Style.Set("Toolbar.Icon.Large", new IMAGE_BRUSH("Icons/icon_tab_Tools_16x", Icon32x32));
 	Style.Set("Toolbar.Icon.Small", new IMAGE_BRUSH("Icons/icon_tab_Tools_16x", Icon16x16));
@@ -76,8 +83,8 @@ TSharedRef<FSlateStyleSet> FInsightsStyle::Create()
 	//////////////////////////////////////////////////
 	// Start Page buttons
 
-	Style.Set("StartPage.Icon.Large", new IMAGE_BRUSH("Icons/icon_tab_Tools_16x", Icon32x32));
-	Style.Set("StartPage.Icon.Small", new IMAGE_BRUSH("Icons/icon_tab_Tools_16x", Icon16x16));
+	Style.Set("StartPage.Icon.Large", new IMAGE_BRUSH_SVG("../../Editor/Slate/Starship/Common/Home", Icon32x32));
+	Style.Set("StartPage.Icon.Small", new IMAGE_BRUSH_SVG("../../Editor/Slate/Starship/Common/Home", Icon16x16));
 
 	Style.Set("Open.Icon.Large", new IMAGE_BRUSH("Icons/Profiler/Profiler_LoadMultiple_Profiler_40x", Icon32x32));
 	Style.Set("Open.Icon.Small", new IMAGE_BRUSH("Icons/Profiler/Profiler_Load_Profiler_40x", Icon16x16));
@@ -85,33 +92,29 @@ TSharedRef<FSlateStyleSet> FInsightsStyle::Create()
 	Style.Set("OpenFile.Icon.Large", new IMAGE_BRUSH("Icons/LV_Load", Icon32x32));
 	Style.Set("OpenFile.Icon.Small", new IMAGE_BRUSH("Icons/LV_Load", Icon16x16));
 	
-
-	Style.Set("Menu.Icon.Large", new CORE_IMAGE_BRUSH_SVG("Starship/Common/menu", Icon32x32));
-	Style.Set("Menu.Icon.Small", new CORE_IMAGE_BRUSH_SVG("Starship/Common/menu", Icon16x16));
-
 	//////////////////////////////////////////////////
 	// Timing Insights
 
-	Style.Set("TimingProfiler.Icon.Large", new IMAGE_BRUSH("Icons/Profiler/profiler_stats_40x", Icon32x32));
-	Style.Set("TimingProfiler.Icon.Small", new IMAGE_BRUSH("Icons/Profiler/profiler_stats_40x", Icon16x16));
+	Style.Set("TimingProfiler.Icon.Large", new IMAGE_BRUSH_SVG("../../Editor/Slate/Starship/Common/Recent", Icon32x32));
+	Style.Set("TimingProfiler.Icon.Small", new IMAGE_BRUSH_SVG("../../Editor/Slate/Starship/Common/Recent", Icon16x16));
 
-	Style.Set("FramesTrack.Icon.Large", new IMAGE_BRUSH("Icons/Profiler/profiler_stats_40x", Icon32x32));
-	Style.Set("FramesTrack.Icon.Small", new IMAGE_BRUSH("Icons/Profiler/profiler_stats_40x", Icon16x16));
+	Style.Set("FramesTrack.Icon.Large", new IMAGE_BRUSH_SVG("../../Editor/Slate/Starship/Common/Statistics", Icon32x32));
+	Style.Set("FramesTrack.Icon.Small", new IMAGE_BRUSH_SVG("../../Editor/Slate/Starship/Common/Statistics", Icon16x16));
 
 	Style.Set("GraphTrack.Icon.Large", new IMAGE_BRUSH("Icons/Profiler/profiler_stats_40x", Icon32x32));
 	Style.Set("GraphTrack.Icon.Small", new IMAGE_BRUSH("Icons/Profiler/profiler_stats_40x", Icon16x16));
 
-	Style.Set("TimingView.Icon.Large", new IMAGE_BRUSH("Icons/Profiler/profiler_stats_40x", Icon32x32));
-	Style.Set("TimingView.Icon.Small", new IMAGE_BRUSH("Icons/Profiler/profiler_stats_40x", Icon16x16));
+	Style.Set("TimingView.Icon.Large", new IMAGE_BRUSH_SVG("../../Editor/Slate/Starship/Common/Timeline", Icon32x32));
+	Style.Set("TimingView.Icon.Small", new IMAGE_BRUSH_SVG("../../Editor/Slate/Starship/Common/Timeline", Icon16x16));
 
-	Style.Set("TimersView.Icon.Large", new IMAGE_BRUSH("Icons/Profiler/Profiler_Data_Capture_40x", Icon32x32));
-	Style.Set("TimersView.Icon.Small", new IMAGE_BRUSH("Icons/Profiler/Profiler_Data_Capture_40x", Icon16x16));
+	Style.Set("TimersView.Icon.Large", new IMAGE_BRUSH_SVG("../../Editor/Slate/Starship/Common/Realtime", Icon32x32));
+	Style.Set("TimersView.Icon.Small", new IMAGE_BRUSH_SVG("../../Editor/Slate/Starship/Common/Realtime", Icon16x16));
 
-	Style.Set("StatsCountersView.Icon.Large", new IMAGE_BRUSH("Icons/Profiler/Profiler_Data_Capture_40x", Icon32x32));
-	Style.Set("StatsCountersView.Icon.Small", new IMAGE_BRUSH("Icons/Profiler/Profiler_Data_Capture_40x", Icon16x16));
+	Style.Set("StatsCountersView.Icon.Large", new IMAGE_BRUSH_SVG("../../Editor/Slate/Starship/Common/Profile", Icon32x32));
+	Style.Set("StatsCountersView.Icon.Small", new IMAGE_BRUSH_SVG("../../Editor/Slate/Starship/Common/Profile", Icon16x16));
 
-	Style.Set("LogView.Icon.Large", new IMAGE_BRUSH("Icons/Profiler/profiler_CopyToClipboard_32x", Icon32x32));
-	Style.Set("LogView.Icon.Small", new IMAGE_BRUSH("Icons/Profiler/profiler_CopyToClipboard_32x", Icon16x16));
+	Style.Set("LogView.Icon.Large", new IMAGE_BRUSH_SVG("../../Editor/Slate/Starship/Common/Log", Icon32x32));
+	Style.Set("LogView.Icon.Small", new IMAGE_BRUSH_SVG("../../Editor/Slate/Starship/Common/Log", Icon16x16));
 
 	Style.Set("TableTreeView.Icon.Large", new IMAGE_BRUSH("Icons/Profiler/Profiler_Data_Capture_40x", Icon32x32));
 	Style.Set("TableTreeView.Icon.Small", new IMAGE_BRUSH("Icons/Profiler/Profiler_Data_Capture_40x", Icon16x16));
@@ -158,10 +161,25 @@ TSharedRef<FSlateStyleSet> FInsightsStyle::Create()
 	//////////////////////////////////////////////////
 
 	Style.Set("FindFirst", new IMAGE_BRUSH("Animation/backward_end", Icon20x20));
-	Style.Set("FindPrevious", new IMAGE_BRUSH("Animation/backward", Icon20x20));
-	Style.Set("FindNext", new IMAGE_BRUSH("Animation/forward", Icon20x20));
+	//Style.Set("FindPrevious", new IMAGE_BRUSH("Animation/backward", Icon20x20));
+	Style.Set("FindPrevious", new IMAGE_BRUSH("../../Editor/Slate/Icons/GeneralTools/Previous_40x", Icon20x20));
+	//Style.Set("FindNext", new IMAGE_BRUSH("Animation/forward", Icon20x20));
+	Style.Set("FindNext", new IMAGE_BRUSH("../../Editor/Slate/Icons/GeneralTools/Next_40x", Icon20x20));
 	Style.Set("FindLast", new IMAGE_BRUSH("Animation/forward_end", Icon20x20));
 
+	//////////////////////////////////////////////////
+
+	Style.Set("Menu.Icon.Large", new CORE_IMAGE_BRUSH_SVG("Starship/Common/menu", Icon32x32));
+	Style.Set("Menu.Icon.Small", new CORE_IMAGE_BRUSH_SVG("Starship/Common/menu", Icon16x16));
+
+	Style.Set("Icon.Dummy", new CORE_IMAGE_BRUSH_SVG("Starship/Common/dummy", Icon24x24));
+
+	Style.Set("Icon.Filter", new CORE_IMAGE_BRUSH_SVG("Starship/Common/filter", Icon24x24));
+	Style.Set("Icon.Find", new IMAGE_BRUSH_SVG("../../Editor/Slate/Starship/Common/TraceDataFiltering", Icon24x24));
+
+	Style.Set("FolderExplore.Icon.Large", new IMAGE_BRUSH_SVG("../../Editor/Slate/Starship/Common/ContentBrowser", Icon32x32));
+	Style.Set("FolderExplore.Icon.Small", new IMAGE_BRUSH_SVG("../../Editor/Slate/Starship/Common/ContentBrowser", Icon16x16));
+
 	Style.Set("Icon.Bug", new IMAGE_BRUSH_SVG("../../Editor/Slate/Starship/Common/bug", Icon24x24));
 
 	return StyleRef;
diff --git a/Engine/Source/Developer/TraceInsights/Private/Insights/ViewModels/MarkersTimingTrack.cpp b/Engine/Source/Developer/TraceInsights/Private/Insights/ViewModels/MarkersTimingTrack.cpp
--- a/Engine/Source/Developer/TraceInsights/Private/Insights/ViewModels/MarkersTimingTrack.cpp
+++ b/Engine/Source/Developer/TraceInsights/Private/Insights/ViewModels/MarkersTimingTrack.cpp
@@ -466,7 +466,7 @@ FLinearColor FTimeMarkerTrackBuilder::GetColorByVerbosity(const ELogVerbosity::T
 	{
 		FLinearColor(0.0f, 0.0f, 0.0f, 1.0f), // NoLogging
 		FLinearColor(1.0f, 0.0f, 0.0f, 1.0f), // Fatal
-		FLinearColor(1.0f, 0.3f, 0.0f, 1.0f), // Error
+		FLinearColor(1.0f, 0.1f, 0.1f, 1.0f), // Error
 		FLinearColor(0.7f, 0.5f, 0.0f, 1.0f), // Warning
 		FLinearColor(0.0f, 0.7f, 0.0f, 1.0f), // Display
 		FLinearColor(0.0f, 0.7f, 1.0f, 1.0f), // Log
diff --git a/Engine/Source/Developer/TraceInsights/Private/Insights/Widgets/SLogView.cpp b/Engine/Source/Developer/TraceInsights/Private/Insights/Widgets/SLogView.cpp
--- a/Engine/Source/Developer/TraceInsights/Private/Insights/Widgets/SLogView.cpp
+++ b/Engine/Source/Developer/TraceInsights/Private/Insights/Widgets/SLogView.cpp
@@ -15,6 +15,7 @@
 #include "HAL/PlatformFileManager.h"
 #include "Logging/MessageLog.h"
 #include "SlateOptMacros.h"
+#include "Styling/StyleColors.h"
 #include "Widgets/Images/SImage.h"
 #include "Widgets/Input/SComboButton.h"
 #include "Widgets/Layout/SScrollBox.h"
@@ -525,134 +526,65 @@ void SLogView::Construct(const FArguments& InArgs)
 	SAssignNew(ExternalScrollbar, SScrollBar)
 	.AlwaysShowScrollbar(true);
 
-	ChildSlot
-	[
-		SNew(SVerticalBox)
+	FSlimHorizontalToolBarBuilder ToolbarBuilder(TSharedPtr<const FUICommandList>(), FMultiBoxCustomization::None);
 
-		// Toolbar
-		+ SVerticalBox::Slot()
-		.AutoHeight()
-		[
-			SNew(SHorizontalBox)
-
-			// Icon
-			//+ SHorizontalBox::Slot()
-			//.VAlign(VAlign_Center)
-			//.AutoWidth()
-			//[
-			//	SNew(SImage)
-			//	.Image(FEditorStyle::GetBrush(TEXT("Log.TabIcon")))
-			//]
-			//
-			//// Label
-			//+ SHorizontalBox::Slot()
-			//.VAlign(VAlign_Center)
-			//.Padding(2.0f, 0.0f)
-			//.AutoWidth()
-			//[
-			//	SNew(STextBlock)
-			//	.Text(LOCTEXT("LogViewLabel", "Log View"))
-			//]
-
-			// Verbosity Threshold Filter
-			+SHorizontalBox::Slot()
-			//.Padding(2.0f, 0.0f, 0.0f, 0.0f)
-			.VAlign(VAlign_Center)
-			.AutoWidth()
-			[
-				SNew(SComboButton)
-				.ComboButtonStyle(FEditorStyle::Get(), "GenericFilters.ComboButtonStyle")
-				.ForegroundColor(FLinearColor::White)
-				.ContentPadding(0)
-				.ToolTipText(LOCTEXT("VerbosityThresholdFilterToolTip", "Filter log messages by verbosity threshold."))
-				.OnGetMenuContent(this, &SLogView::MakeVerbosityThresholdMenu)
-				.HasDownArrow(true)
-				.ContentPadding(FMargin(1.0f, 1.0f, 1.0f, 0.0f))
-				.ButtonContent()
-				[
-					SNew(SHorizontalBox)
+	ToolbarBuilder.BeginSection("Filters");
+	{
+		// Verbosity Threshold
+		ToolbarBuilder.AddComboButton(
+			FUIAction(),
+			FOnGetContent::CreateSP(this, &SLogView::MakeVerbosityThresholdMenu),
+			LOCTEXT("VerbosityThresholdText", "Verbosity Threshold"),
+			LOCTEXT("VerbosityThresholdToolTip", "Filter log messages by verbosity threshold."),
+			FSlateIcon(FInsightsStyle::GetStyleSetName(), "Icon.Filter"),
+			false
+		);
 
-					+SHorizontalBox::Slot()
-					.Padding(0.0f)
-					.AutoWidth()
-					[
-						SNew(STextBlock)
-						.TextStyle(FEditorStyle::Get(), "GenericFilters.TextStyle")
-						.Font(FEditorStyle::Get().GetFontStyle("FontAwesome.9"))
-						.Text(FText::FromString(FString(TEXT("\xf0b0"))) /*fa-filter*/)
-					]
+		// Category Filter
+		ToolbarBuilder.AddComboButton(
+			FUIAction(),
+			FOnGetContent::CreateSP(this, &SLogView::MakeCategoryFilterMenu),
+			LOCTEXT("CategoryFilterText", "Category Filter"),
+			LOCTEXT("CategoryFilterToolTip", "Filter log messages by category."),
+			FSlateIcon(FInsightsStyle::GetStyleSetName(), "Icon.Filter"),
+			false
+		);
 
-					+SHorizontalBox::Slot()
-					.Padding(2.0f, 0.0f, 0.0f, 0.0f)
-					.AutoWidth()
-					[
-						SNew(STextBlock)
-						.TextStyle(FEditorStyle::Get(), "GenericFilters.TextStyle")
-						.Text(LOCTEXT("VerbosityThresholdFilter", "Verbosity Threshold"))
-					]
-				]
-			]
+		// Text Filter (Search Box)
+		ToolbarBuilder.AddWidget(
+			SAssignNew(FilterTextBox, SSearchBox)
+			.HintText(LOCTEXT("FilterTextBoxHint", "Search log messages"))
+			.ToolTipText(LOCTEXT("FilterTextBoxToolTip", "Type here to filter the list of log messages."))
+			.OnTextChanged(this, &SLogView::FilterTextBox_OnTextChanged)
+		);
 
-			// Category Filter
-			+SHorizontalBox::Slot()
-			.Padding(2.0f, 0.0f, 0.0f, 0.0f)
+		// Stats Text (number of logs)
+		ToolbarBuilder.AddWidget(
+			SNew(SBox)
 			.VAlign(VAlign_Center)
-			.AutoWidth()
+			.Padding(FMargin(4.0f, 0.0f, 0.0f, 0.0f))
 			[
-				SNew(SComboButton)
-				.ComboButtonStyle(FEditorStyle::Get(), "GenericFilters.ComboButtonStyle")
-				.ForegroundColor(FLinearColor::White)
-				.ContentPadding(0)
-				.ToolTipText(LOCTEXT("CategoryFilterToolTip", "Filter log messages by category."))
-				.OnGetMenuContent(this, &SLogView::MakeCategoryFilterMenu)
-				.HasDownArrow(true)
-				.ContentPadding(FMargin(1.0f, 1.0f, 1.0f, 0.0f))
-				.ButtonContent()
-				[
-					SNew(SHorizontalBox)
-
-					+SHorizontalBox::Slot()
-					.Padding(0.0f)
-					.AutoWidth()
-					[
-						SNew(STextBlock)
-						.TextStyle(FEditorStyle::Get(), "GenericFilters.TextStyle")
-						.Font(FEditorStyle::Get().GetFontStyle("FontAwesome.9"))
-						.Text(FText::FromString(FString(TEXT("\xf0b0"))) /*fa-filter*/)
-					]
-
-					+SHorizontalBox::Slot()
-					.Padding(2.0f, 0.0f, 0.0f, 0.0f)
-					.AutoWidth()
-					[
-						SNew(STextBlock)
-						.TextStyle(FEditorStyle::Get(), "GenericFilters.TextStyle")
-						.Text(LOCTEXT("CategoryFilter", "Category Filter"))
-					]
-				]
+				SNew(STextBlock)
+				.Text(this, &SLogView::GetStatsText)
+				.ColorAndOpacity(this, &SLogView::GetStatsTextColor)
 			]
+		);
+	}
+	ToolbarBuilder.EndSection();
 
-			// Text Filter (Search Box)
-			+ SHorizontalBox::Slot()
-			.VAlign(VAlign_Center)
-			.Padding(2.0f, 0.0f)
-			.AutoWidth()
-			[
-				SAssignNew(FilterTextBox, SSearchBox)
-				.HintText(LOCTEXT("FilterTextBoxHint", "Search log messages"))
-				.ToolTipText(LOCTEXT("FilterTextBoxToolTip", "Type here to filter the list of log messages."))
-				.OnTextChanged(this, &SLogView::FilterTextBox_OnTextChanged)
-			]
+	ChildSlot
+	[
+		SNew(SVerticalBox)
 
-			// Stats Text
-			+ SHorizontalBox::Slot()
-			.VAlign(VAlign_Center)
-			.Padding(2.0f, 0.0f)
-			.AutoWidth()
+		// Toolbar
+		+ SVerticalBox::Slot()
+		.AutoHeight()
+		[
+			SNew(SBorder)
+			.BorderImage(FCoreStyle::Get().GetBrush("ToolPanel.GroupBorder"))
+			.Padding(4.0f)
 			[
-				SNew(STextBlock)
-				.Text(this, &SLogView::GetStatsText)
-				.ColorAndOpacity(this, &SLogView::GetStatsTextColor)
+				ToolbarBuilder.MakeWidget()
 			]
 		]
 
@@ -1224,7 +1156,7 @@ TSharedRef<SWidget> SLogView::MakeVerbosityThresholdMenu()
 
 ////////////////////////////////////////////////////////////////////////////////////////////////////
 
-void  SLogView::CreateVerbosityThresholdMenuSection(FMenuBuilder& MenuBuilder)
+void SLogView::CreateVerbosityThresholdMenuSection(FMenuBuilder& MenuBuilder)
 {
 	struct FVerbosityThresholdInfo
 	{
@@ -1249,15 +1181,34 @@ void  SLogView::CreateVerbosityThresholdMenuSection(FMenuBuilder& MenuBuilder)
 		const FVerbosityThresholdInfo& Threshold = VerbosityThresholds[Index];
 
 		const TSharedRef<SWidget> TextBlock = SNew(SHorizontalBox)
+#if 0 // shadow
 			+ SHorizontalBox::Slot()
 			.AutoWidth()
 			[
 				SNew(STextBlock)
 				.Text(Threshold.Label)
-				.ShadowColorAndOpacity(FLinearColor(0.05f, 0.05f, 0.05f, 1.0f))
-				.ShadowOffset(FVector2D(1.0f, 1.0f))
 				.ColorAndOpacity(FSlateColor(FTimeMarkerTrackBuilder::GetColorByVerbosity(Threshold.Verbosity)))
+				.ShadowColorAndOpacity(FLinearColor(0.02f, 0.02f, 0.02f, 1.0f))
+				.ShadowOffset(FVector2D(1.0f, 1.0f))
+			]
+#else // rounded background
+			+ SHorizontalBox::Slot()
+			.AutoWidth()
+			.Padding(FMargin(0.0f, -2.0f, 0.0f, -2.0f))
+			[
+				SNew(SBorder)
+				.BorderImage(FInsightsStyle::Get().GetBrush("RoundedBackground"))
+				.BorderBackgroundColor(FLinearColor(0.03f, 0.03f, 0.03f, 1.0f))
+				.Padding(FMargin(12.0f, 2.0f, 12.0f, 2.0f))
+				[
+					SNew(STextBlock)
+					.Text(Threshold.Label)
+					.ColorAndOpacity(FSlateColor(FTimeMarkerTrackBuilder::GetColorByVerbosity(Threshold.Verbosity)))
+					.ShadowColorAndOpacity(FLinearColor(0.02f, 0.02f, 0.02f, 1.0f))
+					.ShadowOffset(FVector2D(1.0f, 1.0f))
+				]
 			]
+#endif
 			+ SHorizontalBox::Slot()
 			.Padding(2.0f, 0.0f)
 			.VAlign(VAlign_Center)
diff --git a/Engine/Source/Developer/TraceInsights/Private/Insights/Widgets/SStartPageWindow.cpp b/Engine/Source/Developer/TraceInsights/Private/Insights/Widgets/SStartPageWindow.cpp
--- a/Engine/Source/Developer/TraceInsights/Private/Insights/Widgets/SStartPageWindow.cpp
+++ b/Engine/Source/Developer/TraceInsights/Private/Insights/Widgets/SStartPageWindow.cpp
@@ -15,6 +15,7 @@
 #include "SlateOptMacros.h"
 #include "SocketSubsystem.h"
 #include "Styling/CoreStyle.h"
+#include "Styling/StyleColors.h"
 #include "Trace/Analysis.h"
 #include "Trace/Analyzer.h"
 #include "Trace/ControlClient.h"
@@ -398,18 +399,20 @@ public:
 		TSharedPtr<FTraceViewModel> TracePin = WeakTrace.Pin();
 		if (TracePin.IsValid())
 		{
-			TSharedRef<ITypedTableView<TSharedPtr<FTraceViewModel>>> OwnerWidget = OwnerTablePtr.Pin().ToSharedRef();
-			const TSharedPtr<FTraceViewModel>* MyItem = OwnerWidget->Private_ItemFromWidget(this);
-			const bool IsSelected = OwnerWidget->Private_IsItemSelected(*MyItem);
-
-			if (IsSelected)
-			{
-				return FSlateColor(FLinearColor(0.0f, 0.0f, 0.0f, 1.0f));
-			}
-			else if (TracePin->Size < 1024ULL * 1024ULL)
+			if (TracePin->Size < 1024ULL * 1024ULL)
 			{
 				// < 1 MiB
-				return FSlateColor(FLinearColor(0.5f, 0.5f, 0.5f, 1.0f));
+				TSharedRef<ITypedTableView<TSharedPtr<FTraceViewModel>>> OwnerWidget = OwnerTablePtr.Pin().ToSharedRef();
+				const TSharedPtr<FTraceViewModel>* MyItem = OwnerWidget->Private_ItemFromWidget(this);
+				const bool IsSelected = OwnerWidget->Private_IsItemSelected(*MyItem);
+				if (IsSelected)
+				{
+					return FSlateColor(FLinearColor(0.75f, 0.75f, 0.75f, 1.0f));
+				}
+				else
+				{
+					return FSlateColor(FLinearColor(0.5f, 0.5f, 0.5f, 1.0f));
+				}
 			}
 			else if (TracePin->Size < 1024ULL * 1024ULL * 1024ULL)
 			{
@@ -633,7 +636,7 @@ SStartPageWindow::~SStartPageWindow()
 #if WITH_EDITOR
 	if (DurationActive > 0.0f && FEngineAnalytics::IsAvailable())
 	{
-		FEngineAnalytics::GetProvider().RecordEvent(TEXT("Editor.Usage.Insights.StartPage"), FAnalyticsEventAttribute(TEXT("Duration"), DurationActive));
+		FEngineAnalytics::GetProvider().RecordEvent(TEXT("Insights.Usage.SessionBrowser"), FAnalyticsEventAttribute(TEXT("Duration"), DurationActive));
 	}
 #endif // WITH_EDITOR
 }
@@ -660,6 +663,22 @@ void SStartPageWindow::Construct(const FArguments& InArgs)
 				.ColorAndOpacity(FLinearColor(0.15f, 0.15f, 0.15f, 1.0f))
 			]
 
+			+ SOverlay::Slot()
+			.HAlign(HAlign_Fill)
+			.VAlign(VAlign_Fill)
+			.Padding(0.0f, 0.0f, 0.0f, 0.0f)
+			[
+				SNew(SBox)
+				[
+					SNew(SBorder)
+					.HAlign(HAlign_Fill)
+					.VAlign(VAlign_Fill)
+					.Padding(0.0f)
+					.BorderImage(FCoreStyle::Get().GetBrush("WhiteBrush"))
+					.BorderBackgroundColor(FSlateColor(EStyleColor::Panel))
+				]
+			]
+
 			// Overlay slot for the main window area
 			+ SOverlay::Slot()
 			.HAlign(HAlign_Fill)
@@ -667,37 +686,30 @@ void SStartPageWindow::Construct(const FArguments& InArgs)
 			[
 				SAssignNew(MainContentPanel, SVerticalBox)
 
+				+ SVerticalBox::Slot()
+				.HAlign(HAlign_Fill)
+				.AutoHeight()
+				.Padding(12.0f, 8.0f, 12.0f, 4.0f)
+				[
+					ConstructTraceStoreDirectoryPanel()
+				]
+
 				+ SVerticalBox::Slot()
 				.HAlign(HAlign_Fill)
 				.VAlign(VAlign_Fill)
 				.FillHeight(1.0f)
-				.Padding(3.0f, 3.0f)
+				.Padding(3.0f, 4.0f)
 				[
-					SNew(SBox)
-					[
-						SNew(SBorder)
-						.BorderImage(FCoreStyle::Get().GetBrush("PopupText.Background"))
-						.Padding(8.0f)
-						[
-							ConstructSessionsPanel()
-						]
-					]
+					ConstructSessionsPanel()
 				]
 
 				+ SVerticalBox::Slot()
 				.HAlign(HAlign_Fill)
+				.VAlign(VAlign_Fill)
 				.AutoHeight()
-				.Padding(3.0f, 3.0f)
+				.Padding(12.0f, 4.0f)
 				[
-					SNew(SBox)
-					[
-						SNew(SBorder)
-						.BorderImage(FCoreStyle::Get().GetBrush("PopupText.Background"))
-						.Padding(8.0f)
-						[
-							ConstructTraceStoreDirectoryPanel()
-						]
-					]
+					ConstructLoadPanel()
 				]
 
 				+ SVerticalBox::Slot()
@@ -773,25 +785,7 @@ void SStartPageWindow::Construct(const FArguments& InArgs)
 
 TSharedRef<SWidget> SStartPageWindow::ConstructSessionsPanel()
 {
-	TSharedRef<SWidget> Widget = SNew(SVerticalBox)
-
-	+ SVerticalBox::Slot()
-	.AutoHeight()
-	.HAlign(HAlign_Left)
-	.Padding(0.0f, 2.0f)
-	[
-		SNew(STextBlock)
-		.Text(LOCTEXT("SessionsPanelTitle", "Trace Sessions"))
-		.Font(FCoreStyle::GetDefaultFontStyle("Bold", 11))
-		.ColorAndOpacity(FLinearColor::Gray)
-	]
-
-	+ SVerticalBox::Slot()
-	.HAlign(HAlign_Fill)
-	.VAlign(VAlign_Fill)
-	.Padding(0.0f, 1.0f, 0.0f, 2.0f)
-	[
-		SAssignNew(TraceListView, SListView<TSharedPtr<FTraceViewModel>>)
+	TSharedRef<SWidget> Widget = SAssignNew(TraceListView, SListView<TSharedPtr<FTraceViewModel>>)
 		.IsFocusable(true)
 		.ItemHeight(20.0f)
 		.SelectionMode(ESelectionMode::Single)
@@ -836,28 +830,7 @@ TSharedRef<SWidget> SStartPageWindow::ConstructSessionsPanel()
 			.HAlignHeader(HAlign_Right)
 			.HAlignCell(HAlign_Right)
 			.DefaultLabel(LOCTEXT("StatusColumn", "Status"))
-		)
-	]
-
-	+ SVerticalBox::Slot()
-	.AutoHeight()
-	.HAlign(HAlign_Fill)
-	.Padding(0.0f, 4.0f, 0.0f, 6.0f)
-	[
-		SNew(SSeparator)
-		.Orientation(Orient_Horizontal)
-		.SeparatorImage(FInsightsStyle::Get().GetBrush("WhiteBrush"))
-		.ColorAndOpacity(FLinearColor::Black)
-		.Thickness(1.0f)
-	]
-
-	+ SVerticalBox::Slot()
-	.AutoHeight()
-	.HAlign(HAlign_Fill)
-	.Padding(0.0f, 2.0f)
-	[
-		ConstructLoadPanel()
-	];
+		);
 
 	return Widget;
 }
@@ -882,6 +855,7 @@ TSharedRef<SWidget> SStartPageWindow::ConstructLoadPanel()
 	.AutoWidth()
 	[
 		SNew(SButton)
+		.ButtonStyle(&FAppStyle::Get().GetWidgetStyle<FButtonStyle>("PrimaryButton"))
 		.IsEnabled(this, &SStartPageWindow::Open_IsEnabled)
 		.OnClicked(this, &SStartPageWindow::Open_OnClicked)
 		.ToolTipText(LOCTEXT("OpenButtonTooltip", "Start analysis for selected trace session."))
@@ -889,22 +863,23 @@ TSharedRef<SWidget> SStartPageWindow::ConstructLoadPanel()
 		.Content()
 		[
 			SNew(SHorizontalBox)
-
-			+ SHorizontalBox::Slot()
-				.AutoWidth()
-				[
-					SNew(SImage)
-					.Image(FInsightsStyle::GetBrush("Open.Icon.Small"))
-				]
-
+			//+ SHorizontalBox::Slot()
+			//.AutoWidth()
+			//.VAlign(VAlign_Center)
+			//[
+			//	SNew(SImage)
+			//	.VAlign(VAlign_Center)
+			//	.Image(FInsightsStyle::GetBrush("Open.Icon.Small"))
+			//]
 			+ SHorizontalBox::Slot()
-				.AutoWidth()
-				.VAlign(VAlign_Center)
-				[
-					SNew(STextBlock)
-					.Text(LOCTEXT("OpenButtonText", "Open"))
-					.TextStyle(FAppStyle::Get(), "ButtonText")
-				]
+			.AutoWidth()
+			.VAlign(VAlign_Center)
+			[
+				SNew(STextBlock)
+				.TextStyle(&FAppStyle::Get().GetWidgetStyle<FTextBlockStyle>("DialogButtonText"))
+				.Justification(ETextJustify::Center)
+				.Text(LOCTEXT("OpenButtonText", "Open Trace"))
+			]
 		]
 	]
 
@@ -930,46 +905,59 @@ TSharedRef<SWidget> SStartPageWindow::ConstructLoadPanel()
 
 TSharedRef<SWidget> SStartPageWindow::ConstructTraceStoreDirectoryPanel()
 {
-	TSharedRef<SWidget> Widget = SNew(SVerticalBox)
+	TSharedRef<SWidget> Widget =
 
-	+ SVerticalBox::Slot()
-	.AutoHeight()
-	.HAlign(HAlign_Left)
-	.Padding(0.0f, 0.0f)
-	[
-		SNew(STextBlock)
-		.Text(LOCTEXT("TraceStoreDirText", "Trace Store Directory:"))
-		.ColorAndOpacity(FLinearColor::Gray)
-	]
-
-	+ SVerticalBox::Slot()
-	.AutoHeight()
-	.HAlign(HAlign_Fill)
-	.Padding(0.0f, 0.0f)
-	[
 		SNew(SHorizontalBox)
 
 		+ SHorizontalBox::Slot()
+		.AutoWidth()
 		.Padding(0.0f, 0.0f)
-		.HAlign(HAlign_Fill)
 		.VAlign(VAlign_Center)
 		[
 			SNew(STextBlock)
+			.Text(LOCTEXT("TraceStoreDirText", "Trace Store Directory"))
+		]
+
+		+ SHorizontalBox::Slot()
+		.AutoWidth()
+		.Padding(4.0f, 0.0f, 0.0f, 0.0f)
+		.VAlign(VAlign_Center)
+		[
+			SNew(SEditableTextBox)
+			.IsReadOnly(true)
+			.BackgroundColor(FSlateColor(EStyleColor::Background))
 			.Text(this, &SStartPageWindow::GetTraceStoreDirectory)
 		]
 
 		+ SHorizontalBox::Slot()
 		.AutoWidth()
 		.Padding(4.0f, 0.0f, 0.0f, 0.0f)
-		.HAlign(HAlign_Right)
 		.VAlign(VAlign_Center)
 		[
 			SNew(SButton)
-			.Text(LOCTEXT("ExploreTraceStoreDirButton", "Explore"))
+			.ButtonStyle(&FAppStyle::Get().GetWidgetStyle<FButtonStyle>("SimpleButton"))
 			.ToolTipText(LOCTEXT("ExploreTraceStoreDirButtonToolTip", "Explore the Trace Store Directory"))
 			.OnClicked(this, &SStartPageWindow::ExploreTraceStoreDirectory_OnClicked)
-		]
-	];
+			[
+				//SNew(SHorizontalBox)
+				//+ SHorizontalBox::Slot()
+				//.AutoWidth()
+				//[
+					SNew(SImage)
+					.Image(FInsightsStyle::Get().GetBrush("FolderExplore.Icon.Small"))
+					.ColorAndOpacity(FSlateColor::UseForeground())
+				//]
+				//+ SHorizontalBox::Slot()
+				//.AutoWidth()
+				//.Padding(2.0f, 0.0f, 0.0f, 0.0f)
+				//[
+				//	SNew(STextBlock)
+				//	.TextStyle(&FAppStyle::Get().GetWidgetStyle<FTextBlockStyle>("ButtonText"))
+				//	.Text(LOCTEXT("ExploreTraceStoreDirButton", "Explore"))
+				//	.ColorAndOpacity(FSlateColor::UseForeground())
+				//]
+			]
+		];
 
 	return Widget;
 }
@@ -1140,11 +1128,11 @@ TSharedRef<SWidget> SStartPageWindow::ConstructConnectPanel()
 
 		+ SHorizontalBox::Slot()
 		.AutoWidth()
-		.Padding(0.0f, 0.0f, 2.0f, 0.0f)
+		.Padding(0.0f, 0.0f, 4.0f, 0.0f)
 		.VAlign(VAlign_Center)
 		[
 			SNew(STextBlock)
-			.Text(LOCTEXT("HostTitle", "Running instance IP:"))
+			.Text(LOCTEXT("HostTitle", "Running instance IP"))
 			.ColorAndOpacity(FLinearColor::Gray)
 		]
 
diff --git a/Engine/Source/Developer/TraceInsights/Private/Insights/Widgets/STimingProfilerToolbar.cpp b/Engine/Source/Developer/TraceInsights/Private/Insights/Widgets/STimingProfilerToolbar.cpp
--- a/Engine/Source/Developer/TraceInsights/Private/Insights/Widgets/STimingProfilerToolbar.cpp
+++ b/Engine/Source/Developer/TraceInsights/Private/Insights/Widgets/STimingProfilerToolbar.cpp
@@ -13,6 +13,7 @@
 // Insights
 #include "Insights/InsightsCommands.h"
 #include "Insights/InsightsManager.h"
+#include "Insights/InsightsStyle.h"
 #include "Insights/TimingProfilerCommands.h"
 #include "Insights/TimingProfilerManager.h"
 
@@ -42,13 +43,27 @@ void STimingProfilerToolbar::Construct(const FArguments& InArgs)
 		{
 			ToolbarBuilder.BeginSection("View");
 			{
-				ToolbarBuilder.AddToolBarButton(FTimingProfilerCommands::Get().ToggleFramesTrackVisibility);
-				ToolbarBuilder.AddToolBarButton(FTimingProfilerCommands::Get().ToggleTimingViewVisibility);
-				ToolbarBuilder.AddToolBarButton(FTimingProfilerCommands::Get().ToggleTimersViewVisibility);
-				ToolbarBuilder.AddToolBarButton(FTimingProfilerCommands::Get().ToggleCallersTreeViewVisibility);
-				ToolbarBuilder.AddToolBarButton(FTimingProfilerCommands::Get().ToggleCalleesTreeViewVisibility);
-				ToolbarBuilder.AddToolBarButton(FTimingProfilerCommands::Get().ToggleStatsCountersViewVisibility);
-				ToolbarBuilder.AddToolBarButton(FTimingProfilerCommands::Get().ToggleLogViewVisibility);
+				ToolbarBuilder.AddToolBarButton(FTimingProfilerCommands::Get().ToggleFramesTrackVisibility,
+					NAME_None, TAttribute<FText>(), TAttribute<FText>(),
+					FSlateIcon(FInsightsStyle::GetStyleSetName(), "FramesTrack.Icon.Large"));
+				ToolbarBuilder.AddToolBarButton(FTimingProfilerCommands::Get().ToggleTimingViewVisibility,
+					NAME_None, TAttribute<FText>(), TAttribute<FText>(),
+					FSlateIcon(FInsightsStyle::GetStyleSetName(), "TimingView.Icon.Large"));
+				ToolbarBuilder.AddToolBarButton(FTimingProfilerCommands::Get().ToggleTimersViewVisibility,
+					NAME_None, TAttribute<FText>(), TAttribute<FText>(),
+					FSlateIcon(FInsightsStyle::GetStyleSetName(), "TimersView.Icon.Large"));
+				ToolbarBuilder.AddToolBarButton(FTimingProfilerCommands::Get().ToggleCallersTreeViewVisibility,
+					NAME_None, TAttribute<FText>(), TAttribute<FText>(),
+					FSlateIcon(FInsightsStyle::GetStyleSetName(), "TimersView.Icon.Large"));
+				ToolbarBuilder.AddToolBarButton(FTimingProfilerCommands::Get().ToggleCalleesTreeViewVisibility,
+					NAME_None, TAttribute<FText>(), TAttribute<FText>(),
+					FSlateIcon(FInsightsStyle::GetStyleSetName(), "TimersView.Icon.Large"));
+				ToolbarBuilder.AddToolBarButton(FTimingProfilerCommands::Get().ToggleStatsCountersViewVisibility,
+					NAME_None, TAttribute<FText>(), TAttribute<FText>(),
+					FSlateIcon(FInsightsStyle::GetStyleSetName(), "StatsCountersView.Icon.Large"));
+				ToolbarBuilder.AddToolBarButton(FTimingProfilerCommands::Get().ToggleLogViewVisibility,
+					NAME_None, TAttribute<FText>(), TAttribute<FText>(),
+					FSlateIcon(FInsightsStyle::GetStyleSetName(), "LogView.Icon.Large"));
 			}
 			ToolbarBuilder.EndSection();
 		}
@@ -57,7 +72,9 @@ void STimingProfilerToolbar::Construct(const FArguments& InArgs)
 		{
 			ToolbarBuilder.BeginSection("Debug");
 			{
-				ToolbarBuilder.AddToolBarButton(FInsightsCommands::Get().ToggleDebugInfo);
+				ToolbarBuilder.AddToolBarButton(FInsightsCommands::Get().ToggleDebugInfo,
+					NAME_None, FText(), TAttribute<FText>(),
+					FSlateIcon(FInsightsStyle::GetStyleSetName(), "Icon.Bug"));
 			}
 			ToolbarBuilder.EndSection();
 		}
@@ -65,10 +82,10 @@ void STimingProfilerToolbar::Construct(const FArguments& InArgs)
 
 	TSharedPtr<FUICommandList> CommandList = FInsightsManager::Get()->GetCommandList();
 
-	FToolBarBuilder ToolbarBuilder(CommandList.ToSharedRef(), FMultiBoxCustomization::None);
+	FSlimHorizontalToolBarBuilder ToolbarBuilder(CommandList.ToSharedRef(), FMultiBoxCustomization::None);
 	Local::FillViewToolbar(ToolbarBuilder);
 
-	FToolBarBuilder RightSideToolbarBuilder(CommandList.ToSharedRef(), FMultiBoxCustomization::None);
+	FSlimHorizontalToolBarBuilder RightSideToolbarBuilder(CommandList.ToSharedRef(), FMultiBoxCustomization::None);
 	Local::FillRightSideToolbar(RightSideToolbarBuilder);
 
 	// Create the tool bar!
