[Insights]
  - Timing View: Updated the Auto Scroll options menu. Added option to define custom values for Viewport Offset and Delay. A compact menu version is also implemented (but not used; added only as a reference/test implementation).
  - Timing View: Added toolbar on top of Time Ruler track.
  - Timing View: Added PosY for the viewport allowing an extra offset on top of viewport (above the docked tracks area).
  - Timing View: Added placeholder menus: All Tracks, Cpu / Gpu Tracks and Other Tracks (all share the same Tracks menu for now; WIP).

diff --git a/Engine/Source/Developer/TraceInsights/Private/Insights/ContextSwitches/ViewModels/ContextSwitchesTimingTrack.cpp b/Engine/Source/Developer/TraceInsights/Private/Insights/ContextSwitches/ViewModels/ContextSwitchesTimingTrack.cpp
--- a/Engine/Source/Developer/TraceInsights/Private/Insights/ContextSwitches/ViewModels/ContextSwitchesTimingTrack.cpp
+++ b/Engine/Source/Developer/TraceInsights/Private/Insights/ContextSwitches/ViewModels/ContextSwitchesTimingTrack.cpp
@@ -119,12 +119,12 @@ void FContextSwitchesTimingTrack::PostDraw(const ITimingTrackDrawContext& Contex
 		{
 			case ETimingTrackLocation::Scrollable:
 			{
-				const float TopY = Viewport.GetTopOffset();
+				const float TopY = Viewport.GetPosY() + Viewport.GetTopOffset();
 				if (LineY1 < TopY)
 				{
 					LineY1 = TopY;
 				}
-				const float BottomY = Viewport.GetHeight() - Viewport.GetBottomOffset();
+				const float BottomY = Viewport.GetPosY() + Viewport.GetHeight() - Viewport.GetBottomOffset();
 				if (LineY2 > BottomY)
 				{
 					LineY2 = BottomY;
@@ -133,12 +133,12 @@ void FContextSwitchesTimingTrack::PostDraw(const ITimingTrackDrawContext& Contex
 			}
 			case ETimingTrackLocation::TopDocked:
 			{
-				const float TopY = 0.0f;
+				const float TopY = Viewport.GetPosY();
 				if (LineY1 < TopY)
 				{
 					LineY1 = TopY;
 				}
-				const float BottomY = Viewport.GetTopOffset();
+				const float BottomY = Viewport.GetPosY() + Viewport.GetTopOffset();
 				if (LineY2 > BottomY)
 				{
 					LineY2 = BottomY;
@@ -147,12 +147,12 @@ void FContextSwitchesTimingTrack::PostDraw(const ITimingTrackDrawContext& Contex
 			}
 			case ETimingTrackLocation::BottomDocked:
 			{
-				const float TopY = Viewport.GetHeight() - Viewport.GetBottomOffset();
+				const float TopY = Viewport.GetPosY() + Viewport.GetHeight() - Viewport.GetBottomOffset();
 				if (LineY1 < TopY)
 				{
 					LineY1 = TopY;
 				}
-				const float BottomY = Viewport.GetHeight();
+				const float BottomY = Viewport.GetPosY() + Viewport.GetHeight();
 				if (LineY2 > BottomY)
 				{
 					LineY2 = BottomY;
diff --git a/Engine/Source/Developer/TraceInsights/Private/Insights/InsightsStyle.cpp b/Engine/Source/Developer/TraceInsights/Private/Insights/InsightsStyle.cpp
--- a/Engine/Source/Developer/TraceInsights/Private/Insights/InsightsStyle.cpp
+++ b/Engine/Source/Developer/TraceInsights/Private/Insights/InsightsStyle.cpp
@@ -182,6 +182,11 @@ TSharedRef<FSlateStyleSet> FInsightsStyle::Create()
 
 	Style.Set("Icon.Bug", new IMAGE_BRUSH_SVG("../../Editor/Slate/Starship/Common/bug", Icon24x24));
 
+	Style.Set("Icon.AutoScroll", new IMAGE_BRUSH_SVG("../../Editor/Slate/Starship/MainToolbar/play", Icon24x24));
+	Style.Set("Icon.AllTracksMenu", new CORE_IMAGE_BRUSH_SVG("Starship/Common/menu", Icon24x24));
+	Style.Set("Icon.CpuGpuTracksMenu", new CORE_IMAGE_BRUSH_SVG("Starship/Common/menu", Icon24x24));
+	Style.Set("Icon.OtherTracksMenu", new CORE_IMAGE_BRUSH_SVG("Starship/Common/menu", Icon24x24));
+
 	return StyleRef;
 }
 
diff --git a/Engine/Source/Developer/TraceInsights/Private/Insights/TimingProfilerCommands.cpp b/Engine/Source/Developer/TraceInsights/Private/Insights/TimingProfilerCommands.cpp
--- a/Engine/Source/Developer/TraceInsights/Private/Insights/TimingProfilerCommands.cpp
+++ b/Engine/Source/Developer/TraceInsights/Private/Insights/TimingProfilerCommands.cpp
@@ -20,7 +20,7 @@
 // FTimingProfilerMenuBuilder
 ////////////////////////////////////////////////////////////////////////////////////////////////////
 
-void FTimingProfilerMenuBuilder::AddMenuEntry(FMenuBuilder& MenuBuilder, const TSharedPtr< FUICommandInfo >& UICommandInfo, const FUIAction& UIAction)
+void FTimingProfilerMenuBuilder::AddMenuEntry(FMenuBuilder& MenuBuilder, const TSharedPtr<FUICommandInfo>& UICommandInfo, const FUIAction& UIAction)
 {
 	MenuBuilder.AddMenuEntry
 	(
@@ -103,7 +103,7 @@ void FTimingViewCommands::RegisterCommands()
 		FInputChord(EKeys::G));
 
 	UI_COMMAND(QuickFind,
-		"Quick Find",
+		"Quick Find...",
 		"Quick find or filter events in the timing view.",
 		EUserInterfaceActionType::Button,
 		FInputChord(EModifierKey::Control, EKeys::F));
diff --git a/Engine/Source/Developer/TraceInsights/Private/Insights/ViewModels/DrawHelpers.cpp b/Engine/Source/Developer/TraceInsights/Private/Insights/ViewModels/DrawHelpers.cpp
--- a/Engine/Source/Developer/TraceInsights/Private/Insights/ViewModels/DrawHelpers.cpp
+++ b/Engine/Source/Developer/TraceInsights/Private/Insights/ViewModels/DrawHelpers.cpp
@@ -166,7 +166,7 @@ void FDrawHelpers::DrawTimeRangeSelection(
 
 		if (SelectionX1 <= MaxX && SelectionX2 >= MinX)
 		{
-			const float Y = 0.0f;
+			const float Y = Viewport.GetPosY();
 			const float H = Viewport.GetHeight();
 
 			const double Duration = EndTime - StartTime;
diff --git a/Engine/Source/Developer/TraceInsights/Private/Insights/ViewModels/MarkersTimingTrack.cpp b/Engine/Source/Developer/TraceInsights/Private/Insights/ViewModels/MarkersTimingTrack.cpp
--- a/Engine/Source/Developer/TraceInsights/Private/Insights/ViewModels/MarkersTimingTrack.cpp
+++ b/Engine/Source/Developer/TraceInsights/Private/Insights/ViewModels/MarkersTimingTrack.cpp
@@ -185,7 +185,7 @@ void FMarkersTimingTrack::PostDraw(const ITimingTrackDrawContext& Context) const
 	}
 	else
 	{
-		BoxY = 0.0f;
+		BoxY = Viewport.GetPosY();
 		BoxH = Viewport.GetHeight();
 	}
 
diff --git a/Engine/Source/Developer/TraceInsights/Private/Insights/ViewModels/ThreadTimingTrack.cpp b/Engine/Source/Developer/TraceInsights/Private/Insights/ViewModels/ThreadTimingTrack.cpp
--- a/Engine/Source/Developer/TraceInsights/Private/Insights/ViewModels/ThreadTimingTrack.cpp
+++ b/Engine/Source/Developer/TraceInsights/Private/Insights/ViewModels/ThreadTimingTrack.cpp
@@ -977,7 +977,7 @@ void FThreadTimingTrack::DrawSelectedEventInfo(const FThreadTrackEvent& Selected
 		const TSharedRef<FSlateFontMeasure> FontMeasureService = FSlateApplication::Get().GetRenderer()->GetFontMeasureService();
 		const FVector2D Size = FontMeasureService->Measure(Str, Font);
 		const float X = Viewport.GetWidth() - Size.X - 23.0f;
-		const float Y = Viewport.GetHeight() - Size.Y - 18.0f;
+		const float Y = Viewport.GetPosY() + Viewport.GetHeight() - Size.Y - 18.0f;
 
 		const FLinearColor BackgroundColor(0.05f, 0.05f, 0.05f, 1.0f);
 		const FLinearColor TextColor(0.7f, 0.7f, 0.7f, 1.0f);
diff --git a/Engine/Source/Developer/TraceInsights/Private/Insights/ViewModels/TimeRulerTrack.cpp b/Engine/Source/Developer/TraceInsights/Private/Insights/ViewModels/TimeRulerTrack.cpp
--- a/Engine/Source/Developer/TraceInsights/Private/Insights/ViewModels/TimeRulerTrack.cpp
+++ b/Engine/Source/Developer/TraceInsights/Private/Insights/ViewModels/TimeRulerTrack.cpp
@@ -293,7 +293,7 @@ void FTimeRulerTrack::PostDraw(const ITimingTrackDrawContext& Context) const
 			{
 				X = SelectionX1 - W / 2;
 			}
-			MousePosTextForegroundColor = FLinearColor(FColor(32, 64, 128, 255));
+			MousePosTextForegroundColor = FLinearColor(0.01f, 0.05f, 0.2f, 1.0f);
 		}
 		else
 		{
@@ -301,10 +301,10 @@ void FTimeRulerTrack::PostDraw(const ITimingTrackDrawContext& Context) const
 			//DrawContext.DrawBox(0.0f, MousePosition.Y, Viewport.Width, 1.0f, WhiteBrush, MousePosLineColor);
 
 			// Draw vertical line at mouse position.
-			DrawContext.DrawBox(MousePosition.X, 0.0f, 1.0f, Viewport.GetHeight(), WhiteBrush, MousePosLineColor);
+			DrawContext.DrawBox(MousePosition.X, Viewport.GetPosY(), 1.0f, Viewport.GetHeight(), WhiteBrush, MousePosLineColor);
 
 			// Stroke the vertical line above current time box.
-			DrawContext.DrawBox(MousePosition.X, 0.0f, 1.0f, TextY, WhiteBrush, MousePosTextBackgroundColor);
+			DrawContext.DrawBox(MousePosition.X, GetPosY(), 1.0f, TextY - GetPosY(), WhiteBrush, MousePosTextBackgroundColor);
 		}
 
 		// Fill the current time box.
@@ -342,13 +342,13 @@ void FTimeRulerTrack::DrawTimeMarker(const ITimingTrackDrawContext& Context, con
 		return;
 	}
 
-	constexpr float TimeMarkerY = 0.0f;
+	const float TimeMarkerY = GetPosY();
 	constexpr float BoxHeight = 12.0f;
 
 	FDrawContext& DrawContext = Context.GetDrawContext();
 
 	// Draw the vertical line.
-	DrawContext.DrawBox(TimeMarkerX, TimeMarkerY, 1.0f, Viewport.GetHeight() - TimeMarkerY, WhiteBrush, TimeMarker.GetColor());
+	DrawContext.DrawBox(TimeMarkerX, TimeMarkerY, 1.0f, Viewport.GetPosY() + Viewport.GetHeight() - TimeMarkerY, WhiteBrush, TimeMarker.GetColor());
 	DrawContext.LayerId++;
 
 	const FVector2D& MousePosition = Context.GetMousePosition();
diff --git a/Engine/Source/Developer/TraceInsights/Private/Insights/ViewModels/TimingViewDrawHelper.cpp b/Engine/Source/Developer/TraceInsights/Private/Insights/ViewModels/TimingViewDrawHelper.cpp
--- a/Engine/Source/Developer/TraceInsights/Private/Insights/ViewModels/TimingViewDrawHelper.cpp
+++ b/Engine/Source/Developer/TraceInsights/Private/Insights/ViewModels/TimingViewDrawHelper.cpp
@@ -296,7 +296,7 @@ FTimingViewDrawHelper::~FTimingViewDrawHelper()
 
 void FTimingViewDrawHelper::DrawBackground() const
 {
-	const float Y = 0.0f;
+	const float Y = Viewport.GetPosY();
 	const float H = FMath::CeilToFloat(Viewport.GetHeight());
 	FDrawHelpers::DrawBackground(DrawContext, BackgroundAreaBrush, Viewport, Y, H, ValidAreaX, ValidAreaW); // also computes ValidAreaX and ValidAreaW
 }
@@ -753,13 +753,13 @@ void FTimingViewDrawHelper::EndDrawTracks() const
 	{
 		const int32 TrackBackgroundLayerId = ReservedLayerId + ToInt32(EDrawLayer::TrackBackground);
 
-		const float TopY = Viewport.GetTopOffset();
-		const float BottomY = Viewport.GetHeight() - Viewport.GetBottomOffset();
+		const float TopY = Viewport.GetPosY() + Viewport.GetTopOffset();
+		const float BottomY = Viewport.GetPosY() + Viewport.GetHeight() - Viewport.GetBottomOffset();
 
 		if (TopY < BottomY)
 		{
 			// Y position of the first pixel below the last track.
-			const float Y = TopY + Viewport.GetScrollHeight() - Viewport.GetScrollPosY();
+			const float Y = TopY + Viewport.GetScrollHeight() - Viewport.GetScrollPosY() - 1.0f;
 
 			if (Y >= TopY && Y < BottomY)
 			{
diff --git a/Engine/Source/Developer/TraceInsights/Private/Insights/Widgets/STimingView.cpp b/Engine/Source/Developer/TraceInsights/Private/Insights/Widgets/STimingView.cpp
--- a/Engine/Source/Developer/TraceInsights/Private/Insights/Widgets/STimingView.cpp
+++ b/Engine/Source/Developer/TraceInsights/Private/Insights/Widgets/STimingView.cpp
@@ -20,11 +20,14 @@
 #include "SlateOptMacros.h"
 #include "Styling/CoreStyle.h"
 #include "Styling/SlateBrush.h"
+#include "Styling/StyleColors.h"
 #include "TraceServices/AnalysisService.h"
 #include "Widgets/Docking/SDockTab.h"
 #include "Widgets/Input/SButton.h"
 #include "Widgets/Input/SCheckBox.h"
 #include "Widgets/Input/SComboButton.h"
+#include "Widgets/Input/SEditableTextBox.h"
+#include "Widgets/Input/SSegmentedControl.h"
 #include "Widgets/Layout/SScrollBar.h"
 #include "Widgets/Layout/SSpacer.h"
 #include "Widgets/SOverlay.h"
@@ -139,6 +142,63 @@ STimingView::~STimingView()
 
 void STimingView::Construct(const FArguments& InArgs)
 {
+	FSlimHorizontalToolBarBuilder LeftToolbar(nullptr, FMultiBoxCustomization::None);
+
+	LeftToolbar.BeginSection("Menus");
+	LeftToolbar.AddComboButton(
+		FUIAction(),
+		FOnGetContent::CreateSP(this, &STimingView::MakeTracksFilterMenu),
+		LOCTEXT("AllTracksFilter", "All Tracks"),
+		LOCTEXT("AllTracksFilterToolTip", "Filter all tracks."),
+		FSlateIcon(FInsightsStyle::GetStyleSetName(), "Icon.AllTracksMenu"));
+	LeftToolbar.AddComboButton(
+		FUIAction(),
+		FOnGetContent::CreateSP(this, &STimingView::MakeTracksFilterMenu),
+		LOCTEXT("CpuGpuTracksFilter", "CPU/GPU Tracks"),
+		LOCTEXT("CpuGpuTracksFilterToolTip", "Filter CPU/GPU timing tracks."),
+		FSlateIcon(FInsightsStyle::GetStyleSetName(), "Icon.CpuGpuTracksMenu"));
+	LeftToolbar.AddComboButton(
+		FUIAction(),
+		FOnGetContent::CreateSP(this, &STimingView::MakeTracksFilterMenu),
+		LOCTEXT("OtherTracksFilter", "Other Tracks"),
+		LOCTEXT("OtherTracksFilterToolTip", "Filter other tracks."),
+		FSlateIcon(FInsightsStyle::GetStyleSetName(), "Icon.OtherTracksMenu"));
+	LeftToolbar.EndSection();
+
+	//////////////////////////////////////////////////
+
+	FSlimHorizontalToolBarBuilder RightToolbar(nullptr, FMultiBoxCustomization::None);
+
+	FUIAction AutoScrollToggleButtonAction;
+	AutoScrollToggleButtonAction.GetActionCheckState.BindLambda([this]
+	{
+		return bAutoScroll ? ECheckBoxState::Checked : ECheckBoxState::Unchecked;
+	});
+	AutoScrollToggleButtonAction.ExecuteAction.BindLambda([this]
+	{
+		bAutoScroll = !bAutoScroll;
+		Viewport.AddDirtyFlags(ETimingTrackViewportDirtyFlags::HInvalidated);
+	});
+
+	RightToolbar.BeginSection("Auto-Scroll");
+	RightToolbar.AddToolBarButton(
+		AutoScrollToggleButtonAction,
+		NAME_None,
+		FText::GetEmpty(),
+		LOCTEXT("AutoScrollToolTip", "Auto-Scroll"),
+		FSlateIcon(FInsightsStyle::GetStyleSetName(),"Icon.AutoScroll"),
+		EUserInterfaceActionType::ToggleButton);
+	RightToolbar.AddComboButton(
+		FUIAction(),
+		FOnGetContent::CreateSP(this, &STimingView::MakeAutoScrollOptionsMenu),
+		TAttribute<FText>(),
+		LOCTEXT("AutoScrollOptionsToolTip", "Auto-Scroll Options"),
+		TAttribute<FSlateIcon>(),
+		true);
+	RightToolbar.EndSection();
+
+	//////////////////////////////////////////////////
+
 	ChildSlot
 	[
 		SNew(SOverlay)
@@ -169,79 +229,17 @@ void STimingView::Construct(const FArguments& InArgs)
 		+ SOverlay::Slot()
 		.HAlign(HAlign_Left)
 		.VAlign(VAlign_Top)
-		.Padding(FMargin(0.0f, 0.0f, 0.0f, 0.0f))
+		.Padding(FMargin(2.0f, 0.0f, 0.0f, 0.0f))
 		[
-			// Tracks Filter
-			SNew(SComboButton)
-			.ComboButtonStyle(FCoreStyle::Get(), "ComboButton")
-			.ButtonColorAndOpacity(FLinearColor::Transparent)
-			.ToolTipText(LOCTEXT("TracksFilterToolTip", "Filter timing tracks."))
-			.OnGetMenuContent(this, &STimingView::MakeTracksFilterMenu)
-			.HasDownArrow(true)
-			.ContentPadding(FMargin(-4.0f, 0.0f, -5.0f, 0.0f))
-			.ButtonContent()
-			[
-				SNew(SHorizontalBox)
-
-				+ SHorizontalBox::Slot()
-				.Padding(FMargin(0.0f, 2.0f, 0.0f, 0.0f))
-				.AutoWidth()
-				[
-					SNew(STextBlock)
-					.Font(FEditorStyle::Get().GetFontStyle("FontAwesome.9"))
-					.Text(FText::FromString(FString(TEXT("\xf0b0"))) /*fa-filter*/)
-				]
-
-				+ SHorizontalBox::Slot()
-				.Padding(FMargin(3.0f, 0.0f, -7.0f, 0.0f))
-				.AutoWidth()
-				[
-					SNew(STextBlock)
-					.Text(LOCTEXT("TracksFilter", "Tracks"))
-				]
-			]
+			LeftToolbar.MakeWidget()
 		]
 
 		+ SOverlay::Slot()
 		.HAlign(HAlign_Right)
 		.VAlign(VAlign_Top)
-		.Padding(FMargin(0.0f, 0.0f, 0.0f, 0.0f))
+		.Padding(FMargin(0.0f, 0.0f, 6.0f, 0.0f))
 		[
-			SNew(SHorizontalBox)
-
-			+ SHorizontalBox::Slot()
-			.Padding(FMargin(0.0f, 1.0f, 0.0f, 1.0f))
-			.AutoWidth()
-			[
-				SNew(SCheckBox)
-				.Style(FCoreStyle::Get(), "ToggleButtonCheckbox")
-				.HAlign(HAlign_Center)
-				.Padding(FMargin(4.0f, 2.0f, 4.0f, 2.0f))
-				.OnCheckStateChanged(this, &STimingView::AutoScroll_OnCheckStateChanged)
-				.IsChecked(this, &STimingView::AutoScroll_IsChecked)
-				.ToolTipText(LOCTEXT("AutoScroll_Tooltip", "Auto Scroll"))
-				[
-					SNew(STextBlock)
-					.Text(LOCTEXT("AutoScroll_Button", ">>"))
-				]
-			]
-
-			+ SHorizontalBox::Slot()
-			.Padding(0.0f)
-			.AutoWidth()
-			[
-				SNew(SComboButton)
-				.ComboButtonStyle(FCoreStyle::Get(), "ComboButton")
-				.ButtonColorAndOpacity(FLinearColor::Transparent)
-				.OnGetMenuContent(this, &STimingView::MakeAutoScrollOptionsMenu)
-				.HasDownArrow(true)
-				.ContentPadding(FMargin(-14.0f, 0.0f, -6.0f, 0.0f))
-				.ButtonContent()
-				[
-					SNew(SBorder)
-					.Padding(0.0f)
-				]
-			]
+			RightToolbar.MakeWidget()
 		]
 
 		+ SOverlay::Slot()
@@ -515,7 +513,8 @@ void STimingView::Tick(const FGeometry& AllottedGeometry, const double InCurrent
 	////////////////////////////////////////////////////////////////////////////////////////////////////
 	// Update viewport.
 
-	Viewport.UpdateSize(ViewWidth, ViewHeight);
+	Viewport.SetPosY(40.0f); // height of toolbar
+	Viewport.UpdateSize(FMath::RoundToFloat(ViewWidth), FMath::RoundToFloat(ViewHeight) - Viewport.GetPosY());
 
 	if (!bIsPanning && !bAutoScroll)
 	{
@@ -718,7 +717,7 @@ void STimingView::Tick(const FGeometry& AllottedGeometry, const double InCurrent
 	float TopOffset = 0.0f;
 	for (TSharedPtr<FBaseTimingTrack>& TrackPtr : TopDockedTracks)
 	{
-		SetTrackPosY(TrackPtr, TopOffset);
+		SetTrackPosY(TrackPtr, Viewport.GetPosY() + TopOffset);
 		if (TrackPtr->IsVisible())
 		{
 			TopOffset += TrackPtr->GetHeight();
@@ -741,33 +740,33 @@ void STimingView::Tick(const FGeometry& AllottedGeometry, const double InCurrent
 	// Compute the total height of bottom docked areas.
 
 	float BottomOffset = 0.0f;
-	for (TSharedPtr<FBaseTimingTrack>& TrackPtr : BottomDockedTracks)
-	{
-		if (TrackPtr->IsVisible())
-		{
-			BottomOffset += TrackPtr->GetHeight();
-		}
-	}
-	float BottomOffsetY = Viewport.GetHeight() - BottomOffset;
 	int32 NumVisibleBottomDockedTracks = 0;
 	for (TSharedPtr<FBaseTimingTrack>& TrackPtr : BottomDockedTracks)
 	{
-		SetTrackPosY(TrackPtr, BottomOffsetY);
 		if (TrackPtr->IsVisible())
 		{
-			BottomOffsetY += TrackPtr->GetHeight();
+			BottomOffset += TrackPtr->GetHeight();
 			NumVisibleBottomDockedTracks++;
 		}
 	}
 	if (NumVisibleBottomDockedTracks > 0)
 	{
-		bDrawBottomSeparatorLine = true;
 		BottomOffset += 2.0f;
+		if (Viewport.GetTopOffset() + BottomOffset > Viewport.GetHeight())
+		{
+			BottomOffset = Viewport.GetHeight() - Viewport.GetTopOffset();
+		}
+		float BottomDockedTrackPosY = Viewport.GetPosY() + Viewport.GetHeight() - BottomOffset + 2.0f;
+		for (TSharedPtr<FBaseTimingTrack>& TrackPtr : BottomDockedTracks)
+		{
+			SetTrackPosY(TrackPtr, BottomDockedTrackPosY);
+			if (TrackPtr->IsVisible())
+			{
+				BottomDockedTrackPosY += TrackPtr->GetHeight();
+			}
+		}
 	}
-	else
-	{
-		bDrawBottomSeparatorLine = false;
-	}
+	bDrawBottomSeparatorLine = (BottomOffset > 0.0f);
 	Viewport.SetBottomOffset(BottomOffset);
 
 	////////////////////////////////////////////////////////////////////////////////////////////////////
@@ -795,7 +794,7 @@ void STimingView::Tick(const FGeometry& AllottedGeometry, const double InCurrent
 	}
 
 	// Set the VerticalScrollBar padding so it is limited to the scrollable area.
-	VerticalScrollBar->SetPadding(FMargin(0.0f, TopOffset + 2.0f, 0.0f, FMath::Max(BottomOffset + 2.0f, 10.0f)));
+	VerticalScrollBar->SetPadding(FMargin(0.0f, Viewport.GetPosY() + TopOffset + 2.0f, 0.0f, FMath::Max(BottomOffset + 2.0f, 10.0f)));
 
 	////////////////////////////////////////////////////////////////////////////////////////////////////
 
@@ -973,7 +972,7 @@ void STimingView::Tick(const FGeometry& AllottedGeometry, const double InCurrent
 	Tooltip.Update();
 	if (!MousePosition.IsZero())
 	{
-		Tooltip.SetPosition(MousePosition, 0.0f, Viewport.GetWidth() - 12.0f, 0.0f, Viewport.GetHeight() - 12.0f); // -12.0f is to avoid overlaping the scrollbars
+		Tooltip.SetPosition(MousePosition, 0.0f, Viewport.GetWidth() - 12.0f, Viewport.GetPosY(), Viewport.GetPosY() + Viewport.GetHeight() - 12.0f); // -12.0f is to avoid overlaping the scrollbars
 	}
 
 	////////////////////////////////////////////////////////////////////////////////////////////////////
@@ -1038,8 +1037,8 @@ void STimingView::Tick(const FGeometry& AllottedGeometry, const double InCurrent
 
 void STimingView::UpdatePositionForScrollableTracks()
 {
-	// Update the Y postion for visible scrollable tracks.
-	float ScrollableTrackPosY = Viewport.GetTopOffset() - Viewport.GetScrollPosY();
+	// Update the Y postion for the visible scrollable tracks.
+	float ScrollableTrackPosY = Viewport.GetPosY() + Viewport.GetTopOffset() - Viewport.GetScrollPosY();
 	for (TSharedPtr<FBaseTimingTrack>& TrackPtr : ScrollableTracks)
 	{
 		SetTrackPosY(TrackPtr, ScrollableTrackPosY);
@@ -1052,12 +1051,12 @@ void STimingView::UpdatePositionForScrollableTracks()
 
 ////////////////////////////////////////////////////////////////////////////////////////////////////
 
-void STimingView::SetTrackPosY(TSharedPtr<FBaseTimingTrack>& TrackPtr, float TopOffsetY) const
+void STimingView::SetTrackPosY(TSharedPtr<FBaseTimingTrack>& TrackPtr, float TrackPosY) const
 {
-	TrackPtr->SetPosY(TopOffsetY);
+	TrackPtr->SetPosY(TrackPosY);
 	if (TrackPtr->GetChildTrack().IsValid())
 	{
-		TrackPtr->GetChildTrack()->SetPosY(TopOffsetY + this->GetViewport().GetLayout().TimelineDY + 1.0f);
+		TrackPtr->GetChildTrack()->SetPosY(TrackPosY + this->GetViewport().GetLayout().TimelineDY + 1.0f);
 	}
 }
 
@@ -1191,13 +1190,13 @@ int32 STimingView::OnPaint(const FPaintArgs& Args, const FGeometry& AllottedGeom
 
 		// Draw the scrollable tracks.
 		{
-			const float TopY = Viewport.GetTopOffset();
-			const float BottomY = Viewport.GetHeight() - Viewport.GetBottomOffset();
+			const float TopY = Viewport.GetPosY() + Viewport.GetTopOffset();
+			const float BottomY = Viewport.GetPosY() + Viewport.GetHeight() - Viewport.GetBottomOffset();
 
 			const float L = Position.X;
-			const float R = Position.X + (Viewport.GetWidth() * Scale);
-			const float T = Position.Y + (TopY * Scale);
-			const float B = Position.Y + (BottomY * Scale);
+			const float R = Position.X + Viewport.GetWidth() * Scale;
+			const float T = Position.Y + TopY * Scale;
+			const float B = Position.Y + BottomY * Scale;
 			const FSlateClippingZone ClipZone(FVector2D(L, T), FVector2D(R, T), FVector2D(L, B), FVector2D(R, B));
 			DrawContext.ElementList.PushClip(ClipZone);
 
@@ -1226,10 +1225,13 @@ int32 STimingView::OnPaint(const FPaintArgs& Args, const FGeometry& AllottedGeom
 
 		// Draw the top docked tracks.
 		{
+			const float TopY = Viewport.GetPosY();
+			const float BottomY = Viewport.GetPosY() + Viewport.GetTopOffset();
+
 			const float L = Position.X;
-			const float R = Position.X + (Viewport.GetWidth() * Scale);
-			const float T = Position.Y;
-			const float B = Position.Y + (Viewport.GetTopOffset() * Scale);
+			const float R = Position.X + Viewport.GetWidth() * Scale;
+			const float T = Position.Y + TopY * Scale;
+			const float B = Position.Y + BottomY * Scale;
 			const FSlateClippingZone ClipZone(FVector2D(L, T), FVector2D(R, T), FVector2D(L, B), FVector2D(R, B));
 			DrawContext.ElementList.PushClip(ClipZone);
 
@@ -1244,7 +1246,7 @@ int32 STimingView::OnPaint(const FPaintArgs& Args, const FGeometry& AllottedGeom
 			if (bDrawTopSeparatorLine)
 			{
 				// Draw separator line between top docked tracks and scrollable tracks.
-				DrawContext.DrawBox(0.0f, Viewport.GetTopOffset() - 2.0f, Viewport.GetWidth(), 2.0f, WhiteBrush, FLinearColor(0.01f, 0.01f, 0.01f, 1.0f));
+				DrawContext.DrawBox(0.0f, BottomY - 2.0f, Viewport.GetWidth(), 2.0f, WhiteBrush, FLinearColor(0.01f, 0.01f, 0.01f, 1.0f));
 				++DrawContext.LayerId;
 			}
 
@@ -1253,10 +1255,13 @@ int32 STimingView::OnPaint(const FPaintArgs& Args, const FGeometry& AllottedGeom
 
 		// Draw the bottom docked tracks.
 		{
+			const float TopY = Viewport.GetPosY() + Viewport.GetHeight() - Viewport.GetBottomOffset();
+			const float BottomY = Viewport.GetPosY() + Viewport.GetHeight();
+
 			const float L = Position.X;
-			const float R = Position.X + (Viewport.GetWidth() * Scale);
-			const float T = Position.Y + ((Viewport.GetHeight() - Viewport.GetBottomOffset()) * Scale);
-			const float B = Position.Y + (Viewport.GetHeight() * Scale);
+			const float R = Position.X + Viewport.GetWidth() * Scale;
+			const float T = Position.Y + TopY * Scale;
+			const float B = Position.Y + BottomY * Scale;
 			const FSlateClippingZone ClipZone(FVector2D(L, T), FVector2D(R, T), FVector2D(L, B), FVector2D(R, B));
 			DrawContext.ElementList.PushClip(ClipZone);
 
@@ -1271,7 +1276,7 @@ int32 STimingView::OnPaint(const FPaintArgs& Args, const FGeometry& AllottedGeom
 			if (bDrawBottomSeparatorLine)
 			{
 				// Draw separator line between top docked tracks and scrollable tracks.
-				DrawContext.DrawBox(0.0f, Viewport.GetHeight() - Viewport.GetBottomOffset(), Viewport.GetWidth(), 2.0f, WhiteBrush, FLinearColor(0.01f, 0.01f, 0.01f, 1.0f));
+				DrawContext.DrawBox(0.0f, TopY, Viewport.GetWidth(), 2.0f, WhiteBrush, FLinearColor(0.01f, 0.01f, 0.01f, 1.0f));
 				++DrawContext.LayerId;
 			}
 
@@ -1374,16 +1379,41 @@ int32 STimingView::OnPaint(const FPaintArgs& Args, const FGeometry& AllottedGeom
 		PostDrawTracksDurationHistory.AddValue(PostDrawTracksStopwatch.AccumulatedTime);
 	}
 
+	//////////////////////////////////////////////////
+	// Draw relations between docked tracks.
+	{
+		bool bIsClipZoneSet = false;
+		if (GraphTrack->IsVisible() &&
+			GraphTrack->GetLocation() == ETimingTrackLocation::TopDocked)
+		{
+			// Avoid overlapping the Main Graph track (when top docked).
+			const float TopY = GraphTrack->GetPosY() + GraphTrack->GetHeight();
+			const float BottomY = Viewport.GetPosY() + Viewport.GetHeight();
+
+			const float L = Position.X;
+			const float R = Position.X + Viewport.GetWidth() * Scale;
+			const float T = Position.Y + TopY * Scale;
+			const float B = Position.Y + BottomY * Scale;
+			const FSlateClippingZone ClipZone(FVector2D(L, T), FVector2D(R, T), FVector2D(L, B), FVector2D(R, B));
+			DrawContext.ElementList.PushClip(ClipZone);
+			bIsClipZoneSet = true;
+		}
+
+		const FTimingViewDrawHelper& TimingHelper = *static_cast<const FTimingViewDrawHelper*>(&TimingDrawContext.GetHelper());
+		TimingHelper.DrawRelations(CurrentRelations, ITimingEventRelation::EDrawFilter::BetweenDockedTracks);
+
+		if (bIsClipZoneSet)
+		{
+			DrawContext.ElementList.PopClip();
+		}
+	}
 	//////////////////////////////////////////////////
 
 	// Draw tooltip with info about hovered event.
 	Tooltip.Draw(DrawContext);
 
-	// Fill background for the "Tracks" filter combobox.
-	DrawContext.DrawBox(0.0f, 0.0f, 73.0f, 24.0f, WhiteBrush, FLinearColor(0.06f, 0.06f, 0.06f, 1.0f));
-
-	// Fill background for the "Auto-scroll" toggle button.
-	DrawContext.DrawBox(ViewWidth - 39.0f, 0.0f, 39.0f, 24.0f, WhiteBrush, FLinearColor(0.06f, 0.06f, 0.06f, 1.0f));
+	// Fill the background of the toolbar.
+	DrawContext.DrawBox(0.0f, 0.0f, ViewWidth, Viewport.GetPosY(), WhiteBrush, FSlateColor(EStyleColor::Panel).GetSpecifiedColor());
 
 	//////////////////////////////////////////////////
 	// Draw the overscroll indication lines.
@@ -1394,8 +1424,8 @@ int32 STimingView::OnPaint(const FPaintArgs& Args, const FGeometry& AllottedGeom
 	if (OverscrollLeft > 0.0f)
 	{
 		// TODO: single box with gradient opacity
-		const float OverscrollLineY = Viewport.GetTopOffset();
-		const float OverscrollLineH = Viewport.GetScrollableAreaHeight();
+		const float OverscrollLineY = Viewport.GetPosY();
+		const float OverscrollLineH = Viewport.GetHeight();
 		for (int32 LineIndex = 0; LineIndex < OverscrollLineCount; ++LineIndex)
 		{
 			const float Opacity = OverscrollLeft * static_cast<float>(OverscrollLineCount - LineIndex) / static_cast<float>(OverscrollLineCount);
@@ -1404,8 +1434,8 @@ int32 STimingView::OnPaint(const FPaintArgs& Args, const FGeometry& AllottedGeom
 	}
 	if (OverscrollRight > 0.0f)
 	{
-		const float OverscrollLineY = Viewport.GetTopOffset();
-		const float OverscrollLineH = Viewport.GetScrollableAreaHeight();
+		const float OverscrollLineY = Viewport.GetPosY();
+		const float OverscrollLineH = Viewport.GetHeight();
 		for (int32 LineIndex = 0; LineIndex < OverscrollLineCount; ++LineIndex)
 		{
 			const float Opacity = OverscrollRight * static_cast<float>(OverscrollLineCount - LineIndex) / static_cast<float>(OverscrollLineCount);
@@ -1414,7 +1444,7 @@ int32 STimingView::OnPaint(const FPaintArgs& Args, const FGeometry& AllottedGeom
 	}
 	if (OverscrollTop > 0.0f)
 	{
-		const float OverscrollLineY = Viewport.GetTopOffset();
+		const float OverscrollLineY = Viewport.GetPosY() + Viewport.GetTopOffset();
 		for (int32 LineIndex = 0; LineIndex < OverscrollLineCount; ++LineIndex)
 		{
 			const float Opacity = OverscrollTop * static_cast<float>(OverscrollLineCount - LineIndex) / static_cast<float>(OverscrollLineCount);
@@ -1423,7 +1453,7 @@ int32 STimingView::OnPaint(const FPaintArgs& Args, const FGeometry& AllottedGeom
 	}
 	if (OverscrollBottom > 0.0f)
 	{
-		const float OverscrollLineY = ViewHeight - Viewport.GetBottomOffset();
+		const float OverscrollLineY = Viewport.GetPosY() + Viewport.GetHeight() - Viewport.GetBottomOffset();
 		for (int32 LineIndex = 0; LineIndex < OverscrollLineCount; ++LineIndex)
 		{
 			const float Opacity = OverscrollBottom * static_cast<float>(OverscrollLineCount - LineIndex) / static_cast<float>(OverscrollLineCount);
@@ -1433,30 +1463,6 @@ int32 STimingView::OnPaint(const FPaintArgs& Args, const FGeometry& AllottedGeom
 
 	//////////////////////////////////////////////////
 
-	// Draw relations between docked/foreground tracks.
-	if (GraphTrack->IsVisible())
-	{
-		const float TopY = GraphTrack->GetPosY() + GraphTrack->GetHeight();
-		const float BottomY = Viewport.GetHeight();
-
-		const float L = Position.X;
-		const float R = Position.X + (Viewport.GetWidth() * Scale);
-		const float T = Position.Y + (TopY * Scale);
-		const float B = Position.Y + (BottomY * Scale);
-		const FSlateClippingZone ClipZone(FVector2D(L, T), FVector2D(R, T), FVector2D(L, B), FVector2D(R, B));
-		DrawContext.ElementList.PushClip(ClipZone);
-	}
-
-	const FTimingViewDrawHelper& TimingHelper = *static_cast<const FTimingViewDrawHelper*>(&TimingDrawContext.GetHelper());
-	TimingHelper.DrawRelations(CurrentRelations, ITimingEventRelation::EDrawFilter::BetweenDockedTracks);
-
-	if (GraphTrack->IsVisible())
-	{
-		DrawContext.ElementList.PopClip();
-	}
-
-	//////////////////////////////////////////////////
-
 	const bool bShouldDisplayDebugInfo = FInsightsManager::Get()->IsDebugInfoEnabled();
 	if (bShouldDisplayDebugInfo)
 	{
@@ -1468,7 +1474,7 @@ int32 STimingView::OnPaint(const FPaintArgs& Args, const FGeometry& AllottedGeom
 		const float DbgW = 320.0f;
 		const float DbgH = DbgDY * 9 + 3.0f;
 		const float DbgX = ViewWidth - DbgW - 20.0f;
-		float DbgY = Viewport.GetTopOffset() + 10.0f;
+		float DbgY = Viewport.GetPosY() + Viewport.GetTopOffset() + 10.0f;
 
 		DrawContext.LayerId++;
 
@@ -2816,7 +2822,7 @@ void STimingView::ShowContextMenu(const FPointerEvent& MouseEvent)
 	{
 		bHasAnyActions |= Extender->ExtendGlobalContextMenu(*this, MenuBuilder);
 	}
-	
+
 	if(!bHasAnyActions)
 	{
 		MenuBuilder.BeginSection(TEXT("Empty"));
@@ -3251,17 +3257,17 @@ void STimingView::ScrollAtPosY(float ScrollPosY)
 
 ////////////////////////////////////////////////////////////////////////////////////////////////////
 
-void STimingView::BringIntoViewY(float InTopY, float InBottomY)
+void STimingView::BringIntoViewY(float InTopScrollY, float InBottomScrollY)
 {
 	const float ScrollY = Viewport.GetScrollPosY();
 	const float ScrollH = Viewport.GetScrollableAreaHeight();
-	if (ScrollY > InTopY)
+	if (ScrollY > InTopScrollY)
 	{
-		ScrollAtPosY(InTopY);
+		ScrollAtPosY(InTopScrollY);
 	}
-	else if (ScrollY + ScrollH < InBottomY)
+	else if (ScrollY + ScrollH < InBottomScrollY)
 	{
-		ScrollAtPosY(FMath::Min(InTopY, InBottomY - ScrollH));
+		ScrollAtPosY(FMath::Min(InTopScrollY, InBottomScrollY - ScrollH));
 	}
 }
 
@@ -3271,8 +3277,8 @@ void STimingView::BringScrollableTrackIntoView(const FBaseTimingTrack& InTrack)
 {
 	if (ensure(InTrack.GetLocation() == ETimingTrackLocation::Scrollable))
 	{
-		const float TopY = Viewport.GetScrollPosY() - Viewport.GetTopOffset() + InTrack.GetPosY();
-		BringIntoViewY(TopY, TopY + InTrack.GetHeight());
+		const float TopScrollY = InTrack.GetPosY() + Viewport.GetScrollPosY() - Viewport.GetTopOffset() - Viewport.GetPosY();
+		BringIntoViewY(TopScrollY, TopScrollY + InTrack.GetHeight());
 	}
 }
 
@@ -3517,11 +3523,11 @@ const TSharedPtr<FBaseTimingTrack> STimingView::GetTrackAt(float InPosX, float I
 {
 	TSharedPtr<FBaseTimingTrack> FoundTrack;
 
-	if (InPosY < 0.0f)
+	if (InPosY < Viewport.GetPosY())
 	{
 		// above viewport
 	}
-	else if (InPosY < Viewport.GetTopOffset())
+	else if (InPosY < Viewport.GetPosY() + Viewport.GetTopOffset())
 	{
 		// Top Docked Tracks
 		for (const TSharedPtr<FBaseTimingTrack>& TrackPtr : TopDockedTracks)
@@ -3537,7 +3543,7 @@ const TSharedPtr<FBaseTimingTrack> STimingView::GetTrackAt(float InPosX, float I
 			}
 		}
 	}
-	else if (InPosY < Viewport.GetHeight() - Viewport.GetBottomOffset())
+	else if (InPosY < Viewport.GetPosY() + Viewport.GetHeight() - Viewport.GetBottomOffset())
 	{
 		// Scrollable Tracks
 		for (const TSharedPtr<FBaseTimingTrack>& TrackPtr : ScrollableTracks)
@@ -3553,7 +3559,7 @@ const TSharedPtr<FBaseTimingTrack> STimingView::GetTrackAt(float InPosX, float I
 			}
 		}
 	}
-	else if (InPosY < Viewport.GetHeight())
+	else if (InPosY < Viewport.GetPosY() + Viewport.GetHeight())
 	{
 		// Bottom Docked Tracks
 		for (const TSharedPtr<FBaseTimingTrack>& TrackPtr : BottomDockedTracks)
@@ -3967,15 +3973,148 @@ bool STimingView::IsFilterByEventType(const uint64 EventType) const
 
 ////////////////////////////////////////////////////////////////////////////////////////////////////
 
+void STimingView::CreateCompactMenuLine(FMenuBuilder& MenuBuilder, FText Label, TSharedRef<SWidget> InnerWidget) const
+{
+	TSharedRef<SWidget> Widget =
+		SNew(SHorizontalBox)
+		+ SHorizontalBox::Slot()
+		.Padding(FMargin(-8.0f, 0.0f, 0.0f, 0.0f))
+		.AutoWidth()
+		.HAlign(HAlign_Fill)
+		.VAlign(VAlign_Center)
+		[
+			SNew(SBox)
+			.WidthOverride(140.0f)
+			.HAlign(HAlign_Right)
+			.Padding(FMargin(0.0f, 0.0f, 4.0f, 0.0f))
+			[
+				SNew(STextBlock)
+				.Text(Label)
+			]
+		]
+		+ SHorizontalBox::Slot()
+		.Padding(FMargin(0.0f, 1.0f, 8.0f, 1.0f))
+		.HAlign(HAlign_Fill)
+		.VAlign(VAlign_Center)
+		[
+			InnerWidget
+		]
+	;
+
+	MenuBuilder.AddWidget(Widget, FText(), true);
+}
+
+////////////////////////////////////////////////////////////////////////////////////////////////////
+
+TSharedRef<SWidget> STimingView::MakeCompactAutoScrollOptionsMenu()
+{
+	FMenuBuilder MenuBuilder(/*bInShouldCloseWindowAfterMenuSelection=*/true, CommandList);
+
+	MenuBuilder.BeginSection("AutoScrollOptions", LOCTEXT("AutoScrollOptionsHeading", "Auto-Scroll Options"));
+	{
+		CreateCompactMenuLine(MenuBuilder,
+			LOCTEXT("FrameAlignment", "Frame Alignment:"),
+			SNew(SSegmentedControl<int32>)
+			.OnValueChanged_Lambda([this](int32 InValue) { SetAutoScrollFrameAlignment(InValue); })
+			.Value_Lambda([this] { return AutoScrollFrameAlignment; })
+			+ SSegmentedControl<int32>::Slot(-1)
+			.Text(LOCTEXT("None", "None"))
+			.ToolTip(LOCTEXT("AutoScrollNoFrameAlignment_Tooltip", "Disables the frame alignment (when autoscrolling)."))
+			+ SSegmentedControl<int32>::Slot((int32)TraceFrameType_Game)
+			.Text(LOCTEXT("Game", "Game"))
+			.ToolTip(LOCTEXT("AutoScrollNoFrameAlignment_Tooltip", "Disables the frame alignment (when autoscrolling)."))
+			+ SSegmentedControl<int32>::Slot((int32)TraceFrameType_Rendering)
+			.Text(LOCTEXT("Rendering", "Rendering"))
+			.ToolTip(LOCTEXT("AutoScrollAlignWithGameFrames_Tooltip", "Aligns the viewport's center position with the start time of a Game frame (when autoscrolling)."))
+		);
+
+		CreateCompactMenuLine(MenuBuilder,
+			LOCTEXT("ViewportOffset", "Viewport Offset:"),
+			SNew(SHorizontalBox)
+			+ SHorizontalBox::Slot()
+			.AutoWidth()
+			[
+				SNew(SEditableTextBox)
+				.MinDesiredWidth(40.0f)
+				.HintText(LOCTEXT("ViewportOffsetCustomHint", "Custom"))
+				.Text_Lambda([this]
+				{
+					const FString ValueStr = (AutoScrollViewportOffsetPercent == 0.0) ? FString(TEXT("0")) : FString::Printf(TEXT("%g%%"), AutoScrollViewportOffsetPercent * 100.0);
+					return FText::FromString(ValueStr);
+				})
+				.OnTextChanged_Lambda([this](const FText& InText) { SetAutoScrollViewportOffset(atof(TCHAR_TO_ANSI(*InText.ToString())) * 0.01); })
+			]
+			+ SHorizontalBox::Slot()
+			.FillWidth(1.0f)
+			[
+				SNew(SSegmentedControl<double>)
+				.OnValueChanged_Lambda([this](double InValue) { SetAutoScrollViewportOffset(InValue); })
+				.Value_Lambda([this] { return AutoScrollViewportOffsetPercent; })
+				+ SSegmentedControl<double>::Slot(-0.1)
+				.Text(LOCTEXT("AutoScrollViewportOffset-10", "-10%"))
+				.ToolTip(LOCTEXT("AutoScrollViewportOffset-10_Tooltip", "Sets the viewport offset to -10% (i.e. backward) of the viewport's width (when autoscrolling).\nAvoids flickering as the end of the session will be outside of the viewport."))
+				+ SSegmentedControl<double>::Slot(0.0)
+				.Text(LOCTEXT("AutoScrollViewportOffset0", "0"))
+				.ToolTip(LOCTEXT("AutoScrollViewportOffset0_Tooltip", "Sets the viewport offset to 0 (when autoscrolling).\nThe right side of the viewport will correspond to the current session time."))
+				+ SSegmentedControl<double>::Slot(+0.1)
+				.Text(LOCTEXT("AutoScrollViewportOffset+10", "+10%"))
+				.ToolTip(LOCTEXT("AutoScrollViewportOffset+10_Tooltip", "Sets the viewport offset to +10% (i.e. forward) of the viewport's width (when autoscrolling).\nAllows 10% empty space on the right side of the viewport."))
+			]
+		);
+
+		CreateCompactMenuLine(MenuBuilder,
+			LOCTEXT("Delay", "Delay:"),
+			SNew(SHorizontalBox)
+			+ SHorizontalBox::Slot()
+			.AutoWidth()
+			[
+				SNew(SEditableTextBox)
+				.MinDesiredWidth(40.0f)
+				.HintText(LOCTEXT("AutoScrollDelayCustomHint", "Custom"))
+				.Text_Lambda([this]
+				{
+					const FString ValueStr = (AutoScrollMinDelay == 0.0) ? FString(TEXT("0")) : FString::Printf(TEXT("%gs"), AutoScrollMinDelay);
+					return FText::FromString(ValueStr);
+				})
+				.OnTextChanged_Lambda([this](const FText& InText) { SetAutoScrollDelay(atof(TCHAR_TO_ANSI(*InText.ToString()))); })
+			]
+			+ SHorizontalBox::Slot()
+			.FillWidth(1.0f)
+			[
+				SNew(SSegmentedControl<double>)
+				.OnValueChanged_Lambda([this](double InValue) { SetAutoScrollDelay(InValue); })
+				.Value_Lambda([this] { return AutoScrollMinDelay; })
+				+ SSegmentedControl<double>::Slot(0.0)
+				.Text(LOCTEXT("AutoScrollDelay0", "0"))
+				.ToolTip(LOCTEXT("AutoScrollDelay0_Tooltip", "Sets the time delay of the auto-scroll update to 0."))
+				+ SSegmentedControl<double>::Slot(0.3)
+				.Text(LOCTEXT("AutoScrollDelay300ms", "300ms"))
+				.ToolTip(LOCTEXT("AutoScrollDelay300ms_Tooltip", "Sets the time delay of the auto-scroll update to 300ms."))
+				+ SSegmentedControl<double>::Slot(1.0)
+				.Text(LOCTEXT("AutoScrollDelay1s", "1s"))
+				.ToolTip(LOCTEXT("AutoScrollDelay1s_Tooltip", "Sets the time delay of the auto-scroll update to 1s."))
+				+ SSegmentedControl<double>::Slot(3.0)
+				.Text(LOCTEXT("AutoScrollDelay3s", "3s"))
+				.ToolTip(LOCTEXT("AutoScrollDelay3s_Tooltip", "Sets the time delay of the auto-scroll update to 3s."))
+			]
+		);
+	}
+	MenuBuilder.EndSection();
+
+	return MenuBuilder.MakeWidget();
+}
+
+////////////////////////////////////////////////////////////////////////////////////////////////////
+
 TSharedRef<SWidget> STimingView::MakeAutoScrollOptionsMenu()
 {
 	FMenuBuilder MenuBuilder(/*bInShouldCloseWindowAfterMenuSelection=*/true, CommandList);
 
-	MenuBuilder.BeginSection("QuickFilter", LOCTEXT("AutoScrollHeading", "Auto Scroll Options"));
+	MenuBuilder.BeginSection("Alignment", LOCTEXT("AlignmentHeading", "Alignment"));
 	{
 		MenuBuilder.AddMenuEntry(
-			LOCTEXT("AutoScrollNoFrameAlignment", "No Frame Alignment"),
-			LOCTEXT("AutoScrollNoFrameAlignment_Tooltip", "Disable frame alignment."),
+			LOCTEXT("AutoScrollNoFrameAlignment", "None"),
+			LOCTEXT("AutoScrollNoFrameAlignment_Tooltip", "Disables the frame alignment (when autoscrolling)."),
 			FSlateIcon(),
 			FUIAction(FExecuteAction::CreateSP(this, &STimingView::SetAutoScrollFrameAlignment, -1),
 				FCanExecuteAction(),
@@ -3985,8 +4124,8 @@ TSharedRef<SWidget> STimingView::MakeAutoScrollOptionsMenu()
 		);
 
 		MenuBuilder.AddMenuEntry(
-			LOCTEXT("AutoScrollAlignWithGameFrames", "Align with Game Frames"),
-			LOCTEXT("AutoScrollAlignWithGameFrames_Tooltip", "Align the viewport's center position with the start time of a Game frame."),
+			LOCTEXT("AutoScrollAlignWithGameFrames", "Game Frames"),
+			LOCTEXT("AutoScrollAlignWithGameFrames_Tooltip", "Aligns the viewport's center position with the start time of a Game frame (when autoscrolling)."),
 			FSlateIcon(),
 			FUIAction(FExecuteAction::CreateSP(this, &STimingView::SetAutoScrollFrameAlignment, (int32)TraceFrameType_Game),
 				FCanExecuteAction(),
@@ -3996,8 +4135,8 @@ TSharedRef<SWidget> STimingView::MakeAutoScrollOptionsMenu()
 		);
 
 		MenuBuilder.AddMenuEntry(
-			LOCTEXT("AutoScrollAlignWithRenderingFrames", "Align with Rendering Frames"),
-			LOCTEXT("AutoScrollAlignWithRenderingFrames_Tooltip", "Align the viewport's center position with the start time of a Rendering frame."),
+			LOCTEXT("AutoScrollAlignWithRenderingFrames", "Rendering Frames"),
+			LOCTEXT("AutoScrollAlignWithRenderingFrames_Tooltip", "Aligns the viewport's center position with the start time of a Rendering frame (when autoscrolling)."),
 			FSlateIcon(),
 			FUIAction(FExecuteAction::CreateSP(this, &STimingView::SetAutoScrollFrameAlignment, (int32)TraceFrameType_Rendering),
 				FCanExecuteAction(),
@@ -4005,12 +4144,14 @@ TSharedRef<SWidget> STimingView::MakeAutoScrollOptionsMenu()
 			NAME_None,
 			EUserInterfaceActionType::RadioButton
 		);
+	}
+	MenuBuilder.EndSection();
 
-		MenuBuilder.AddMenuSeparator();
-
+	MenuBuilder.BeginSection("ViewportOffset", LOCTEXT("ViewportOffsetHeading", "Viewport Offset"));
+	{
 		MenuBuilder.AddMenuEntry(
-			LOCTEXT("AutoScrollViewportOffset-10", "Viewport Offset: -10%"),
-			LOCTEXT("AutoScrollViewportOffset-10_Tooltip", "Set the viewport offset to -10% (i.e. backward) of the viewport's width.\nAvoids flickering as the end of session will be outside of the viewport."),
+			LOCTEXT("AutoScrollViewportOffset-10", "-10%"),
+			LOCTEXT("AutoScrollViewportOffset-10_Tooltip", "Sets the viewport offset to -10% (i.e. backward) of the viewport's width (when autoscrolling).\nAvoids flickering as the end of the session will be outside of the viewport."),
 			FSlateIcon(),
 			FUIAction(FExecuteAction::CreateSP(this, &STimingView::SetAutoScrollViewportOffset, -0.1),
 				FCanExecuteAction(),
@@ -4020,8 +4161,8 @@ TSharedRef<SWidget> STimingView::MakeAutoScrollOptionsMenu()
 		);
 
 		MenuBuilder.AddMenuEntry(
-			LOCTEXT("AutoScrollViewportOffset0", "Viewport Offset: 0"),
-			LOCTEXT("AutoScrollViewportOffset0_Tooltip", "Set the viewport offset to 0."),
+			LOCTEXT("AutoScrollViewportOffset0", "0"),
+			LOCTEXT("AutoScrollViewportOffset0_Tooltip", "Sets the viewport offset to 0 (when autoscrolling).\nThe right side of the viewport will correspond to the current session time."),
 			FSlateIcon(),
 			FUIAction(FExecuteAction::CreateSP(this, &STimingView::SetAutoScrollViewportOffset, 0.0),
 				FCanExecuteAction(),
@@ -4031,8 +4172,8 @@ TSharedRef<SWidget> STimingView::MakeAutoScrollOptionsMenu()
 		);
 
 		MenuBuilder.AddMenuEntry(
-			LOCTEXT("AutoScrollViewportOffset+10", "Viewport Offset: +10%"),
-			LOCTEXT("AutoScrollViewportOffset+10_Tooltip", "Set the viewport offset to +10% (i.e. forward) of the viewport's width.\nAllows 10% empty space on the right side of the viewport."),
+			LOCTEXT("AutoScrollViewportOffset+10", "+10%"),
+			LOCTEXT("AutoScrollViewportOffset+10_Tooltip", "Sets the viewport offset to +10% (i.e. forward) of the viewport's width (when autoscrolling).\nAllows 10% empty space on the right side of the viewport."),
 			FSlateIcon(),
 			FUIAction(FExecuteAction::CreateSP(this, &STimingView::SetAutoScrollViewportOffset, +0.1),
 				FCanExecuteAction(),
@@ -4041,11 +4182,38 @@ TSharedRef<SWidget> STimingView::MakeAutoScrollOptionsMenu()
 			EUserInterfaceActionType::RadioButton
 		);
 
-		MenuBuilder.AddMenuSeparator();
+		MenuBuilder.AddMenuEntry(
+			FUIAction(FExecuteAction(), FCanExecuteAction(), FIsActionChecked::CreateLambda(
+				[this]() -> bool
+				{
+					return	AutoScrollViewportOffsetPercent != 0.0 &&
+							AutoScrollViewportOffsetPercent != -0.1 &&
+							AutoScrollViewportOffsetPercent != +0.1;
+				})),
+			SNew(SBox)
+			.Padding(FMargin(0.0f, 0.0f, 12.0f, 0.0f))
+			[
+				SNew(SEditableTextBox)
+				.HintText(LOCTEXT("ViewportOffsetCustomHint", "Custom"))
+				.Text_Lambda([this]
+				{
+					const FString ValueStr = (AutoScrollViewportOffsetPercent == 0.0) ? FString(TEXT("0")) : FString::Printf(TEXT("%g%%"), AutoScrollViewportOffsetPercent * 100.0);
+					return FText::FromString(ValueStr);
+				})
+				.OnTextChanged_Lambda([this](const FText& InText) { SetAutoScrollViewportOffset(atof(TCHAR_TO_ANSI(*InText.ToString())) * 0.01); })
+			],
+			NAME_None,
+			LOCTEXT("AutoScrollViewportOffsetCustom_Tooltip", "Sets a custom value for the viewport offset as percent from viewport's width (when autoscrolling)."),
+			EUserInterfaceActionType::RadioButton
+		);
+	}
+	MenuBuilder.EndSection();
 
+	MenuBuilder.BeginSection("Delay", LOCTEXT("DelayHeading", "Delay"));
+	{
 		MenuBuilder.AddMenuEntry(
-			LOCTEXT("AutoScrollDelay0", "Delay: 0"),
-			LOCTEXT("AutoScrollDelay0_Tooltip", "Set the time delay of the auto-scroll update to 0."),
+			LOCTEXT("AutoScrollDelay0", "0"),
+			LOCTEXT("AutoScrollDelay0_Tooltip", "Sets the time delay of the auto-scroll update to 0."),
 			FSlateIcon(),
 			FUIAction(FExecuteAction::CreateSP(this, &STimingView::SetAutoScrollDelay, 0.0),
 				FCanExecuteAction(),
@@ -4055,8 +4223,8 @@ TSharedRef<SWidget> STimingView::MakeAutoScrollOptionsMenu()
 		);
 
 		MenuBuilder.AddMenuEntry(
-			LOCTEXT("AutoScrollDelay300ms", "Delay: 300ms"),
-			LOCTEXT("AutoScrollDelay300ms_Tooltip", "Set the time delay of the auto-scroll update to 300ms."),
+			LOCTEXT("AutoScrollDelay300ms", "300ms"),
+			LOCTEXT("AutoScrollDelay300ms_Tooltip", "Sets the time delay of the auto-scroll update to 300ms."),
 			FSlateIcon(),
 			FUIAction(FExecuteAction::CreateSP(this, &STimingView::SetAutoScrollDelay, 0.3),
 				FCanExecuteAction(),
@@ -4066,8 +4234,8 @@ TSharedRef<SWidget> STimingView::MakeAutoScrollOptionsMenu()
 		);
 
 		MenuBuilder.AddMenuEntry(
-			LOCTEXT("AutoScrollDelay1s", "Delay: 1s"),
-			LOCTEXT("AutoScrollDelay1s_Tooltip", "Set the time delay of the auto-scroll update to 1s."),
+			LOCTEXT("AutoScrollDelay1s", "1s"),
+			LOCTEXT("AutoScrollDelay1s_Tooltip", "Sets the time delay of the auto-scroll update to 1s."),
 			FSlateIcon(),
 			FUIAction(FExecuteAction::CreateSP(this, &STimingView::SetAutoScrollDelay, 1.0),
 				FCanExecuteAction(),
@@ -4077,8 +4245,8 @@ TSharedRef<SWidget> STimingView::MakeAutoScrollOptionsMenu()
 		);
 
 		MenuBuilder.AddMenuEntry(
-			LOCTEXT("AutoScrollDelay3s", "Delay: 3s"),
-			LOCTEXT("AutoScrollDelay3s_Tooltip", "Set the time delay of the auto-scroll update to 3s."),
+			LOCTEXT("AutoScrollDelay3s", "3s"),
+			LOCTEXT("AutoScrollDelay3s_Tooltip", "Sets the time delay of the auto-scroll update to 3s."),
 			FSlateIcon(),
 			FUIAction(FExecuteAction::CreateSP(this, &STimingView::SetAutoScrollDelay, 3.0),
 				FCanExecuteAction(),
@@ -4086,7 +4254,34 @@ TSharedRef<SWidget> STimingView::MakeAutoScrollOptionsMenu()
 			NAME_None,
 			EUserInterfaceActionType::RadioButton
 		);
+
+		MenuBuilder.AddMenuEntry(
+			FUIAction(FExecuteAction(), FCanExecuteAction(), FIsActionChecked::CreateLambda(
+				[this]() -> bool
+				{
+					return AutoScrollMinDelay != 0.0 &&
+						   AutoScrollMinDelay != 0.3 &&
+						   AutoScrollMinDelay != 1.0 &&
+						   AutoScrollMinDelay != 3.0;
+				})),
+			SNew(SBox)
+			.Padding(FMargin(0.0f, 0.0f, 12.0f, 0.0f))
+			[
+				SNew(SEditableTextBox)
+				.HintText(LOCTEXT("AutoScrollDelayCustomHint", "Custom"))
+				.Text_Lambda([this]
+				{
+					const FString ValueStr = (AutoScrollMinDelay == 0.0) ? FString(TEXT("0")) : FString::Printf(TEXT("%gs"), AutoScrollMinDelay);
+					return FText::FromString(ValueStr);
+				})
+				.OnTextChanged_Lambda([this](const FText& InText) { SetAutoScrollDelay(atof(TCHAR_TO_ANSI(*InText.ToString()))); })
+			],
+			NAME_None,
+			LOCTEXT("AutoScrollDelayCustom_Tooltip", "Sets a custom time delay (in seconds) for the auto-scroll update."),
+			EUserInterfaceActionType::RadioButton
+		);
 	}
+	MenuBuilder.EndSection();
 
 	return MenuBuilder.MakeWidget();
 }
diff --git a/Engine/Source/Developer/TraceInsights/Private/Insights/Widgets/STimingView.h b/Engine/Source/Developer/TraceInsights/Private/Insights/Widgets/STimingView.h
--- a/Engine/Source/Developer/TraceInsights/Private/Insights/Widgets/STimingView.h
+++ b/Engine/Source/Developer/TraceInsights/Private/Insights/Widgets/STimingView.h
@@ -71,6 +71,8 @@ public:
 	 */
 	void Construct(const FArguments& InArgs);
 
+	void CreateCompactMenuLine(FMenuBuilder& MenuBuilder, FText Label, TSharedRef<SWidget> InnerWidget) const;
+	TSharedRef<SWidget> MakeCompactAutoScrollOptionsMenu();
 	TSharedRef<SWidget> MakeAutoScrollOptionsMenu();
 
 	TSharedRef<SWidget> MakeTracksFilterMenu();
@@ -453,7 +455,7 @@ protected:
 	FReply AllowTracksToProcessOnMouseButtonUp(const FGeometry& MyGeometry, const FPointerEvent& MouseEvent);
 	FReply AllowTracksToProcessOnMouseButtonDoubleClick(const FGeometry& MyGeometry, const FPointerEvent& MouseEvent);
 
-	void SetTrackPosY(TSharedPtr<FBaseTimingTrack>& TrackPtr, float TopOffsetY) const;
+	void SetTrackPosY(TSharedPtr<FBaseTimingTrack>& TrackPtr, float TrackPosY) const;
 	
 	void FindNextEvent();
 	void FindPrevEvent();
diff --git a/Engine/Source/Developer/TraceInsights/Public/Insights/ViewModels/TimingTrackViewport.h b/Engine/Source/Developer/TraceInsights/Public/Insights/ViewModels/TimingTrackViewport.h
--- a/Engine/Source/Developer/TraceInsights/Public/Insights/ViewModels/TimingTrackViewport.h
+++ b/Engine/Source/Developer/TraceInsights/Public/Insights/ViewModels/TimingTrackViewport.h
@@ -81,6 +81,7 @@ public:
 		MaxScaleX = 1.0E10; // 10ns between major tick marks
 		ScaleX = (5 * 20) / 5.0; // 5s between major tick marks
 
+		PosY = 0.0f;
 		TopOffset = 0.0f;
 		BottomOffset = 0.0f;
 		ScrollHeight = 1.0f;
@@ -182,6 +183,10 @@ public:
 	//////////////////////////////////////////////////
 	// Vertical axis
 
+	/** The Y position where it starts the top docked tracks (i.e. height of the toolbar), in pixels. [px] */
+	float GetPosY() const { return PosY; }
+	void SetPosY(float InPosY) { PosY = InPosY; }
+
 	/** Top offset (total height of the visible top docked tracks), in pixels. [px] */
 	float GetTopOffset() const { return TopOffset; }
 	void SetTopOffset(float InTopOffset) { TopOffset = InTopOffset; }
@@ -200,7 +205,7 @@ public:
 	float GetScrollPosY() const { return ScrollPosY; }
 	void SetScrollPosY(const float InScrollPosY) { ScrollPosY = InScrollPosY; }
 
-	float GetViewportY(const float Y) const { return TopOffset + Y - ScrollPosY; }
+	float GetViewportY(const float Y) const { return PosY + TopOffset + Y - ScrollPosY; }
 
 	const FTimingViewLayout& GetLayout() const { return Layout; }
 	void UpdateLayout();
@@ -260,6 +265,7 @@ private:
 	double MaxScaleX; // max scale factor; [px/s]
 	double ScaleX; // scale factor between seconds and pixels; [px/s]
 
+	float PosY; // position of top docked tracks == height of toolbar; [px]
 	float TopOffset; // top offset, in pixels; [px]
 	float BottomOffset; // bottom offset, in pixels; [px]
 	float ScrollHeight; // height of the vertical scrollable area, in pixels; [px]
