[Insights] Add filtering by Track name to the quick find widget

diff --git a/Engine/Source/Developer/TraceInsights/Private/Insights/ViewModels/FilterConfiguratorNode.cpp b/Engine/Source/Developer/TraceInsights/Private/Insights/ViewModels/FilterConfiguratorNode.cpp
--- a/Engine/Source/Developer/TraceInsights/Private/Insights/ViewModels/FilterConfiguratorNode.cpp
+++ b/Engine/Source/Developer/TraceInsights/Private/Insights/ViewModels/FilterConfiguratorNode.cpp
@@ -131,7 +131,7 @@ void FFilterConfiguratorNode::SetSelectedFilter(TSharedPtr<struct FFilter> InSel
 		SetSelectedFilterOperator(SelectedFilter->GetSupportedOperators()->GetData()[0]);
 
 		AvailableFilterOperators->Empty();
-		Insights::FFilter::SupportedOperatorsArrayPtr AvailableOperators = InSelectedFilter->SupportedOperators;
+		Insights::SupportedOperatorsArrayPtr AvailableOperators = InSelectedFilter->SupportedOperators;
 		for (auto& FilterOperator : *AvailableOperators)
 		{
 			AvailableFilterOperators->Add(FilterOperator);
@@ -238,6 +238,12 @@ bool FFilterConfiguratorNode::ApplyFilters(const FFilterContext& Context) const
 	}
 	else
 	{
+		if (!Context.HasFilterData(SelectedFilter->Key))
+		{
+			// If data is not set for this filter we consider it passed the test.
+			return true;
+		}
+
 		switch (SelectedFilter->DataType)
 		{
 		case EFilterDataType::Double:
diff --git a/Engine/Source/Developer/TraceInsights/Private/Insights/ViewModels/Filters.cpp b/Engine/Source/Developer/TraceInsights/Private/Insights/ViewModels/Filters.cpp
--- a/Engine/Source/Developer/TraceInsights/Private/Insights/ViewModels/Filters.cpp
+++ b/Engine/Source/Developer/TraceInsights/Private/Insights/ViewModels/Filters.cpp
@@ -12,6 +12,9 @@
 namespace Insights
 {
 
+INSIGHTS_IMPLEMENT_RTTI(FFilter)
+INSIGHTS_IMPLEMENT_RTTI(FFilterWithSuggestions)
+
 ////////////////////////////////////////////////////////////////////////////////////////////////////
 
 TSharedPtr<FFilterService> FFilterService::Instance;
diff --git a/Engine/Source/Developer/TraceInsights/Private/Insights/ViewModels/Filters.h b/Engine/Source/Developer/TraceInsights/Private/Insights/ViewModels/Filters.h
--- a/Engine/Source/Developer/TraceInsights/Private/Insights/ViewModels/Filters.h
+++ b/Engine/Source/Developer/TraceInsights/Private/Insights/ViewModels/Filters.h
@@ -7,6 +7,8 @@
 #include "CoreMinimal.h"
 #include "Misc/TVariant.h"
 
+#include "Insights/Common/SimpleRtti.h"
+
 #define LOCTEXT_NAMESPACE "Filters"
 
 class FSpawnTabArgs;
@@ -108,10 +110,13 @@ public:
 
 ////////////////////////////////////////////////////////////////////////////////////////////////////
 
+typedef TSharedPtr<const TArray<TSharedPtr<IFilterOperator>>> SupportedOperatorsArrayPtr;
+
 struct FFilter
 {
-	typedef TSharedPtr<const TArray<TSharedPtr<IFilterOperator>>> SupportedOperatorsArrayPtr;
+	INSIGHTS_DECLARE_RTTI_BASE(FFilter)
 
+public:
 	FFilter(int32 InKey, FText InName, FText InDesc, EFilterDataType InDataType, SupportedOperatorsArrayPtr InSupportedOperators)
 		: Key(InKey)
 		, Name(InName)
@@ -119,6 +124,7 @@ struct FFilter
 		, DataType(InDataType)
 		, SupportedOperators(InSupportedOperators)
 	{}
+	virtual ~FFilter() {}
 
 	SupportedOperatorsArrayPtr GetSupportedOperators() { return SupportedOperators;	}
 
@@ -132,6 +138,24 @@ struct FFilter
 
 ////////////////////////////////////////////////////////////////////////////////////////////////////
 
+struct FFilterWithSuggestions : FFilter
+{
+	typedef TFunction<void(const FString& /*Text*/, TArray<FString>& OutSuggestions)> GetSuggestionsCallback;
+
+	INSIGHTS_DECLARE_RTTI(FFilterWithSuggestions, FFilter)
+
+public:
+	FFilterWithSuggestions(int32 InKey, FText InName, FText InDesc, EFilterDataType InDataType, SupportedOperatorsArrayPtr InSupportedOperators)
+		: FFilter(InKey, InName, InDesc, InDataType, InSupportedOperators)
+	{ }
+
+	virtual ~FFilterWithSuggestions() {}
+
+	GetSuggestionsCallback Callback;
+};
+
+////////////////////////////////////////////////////////////////////////////////////////////////////
+
 class FFilterStorage
 {
 public:
@@ -235,6 +259,11 @@ public:
 		OutData = Data->Get<T>();
 	}
 
+	bool HasFilterData(int32 Key) const
+	{
+		return DataMap.Contains(Key);
+	}
+
 private:
 	TMap<int32, ContextData> DataMap;
 };
diff --git a/Engine/Source/Developer/TraceInsights/Private/Insights/ViewModels/ThreadTimingTrack.cpp b/Engine/Source/Developer/TraceInsights/Private/Insights/ViewModels/ThreadTimingTrack.cpp
--- a/Engine/Source/Developer/TraceInsights/Private/Insights/ViewModels/ThreadTimingTrack.cpp
+++ b/Engine/Source/Developer/TraceInsights/Private/Insights/ViewModels/ThreadTimingTrack.cpp
@@ -881,7 +881,7 @@ void FThreadTimingTrack::BuildFilteredDrawState(ITimingEventsTrackDrawStateBuild
 
 					// Note: Enumerating events for filtering should not use downsampling.
 					Params.Resolution = 0.0;
-					Params.SetupCallback = [&FilteredEvents, &FilterContexts](uint32 NumTasks)
+					Params.SetupCallback = [&FilteredEvents, &FilterContexts, this](uint32 NumTasks)
 					{
 						FilteredEvents.AddDefaulted(NumTasks);
 						FilterContexts.AddDefaulted(NumTasks);
@@ -891,6 +891,7 @@ void FThreadTimingTrack::BuildFilteredDrawState(ITimingEventsTrackDrawStateBuild
 							Context.AddFilterData<double>(static_cast<int32>(EFilterField::EndTime), 0.0f);
 							Context.AddFilterData<double>(static_cast<int32>(EFilterField::Duration), 0.0f);
 							Context.AddFilterData<int64>(static_cast<int32>(EFilterField::EventType), 0);
+							Context.AddFilterData<FString>(static_cast<int32>(EFilterField::TrackName), this->GetName());
 						}
 					};
 					Params.Callback = [this, &Builder, TimerReader, &FilteredEvents, &FilterContexts](double StartTime, double EndTime, uint32 Depth, const TraceServices::FTimingProfilerEvent& Event, uint32 TaskIndex)
diff --git a/Engine/Source/Developer/TraceInsights/Private/Insights/ViewModels/ThreadTimingTrack.h b/Engine/Source/Developer/TraceInsights/Private/Insights/ViewModels/ThreadTimingTrack.h
--- a/Engine/Source/Developer/TraceInsights/Private/Insights/ViewModels/ThreadTimingTrack.h
+++ b/Engine/Source/Developer/TraceInsights/Private/Insights/ViewModels/ThreadTimingTrack.h
@@ -106,7 +106,8 @@ enum class EFilterField : int32
 	StartTime = 0,
 	EndTime = 1,
 	Duration = 2,
-	EventType = 3
+	EventType = 3,
+	TrackName = 4
 };
 
 ////////////////////////////////////////////////////////////////////////////////////////////////////
diff --git a/Engine/Source/Developer/TraceInsights/Private/Insights/Widgets/SFilterConfiguratorRow.cpp b/Engine/Source/Developer/TraceInsights/Private/Insights/Widgets/SFilterConfiguratorRow.cpp
--- a/Engine/Source/Developer/TraceInsights/Private/Insights/Widgets/SFilterConfiguratorRow.cpp
+++ b/Engine/Source/Developer/TraceInsights/Private/Insights/Widgets/SFilterConfiguratorRow.cpp
@@ -9,6 +9,7 @@
 #include "Widgets/Images/SImage.h"
 #include "Widgets/Input/SButton.h"
 #include "Widgets/Input/SEditableTextBox.h"
+#include "Widgets/Input/SSuggestionTextBox.h"
 #include "Widgets/Views/STableRow.h"
 #include "Widgets/Views/STableViewBase.h"
 #include "Widgets/Views/STreeView.h"
@@ -95,13 +96,37 @@ TSharedRef<SWidget> SFilterConfiguratorRow::GenerateWidgetForColumn(const FName&
 				]
 			];
 
-		GeneratedWidget->AddSlot()
-			.AutoWidth()
-			.VAlign(VAlign_Center)
-			.HAlign(HAlign_Left)
-			.FillWidth(1.0)
-			.Padding(2.0f)
-			[
+		if (FilterConfiguratorNodePtr->GetSelectedFilter()->Is<FFilterWithSuggestions>())
+		{
+			SuggestionTextBoxValue = FilterConfiguratorNodePtr->GetTextBoxValue();
+
+			GeneratedWidget->AddSlot()
+				.AutoWidth()
+				.VAlign(VAlign_Center)
+				.HAlign(HAlign_Left)
+				.FillWidth(1.0)
+				.Padding(2.0f)
+				[
+					SNew(SSuggestionTextBox)
+					.MinDesiredWidth(50.0f)
+					.ForegroundColor(FLinearColor::White)
+					.OnTextCommitted(this, &SFilterConfiguratorRow::SuggestionTextBox_OnValueCommitted)
+					.OnTextChanged(this, &SFilterConfiguratorRow::SuggestionTextBox_OnValueChanged)
+					.Text(this, &SFilterConfiguratorRow::SuggestionTextBox_GetValue)
+					.ToolTipText(this, &SFilterConfiguratorRow::GetTextBoxTooltipText)
+					.HintText(this, &SFilterConfiguratorRow::GetTextBoxHintText)
+					.OnShowingSuggestions(FOnShowingSuggestions::CreateSP(this, &SFilterConfiguratorRow::SuggestionTextBox_GetSuggestions))
+				];
+		}
+		else
+		{
+			GeneratedWidget->AddSlot()
+				.AutoWidth()
+				.VAlign(VAlign_Center)
+				.HAlign(HAlign_Left)
+				.FillWidth(1.0)
+				.Padding(2.0f)
+				[
 				SNew(SEditableTextBox)
 				.MinDesiredWidth(50.0f)
 				.OnTextCommitted(this, &SFilterConfiguratorRow::OnTextBoxValueCommitted)
@@ -109,7 +134,8 @@ TSharedRef<SWidget> SFilterConfiguratorRow::GenerateWidgetForColumn(const FName&
 				.ToolTipText(this, &SFilterConfiguratorRow::GetTextBoxTooltipText)
 				.HintText(this, &SFilterConfiguratorRow::GetTextBoxHintText)
 				.OnVerifyTextChanged(this, &SFilterConfiguratorRow::TextBox_OnVerifyTextChanged)
-			];
+				];
+		}
 			
 		GeneratedWidget->AddSlot()
 			.VAlign(VAlign_Center)
@@ -278,6 +304,11 @@ const TArray<TSharedPtr<struct FFilter>>* SFilterConfiguratorRow::GetAvailableFi
 void SFilterConfiguratorRow::AvailableFilters_OnSelectionChanged(TSharedPtr<FFilter> InFilter, ESelectInfo::Type SelectInfo)
 {
 	FilterConfiguratorNodePtr->SetSelectedFilter(InFilter);
+
+	TSharedPtr<STreeView<FFilterConfiguratorNodePtr>> OwnerTable = StaticCastSharedPtr<STreeView<FFilterConfiguratorNodePtr>>(OwnerTablePtr.Pin());
+	ensure(OwnerTable.IsValid());
+	ClearCellCache();
+	GenerateColumns(OwnerTable->GetHeaderRow().ToSharedRef());
 }
 
 ///////////////////////////////////////////////////////////////////////////////////////////////////
@@ -448,6 +479,41 @@ bool SFilterConfiguratorRow::TextBox_OnVerifyTextChanged(const FText& InText, FT
 
 ///////////////////////////////////////////////////////////////////////////////////////////////////
 
+void SFilterConfiguratorRow::SuggestionTextBox_GetSuggestions(const FString& Text, TArray<FString>& Suggestions)
+{
+	TSharedPtr<FFilter> Filter = FilterConfiguratorNodePtr->GetSelectedFilter();
+	
+	if (Filter->Is<FFilterWithSuggestions>())
+	{
+		TSharedPtr<FFilterWithSuggestions> FilterWithSuggestions = StaticCastSharedPtr<FFilterWithSuggestions>(Filter);
+		FilterWithSuggestions->Callback(Text, Suggestions);
+	}
+}
+
+///////////////////////////////////////////////////////////////////////////////////////////////////
+
+void SFilterConfiguratorRow::SuggestionTextBox_OnValueChanged(const FText& InNewText)
+{
+	SuggestionTextBoxValue = InNewText.ToString();
+}
+
+///////////////////////////////////////////////////////////////////////////////////////////////////
+
+FText SFilterConfiguratorRow::SuggestionTextBox_GetValue() const
+{
+	return FText::FromString(SuggestionTextBoxValue);
+}
+
+///////////////////////////////////////////////////////////////////////////////////////////////////
+
+void SFilterConfiguratorRow::SuggestionTextBox_OnValueCommitted(const FText& InNewText, ETextCommit::Type InTextCommit)
+{
+	FilterConfiguratorNodePtr->SetTextBoxValue(InNewText.ToString());
+}
+
+///////////////////////////////////////////////////////////////////////////////////////////////////
+
+
 } // namespace Insights
 
 #undef LOCTEXT_NAMESPACE
diff --git a/Engine/Source/Developer/TraceInsights/Private/Insights/Widgets/SFilterConfiguratorRow.h b/Engine/Source/Developer/TraceInsights/Private/Insights/Widgets/SFilterConfiguratorRow.h
--- a/Engine/Source/Developer/TraceInsights/Private/Insights/Widgets/SFilterConfiguratorRow.h
+++ b/Engine/Source/Developer/TraceInsights/Private/Insights/Widgets/SFilterConfiguratorRow.h
@@ -62,10 +62,16 @@ private:
 
 	FText GetTextBoxValue() const;
 	void OnTextBoxValueCommitted(const FText& InNewText, ETextCommit::Type InTextCommit);
+
 	FText GetTextBoxTooltipText() const;
 	FText GetTextBoxHintText() const;
 	bool TextBox_OnVerifyTextChanged(const FText& InText, FText& OutErrorMessage);
 
+	void SuggestionTextBox_GetSuggestions(const FString& Text, TArray<FString>& Suggestions);
+	void SuggestionTextBox_OnValueChanged(const FText& InNewText);
+	FText SuggestionTextBox_GetValue() const;
+	void SuggestionTextBox_OnValueCommitted(const FText& InNewText, ETextCommit::Type InTextCommit);
+
 private:
 	FFilterConfiguratorNodePtr FilterConfiguratorNodePtr;
 
@@ -74,6 +80,8 @@ private:
 	TSharedPtr<SComboBox<TSharedPtr<class IFilterOperator>>> FilterOperatorComboBox;
 
 	TSharedPtr<SComboBox<TSharedPtr<struct FFilterGroupOperator>>> FilterGroupOperatorComboBox;
+
+	FString SuggestionTextBoxValue;
 };
 
 } // namnespace Insights
\ No newline at end of file
diff --git a/Engine/Source/Developer/TraceInsights/Private/Insights/Widgets/STimingView.cpp b/Engine/Source/Developer/TraceInsights/Private/Insights/Widgets/STimingView.cpp
--- a/Engine/Source/Developer/TraceInsights/Private/Insights/Widgets/STimingView.cpp
+++ b/Engine/Source/Developer/TraceInsights/Private/Insights/Widgets/STimingView.cpp
@@ -4504,6 +4504,20 @@ void STimingView::QuickFind_Execute()
 		AvailableFilters->Add(MakeShared<FFilter>(static_cast<int32>(EFilterField::Duration), LOCTEXT("Duration", "Duration"), LOCTEXT("Duration", "Duration"), EFilterDataType::Double, FFilterService::Get()->GetDoubleOperators()));
 		AvailableFilters->Add(MakeShared<FFilter>(static_cast<int32>(EFilterField::EventType), LOCTEXT("Type", "Type"), LOCTEXT("Type", "Type"), EFilterDataType::Int64, FFilterService::Get()->GetIntegerOperators()));
 
+		TSharedPtr<TArray<TSharedPtr<IFilterOperator>>>  TrackNameFilterOperators = MakeShared<TArray<TSharedPtr<IFilterOperator>>>();
+		TrackNameFilterOperators->Add(StaticCastSharedRef<IFilterOperator>(MakeShared<FFilterOperator<FString>>(EFilterOperator::Eq, TEXT("Is"), [](const FString& lhs, const FString& rhs) { return lhs.Equals(rhs); })));
+		TrackNameFilterOperators->Add(StaticCastSharedRef<IFilterOperator>(MakeShared<FFilterOperator<FString>>(EFilterOperator::Contains, TEXT("Contains"), [](const FString& lhs, const FString& rhs) { return lhs.Contains(rhs); })));
+		TrackNameFilterOperators->Add(StaticCastSharedRef<IFilterOperator>(MakeShared<FFilterOperator<FString>>(EFilterOperator::Contains, TEXT("Is Not"), [](const FString& lhs, const FString& rhs) { return !lhs.Equals(rhs); })));
+		TrackNameFilterOperators->Add(StaticCastSharedRef<IFilterOperator>(MakeShared<FFilterOperator<FString>>(EFilterOperator::Contains, TEXT("Does Not Contain"), [](const FString& lhs, const FString& rhs) { return !lhs.Contains(rhs); })));
+
+		TSharedPtr<FFilterWithSuggestions> TrackFilter = MakeShared<FFilterWithSuggestions>(static_cast<int32>(EFilterField::TrackName), LOCTEXT("Track", "Track"), LOCTEXT("Track", "Track"), EFilterDataType::String, TrackNameFilterOperators);
+		TrackFilter->Callback = [this](const FString& Text, TArray<FString>& OutSuggestions)
+		{
+			this->PopulateTrackSuggestionList(Text, OutSuggestions);
+		};
+
+		AvailableFilters->Add(TrackFilter);
+
 		QuickFindVm = MakeShared<FQuickFind>(FilterConfigurator);
 		QuickFindVm->GetOnFindNextEvent().AddSP(this, &STimingView::FindNextEvent);
 		QuickFindVm->GetOnFindPreviousEvent().AddSP(this, &STimingView::FindPrevEvent);
@@ -4631,16 +4645,16 @@ void STimingView::FindNextEvent()
 	FTimingEventSearchParameters Params(StartTime, std::numeric_limits<double>::max(), ETimingEventSearchFlags::StopAtFirstMatch, EventFilter);
 	Params.FilterExecutor = QuickFindVm->GetFilterConfigurator();
 
-	for (auto& Entry : AllTracks)
+	EnumerateFilteredTracks(QuickFindVm->GetFilterConfigurator(), [&Params, &BestMatchEvent](TSharedPtr<FBaseTimingTrack>& Track)
 	{
-		if (!Entry.Value->IsVisible())
+		if (!Track->IsVisible())
 		{
-			continue;
+			return;
 		}
 
 		Params.EndTime = BestMatchEvent.IsValid() ? BestMatchEvent->GetStartTime() : std::numeric_limits<double>::max();
 
-		TSharedPtr<const ITimingEvent> FoundEvent = Entry.Value->SearchEvent(Params);
+		TSharedPtr<const ITimingEvent> FoundEvent = Track->SearchEvent(Params);
 		if (FoundEvent.IsValid())
 		{
 			if (BestMatchEvent.IsValid())
@@ -4649,7 +4663,7 @@ void STimingView::FindNextEvent()
 			}
 			BestMatchEvent = FoundEvent;
 		}
-	}
+	});
 
 	if (BestMatchEvent)
 	{
@@ -4674,16 +4688,16 @@ void STimingView::FindPrevEvent()
 	Params.FilterExecutor = QuickFindVm->GetFilterConfigurator();
 	Params.SearchDirection = FTimingEventSearchParameters::ESearchDirection::Backward;
 
-	for (auto& Entry : AllTracks)
+	EnumerateFilteredTracks(QuickFindVm->GetFilterConfigurator(), [&Params, &BestMatchEvent](TSharedPtr<FBaseTimingTrack>& Track)
 	{
-		if (!Entry.Value->IsVisible())
+		if (!Track->IsVisible())
 		{
-			continue;
+			return;
 		}
 
 		Params.StartTime = BestMatchEvent.IsValid() ? BestMatchEvent->GetStartTime() : std::numeric_limits<double>::lowest();
 
-		TSharedPtr<const ITimingEvent> FoundEvent = Entry.Value->SearchEvent(Params);
+		TSharedPtr<const ITimingEvent> FoundEvent = Track->SearchEvent(Params);
 		if (FoundEvent.IsValid())
 		{
 			if (!BestMatchEvent.IsValid() || BestMatchEvent->GetStartTime() < FoundEvent->GetStartTime())
@@ -4691,7 +4705,7 @@ void STimingView::FindPrevEvent()
 				BestMatchEvent = FoundEvent;
 			}
 		}
-	}
+	});
 
 	if (BestMatchEvent)
 	{
@@ -4723,5 +4737,34 @@ void STimingView::ClearFilters()
 
 ////////////////////////////////////////////////////////////////////////////////////////////////////
 
+void STimingView::PopulateTrackSuggestionList(const FString& Text, TArray<FString>& OutSuggestions)
+{
+	for (auto& Entry : AllTracks)
+	{
+		if (Entry.Value->GetName().Contains(Text))
+		{
+			OutSuggestions.Add(Entry.Value->GetName());
+		}
+	}
+}
+
+////////////////////////////////////////////////////////////////////////////////////////////////////
+
+void STimingView::EnumerateFilteredTracks(TSharedPtr<Insights::FFilterConfigurator> FilterConfigurator, EnumerateFilteredTracksCallback Callback)
+{
+	Insights::FFilterContext FilterContext;
+	FilterContext.AddFilterData(static_cast<int32>(EFilterField::TrackName), FString());
+	for (auto& Entry : AllTracks)
+	{
+		FilterContext.SetFilterData(static_cast<int32>(EFilterField::TrackName), Entry.Value->GetName());
+		if (FilterConfigurator->ApplyFilters(FilterContext))
+		{
+			Callback(Entry.Value);
+		}
+	}
+}
+
+////////////////////////////////////////////////////////////////////////////////////////////////////
+
 #undef ACTIVATE_BENCHMARK
 #undef LOCTEXT_NAMESPACE
diff --git a/Engine/Source/Developer/TraceInsights/Private/Insights/Widgets/STimingView.h b/Engine/Source/Developer/TraceInsights/Private/Insights/Widgets/STimingView.h
--- a/Engine/Source/Developer/TraceInsights/Private/Insights/Widgets/STimingView.h
+++ b/Engine/Source/Developer/TraceInsights/Private/Insights/Widgets/STimingView.h
@@ -463,6 +463,10 @@ protected:
 	void ClearFilters();
 
 	TSharedRef<SDockTab> SpawnQuickFindTab(const FSpawnTabArgs& Args);
+	void PopulateTrackSuggestionList(const FString& Text, TArray<FString>& OutSuggestions);
+
+	typedef TFunctionRef<void(TSharedPtr<FBaseTimingTrack>& Track)> EnumerateFilteredTracksCallback;
+	void EnumerateFilteredTracks(TSharedPtr<Insights::FFilterConfigurator> FilterConfigurator, EnumerateFilteredTracksCallback Callback);
 
 protected:
 	/** The track's viewport. Encapsulates info about position and scale. */
