[Insights] Bring into view the track containing an event found using the Quick Find Widget

diff --git a/Engine/Source/Developer/TraceInsights/Private/Insights/Widgets/STimingView.cpp b/Engine/Source/Developer/TraceInsights/Private/Insights/Widgets/STimingView.cpp
--- a/Engine/Source/Developer/TraceInsights/Private/Insights/Widgets/STimingView.cpp
+++ b/Engine/Source/Developer/TraceInsights/Private/Insights/Widgets/STimingView.cpp
@@ -4669,6 +4669,11 @@ void STimingView::FindNextEvent()
 	{
 		SelectedEvent = BestMatchEvent;
 		BringIntoView(SelectedEvent->GetStartTime(), SelectedEvent->GetEndTime());
+		if (SelectedEvent->GetTrack()->GetLocation() == ETimingTrackLocation::Scrollable)
+		{
+			BringScrollableTrackIntoView(*SelectedEvent->GetTrack());
+		}
+
 		OnSelectedTimingEventChanged();
 	}
 }
@@ -4711,6 +4716,11 @@ void STimingView::FindPrevEvent()
 	{
 		SelectedEvent = BestMatchEvent;
 		BringIntoView(SelectedEvent->GetStartTime(), SelectedEvent->GetEndTime());
+		if (SelectedEvent->GetTrack()->GetLocation() == ETimingTrackLocation::Scrollable)
+		{
+			BringScrollableTrackIntoView(*SelectedEvent->GetTrack());
+		}
+
 		OnSelectedTimingEventChanged();
 	}
 }
