[Insights] UI tuning for Trace Store page (made the "trace store dir" box to fill the width of the window; removed icon for the Open Trace menu dropdown.

diff --git a/Engine/Source/Developer/TraceInsights/Private/Insights/Widgets/SStartPageWindow.cpp b/Engine/Source/Developer/TraceInsights/Private/Insights/Widgets/SStartPageWindow.cpp
--- a/Engine/Source/Developer/TraceInsights/Private/Insights/Widgets/SStartPageWindow.cpp
+++ b/Engine/Source/Developer/TraceInsights/Private/Insights/Widgets/SStartPageWindow.cpp
@@ -859,20 +859,10 @@ TSharedRef<SWidget> SStartPageWindow::ConstructLoadPanel()
 		.IsEnabled(this, &SStartPageWindow::Open_IsEnabled)
 		.OnClicked(this, &SStartPageWindow::Open_OnClicked)
 		.ToolTipText(LOCTEXT("OpenButtonTooltip", "Start analysis for selected trace session."))
-		.ContentPadding(FMargin(0.0f, 0.0f))
+		.ContentPadding(FMargin(0.0f, 0.0f, 0.0f, 0.0f))
 		.Content()
 		[
-			SNew(SHorizontalBox)
-			//+ SHorizontalBox::Slot()
-			//.AutoWidth()
-			//.VAlign(VAlign_Center)
-			//[
-			//	SNew(SImage)
-			//	.VAlign(VAlign_Center)
-			//	.Image(FInsightsStyle::GetBrush("Open.Icon.Small"))
-			//]
-			+ SHorizontalBox::Slot()
-			.AutoWidth()
+			SNew(SBox)
 			.VAlign(VAlign_Center)
 			[
 				SNew(STextBlock)
@@ -891,16 +881,11 @@ TSharedRef<SWidget> SStartPageWindow::ConstructLoadPanel()
 		.ToolTipText(LOCTEXT("MRU_Tooltip", "Open a trace file or choose a trace session."))
 		.OnGetMenuContent(this, &SStartPageWindow::MakeTraceListMenu)
 		.HasDownArrow(true)
-		.ContentPadding(FMargin(0.0f, 0.0f, 0.0f, 0.0f))
-		.ButtonContent()
-		[
-			SNew(SImage)
-			.Image(FInsightsStyle::GetBrush("Menu.Icon.Small"))
-		]
 	];
 
 	return Widget;
 }
+
 ////////////////////////////////////////////////////////////////////////////////////////////////////
 
 TSharedRef<SWidget> SStartPageWindow::ConstructTraceStoreDirectoryPanel()
@@ -919,7 +904,7 @@ TSharedRef<SWidget> SStartPageWindow::ConstructTraceStoreDirectoryPanel()
 		]
 
 		+ SHorizontalBox::Slot()
-		.AutoWidth()
+		.FillWidth(1.0f)
 		.Padding(4.0f, 0.0f, 0.0f, 0.0f)
 		.VAlign(VAlign_Center)
 		[
@@ -939,23 +924,9 @@ TSharedRef<SWidget> SStartPageWindow::ConstructTraceStoreDirectoryPanel()
 			.ToolTipText(LOCTEXT("ExploreTraceStoreDirButtonToolTip", "Explore the Trace Store Directory"))
 			.OnClicked(this, &SStartPageWindow::ExploreTraceStoreDirectory_OnClicked)
 			[
-				//SNew(SHorizontalBox)
-				//+ SHorizontalBox::Slot()
-				//.AutoWidth()
-				//[
-					SNew(SImage)
-					.Image(FInsightsStyle::Get().GetBrush("FolderExplore.Icon.Small"))
-					.ColorAndOpacity(FSlateColor::UseForeground())
-				//]
-				//+ SHorizontalBox::Slot()
-				//.AutoWidth()
-				//.Padding(2.0f, 0.0f, 0.0f, 0.0f)
-				//[
-				//	SNew(STextBlock)
-				//	.TextStyle(&FAppStyle::Get().GetWidgetStyle<FTextBlockStyle>("ButtonText"))
-				//	.Text(LOCTEXT("ExploreTraceStoreDirButton", "Explore"))
-				//	.ColorAndOpacity(FSlateColor::UseForeground())
-				//]
+				SNew(SImage)
+				.Image(FInsightsStyle::Get().GetBrush("FolderExplore.Icon.Small"))
+				.ColorAndOpacity(FSlateColor::UseForeground())
 			]
 		];
 
