[Insights] Log View: Fixed the height of the toolbar.

diff --git a/Engine/Source/Developer/TraceInsights/Private/Insights/Widgets/SLogView.cpp b/Engine/Source/Developer/TraceInsights/Private/Insights/Widgets/SLogView.cpp
--- a/Engine/Source/Developer/TraceInsights/Private/Insights/Widgets/SLogView.cpp
+++ b/Engine/Source/Developer/TraceInsights/Private/Insights/Widgets/SLogView.cpp
@@ -580,9 +580,8 @@ void SLogView::Construct(const FArguments& InArgs)
 		+ SVerticalBox::Slot()
 		.AutoHeight()
 		[
-			SNew(SBorder)
-			.BorderImage(FCoreStyle::Get().GetBrush("ToolPanel.GroupBorder"))
-			.Padding(4.0f)
+			SNew(SBox)
+			.Padding(FMargin(2.0f, 0.0f, 0.0f, 0.0f))
 			[
 				ToolbarBuilder.MakeWidget()
 			]
