[Insights] Context Switches: Made the option to turn on/off the visibility of extended lines to be independent of the overlay areas toggle.

diff --git a/Engine/Source/Developer/TraceInsights/Private/Insights/ContextSwitches/ViewModels/ContextSwitchesTimingTrack.cpp b/Engine/Source/Developer/TraceInsights/Private/Insights/ContextSwitches/ViewModels/ContextSwitchesTimingTrack.cpp
--- a/Engine/Source/Developer/TraceInsights/Private/Insights/ContextSwitches/ViewModels/ContextSwitchesTimingTrack.cpp
+++ b/Engine/Source/Developer/TraceInsights/Private/Insights/ContextSwitches/ViewModels/ContextSwitchesTimingTrack.cpp
@@ -95,7 +95,7 @@ void FContextSwitchesTimingTrack::DrawLineEvents(const ITimingTrackDrawContext&
 
 void FContextSwitchesTimingTrack::PostDraw(const ITimingTrackDrawContext& Context) const
 {
-	if (SharedState.AreOverlaysVisible())
+	if (SharedState.AreOverlaysVisible() || SharedState.AreExtendedLinesVisible())
 	{
 		float LineY1 = 0.0f;
 		float LineY2 = 0.0f;
@@ -165,7 +165,7 @@ void FContextSwitchesTimingTrack::PostDraw(const ITimingTrackDrawContext& Contex
 		if (LineH > 0.0f)
 		{
 			const FTimingViewDrawHelper& Helper = *static_cast<const FTimingViewDrawHelper*>(&Context.GetHelper());
-			Helper.DrawContextSwitchMarkers(GetDrawState(), LineY1, LineH, 0.25f, SharedState.AreExtendedLinesVisible());
+			Helper.DrawContextSwitchMarkers(GetDrawState(), LineY1, LineH, 0.25f, SharedState.AreOverlaysVisible(), SharedState.AreExtendedLinesVisible());
 		}
 	}
 }
diff --git a/Engine/Source/Developer/TraceInsights/Private/Insights/ViewModels/TimingViewDrawHelper.cpp b/Engine/Source/Developer/TraceInsights/Private/Insights/ViewModels/TimingViewDrawHelper.cpp
--- a/Engine/Source/Developer/TraceInsights/Private/Insights/ViewModels/TimingViewDrawHelper.cpp
+++ b/Engine/Source/Developer/TraceInsights/Private/Insights/ViewModels/TimingViewDrawHelper.cpp
@@ -580,13 +580,13 @@ void FTimingViewDrawHelper::DrawFadedLineEvents(const FTimingEventsTrackDrawStat
 
 ////////////////////////////////////////////////////////////////////////////////////////////////////
 
-void FTimingViewDrawHelper::DrawContextSwitchMarkers(const FTimingEventsTrackDrawState& DrawState, float LineY, float LineH, float Opacity, bool bShowVerticalLines) const
+void FTimingViewDrawHelper::DrawContextSwitchMarkers(const FTimingEventsTrackDrawState& DrawState, float LineY, float LineH, float Opacity, bool bDrawOverlays, bool bDrawVerticalLines) const
 {
 	if (LineH > 0.0f)
 	{
 		const int32 LayerId = ReservedLayerId + ToInt32(EDrawLayer::EventHighlight);
 
-		if (bShowVerticalLines)
+		if (bDrawVerticalLines)
 		{
 			// Draw vertical lines (merged into large boxes).
 			for (const FTimingEventsTrackDrawState::FBoxPrimitive& Box : DrawState.Boxes)
@@ -604,43 +604,46 @@ void FTimingViewDrawHelper::DrawContextSwitchMarkers(const FTimingEventsTrackDra
 			}
 		}
 
-		// Draw overlay for idle areas.
-		const FSlateBrush* IdleBrush = WhiteBrush;
-		const FLinearColor IdleColor(0.0f, 0.0f, 0.0f, 0.3f);
-		const int32 Count1 = DrawState.Boxes.Num();
-		const int32 Count2 = DrawState.Borders.Num();
-		int32 Index1 = 0;
-		int32 Index2 = 0;
-		const float MinX = Viewport.TimeToSlateUnitsRounded(Viewport.GetMinValidTime());
-		float CurrentX = FMath::Max(MinX, 0.0f);
-		while (Index1 < Count1 || Index2 < Count2)
+		if (bDrawOverlays)
 		{
-			const float X1 = (Index1 < Count1) ? DrawState.Boxes[Index1].X : FLT_MAX;
-			const float X2 = (Index2 < Count2) ? DrawState.Borders[Index2].X : FLT_MAX;
-			if (X1 < X2)
+			// Draw overlay for idle areas.
+			const FSlateBrush* IdleBrush = WhiteBrush;
+			const FLinearColor IdleColor(0.0f, 0.0f, 0.0f, 0.3f);
+			const int32 Count1 = DrawState.Boxes.Num();
+			const int32 Count2 = DrawState.Borders.Num();
+			int32 Index1 = 0;
+			int32 Index2 = 0;
+			const float MinX = Viewport.TimeToSlateUnitsRounded(Viewport.GetMinValidTime());
+			float CurrentX = FMath::Max(MinX, 0.0f);
+			while (Index1 < Count1 || Index2 < Count2)
 			{
-				if (X1 - CurrentX > 0.0f)
+				const float X1 = (Index1 < Count1) ? DrawState.Boxes[Index1].X : FLT_MAX;
+				const float X2 = (Index2 < Count2) ? DrawState.Borders[Index2].X : FLT_MAX;
+				if (X1 < X2)
 				{
-					DrawContext.DrawBox(LayerId, CurrentX, LineY, X1 - CurrentX, LineH, IdleBrush, IdleColor);
+					if (X1 - CurrentX > 0.0f)
+					{
+						DrawContext.DrawBox(LayerId, CurrentX, LineY, X1 - CurrentX, LineH, IdleBrush, IdleColor);
+					}
+					CurrentX = FMath::Max(CurrentX, X1 + DrawState.Boxes[Index1].W);
+					++Index1;
 				}
-				CurrentX = FMath::Max(CurrentX, X1 + DrawState.Boxes[Index1].W);
-				++Index1;
-			}
-			else
-			{
-				if (X2 - CurrentX > 0.0f)
+				else
 				{
-					DrawContext.DrawBox(LayerId, CurrentX, LineY, X2 - CurrentX, LineH, IdleBrush, IdleColor);
+					if (X2 - CurrentX > 0.0f)
+					{
+						DrawContext.DrawBox(LayerId, CurrentX, LineY, X2 - CurrentX, LineH, IdleBrush, IdleColor);
+					}
+					CurrentX = FMath::Max(CurrentX, X2 + DrawState.Borders[Index2].W);
+					++Index2;
 				}
-				CurrentX = FMath::Max(CurrentX, X2 + DrawState.Borders[Index2].W);
-				++Index2;
 			}
-		}
-		const float MaxX = Viewport.TimeToSlateUnitsRounded(Viewport.GetMaxValidTime());
-		const float LastAreaW = FMath::Min(MaxX, Viewport.GetWidth()) - CurrentX;
-		if (LastAreaW > 0.0f)
-		{
-			DrawContext.DrawBox(LayerId, CurrentX, LineY, LastAreaW, LineH, IdleBrush, IdleColor);
+			const float MaxX = Viewport.TimeToSlateUnitsRounded(Viewport.GetMaxValidTime());
+			const float LastAreaW = FMath::Min(MaxX, Viewport.GetWidth()) - CurrentX;
+			if (LastAreaW > 0.0f)
+			{
+				DrawContext.DrawBox(LayerId, CurrentX, LineY, LastAreaW, LineH, IdleBrush, IdleColor);
+			}
 		}
 	}
 }
diff --git a/Engine/Source/Developer/TraceInsights/Private/Insights/ViewModels/TimingViewDrawHelper.h b/Engine/Source/Developer/TraceInsights/Private/Insights/ViewModels/TimingViewDrawHelper.h
--- a/Engine/Source/Developer/TraceInsights/Private/Insights/ViewModels/TimingViewDrawHelper.h
+++ b/Engine/Source/Developer/TraceInsights/Private/Insights/ViewModels/TimingViewDrawHelper.h
@@ -193,7 +193,7 @@ public:
 
 	void DrawLineEvents(const FTimingEventsTrackDrawState& DrawState, const FTimingEventsTrack& Track, const float OffsetY = 1.0f) const;
 	void DrawFadedLineEvents(const FTimingEventsTrackDrawState& DrawState, const FTimingEventsTrack& Track, const float OffsetY = 1.0f, const float Opacity = 0.1f) const;
-	void DrawContextSwitchMarkers(const FTimingEventsTrackDrawState& DrawState, float LineY, float LineH, float Opacity, bool bShowVerticalLines) const;
+	void DrawContextSwitchMarkers(const FTimingEventsTrackDrawState& DrawState, float LineY, float LineH, float Opacity, bool bDrawOverlays, bool bDrawVerticalLines) const;
 
 	void DrawMarkers(const FTimingEventsTrackDrawState& DrawState, float LineY, float LineH, float Opacity) const;
 
