[Insights] Timing view: Fixed the toggle between Bookmarks and Logs for the Markers track to apply immediately when changed from the context menu.

diff --git a/Engine/Source/Developer/TraceInsights/Private/Insights/ViewModels/MarkersTimingTrack.h b/Engine/Source/Developer/TraceInsights/Private/Insights/ViewModels/MarkersTimingTrack.h
--- a/Engine/Source/Developer/TraceInsights/Private/Insights/ViewModels/MarkersTimingTrack.h
+++ b/Engine/Source/Developer/TraceInsights/Private/Insights/ViewModels/MarkersTimingTrack.h
@@ -65,8 +65,8 @@ public:
 		bUseOnlyBookmarks = bInUseOnlyBookmarks;
 		UpdateTrackNameAndHeight();
 	}
-	void SetBookmarksTrack() { SetBookmarksTrackFlag(true); }
-	void SetLogsTrack() { SetBookmarksTrackFlag(false); }
+	void SetBookmarksTrack() { SetBookmarksTrackFlag(true); SetDirtyFlag(); }
+	void SetLogsTrack() { SetBookmarksTrackFlag(false); SetDirtyFlag(); }
 
 	// Stats
 	int32 GetNumLogMessages() const { return NumLogMessages; }
