[Insights]
  - Timing View: Added extension hooks for the CPU/GPU, Other and Plugins menus.
  - Timing View: Split the options from the previous "Tracks" menu into CPU/GPU, Other and Plugins menus. Some submenus are now inlined (ex. Docked and Scrollable track filter, File Activity, Context Switches, etc.).
  - Timing View: Fixed position of the Frame tracks (to take into consideration the PosY, i.e. height of top toolbar).
  - Context Switches: Enabled the Core Tracks by default.

diff --git a/Engine/Source/Developer/TraceInsights/Private/Insights/ContextSwitches/ViewModels/ContextSwitchesSharedState.cpp b/Engine/Source/Developer/TraceInsights/Private/Insights/ContextSwitches/ViewModels/ContextSwitchesSharedState.cpp
--- a/Engine/Source/Developer/TraceInsights/Private/Insights/ContextSwitches/ViewModels/ContextSwitchesSharedState.cpp
+++ b/Engine/Source/Developer/TraceInsights/Private/Insights/ContextSwitches/ViewModels/ContextSwitchesSharedState.cpp
@@ -108,7 +108,7 @@ FContextSwitchesSharedState::FContextSwitchesSharedState(STimingView* InTimingVi
 	: TimingView(InTimingView)
 	, ThreadsSerial(0)
 	, CpuCoresSerial(0)
-	, bAreCoreTracksVisible(false)
+	, bAreCoreTracksVisible(true)
 	, bAreContextSwitchesVisible(true)
 	, bAreOverlaysVisible(true)
 	, bAreExtendedLinesVisible(true)
@@ -146,7 +146,7 @@ void FContextSwitchesSharedState::OnBeginSession(Insights::ITimingViewSession& I
 	ThreadsSerial = 0;
 	CpuCoresSerial = 0;
 
-	bAreCoreTracksVisible = false;
+	bAreCoreTracksVisible = true;
 	bAreContextSwitchesVisible = true;
 	bAreOverlaysVisible = true;
 	bAreExtendedLinesVisible = true;
@@ -166,7 +166,7 @@ void FContextSwitchesSharedState::OnEndSession(Insights::ITimingViewSession& InS
 	ThreadsSerial = 0;
 	CpuCoresSerial = 0;
 
-	bAreCoreTracksVisible = false;
+	bAreCoreTracksVisible = true;
 	bAreContextSwitchesVisible = true;
 	bAreOverlaysVisible = true;
 	bAreExtendedLinesVisible = true;
@@ -206,24 +206,26 @@ void FContextSwitchesSharedState::Tick(Insights::ITimingViewSession& InSession,
 
 ////////////////////////////////////////////////////////////////////////////////////////////////////
 
-void FContextSwitchesSharedState::ExtendFilterMenu(Insights::ITimingViewSession& InSession, FMenuBuilder& InMenuBuilder)
+void FContextSwitchesSharedState::ExtendCpuTracksFilterMenu(Insights::ITimingViewSession& InSession, FMenuBuilder& InMenuBuilder)
 {
 	if (&InSession != TimingView)
 	{
 		return;
 	}
 
-	InMenuBuilder.BeginSection("ContextSwitches");
-	{
-		InMenuBuilder.AddSubMenu(
-			LOCTEXT("ContextSwitches_SubMenu", "Context Switches"),
-			LOCTEXT("ContextSwitches_SubMenu_Desc", "Context Switch track options"),
-			FNewMenuDelegate::CreateSP(this, &FContextSwitchesSharedState::BuildSubMenu),
-			false,
-			FSlateIcon()
-		);
-	}
-	InMenuBuilder.EndSection();
+	BuildSubMenu(InMenuBuilder);
+
+	//InMenuBuilder.BeginSection("ContextSwitches");
+	//{
+	//	InMenuBuilder.AddSubMenu(
+	//		LOCTEXT("ContextSwitches_SubMenu", "Context Switches"),
+	//		LOCTEXT("ContextSwitches_SubMenu_Desc", "Context Switch track options"),
+	//		FNewMenuDelegate::CreateSP(this, &FContextSwitchesSharedState::BuildSubMenu),
+	//		false,
+	//		FSlateIcon()
+	//	);
+	//}
+	//InMenuBuilder.EndSection();
 }
 
 ////////////////////////////////////////////////////////////////////////////////////////////////////
diff --git a/Engine/Source/Developer/TraceInsights/Private/Insights/ContextSwitches/ViewModels/ContextSwitchesSharedState.h b/Engine/Source/Developer/TraceInsights/Private/Insights/ContextSwitches/ViewModels/ContextSwitchesSharedState.h
--- a/Engine/Source/Developer/TraceInsights/Private/Insights/ContextSwitches/ViewModels/ContextSwitchesSharedState.h
+++ b/Engine/Source/Developer/TraceInsights/Private/Insights/ContextSwitches/ViewModels/ContextSwitchesSharedState.h
@@ -66,7 +66,7 @@ public:
 	virtual void OnBeginSession(Insights::ITimingViewSession& InSession) override;
 	virtual void OnEndSession(Insights::ITimingViewSession& InSession) override;
 	virtual void Tick(Insights::ITimingViewSession& InSession, const TraceServices::IAnalysisSession& InAnalysisSession) override;
-	virtual void ExtendFilterMenu(Insights::ITimingViewSession& InSession, FMenuBuilder& InMenuBuilder) override;
+	virtual void ExtendCpuTracksFilterMenu(Insights::ITimingViewSession& InSession, FMenuBuilder& InMenuBuilder) override;
 
 	//////////////////////////////////////////////////
 
@@ -131,7 +131,7 @@ private:
 	bool Command_ShowOverlays_IsChecked() const { return AreOverlaysVisible(); }
 
 	void Command_ShowExtendedLines_Execute() { ToggleExtendedLines(); }
-	bool Command_ShowExtendedLines_CanExecute() const { return AreContextSwitchesAvailable() && AreContextSwitchesVisible() && AreOverlaysVisible(); }
+	bool Command_ShowExtendedLines_CanExecute() const { return AreContextSwitchesAvailable() && AreContextSwitchesVisible(); }
 	bool Command_ShowExtendedLines_IsChecked() const { return AreExtendedLinesVisible(); }
 
 	bool IsValidCpuCoreEventSelected() const;
diff --git a/Engine/Source/Developer/TraceInsights/Private/Insights/ViewModels/FileActivityTimingTrack.cpp b/Engine/Source/Developer/TraceInsights/Private/Insights/ViewModels/FileActivityTimingTrack.cpp
--- a/Engine/Source/Developer/TraceInsights/Private/Insights/ViewModels/FileActivityTimingTrack.cpp
+++ b/Engine/Source/Developer/TraceInsights/Private/Insights/ViewModels/FileActivityTimingTrack.cpp
@@ -342,39 +342,26 @@ void FFileActivitySharedState::Tick(Insights::ITimingViewSession& InSession, con
 
 ////////////////////////////////////////////////////////////////////////////////////////////////////
 
-void FFileActivitySharedState::ExtendFilterMenu(Insights::ITimingViewSession& InSession, FMenuBuilder& InOutMenuBuilder)
+void FFileActivitySharedState::ExtendOtherTracksFilterMenu(Insights::ITimingViewSession& InSession, FMenuBuilder& InOutMenuBuilder)
 {
 	if (&InSession != TimingView)
 	{
 		return;
 	}
 
-	//InOutMenuBuilder.BeginSection("File Activity", LOCTEXT("FileActivityHeading", "File Activity"));
+	BuildSubMenu(InOutMenuBuilder);
+
+	//InOutMenuBuilder.BeginSection("File Activity");
 	//{
-	//	InOutMenuBuilder.AddMenuEntry(
-	//		LOCTEXT("ShowAllIoTracks", "I/O Tracks - I"),
-	//		LOCTEXT("ShowAllIoTracks_Tooltip", "Show/hide the I/O (File Activity) tracks"),
-	//		FSlateIcon(),
-	//		FUIAction(FExecuteAction::CreateSP(this, &FFileActivitySharedState::ShowHideAllIoTracks),
-	//				  FCanExecuteAction(),
-	//				  FIsActionChecked::CreateSP(this, &FFileActivitySharedState::IsAllIoTracksToggleOn)),
-	//		NAME_None, //"QuickFilterSeparator",
-	//		EUserInterfaceActionType::ToggleButton
+	//	InOutMenuBuilder.AddSubMenu(
+	//		LOCTEXT("FileActivity_SubMenu", "File Activity"),
+	//		LOCTEXT("FileActivity_SubMenu_Desc", "File Activity track options"),
+	//		FNewMenuDelegate::CreateSP(this, &FFileActivitySharedState::BuildSubMenu),
+	//		false,
+	//		FSlateIcon()
 	//	);
 	//}
 	//InOutMenuBuilder.EndSection();
-
-	InOutMenuBuilder.BeginSection("File Activity");
-	{
-		InOutMenuBuilder.AddSubMenu(
-			LOCTEXT("FileActivity_SubMenu", "File Activity"),
-			LOCTEXT("FileActivity_SubMenu_Desc", "File Activity track options"),
-			FNewMenuDelegate::CreateSP(this, &FFileActivitySharedState::BuildSubMenu),
-			false,
-			FSlateIcon()
-		);
-	}
-	InOutMenuBuilder.EndSection();
 }
 
 ////////////////////////////////////////////////////////////////////////////////////////////////////
@@ -385,7 +372,7 @@ void FFileActivitySharedState::BuildSubMenu(FMenuBuilder& InOutMenuBuilder)
 	{
 		InOutMenuBuilder.AddMenuEntry(
 			LOCTEXT("ShowIoOverviewTrack", "I/O Overview Track - I"),
-			LOCTEXT("ShowIoOverviewTrack_Tooltip", "Show/hide the I/O Overview track"),
+			LOCTEXT("ShowIoOverviewTrack_Tooltip", "Show/hide the I/O Overview track."),
 			FSlateIcon(),
 			FUIAction(FExecuteAction::CreateSP(this, &FFileActivitySharedState::ShowHideIoOverviewTrack),
 				FCanExecuteAction(),
@@ -407,7 +394,7 @@ void FFileActivitySharedState::BuildSubMenu(FMenuBuilder& InOutMenuBuilder)
 
 		InOutMenuBuilder.AddMenuEntry(
 			LOCTEXT("ShowIoActivityTrack", "I/O Activity Track - I"),
-			LOCTEXT("ShowIoActivityTrack_Tooltip", "Show/hide the I/O Activity track"),
+			LOCTEXT("ShowIoActivityTrack_Tooltip", "Show/hide the I/O Activity track."),
 			FSlateIcon(),
 			FUIAction(FExecuteAction::CreateSP(this, &FFileActivitySharedState::ShowHideIoActivityTrack),
 				FCanExecuteAction(),
diff --git a/Engine/Source/Developer/TraceInsights/Private/Insights/ViewModels/FileActivityTimingTrack.h b/Engine/Source/Developer/TraceInsights/Private/Insights/ViewModels/FileActivityTimingTrack.h
--- a/Engine/Source/Developer/TraceInsights/Private/Insights/ViewModels/FileActivityTimingTrack.h
+++ b/Engine/Source/Developer/TraceInsights/Private/Insights/ViewModels/FileActivityTimingTrack.h
@@ -59,7 +59,7 @@ public:
 	virtual void OnBeginSession(Insights::ITimingViewSession& InSession) override;
 	virtual void OnEndSession(Insights::ITimingViewSession& InSession) override;
 	virtual void Tick(Insights::ITimingViewSession& InSession, const TraceServices::IAnalysisSession& InAnalysisSession) override;
-	virtual void ExtendFilterMenu(Insights::ITimingViewSession& InSession, FMenuBuilder& InOutMenuBuilder) override;
+	virtual void ExtendOtherTracksFilterMenu(Insights::ITimingViewSession& InSession, FMenuBuilder& InOutMenuBuilder) override;
 
 	const TArray<FIoTimingEvent>& GetAllEvents() const { return AllIoEvents; }
 
diff --git a/Engine/Source/Developer/TraceInsights/Private/Insights/ViewModels/FrameTimingTrack.cpp b/Engine/Source/Developer/TraceInsights/Private/Insights/ViewModels/FrameTimingTrack.cpp
--- a/Engine/Source/Developer/TraceInsights/Private/Insights/ViewModels/FrameTimingTrack.cpp
+++ b/Engine/Source/Developer/TraceInsights/Private/Insights/ViewModels/FrameTimingTrack.cpp
@@ -116,7 +116,7 @@ void FFrameSharedState::Tick(Insights::ITimingViewSession& InSession, const Trac
 
 ////////////////////////////////////////////////////////////////////////////////////////////////////
 
-void FFrameSharedState::ExtendFilterMenu(Insights::ITimingViewSession& InSession, FMenuBuilder& InOutMenuBuilder)
+void FFrameSharedState::ExtendOtherTracksFilterMenu(Insights::ITimingViewSession& InSession, FMenuBuilder& InOutMenuBuilder)
 {
 	if (&InSession != TimingView)
 	{
@@ -128,7 +128,7 @@ void FFrameSharedState::ExtendFilterMenu(Insights::ITimingViewSession& InSession
 		//TODO: MenuBuilder.AddMenuEntry(Commands.ShowAllFrameTracks);
 		InOutMenuBuilder.AddMenuEntry(
 			LOCTEXT("ShowAllFrameTracks", "Frame Tracks - R"),
-			LOCTEXT("ShowAllFrameTracks_Tooltip", "Show/hide all Frame tracks"),
+			LOCTEXT("ShowAllFrameTracks_Tooltip", "Show/hide all Frame tracks."),
 			FSlateIcon(),
 			FUIAction(FExecuteAction::CreateSP(this, &FFrameSharedState::ShowHideAllFrameTracks),
 					  FCanExecuteAction(),
@@ -282,7 +282,7 @@ void FFrameTimingTrack::PostDraw(const ITimingTrackDrawContext& Context) const
 	if (!Header.IsCollapsed())
 	{
 		const FTimingTrackViewport& Viewport = Context.GetViewport();
-		DrawMarkers(Context, 0.0f, Viewport.GetHeight());
+		DrawMarkers(Context, Viewport.GetPosY(), Viewport.GetHeight());
 	}
 
 	const TSharedPtr<const ITimingEvent> SelectedEventPtr = Context.GetSelectedEvent();
@@ -311,7 +311,7 @@ void FFrameTimingTrack::DrawSelectedEventInfo(const FTimingEvent& SelectedEvent,
 		const TSharedRef<FSlateFontMeasure> FontMeasureService = FSlateApplication::Get().GetRenderer()->GetFontMeasureService();
 		const FVector2D Size = FontMeasureService->Measure(Str, Font);
 		const float X = Viewport.GetWidth() - Size.X - 23.0f;
-		const float Y = Viewport.GetHeight() - Size.Y - 18.0f;
+		const float Y = Viewport.GetPosY() + Viewport.GetHeight() - Size.Y - 18.0f;
 
 		const FLinearColor BackgroundColor(0.05f, 0.05f, 0.05f, 1.0f);
 		const FLinearColor TextColor(0.7f, 0.7f, 0.7f, 1.0f);
diff --git a/Engine/Source/Developer/TraceInsights/Private/Insights/ViewModels/FrameTimingTrack.h b/Engine/Source/Developer/TraceInsights/Private/Insights/ViewModels/FrameTimingTrack.h
--- a/Engine/Source/Developer/TraceInsights/Private/Insights/ViewModels/FrameTimingTrack.h
+++ b/Engine/Source/Developer/TraceInsights/Private/Insights/ViewModels/FrameTimingTrack.h
@@ -36,7 +36,7 @@ public:
 	virtual void OnBeginSession(Insights::ITimingViewSession& InSession) override;
 	virtual void OnEndSession(Insights::ITimingViewSession& InSession) override;
 	virtual void Tick(Insights::ITimingViewSession& InSession, const TraceServices::IAnalysisSession& InAnalysisSession) override;
-	virtual void ExtendFilterMenu(Insights::ITimingViewSession& InSession, FMenuBuilder& InMenuBuilder) override;
+	virtual void ExtendOtherTracksFilterMenu(Insights::ITimingViewSession& InSession, FMenuBuilder& InMenuBuilder) override;
 
 	//////////////////////////////////////////////////
 
diff --git a/Engine/Source/Developer/TraceInsights/Private/Insights/ViewModels/GraphTrack.cpp b/Engine/Source/Developer/TraceInsights/Private/Insights/ViewModels/GraphTrack.cpp
--- a/Engine/Source/Developer/TraceInsights/Private/Insights/ViewModels/GraphTrack.cpp
+++ b/Engine/Source/Developer/TraceInsights/Private/Insights/ViewModels/GraphTrack.cpp
@@ -651,7 +651,7 @@ const TSharedPtr<const ITimingEvent> FGraphTrack::GetEvent(float InPosX, float I
 
 void FGraphTrack::BuildContextMenu(FMenuBuilder& MenuBuilder)
 {
-	MenuBuilder.BeginSection("Misc");
+	MenuBuilder.BeginSection("Options", LOCTEXT("OptionsHeading", "Options"));
 	{
 		if (EnumHasAnyFlags(VisibleOptions, EGraphOptions::ShowDebugInfo)) // debug functionality
 		{
diff --git a/Engine/Source/Developer/TraceInsights/Private/Insights/ViewModels/LoadingTimingTrack.cpp b/Engine/Source/Developer/TraceInsights/Private/Insights/ViewModels/LoadingTimingTrack.cpp
--- a/Engine/Source/Developer/TraceInsights/Private/Insights/ViewModels/LoadingTimingTrack.cpp
+++ b/Engine/Source/Developer/TraceInsights/Private/Insights/ViewModels/LoadingTimingTrack.cpp
@@ -111,7 +111,7 @@ void FLoadingSharedState::Tick(Insights::ITimingViewSession& InSession, const Tr
 
 ////////////////////////////////////////////////////////////////////////////////////////////////////
 
-void FLoadingSharedState::ExtendFilterMenu(Insights::ITimingViewSession& InSession, FMenuBuilder& InOutMenuBuilder)
+void FLoadingSharedState::ExtendOtherTracksFilterMenu(Insights::ITimingViewSession& InSession, FMenuBuilder& InOutMenuBuilder)
 {
 	if (&InSession != TimingView)
 	{
diff --git a/Engine/Source/Developer/TraceInsights/Private/Insights/ViewModels/LoadingTimingTrack.h b/Engine/Source/Developer/TraceInsights/Private/Insights/ViewModels/LoadingTimingTrack.h
--- a/Engine/Source/Developer/TraceInsights/Private/Insights/ViewModels/LoadingTimingTrack.h
+++ b/Engine/Source/Developer/TraceInsights/Private/Insights/ViewModels/LoadingTimingTrack.h
@@ -28,7 +28,7 @@ public:
 	virtual void OnBeginSession(Insights::ITimingViewSession& InSession) override;
 	virtual void OnEndSession(Insights::ITimingViewSession& InSession) override;
 	virtual void Tick(Insights::ITimingViewSession& InSession, const TraceServices::IAnalysisSession& InAnalysisSession) override;
-	virtual void ExtendFilterMenu(Insights::ITimingViewSession& InSession, FMenuBuilder& InMenuBuilder) override;
+	virtual void ExtendOtherTracksFilterMenu(Insights::ITimingViewSession& InSession, FMenuBuilder& InMenuBuilder) override;
 
 	const TCHAR* GetEventName(uint32 Depth, const TraceServices::FLoadTimeProfilerCpuEvent& Event) const;
 	void SetColorSchema(int32 Schema);
diff --git a/Engine/Source/Developer/TraceInsights/Private/Insights/ViewModels/MarkersTimingTrack.cpp b/Engine/Source/Developer/TraceInsights/Private/Insights/ViewModels/MarkersTimingTrack.cpp
--- a/Engine/Source/Developer/TraceInsights/Private/Insights/ViewModels/MarkersTimingTrack.cpp
+++ b/Engine/Source/Developer/TraceInsights/Private/Insights/ViewModels/MarkersTimingTrack.cpp
@@ -258,19 +258,8 @@ FReply FMarkersTimingTrack::OnMouseButtonDoubleClick(const FGeometry& MyGeometry
 
 void FMarkersTimingTrack::BuildContextMenu(FMenuBuilder& MenuBuilder)
 {
-	MenuBuilder.BeginSection(TEXT("Misc"));
+	MenuBuilder.BeginSection(TEXT("Content"), LOCTEXT("ContentHeading", "Content"));
 	{
-		MenuBuilder.AddMenuEntry(
-			LOCTEXT("ContextMenu_ToggleCollapsed", "Collapsed"),
-			LOCTEXT("ContextMenu_ToggleCollapsed_Desc", "Whether the vertical marker lines are collapsed or expanded."),
-			FSlateIcon(),
-			FUIAction(FExecuteAction::CreateSP(this, &FMarkersTimingTrack::ToggleCollapsed),
-					  FCanExecuteAction(),
-					  FIsActionChecked::CreateSP(this, &FMarkersTimingTrack::IsCollapsed)),
-			NAME_None,
-			EUserInterfaceActionType::ToggleButton
-		);
-
 		MenuBuilder.AddMenuEntry(
 			LOCTEXT("ContextMenu_Bookmarks", "Bookmarks"),
 			LOCTEXT("ContextMenu_Bookmarks_Desc", "Change this track to show only the bookmarks."),
@@ -294,6 +283,21 @@ void FMarkersTimingTrack::BuildContextMenu(FMenuBuilder& MenuBuilder)
 		);
 	}
 	MenuBuilder.EndSection();
+
+	MenuBuilder.BeginSection(TEXT("MarkerLines"), LOCTEXT("MarkerLinesHeading", "Marker Lines"));
+	{
+		MenuBuilder.AddMenuEntry(
+			LOCTEXT("ContextMenu_ToggleCollapsed", "Collapsed"),
+			LOCTEXT("ContextMenu_ToggleCollapsed_Desc", "Whether the vertical marker lines are collapsed or expanded."),
+			FSlateIcon(),
+			FUIAction(FExecuteAction::CreateSP(this, &FMarkersTimingTrack::ToggleCollapsed),
+				FCanExecuteAction(),
+				FIsActionChecked::CreateSP(this, &FMarkersTimingTrack::IsCollapsed)),
+			NAME_None,
+			EUserInterfaceActionType::ToggleButton
+		);
+	}
+	MenuBuilder.EndSection();
 }
 
 ////////////////////////////////////////////////////////////////////////////////////////////////////
diff --git a/Engine/Source/Developer/TraceInsights/Private/Insights/ViewModels/ThreadTimingTrack.cpp b/Engine/Source/Developer/TraceInsights/Private/Insights/ViewModels/ThreadTimingTrack.cpp
--- a/Engine/Source/Developer/TraceInsights/Private/Insights/ViewModels/ThreadTimingTrack.cpp
+++ b/Engine/Source/Developer/TraceInsights/Private/Insights/ViewModels/ThreadTimingTrack.cpp
@@ -496,14 +496,14 @@ void FThreadTimingSharedState::Tick(Insights::ITimingViewSession& InSession, con
 
 ////////////////////////////////////////////////////////////////////////////////////////////////////
 
-void FThreadTimingSharedState::ExtendFilterMenu(Insights::ITimingViewSession& InSession, FMenuBuilder& InOutMenuBuilder)
+void FThreadTimingSharedState::ExtendGpuTracksFilterMenu(Insights::ITimingViewSession& InSession, FMenuBuilder& InOutMenuBuilder)
 {
 	if (&InSession != TimingView)
 	{
 		return;
 	}
 
-	InOutMenuBuilder.BeginSection("ThreadProfiler", LOCTEXT("ThreadProfilerHeading", "Threads"));
+	InOutMenuBuilder.BeginSection("GpuTracks", LOCTEXT("GpuTracksHeading", "GPU Tracks"));
 	{
 		//TODO: MenuBuilder.AddMenuEntry(Commands.ShowAllGpuTracks);
 		InOutMenuBuilder.AddMenuEntry(
@@ -513,10 +513,24 @@ void FThreadTimingSharedState::ExtendFilterMenu(Insights::ITimingViewSession& In
 			FUIAction(FExecuteAction::CreateSP(this, &FThreadTimingSharedState::ShowHideAllGpuTracks),
 					  FCanExecuteAction(),
 					  FIsActionChecked::CreateSP(this, &FThreadTimingSharedState::IsAllGpuTracksToggleOn)),
-			NAME_None, //"QuickFilterSeparator",
+			NAME_None,
 			EUserInterfaceActionType::ToggleButton
 		);
+	}
+	InOutMenuBuilder.EndSection();
+}
 
+////////////////////////////////////////////////////////////////////////////////////////////////////
+
+void FThreadTimingSharedState::ExtendCpuTracksFilterMenu(Insights::ITimingViewSession& InSession, FMenuBuilder& InOutMenuBuilder)
+{
+	if (&InSession != TimingView)
+	{
+		return;
+	}
+
+	InOutMenuBuilder.BeginSection("CpuTracks", LOCTEXT("CpuTracksHeading", "CPU Tracks"));
+	{
 		//TODO: MenuBuilder.AddMenuEntry(Commands.ShowAllCpuTracks);
 		InOutMenuBuilder.AddMenuEntry(
 			LOCTEXT("ShowAllCpuTracks", "CPU Thread Tracks - U"),
@@ -525,7 +539,7 @@ void FThreadTimingSharedState::ExtendFilterMenu(Insights::ITimingViewSession& In
 			FUIAction(FExecuteAction::CreateSP(this, &FThreadTimingSharedState::ShowHideAllCpuTracks),
 					  FCanExecuteAction(),
 					  FIsActionChecked::CreateSP(this, &FThreadTimingSharedState::IsAllCpuTracksToggleOn)),
-			NAME_None, //"QuickFilterSeparator",
+			NAME_None,
 			EUserInterfaceActionType::ToggleButton
 		);
 	}
@@ -1417,9 +1431,7 @@ void FThreadTimingTrack::BuildContextMenu(FMenuBuilder& MenuBuilder)
 		ChildTrack->BuildContextMenu(MenuBuilder);
 	}
 
-	MenuBuilder.AddSeparator();
-
-	MenuBuilder.BeginSection("Options");
+	MenuBuilder.BeginSection("TimingEvents", LOCTEXT("TimingEventsSection", "Timing Events"));
 	{
 		FExecuteAction FilterTrackAction;
 		FilterTrackAction.BindSP(this, &FThreadTimingTrack::OnFilterTrackClicked);
diff --git a/Engine/Source/Developer/TraceInsights/Private/Insights/ViewModels/ThreadTimingTrack.h b/Engine/Source/Developer/TraceInsights/Private/Insights/ViewModels/ThreadTimingTrack.h
--- a/Engine/Source/Developer/TraceInsights/Private/Insights/ViewModels/ThreadTimingTrack.h
+++ b/Engine/Source/Developer/TraceInsights/Private/Insights/ViewModels/ThreadTimingTrack.h
@@ -58,7 +58,8 @@ public:
 	virtual void OnBeginSession(Insights::ITimingViewSession& InSession) override;
 	virtual void OnEndSession(Insights::ITimingViewSession& InSession) override;
 	virtual void Tick(Insights::ITimingViewSession& InSession, const TraceServices::IAnalysisSession& InAnalysisSession) override;
-	virtual void ExtendFilterMenu(Insights::ITimingViewSession& InSession, FMenuBuilder& InMenuBuilder) override;
+	virtual void ExtendGpuTracksFilterMenu(Insights::ITimingViewSession& InSession, FMenuBuilder& InMenuBuilder) override;
+	virtual void ExtendCpuTracksFilterMenu(Insights::ITimingViewSession& InSession, FMenuBuilder& InMenuBuilder) override;
 
 	//////////////////////////////////////////////////
 
diff --git a/Engine/Source/Developer/TraceInsights/Private/Insights/Widgets/STimingView.cpp b/Engine/Source/Developer/TraceInsights/Private/Insights/Widgets/STimingView.cpp
--- a/Engine/Source/Developer/TraceInsights/Private/Insights/Widgets/STimingView.cpp
+++ b/Engine/Source/Developer/TraceInsights/Private/Insights/Widgets/STimingView.cpp
@@ -147,21 +147,33 @@ void STimingView::Construct(const FArguments& InArgs)
 	LeftToolbar.BeginSection("Menus");
 	LeftToolbar.AddComboButton(
 		FUIAction(),
-		FOnGetContent::CreateSP(this, &STimingView::MakeTracksFilterMenu),
-		LOCTEXT("AllTracksFilter", "All Tracks"),
-		LOCTEXT("AllTracksFilterToolTip", "Filter all tracks."),
+		FOnGetContent::CreateSP(this, &STimingView::MakeAllTracksMenu),
+		LOCTEXT("AllTracksMenu", "All Tracks"),
+		LOCTEXT("AllTracksMenuToolTip", "Filter all tracks."),
 		FSlateIcon(FInsightsStyle::GetStyleSetName(), "Icon.AllTracksMenu"));
 	LeftToolbar.AddComboButton(
 		FUIAction(),
-		FOnGetContent::CreateSP(this, &STimingView::MakeTracksFilterMenu),
-		LOCTEXT("CpuGpuTracksFilter", "CPU/GPU Tracks"),
-		LOCTEXT("CpuGpuTracksFilterToolTip", "Filter CPU/GPU timing tracks."),
+		FOnGetContent::CreateSP(this, &STimingView::MakeCpuGpuTracksFilterMenu),
+		LOCTEXT("CpuGpuTracksMenu", "CPU/GPU"),
+		LOCTEXT("CpuGpuTracksMenuToolTip", "Filter CPU/GPU timing tracks."),
 		FSlateIcon(FInsightsStyle::GetStyleSetName(), "Icon.CpuGpuTracksMenu"));
 	LeftToolbar.AddComboButton(
 		FUIAction(),
-		FOnGetContent::CreateSP(this, &STimingView::MakeTracksFilterMenu),
-		LOCTEXT("OtherTracksFilter", "Other Tracks"),
-		LOCTEXT("OtherTracksFilterToolTip", "Filter other tracks."),
+		FOnGetContent::CreateSP(this, &STimingView::MakeOtherTracksFilterMenu),
+		LOCTEXT("OtherTracksMenu", "Other"),
+		LOCTEXT("OtherTracksMenuToolTip", "Filter other tracks."),
+		FSlateIcon(FInsightsStyle::GetStyleSetName(), "Icon.OtherTracksMenu"));
+	LeftToolbar.AddComboButton(
+		FUIAction(),
+		FOnGetContent::CreateSP(this, &STimingView::MakePluginTracksFilterMenu),
+		LOCTEXT("PluginTracksMenu", "Plugins"),
+		LOCTEXT("PluginTracksMenuToolTip", "Filter tracks added by plugins."),
+		FSlateIcon(FInsightsStyle::GetStyleSetName(), "Icon.OtherTracksMenu"));
+	LeftToolbar.AddComboButton(
+		FUIAction(),
+		FOnGetContent::CreateSP(this, &STimingView::MakeViewModeMenu),
+		LOCTEXT("ViewModeMenu", "View Mode"),
+		LOCTEXT("ViewModeMenuToolTip", "Various options for the Timing view."),
 		FSlateIcon(FInsightsStyle::GetStyleSetName(), "Icon.OtherTracksMenu"));
 	LeftToolbar.EndSection();
 
@@ -2798,8 +2810,8 @@ void STimingView::ShowContextMenu(const FPointerEvent& MouseEvent)
 	const FTimingViewCommands& Commands = FTimingViewCommands::Get();
 
 	const bool bShouldCloseWindowAfterMenuSelection = true;
-
 	FMenuBuilder MenuBuilder(bShouldCloseWindowAfterMenuSelection, CommandList);
+
 	bool bHasAnyActions = false;
 
 	if (HoveredTrack.IsValid())
@@ -2809,14 +2821,20 @@ void STimingView::ShowContextMenu(const FPointerEvent& MouseEvent)
 		MenuBuilder.EndSection();
 
 		HoveredTrack->BuildContextMenu(MenuBuilder);
+		MenuBuilder.AddSeparator();
 		bHasAnyActions = true;
 	}
 
-	MenuBuilder.AddMenuEntry(Commands.QuickFind,
-		NAME_None,
-		TAttribute<FText>(),
-		TAttribute<FText>(),
-		FSlateIcon(FInsightsStyle::GetStyleSetName(), "Icon.Find"));
+	MenuBuilder.BeginSection(TEXT("Misc"));
+	{
+		MenuBuilder.AddMenuEntry(Commands.QuickFind,
+			FName("QuickFind"),
+			TAttribute<FText>(),
+			TAttribute<FText>(),
+			FSlateIcon(FInsightsStyle::GetStyleSetName(), "Icon.Find"));
+		bHasAnyActions = true;
+	}
+	MenuBuilder.EndSection();
 
 	for (Insights::ITimingViewExtender* Extender : GetExtenders())
 	{
@@ -2852,110 +2870,59 @@ void STimingView::CreateTrackLocationMenu(FMenuBuilder& MenuBuilder, TSharedRef<
 {
 	if (EnumHasAnyFlags(Track->GetValidLocations(), ETimingTrackLocation::TopDocked))
 	{
-		if (Track->GetLocation() == ETimingTrackLocation::TopDocked)
-		{
-			MenuBuilder.AddMenuEntry(
-				LOCTEXT("ContextMenu_LocationTopDocked", "Location: Top Docked Tracks"),
-				LOCTEXT("ContextMenu_LocationTopDocked_Desc", "This track is in the list of top docked tracks."),
-				FSlateIcon(),
-				FUIAction(FExecuteAction(), FCanExecuteAction::CreateLambda([] { return false; })),
-				NAME_None,
-				EUserInterfaceActionType::Button
-			);
-		}
-		else
-		{
-			MenuBuilder.AddMenuEntry(
-				LOCTEXT("ContextMenu_DockToTop", "Dock To Top"),
-				LOCTEXT("ContextMenu_DockToTop_Desc", "Dock this track to the top."),
-				FSlateIcon(),
-				FUIAction(FExecuteAction::CreateSP(this, &STimingView::ChangeTrackLocation, Track, ETimingTrackLocation::TopDocked),
-						  FCanExecuteAction::CreateSP(this, &STimingView::CanChangeTrackLocation, Track, ETimingTrackLocation::TopDocked)),
-				NAME_None,
-				EUserInterfaceActionType::Button
-			);
-		}
+		MenuBuilder.AddMenuEntry(
+			LOCTEXT("ContextMenu_DockToTop", "Top Docked"),
+			LOCTEXT("ContextMenu_DockToTop_Desc", "Dock this track to the top."),
+			FSlateIcon(),
+			FUIAction(
+				FExecuteAction::CreateSP(this, &STimingView::ChangeTrackLocation, Track, ETimingTrackLocation::TopDocked),
+				FCanExecuteAction::CreateSP(this, &STimingView::CanChangeTrackLocation, Track, ETimingTrackLocation::TopDocked),
+				FIsActionChecked::CreateSP(this, &STimingView::CheckTrackLocation, Track, ETimingTrackLocation::TopDocked)),
+			NAME_None,
+			EUserInterfaceActionType::RadioButton);
 	}
 
 	if (EnumHasAnyFlags(Track->GetValidLocations(), ETimingTrackLocation::Scrollable))
 	{
-		if (Track->GetLocation() == ETimingTrackLocation::Scrollable)
-		{
-			MenuBuilder.AddMenuEntry(
-				LOCTEXT("ContextMenu_LocationScrollable", "Location: Scrollable Tracks"),
-				LOCTEXT("ContextMenu_LocationScrollable_Desc", "This track is in the list of scrollable tracks."),
-				FSlateIcon(),
-				FUIAction(FExecuteAction(), FCanExecuteAction::CreateLambda([] { return false; })),
-				NAME_None,
-				EUserInterfaceActionType::Button
-			);
-		}
-		else
-		{
-			MenuBuilder.AddMenuEntry(
-				LOCTEXT("ContextMenu_MoveToScrollable", "Move to Scrollable Tracks"),
-				LOCTEXT("ContextMenu_MoveToScrollable_Desc", "Move this track to the list of scrollable tracks."),
-				FSlateIcon(),
-				FUIAction(FExecuteAction::CreateSP(this, &STimingView::ChangeTrackLocation, Track, ETimingTrackLocation::Scrollable),
-						  FCanExecuteAction::CreateSP(this, &STimingView::CanChangeTrackLocation, Track, ETimingTrackLocation::Scrollable)),
-				NAME_None,
-				EUserInterfaceActionType::Button
-			);
-		}
+		MenuBuilder.AddMenuEntry(
+			LOCTEXT("ContextMenu_MoveToScrollable", "Scrollable"),
+			LOCTEXT("ContextMenu_MoveToScrollable_Desc", "Move this track to the list of scrollable tracks."),
+			FSlateIcon(),
+			FUIAction(
+				FExecuteAction::CreateSP(this, &STimingView::ChangeTrackLocation, Track, ETimingTrackLocation::Scrollable),
+				FCanExecuteAction::CreateSP(this, &STimingView::CanChangeTrackLocation, Track, ETimingTrackLocation::Scrollable),
+				FIsActionChecked::CreateSP(this, &STimingView::CheckTrackLocation, Track, ETimingTrackLocation::Scrollable)),
+			NAME_None,
+			EUserInterfaceActionType::RadioButton);
 	}
 
 	if (EnumHasAnyFlags(Track->GetValidLocations(), ETimingTrackLocation::BottomDocked))
 	{
-		if (Track->GetLocation() == ETimingTrackLocation::BottomDocked)
-		{
-			MenuBuilder.AddMenuEntry(
-				LOCTEXT("ContextMenu_LocationBottomDocked", "Location: Bottom Docked Tracks"),
-				LOCTEXT("ContextMenu_LocationBottomDocked_Desc", "This track is in the list of bottom docked tracks."),
-				FSlateIcon(),
-				FUIAction(FExecuteAction(), FCanExecuteAction::CreateLambda([] { return false; })),
-				NAME_None,
-				EUserInterfaceActionType::Button
-			);
-		}
-		else
-		{
-			MenuBuilder.AddMenuEntry(
-				LOCTEXT("ContextMenu_DockToBottom", "Dock To Bottom"),
-				LOCTEXT("ContextMenu_DockToBottom_Desc", "Dock this track to the bottom."),
-				FSlateIcon(),
-				FUIAction(FExecuteAction::CreateSP(this, &STimingView::ChangeTrackLocation, Track, ETimingTrackLocation::BottomDocked),
-						  FCanExecuteAction::CreateSP(this, &STimingView::CanChangeTrackLocation, Track, ETimingTrackLocation::BottomDocked)),
-				NAME_None,
-				EUserInterfaceActionType::Button
-			);
-		}
+		MenuBuilder.AddMenuEntry(
+			LOCTEXT("ContextMenu_DockToBottom", "Bottom Docked"),
+			LOCTEXT("ContextMenu_DockToBottom_Desc", "Dock this track to the bottom."),
+			FSlateIcon(),
+			FUIAction(
+				FExecuteAction::CreateSP(this, &STimingView::ChangeTrackLocation, Track, ETimingTrackLocation::BottomDocked),
+				FCanExecuteAction::CreateSP(this, &STimingView::CanChangeTrackLocation, Track, ETimingTrackLocation::BottomDocked),
+				FIsActionChecked::CreateSP(this, &STimingView::CheckTrackLocation, Track, ETimingTrackLocation::BottomDocked)),
+			NAME_None,
+			EUserInterfaceActionType::RadioButton);
 	}
 
 	if (EnumHasAnyFlags(Track->GetValidLocations(), ETimingTrackLocation::Foreground))
 	{
-		if (Track->GetLocation() == ETimingTrackLocation::Foreground)
-		{
-			MenuBuilder.AddMenuEntry(
-				LOCTEXT("ContextMenu_LocationForeground", "Location: Foreground Tracks"),
-				LOCTEXT("ContextMenu_LocationForeground_Desc", "This track is in the list of foreground tracks."),
-				FSlateIcon(),
-				FUIAction(FExecuteAction(), FCanExecuteAction::CreateLambda([] { return false; })),
-				NAME_None,
-				EUserInterfaceActionType::Button
-			);
-		}
-		else
-		{
-			MenuBuilder.AddMenuEntry(
-				LOCTEXT("ContextMenu_MoveToForeground", "Move to Foreground Tracks"),
-				LOCTEXT("ContextMenu_MoveToForeground_Desc", "Move this track to the list of foreground tracks."),
-				FSlateIcon(),
-				FUIAction(FExecuteAction::CreateSP(this, &STimingView::ChangeTrackLocation, Track, ETimingTrackLocation::Foreground),
-						  FCanExecuteAction::CreateSP(this, &STimingView::CanChangeTrackLocation, Track, ETimingTrackLocation::Foreground)),
-				NAME_None,
-				EUserInterfaceActionType::Button
-			);
-		}
+		MenuBuilder.AddMenuEntry(
+			LOCTEXT("ContextMenu_MoveToForeground", "Foreground"),
+			LOCTEXT("ContextMenu_MoveToForeground_Desc", "Move this track to the list of foreground tracks."),
+			FSlateIcon(),
+			FUIAction(
+				FExecuteAction::CreateSP(this, &STimingView::ChangeTrackLocation, Track, ETimingTrackLocation::Foreground),
+				FCanExecuteAction::CreateSP(this, &STimingView::CanChangeTrackLocation, Track, ETimingTrackLocation::Foreground),
+				FIsActionChecked::CreateSP(this, &STimingView::CheckTrackLocation, Track, ETimingTrackLocation::Foreground)),
+			NAME_None,
+			EUserInterfaceActionType::Button
+		);
 	}
 }
 
@@ -2997,7 +2964,8 @@ void STimingView::EnumerateAllTracks(TFunctionRef<bool(TSharedPtr<FBaseTimingTra
 
 void STimingView::ChangeTrackLocation(TSharedRef<FBaseTimingTrack> Track, ETimingTrackLocation NewLocation)
 {
-	if (ensure(CanChangeTrackLocation(Track, NewLocation)))
+	if (EnumHasAnyFlags(Track->GetValidLocations(), NewLocation) &&
+		Track->GetLocation() != NewLocation)
 	{
 		switch (Track->GetLocation())
 		{
@@ -3045,7 +3013,14 @@ void STimingView::ChangeTrackLocation(TSharedRef<FBaseTimingTrack> Track, ETimin
 
 bool STimingView::CanChangeTrackLocation(TSharedRef<FBaseTimingTrack> Track, ETimingTrackLocation NewLocation) const
 {
-	return EnumHasAnyFlags(Track->GetValidLocations(), NewLocation) && Track->GetLocation() != NewLocation;
+	return EnumHasAnyFlags(Track->GetValidLocations(), NewLocation);
+}
+
+////////////////////////////////////////////////////////////////////////////////////////////////////
+
+bool STimingView::CheckTrackLocation(TSharedRef<FBaseTimingTrack> Track, ETimingTrackLocation Location) const
+{
+	return Track->GetLocation() == Location;
 }
 
 ////////////////////////////////////////////////////////////////////////////////////////////////////
@@ -4288,26 +4263,126 @@ TSharedRef<SWidget> STimingView::MakeAutoScrollOptionsMenu()
 
 ////////////////////////////////////////////////////////////////////////////////////////////////////
 
-TSharedRef<SWidget> STimingView::MakeTracksFilterMenu()
+TSharedRef<SWidget> STimingView::MakeAllTracksMenu()
 {
 	FMenuBuilder MenuBuilder(/*bInShouldCloseWindowAfterMenuSelection=*/true, CommandList);
 
 	CreateAllTracksMenu(MenuBuilder);
 
-	MenuBuilder.BeginSection("QuickFilter", LOCTEXT("TracksFilterHeading", "Quick Filter"));
+	return MenuBuilder.MakeWidget();
+}
+
+////////////////////////////////////////////////////////////////////////////////////////////////////
+
+void STimingView::CreateAllTracksMenu(FMenuBuilder& MenuBuilder)
+{
+	constexpr float MaxDesiredHeight = 24.0f + 8 * 18.0f;
+	constexpr float MaxDesiredHeightScrollableTracks = 24.0f + 16 * 18.0f;
+	constexpr float MinDesiredWidth = 300.0f;
+	constexpr float MaxDesiredWidth = 300.0f;
+
+	if (TopDockedTracks.Num() > 0)
 	{
-		const FTimingViewCommands& Commands = FTimingViewCommands::Get();
+		MenuBuilder.BeginSection("TopDockedTracks", LOCTEXT("TopDockedTracksHeading", "Top Docked Tracks"));
+		//MenuBuilder.AddWidget(
+		//	SNew(STextBlock)
+		//	.Text(LOCTEXT("TopDockedTracks", "Top Docked Tracks")),
+		//	FText(), true);
 
-		MenuBuilder.AddMenuEntry(Commands.ShowMainGraphTrack);
-		MenuBuilder.AddMenuEntry(Commands.ToggleCompactMode);
-		MenuBuilder.AddMenuEntry(Commands.AutoHideEmptyTracks);
+		MenuBuilder.AddWidget(
+			SNew(SBox)
+			.MaxDesiredHeight(MaxDesiredHeight)
+			.MinDesiredWidth(MinDesiredWidth)
+			.MaxDesiredWidth(MaxDesiredWidth)
+			[
+				SNew(STimingViewTrackList, SharedThis(this), ETimingTrackLocation::TopDocked)
+			],
+			FText(), true);
+
+		MenuBuilder.EndSection();
 	}
-	MenuBuilder.EndSection();
 
-	// Let any plugin extend the filter menu.
+	if (ScrollableTracks.Num() > 0)
+	{
+		MenuBuilder.BeginSection("ScrollableTracks", LOCTEXT("ScrollableTracksHeading", "Scrollable Tracks"));
+		//MenuBuilder.AddWidget(
+		//	SNew(STextBlock)
+		//	.Text(LOCTEXT("ScrollableTracks", "Scrollable Tracks")),
+		//	FText(), true);
+
+		MenuBuilder.AddWidget(
+			SNew(SBox)
+			.Padding(FMargin(0.0f, 0.0f))
+			.MaxDesiredHeight(MaxDesiredHeightScrollableTracks)
+			.MinDesiredWidth(MinDesiredWidth)
+			.MaxDesiredWidth(MaxDesiredWidth)
+			[
+				SNew(STimingViewTrackList, SharedThis(this), ETimingTrackLocation::Scrollable)
+			],
+			FText(), true);
+
+		MenuBuilder.EndSection();
+	}
+
+	if (BottomDockedTracks.Num() > 0)
+	{
+		MenuBuilder.BeginSection("BottomDockedTracks", LOCTEXT("BottomDockedTracksHeading", "Bottom Docked Tracks"));
+		//MenuBuilder.AddWidget(
+		//	SNew(STextBlock)
+		//	.Text(LOCTEXT("BottomDockedTracks", "Bottom Docked Tracks")),
+		//	FText(), true);
+
+		MenuBuilder.AddWidget(
+			SNew(SBox)
+			.MaxDesiredHeight(MaxDesiredHeight)
+			.MinDesiredWidth(MinDesiredWidth)
+			.MaxDesiredWidth(MaxDesiredWidth)
+			[
+				SNew(STimingViewTrackList, SharedThis(this), ETimingTrackLocation::BottomDocked)
+			],
+			FText(), true);
+
+		MenuBuilder.EndSection();
+	}
+
+	if (ForegroundTracks.Num() > 0)
+	{
+		MenuBuilder.BeginSection("ForegroundTracks", LOCTEXT("ForegroundTracksHeading", "Foreground Tracks"));
+		//MenuBuilder.AddWidget(
+		//	SNew(STextBlock)
+		//	.Text(LOCTEXT("ForegroundTracks", "Foreground Tracks")),
+		//	FText(), true);
+		
+		MenuBuilder.AddWidget(
+			SNew(SBox)
+			.MaxDesiredHeight(MaxDesiredHeight)
+			.MinDesiredWidth(MinDesiredWidth)
+			.MaxDesiredWidth(MaxDesiredWidth)
+			[
+				SNew(STimingViewTrackList, SharedThis(this), ETimingTrackLocation::Foreground)
+			],
+			FText(), true);
+
+		MenuBuilder.EndSection();
+	}
+}
+
+////////////////////////////////////////////////////////////////////////////////////////////////////
+
+TSharedRef<SWidget> STimingView::MakeCpuGpuTracksFilterMenu()
+{
+	FMenuBuilder MenuBuilder(/*bInShouldCloseWindowAfterMenuSelection=*/true, CommandList);
+
+	// Let any plugin extend the GPU Tracks Filter menu.
 	for (Insights::ITimingViewExtender* Extender : GetExtenders())
 	{
-		Extender->ExtendFilterMenu(*this, MenuBuilder);
+		Extender->ExtendGpuTracksFilterMenu(*this, MenuBuilder);
+	}
+
+	// Let any plugin extend the CPU Tracks Filter menu.
+	for (Insights::ITimingViewExtender* Extender : GetExtenders())
+	{
+		Extender->ExtendCpuTracksFilterMenu(*this, MenuBuilder);
 	}
 
 	return MenuBuilder.MakeWidget();
@@ -4315,91 +4390,58 @@ TSharedRef<SWidget> STimingView::MakeTracksFilterMenu()
 
 ////////////////////////////////////////////////////////////////////////////////////////////////////
 
-void STimingView::CreateAllTracksMenu(FMenuBuilder& MenuBuilder)
+TSharedRef<SWidget> STimingView::MakeOtherTracksFilterMenu()
 {
-	MenuBuilder.BeginSection("AllTracks", LOCTEXT("AllTracksHeading", "All Tracks"));
+	FMenuBuilder MenuBuilder(/*bInShouldCloseWindowAfterMenuSelection=*/true, CommandList);
 
-	if (TopDockedTracks.Num() > 0)
+	const FTimingViewCommands& Commands = FTimingViewCommands::Get();
+
+	MenuBuilder.BeginSection("GraphTracks", LOCTEXT("GraphTracksHeading", "Main Graph Track"));
 	{
-		MenuBuilder.AddSubMenu(
-			LOCTEXT("TopDockedTracks", "Top Docked Tracks"),
-			LOCTEXT("TopDockedTracks_Tooltip", "Show/hide individual top docked tracks"),
-			FNewMenuDelegate::CreateLambda([this](FMenuBuilder& InSubMenuBuilder)
-				{
-					InSubMenuBuilder.AddWidget(
-						SNew(SBox)
-						.MaxDesiredHeight(300.0f)
-						.MinDesiredWidth(300.0f)
-						.MaxDesiredWidth(300.0f)
-						[
-							SNew(STimingViewTrackList, SharedThis(this), ETimingTrackLocation::TopDocked)
-						],
-						FText(), true);
-				})
-		);
+		MenuBuilder.AddMenuEntry(Commands.ShowMainGraphTrack);
 	}
+	MenuBuilder.EndSection();
 
-	if (BottomDockedTracks.Num() > 0)
+	// Let any plugin extend the Other Tracks Filter menu.
+	for (Insights::ITimingViewExtender* Extender : GetExtenders())
 	{
-		MenuBuilder.AddSubMenu(
-			LOCTEXT("BottomDockedTracks", "Bottom Docked Tracks"),
-			LOCTEXT("BottomDockedTracks_Tooltip", "Show/hide individual bottom docked tracks"),
-			FNewMenuDelegate::CreateLambda([this](FMenuBuilder& InSubMenuBuilder)
-				{
-					InSubMenuBuilder.AddWidget(
-						SNew(SBox)
-						.MaxDesiredHeight(300.0f)
-						.MinDesiredWidth(300.0f)
-						.MaxDesiredWidth(300.0f)
-						[
-							SNew(STimingViewTrackList, SharedThis(this), ETimingTrackLocation::BottomDocked)
-						],
-						FText(), true);
-				})
-		);
+		Extender->ExtendOtherTracksFilterMenu(*this, MenuBuilder);
 	}
 
-	if (ScrollableTracks.Num() > 0)
+	return MenuBuilder.MakeWidget();
+}
+
+////////////////////////////////////////////////////////////////////////////////////////////////////
+
+TSharedRef<SWidget> STimingView::MakePluginTracksFilterMenu()
+{
+	FMenuBuilder MenuBuilder(/*bInShouldCloseWindowAfterMenuSelection=*/true, CommandList);
+
+	// Let any plugin extend the filter menu.
+	for (Insights::ITimingViewExtender* Extender : GetExtenders())
 	{
-		MenuBuilder.AddSubMenu(
-			LOCTEXT("ScrollableTracks", "Scrollable Tracks"),
-			LOCTEXT("ScrollableTracks_Tooltip", "Show/hide individual scrollable tracks"),
-			FNewMenuDelegate::CreateLambda([this](FMenuBuilder& InSubMenuBuilder)
-				{
-					InSubMenuBuilder.AddWidget(
-						SNew(SBox)
-						.MaxDesiredHeight(300.0f)
-						.MinDesiredWidth(300.0f)
-						.MaxDesiredWidth(300.0f)
-						[
-							SNew(STimingViewTrackList, SharedThis(this), ETimingTrackLocation::Scrollable)
-						],
-						FText(), true);
-				})
-		);
+		Extender->ExtendFilterMenu(*this, MenuBuilder);
 	}
 
-	if (ForegroundTracks.Num() > 0)
+	return MenuBuilder.MakeWidget();
+}
+
+////////////////////////////////////////////////////////////////////////////////////////////////////
+
+TSharedRef<SWidget> STimingView::MakeViewModeMenu()
+{
+	FMenuBuilder MenuBuilder(/*bInShouldCloseWindowAfterMenuSelection=*/true, CommandList);
+
+	const FTimingViewCommands& Commands = FTimingViewCommands::Get();
+
+	MenuBuilder.BeginSection("ViewMode", LOCTEXT("ViewModeHeading", "View Mode"));
 	{
-		MenuBuilder.AddSubMenu(
-			LOCTEXT("ForegroundTracks", "Foreground Tracks"),
-			LOCTEXT("ForegroundTracks_Tooltip", "Show/hide individual foreground tracks"),
-			FNewMenuDelegate::CreateLambda([this](FMenuBuilder& InSubMenuBuilder)
-				{
-					InSubMenuBuilder.AddWidget(
-						SNew(SBox)
-						.MaxDesiredHeight(300.0f)
-						.MinDesiredWidth(300.0f)
-						.MaxDesiredWidth(300.0f)
-						[
-							SNew(STimingViewTrackList, SharedThis(this), ETimingTrackLocation::Foreground)
-						],
-						FText(), true);
-				})
-		);
+		MenuBuilder.AddMenuEntry(Commands.ToggleCompactMode);
+		MenuBuilder.AddMenuEntry(Commands.AutoHideEmptyTracks);
 	}
-
 	MenuBuilder.EndSection();
+
+	return MenuBuilder.MakeWidget();
 }
 
 ////////////////////////////////////////////////////////////////////////////////////////////////////
diff --git a/Engine/Source/Developer/TraceInsights/Private/Insights/Widgets/STimingView.h b/Engine/Source/Developer/TraceInsights/Private/Insights/Widgets/STimingView.h
--- a/Engine/Source/Developer/TraceInsights/Private/Insights/Widgets/STimingView.h
+++ b/Engine/Source/Developer/TraceInsights/Private/Insights/Widgets/STimingView.h
@@ -75,8 +75,12 @@ public:
 	TSharedRef<SWidget> MakeCompactAutoScrollOptionsMenu();
 	TSharedRef<SWidget> MakeAutoScrollOptionsMenu();
 
-	TSharedRef<SWidget> MakeTracksFilterMenu();
+	TSharedRef<SWidget> MakeAllTracksMenu();
 	void CreateAllTracksMenu(FMenuBuilder& MenuBuilder);
+	TSharedRef<SWidget> MakeCpuGpuTracksFilterMenu();
+	TSharedRef<SWidget> MakeOtherTracksFilterMenu();
+	TSharedRef<SWidget> MakePluginTracksFilterMenu();
+	TSharedRef<SWidget> MakeViewModeMenu();
 
 	bool ShowHideGraphTrack_IsChecked() const;
 	void ShowHideGraphTrack_Execute();
@@ -303,6 +307,7 @@ public:
 
 	void ChangeTrackLocation(TSharedRef<FBaseTimingTrack> Track, ETimingTrackLocation NewLocation);
 	bool CanChangeTrackLocation(TSharedRef<FBaseTimingTrack> Track, ETimingTrackLocation NewLocation) const;
+	bool CheckTrackLocation(TSharedRef<FBaseTimingTrack> Track, ETimingTrackLocation Location) const;
 
 	void UpdateScrollableTracksOrder();
 	int32 GetFirstScrollableTrackOrder() const;
diff --git a/Engine/Source/Developer/TraceInsights/Private/Insights/Widgets/STimingViewTrackList.cpp b/Engine/Source/Developer/TraceInsights/Private/Insights/Widgets/STimingViewTrackList.cpp
--- a/Engine/Source/Developer/TraceInsights/Private/Insights/Widgets/STimingViewTrackList.cpp
+++ b/Engine/Source/Developer/TraceInsights/Private/Insights/Widgets/STimingViewTrackList.cpp
@@ -126,7 +126,7 @@ void STimingViewTrackList::Construct(const FArguments& InArgs, const TSharedRef<
 			SNew(SHorizontalBox)
 			+SHorizontalBox::Slot()
 			.AutoWidth()
-			.Padding(2.0f, 0.0f)
+			.Padding(2.0f, 0.0f, 0.0f, 0.0f)
 			[
 				// Checkbox for bulk operations
 				SNew(SCheckBox)
@@ -164,6 +164,7 @@ void STimingViewTrackList::Construct(const FArguments& InArgs, const TSharedRef<
 				})
 			]
 			+ SHorizontalBox::Slot()
+			.Padding(FMargin(0.0f, 0.0f, 4.0f, 0.0f))
 			.FillWidth(1.0f)
 			[
 				// Search box allows for filtering
@@ -172,6 +173,7 @@ void STimingViewTrackList::Construct(const FArguments& InArgs, const TSharedRef<
 			]
 		]
 		+ SVerticalBox::Slot()
+		.Padding(FMargin(2.0f, 0.0f))
 		.FillHeight(1.0f)
 		[
 			SNew(SScrollBorder, ListView.ToSharedRef())
diff --git a/Engine/Source/Developer/TraceInsights/Public/Insights/ITimingViewExtender.h b/Engine/Source/Developer/TraceInsights/Public/Insights/ITimingViewExtender.h
--- a/Engine/Source/Developer/TraceInsights/Public/Insights/ITimingViewExtender.h
+++ b/Engine/Source/Developer/TraceInsights/Public/Insights/ITimingViewExtender.h
@@ -28,8 +28,14 @@ public:
 	/** Called each frame. If any new tracks are created they can be added via ITimingViewSession::Add*Track() */
 	virtual void Tick(ITimingViewSession& InSession, const TraceServices::IAnalysisSession& InAnalysisSession) = 0;
 
-	/** Extension hook for the 'quick filter' menu */
-	virtual void ExtendFilterMenu(ITimingViewSession& InSession, FMenuBuilder& InMenuBuilder) = 0;
+	/** Extension hook for the 'CPU Tracks Filter' menu */
+	virtual void ExtendCpuTracksFilterMenu(ITimingViewSession& InSession, FMenuBuilder& InMenuBuilder) {}
+	/** Extension hook for the 'GPU Tracks Filter' menu */
+	virtual void ExtendGpuTracksFilterMenu(ITimingViewSession& InSession, FMenuBuilder& InMenuBuilder) {}
+	/** Extension hook for the 'Other Tracks Filter' menu */
+	virtual void ExtendOtherTracksFilterMenu(ITimingViewSession& InSession, FMenuBuilder& InMenuBuilder) {}
+	/** Extension hook for the 'Plugins' menu */
+	virtual void ExtendFilterMenu(ITimingViewSession& InSession, FMenuBuilder& InMenuBuilder) {}
 
 	/** Extension hook for the context menu for all tracks
 	@return True if any menu option was added and False if no option was added */ 
