[Insights] Timing view: Added "Detph Limit" and "Coloring Mode (CPU Thread Tracks)" in the View Mode menu from toolbar.

diff --git a/Engine/Source/Developer/TraceInsights/Private/Insights/TimingProfilerManager.h b/Engine/Source/Developer/TraceInsights/Private/Insights/TimingProfilerManager.h
--- a/Engine/Source/Developer/TraceInsights/Private/Insights/TimingProfilerManager.h
+++ b/Engine/Source/Developer/TraceInsights/Private/Insights/TimingProfilerManager.h
@@ -18,8 +18,8 @@ namespace Insights
 
 	enum class ETimingEventsColoringMode : uint32
 	{
-		ByName,
-		ById,
+		ByTimerName,
+		ByTimerId,
 		ByDuration,
 
 		Count
@@ -36,6 +36,9 @@ class FTimingProfilerManager : public TSharedFromThis<FTimingProfilerManager>, p
 {
 	friend class FTimingProfilerActionManager;
 
+public:
+	static const uint32 UnlimitedEventDepth = 1000;
+
 public:
 	/** Creates the Timing Profiler manager, only one instance can exist. */
 	FTimingProfilerManager(TSharedRef<FUICommandList> InCommandList);
@@ -232,8 +235,8 @@ private:
 	/** The name of the Timing Insights log listing. */
 	FName LogListingName;
 
-	Insights::ETimingEventsColoringMode ColoringMode = Insights::ETimingEventsColoringMode::ByName;
-	uint32 EventDepthLimit = 1000;
+	Insights::ETimingEventsColoringMode ColoringMode = Insights::ETimingEventsColoringMode::ByTimerName;
+	uint32 EventDepthLimit = UnlimitedEventDepth;
 
 	/** A shared pointer to the global instance of the Timing Profiler manager. */
 	static TSharedPtr<FTimingProfilerManager> Instance;
diff --git a/Engine/Source/Developer/TraceInsights/Private/Insights/ViewModels/FileActivityTimingTrack.cpp b/Engine/Source/Developer/TraceInsights/Private/Insights/ViewModels/FileActivityTimingTrack.cpp
--- a/Engine/Source/Developer/TraceInsights/Private/Insights/ViewModels/FileActivityTimingTrack.cpp
+++ b/Engine/Source/Developer/TraceInsights/Private/Insights/ViewModels/FileActivityTimingTrack.cpp
@@ -377,18 +377,18 @@ void FFileActivitySharedState::BuildSubMenu(FMenuBuilder& InOutMenuBuilder)
 			FUIAction(FExecuteAction::CreateSP(this, &FFileActivitySharedState::ShowHideIoOverviewTrack),
 				FCanExecuteAction(),
 				FIsActionChecked::CreateSP(this, &FFileActivitySharedState::IsIoOverviewTrackVisible)),
-			NAME_None, //"QuickFilterSeparator",
+			NAME_None,
 			EUserInterfaceActionType::ToggleButton
 		);
 
 		InOutMenuBuilder.AddMenuEntry(
-			LOCTEXT("ShowOnlyErrors", "Show Only Errors (I/O Overview Track)"),
+			LOCTEXT("ShowOnlyErrors", "Only Errors (I/O Overview Track)"),
 			LOCTEXT("ShowOnlyErrors_Tooltip", "Show only the events with errors, in the I/O Overview track."),
 			FSlateIcon(),
 			FUIAction(FExecuteAction::CreateSP(this, &FFileActivitySharedState::ToggleOnlyErrors),
 				FCanExecuteAction(),
 				FIsActionChecked::CreateSP(this, &FFileActivitySharedState::IsOnlyErrorsToggleOn)),
-			NAME_None, //"QuickFilterSeparator",
+			NAME_None,
 			EUserInterfaceActionType::ToggleButton
 		);
 
@@ -399,18 +399,18 @@ void FFileActivitySharedState::BuildSubMenu(FMenuBuilder& InOutMenuBuilder)
 			FUIAction(FExecuteAction::CreateSP(this, &FFileActivitySharedState::ShowHideIoActivityTrack),
 				FCanExecuteAction(),
 				FIsActionChecked::CreateSP(this, &FFileActivitySharedState::IsIoActivityTrackVisible)),
-			NAME_None, //"QuickFilterSeparator",
+			NAME_None,
 			EUserInterfaceActionType::ToggleButton
 		);
 
 		InOutMenuBuilder.AddMenuEntry(
-			LOCTEXT("ShowBackgroundEvents", "Show Background Events (I/O Activity Track) - O"),
+			LOCTEXT("ShowBackgroundEvents", "Background Events (I/O Activity Track) - O"),
 			LOCTEXT("ShowBackgroundEvents_Tooltip", "Show background events for file activities, in the I/O Activity track."),
 			FSlateIcon(),
 			FUIAction(FExecuteAction::CreateSP(this, &FFileActivitySharedState::ToggleBackgroundEvents),
 				FCanExecuteAction(),
 				FIsActionChecked::CreateSP(this, &FFileActivitySharedState::AreBackgroundEventsVisible)),
-			NAME_None, //"QuickFilterSeparator",
+			NAME_None,
 			EUserInterfaceActionType::ToggleButton
 		);
 	}
@@ -829,7 +829,7 @@ void FOverviewFileActivityTimingTrack::BuildContextMenu(FMenuBuilder& InOutMenuB
 			FUIAction(FExecuteAction::CreateSP(this, &FOverviewFileActivityTimingTrack::ToggleOnlyErrors),
 				FCanExecuteAction(),
 				FIsActionChecked::CreateSP(this, &FOverviewFileActivityTimingTrack::IsOnlyErrorsToggleOn)),
-			NAME_None, //"QuickFilterSeparator",
+			NAME_None,
 			EUserInterfaceActionType::ToggleButton
 		);
 	}
@@ -968,7 +968,7 @@ void FDetailedFileActivityTimingTrack::BuildContextMenu(FMenuBuilder& InOutMenuB
 			FUIAction(FExecuteAction::CreateSP(this, &FDetailedFileActivityTimingTrack::ToggleOnlyErrors),
 				FCanExecuteAction(),
 				FIsActionChecked::CreateSP(this, &FDetailedFileActivityTimingTrack::IsOnlyErrorsToggleOn)),
-			NAME_None, //"QuickFilterSeparator",
+			NAME_None,
 			EUserInterfaceActionType::ToggleButton
 		);
 
@@ -979,7 +979,7 @@ void FDetailedFileActivityTimingTrack::BuildContextMenu(FMenuBuilder& InOutMenuB
 			FUIAction(FExecuteAction::CreateSP(this, &FDetailedFileActivityTimingTrack::ToggleBackgroundEvents),
 				FCanExecuteAction(),
 				FIsActionChecked::CreateSP(this, &FDetailedFileActivityTimingTrack::AreBackgroundEventsVisible)),
-			NAME_None, //"QuickFilterSeparator",
+			NAME_None,
 			EUserInterfaceActionType::ToggleButton
 		);
 	}
diff --git a/Engine/Source/Developer/TraceInsights/Private/Insights/ViewModels/ThreadTimingTrack.cpp b/Engine/Source/Developer/TraceInsights/Private/Insights/ViewModels/ThreadTimingTrack.cpp
--- a/Engine/Source/Developer/TraceInsights/Private/Insights/ViewModels/ThreadTimingTrack.cpp
+++ b/Engine/Source/Developer/TraceInsights/Private/Insights/ViewModels/ThreadTimingTrack.cpp
@@ -213,10 +213,10 @@ static void AddTimingEventToBuilder(ITimingEventsTrackDrawStateBuilder& Builder,
 	uint32 EventColor;
 	switch (FTimingProfilerManager::Get()->GetColoringMode())
 	{
-		case Insights::ETimingEventsColoringMode::ByName:
+		case Insights::ETimingEventsColoringMode::ByTimerName:
 			EventColor = FTimingEvent::ComputeEventColor(Timer->Name);
 			break;
-		case Insights::ETimingEventsColoringMode::ById:
+		case Insights::ETimingEventsColoringMode::ByTimerId:
 			EventColor = FTimingEvent::ComputeEventColor(Timer->Id);
 			break;
 		case Insights::ETimingEventsColoringMode::ByDuration:
diff --git a/Engine/Source/Developer/TraceInsights/Private/Insights/Widgets/STimingView.cpp b/Engine/Source/Developer/TraceInsights/Private/Insights/Widgets/STimingView.cpp
--- a/Engine/Source/Developer/TraceInsights/Private/Insights/Widgets/STimingView.cpp
+++ b/Engine/Source/Developer/TraceInsights/Private/Insights/Widgets/STimingView.cpp
@@ -2747,29 +2747,12 @@ FReply STimingView::OnKeyDown(const FGeometry& MyGeometry, const FKeyEvent& InKe
 	}
 	else if (InKeyEvent.GetKey() == EKeys::Zero)
 	{
-		uint32 Mode = (uint32)FTimingProfilerManager::Get()->GetColoringMode();
-		Mode = (Mode + 1) % (uint32)Insights::ETimingEventsColoringMode::Count;
-		FTimingProfilerManager::Get()->SetColoringMode((Insights::ETimingEventsColoringMode)Mode);
-		Viewport.AddDirtyFlags(ETimingTrackViewportDirtyFlags::HInvalidated);
+		ChooseNextCpuThreadTrackColoringMode();
 		return FReply::Handled();
 	}
 	else if (InKeyEvent.GetKey() == EKeys::X)
 	{
-		uint32 DepthLimit = FTimingProfilerManager::Get()->GetEventDepthLimit();
-		if (DepthLimit == 1)
-		{
-			DepthLimit = 4;
-		}
-		else if (DepthLimit == 4)
-		{
-			DepthLimit = 1000;
-		}
-		else
-		{
-			DepthLimit = 1;
-		}
-		FTimingProfilerManager::Get()->SetEventDepthLimit(DepthLimit);
-		Viewport.AddDirtyFlags(ETimingTrackViewportDirtyFlags::HInvalidated);
+		ChooseNextEventDepthLimit();
 		return FReply::Handled();
 	}
 #if ACTIVATE_BENCHMARK
@@ -2871,50 +2854,53 @@ void STimingView::CreateTrackLocationMenu(FMenuBuilder& MenuBuilder, TSharedRef<
 	if (EnumHasAnyFlags(Track->GetValidLocations(), ETimingTrackLocation::TopDocked))
 	{
 		MenuBuilder.AddMenuEntry(
-			LOCTEXT("ContextMenu_DockToTop", "Top Docked"),
-			LOCTEXT("ContextMenu_DockToTop_Desc", "Dock this track to the top."),
+			LOCTEXT("TrackLocationMenu_DockToTop", "Top Docked"),
+			LOCTEXT("TrackLocationMenu_DockToTop_Desc", "Dock this track to the top."),
 			FSlateIcon(),
 			FUIAction(
 				FExecuteAction::CreateSP(this, &STimingView::ChangeTrackLocation, Track, ETimingTrackLocation::TopDocked),
 				FCanExecuteAction::CreateSP(this, &STimingView::CanChangeTrackLocation, Track, ETimingTrackLocation::TopDocked),
 				FIsActionChecked::CreateSP(this, &STimingView::CheckTrackLocation, Track, ETimingTrackLocation::TopDocked)),
 			NAME_None,
-			EUserInterfaceActionType::RadioButton);
+			EUserInterfaceActionType::RadioButton
+		);
 	}
 
 	if (EnumHasAnyFlags(Track->GetValidLocations(), ETimingTrackLocation::Scrollable))
 	{
 		MenuBuilder.AddMenuEntry(
-			LOCTEXT("ContextMenu_MoveToScrollable", "Scrollable"),
-			LOCTEXT("ContextMenu_MoveToScrollable_Desc", "Move this track to the list of scrollable tracks."),
+			LOCTEXT("TrackLocationMenu_MoveToScrollable", "Scrollable"),
+			LOCTEXT("TrackLocationMenu_MoveToScrollable_Desc", "Move this track to the list of scrollable tracks."),
 			FSlateIcon(),
 			FUIAction(
 				FExecuteAction::CreateSP(this, &STimingView::ChangeTrackLocation, Track, ETimingTrackLocation::Scrollable),
 				FCanExecuteAction::CreateSP(this, &STimingView::CanChangeTrackLocation, Track, ETimingTrackLocation::Scrollable),
 				FIsActionChecked::CreateSP(this, &STimingView::CheckTrackLocation, Track, ETimingTrackLocation::Scrollable)),
 			NAME_None,
-			EUserInterfaceActionType::RadioButton);
+			EUserInterfaceActionType::RadioButton
+		);
 	}
 
 	if (EnumHasAnyFlags(Track->GetValidLocations(), ETimingTrackLocation::BottomDocked))
 	{
 		MenuBuilder.AddMenuEntry(
-			LOCTEXT("ContextMenu_DockToBottom", "Bottom Docked"),
-			LOCTEXT("ContextMenu_DockToBottom_Desc", "Dock this track to the bottom."),
+			LOCTEXT("TrackLocationMenu_DockToBottom", "Bottom Docked"),
+			LOCTEXT("TrackLocationMenu_DockToBottom_Desc", "Dock this track to the bottom."),
 			FSlateIcon(),
 			FUIAction(
 				FExecuteAction::CreateSP(this, &STimingView::ChangeTrackLocation, Track, ETimingTrackLocation::BottomDocked),
 				FCanExecuteAction::CreateSP(this, &STimingView::CanChangeTrackLocation, Track, ETimingTrackLocation::BottomDocked),
 				FIsActionChecked::CreateSP(this, &STimingView::CheckTrackLocation, Track, ETimingTrackLocation::BottomDocked)),
 			NAME_None,
-			EUserInterfaceActionType::RadioButton);
+			EUserInterfaceActionType::RadioButton
+		);
 	}
 
 	if (EnumHasAnyFlags(Track->GetValidLocations(), ETimingTrackLocation::Foreground))
 	{
 		MenuBuilder.AddMenuEntry(
-			LOCTEXT("ContextMenu_MoveToForeground", "Foreground"),
-			LOCTEXT("ContextMenu_MoveToForeground_Desc", "Move this track to the list of foreground tracks."),
+			LOCTEXT("TrackLocationMenu_MoveToForeground", "Foreground"),
+			LOCTEXT("TrackLocationMenu_MoveToForeground_Desc", "Move this track to the list of foreground tracks."),
 			FSlateIcon(),
 			FUIAction(
 				FExecuteAction::CreateSP(this, &STimingView::ChangeTrackLocation, Track, ETimingTrackLocation::Foreground),
@@ -4413,6 +4399,20 @@ TSharedRef<SWidget> STimingView::MakeOtherTracksFilterMenu()
 
 ////////////////////////////////////////////////////////////////////////////////////////////////////
 
+bool STimingView::ShowHideGraphTrack_IsChecked() const
+{
+	return GraphTrack->IsVisible();
+}
+
+////////////////////////////////////////////////////////////////////////////////////////////////////
+
+void STimingView::ShowHideGraphTrack_Execute()
+{
+	GraphTrack->ToggleVisibility();
+}
+
+////////////////////////////////////////////////////////////////////////////////////////////////////
+
 TSharedRef<SWidget> STimingView::MakePluginTracksFilterMenu()
 {
 	FMenuBuilder MenuBuilder(/*bInShouldCloseWindowAfterMenuSelection=*/true, CommandList);
@@ -4441,21 +4441,11 @@ TSharedRef<SWidget> STimingView::MakeViewModeMenu()
 	}
 	MenuBuilder.EndSection();
 
-	return MenuBuilder.MakeWidget();
-}
+	CreateDepthLimitMenu(MenuBuilder);
 
-////////////////////////////////////////////////////////////////////////////////////////////////////
+	CreateCpuThreadTrackColoringModeMenu(MenuBuilder);
 
-bool STimingView::ShowHideGraphTrack_IsChecked() const
-{
-	return GraphTrack->IsVisible();
-}
-
-////////////////////////////////////////////////////////////////////////////////////////////////////
-
-void STimingView::ShowHideGraphTrack_Execute()
-{
-	GraphTrack->ToggleVisibility();
+	return MenuBuilder.MakeWidget();
 }
 
 ////////////////////////////////////////////////////////////////////////////////////////////////////
@@ -4501,6 +4491,153 @@ void STimingView::ToggleAutoHideEmptyTracks()
 
 ////////////////////////////////////////////////////////////////////////////////////////////////////
 
+void STimingView::CreateDepthLimitMenu(FMenuBuilder& MenuBuilder)
+{
+	MenuBuilder.BeginSection("DepthLimit", LOCTEXT("DepthLimitHeading", "Depth Limit"));
+	{
+		MenuBuilder.AddMenuEntry(
+			LOCTEXT("UnlimitedDepth", "Unlimited"),
+			LOCTEXT("UnlimitedDepth_Desc", "Timing Events tracks (like the Cpu Thread tracks) can have unlimited depth (lanes)."),
+			FSlateIcon(),
+			FUIAction(
+				FExecuteAction::CreateSP(this, &STimingView::SetEventDepthLimit, FTimingProfilerManager::UnlimitedEventDepth),
+				FCanExecuteAction(),
+				FIsActionChecked::CreateSP(this, &STimingView::CheckEventDepthLimit, FTimingProfilerManager::UnlimitedEventDepth)),
+			NAME_None,
+			EUserInterfaceActionType::RadioButton
+		);
+		MenuBuilder.AddMenuEntry(
+			LOCTEXT("DepthLimit4", "4 Lanes"),
+			LOCTEXT("DepthLimit4_Desc", "Timing Events tracks (like the Cpu Thread tracks) can have maximum 4 lanes."),
+			FSlateIcon(),
+			FUIAction(
+				FExecuteAction::CreateSP(this, &STimingView::SetEventDepthLimit, (uint32)4),
+				FCanExecuteAction(),
+				FIsActionChecked::CreateSP(this, &STimingView::CheckEventDepthLimit, (uint32)4)),
+			NAME_None,
+			EUserInterfaceActionType::RadioButton
+		);
+		MenuBuilder.AddMenuEntry(
+			LOCTEXT("DepthLimit1", "Single Lane"),
+			LOCTEXT("DepthLimit1_Desc", "Timing Events tracks (like the Cpu Thread tracks) can have a single lane."),
+			FSlateIcon(),
+			FUIAction(
+				FExecuteAction::CreateSP(this, &STimingView::SetEventDepthLimit, (uint32)1),
+				FCanExecuteAction(),
+				FIsActionChecked::CreateSP(this, &STimingView::CheckEventDepthLimit, (uint32)1)),
+			NAME_None,
+			EUserInterfaceActionType::RadioButton
+		);
+	}
+	MenuBuilder.EndSection();
+}
+
+////////////////////////////////////////////////////////////////////////////////////////////////////
+
+void STimingView::ChooseNextEventDepthLimit()
+{
+	uint32 DepthLimit = FTimingProfilerManager::Get()->GetEventDepthLimit();
+	if (DepthLimit == 1)
+	{
+		DepthLimit = 4;
+	}
+	else if (DepthLimit == 4)
+	{
+		DepthLimit = FTimingProfilerManager::UnlimitedEventDepth;
+	}
+	else
+	{
+		DepthLimit = 1;
+	}
+	FTimingProfilerManager::Get()->SetEventDepthLimit(DepthLimit);
+	Viewport.AddDirtyFlags(ETimingTrackViewportDirtyFlags::HInvalidated);
+}
+
+////////////////////////////////////////////////////////////////////////////////////////////////////
+
+void STimingView::SetEventDepthLimit(uint32 DepthLimit)
+{
+	FTimingProfilerManager::Get()->SetEventDepthLimit(DepthLimit);
+	Viewport.AddDirtyFlags(ETimingTrackViewportDirtyFlags::HInvalidated);
+}
+
+////////////////////////////////////////////////////////////////////////////////////////////////////
+
+bool STimingView::CheckEventDepthLimit(uint32 DepthLimit)
+{
+	return DepthLimit == FTimingProfilerManager::Get()->GetEventDepthLimit();
+}
+
+////////////////////////////////////////////////////////////////////////////////////////////////////
+
+void STimingView::CreateCpuThreadTrackColoringModeMenu(FMenuBuilder& MenuBuilder)
+{
+	MenuBuilder.BeginSection("CpuThreadTrackColoringMode", LOCTEXT("CpuThreadTrackColoringModeHeading", "Coloring Mode (CPU Thread Tracks)"));
+	{
+		MenuBuilder.AddMenuEntry(
+			LOCTEXT("CpuThreadTrackColoringMode_ByTimerName", "By Timer Name"),
+			LOCTEXT("CpuThreadTrackColoringMode_ByTimerName_Desc", "Assign a color to CPU/GPU timing events based on their timer name."),
+			FSlateIcon(),
+			FUIAction(
+				FExecuteAction::CreateSP(this, &STimingView::SetCpuThreadTrackColoringMode, Insights::ETimingEventsColoringMode::ByTimerName),
+				FCanExecuteAction(),
+				FIsActionChecked::CreateSP(this, &STimingView::CheckCpuThreadTrackColoringMode, Insights::ETimingEventsColoringMode::ByTimerName)),
+			NAME_None,
+			EUserInterfaceActionType::RadioButton
+		);
+		MenuBuilder.AddMenuEntry(
+			LOCTEXT("CpuThreadTrackColoringMode_ByTimerId", "By Timer Id"),
+			LOCTEXT("CpuThreadTrackColoringMode_ByTimerId_Desc", "Assign a color to CPU/GPU timing events based on their timer id."),
+			FSlateIcon(),
+			FUIAction(
+				FExecuteAction::CreateSP(this, &STimingView::SetCpuThreadTrackColoringMode, Insights::ETimingEventsColoringMode::ByTimerId),
+				FCanExecuteAction(),
+				FIsActionChecked::CreateSP(this, &STimingView::CheckCpuThreadTrackColoringMode, Insights::ETimingEventsColoringMode::ByTimerId)),
+			NAME_None,
+			EUserInterfaceActionType::RadioButton
+		);
+		MenuBuilder.AddMenuEntry(
+			LOCTEXT("CpuThreadTrackColoringMode_ByDuration", "By Duration"),
+			LOCTEXT("CpuThreadTrackColoringMode_ByDuration_Desc", "Assign a color to CPU/GPU timing events based on their duration (inclusive time).\nred: > 10 ms,   yellow: (1 ms .. 10 ms],   green: (100 us .. 1 ms],   blue: ≤ 100 us"),
+			FSlateIcon(),
+			FUIAction(
+				FExecuteAction::CreateSP(this, &STimingView::SetCpuThreadTrackColoringMode, Insights::ETimingEventsColoringMode::ByDuration),
+				FCanExecuteAction(),
+				FIsActionChecked::CreateSP(this, &STimingView::CheckCpuThreadTrackColoringMode, Insights::ETimingEventsColoringMode::ByDuration)),
+			NAME_None,
+			EUserInterfaceActionType::RadioButton
+		);
+	}
+	MenuBuilder.EndSection();
+}
+
+////////////////////////////////////////////////////////////////////////////////////////////////////
+
+void STimingView::ChooseNextCpuThreadTrackColoringMode()
+{
+	uint32 Mode = (uint32)FTimingProfilerManager::Get()->GetColoringMode();
+	Mode = (Mode + 1) % (uint32)Insights::ETimingEventsColoringMode::Count;
+	FTimingProfilerManager::Get()->SetColoringMode((Insights::ETimingEventsColoringMode)Mode);
+	Viewport.AddDirtyFlags(ETimingTrackViewportDirtyFlags::HInvalidated);
+}
+
+////////////////////////////////////////////////////////////////////////////////////////////////////
+
+void STimingView::SetCpuThreadTrackColoringMode(Insights::ETimingEventsColoringMode Mode)
+{
+	FTimingProfilerManager::Get()->SetColoringMode(Mode);
+	Viewport.AddDirtyFlags(ETimingTrackViewportDirtyFlags::HInvalidated);
+}
+
+////////////////////////////////////////////////////////////////////////////////////////////////////
+
+bool STimingView::CheckCpuThreadTrackColoringMode(Insights::ETimingEventsColoringMode Mode)
+{
+	return Mode == FTimingProfilerManager::Get()->GetColoringMode();
+}
+
+////////////////////////////////////////////////////////////////////////////////////////////////////
+
 bool STimingView::ToggleTrackVisibility_IsChecked(uint64 InTrackId) const
 {
 	const TSharedPtr<FBaseTimingTrack>* const TrackPtrPtr = AllTracks.Find(InTrackId);
diff --git a/Engine/Source/Developer/TraceInsights/Private/Insights/Widgets/STimingView.h b/Engine/Source/Developer/TraceInsights/Private/Insights/Widgets/STimingView.h
--- a/Engine/Source/Developer/TraceInsights/Private/Insights/Widgets/STimingView.h
+++ b/Engine/Source/Developer/TraceInsights/Private/Insights/Widgets/STimingView.h
@@ -46,6 +46,7 @@ namespace Insights
 	class FTimeMarker;
 	class FQuickFind;
 	class SQuickFind;
+	enum class ETimingEventsColoringMode : uint32;
 }
 
 /** A custom widget used to display timing events. */
@@ -71,20 +72,6 @@ public:
 	 */
 	void Construct(const FArguments& InArgs);
 
-	void CreateCompactMenuLine(FMenuBuilder& MenuBuilder, FText Label, TSharedRef<SWidget> InnerWidget) const;
-	TSharedRef<SWidget> MakeCompactAutoScrollOptionsMenu();
-	TSharedRef<SWidget> MakeAutoScrollOptionsMenu();
-
-	TSharedRef<SWidget> MakeAllTracksMenu();
-	void CreateAllTracksMenu(FMenuBuilder& MenuBuilder);
-	TSharedRef<SWidget> MakeCpuGpuTracksFilterMenu();
-	TSharedRef<SWidget> MakeOtherTracksFilterMenu();
-	TSharedRef<SWidget> MakePluginTracksFilterMenu();
-	TSharedRef<SWidget> MakeViewModeMenu();
-
-	bool ShowHideGraphTrack_IsChecked() const;
-	void ShowHideGraphTrack_Execute();
-
 	bool IsCompactModeEnabled() const;
 	void ToggleCompactMode();
 
@@ -385,12 +372,39 @@ protected:
 		return FVector2D(16.0f, 16.0f);
 	}
 
-	void ShowContextMenu(const FPointerEvent& MouseEvent);
-	void CreateTrackLocationMenu(FMenuBuilder& MenuBuilder, TSharedRef<FBaseTimingTrack> Track);
-
 	/** Binds our UI commands to delegates. */
 	void BindCommands();
 
+	void CreateCompactMenuLine(FMenuBuilder& MenuBuilder, FText Label, TSharedRef<SWidget> InnerWidget) const;
+	TSharedRef<SWidget> MakeCompactAutoScrollOptionsMenu();
+	TSharedRef<SWidget> MakeAutoScrollOptionsMenu();
+
+	TSharedRef<SWidget> MakeAllTracksMenu();
+	void CreateAllTracksMenu(FMenuBuilder& MenuBuilder);
+
+	TSharedRef<SWidget> MakeCpuGpuTracksFilterMenu();
+
+	TSharedRef<SWidget> MakeOtherTracksFilterMenu();
+	bool ShowHideGraphTrack_IsChecked() const;
+	void ShowHideGraphTrack_Execute();
+
+	TSharedRef<SWidget> MakePluginTracksFilterMenu();
+
+	TSharedRef<SWidget> MakeViewModeMenu();
+
+	void CreateDepthLimitMenu(FMenuBuilder& MenuBuilder);
+	void ChooseNextEventDepthLimit();
+	void SetEventDepthLimit(uint32 DepthLimit);
+	bool CheckEventDepthLimit(uint32 DepthLimit);
+
+	void CreateCpuThreadTrackColoringModeMenu(FMenuBuilder& MenuBuilder);
+	void ChooseNextCpuThreadTrackColoringMode();
+	void SetCpuThreadTrackColoringMode(Insights::ETimingEventsColoringMode Mode);
+	bool CheckCpuThreadTrackColoringMode(Insights::ETimingEventsColoringMode Mode);
+
+	void ShowContextMenu(const FPointerEvent& MouseEvent);
+	void CreateTrackLocationMenu(FMenuBuilder& MenuBuilder, TSharedRef<FBaseTimingTrack> Track);
+
 	////////////////////////////////////////////////////////////////////////////////////////////////////
 	// Auto-Scroll
 
