[Insights] Fixed typo.

diff --git a/Engine/Source/Developer/TraceInsights/Private/Insights/NetworkingProfiler/Widgets/SPacketView.cpp b/Engine/Source/Developer/TraceInsights/Private/Insights/NetworkingProfiler/Widgets/SPacketView.cpp
--- a/Engine/Source/Developer/TraceInsights/Private/Insights/NetworkingProfiler/Widgets/SPacketView.cpp
+++ b/Engine/Source/Developer/TraceInsights/Private/Insights/NetworkingProfiler/Widgets/SPacketView.cpp
@@ -751,7 +751,7 @@ int32 SPacketView::OnPaint(const FPaintArgs& Args, const FGeometry& AllottedGeom
 											"({1})\n"
 											"Largest Packet\n"
 											"    Index: {2}\n"
-											"    Sequance Number: {3}\n"
+											"    Sequence Number: {3}\n"
 											"    Content Size: {4} bits\n"
 											"    Total Size: {5} bytes ({6} unused bits)\n"
 											"    Timestamp: {7}\n"
