[Insights] Networking Insights: Fixed UI for the NetStats view.

diff --git a/Engine/Source/Developer/TraceInsights/Private/Insights/NetworkingProfiler/Widgets/SNetStatsView.cpp b/Engine/Source/Developer/TraceInsights/Private/Insights/NetworkingProfiler/Widgets/SNetStatsView.cpp
--- a/Engine/Source/Developer/TraceInsights/Private/Insights/NetworkingProfiler/Widgets/SNetStatsView.cpp
+++ b/Engine/Source/Developer/TraceInsights/Private/Insights/NetworkingProfiler/Widgets/SNetStatsView.cpp
@@ -95,9 +95,9 @@ void SNetStatsView::Construct(const FArguments& InArgs, TSharedPtr<SNetworkingPr
 
 				// Search box
 				+ SHorizontalBox::Slot()
-				.VAlign(VAlign_Center)
 				.Padding(2.0f)
 				.FillWidth(1.0f)
+				.VAlign(VAlign_Center)
 				[
 					SAssignNew(SearchBox, SSearchBox)
 					.HintText(LOCTEXT("SearchBoxHint", "Search net events or groups"))
@@ -108,9 +108,9 @@ void SNetStatsView::Construct(const FArguments& InArgs, TSharedPtr<SNetworkingPr
 
 				// Filter out net event types with zero instance count
 				+ SHorizontalBox::Slot()
-				.VAlign(VAlign_Center)
 				.Padding(2.0f)
 				.AutoWidth()
+				.VAlign(VAlign_Center)
 				[
 					SNew(SCheckBox)
 					.Style(FCoreStyle::Get(), "ToggleButtonCheckbox")
@@ -135,7 +135,8 @@ void SNetStatsView::Construct(const FArguments& InArgs, TSharedPtr<SNetworkingPr
 				SNew(SHorizontalBox)
 
 				+ SHorizontalBox::Slot()
-				.FillWidth(1.0f)
+				.AutoWidth()
+				.Padding(0.0f, 0.0f, 4.0f, 0.0f)
 				.VAlign(VAlign_Center)
 				[
 					SNew(STextBlock)
@@ -143,58 +144,30 @@ void SNetStatsView::Construct(const FArguments& InArgs, TSharedPtr<SNetworkingPr
 				]
 
 				+ SHorizontalBox::Slot()
-				.FillWidth(2.0f)
+				.AutoWidth()
 				.VAlign(VAlign_Center)
 				[
-					SAssignNew(GroupByComboBox, SComboBox<TSharedPtr<ENetEventGroupingMode>>)
-					.ToolTipText(this, &SNetStatsView::GroupBy_GetSelectedTooltipText)
-					.OptionsSource(&GroupByOptionsSource)
-					.OnSelectionChanged(this, &SNetStatsView::GroupBy_OnSelectionChanged)
-					.OnGenerateWidget(this, &SNetStatsView::GroupBy_OnGenerateWidget)
+					SNew(SBox)
+					.MinDesiredWidth(128.0f)
 					[
-						SNew(STextBlock)
-						.Text(this, &SNetStatsView::GroupBy_GetSelectedText)
+						SAssignNew(GroupByComboBox, SComboBox<TSharedPtr<ENetEventGroupingMode>>)
+						.ToolTipText(this, &SNetStatsView::GroupBy_GetSelectedTooltipText)
+						.OptionsSource(&GroupByOptionsSource)
+						.OnSelectionChanged(this, &SNetStatsView::GroupBy_OnSelectionChanged)
+						.OnGenerateWidget(this, &SNetStatsView::GroupBy_OnGenerateWidget)
+						[
+							SNew(STextBlock)
+							.Text(this, &SNetStatsView::GroupBy_GetSelectedText)
+						]
 					]
 				]
 			]
-
-			// TODO: Check boxes for: NetEvent, ...
-			/*
-			+ SVerticalBox::Slot()
-			.VAlign(VAlign_Center)
-			.Padding(2.0f)
-			.AutoHeight()
-			[
-				SNew(SHorizontalBox)
-
-				+ SHorizontalBox::Slot()
-				.Padding(FMargin(0.0f,0.0f,1.0f,0.0f))
-				.FillWidth(1.0f)
-				[
-					GetToggleButtonForNetEventType(ENetEventNodeType::NetEvent)
-				]
-
-				//+ SHorizontalBox::Slot()
-				//.Padding(FMargin(1.0f,0.0f,1.0f,0.0f))
-				//.FillWidth(1.0f)
-				//[
-				//	GetToggleButtonForNetEventType(ENetEventNodeType::NetEvent1)
-				//]
-
-				//+ SHorizontalBox::Slot()
-				.//Padding(FMargin(1.0f,0.0f,1.0f,0.0f))
-				//.FillWidth(1.0f)
-				//[
-				//	GetToggleButtonForNetEventType(ENetEventNodeType::NetEvent2)
-				//]
-			]
-			*/
 		]
 
 		// Tree view
 		+ SVerticalBox::Slot()
 		.FillHeight(1.0f)
-		.Padding(0.0f, 6.0f, 0.0f, 0.0f)
+		.Padding(0.0f, 2.0f, 0.0f, 0.0f)
 		[
 			SNew(SHorizontalBox)
 
@@ -317,19 +290,13 @@ TSharedPtr<SWidget> SNetStatsView::TreeView_GetMenuContent()
 	MenuBuilder.BeginSection("Misc", LOCTEXT("ContextMenu_Header_Misc", "Miscellaneous"));
 	{
 		/*TODO
-		FUIAction Action_CopySelectedToClipboard
-		(
-			FExecuteAction::CreateSP(this, &SNetStatsView::ContextMenu_CopySelectedToClipboard_Execute),
-			FCanExecuteAction::CreateSP(this, &SNetStatsView::ContextMenu_CopySelectedToClipboard_CanExecute)
-		);
 		MenuBuilder.AddMenuEntry
 		(
-			LOCTEXT("ContextMenu_Header_Misc_CopySelectedToClipboard", "Copy To Clipboard"),
-			LOCTEXT("ContextMenu_Header_Misc_CopySelectedToClipboard_Desc", "Copies selection to clipboard"),
-			FSlateIcon(FCoreStyle::Get().GetStyleSetName(), "GenericCommands.Copy"),
-			Action_CopySelectedToClipboard,
+			FTimersViewCommands::Get().Command_CopyToClipboard,
 			NAME_None,
-			EUserInterfaceActionType::Button
+			TAttribute<FText>(),
+			TAttribute<FText>(),
+			FSlateIcon(FCoreStyle::Get().GetStyleSetName(), "GenericCommands.Copy")
 		);
 		*/
 
@@ -621,12 +588,12 @@ void SNetStatsView::UpdateTree()
 	FStopwatch Stopwatch;
 	Stopwatch.Start();
 
-	CreateSortedGroups();
+	CreateGroups();
 
 	Stopwatch.Update();
 	const double Time1 = Stopwatch.GetAccumulatedTime();
 
-	SortTreeChildNodes();
+	SortTreeNodes();
 
 	Stopwatch.Update();
 	const double Time2 = Stopwatch.GetAccumulatedTime();
@@ -657,17 +624,17 @@ void SNetStatsView::ApplyFiltering()
 		const bool bIsGroupVisible = Filters->PassesAllFilters(GroupPtr);
 
 		const TArray<Insights::FBaseTreeNodePtr>& GroupChildren = GroupPtr->GetChildren();
-		const int32 NumChildren = GroupChildren.Num();
 		int32 NumVisibleChildren = 0;
-		for (int32 Cx = 0; Cx < NumChildren; ++Cx)
+		for (const Insights::FBaseTreeNodePtr& ChildPtr : GroupChildren)
 		{
-			// Add a child.
-			const FNetEventNodePtr& NodePtr = StaticCastSharedPtr<FNetEventNode, Insights::FBaseTreeNode>(GroupChildren[Cx]);
+			const FNetEventNodePtr& NodePtr = StaticCastSharedPtr<FNetEventNode, Insights::FBaseTreeNode>(ChildPtr);
+
 			const bool bIsChildVisible = (!bFilterOutZeroCountEvents || NodePtr->GetAggregatedStats().InstanceCount > 0)
 									  && bNetEventTypeIsVisible[static_cast<int>(NodePtr->GetType())]
 									  && Filters->PassesAllFilters(NodePtr);
 			if (bIsChildVisible)
 			{
+				// Add a child.
 				GroupPtr->AddFilteredChild(NodePtr);
 				NumVisibleChildren++;
 			}
@@ -696,9 +663,8 @@ void SNetStatsView::ApplyFiltering()
 			bExpansionSaved = true;
 		}
 
-		for (int32 Fx = 0; Fx < FilteredGroupNodes.Num(); Fx++)
+		for (const FNetEventNodePtr& GroupPtr : FilteredGroupNodes)
 		{
-			const FNetEventNodePtr& GroupPtr = FilteredGroupNodes[Fx];
 			TreeView->SetItemExpansion(GroupPtr, GroupPtr->IsExpanded());
 		}
 	}
@@ -733,24 +699,24 @@ TSharedRef<SWidget> SNetStatsView::GetToggleButtonForNetEventType(const ENetEven
 {
 	return SNew(SCheckBox)
 		.Style(FCoreStyle::Get(), "ToggleButtonCheckbox")
+		.Padding(FMargin(4.0f, 2.0f, 4.0f, 2.0f))
 		.HAlign(HAlign_Center)
-		.Padding(2.0f)
 		.OnCheckStateChanged(this, &SNetStatsView::FilterByNetEventType_OnCheckStateChanged, NodeType)
 		.IsChecked(this, &SNetStatsView::FilterByNetEventType_IsChecked, NodeType)
 		.ToolTipText(NetEventNodeTypeHelper::ToDescription(NodeType))
 		[
 			SNew(SHorizontalBox)
 
-			+ SHorizontalBox::Slot()
-			.AutoWidth()
-			.VAlign(VAlign_Center)
-			[
-				SNew(SImage)
-				.Image(NetEventNodeTypeHelper::GetIconForNetEventNodeType(NodeType))
-			]
+			//+ SHorizontalBox::Slot()
+			//.Padding(0.0f, 0.0f, 2.0f, 0.0f)
+			//.AutoWidth()
+			//.VAlign(VAlign_Center)
+			//[
+			//	SNew(SImage)
+			//	.Image(NetEventNodeTypeHelper::GetIconForNetEventNodeType(NodeType))
+			//]
 
 			+ SHorizontalBox::Slot()
-			.Padding(2.0f, 0.0f, 0.0f, 0.0f)
 			.VAlign(VAlign_Center)
 			[
 				SNew(STextBlock)
@@ -956,7 +922,7 @@ bool SNetStatsView::SearchBox_IsEnabled() const
 // Grouping
 ////////////////////////////////////////////////////////////////////////////////////////////////////
 
-void SNetStatsView::CreateSortedGroups()
+void SNetStatsView::CreateGroups()
 {
 	if (GroupingMode == ENetEventGroupingMode::Flat)
 	{
@@ -1060,8 +1026,8 @@ void SNetStatsView::GroupBy_OnSelectionChanged(TSharedPtr<ENetEventGroupingMode>
 	{
 		GroupingMode = *NewGroupingMode;
 
-		CreateSortedGroups();
-		SortTreeChildNodes();
+		CreateGroups();
+		SortTreeNodes();
 		ApplyFiltering();
 	}
 }
@@ -1137,11 +1103,10 @@ void SNetStatsView::UpdateCurrentSortingByColumn()
 
 ////////////////////////////////////////////////////////////////////////////////////////////////////
 
-void SNetStatsView::SortTreeChildNodes()
+void SNetStatsView::SortTreeNodes()
 {
 	if (CurrentSorter.IsValid())
 	{
-		// Sort nodes in each group.
 		for (FNetEventNodePtr& Root : GroupNodes)
 		{
 			SortTreeNodesRec(*Root, *CurrentSorter);
@@ -1164,7 +1129,6 @@ void SNetStatsView::SortTreeNodesRec(FNetEventNode& Node, const Insights::ITable
 
 	for (Insights::FBaseTreeNodePtr ChildPtr : Node.GetChildren())
 	{
-		//if (ChildPtr->IsGroup())
 		if (ChildPtr->GetChildren().Num() > 0)
 		{
 			SortTreeNodesRec(*StaticCastSharedPtr<FNetEventNode>(ChildPtr), Sorter);
@@ -1192,7 +1156,7 @@ void SNetStatsView::SetSortModeForColumn(const FName& ColumnId, const EColumnSor
 	ColumnSortMode = SortMode;
 	UpdateCurrentSortingByColumn();
 
-	SortTreeChildNodes();
+	SortTreeNodes();
 	ApplyFiltering();
 }
 
@@ -1362,7 +1326,7 @@ void SNetStatsView::HideColumn(const FName ColumnId)
 // ToggleColumn action
 ////////////////////////////////////////////////////////////////////////////////////////////////////
 
-bool SNetStatsView::IsColumnVisible(const FName ColumnId)
+bool SNetStatsView::IsColumnVisible(const FName ColumnId) const
 {
 	const Insights::FTableColumn& Column = *Table->FindColumnChecked(ColumnId);
 	return Column.IsVisible();
@@ -1614,7 +1578,7 @@ void SNetStatsView::RebuildTree(bool bResync)
 	if (TotalTime > 0.01)
 	{
 		const double SyncTime = SyncStopwatch.GetAccumulatedTime();
-		UE_LOG(NetworkingProfiler, Log, TEXT("[NetStats] Tree view rebuilt in %.3fs (%.3fs + %.3fs) --> %d net events (%d added)"),
+		UE_LOG(NetworkingProfiler, Log, TEXT("[NetStats] Tree view rebuilt in %.4fs (%.4fs + %.4fs) --> %d net events (%d added)"),
 			TotalTime, SyncTime, TotalTime - SyncTime, NetEventNodes.Num(), NetEventNodes.Num() - PreviousNodeCount);
 	}
 }
diff --git a/Engine/Source/Developer/TraceInsights/Private/Insights/NetworkingProfiler/Widgets/SNetStatsView.h b/Engine/Source/Developer/TraceInsights/Private/Insights/NetworkingProfiler/Widgets/SNetStatsView.h
--- a/Engine/Source/Developer/TraceInsights/Private/Insights/NetworkingProfiler/Widgets/SNetStatsView.h
+++ b/Engine/Source/Developer/TraceInsights/Private/Insights/NetworkingProfiler/Widgets/SNetStatsView.h
@@ -175,7 +175,7 @@ protected:
 	////////////////////////////////////////////////////////////////////////////////////////////////////
 	// Grouping
 
-	void CreateSortedGroups();
+	void CreateGroups();
 	void CreateGroupByOptionsSources();
 
 	void GroupBy_OnSelectionChanged(TSharedPtr<ENetEventGroupingMode> NewGroupingMode, ESelectInfo::Type SelectInfo);
@@ -195,7 +195,7 @@ protected:
 	void CreateSortings();
 
 	void UpdateCurrentSortingByColumn();
-	void SortTreeChildNodes();
+	void SortTreeNodes();
 	void SortTreeNodesRec(FNetEventNode& Node, const Insights::ITableCellValueSorter& Sorter);
 
 	EColumnSortMode::Type GetSortModeForColumn(const FName ColumnId) const;
@@ -232,7 +232,7 @@ protected:
 	void HideColumn(const FName ColumnId);
 
 	// ToggleColumnVisibility
-	bool IsColumnVisible(const FName ColumnId);
+	bool IsColumnVisible(const FName ColumnId) const;
 	bool CanToggleColumnVisibility(const FName ColumnId) const;
 	void ToggleColumnVisibility(const FName ColumnId);
 
@@ -306,7 +306,7 @@ protected:
 
 	//bool bUseFiltering;
 
-	/** The search box widget used to filter items displayed in the stats and groups tree. */
+	/** The search box widget used to filter items displayed in the tree. */
 	TSharedPtr<SSearchBox> SearchBox;
 
 	/** The text based filter. */
