[Insights] Networking Insights / Packets view: Fixed the draw of the highlight border for selected packet.

diff --git a/Engine/Source/Developer/TraceInsights/Private/Insights/NetworkingProfiler/ViewModels/PacketViewDrawHelper.cpp b/Engine/Source/Developer/TraceInsights/Private/Insights/NetworkingProfiler/ViewModels/PacketViewDrawHelper.cpp
--- a/Engine/Source/Developer/TraceInsights/Private/Insights/NetworkingProfiler/ViewModels/PacketViewDrawHelper.cpp
+++ b/Engine/Source/Developer/TraceInsights/Private/Insights/NetworkingProfiler/ViewModels/PacketViewDrawHelper.cpp
@@ -14,6 +14,8 @@
 
 #include <limits>
 
+#define INSIGHTS_USE_LEGACY_BORDER 0
+
 ////////////////////////////////////////////////////////////////////////////////////////////////////
 // FNetworkPacketAggregatedSample
 ////////////////////////////////////////////////////////////////////////////////////////////////////
@@ -238,12 +240,23 @@ void FPacketViewDrawHelper::DrawSampleHighlight(const FNetworkPacketAggregatedSa
 	const float H = ValueY - BaselineY;
 	const float Y = ViewHeight - H;
 
+#if INSIGHTS_USE_LEGACY_BORDER
+	constexpr float BorderOffset = 1.0f;
+#else
+	constexpr float BorderOffset = 2.0f;
+#endif
+
 	if (Mode == EHighlightMode::Hovered)
 	{
 		const FLinearColor Color(1.0f, 1.0f, 0.0f, 1.0f); // yellow
 
 		// Draw border around the hovered box.
-		DrawContext.DrawBox(X - 1.0f, Y - 1.0f, SampleW + 2.0f, H + 2.0f, HoveredEventBorderBrush, Color);
+#if INSIGHTS_USE_LEGACY_BORDER
+		DrawContext.DrawBox(X - BorderOffset, Y - BorderOffset, SampleW + 2 * BorderOffset, H + 2 * BorderOffset, HoveredEventBorderBrush, Color);
+#else
+		FSlateRoundedBoxBrush Brush(FLinearColor(0.0f, 0.0f, 0.0f, 0.0f), 2.0f, Color, 2.0f);
+		DrawContext.DrawBox(X - BorderOffset, Y - BorderOffset, SampleW + 2 * BorderOffset, H + 2 * BorderOffset, &Brush, FLinearColor(0.0f, 0.0f, 0.0f, 0.0f));
+#endif
 	}
 	else // EHighlightMode::Selected or EHighlightMode::SelectedAndHovered
 	{
@@ -255,7 +268,12 @@ void FPacketViewDrawHelper::DrawSampleHighlight(const FNetworkPacketAggregatedSa
 		const FLinearColor Color(S, S, Blue, 1.0f);
 
 		// Draw border around the selected box.
-		DrawContext.DrawBox(X - 1.0f, Y - 1.0f, SampleW + 2.0f, H + 2.0f, SelectedEventBorderBrush, Color);
+#if INSIGHTS_USE_LEGACY_BORDER
+		DrawContext.DrawBox(X - BorderOffset, Y - BorderOffset, SampleW + 2 * BorderOffset, H + 2 * BorderOffset, SelectedEventBorderBrush, Color);
+#else
+		FSlateRoundedBoxBrush Brush(FLinearColor(0.0f, 0.0f, 0.0f, 0.0f), 2.0f, Color, 2.0f);
+		DrawContext.DrawBox(X - BorderOffset, Y - BorderOffset, SampleW + 2 * BorderOffset, H + 2 * BorderOffset, &Brush, FLinearColor(0.0f, 0.0f, 0.0f, 0.0f));
+#endif
 	}
 	DrawContext.LayerId++;
 }
@@ -305,3 +323,5 @@ void FPacketViewDrawHelper::DrawSelection(int32 StartPacketIndex, int32 EndPacke
 }
 
 ////////////////////////////////////////////////////////////////////////////////////////////////////
+
+#undef INSIGHTS_USE_LEGACY_BORDER
