[Insights]
  - Networking Insights: Updated main toolbar (using placeholder icons).
  - Networking Insights: Updated toolbar for Packet Content view.

diff --git a/Engine/Source/Developer/TraceInsights/Private/Insights/InsightsStyle.cpp b/Engine/Source/Developer/TraceInsights/Private/Insights/InsightsStyle.cpp
--- a/Engine/Source/Developer/TraceInsights/Private/Insights/InsightsStyle.cpp
+++ b/Engine/Source/Developer/TraceInsights/Private/Insights/InsightsStyle.cpp
@@ -19,6 +19,8 @@
 #define CORE_BOX_BRUSH_SVG(RelativePath, ...)    FSlateVectorBoxBrush   (Style.RootToCoreContentDir(RelativePath, TEXT(".svg")), __VA_ARGS__)
 #define CORE_BORDER_BRUSH_SVG(RelativePath, ...) FSlateVectorBorderBrush(Style.RootToCoreContentDir(RelativePath, TEXT(".svg")), __VA_ARGS__)
 
+#define TODO_IMAGE_BRUSH(...) IMAGE_BRUSH_SVG("../../Editor/Slate/Starship/Common/StaticMesh", __VA_ARGS__)
+
 TSharedPtr<FSlateStyleSet> FInsightsStyle::StyleInstance = nullptr;
 
 void FInsightsStyle::Initialize()
@@ -116,56 +118,54 @@ TSharedRef<FSlateStyleSet> FInsightsStyle::Create()
 	Style.Set("LogView.Icon.Large", new IMAGE_BRUSH_SVG("../../Editor/Slate/Starship/Common/Log", Icon32x32));
 	Style.Set("LogView.Icon.Small", new IMAGE_BRUSH_SVG("../../Editor/Slate/Starship/Common/Log", Icon16x16));
 
-	Style.Set("TableTreeView.Icon.Large", new IMAGE_BRUSH("Icons/Profiler/Profiler_Data_Capture_40x", Icon32x32));
-	Style.Set("TableTreeView.Icon.Small", new IMAGE_BRUSH("Icons/Profiler/Profiler_Data_Capture_40x", Icon16x16));
+	Style.Set("TableTreeView.Icon.Large", new TODO_IMAGE_BRUSH(Icon32x32));
+	Style.Set("TableTreeView.Icon.Small", new TODO_IMAGE_BRUSH(Icon16x16));
 
 	//////////////////////////////////////////////////
 	// Asset Loading Insights
 
-	Style.Set("LoadingProfiler.Icon.Large", new IMAGE_BRUSH("Icons/Profiler/profiler_stats_40x", Icon32x32));
-	Style.Set("LoadingProfiler.Icon.Small", new IMAGE_BRUSH("Icons/Profiler/profiler_stats_40x", Icon16x16));
+	Style.Set("LoadingProfiler.Icon.Large", new TODO_IMAGE_BRUSH(Icon32x32));
+	Style.Set("LoadingProfiler.Icon.Small", new TODO_IMAGE_BRUSH(Icon16x16));
 
 	//////////////////////////////////////////////////
 	// Networking Insights
 
-	Style.Set("NetworkingProfiler.Icon.Large", new IMAGE_BRUSH("Icons/Profiler/profiler_stats_40x", Icon32x32));
-	Style.Set("NetworkingProfiler.Icon.Small", new IMAGE_BRUSH("Icons/Profiler/profiler_stats_40x", Icon16x16));
+	Style.Set("NetworkingProfiler.Icon.Large", new TODO_IMAGE_BRUSH(Icon32x32));
+	Style.Set("NetworkingProfiler.Icon.Small", new TODO_IMAGE_BRUSH(Icon16x16));
 
-	Style.Set("PacketOveriew.Icon.Large", new IMAGE_BRUSH("Icons/Profiler/profiler_stats_40x", Icon32x32));
-	Style.Set("PacketOveriew.Icon.Small", new IMAGE_BRUSH("Icons/Profiler/profiler_stats_40x", Icon16x16));
+	Style.Set("PacketView.Icon.Large", new TODO_IMAGE_BRUSH(Icon32x32));
+	Style.Set("PacketView.Icon.Small", new TODO_IMAGE_BRUSH(Icon16x16));
 
-	Style.Set("PacketContentView.Icon.Large", new IMAGE_BRUSH("Icons/Profiler/profiler_stats_40x", Icon32x32));
-	Style.Set("PacketContentView.Icon.Small", new IMAGE_BRUSH("Icons/Profiler/profiler_stats_40x", Icon16x16));
+	Style.Set("PacketContentView.Icon.Large", new TODO_IMAGE_BRUSH(Icon32x32));
+	Style.Set("PacketContentView.Icon.Small", new TODO_IMAGE_BRUSH(Icon16x16));
 
-	Style.Set("NetStatsView.Icon.Large", new IMAGE_BRUSH("Icons/Profiler/profiler_stats_40x", Icon32x32));
-	Style.Set("NetStatsView.Icon.Small", new IMAGE_BRUSH("Icons/Profiler/profiler_stats_40x", Icon16x16));
+	Style.Set("NetStatsView.Icon.Large", new TODO_IMAGE_BRUSH(Icon32x32));
+	Style.Set("NetStatsView.Icon.Small", new TODO_IMAGE_BRUSH(Icon16x16));
 
 	//////////////////////////////////////////////////
 	// Memory Insights
 
-	Style.Set("MemoryProfiler.Icon.Large", new IMAGE_BRUSH("Icons/Profiler/profiler_stats_40x", Icon32x32));
-	Style.Set("MemoryProfiler.Icon.Small", new IMAGE_BRUSH("Icons/Profiler/profiler_stats_40x", Icon16x16));
+	Style.Set("MemoryProfiler.Icon.Large", new TODO_IMAGE_BRUSH(Icon32x32));
+	Style.Set("MemoryProfiler.Icon.Small", new TODO_IMAGE_BRUSH(Icon16x16));
 
-	Style.Set("MemInvestigationView.Icon.Large", new IMAGE_BRUSH("Icons/Profiler/profiler_stats_40x", Icon32x32));
-	Style.Set("MemInvestigationView.Icon.Small", new IMAGE_BRUSH("Icons/Profiler/profiler_stats_40x", Icon16x16));
+	Style.Set("MemInvestigationView.Icon.Large", new TODO_IMAGE_BRUSH(Icon32x32));
+	Style.Set("MemInvestigationView.Icon.Small", new TODO_IMAGE_BRUSH(Icon16x16));
 
-	Style.Set("MemAllocTableTreeView.Icon.Large", new IMAGE_BRUSH("Icons/Profiler/profiler_stats_40x", Icon32x32));
-	Style.Set("MemAllocTableTreeView.Icon.Small", new IMAGE_BRUSH("Icons/Profiler/profiler_stats_40x", Icon16x16));
+	Style.Set("MemAllocTableTreeView.Icon.Large", new TODO_IMAGE_BRUSH(Icon32x32));
+	Style.Set("MemAllocTableTreeView.Icon.Small", new TODO_IMAGE_BRUSH(Icon16x16));
 
-	Style.Set("MemTagTreeView.Icon.Large", new IMAGE_BRUSH("Icons/Profiler/profiler_stats_40x", Icon32x32));
-	Style.Set("MemTagTreeView.Icon.Small", new IMAGE_BRUSH("Icons/Profiler/profiler_stats_40x", Icon16x16));
+	Style.Set("MemTagTreeView.Icon.Large", new TODO_IMAGE_BRUSH(Icon32x32));
+	Style.Set("MemTagTreeView.Icon.Small", new TODO_IMAGE_BRUSH(Icon16x16));
 
 	Style.Set("Mem.Add.Small", new IMAGE_BRUSH("Icons/icon_Cascade_AddLOD2_40x", Icon20x20));
 	Style.Set("Mem.Remove.Small", new IMAGE_BRUSH("Icons/icon_Cascade_DeleteLOD_40x", Icon20x20));
 
 	//////////////////////////////////////////////////
 
-	Style.Set("FindFirst", new IMAGE_BRUSH("Animation/backward_end", Icon20x20));
-	//Style.Set("FindPrevious", new IMAGE_BRUSH("Animation/backward", Icon20x20));
+	Style.Set("FindFirst", new CORE_IMAGE_BRUSH_SVG("Starship/Common/arrow-left", Icon20x20));
 	Style.Set("FindPrevious", new IMAGE_BRUSH("../../Editor/Slate/Icons/GeneralTools/Previous_40x", Icon20x20));
-	//Style.Set("FindNext", new IMAGE_BRUSH("Animation/forward", Icon20x20));
 	Style.Set("FindNext", new IMAGE_BRUSH("../../Editor/Slate/Icons/GeneralTools/Next_40x", Icon20x20));
-	Style.Set("FindLast", new IMAGE_BRUSH("Animation/forward_end", Icon20x20));
+	Style.Set("FindLast", new CORE_IMAGE_BRUSH_SVG("Starship/Common/arrow-right", Icon20x20));
 
 	//////////////////////////////////////////////////
 
@@ -190,6 +190,8 @@ TSharedRef<FSlateStyleSet> FInsightsStyle::Create()
 	return StyleRef;
 }
 
+#undef TODO_IMAGE_BRUSH
+
 #undef IMAGE_BRUSH
 #undef BOX_BRUSH
 #undef BORDER_BRUSH
diff --git a/Engine/Source/Developer/TraceInsights/Private/Insights/NetworkingProfiler/ViewModels/PacketContentViewDrawHelper.cpp b/Engine/Source/Developer/TraceInsights/Private/Insights/NetworkingProfiler/ViewModels/PacketContentViewDrawHelper.cpp
--- a/Engine/Source/Developer/TraceInsights/Private/Insights/NetworkingProfiler/ViewModels/PacketContentViewDrawHelper.cpp
+++ b/Engine/Source/Developer/TraceInsights/Private/Insights/NetworkingProfiler/ViewModels/PacketContentViewDrawHelper.cpp
@@ -203,7 +203,7 @@ void FPacketContentViewDrawStateBuilder::AddEvent(const TraceServices::FNetProfi
 				}
 				else if (Event.BunchInfo.bClose)
 				{
-					Builder.Appendf(TEXT(" | Close: %s"), LexToString(TraceServices::ENetProfilerChannelCloseReason(Info.ChannelCloseReason)));		
+					Builder.Appendf(TEXT(" | Close: %s"), LexToString(TraceServices::ENetProfilerChannelCloseReason(Info.ChannelCloseReason)));
 				}
 				if (Event.BunchInfo.bReliable)
 				{
@@ -309,9 +309,6 @@ FPacketContentViewDrawHelper::FPacketContentViewDrawHelper(const FDrawContext& I
 	, HoveredEventBorderBrush(FInsightsStyle::Get().GetBrush("HoveredEventBorder"))
 	, SelectedEventBorderBrush(FInsightsStyle::Get().GetBrush("SelectedEventBorder"))
 	, EventFont(FCoreStyle::GetDefaultFontStyle("Regular", 8))
-	, LayoutPosY(0.0f)
-	, LayoutEventH(14.0f)
-	, LayoutEventDY(2.0f)
 {
 }
 
@@ -373,59 +370,64 @@ FLinearColor FPacketContentViewDrawHelper::GetColorByType(int32 Type)
 void FPacketContentViewDrawHelper::Draw(const FPacketContentViewDrawState& DrawState, const float Opacity) const
 {
 	// Draw filled boxes (merged borders).
-	//if (LayoutEventH > 0.0f)
+	//if (Viewport.GetEventHeight() > 0.0f)
 	{
-		const float EventFillH = LayoutEventH;
+		const float BoxY0 = Viewport.GetTopEventPosY();
+		const float BoxDY = Viewport.GetEventHeight() + Viewport.GetEventDY();
+		const float EventFillH = Viewport.GetEventHeight();
 		for (const FPacketContentViewDrawState::FBoxPrimitive& Box : DrawState.Boxes)
 		{
-			const float Y = LayoutPosY + (LayoutEventH + LayoutEventDY) * Box.Depth;
-			DrawContext.DrawBox(Box.X, Y, Box.W, EventFillH, WhiteBrush, Box.Color.CopyWithNewOpacity(Opacity));
+			DrawContext.DrawBox(Box.X, BoxY0 + BoxDY * Box.Depth, Box.W, EventFillH, WhiteBrush, Box.Color.CopyWithNewOpacity(Opacity));
 		}
 		DrawContext.LayerId++;
 	}
 
 	// Draw filled boxes (event inside area).
-	if (LayoutEventH > 2.0f)
+	if (Viewport.GetEventHeight() > 2.0f)
 	{
-		const float EventFillH = LayoutEventH - 2.0f;
+		const float BoxY0 = Viewport.GetTopEventPosY() + 1.0f;
+		const float BoxDY = Viewport.GetEventHeight() + Viewport.GetEventDY();
+		const float EventFillH = Viewport.GetEventHeight() - 2.0f;
 		for (const FPacketContentViewDrawState::FBoxPrimitive& Box : DrawState.InsideBoxes)
 		{
-			const float Y = LayoutPosY + (LayoutEventH + LayoutEventDY) * Box.Depth + 1.0f;
-			DrawContext.DrawBox(Box.X, Y, Box.W, EventFillH, WhiteBrush, Box.Color.CopyWithNewOpacity(Opacity));
+			DrawContext.DrawBox(Box.X, BoxY0 + BoxDY * Box.Depth, Box.W, EventFillH, WhiteBrush, Box.Color.CopyWithNewOpacity(Opacity));
 		}
 		DrawContext.LayerId++;
 	}
 
 	// Draw borders.
-	//if (LayoutEventH > 0.0f)
+	//if (Viewport.GetEventHeight() > 0.0f)
 	{
 #if INSIGHTS_USE_LEGACY_BORDER
-		const float EventBorderH = LayoutEventH;
+		const float BoxY0 = Viewport.GetTopEventPosY();
+		const float BoxDY = Viewport.GetEventHeight() + Viewport.GetEventDY();
+		const float EventBorderH = Viewport.GetEventHeight();
 		for (const FPacketContentViewDrawState::FBoxPrimitive& Box : DrawState.Borders)
 		{
-			const float Y = LayoutPosY + (LayoutEventH + LayoutEventDY) * Box.Depth;
-			DrawContext.DrawBox(Box.X, Y, Box.W, EventBorderH, EventBorderBrush, Box.Color.CopyWithNewOpacity(Opacity));
+			DrawContext.DrawBox(Box.X, BoxY0 + BoxDY * Box.Depth, Box.W, EventBorderH, EventBorderBrush, Box.Color.CopyWithNewOpacity(Opacity));
 		}
 		DrawContext.LayerId++;
 #else
 		if (Opacity == 1.0f)
 		{
 			const int32 EventBorderLayerId = DrawContext.LayerId - 2; // draw under the filled boxes
-			const float EventBorderH = LayoutEventH;
+			const float BoxY0 = Viewport.GetTopEventPosY();
+			const float BoxDY = Viewport.GetEventHeight() + Viewport.GetEventDY();
+			const float EventBorderH = Viewport.GetEventHeight();
 			for (const FPacketContentViewDrawState::FBoxPrimitive& Box : DrawState.Borders)
 			{
-				const float Y = LayoutPosY + (LayoutEventH + LayoutEventDY) * Box.Depth;
-				DrawContext.DrawBox(EventBorderLayerId, Box.X, Y, Box.W, EventBorderH, WhiteBrush, Box.Color);
+				DrawContext.DrawBox(EventBorderLayerId, Box.X, BoxY0 + BoxDY * Box.Depth, Box.W, EventBorderH, WhiteBrush, Box.Color);
 			}
 		}
 		else
 		{
-			const float EventBorderH = LayoutEventH;
+			const float BoxY0 = Viewport.GetTopEventPosY();
+			const float BoxDY = Viewport.GetEventHeight() + Viewport.GetEventDY();
+			const float EventBorderH = Viewport.GetEventHeight();
 			for (const FPacketContentViewDrawState::FBoxPrimitive& Box : DrawState.Borders)
 			{
-				const float Y = LayoutPosY + (LayoutEventH + LayoutEventDY) * Box.Depth;
 				FSlateRoundedBoxBrush Brush(FLinearColor(0.0f, 0.0f, 0.0f, 0.0f), 0.0f, Box.Color.CopyWithNewOpacity(Opacity), 1.0f);
-				DrawContext.DrawBox(Box.X, Y, Box.W, EventBorderH, &Brush, FLinearColor(0.0f, 0.0f, 0.0f, 0.0f));
+				DrawContext.DrawBox(Box.X, BoxY0 + BoxDY * Box.Depth, Box.W, EventBorderH, &Brush, FLinearColor(0.0f, 0.0f, 0.0f, 0.0f));
 			}
 			DrawContext.LayerId++;
 		}
@@ -433,15 +435,15 @@ void FPacketContentViewDrawHelper::Draw(const FPacketContentViewDrawState& DrawS
 	}
 
 	// Draw texts.
-	if (LayoutEventH > 10.0f)
+	if (Viewport.GetEventHeight() > 10.0f)
 	{
+		const float TextY0 = Viewport.GetTopEventPosY() + 1.0f;
+		const float TextDY = Viewport.GetEventHeight() + Viewport.GetEventDY();
 		const FLinearColor WhiteColor(1.0f, 1.0f, 1.0f, Opacity);
 		const FLinearColor BlackColor(0.0f, 0.0f, 0.0f, Opacity);
-
 		for (const FPacketContentViewDrawState::FTextPrimitive& Text : DrawState.Texts)
 		{
-			const float Y = LayoutPosY + (LayoutEventH + LayoutEventDY) * Text.Depth + 1.0f;
-			DrawContext.DrawText(Text.X, Y, Text.Text, EventFont, Text.bWhite ? WhiteColor : BlackColor);
+			DrawContext.DrawText(Text.X, TextY0 + TextDY * Text.Depth, Text.Text, EventFont, Text.bWhite ? WhiteColor : BlackColor);
 		}
 		DrawContext.LayerId++;
 	}
@@ -486,7 +488,7 @@ void FPacketContentViewDrawHelper::DrawEventHighlight(const FNetworkPacketEvent&
 	}
 
 	const float EventW = EventX2 - EventX1;
-	const float EventY = LayoutPosY + (LayoutEventH + LayoutEventDY) * Event.Level;
+	const float EventY = Viewport.GetTopEventPosY() + (Viewport.GetEventHeight() + Viewport.GetEventDY()) * Event.Level;
 
 #if INSIGHTS_USE_LEGACY_BORDER
 	constexpr float BorderOffset = 1.0f;
@@ -500,10 +502,10 @@ void FPacketContentViewDrawHelper::DrawEventHighlight(const FNetworkPacketEvent&
 
 		// Draw border around the timing event box.
 #if INSIGHTS_USE_LEGACY_BORDER
-		DrawContext.DrawBox(EventX1 - BorderOffset, EventY - BorderOffset, EventW + 2 * BorderOffset, LayoutEventH + 2 * BorderOffset, HoveredEventBorderBrush, Color);
+		DrawContext.DrawBox(EventX1 - BorderOffset, EventY - BorderOffset, EventW + 2 * BorderOffset, Viewport.GetEventHeight() + 2 * BorderOffset, HoveredEventBorderBrush, Color);
 #else
 		FSlateRoundedBoxBrush Brush(FLinearColor(0.0f, 0.0f, 0.0f, 0.0f), 2.0f, Color, 2.0f);
-		DrawContext.DrawBox(EventX1 - BorderOffset, EventY - BorderOffset, EventW + 2 * BorderOffset, LayoutEventH + 2 * BorderOffset, &Brush, FLinearColor(0.0f, 0.0f, 0.0f, 0.0f));
+		DrawContext.DrawBox(EventX1 - BorderOffset, EventY - BorderOffset, EventW + 2 * BorderOffset, Viewport.GetEventHeight() + 2 * BorderOffset, &Brush, FLinearColor(0.0f, 0.0f, 0.0f, 0.0f));
 #endif
 	}
 	else // EHighlightMode::Selected or EHighlightMode::SelectedAndHovered
@@ -517,10 +519,10 @@ void FPacketContentViewDrawHelper::DrawEventHighlight(const FNetworkPacketEvent&
 
 		// Draw border around the timing event box.
 #if INSIGHTS_USE_LEGACY_BORDER
-		DrawContext.DrawBox(EventX1 - BorderOffset, EventY - BorderOffset, EventW + 2 * BorderOffset, LayoutEventH + 2 * BorderOffset, SelectedEventBorderBrush, Color);
+		DrawContext.DrawBox(EventX1 - BorderOffset, EventY - BorderOffset, EventW + 2 * BorderOffset, Viewport.GetEventHeight() + 2 * BorderOffset, SelectedEventBorderBrush, Color);
 #else
 		FSlateRoundedBoxBrush Brush(FLinearColor(0.0f, 0.0f, 0.0f, 0.0f), 2.0f, Color, 2.0f);
-		DrawContext.DrawBox(EventX1 - BorderOffset, EventY - BorderOffset, EventW + 2 * BorderOffset, LayoutEventH + 2 * BorderOffset, &Brush, FLinearColor(0.0f, 0.0f, 0.0f, 0.0f));
+		DrawContext.DrawBox(EventX1 - BorderOffset, EventY - BorderOffset, EventW + 2 * BorderOffset, Viewport.GetEventHeight() + 2 * BorderOffset, &Brush, FLinearColor(0.0f, 0.0f, 0.0f, 0.0f));
 #endif
 	}
 	DrawContext.LayerId++;
diff --git a/Engine/Source/Developer/TraceInsights/Private/Insights/NetworkingProfiler/ViewModels/PacketContentViewDrawHelper.h b/Engine/Source/Developer/TraceInsights/Private/Insights/NetworkingProfiler/ViewModels/PacketContentViewDrawHelper.h
--- a/Engine/Source/Developer/TraceInsights/Private/Insights/NetworkingProfiler/ViewModels/PacketContentViewDrawHelper.h
+++ b/Engine/Source/Developer/TraceInsights/Private/Insights/NetworkingProfiler/ViewModels/PacketContentViewDrawHelper.h
@@ -211,14 +211,6 @@ public:
 	const FSlateBrush* GetWhiteBrush() const { return WhiteBrush; }
 	const FSlateFontInfo& GetEventFont() const { return EventFont; }
 
-	float GetLayoutPosY() const { return LayoutPosY; }
-	float GetLayoutEventH() const { return LayoutEventH; }
-	float GetLayoutEventDY() const { return LayoutEventDY; }
-
-	void SetLayoutPosY(const float InLayoutPosY) { LayoutPosY = InLayoutPosY; }
-	void SetLayoutEventH(const float InLayoutEventH) { LayoutEventH = InLayoutEventH; }
-	void SetLayoutEventDY(const float InLayoutEventDY) { LayoutEventDY = InLayoutEventDY; }
-
 	void DrawBackground() const;
 	void Draw(const FPacketContentViewDrawState& DrawState, const float Opacity = 1.0f) const;
 	void DrawEventHighlight(const FNetworkPacketEvent& Event, EHighlightMode Mode) const;
@@ -234,10 +226,6 @@ private:
 	const FSlateBrush* HoveredEventBorderBrush;
 	const FSlateBrush* SelectedEventBorderBrush;
 	const FSlateFontInfo EventFont;
-
-	float LayoutPosY;
-	float LayoutEventH;
-	float LayoutEventDY;
 };
 
 ////////////////////////////////////////////////////////////////////////////////////////////////////
diff --git a/Engine/Source/Developer/TraceInsights/Private/Insights/NetworkingProfiler/ViewModels/PacketContentViewport.h b/Engine/Source/Developer/TraceInsights/Private/Insights/NetworkingProfiler/ViewModels/PacketContentViewport.h
--- a/Engine/Source/Developer/TraceInsights/Private/Insights/NetworkingProfiler/ViewModels/PacketContentViewport.h
+++ b/Engine/Source/Developer/TraceInsights/Private/Insights/NetworkingProfiler/ViewModels/PacketContentViewport.h
@@ -21,6 +21,9 @@ public:
 	{
 		HorizontalAxisViewport.Reset();
 		Height = 0.0f;
+		TopEventPosY = 40.0f + 4.0f; // toolbar height + 4px
+		EventHeight = 14.0f;
+		EventDY = 2.0f;
 	}
 
 	const FAxisViewportDouble& GetHorizontalAxisViewport() const { return HorizontalAxisViewport; }
@@ -41,6 +44,15 @@ public:
 		return false;
 	}
 
+	float GetTopEventPosY() const { return TopEventPosY; }
+	void SetTopEventPosY(float InTopEventPosY) { TopEventPosY = InTopEventPosY; }
+
+	float GetEventHeight() const { return EventHeight; }
+	void SetEventHeight(float InEventHeight) { EventHeight = InEventHeight; }
+
+	float GetEventDY() const { return EventDY; }
+	void SetEventDY(float InEventDY) { EventDY = InEventDY; }
+
 private:
 	void OnSizeChanged()
 	{
@@ -49,6 +61,9 @@ private:
 private:
 	FAxisViewportDouble HorizontalAxisViewport;
 	float Height;
+	float TopEventPosY; // Y postion of events on level 0 (top lane)
+	float EventHeight; // height of an event
+	float EventDY; // vertical space between events in two adjacent lanes
 };
 
 ////////////////////////////////////////////////////////////////////////////////////////////////////
diff --git a/Engine/Source/Developer/TraceInsights/Private/Insights/NetworkingProfiler/Widgets/SNetworkingProfilerToolbar.cpp b/Engine/Source/Developer/TraceInsights/Private/Insights/NetworkingProfiler/Widgets/SNetworkingProfilerToolbar.cpp
--- a/Engine/Source/Developer/TraceInsights/Private/Insights/NetworkingProfiler/Widgets/SNetworkingProfilerToolbar.cpp
+++ b/Engine/Source/Developer/TraceInsights/Private/Insights/NetworkingProfiler/Widgets/SNetworkingProfilerToolbar.cpp
@@ -2,22 +2,14 @@
 
 #include "SNetworkingProfilerToolbar.h"
 
-#include "EditorStyleSet.h"
-#include "Framework/Application/SlateApplication.h"
-#include "Framework/Docking/TabManager.h"
-#include "Framework/MultiBox/MultiBoxDefs.h"
 #include "Framework/MultiBox/MultiBoxBuilder.h"
-#include "Widgets/Docking/SDockTab.h"
-#include "Widgets/Input/SCheckBox.h"
-#include "Widgets/Layout/SBorder.h"
 #include "Widgets/SBoxPanel.h"
-#include "Widgets/SToolTip.h"
 
 // Insights
 #include "Insights/InsightsCommands.h"
 #include "Insights/InsightsManager.h"
+#include "Insights/InsightsStyle.h"
 #include "Insights/NetworkingProfiler/NetworkingProfilerCommands.h"
-#include "Insights/NetworkingProfiler/NetworkingProfilerManager.h"
 #include "Insights/NetworkingProfiler/Widgets/SNetworkingProfilerWindow.h"
 
 ////////////////////////////////////////////////////////////////////////////////////////////////////
@@ -48,9 +40,15 @@ void SNetworkingProfilerToolbar::Construct(const FArguments& InArgs, TSharedPtr<
 		{
 			ToolbarBuilder.BeginSection("View");
 			{
-				ToolbarBuilder.AddToolBarButton(FNetworkingProfilerManager::GetCommands().TogglePacketViewVisibility);
-				ToolbarBuilder.AddToolBarButton(FNetworkingProfilerManager::GetCommands().TogglePacketContentViewVisibility);
-				ToolbarBuilder.AddToolBarButton(FNetworkingProfilerManager::GetCommands().ToggleNetStatsViewVisibility);
+				ToolbarBuilder.AddToolBarButton(FNetworkingProfilerCommands::Get().TogglePacketViewVisibility,
+					NAME_None, TAttribute<FText>(), TAttribute<FText>(),
+					FSlateIcon(FInsightsStyle::GetStyleSetName(), "PacketView.Icon.Large"));
+				ToolbarBuilder.AddToolBarButton(FNetworkingProfilerCommands::Get().TogglePacketContentViewVisibility,
+					NAME_None, TAttribute<FText>(), TAttribute<FText>(),
+					FSlateIcon(FInsightsStyle::GetStyleSetName(), "PacketContentView.Icon.Large"));
+				ToolbarBuilder.AddToolBarButton(FNetworkingProfilerCommands::Get().ToggleNetStatsViewVisibility,
+					NAME_None, TAttribute<FText>(), TAttribute<FText>(),
+					FSlateIcon(FInsightsStyle::GetStyleSetName(), "NetStatsView.Icon.Large"));
 			}
 			ToolbarBuilder.EndSection();
 			ToolbarBuilder.BeginSection("Connection");
@@ -65,65 +63,55 @@ void SNetworkingProfilerToolbar::Construct(const FArguments& InArgs, TSharedPtr<
 				ToolbarBuilder.AddWidget(ConnectionModeWidget);
 			}
 			ToolbarBuilder.EndSection();
-			ToolbarBuilder.BeginSection("New");
-			{
-				//TODO: ToolbarBuilder.AddToolBarButton(FNetworkingProfilerManager::GetCommands().OpenNewWindow);
-			}
-			ToolbarBuilder.EndSection();
+			//ToolbarBuilder.BeginSection("New");
+			//{
+			//	//TODO: ToolbarBuilder.AddToolBarButton(FNetworkingProfilerCommands::Get().OpenNewWindow,
+			//		NAME_None, TAttribute<FText>(), TAttribute<FText>(),
+			//		FSlateIcon(FInsightsStyle::GetStyleSetName(), "NewNetworkingInsights.Icon"));
+			//}
+			//ToolbarBuilder.EndSection();
 		}
 
 		static void FillRightSideToolbar(FToolBarBuilder& ToolbarBuilder)
 		{
 			ToolbarBuilder.BeginSection("Debug");
 			{
-				ToolbarBuilder.AddToolBarButton(FInsightsCommands::Get().ToggleDebugInfo);
+				ToolbarBuilder.AddToolBarButton(FInsightsCommands::Get().ToggleDebugInfo,
+					NAME_None, FText(), TAttribute<FText>(),
+					FSlateIcon(FInsightsStyle::GetStyleSetName(), "Icon.Bug"));
 			}
 			ToolbarBuilder.EndSection();
 		}
 	};
 
-	//TSharedPtr<FUICommandList> CommandList = FInsightsManager::Get()->GetCommandList();
 	TSharedPtr<FUICommandList> CommandList = ProfilerWindow->GetCommandList();
 
-	FToolBarBuilder ToolbarBuilder(CommandList.ToSharedRef(), FMultiBoxCustomization::None);
+	FSlimHorizontalToolBarBuilder ToolbarBuilder(CommandList.ToSharedRef(), FMultiBoxCustomization::None);
 	Local::FillViewToolbar(ProfilerWindow, ToolbarBuilder);
 
-	FToolBarBuilder RightSideToolbarBuilder(CommandList.ToSharedRef(), FMultiBoxCustomization::None);
+	FSlimHorizontalToolBarBuilder RightSideToolbarBuilder(CommandList.ToSharedRef(), FMultiBoxCustomization::None);
 	Local::FillRightSideToolbar(RightSideToolbarBuilder);
 
-	// Create the tool bar!
 	ChildSlot
 	[
 		SNew(SHorizontalBox)
 
-		+SHorizontalBox::Slot()
-		.HAlign(HAlign_Left)
+		+ SHorizontalBox::Slot()
+		.HAlign(HAlign_Fill)
 		.VAlign(VAlign_Center)
 		.FillWidth(1.0)
 		.Padding(0.0f)
 		[
-			SNew(SBorder)
-			.Padding(0)
-			.BorderImage(FEditorStyle::GetBrush("NoBorder"))
-			.IsEnabled(FSlateApplication::Get().GetNormalExecutionAttribute())
-			[
-				ToolbarBuilder.MakeWidget()
-			]
+			ToolbarBuilder.MakeWidget()
 		]
 
-		+SHorizontalBox::Slot()
+		+ SHorizontalBox::Slot()
 		.HAlign(HAlign_Right)
 		.VAlign(VAlign_Center)
 		.AutoWidth()
 		.Padding(0.0f)
 		[
-			SNew(SBorder)
-			.Padding(0)
-			.BorderImage(FEditorStyle::GetBrush("NoBorder"))
-			.IsEnabled(FSlateApplication::Get().GetNormalExecutionAttribute())
-			[
-				RightSideToolbarBuilder.MakeWidget()
-			]
+			RightSideToolbarBuilder.MakeWidget()
 		]
 	];
 }
diff --git a/Engine/Source/Developer/TraceInsights/Private/Insights/NetworkingProfiler/Widgets/SNetworkingProfilerToolbar.h b/Engine/Source/Developer/TraceInsights/Private/Insights/NetworkingProfiler/Widgets/SNetworkingProfilerToolbar.h
--- a/Engine/Source/Developer/TraceInsights/Private/Insights/NetworkingProfiler/Widgets/SNetworkingProfilerToolbar.h
+++ b/Engine/Source/Developer/TraceInsights/Private/Insights/NetworkingProfiler/Widgets/SNetworkingProfilerToolbar.h
@@ -2,8 +2,7 @@
 
 #pragma once
 
-#include "CoreMinimal.h"
-#include "Styling/SlateTypes.h"
+#include "Templates/SharedPointer.h"
 #include "Widgets/DeclarativeSyntaxSupport.h"
 #include "Widgets/SCompoundWidget.h"
 
@@ -24,7 +23,7 @@ public:
 	/**
 	 * Construct this widget
 	 *
-	 * @param	InArgs	The declaration data for this widget
+	 * @param InArgs The declaration data for this widget
 	 */
 	void Construct(const FArguments& InArgs, TSharedPtr<SNetworkingProfilerWindow> InProfilerWindow);
 
diff --git a/Engine/Source/Developer/TraceInsights/Private/Insights/NetworkingProfiler/Widgets/SNetworkingProfilerWindow.cpp b/Engine/Source/Developer/TraceInsights/Private/Insights/NetworkingProfiler/Widgets/SNetworkingProfilerWindow.cpp
--- a/Engine/Source/Developer/TraceInsights/Private/Insights/NetworkingProfiler/Widgets/SNetworkingProfilerWindow.cpp
+++ b/Engine/Source/Developer/TraceInsights/Private/Insights/NetworkingProfiler/Widgets/SNetworkingProfilerWindow.cpp
@@ -850,8 +850,11 @@ TSharedRef<SWidget> SNetworkingProfilerWindow::CreateGameInstanceComboBox()
 			SNew(STextBlock)
 			.Text(this, &SNetworkingProfilerWindow::GameInstance_GetSelectedText)
 		];
-
-	return GameInstanceComboBox.ToSharedRef();
+	return SNew(SBox)
+		.Padding(FMargin(4.0f, 0.0f, 2.0f, 0.0f))
+		[
+			GameInstanceComboBox.ToSharedRef()
+		];
 }
 END_SLATE_FUNCTION_BUILD_OPTIMIZATION
 
@@ -915,8 +918,11 @@ TSharedRef<SWidget> SNetworkingProfilerWindow::CreateConnectionComboBox()
 			SNew(STextBlock)
 			.Text(this, &SNetworkingProfilerWindow::Connection_GetSelectedText)
 		];
-
-	return ConnectionComboBox.ToSharedRef();
+	return SNew(SBox)
+		.Padding(FMargin(2.0f, 0.0f, 2.0f, 0.0f))
+		[
+			ConnectionComboBox.ToSharedRef()
+		];
 }
 END_SLATE_FUNCTION_BUILD_OPTIMIZATION
 
@@ -981,7 +987,11 @@ TSharedRef<SWidget> SNetworkingProfilerWindow::CreateConnectionModeComboBox()
 			.Text(this, &SNetworkingProfilerWindow::ConnectionMode_GetSelectedText)
 		];
 
-	return ConnectionModeComboBox.ToSharedRef();
+	return SNew(SBox)
+		.Padding(FMargin(2.0f, 0.0f, 4.0f, 0.0f))
+		[
+			ConnectionModeComboBox.ToSharedRef()
+		];
 }
 END_SLATE_FUNCTION_BUILD_OPTIMIZATION
 
diff --git a/Engine/Source/Developer/TraceInsights/Private/Insights/NetworkingProfiler/Widgets/SPacketContentView.cpp b/Engine/Source/Developer/TraceInsights/Private/Insights/NetworkingProfiler/Widgets/SPacketContentView.cpp
--- a/Engine/Source/Developer/TraceInsights/Private/Insights/NetworkingProfiler/Widgets/SPacketContentView.cpp
+++ b/Engine/Source/Developer/TraceInsights/Private/Insights/NetworkingProfiler/Widgets/SPacketContentView.cpp
@@ -107,221 +107,213 @@ void SPacketContentView::Construct(const FArguments& InArgs, TSharedPtr<SNetwork
 {
 	ProfilerWindow = InProfilerWindow;
 
-	ChildSlot
-	[
-		SNew(SOverlay)
-		.Visibility(EVisibility::SelfHitTestInvisible)
-
-		+ SOverlay::Slot()
-		.HAlign(HAlign_Center)
-		.VAlign(VAlign_Top)
-		.Padding(FMargin(0, 0, 0, 0))
-		[
-			SNew(SHorizontalBox)
+	FSlimHorizontalToolBarBuilder ToolbarBuilder(TSharedPtr<const FUICommandList>(), FMultiBoxCustomization::None);
 
-			//////////////////////////////////////////////////
-			// Find Packet
-
-			+ SHorizontalBox::Slot()
-			.AutoWidth()
+	ToolbarBuilder.BeginSection("FindPacket");
+	{
+		ToolbarBuilder.AddWidget(
+			SNew(SBox)
 			.VAlign(VAlign_Center)
-			.Padding(0.0f, 0.0f, 0.0f, 0.0f)
+			.Padding(FMargin(8.0f, 0.0f, 0.0f, 0.0f))
 			[
 				SNew(STextBlock)
 				.Text(LOCTEXT("FindPacketText", "Find Packet:"))
 			]
+		);
 
-			+ SHorizontalBox::Slot()
-			.AutoWidth()
-			.Padding(4.0f, 0.0f, 0.0f, 0.0f)
-			[
-				SNew(SButton)
-				.ToolTipText(LOCTEXT("PreviousPacketToolTip", "Previous Packet"))
-				.ContentPadding(0.0f)
-				.OnClicked(this, &SPacketContentView::FindPreviousPacket_OnClicked)
-				.Content()
-				[
-					SNew(SImage)
-					.Image(FInsightsStyle::GetBrush("FindPrevious"))
-				]
-			]
+		ToolbarBuilder.AddToolBarButton(
+			FUIAction(FExecuteAction::CreateSP(this, &SPacketContentView::FindPreviousPacket)),
+			NAME_None,
+			TAttribute<FText>(),
+			LOCTEXT("PreviousPacketToolTip", "Previous Packet"),
+			FSlateIcon(FInsightsStyle::GetStyleSetName(), "FindPrevious"),
+			EUserInterfaceActionType::Button
+		);
 
-			+ SHorizontalBox::Slot()
-			.AutoWidth()
-			.Padding(4.0f, 0.0f, 0.0f, 0.0f)
+		ToolbarBuilder.AddWidget(
+			SNew(SBox)
+			.Padding(FMargin(0.0f, 0.0f, 0.0f, 0.0f))
+			.VAlign(VAlign_Center)
 			[
 				SNew(SEditableTextBox)
 				.RevertTextOnEscape(true)
 				.SelectAllTextWhenFocused(true)
 				.Text(this, &SPacketContentView::GetPacketText)
+				.ToolTipText(LOCTEXT("SequenceNumber_Tooltip", "Sequence Number"))
 				.OnTextCommitted(this, &SPacketContentView::Packet_OnTextCommitted)
 				.MinDesiredWidth(30.0f)
 			]
+		);
 
-			+ SHorizontalBox::Slot()
-			.AutoWidth()
-			.Padding(4.0f, 0.0f, 0.0f, 0.0f)
-			[
-				SNew(SButton)
-				.ToolTipText(LOCTEXT("NextPacketToolTip", "Next Packet"))
-				.ContentPadding(0.0f)
-				.OnClicked(this, &SPacketContentView::FindNextPacket_OnClicked)
-				.Content()
-				[
-					SNew(SImage)
-					.Image(FInsightsStyle::GetBrush("FindNext"))
-				]
-			]
-
-			//////////////////////////////////////////////////
-			// Find Event
+		ToolbarBuilder.AddToolBarButton(
+			FUIAction(FExecuteAction::CreateSP(this, &SPacketContentView::FindNextPacket)),
+			NAME_None,
+			TAttribute<FText>(),
+			LOCTEXT("NextPacketToolTip", "Next Packet"),
+			FSlateIcon(FInsightsStyle::GetStyleSetName(), "FindNext"),
+			EUserInterfaceActionType::Button
+		);
+	}
+	ToolbarBuilder.EndSection();
 
-			+ SHorizontalBox::Slot()
-			.AutoWidth()
+	ToolbarBuilder.BeginSection("FindEvent");
+	{
+		ToolbarBuilder.AddWidget(
+			SNew(SBox)
 			.VAlign(VAlign_Center)
-			.Padding(12.0f, 0.0f, 0.0f, 0.0f)
+			.Padding(FMargin(8.0f, 0.0f, 0.0f, 0.0f))
 			[
 				SNew(STextBlock)
 				.Text(LOCTEXT("FindEventText", "Find Event:"))
 			]
+		);
 
-			+ SHorizontalBox::Slot()
-			.AutoWidth()
-			.Padding(4.0f, 0.0f, 0.0f, 0.0f)
-			[
-				SNew(SButton)
-				.ToolTipText(LOCTEXT("FindFirstEventToolTip", "First Event"))
-				.ContentPadding(0.0f)
-				.OnClicked(this, &SPacketContentView::FindFirstEvent)
-				.Content()
-				[
-					SNew(SImage)
-					.Image(FInsightsStyle::GetBrush("FindFirst"))
-				]
-			]
-
-			+ SHorizontalBox::Slot()
-			.AutoWidth()
-			.Padding(4.0f, 0.0f, 0.0f, 0.0f)
-			[
-				SNew(SButton)
-				.ToolTipText(LOCTEXT("FindPreviousEventToolTip", "Previous Event"))
-				.ContentPadding(0.0f)
-				.OnClicked(this, &SPacketContentView::FindPreviousEvent, EEventNavigationType::AnyLevel)
-				.Content()
-				[
-					SNew(SImage)
-					.Image(FInsightsStyle::GetBrush("FindPrevious"))
-				]
-			]
+		ToolbarBuilder.AddToolBarButton(
+			FUIAction(FExecuteAction::CreateSP(this, &SPacketContentView::FindFirstEvent)),
+			NAME_None,
+			TAttribute<FText>(),
+			LOCTEXT("FindFirstEventToolTip", "First Event"),
+			FSlateIcon(FInsightsStyle::GetStyleSetName(), "FindFirst"),
+			EUserInterfaceActionType::Button
+		);
 
-			+ SHorizontalBox::Slot()
-			.AutoWidth()
-			.Padding(4.0f, 0.0f, 0.0f, 0.0f)
-			[
-				SNew(SButton)
-				.ToolTipText(LOCTEXT("FindNextEventToolTip", "Next Event"))
-				.ContentPadding(0.0f)
-				.OnClicked(this, &SPacketContentView::FindNextEvent, EEventNavigationType::AnyLevel)
-				.Content()
-				[
-					SNew(SImage)
-					.Image(FInsightsStyle::GetBrush("FindNext"))
-				]
-			]
+		ToolbarBuilder.AddToolBarButton(
+			FUIAction(FExecuteAction::CreateSP(this, &SPacketContentView::FindPreviousEvent, EEventNavigationType::AnyLevel)),
+			NAME_None,
+			TAttribute<FText>(),
+			LOCTEXT("FindPreviousEventToolTip", "Previous Event"),
+			FSlateIcon(FInsightsStyle::GetStyleSetName(), "FindPrevious"),
+			EUserInterfaceActionType::Button
+		);
 
-			+ SHorizontalBox::Slot()
-			.AutoWidth()
-			.Padding(4.0f, 0.0f, 0.0f, 0.0f)
-			[
-				SNew(SButton)
-				.ToolTipText(LOCTEXT("FindLastEventToolTip", "Last Event"))
-				.ContentPadding(0.0f)
-				.OnClicked(this, &SPacketContentView::FindLastEvent)
-				.Content()
-				[
-					SNew(SImage)
-					.Image(FInsightsStyle::GetBrush("FindLast"))
-				]
-			]
+		ToolbarBuilder.AddToolBarButton(
+			FUIAction(FExecuteAction::CreateSP(this, &SPacketContentView::FindNextEvent, EEventNavigationType::AnyLevel)),
+			NAME_None,
+			TAttribute<FText>(),
+			LOCTEXT("FindNextEventToolTip", "Next Event"),
+			FSlateIcon(FInsightsStyle::GetStyleSetName(), "FindNext"),
+			EUserInterfaceActionType::Button
+		);
 
-			//////////////////////////////////////////////////
-			// By NetId
+		ToolbarBuilder.AddToolBarButton(
+			FUIAction(FExecuteAction::CreateSP(this, &SPacketContentView::FindLastEvent)),
+			NAME_None,
+			TAttribute<FText>(),
+			LOCTEXT("FindLastEventToolTip", "Last Event"),
+			FSlateIcon(FInsightsStyle::GetStyleSetName(), "FindLast"),
+			EUserInterfaceActionType::Button
+		);
+	}
+	ToolbarBuilder.EndSection();
 
-			+ SHorizontalBox::Slot()
-			.AutoWidth()
-			.VAlign(VAlign_Center)
-			.Padding(12.0f, 0.0f, 0.0f, 0.0f)
+	ToolbarBuilder.BeginSection("FilterByNetId");
+	{
+		ToolbarBuilder.AddWidget(
+			SNew(SBox)
+			.Padding(FMargin(4.0f, 0.0f, 4.0f, 0.0f))
 			[
 				SNew(SCheckBox)
 				.ToolTipText(LOCTEXT("FilterByNetId_Tooltip", "Filter events that have the specified NetId."))
 				.IsChecked(this, &SPacketContentView::FilterByNetId_IsChecked)
 				.OnCheckStateChanged(this, &SPacketContentView::FilterByNetId_OnCheckStateChanged)
 				[
-					SNew(STextBlock)
-					.Text(LOCTEXT("FilterByNetId_Text", "By NetId:"))
+					SNew(SBox)
+					.Padding(FMargin(4.0f, 0.0f, 0.0f, 0.0f))
+					[
+						SNew(STextBlock)
+						.Text(LOCTEXT("FilterByNetId_Text", "By NetId:"))
+					]
 				]
 			]
+		);
 
-			+ SHorizontalBox::Slot()
-			.AutoWidth()
-			.Padding(4.0f, 0.0f, 0.0f, 0.0f)
+		ToolbarBuilder.AddWidget(
+			SNew(SBox)
+			.Padding(FMargin(0.0f, 0.0f, 4.0f, 0.0f))
+			.VAlign(VAlign_Center)
 			[
 				SNew(SEditableTextBox)
 				.RevertTextOnEscape(true)
 				.SelectAllTextWhenFocused(true)
 				.Text(this, &SPacketContentView::GetFilterNetIdText)
+				.ToolTipText(LOCTEXT("NetId_Tooltip", "NetId"))
 				.OnTextCommitted(this, &SPacketContentView::FilterNetId_OnTextCommitted)
 				.MinDesiredWidth(40.0f)
 			]
+		);
+	}
+	ToolbarBuilder.EndSection();
 
-			//////////////////////////////////////////////////
-			// By Event Type
-
-			+ SHorizontalBox::Slot()
-			.AutoWidth()
-			.VAlign(VAlign_Center)
-			.Padding(12.0f, 0.0f, 0.0f, 0.0f)
+	ToolbarBuilder.BeginSection("FilterByEventType");
+	{
+		ToolbarBuilder.AddWidget(
+			SNew(SBox)
+			.Padding(FMargin(4.0f, 0.0f, 4.0f, 0.0f))
 			[
 				SNew(SCheckBox)
 				.ToolTipText(LOCTEXT("FilterByEventType_Tooltip", "Filter events that have the specified type.\n\nTo set the event type:\n\tdouble click either an event in the Packet Content view\n\tor an event type in the NetStats tree view."))
 				.IsChecked(this, &SPacketContentView::FilterByEventType_IsChecked)
 				.OnCheckStateChanged(this, &SPacketContentView::FilterByEventType_OnCheckStateChanged)
 				[
-					SNew(STextBlock)
-					.Text(LOCTEXT("FilterByEventType_Text", "By Type:"))
+					SNew(SBox)
+					.Padding(FMargin(4.0f, 0.0f, 0.0f, 0.0f))
+					[
+						SNew(STextBlock)
+						.Text(LOCTEXT("FilterByEventType_Text", "By Type:"))
+					]
 				]
 			]
+		);
 
-			+ SHorizontalBox::Slot()
-			.AutoWidth()
-			.Padding(4.0f, 0.0f, 0.0f, 0.0f)
+		ToolbarBuilder.AddWidget(
+			SNew(SBox)
+			.Padding(FMargin(0.0f, 0.0f, 4.0f, 0.0f))
+			.VAlign(VAlign_Center)
 			[
 				SNew(SEditableTextBox)
 				.Text(this, &SPacketContentView::GetFilterEventTypeText)
+				.ToolTipText(LOCTEXT("EventType_Tooltip", "Event Type\n\nTo set the event type:\n\tdouble click either an event in the Packet Content view\n\tor an event type in the NetStats tree view."))
 				.IsReadOnly(true)
 				.MinDesiredWidth(120.0f)
 			]
+		);
+	}
+	ToolbarBuilder.EndSection();
 
-			//////////////////////////////////////////////////
-			// Highlight Filtered Events
-
-			+ SHorizontalBox::Slot()
-			.AutoWidth()
-			.VAlign(VAlign_Center)
-			.Padding(12.0f, 0.0f, 0.0f, 0.0f)
+	ToolbarBuilder.BeginSection("HighlightFilteredEvents");
+	{
+		ToolbarBuilder.AddWidget(
+			SNew(SBox)
+			.Padding(FMargin(4.0f, 0.0f, 4.0f, 0.0f))
 			[
 				SNew(SCheckBox)
 				.ToolTipText(LOCTEXT("HighlightFilteredEvents_Tooltip", "Highlight filtered events."))
 				.IsChecked(this, &SPacketContentView::HighlightFilteredEvents_IsChecked)
 				.OnCheckStateChanged(this, &SPacketContentView::HighlightFilteredEvents_OnCheckStateChanged)
 				[
-					SNew(STextBlock)
-					.Text(LOCTEXT("HighlightFilteredEvents_Text", "Highlight"))
+					SNew(SBox)
+					.Padding(FMargin(4.0f, 0.0f, 4.0f, 0.0f))
+					[
+						SNew(STextBlock)
+						.Text(LOCTEXT("HighlightFilteredEvents_Text", "Highlight"))
+					]
 				]
 			]
+		);
+	}
+	ToolbarBuilder.EndSection();
 
-			//////////////////////////////////////////////////
+	ChildSlot
+	[
+		SNew(SOverlay)
+		.Visibility(EVisibility::SelfHitTestInvisible)
+
+		+ SOverlay::Slot()
+		.HAlign(HAlign_Fill)
+		.VAlign(VAlign_Top)
+		.Padding(FMargin(0, 0, 0, 0))
+		[
+			ToolbarBuilder.MakeWidget()
 		]
 
 		+ SOverlay::Slot()
@@ -345,28 +337,24 @@ void SPacketContentView::Construct(const FArguments& InArgs, TSharedPtr<SNetwork
 
 ////////////////////////////////////////////////////////////////////////////////////////////////////
 
-FReply SPacketContentView::FindPreviousPacket_OnClicked()
+void SPacketContentView::FindPreviousPacket()
 {
 	TSharedPtr<SPacketView> PacketView = ProfilerWindow->GetPacketView();
 	if (PacketView.IsValid())
 	{
 		PacketView->SelectPreviousPacket();
 	}
-
-	return FReply::Handled();
 }
 
 ////////////////////////////////////////////////////////////////////////////////////////////////////
 
-FReply SPacketContentView::FindNextPacket_OnClicked()
+void SPacketContentView::FindNextPacket()
 {
 	TSharedPtr<SPacketView> PacketView = ProfilerWindow->GetPacketView();
 	if (PacketView.IsValid())
 	{
 		PacketView->SelectNextPacket();
 	}
-
-	return FReply::Handled();
 }
 
 ////////////////////////////////////////////////////////////////////////////////////////////////////
@@ -395,7 +383,7 @@ void SPacketContentView::Packet_OnTextCommitted(const FText& InNewText, ETextCom
 
 ////////////////////////////////////////////////////////////////////////////////////////////////////
 
-FReply SPacketContentView::FindFirstEvent()
+void SPacketContentView::FindFirstEvent()
 {
 	if (FilteredDrawState->Events.Num() > 0)
 	{
@@ -403,17 +391,16 @@ FReply SPacketContentView::FindFirstEvent()
 		OnSelectedEventChanged();
 		BringEventIntoView(SelectedEvent);
 	}
-
-	return FReply::Handled();
 }
 
 ////////////////////////////////////////////////////////////////////////////////////////////////////
 
-FReply SPacketContentView::FindPreviousEvent(EEventNavigationType NavigationType)
+void SPacketContentView::FindPreviousEvent(EEventNavigationType NavigationType)
 {
 	if (!SelectedEvent.IsValid())
 	{
-		return FindFirstEvent();
+		FindFirstEvent();
+		return;
 	}
 
 	const int32 EventCount = FilteredDrawState->Events.Num();
@@ -460,16 +447,16 @@ FReply SPacketContentView::FindPreviousEvent(EEventNavigationType NavigationType
 	}
 
 	BringEventIntoView(SelectedEvent);
-	return FReply::Handled();
 }
 
 ////////////////////////////////////////////////////////////////////////////////////////////////////
 
-FReply SPacketContentView::FindNextEvent(EEventNavigationType NavigationType)
+void SPacketContentView::FindNextEvent(EEventNavigationType NavigationType)
 {
 	if (!SelectedEvent.IsValid())
 	{
-		return FindLastEvent();
+		FindLastEvent();
+		return;
 	}
 
 	const int32 EventCount = FilteredDrawState->Events.Num();
@@ -516,12 +503,11 @@ FReply SPacketContentView::FindNextEvent(EEventNavigationType NavigationType)
 	}
 
 	BringEventIntoView(SelectedEvent);
-	return FReply::Handled();
 }
 
 ////////////////////////////////////////////////////////////////////////////////////////////////////
 
-FReply SPacketContentView::FindLastEvent()
+void SPacketContentView::FindLastEvent()
 {
 	if (FilteredDrawState->Events.Num() > 0)
 	{
@@ -529,17 +515,16 @@ FReply SPacketContentView::FindLastEvent()
 		OnSelectedEventChanged();
 		BringEventIntoView(SelectedEvent);
 	}
-
-	return FReply::Handled();
 }
 
 ////////////////////////////////////////////////////////////////////////////////////////////////////
 
-FReply SPacketContentView::FindPreviousLevel()
+void SPacketContentView::FindPreviousLevel()
 {
 	if (!SelectedEvent.IsValid())
 	{
-		return FindFirstEvent();
+		FindFirstEvent();
+		return;
 	}
 
 	const int32 EventCount = FilteredDrawState->Events.Num();
@@ -565,16 +550,16 @@ FReply SPacketContentView::FindPreviousLevel()
 	}
 
 	BringEventIntoView(SelectedEvent);
-	return FReply::Handled();
 }
 
 ////////////////////////////////////////////////////////////////////////////////////////////////////
 
-FReply SPacketContentView::FindNextLevel()
+void SPacketContentView::FindNextLevel()
 {
 	if (!SelectedEvent.IsValid())
 	{
-		return FindLastEvent();
+		FindLastEvent();
+		return;
 	}
 
 	const int32 EventCount = FilteredDrawState->Events.Num();
@@ -600,7 +585,6 @@ FReply SPacketContentView::FindNextLevel()
 	}
 
 	BringEventIntoView(SelectedEvent);
-	return FReply::Handled();
 }
 
 ////////////////////////////////////////////////////////////////////////////////////////////////////
@@ -1039,15 +1023,12 @@ FNetworkPacketEventRef SPacketContentView::GetEventAtMousePosition(float X, floa
 			const float EventX1 = ViewportX.GetRoundedOffsetForValue(static_cast<double>(Event.BitOffset));
 			const float EventX2 = ViewportX.GetRoundedOffsetForValue(static_cast<double>(Event.BitOffset + Event.BitSize));
 
-			constexpr float EventsPosY = 32.0f;
-			constexpr float EventH = 14.0f;
-			constexpr float EventDY = 2.0f;
-			const float EventY = EventsPosY + (EventH + EventDY) * Event.Level;
+			const float EventY = Viewport.GetTopEventPosY() + (Viewport.GetEventHeight() + Viewport.GetEventDY()) * Event.Level;
 
 			constexpr float ToleranceX = 1.0f;
 
 			if (X >= EventX1 - ToleranceX && X <= EventX2 &&
-				Y >= EventY - EventDY / 2 && Y < EventY + EventH + EventDY / 2)
+				Y >= EventY - Viewport.GetEventDY() / 2 && Y < EventY + Viewport.GetEventHeight() + Viewport.GetEventDY() / 2)
 			{
 				return FNetworkPacketEventRef(Event);
 			}
@@ -1074,10 +1055,6 @@ int32 SPacketContentView::OnPaint(const FPaintArgs& Args, const FGeometry& Allot
 
 		FPacketContentViewDrawHelper Helper(DrawContext, Viewport);
 
-		Helper.SetLayoutPosY(32.0f);
-		//Helper.SetLayoutEventH(14.0f);
-		//Helper.SetLayoutEventDY(2.0f);
-
 		Helper.DrawBackground();
 
 		if (bHighlightFilteredEvents && (bFilterByNetId || bFilterByEventType))
diff --git a/Engine/Source/Developer/TraceInsights/Private/Insights/NetworkingProfiler/Widgets/SPacketContentView.h b/Engine/Source/Developer/TraceInsights/Private/Insights/NetworkingProfiler/Widgets/SPacketContentView.h
--- a/Engine/Source/Developer/TraceInsights/Private/Insights/NetworkingProfiler/Widgets/SPacketContentView.h
+++ b/Engine/Source/Developer/TraceInsights/Private/Insights/NetworkingProfiler/Widgets/SPacketContentView.h
@@ -152,16 +152,16 @@ public:
 	void EnableFilterEventType(const uint32 InEventTypeIndex);
 	void DisableFilterEventType();
 
-	FReply FindFirstEvent();
-	FReply FindPreviousEvent(EEventNavigationType NavigationType);
-	FReply FindNextEvent(EEventNavigationType NavigationType);
-	FReply FindLastEvent();
-	FReply FindPreviousLevel();
-	FReply FindNextLevel();
+	void FindFirstEvent();
+	void FindPreviousEvent(EEventNavigationType NavigationType);
+	void FindNextEvent(EEventNavigationType NavigationType);
+	void FindLastEvent();
+	void FindPreviousLevel();
+	void FindNextLevel();
 
 private:
-	FReply FindPreviousPacket_OnClicked();
-	FReply FindNextPacket_OnClicked();
+	void FindPreviousPacket();
+	void FindNextPacket();
 
 	FText GetPacketText() const;
 	void Packet_OnTextCommitted(const FText& InNewText, ETextCommit::Type InTextCommit);
