[Insights] Minor UX fixes.

diff --git a/Engine/Source/Developer/TraceInsights/Private/Insights/Widgets/SFrameTrack.cpp b/Engine/Source/Developer/TraceInsights/Private/Insights/Widgets/SFrameTrack.cpp
--- a/Engine/Source/Developer/TraceInsights/Private/Insights/Widgets/SFrameTrack.cpp
+++ b/Engine/Source/Developer/TraceInsights/Private/Insights/Widgets/SFrameTrack.cpp
@@ -1110,7 +1110,7 @@ void SFrameTrack::ShowContextMenu(const FPointerEvent& MouseEvent)
 	const bool bShouldCloseWindowAfterMenuSelection = true;
 	FMenuBuilder MenuBuilder(bShouldCloseWindowAfterMenuSelection, NULL);
 
-	MenuBuilder.BeginSection("FramesTrack", LOCTEXT("ContextMenu_Header", "Frames Track"));
+	MenuBuilder.BeginSection("Frames", LOCTEXT("FramesHeading", "Frames"));
 	{
 		struct FLocal
 		{
@@ -1151,7 +1151,11 @@ void SFrameTrack::ShowContextMenu(const FPointerEvent& MouseEvent)
 			NAME_None,
 			EUserInterfaceActionType::ToggleButton
 		);
+	}
+	MenuBuilder.EndSection();
 
+	MenuBuilder.BeginSection("Zoom", LOCTEXT("ZoomHeading", "Zoom"));
+	{
 		FUIAction Action_AutoZoom
 		(
 			FExecuteAction::CreateSP(this, &SFrameTrack::ContextMenu_AutoZoom_Execute),
diff --git a/Engine/Source/Developer/TraceInsights/Private/Insights/Widgets/SLogView.cpp b/Engine/Source/Developer/TraceInsights/Private/Insights/Widgets/SLogView.cpp
--- a/Engine/Source/Developer/TraceInsights/Private/Insights/Widgets/SLogView.cpp
+++ b/Engine/Source/Developer/TraceInsights/Private/Insights/Widgets/SLogView.cpp
@@ -108,11 +108,11 @@ public:
 		ChildSlot
 		[
 			SNew(SBorder)
-				.BorderImage(FInsightsStyle::Get().GetBrush("WhiteBrush"))
-				.BorderBackgroundColor(this, &SLogMessageRow::GetBackgroundColor)
-				[
-					Row
-				]
+			.BorderImage(FInsightsStyle::Get().GetBrush("WhiteBrush"))
+			.BorderBackgroundColor(this, &SLogMessageRow::GetBackgroundColor)
+			[
+				Row
+			]
 		];
 	}
 
@@ -579,96 +579,86 @@ void SLogView::Construct(const FArguments& InArgs)
 		// Toolbar
 		+ SVerticalBox::Slot()
 		.AutoHeight()
+		.Padding(FMargin(2.0f, 0.0f, 0.0f, 0.0f))
 		[
-			SNew(SBox)
-			.Padding(FMargin(2.0f, 0.0f, 0.0f, 0.0f))
-			[
-				ToolbarBuilder.MakeWidget()
-			]
+			ToolbarBuilder.MakeWidget()
 		]
 
 		+ SVerticalBox::Slot()
 		.FillHeight(1.0f)
 		[
-			SNew(SBox)
+			SNew(SHorizontalBox)
+
+			+ SHorizontalBox::Slot()
+			.FillWidth(1.0f)
+			.Padding(0.0f)
 			.VAlign(VAlign_Fill)
 			[
-				SNew(SHorizontalBox)
+				SNew(SScrollBox)
+				.Orientation(Orient_Horizontal)
 
-				+ SHorizontalBox::Slot()
-				.FillWidth(1.0f)
-				.Padding(0.0f)
+				+ SScrollBox::Slot()
 				.VAlign(VAlign_Fill)
 				[
-					SNew(SScrollBox)
-					.Orientation(Orient_Horizontal)
-
-					+ SScrollBox::Slot()
-					.VAlign(VAlign_Fill)
-					[
-						SAssignNew(ListView, SListView<TSharedPtr<FLogMessage>>)
-						.ExternalScrollbar(ExternalScrollbar)
-						.ItemHeight(20.0f)
-						.SelectionMode(ESelectionMode::Single)
-						.OnMouseButtonClick(this, &SLogView::OnMouseButtonClick)
-						.OnSelectionChanged(this, &SLogView::OnSelectionChanged)
-						.ListItemsSource(&Messages)
-						.OnGenerateRow(this, &SLogView::OnGenerateRow)
-						.ConsumeMouseWheel(EConsumeMouseWheel::Always)
-						.OnContextMenuOpening(FOnContextMenuOpening::CreateSP(this, &SLogView::ListView_GetContextMenu))
-						.HeaderRow
-						(
-							SNew(SHeaderRow)
-
-							+ SHeaderRow::Column(LogViewColumns::IdColumnName)
-							.ManualWidth(60.0f)
-							.DefaultLabel(LOCTEXT("IdColumn", "Index"))
-
-							+ SHeaderRow::Column(LogViewColumns::TimeColumnName)
-							.ManualWidth(94.0f)
-							.DefaultLabel(LOCTEXT("TimeColumn", "Time"))
-
-							+ SHeaderRow::Column(LogViewColumns::VerbosityColumnName)
-							.ManualWidth(80.0f)
-							.DefaultLabel(LOCTEXT("VerbosityColumn", "Verbosity"))
-
-							+ SHeaderRow::Column(LogViewColumns::CategoryColumnName)
-							.ManualWidth(120.0f)
-							.DefaultLabel(LOCTEXT("CategoryColumn", "Category"))
-
-							+ SHeaderRow::Column(LogViewColumns::MessageColumnName)
-							.ManualWidth(880.0f)
-							.DefaultLabel(LOCTEXT("MessageColumn", "Message"))
-
-							+ SHeaderRow::Column(LogViewColumns::FileColumnName)
-							.ManualWidth(600.0f)
-							.DefaultLabel(LOCTEXT("FileColumn", "File"))
-
-							+ SHeaderRow::Column(LogViewColumns::LineColumnName)
-							.ManualWidth(60.0f)
-							.DefaultLabel(LOCTEXT("LineColumn", "Line"))
-						)
-					]
+					SAssignNew(ListView, SListView<TSharedPtr<FLogMessage>>)
+					.ExternalScrollbar(ExternalScrollbar)
+					.ItemHeight(20.0f)
+					.SelectionMode(ESelectionMode::Single)
+					.OnMouseButtonClick(this, &SLogView::OnMouseButtonClick)
+					.OnSelectionChanged(this, &SLogView::OnSelectionChanged)
+					.ListItemsSource(&Messages)
+					.OnGenerateRow(this, &SLogView::OnGenerateRow)
+					.ConsumeMouseWheel(EConsumeMouseWheel::Always)
+					.OnContextMenuOpening(FOnContextMenuOpening::CreateSP(this, &SLogView::ListView_GetContextMenu))
+					.HeaderRow
+					(
+						SNew(SHeaderRow)
+
+						+ SHeaderRow::Column(LogViewColumns::IdColumnName)
+						.ManualWidth(60.0f)
+						.DefaultLabel(LOCTEXT("IdColumn", "Index"))
+
+						+ SHeaderRow::Column(LogViewColumns::TimeColumnName)
+						.ManualWidth(94.0f)
+						.DefaultLabel(LOCTEXT("TimeColumn", "Time"))
+
+						+ SHeaderRow::Column(LogViewColumns::VerbosityColumnName)
+						.ManualWidth(80.0f)
+						.DefaultLabel(LOCTEXT("VerbosityColumn", "Verbosity"))
+
+						+ SHeaderRow::Column(LogViewColumns::CategoryColumnName)
+						.ManualWidth(120.0f)
+						.DefaultLabel(LOCTEXT("CategoryColumn", "Category"))
+
+						+ SHeaderRow::Column(LogViewColumns::MessageColumnName)
+						.ManualWidth(880.0f)
+						.DefaultLabel(LOCTEXT("MessageColumn", "Message"))
+
+						+ SHeaderRow::Column(LogViewColumns::FileColumnName)
+						.ManualWidth(600.0f)
+						.DefaultLabel(LOCTEXT("FileColumn", "File"))
+
+						+ SHeaderRow::Column(LogViewColumns::LineColumnName)
+						.ManualWidth(60.0f)
+						.DefaultLabel(LOCTEXT("LineColumn", "Line"))
+					)
 				]
+			]
 
-				+ SHorizontalBox::Slot()
-				.AutoWidth()
-				.Padding(0.0f)
+			+ SHorizontalBox::Slot()
+			.AutoWidth()
+			.Padding(0.0f)
+			[
+				SNew(SBox)
+				.WidthOverride(FOptionalSize(13.0f))
 				[
-					SNew(SBox)
-					.WidthOverride(FOptionalSize(13.0f))
-					[
-						ExternalScrollbar.ToSharedRef()
-					]
+					ExternalScrollbar.ToSharedRef()
 				]
 			]
 		]
 	];
 
 	InitCommandList();
-
-	// Register ourselves with the profiler manager.
-	//TODO: FTimingProfilerManager::Get()->OnRequestLogViewUpdate().AddSP(this, &SLogView::ProfilerManager_OnRequestLogViewUpdate);
 }
 END_SLATE_FUNCTION_BUILD_OPTIMIZATION
 
diff --git a/Engine/Source/Developer/TraceInsights/Private/Insights/Widgets/SStartPageWindow.cpp b/Engine/Source/Developer/TraceInsights/Private/Insights/Widgets/SStartPageWindow.cpp
--- a/Engine/Source/Developer/TraceInsights/Private/Insights/Widgets/SStartPageWindow.cpp
+++ b/Engine/Source/Developer/TraceInsights/Private/Insights/Widgets/SStartPageWindow.cpp
@@ -648,132 +648,132 @@ BEGIN_SLATE_FUNCTION_BUILD_OPTIMIZATION
 void SStartPageWindow::Construct(const FArguments& InArgs)
 {
 	ChildSlot
+	[
+		SNew(SOverlay)
+
+		// Version
+		+ SOverlay::Slot()
+		.HAlign(HAlign_Right)
+		.VAlign(VAlign_Top)
+		.Padding(0.0f, -16.0f, 4.0f, 0.0f)
 		[
-			SNew(SOverlay)
+			SNew(STextBlock)
+			.Clipping(EWidgetClipping::ClipToBoundsWithoutIntersecting)
+			.Text(LOCTEXT("UnrealInsightsVersion", UNREAL_INSIGHTS_VERSION_STRING_EX))
+			.ColorAndOpacity(FLinearColor(0.15f, 0.15f, 0.15f, 1.0f))
+		]
 
-			// Version
-			+ SOverlay::Slot()
-			.HAlign(HAlign_Right)
-			.VAlign(VAlign_Top)
-			.Padding(0.0f, -16.0f, 4.0f, 0.0f)
+		+ SOverlay::Slot()
+		.HAlign(HAlign_Fill)
+		.VAlign(VAlign_Fill)
+		.Padding(0.0f, 0.0f, 0.0f, 0.0f)
+		[
+			SNew(SBox)
 			[
-				SNew(STextBlock)
-				.Clipping(EWidgetClipping::ClipToBoundsWithoutIntersecting)
-				.Text(LOCTEXT("UnrealInsightsVersion", UNREAL_INSIGHTS_VERSION_STRING_EX))
-				.ColorAndOpacity(FLinearColor(0.15f, 0.15f, 0.15f, 1.0f))
+				SNew(SBorder)
+				.HAlign(HAlign_Fill)
+				.VAlign(VAlign_Fill)
+				.Padding(0.0f)
+				.BorderImage(FCoreStyle::Get().GetBrush("WhiteBrush"))
+				.BorderBackgroundColor(FSlateColor(EStyleColor::Panel))
 			]
+		]
 
-			+ SOverlay::Slot()
+		// Overlay slot for the main window area
+		+ SOverlay::Slot()
+		.HAlign(HAlign_Fill)
+		.VAlign(VAlign_Fill)
+		[
+			SAssignNew(MainContentPanel, SVerticalBox)
+
+			+ SVerticalBox::Slot()
 			.HAlign(HAlign_Fill)
-			.VAlign(VAlign_Fill)
-			.Padding(0.0f, 0.0f, 0.0f, 0.0f)
+			.AutoHeight()
+			.Padding(12.0f, 8.0f, 12.0f, 4.0f)
 			[
-				SNew(SBox)
-				[
-					SNew(SBorder)
-					.HAlign(HAlign_Fill)
-					.VAlign(VAlign_Fill)
-					.Padding(0.0f)
-					.BorderImage(FCoreStyle::Get().GetBrush("WhiteBrush"))
-					.BorderBackgroundColor(FSlateColor(EStyleColor::Panel))
-				]
+				ConstructTraceStoreDirectoryPanel()
 			]
 
-			// Overlay slot for the main window area
-			+ SOverlay::Slot()
+			+ SVerticalBox::Slot()
 			.HAlign(HAlign_Fill)
 			.VAlign(VAlign_Fill)
+			.FillHeight(1.0f)
+			.Padding(3.0f, 4.0f)
 			[
-				SAssignNew(MainContentPanel, SVerticalBox)
-
-				+ SVerticalBox::Slot()
-				.HAlign(HAlign_Fill)
-				.AutoHeight()
-				.Padding(12.0f, 8.0f, 12.0f, 4.0f)
-				[
-					ConstructTraceStoreDirectoryPanel()
-				]
-
-				+ SVerticalBox::Slot()
-				.HAlign(HAlign_Fill)
-				.VAlign(VAlign_Fill)
-				.FillHeight(1.0f)
-				.Padding(3.0f, 4.0f)
-				[
-					ConstructSessionsPanel()
-				]
-
-				+ SVerticalBox::Slot()
-				.HAlign(HAlign_Fill)
-				.VAlign(VAlign_Fill)
-				.AutoHeight()
-				.Padding(12.0f, 4.0f)
-				[
-					ConstructLoadPanel()
-				]
-
-				+ SVerticalBox::Slot()
-				.HAlign(HAlign_Fill)
-				.AutoHeight()
-				.Padding(3.0f, 3.0f)
-				[
-					SNew(SBox)
-					.Visibility(this, &SStartPageWindow::StopTraceRecorder_Visibility)
-					[
-						SNew(SBorder)
-						.BorderImage(FCoreStyle::Get().GetBrush("PopupText.Background"))
-						.Padding(8.0f)
-						[
-							ConstructConnectPanel()
-						]
-					]
-				]
+				ConstructSessionsPanel()
 			]
 
-			// Overlay for fake splashscreen.
-			+ SOverlay::Slot()
+			+ SVerticalBox::Slot()
 			.HAlign(HAlign_Fill)
 			.VAlign(VAlign_Fill)
-			.Padding(0.0f)
+			.AutoHeight()
+			.Padding(12.0f, 4.0f)
+			[
+				ConstructLoadPanel()
+			]
+
+			+ SVerticalBox::Slot()
+			.HAlign(HAlign_Fill)
+			.AutoHeight()
+			.Padding(3.0f, 3.0f)
 			[
 				SNew(SBox)
-				.Visibility(this, &SStartPageWindow::SplashScreenOverlay_Visibility)
+				.Visibility(this, &SStartPageWindow::StopTraceRecorder_Visibility)
 				[
 					SNew(SBorder)
 					.BorderImage(FCoreStyle::Get().GetBrush("PopupText.Background"))
-					.BorderBackgroundColor(this, &SStartPageWindow::SplashScreenOverlay_ColorAndOpacity)
-					.Padding(0.0f)
-					.HAlign(HAlign_Fill)
-					.VAlign(VAlign_Fill)
+					.Padding(8.0f)
 					[
-						SNew(SBox)
-						.HAlign(HAlign_Center)
-						.VAlign(VAlign_Center)
-						[
-							SNew(STextBlock)
-							.Text(this, &SStartPageWindow::GetSplashScreenOverlayText)
-							.Font(FCoreStyle::GetDefaultFontStyle("Bold", 11))
-							.ColorAndOpacity(this, &SStartPageWindow::SplashScreenOverlay_TextColorAndOpacity)
-						]
+						ConstructConnectPanel()
 					]
 				]
 			]
+		]
 
-			// Notification area overlay
-			+ SOverlay::Slot()
-			.HAlign(HAlign_Right)
-			.VAlign(VAlign_Bottom)
-			.Padding(16.0f)
+		// Overlay for fake splashscreen.
+		+ SOverlay::Slot()
+		.HAlign(HAlign_Fill)
+		.VAlign(VAlign_Fill)
+		.Padding(0.0f)
+		[
+			SNew(SBox)
+			.Visibility(this, &SStartPageWindow::SplashScreenOverlay_Visibility)
 			[
-				SAssignNew(NotificationList, SNotificationList)
+				SNew(SBorder)
+				.BorderImage(FCoreStyle::Get().GetBrush("PopupText.Background"))
+				.BorderBackgroundColor(this, &SStartPageWindow::SplashScreenOverlay_ColorAndOpacity)
+				.Padding(0.0f)
+				.HAlign(HAlign_Fill)
+				.VAlign(VAlign_Fill)
+				[
+					SNew(SBox)
+					.HAlign(HAlign_Center)
+					.VAlign(VAlign_Center)
+					[
+						SNew(STextBlock)
+						.Text(this, &SStartPageWindow::GetSplashScreenOverlayText)
+						.Font(FCoreStyle::GetDefaultFontStyle("Bold", 11))
+						.ColorAndOpacity(this, &SStartPageWindow::SplashScreenOverlay_TextColorAndOpacity)
+					]
+				]
 			]
+		]
 
-			// Settings dialog overlay
-			+ SOverlay::Slot()
-			.HAlign(HAlign_Center)
-			.VAlign(VAlign_Center)
-			.Expose(OverlaySettingsSlot)
-		];
+		// Notification area overlay
+		+ SOverlay::Slot()
+		.HAlign(HAlign_Right)
+		.VAlign(VAlign_Bottom)
+		.Padding(16.0f)
+		[
+			SAssignNew(NotificationList, SNotificationList)
+		]
+
+		// Settings dialog overlay
+		+ SOverlay::Slot()
+		.HAlign(HAlign_Center)
+		.VAlign(VAlign_Center)
+		.Expose(OverlaySettingsSlot)
+	];
 
 	RefreshTraceList();
 
