[Insights] Filter Configurator UI improvements

diff --git a/Engine/Source/Developer/TraceInsights/Private/Insights/TimingProfilerManager.cpp b/Engine/Source/Developer/TraceInsights/Private/Insights/TimingProfilerManager.cpp
--- a/Engine/Source/Developer/TraceInsights/Private/Insights/TimingProfilerManager.cpp
+++ b/Engine/Source/Developer/TraceInsights/Private/Insights/TimingProfilerManager.cpp
@@ -213,6 +213,7 @@ bool FTimingProfilerManager::CanSpawnTab(const FSpawnTabArgs& Args) const
 
 void FTimingProfilerManager::OnTabClosed(TSharedRef<SDockTab> TabBeingClosed)
 {
+	OnWindowClosedEvent();
 	RemoveProfilerWindow();
 
 	// Disable TabClosed delegate.
diff --git a/Engine/Source/Developer/TraceInsights/Private/Insights/ViewModels/Filters.cpp b/Engine/Source/Developer/TraceInsights/Private/Insights/ViewModels/Filters.cpp
--- a/Engine/Source/Developer/TraceInsights/Private/Insights/ViewModels/Filters.cpp
+++ b/Engine/Source/Developer/TraceInsights/Private/Insights/ViewModels/Filters.cpp
@@ -4,6 +4,7 @@
 
 #include "Widgets/Docking/SDockTab.h"
 
+#include "Insights/InsightsStyle.h"
 #include "Insights/Widgets/SAdvancedFilter.h"
 #include "Insights/ViewModels/FilterConfigurator.h"
 
@@ -56,6 +57,12 @@ TSharedRef<SDockTab> FFilterService::SpawnTab(const FSpawnTabArgs& Args)
 	const TSharedRef<SDockTab> DockTab = SNew(SDockTab)
 		.TabRole(ETabRole::NomadTab);
 
+	const TSharedPtr<SWindow>& OwnerWindow = Args.GetOwnerWindow();
+	if (OwnerWindow.IsValid())
+	{
+		OwnerWindow->Resize(FVector2D(600, 400));
+	}
+
 	DockTab->SetContent(PendingWidget.ToSharedRef());
 	PendingWidget->SetParentTab(DockTab);
 
@@ -70,7 +77,8 @@ void FFilterService::RegisterTabSpawner()
 	FTabSpawnerEntry& TabSpawnerEntry = FGlobalTabmanager::Get()->RegisterNomadTabSpawner(FilterConfiguratorTabId,
 		FOnSpawnTab::CreateRaw(this, &FFilterService::SpawnTab))
 		.SetDisplayName(LOCTEXT("FilterConfiguratorTabTitle", "Filter Configurator"))
-		.SetMenuType(ETabSpawnerMenuType::Hidden);
+		.SetMenuType(ETabSpawnerMenuType::Hidden)
+		.SetIcon(FSlateIcon(FInsightsStyle::GetStyleSetName(), "FolderExplore.Icon.Large"));
 }
 
 ////////////////////////////////////////////////////////////////////////////////////////////////////
diff --git a/Engine/Source/Developer/TraceInsights/Private/Insights/Widgets/SAdvancedFilter.cpp b/Engine/Source/Developer/TraceInsights/Private/Insights/Widgets/SAdvancedFilter.cpp
--- a/Engine/Source/Developer/TraceInsights/Private/Insights/Widgets/SAdvancedFilter.cpp
+++ b/Engine/Source/Developer/TraceInsights/Private/Insights/Widgets/SAdvancedFilter.cpp
@@ -78,8 +78,9 @@ void SAdvancedFilter::Construct(const FArguments& InArgs, TSharedPtr<FFilterConf
 			.Padding(1.0f)
 			[
 				SNew(SButton)
+				.ButtonStyle(&FAppStyle::Get().GetWidgetStyle<FButtonStyle>("PrimaryButton"))
 				.Text(LOCTEXT("OK", "OK"))
-				.ToolTipText(LOCTEXT("OKDESC", "Apply the changes to the filters."))
+				.ToolTipText(LOCTEXT("OKDesc", "Apply the changes to the filters."))
 				.OnClicked(this, &SAdvancedFilter::OK_OnClicked)
 			]
 
diff --git a/Engine/Source/Developer/TraceInsights/Private/Insights/Widgets/SFilterConfigurator.cpp b/Engine/Source/Developer/TraceInsights/Private/Insights/Widgets/SFilterConfigurator.cpp
--- a/Engine/Source/Developer/TraceInsights/Private/Insights/Widgets/SFilterConfigurator.cpp
+++ b/Engine/Source/Developer/TraceInsights/Private/Insights/Widgets/SFilterConfigurator.cpp
@@ -49,9 +49,6 @@ BEGIN_SLATE_FUNCTION_BUILD_OPTIMIZATION
 
 void SFilterConfigurator::Construct(const FArguments& InArgs, TSharedPtr<FFilterConfigurator> InFilterConfiguratorViewModel)
 {
-	SAssignNew(ExternalScrollbar, SScrollBar)
-	.AlwaysShowScrollbar(true);
-
 	ChildSlot
 	[
 		SNew(SVerticalBox)
@@ -67,42 +64,17 @@ void SFilterConfigurator::Construct(const FArguments& InArgs, TSharedPtr<FFilter
 			.FillWidth(1.0f)
 			.Padding(0.0f)
 			[
-				SNew(SScrollBox)
-				.Orientation(Orient_Horizontal)
-
-				+ SScrollBox::Slot()
-				[
-					SNew(SOverlay)
-
-					+ SOverlay::Slot()
-					.HAlign(HAlign_Fill)
-					.VAlign(VAlign_Fill)
-					[
-						SAssignNew(TreeView, STreeView<FFilterConfiguratorNodePtr>)
-						.ExternalScrollbar(ExternalScrollbar)
-						.SelectionMode(ESelectionMode::None)
-						.TreeItemsSource(&GroupNodes)
-						.OnGetChildren(this, &SFilterConfigurator::TreeView_OnGetChildren)
-						.OnGenerateRow(this, &SFilterConfigurator::TreeView_OnGenerateRow)
-						.ItemHeight(12.0f)
-						.HeaderRow
-						(
-							SAssignNew(TreeViewHeaderRow, SHeaderRow)
-							.Visibility(EVisibility::Visible)
-						)
-					]
-				]
-			]
-
-			+ SHorizontalBox::Slot()
-			.AutoWidth()
-			.Padding(0.0f)
-			[
-				SNew(SBox)
-				.WidthOverride(FOptionalSize(13.0f))
-				[
-					ExternalScrollbar.ToSharedRef()
-				]
+				SAssignNew(TreeView, STreeView<FFilterConfiguratorNodePtr>)
+				.SelectionMode(ESelectionMode::None)
+				.TreeItemsSource(&GroupNodes)
+				.OnGetChildren(this, &SFilterConfigurator::TreeView_OnGetChildren)
+				.OnGenerateRow(this, &SFilterConfigurator::TreeView_OnGenerateRow)
+				.ItemHeight(12.0f)
+				.HeaderRow
+				(
+					SAssignNew(TreeViewHeaderRow, SHeaderRow)
+					.Visibility(EVisibility::Visible)
+				)
 			]
 		]
 	];
@@ -132,6 +104,8 @@ void SFilterConfigurator::Construct(const FArguments& InArgs, TSharedPtr<FFilter
 	FilterConfiguratorViewModel = InFilterConfiguratorViewModel;
 
 	GroupNodes.Add(FilterConfiguratorViewModel->GetRootNode());
+
+	SetInitialExpansionRec(FilterConfiguratorViewModel->GetRootNode(), true);
 }
 
 END_SLATE_FUNCTION_BUILD_OPTIMIZATION
@@ -161,6 +135,15 @@ TSharedRef<ITableRow> SFilterConfigurator::TreeView_OnGenerateRow(FFilterConfigu
 
 ////////////////////////////////////////////////////////////////////////////////////////////////////
 
+void SFilterConfigurator::SetInitialExpansionRec(const FFilterConfiguratorNodePtr& Node, bool Value)
+{
+	TreeView->SetItemExpansion(Node, Value);
+	for (const Insights::FBaseTreeNodePtr& Child : Node->GetChildren())
+	{
+		SetInitialExpansionRec(StaticCastSharedPtr<FFilterConfiguratorNode, Insights::FBaseTreeNode>(Child), Value);
+	}
+}
+
 } // namespace Insights
 
 #undef LOCTEXT_NAMESPACE
diff --git a/Engine/Source/Developer/TraceInsights/Private/Insights/Widgets/SFilterConfigurator.h b/Engine/Source/Developer/TraceInsights/Private/Insights/Widgets/SFilterConfigurator.h
--- a/Engine/Source/Developer/TraceInsights/Private/Insights/Widgets/SFilterConfigurator.h
+++ b/Engine/Source/Developer/TraceInsights/Private/Insights/Widgets/SFilterConfigurator.h
@@ -59,8 +59,9 @@ public:
 private:
 	void InitCommandList();
 
-private:
+	void SetInitialExpansionRec(const FFilterConfiguratorNodePtr& Node, bool Value);
 
+private:
 	TSharedPtr<class FFilterConfigurator> FilterConfiguratorViewModel;
 
 	/** The tree widget which holds the filters. */
diff --git a/Engine/Source/Developer/TraceInsights/Private/Insights/Widgets/SFilterConfiguratorRow.cpp b/Engine/Source/Developer/TraceInsights/Private/Insights/Widgets/SFilterConfiguratorRow.cpp
--- a/Engine/Source/Developer/TraceInsights/Private/Insights/Widgets/SFilterConfiguratorRow.cpp
+++ b/Engine/Source/Developer/TraceInsights/Private/Insights/Widgets/SFilterConfiguratorRow.cpp
@@ -5,11 +5,13 @@
 #include "EditorStyleSet.h"
 #include "Framework/Application/SlateApplication.h"
 #include "SlateOptMacros.h"
+#include "Styling/CoreStyle.h"
 #include "Widgets/DeclarativeSyntaxSupport.h"
 #include "Widgets/Images/SImage.h"
 #include "Widgets/Input/SButton.h"
 #include "Widgets/Input/SEditableTextBox.h"
 #include "Widgets/Input/SSuggestionTextBox.h"
+#include "Widgets/Layout/SSeparator.h"
 #include "Widgets/Views/STableRow.h"
 #include "Widgets/Views/STableViewBase.h"
 #include "Widgets/Views/STreeView.h"
@@ -38,9 +40,13 @@ void SFilterConfiguratorRow::Construct(const FArguments& InArgs, const TSharedRe
 TSharedRef<SWidget> SFilterConfiguratorRow::GenerateWidgetForColumn(const FName& InColumnName)
 {
 	TSharedRef<SHorizontalBox> GeneratedWidget = SNew(SHorizontalBox);
+
+	TSharedRef<SHorizontalBox> LeftBox = SNew(SHorizontalBox);
+	TSharedRef<SHorizontalBox> RightBox = SNew(SHorizontalBox);
+
 	if (!FilterConfiguratorNodePtr->IsGroup())
 	{
-		GeneratedWidget->AddSlot()
+		LeftBox->AddSlot()
 			.AutoWidth()
 			.HAlign(HAlign_Left)
 			.VAlign(VAlign_Center)
@@ -50,7 +56,7 @@ TSharedRef<SWidget> SFilterConfiguratorRow::GenerateWidgetForColumn(const FName&
 			];
 
 		// Filter combo box
-		GeneratedWidget->AddSlot()
+		LeftBox->AddSlot()
 			.HAlign(HAlign_Left)
 			.VAlign(VAlign_Center)
 			.AutoWidth()
@@ -73,7 +79,7 @@ TSharedRef<SWidget> SFilterConfiguratorRow::GenerateWidgetForColumn(const FName&
 			];
 
 		// Operator combo box
-		GeneratedWidget->AddSlot()
+		LeftBox->AddSlot()
 			.HAlign(HAlign_Left)
 			.VAlign(VAlign_Center)
 			.FillWidth(1.0)
@@ -89,10 +95,10 @@ TSharedRef<SWidget> SFilterConfiguratorRow::GenerateWidgetForColumn(const FName&
 					.OptionsSource(GetAvailableFilterOperators())
 					.OnSelectionChanged(this, &SFilterConfiguratorRow::AvailableFilterOperators_OnSelectionChanged)
 					.OnGenerateWidget(this, &SFilterConfiguratorRow::AvailableFilterOperators_OnGenerateWidget)
-						[
+					[
 						SNew(STextBlock)
 						.Text(this, &SFilterConfiguratorRow::AvailableFilterOperators_GetSelectionText)
-						]
+					]
 				]
 			];
 
@@ -100,11 +106,10 @@ TSharedRef<SWidget> SFilterConfiguratorRow::GenerateWidgetForColumn(const FName&
 		{
 			SuggestionTextBoxValue = FilterConfiguratorNodePtr->GetTextBoxValue();
 
-			GeneratedWidget->AddSlot()
+			LeftBox->AddSlot()
 				.AutoWidth()
 				.VAlign(VAlign_Center)
 				.HAlign(HAlign_Left)
-				.FillWidth(1.0)
 				.Padding(2.0f)
 				[
 					SNew(SSuggestionTextBox)
@@ -120,46 +125,41 @@ TSharedRef<SWidget> SFilterConfiguratorRow::GenerateWidgetForColumn(const FName&
 		}
 		else
 		{
-			GeneratedWidget->AddSlot()
+			LeftBox->AddSlot()
 				.AutoWidth()
 				.VAlign(VAlign_Center)
 				.HAlign(HAlign_Left)
 				.FillWidth(1.0)
 				.Padding(2.0f)
 				[
-				SNew(SEditableTextBox)
-				.MinDesiredWidth(50.0f)
-				.OnTextCommitted(this, &SFilterConfiguratorRow::OnTextBoxValueCommitted)
-				.Text(this, &SFilterConfiguratorRow::GetTextBoxValue)
-				.ToolTipText(this, &SFilterConfiguratorRow::GetTextBoxTooltipText)
-				.HintText(this, &SFilterConfiguratorRow::GetTextBoxHintText)
-				.OnVerifyTextChanged(this, &SFilterConfiguratorRow::TextBox_OnVerifyTextChanged)
+					SNew(SEditableTextBox)
+					.MinDesiredWidth(50.0f)
+					.OnTextCommitted(this, &SFilterConfiguratorRow::OnTextBoxValueCommitted)
+					.Text(this, &SFilterConfiguratorRow::GetTextBoxValue)
+					.ToolTipText(this, &SFilterConfiguratorRow::GetTextBoxTooltipText)
+					.HintText(this, &SFilterConfiguratorRow::GetTextBoxHintText)
+					.OnVerifyTextChanged(this, &SFilterConfiguratorRow::TextBox_OnVerifyTextChanged)
 				];
 		}
 			
-		GeneratedWidget->AddSlot()
+		RightBox->AddSlot()
 			.VAlign(VAlign_Center)
-			.HAlign(HAlign_Left)
-			.Padding(2.0f)
+			.HAlign(HAlign_Right)
 			.AutoWidth()
 			[
 				SNew(SButton)
 				.ToolTipText(LOCTEXT("DeleteFilterTooptip", "Delete Filter"))
-				.ContentPadding(FMargin(2.0f, 0.0f, 2.0f, -1.0f))
 				.OnClicked(this, &SFilterConfiguratorRow::DeleteFilter_OnClicked)
 				.Content()
 				[
 					SNew(SImage)
-					.ColorAndOpacity(FSlateColor(FLinearColor(0.5f, 0.0f, 0.0f, 1.0f)))
-					.Image(FInsightsStyle::Get().GetBrush("Mem.Remove.Small"))
+					.Image(FCoreStyle::Get().GetBrush("Icons.Delete"))
 				]
 			];
 	}
 	else
 	{
-		OwnerTablePtr.Pin()->Private_SetItemExpansion(FilterConfiguratorNodePtr, true);
-
-		GeneratedWidget->AddSlot()
+		LeftBox->AddSlot()
 			.AutoWidth()
 			.HAlign(HAlign_Left)
 			.VAlign(VAlign_Center)
@@ -169,7 +169,7 @@ TSharedRef<SWidget> SFilterConfiguratorRow::GenerateWidgetForColumn(const FName&
 				.ShouldDrawWires(true)
 			];
 
-		GeneratedWidget->AddSlot()
+		LeftBox->AddSlot()
 			.HAlign(HAlign_Left)
 			.VAlign(VAlign_Center)
 			.AutoWidth()
@@ -191,54 +191,120 @@ TSharedRef<SWidget> SFilterConfiguratorRow::GenerateWidgetForColumn(const FName&
 				]
 			];
 		
-		GeneratedWidget->AddSlot()
-			.HAlign(HAlign_Left)
+		RightBox->AddSlot()
+			.HAlign(HAlign_Right)
 			.VAlign(VAlign_Center)
 			.AutoWidth()
 			.Padding(2.0f)
 			[
 				SNew(SButton)
-				.Text(LOCTEXT("AddFilter", "Add Filter"))
 				.ToolTipText(LOCTEXT("AddFilterDesc", "Add a filter node as a child to this group node."))
 				.OnClicked(this, &SFilterConfiguratorRow::AddFilter_OnClicked)
+				.Content()
+				[
+					SNew(SHorizontalBox)
+
+					+ SHorizontalBox::Slot()
+					.AutoWidth()
+					.HAlign(HAlign_Center)
+					.VAlign(VAlign_Center)
+					[
+						SNew(SImage)
+						.Image(FCoreStyle::Get().GetBrush("Icons.Filter"))
+					]
+
+					+ SHorizontalBox::Slot()
+						.AutoWidth()
+						[
+							SNew(STextBlock)
+							.Margin(FMargin(2.0, 0.0f, 0.0f, 0.0f))
+							.Text(LOCTEXT("AddFilter", "Add Filter"))
+						]
+					]
 			];
 
-		GeneratedWidget->AddSlot()
-			.HAlign(HAlign_Left)
+		RightBox->AddSlot()
+			.HAlign(HAlign_Right)
 			.VAlign(VAlign_Center)
 			.AutoWidth()
 			.Padding(2.0f)
 			[
 				SNew(SButton)
-				.Text(LOCTEXT("AddGroup", "Add Group"))
 				.ToolTipText(LOCTEXT("AddGroupDesc", "Add a group node as a child to this group node."))
 				.OnClicked(this, &SFilterConfiguratorRow::AddGroup_OnClicked)
-			];
-
-		// Do not show Delete button for the Root node
-		if (FilterConfiguratorNodePtr->GetGroupPtr().IsValid())
-		{
-			GeneratedWidget->AddSlot()
-				.VAlign(VAlign_Center)
-				.HAlign(HAlign_Left)
-				.Padding(2.0f)
-				.AutoWidth()
+				.Content()
 				[
-					SNew(SButton)
-					.ToolTipText(LOCTEXT("DeleteGroupTooptip", "Delete Group"))
-					.ContentPadding(FMargin(2.0f, 0.0f, 2.0f, -1.0f))
-					.OnClicked(this, &SFilterConfiguratorRow::DeleteGroup_OnClicked)
-					.Content()
+					SNew(SHorizontalBox)
+
+					+ SHorizontalBox::Slot()
+					.AutoWidth()
+					.HAlign(HAlign_Center)
+					.VAlign(VAlign_Center)
 					[
 						SNew(SImage)
-						.ColorAndOpacity(FSlateColor(FLinearColor(0.5f, 0.0f, 0.0f, 1.0f)))
-						.Image(FInsightsStyle::Get().GetBrush("Mem.Remove.Small"))
+						.Image(FEditorStyle::GetBrush("LevelEditor.Tabs.Outliner"))
 					]
-				];
-		}
+
+					+ SHorizontalBox::Slot()
+					.AutoWidth()
+					[
+						SNew(STextBlock)
+						.Margin(FMargin(2.0, 0.0f, 0.0f, 0.0f))
+						.Text(LOCTEXT("AddGroup", "Add Group"))
+					]
+				]
+			];
+
+		// Do not show Delete button for the Root node
+		EVisibility DeleteVisibility = FilterConfiguratorNodePtr->GetGroupPtr().IsValid() ? EVisibility::Visible : EVisibility::Hidden;
+		RightBox->AddSlot()
+			.VAlign(VAlign_Center)
+			.HAlign(HAlign_Right)
+			.AutoWidth()
+			[
+				SNew(SButton)
+				.ToolTipText(LOCTEXT("DeleteGroupTooptip", "Delete Group"))
+				.OnClicked(this, &SFilterConfiguratorRow::DeleteGroup_OnClicked)
+				.HAlign(HAlign_Right)
+				.Visibility(DeleteVisibility)
+				.Content()
+				[
+					SNew(SImage)
+					.Image(FCoreStyle::Get().GetBrush("Icons.Delete"))
+				]
+			];
 	}
 
-	return GeneratedWidget;
+	GeneratedWidget->AddSlot()
+		.VAlign(VAlign_Center)
+		.HAlign(HAlign_Left)
+		.AutoWidth()
+		[
+			LeftBox
+		];
+
+	GeneratedWidget->AddSlot()
+		.VAlign(VAlign_Center)
+		.HAlign(HAlign_Right)
+		.FillWidth(1.0f)
+		[
+			SNew(SBorder)
+			.Padding(FMargin(2.0f, 0.0f, 0.0f, 0.0f))
+			.HAlign(EHorizontalAlignment::HAlign_Fill)
+			.BorderBackgroundColor(FLinearColor::Black)
+			.BorderImage(FEditorStyle::GetBrush("Border"))
+			[
+				RightBox
+			]
+		];
+
+	return	SNew(SBorder)
+			.Padding(FMargin(2))
+			.BorderBackgroundColor(FLinearColor::Black)
+			.BorderImage(FEditorStyle::GetBrush("Border"))
+			[
+				GeneratedWidget
+			];
 }
 
 ///////////////////////////////////////////////////////////////////////////////////////////////////
diff --git a/Engine/Source/Developer/TraceInsights/Private/Insights/Widgets/STimingView.cpp b/Engine/Source/Developer/TraceInsights/Private/Insights/Widgets/STimingView.cpp
--- a/Engine/Source/Developer/TraceInsights/Private/Insights/Widgets/STimingView.cpp
+++ b/Engine/Source/Developer/TraceInsights/Private/Insights/Widgets/STimingView.cpp
@@ -136,6 +136,8 @@ STimingView::~STimingView()
 	IModularFeatures::Get().UnregisterModularFeature(Insights::TimingViewExtenderFeatureName, LoadingSharedState.Get());
 	IModularFeatures::Get().UnregisterModularFeature(Insights::TimingViewExtenderFeatureName, ThreadTimingSharedState.Get());
 	IModularFeatures::Get().UnregisterModularFeature(Insights::TimingViewExtenderFeatureName, FrameSharedState.Get());
+
+	FGlobalTabmanager::Get()->UnregisterNomadTabSpawner(QuickFindTabId);
 }
 
 ////////////////////////////////////////////////////////////////////////////////////////////////////
@@ -271,7 +273,8 @@ void STimingView::Construct(const FArguments& InArgs)
 	FTabSpawnerEntry& TabSpawnerEntry = FGlobalTabmanager::Get()->RegisterNomadTabSpawner(QuickFindTabId,
 		FOnSpawnTab::CreateSP(this, &STimingView::SpawnQuickFindTab))
 		.SetDisplayName(LOCTEXT("QuickFindTabTitle", "Quick Find"))
-		.SetMenuType(ETabSpawnerMenuType::Hidden);
+		.SetMenuType(ETabSpawnerMenuType::Hidden)
+		.SetIcon(FSlateIcon(FInsightsStyle::GetStyleSetName(), "FolderExplore.Icon.Large"));
 }
 
 ////////////////////////////////////////////////////////////////////////////////////////////////////
@@ -4719,6 +4722,12 @@ TSharedRef<SDockTab> STimingView::SpawnQuickFindTab(const FSpawnTabArgs& Args)
 	const TSharedRef<SDockTab> DockTab = SNew(SDockTab)
 		.TabRole(ETabRole::NomadTab);
 
+	const TSharedPtr<SWindow>& OwnerWindow = Args.GetOwnerWindow();
+	if (OwnerWindow.IsValid() && OwnerWindow != FSlateApplication::Get().FindWidgetWindow(SharedThis(this)))
+	{
+		OwnerWindow->Resize(FVector2D(600, 400));
+	}
+
 	if (!QuickFindWidgetSharedPtr.IsValid())
 	{
 		return DockTab;
