[Insights]
  - Updated toolbar for Asset Loading Insights (with placeholder icons).
  - Updated toolbar for Memory Insights (with placeholder icons).
  - Made the main toolbar fixed (not dockable) for all major views.
  - Moved the Menu on top of the window (currently hardcoded for Session Info, Timing Insights, Asset Loading Insights and Memory Insights).
  - Other minor UI fixes.

diff --git a/Engine/Source/Developer/TraceInsights/Private/Insights/LoadingProfiler/Widgets/SLoadingProfilerToolbar.cpp b/Engine/Source/Developer/TraceInsights/Private/Insights/LoadingProfiler/Widgets/SLoadingProfilerToolbar.cpp
--- a/Engine/Source/Developer/TraceInsights/Private/Insights/LoadingProfiler/Widgets/SLoadingProfilerToolbar.cpp
+++ b/Engine/Source/Developer/TraceInsights/Private/Insights/LoadingProfiler/Widgets/SLoadingProfilerToolbar.cpp
@@ -2,22 +2,14 @@
 
 #include "SLoadingProfilerToolbar.h"
 
-#include "EditorStyleSet.h"
-#include "Framework/Application/SlateApplication.h"
-#include "Framework/Docking/TabManager.h"
-#include "Framework/MultiBox/MultiBoxDefs.h"
 #include "Framework/MultiBox/MultiBoxBuilder.h"
-#include "Widgets/Docking/SDockTab.h"
-#include "Widgets/Input/SCheckBox.h"
-#include "Widgets/Layout/SBorder.h"
 #include "Widgets/SBoxPanel.h"
-#include "Widgets/SToolTip.h"
 
 // Insights
 #include "Insights/InsightsCommands.h"
 #include "Insights/InsightsManager.h"
+#include "Insights/InsightsStyle.h"
 #include "Insights/LoadingProfiler/LoadingProfilerCommands.h"
-#include "Insights/LoadingProfiler/LoadingProfilerManager.h"
 
 ////////////////////////////////////////////////////////////////////////////////////////////////////
 
@@ -45,12 +37,24 @@ void SLoadingProfilerToolbar::Construct(const FArguments& InArgs)
 		{
 			ToolbarBuilder.BeginSection("View");
 			{
-				ToolbarBuilder.AddToolBarButton(FLoadingProfilerCommands::Get().ToggleTimingViewVisibility);
-				ToolbarBuilder.AddToolBarButton(FLoadingProfilerCommands::Get().ToggleEventAggregationTreeViewVisibility);
-				ToolbarBuilder.AddToolBarButton(FLoadingProfilerCommands::Get().ToggleObjectTypeAggregationTreeViewVisibility);
-				ToolbarBuilder.AddToolBarButton(FLoadingProfilerCommands::Get().TogglePackageDetailsTreeViewVisibility);
-				ToolbarBuilder.AddToolBarButton(FLoadingProfilerCommands::Get().ToggleExportDetailsTreeViewVisibility);
-				ToolbarBuilder.AddToolBarButton(FLoadingProfilerCommands::Get().ToggleRequestsTreeViewVisibility);
+				ToolbarBuilder.AddToolBarButton(FLoadingProfilerCommands::Get().ToggleTimingViewVisibility,
+					NAME_None, TAttribute<FText>(), TAttribute<FText>(),
+					FSlateIcon(FInsightsStyle::GetStyleSetName(), "TimingView.Icon.Large"));
+				ToolbarBuilder.AddToolBarButton(FLoadingProfilerCommands::Get().ToggleEventAggregationTreeViewVisibility,
+					NAME_None, TAttribute<FText>(), TAttribute<FText>(),
+					FSlateIcon(FInsightsStyle::GetStyleSetName(), "TableTreeView.Icon.Large"));
+				ToolbarBuilder.AddToolBarButton(FLoadingProfilerCommands::Get().ToggleObjectTypeAggregationTreeViewVisibility,
+					NAME_None, TAttribute<FText>(), TAttribute<FText>(),
+					FSlateIcon(FInsightsStyle::GetStyleSetName(), "TableTreeView.Icon.Large"));
+				ToolbarBuilder.AddToolBarButton(FLoadingProfilerCommands::Get().TogglePackageDetailsTreeViewVisibility,
+					NAME_None, TAttribute<FText>(), TAttribute<FText>(),
+					FSlateIcon(FInsightsStyle::GetStyleSetName(), "TableTreeView.Icon.Large"));
+				ToolbarBuilder.AddToolBarButton(FLoadingProfilerCommands::Get().ToggleExportDetailsTreeViewVisibility,
+					NAME_None, TAttribute<FText>(), TAttribute<FText>(),
+					FSlateIcon(FInsightsStyle::GetStyleSetName(), "TableTreeView.Icon.Large"));
+				ToolbarBuilder.AddToolBarButton(FLoadingProfilerCommands::Get().ToggleRequestsTreeViewVisibility,
+					NAME_None, TAttribute<FText>(), TAttribute<FText>(),
+					FSlateIcon(FInsightsStyle::GetStyleSetName(), "TableTreeView.Icon.Large"));
 			}
 			ToolbarBuilder.EndSection();
 		}
@@ -59,7 +63,9 @@ void SLoadingProfilerToolbar::Construct(const FArguments& InArgs)
 		{
 			ToolbarBuilder.BeginSection("Debug");
 			{
-				ToolbarBuilder.AddToolBarButton(FInsightsCommands::Get().ToggleDebugInfo);
+				ToolbarBuilder.AddToolBarButton(FInsightsCommands::Get().ToggleDebugInfo,
+					NAME_None, FText(), TAttribute<FText>(),
+					FSlateIcon(FInsightsStyle::GetStyleSetName(), "Icon.Bug"));
 			}
 			ToolbarBuilder.EndSection();
 		}
@@ -67,45 +73,32 @@ void SLoadingProfilerToolbar::Construct(const FArguments& InArgs)
 
 	TSharedPtr<FUICommandList> CommandList = FInsightsManager::Get()->GetCommandList();
 
-	FToolBarBuilder ToolbarBuilder(CommandList.ToSharedRef(), FMultiBoxCustomization::None);
+	FSlimHorizontalToolBarBuilder ToolbarBuilder(CommandList.ToSharedRef(), FMultiBoxCustomization::None);
 	Local::FillViewToolbar(ToolbarBuilder);
 
-	FToolBarBuilder RightSideToolbarBuilder(CommandList.ToSharedRef(), FMultiBoxCustomization::None);
+	FSlimHorizontalToolBarBuilder RightSideToolbarBuilder(CommandList.ToSharedRef(), FMultiBoxCustomization::None);
 	Local::FillRightSideToolbar(RightSideToolbarBuilder);
 
-	// Create the tool bar!
 	ChildSlot
 	[
 		SNew(SHorizontalBox)
 
-		+SHorizontalBox::Slot()
-		.HAlign(HAlign_Left)
+		+ SHorizontalBox::Slot()
+		.HAlign(HAlign_Fill)
 		.VAlign(VAlign_Center)
 		.FillWidth(1.0)
 		.Padding(0.0f)
 		[
-			SNew(SBorder)
-			.Padding(0)
-			.BorderImage(FEditorStyle::GetBrush("NoBorder"))
-			.IsEnabled(FSlateApplication::Get().GetNormalExecutionAttribute())
-			[
-				ToolbarBuilder.MakeWidget()
-			]
+			ToolbarBuilder.MakeWidget()
 		]
 
-		+SHorizontalBox::Slot()
+		+ SHorizontalBox::Slot()
 		.HAlign(HAlign_Right)
 		.VAlign(VAlign_Center)
 		.AutoWidth()
 		.Padding(0.0f)
 		[
-			SNew(SBorder)
-			.Padding(0)
-			.BorderImage(FEditorStyle::GetBrush("NoBorder"))
-			.IsEnabled(FSlateApplication::Get().GetNormalExecutionAttribute())
-			[
-				RightSideToolbarBuilder.MakeWidget()
-			]
+			RightSideToolbarBuilder.MakeWidget()
 		]
 	];
 }
diff --git a/Engine/Source/Developer/TraceInsights/Private/Insights/LoadingProfiler/Widgets/SLoadingProfilerToolbar.h b/Engine/Source/Developer/TraceInsights/Private/Insights/LoadingProfiler/Widgets/SLoadingProfilerToolbar.h
--- a/Engine/Source/Developer/TraceInsights/Private/Insights/LoadingProfiler/Widgets/SLoadingProfilerToolbar.h
+++ b/Engine/Source/Developer/TraceInsights/Private/Insights/LoadingProfiler/Widgets/SLoadingProfilerToolbar.h
@@ -1,16 +1,10 @@
 // Copyright Epic Games, Inc. All Rights Reserved.
 
-
 #pragma once
 
-#include "CoreMinimal.h"
-#include "Styling/SlateTypes.h"
 #include "Widgets/DeclarativeSyntaxSupport.h"
 #include "Widgets/SCompoundWidget.h"
 
-class SBorder;
-class FToolBarBuilder;
-
 /** Ribbon based toolbar used as a main menu in the Profiler window. */
 class SLoadingProfilerToolbar : public SCompoundWidget
 {
@@ -27,9 +21,9 @@ public:
 	/**
 	 * Construct this widget
 	 *
-	 * @param	InArgs	The declaration data for this widget
+	 * @param InArgs The declaration data for this widget
 	 */
-	void Construct( const FArguments& InArgs );
+	void Construct(const FArguments& InArgs);
 
 private:
 	/** Create the UI commands for the toolbar */
diff --git a/Engine/Source/Developer/TraceInsights/Private/Insights/LoadingProfiler/Widgets/SLoadingProfilerWindow.cpp b/Engine/Source/Developer/TraceInsights/Private/Insights/LoadingProfiler/Widgets/SLoadingProfilerWindow.cpp
--- a/Engine/Source/Developer/TraceInsights/Private/Insights/LoadingProfiler/Widgets/SLoadingProfilerWindow.cpp
+++ b/Engine/Source/Developer/TraceInsights/Private/Insights/LoadingProfiler/Widgets/SLoadingProfilerWindow.cpp
@@ -2,7 +2,9 @@
 
 #include "SLoadingProfilerWindow.h"
 
+#include "Framework/Commands/UICommandList.h"
 #include "Framework/Docking/LayoutService.h"
+#include "Framework/Docking/TabManager.h"
 #include "Framework/Docking/WorkspaceItem.h"
 #include "Framework/MultiBox/MultiBoxBuilder.h"
 #include "SlateOptMacros.h"
@@ -11,8 +13,6 @@
 #include "Widgets/Images/SImage.h"
 #include "Widgets/Layout/SBorder.h"
 #include "Widgets/Layout/SBox.h"
-#include "Widgets/Layout/SSpacer.h"
-#include "Widgets/Notifications/SNotificationList.h"
 #include "Widgets/SBoxPanel.h"
 #include "Widgets/Text/STextBlock.h"
 
@@ -40,7 +40,6 @@
 
 ////////////////////////////////////////////////////////////////////////////////////////////////////
 
-const FName FLoadingProfilerTabs::ToolbarID(TEXT("Toolbar"));
 const FName FLoadingProfilerTabs::TimingViewID(TEXT("TimingView"));
 const FName FLoadingProfilerTabs::EventAggregationTreeViewID(TEXT("EventAggregation"));
 const FName FLoadingProfilerTabs::ObjectTypeAggregationTreeViewID(TEXT("ObjectTypeAggregation"));
@@ -95,8 +94,6 @@ SLoadingProfilerWindow::~SLoadingProfilerWindow()
 		check(TimingView == nullptr);
 	}
 
-	HideTab(FLoadingProfilerTabs::ToolbarID);
-
 #if WITH_EDITOR
 	if (DurationActive > 0.0f && FEngineAnalytics::IsAvailable())
 	{
@@ -292,29 +289,6 @@ void SLoadingProfilerWindow::UpdateRequestsTreeView()
 ////////////////////////////////////////////////////////////////////////////////////////////////////
 
 BEGIN_SLATE_FUNCTION_BUILD_OPTIMIZATION
-
-TSharedRef<SDockTab> SLoadingProfilerWindow::SpawnTab_Toolbar(const FSpawnTabArgs& Args)
-{
-	const TSharedRef<SDockTab> DockTab = SNew(SDockTab)
-		.ShouldAutosize(true)
-		.TabRole(ETabRole::PanelTab)
-		[
-			SNew(SLoadingProfilerToolbar)
-		];
-
-	DockTab->SetOnTabClosed(SDockTab::FOnTabClosedCallback::CreateRaw(this, &SLoadingProfilerWindow::OnToolbarTabClosed));
-
-	return DockTab;
-}
-
-////////////////////////////////////////////////////////////////////////////////////////////////////
-
-void SLoadingProfilerWindow::OnToolbarTabClosed(TSharedRef<SDockTab> TabBeingClosed)
-{
-}
-
-////////////////////////////////////////////////////////////////////////////////////////////////////
-
 TSharedRef<SDockTab> SLoadingProfilerWindow::SpawnTab_TimingView(const FSpawnTabArgs& Args)
 {
 	FLoadingProfilerManager::Get()->SetTimingViewVisible(true);
@@ -333,6 +307,7 @@ TSharedRef<SDockTab> SLoadingProfilerWindow::SpawnTab_TimingView(const FSpawnTab
 
 	return DockTab;
 }
+END_SLATE_FUNCTION_BUILD_OPTIMIZATION
 
 ////////////////////////////////////////////////////////////////////////////////////////////////////
 
@@ -344,6 +319,7 @@ void SLoadingProfilerWindow::OnTimingViewTabClosed(TSharedRef<SDockTab> TabBeing
 
 ////////////////////////////////////////////////////////////////////////////////////////////////////
 
+BEGIN_SLATE_FUNCTION_BUILD_OPTIMIZATION
 TSharedRef<SDockTab> SLoadingProfilerWindow::SpawnTab_EventAggregationTreeView(const FSpawnTabArgs& Args)
 {
 	FLoadingProfilerManager::Get()->SetEventAggregationTreeViewVisible(true);
@@ -364,6 +340,7 @@ TSharedRef<SDockTab> SLoadingProfilerWindow::SpawnTab_EventAggregationTreeView(c
 
 	return DockTab;
 }
+END_SLATE_FUNCTION_BUILD_OPTIMIZATION
 
 ////////////////////////////////////////////////////////////////////////////////////////////////////
 
@@ -375,6 +352,7 @@ void SLoadingProfilerWindow::OnEventAggregationTreeViewTabClosed(TSharedRef<SDoc
 
 ////////////////////////////////////////////////////////////////////////////////////////////////////
 
+BEGIN_SLATE_FUNCTION_BUILD_OPTIMIZATION
 TSharedRef<SDockTab> SLoadingProfilerWindow::SpawnTab_ObjectTypeAggregationTreeView(const FSpawnTabArgs& Args)
 {
 	FLoadingProfilerManager::Get()->SetObjectTypeAggregationTreeViewVisible(true);
@@ -395,6 +373,7 @@ TSharedRef<SDockTab> SLoadingProfilerWindow::SpawnTab_ObjectTypeAggregationTreeV
 
 	return DockTab;
 }
+END_SLATE_FUNCTION_BUILD_OPTIMIZATION
 
 ////////////////////////////////////////////////////////////////////////////////////////////////////
 
@@ -406,6 +385,7 @@ void SLoadingProfilerWindow::OnObjectTypeAggregationTreeViewTabClosed(TSharedRef
 
 ////////////////////////////////////////////////////////////////////////////////////////////////////
 
+BEGIN_SLATE_FUNCTION_BUILD_OPTIMIZATION
 TSharedRef<SDockTab> SLoadingProfilerWindow::SpawnTab_PackageDetailsTreeView(const FSpawnTabArgs& Args)
 {
 	FLoadingProfilerManager::Get()->SetPackageDetailsTreeViewVisible(true);
@@ -426,6 +406,7 @@ TSharedRef<SDockTab> SLoadingProfilerWindow::SpawnTab_PackageDetailsTreeView(con
 
 	return DockTab;
 }
+END_SLATE_FUNCTION_BUILD_OPTIMIZATION
 
 ////////////////////////////////////////////////////////////////////////////////////////////////////
 
@@ -437,6 +418,7 @@ void SLoadingProfilerWindow::OnPackageDetailsTreeViewTabClosed(TSharedRef<SDockT
 
 ////////////////////////////////////////////////////////////////////////////////////////////////////
 
+BEGIN_SLATE_FUNCTION_BUILD_OPTIMIZATION
 TSharedRef<SDockTab> SLoadingProfilerWindow::SpawnTab_ExportDetailsTreeView(const FSpawnTabArgs& Args)
 {
 	FLoadingProfilerManager::Get()->SetExportDetailsTreeViewVisible(true);
@@ -457,6 +439,7 @@ TSharedRef<SDockTab> SLoadingProfilerWindow::SpawnTab_ExportDetailsTreeView(cons
 
 	return DockTab;
 }
+END_SLATE_FUNCTION_BUILD_OPTIMIZATION
 
 ////////////////////////////////////////////////////////////////////////////////////////////////////
 
@@ -468,6 +451,7 @@ void SLoadingProfilerWindow::OnExportDetailsTreeViewTabClosed(TSharedRef<SDockTa
 
 ////////////////////////////////////////////////////////////////////////////////////////////////////
 
+BEGIN_SLATE_FUNCTION_BUILD_OPTIMIZATION
 TSharedRef<SDockTab> SLoadingProfilerWindow::SpawnTab_RequestsTreeView(const FSpawnTabArgs& Args)
 {
 	FLoadingProfilerManager::Get()->SetRequestsTreeViewVisible(true);
@@ -488,6 +472,7 @@ TSharedRef<SDockTab> SLoadingProfilerWindow::SpawnTab_RequestsTreeView(const FSp
 
 	return DockTab;
 }
+END_SLATE_FUNCTION_BUILD_OPTIMIZATION
 
 ////////////////////////////////////////////////////////////////////////////////////////////////////
 
@@ -499,6 +484,7 @@ void SLoadingProfilerWindow::OnRequestsTreeViewTabClosed(TSharedRef<SDockTab> Ta
 
 ////////////////////////////////////////////////////////////////////////////////////////////////////
 
+BEGIN_SLATE_FUNCTION_BUILD_OPTIMIZATION
 void SLoadingProfilerWindow::Construct(const FArguments& InArgs, const TSharedRef<SDockTab>& ConstructUnderMajorTab, const TSharedPtr<SWindow>& ConstructUnderWindow)
 {
 	// Create & initialize tab manager.
@@ -511,38 +497,33 @@ void SLoadingProfilerWindow::Construct(const FArguments& InArgs, const TSharedRe
 
 	TSharedRef<FWorkspaceItem> AppMenuGroup = TabManager->AddLocalWorkspaceMenuCategory(LOCTEXT("LoadingProfilerMenuGroupName", "Asset Loading Insights"));
 
-	TabManager->RegisterTabSpawner(FLoadingProfilerTabs::ToolbarID, FOnSpawnTab::CreateRaw(this, &SLoadingProfilerWindow::SpawnTab_Toolbar))
-		.SetDisplayName(LOCTEXT("DeviceToolbarTabTitle", "Toolbar"))
-		.SetIcon(FSlateIcon(FInsightsStyle::GetStyleSetName(), "Toolbar.Icon.Small"))
-		.SetGroup(AppMenuGroup);
-
 	TabManager->RegisterTabSpawner(FLoadingProfilerTabs::TimingViewID, FOnSpawnTab::CreateRaw(this, &SLoadingProfilerWindow::SpawnTab_TimingView))
-		.SetDisplayName(LOCTEXT("LoadingProfiler.TimingViewTabTitle", "Timing View"))
+		.SetDisplayName(LOCTEXT("TimingViewTabTitle", "Timing View"))
 		.SetIcon(FSlateIcon(FInsightsStyle::GetStyleSetName(), "TimingView.Icon.Small"))
 		.SetGroup(AppMenuGroup);
 
 	TabManager->RegisterTabSpawner(FLoadingProfilerTabs::EventAggregationTreeViewID, FOnSpawnTab::CreateRaw(this, &SLoadingProfilerWindow::SpawnTab_EventAggregationTreeView))
-		.SetDisplayName(LOCTEXT("LoadingProfiler.EventAggregationTreeViewTabTitle", "Event Aggregation"))
+		.SetDisplayName(LOCTEXT("EventAggregationTreeViewTabTitle", "Event Aggregation"))
 		.SetIcon(FSlateIcon(FInsightsStyle::GetStyleSetName(), "TableTreeView.Icon.Small"))
 		.SetGroup(AppMenuGroup);
 
 	TabManager->RegisterTabSpawner(FLoadingProfilerTabs::ObjectTypeAggregationTreeViewID, FOnSpawnTab::CreateRaw(this, &SLoadingProfilerWindow::SpawnTab_ObjectTypeAggregationTreeView))
-		.SetDisplayName(LOCTEXT("LoadingProfiler.ObjectTypeAggregationTreeViewTabTitle", "Object Type Aggregation"))
+		.SetDisplayName(LOCTEXT("ObjectTypeAggregationTreeViewTabTitle", "Object Type Aggregation"))
 		.SetIcon(FSlateIcon(FInsightsStyle::GetStyleSetName(), "TableTreeView.Icon.Small"))
 		.SetGroup(AppMenuGroup);
 
 	TabManager->RegisterTabSpawner(FLoadingProfilerTabs::PackageDetailsTreeViewID, FOnSpawnTab::CreateRaw(this, &SLoadingProfilerWindow::SpawnTab_PackageDetailsTreeView))
-		.SetDisplayName(LOCTEXT("LoadingProfiler.PackageDetailsTreeViewTabTitle", "Package Details"))
+		.SetDisplayName(LOCTEXT("PackageDetailsTreeViewTabTitle", "Package Details"))
 		.SetIcon(FSlateIcon(FInsightsStyle::GetStyleSetName(), "TableTreeView.Icon.Small"))
 		.SetGroup(AppMenuGroup);
 
 	TabManager->RegisterTabSpawner(FLoadingProfilerTabs::ExportDetailsTreeViewID, FOnSpawnTab::CreateRaw(this, &SLoadingProfilerWindow::SpawnTab_ExportDetailsTreeView))
-		.SetDisplayName(LOCTEXT("LoadingProfiler.ExportDetailsTreeViewTabTitle", "Export Details"))
+		.SetDisplayName(LOCTEXT("ExportDetailsTreeViewTabTitle", "Export Details"))
 		.SetIcon(FSlateIcon(FInsightsStyle::GetStyleSetName(), "TableTreeView.Icon.Small"))
 		.SetGroup(AppMenuGroup);
 
 	TabManager->RegisterTabSpawner(FLoadingProfilerTabs::RequestsTreeViewID, FOnSpawnTab::CreateRaw(this, &SLoadingProfilerWindow::SpawnTab_RequestsTreeView))
-		.SetDisplayName(LOCTEXT("LoadingProfiler.RequestsTreeViewTabTitle", "Requests"))
+		.SetDisplayName(LOCTEXT("RequestsTreeViewTabTitle", "Requests"))
 		.SetIcon(FSlateIcon(FInsightsStyle::GetStyleSetName(), "TableTreeView.Icon.Small"))
 		.SetGroup(AppMenuGroup);
 
@@ -550,40 +531,28 @@ void SLoadingProfilerWindow::Construct(const FArguments& InArgs, const TSharedRe
 	ensure(LoadingProfilerManager.IsValid());
 
 	// Create tab layout.
-	TSharedRef<FTabManager::FLayout> Layout = FTabManager::NewLayout("InsightsLoadingProfilerLayout_v1.0")
+	TSharedRef<FTabManager::FLayout> Layout = FTabManager::NewLayout("InsightsLoadingProfilerLayout_v1.2")
 		->AddArea
 		(
 			FTabManager::NewPrimaryArea()
-			->SetOrientation(Orient_Vertical)
+			->SetOrientation(Orient_Horizontal)
 			->Split
 			(
 				FTabManager::NewStack()
-				->AddTab(FLoadingProfilerTabs::ToolbarID, ETabState::OpenedTab)
+				->SetSizeCoefficient(0.5f)
 				->SetHideTabWell(true)
+				->AddTab(FLoadingProfilerTabs::TimingViewID, ETabState::OpenedTab)
 			)
 			->Split
 			(
-				FTabManager::NewSplitter()
-				->SetOrientation(Orient_Horizontal)
-				->SetSizeCoefficient(1.0f)
-				->Split
-				(
-					FTabManager::NewStack()
-					->SetSizeCoefficient(0.5f)
-					->SetHideTabWell(true)
-					->AddTab(FLoadingProfilerTabs::TimingViewID, ETabState::OpenedTab)
-				)
-				->Split
-				(
-					FTabManager::NewStack()
-					->SetSizeCoefficient(0.35f)
-					->AddTab(FLoadingProfilerTabs::EventAggregationTreeViewID, ETabState::OpenedTab)
-					->AddTab(FLoadingProfilerTabs::ObjectTypeAggregationTreeViewID, ETabState::OpenedTab)
-					->AddTab(FLoadingProfilerTabs::PackageDetailsTreeViewID, ETabState::OpenedTab)
-					->AddTab(FLoadingProfilerTabs::ExportDetailsTreeViewID, ETabState::OpenedTab)
-					->AddTab(FLoadingProfilerTabs::RequestsTreeViewID, ETabState::OpenedTab)
-					->SetForegroundTab(FLoadingProfilerTabs::EventAggregationTreeViewID)
-				)
+				FTabManager::NewStack()
+				->SetSizeCoefficient(0.35f)
+				->AddTab(FLoadingProfilerTabs::EventAggregationTreeViewID, ETabState::OpenedTab)
+				->AddTab(FLoadingProfilerTabs::ObjectTypeAggregationTreeViewID, ETabState::OpenedTab)
+				->AddTab(FLoadingProfilerTabs::PackageDetailsTreeViewID, ETabState::OpenedTab)
+				->AddTab(FLoadingProfilerTabs::ExportDetailsTreeViewID, ETabState::OpenedTab)
+				->AddTab(FLoadingProfilerTabs::RequestsTreeViewID, ETabState::OpenedTab)
+				->SetForegroundTab(FLoadingProfilerTabs::EventAggregationTreeViewID)
 			)
 		);
 
@@ -591,7 +560,6 @@ void SLoadingProfilerWindow::Construct(const FArguments& InArgs, const TSharedRe
 
 	// Create & initialize main menu.
 	FMenuBarBuilder MenuBarBuilder = FMenuBarBuilder(TSharedPtr<FUICommandList>());
-
 	MenuBarBuilder.AddPullDownMenu(
 		LOCTEXT("MenuLabel", "Menu"),
 		FText::GetEmpty(),
@@ -599,63 +567,80 @@ void SLoadingProfilerWindow::Construct(const FArguments& InArgs, const TSharedRe
 		FName(TEXT("Menu"))
 	);
 
+#if !WITH_EDITOR
 	TSharedRef<SWidget> MenuWidget = MenuBarBuilder.MakeWidget();
+	MenuWidget->SetClipping(EWidgetClipping::ClipToBoundsWithoutIntersecting);
+#endif
+
 	ChildSlot
+	[
+		SNew(SOverlay)
+
+#if !WITH_EDITOR
+		// Menu
+		+ SOverlay::Slot()
+		.HAlign(HAlign_Left)
+		.VAlign(VAlign_Top)
+		.Padding(34.0f, -60.0f, 0.0f, 0.0f)
 		[
-			SNew(SOverlay)
-
-			// Version
-			+ SOverlay::Slot()
-				.HAlign(HAlign_Right)
-				.VAlign(VAlign_Top)
-				.Padding(0.0f, -16.0f, 0.0f, 0.0f)
-				[
-					SNew(STextBlock)
-						.Clipping(EWidgetClipping::ClipToBoundsWithoutIntersecting)
-						.Text(LOCTEXT("UnrealInsightsVersion", UNREAL_INSIGHTS_VERSION_STRING_EX))
-						.ColorAndOpacity(FLinearColor(0.15f, 0.15f, 0.15f, 1.0f))
-				]
-
-			// Overlay slot for the main window area
-			+ SOverlay::Slot()
-				.HAlign(HAlign_Fill)
-				.VAlign(VAlign_Fill)
-				[
-					SNew(SVerticalBox)
-
-					+ SVerticalBox::Slot()
-						.AutoHeight()
-						[
-							MenuWidget
-						]
-
-					+ SVerticalBox::Slot()
-						.FillHeight(1.0f)
-						[
-							TabManager->RestoreFrom(Layout, ConstructUnderWindow).ToSharedRef()
-						]
-				]
-
-			// Session hint overlay
-			+ SOverlay::Slot()
-				.HAlign(HAlign_Center)
-				.VAlign(VAlign_Center)
-				[
-					SNew(SBorder)
-						.Visibility(this, &SLoadingProfilerWindow::IsSessionOverlayVisible)
-						.BorderImage(FCoreStyle::Get().GetBrush("PopupText.Background"))
-						.Padding(8.0f)
-						[
-							SNew(STextBlock)
-								.Text(LOCTEXT("SelectTraceOverlayText", "Please select a trace."))
-						]
-				]
-		];
-
+			MenuWidget
+		]
+#endif
+
+		// Version
+		+ SOverlay::Slot()
+		.HAlign(HAlign_Right)
+		.VAlign(VAlign_Top)
+		.Padding(0.0f, -16.0f, 0.0f, 0.0f)
+		[
+			SNew(STextBlock)
+			.Clipping(EWidgetClipping::ClipToBoundsWithoutIntersecting)
+			.Text(LOCTEXT("UnrealInsightsVersion", UNREAL_INSIGHTS_VERSION_STRING_EX))
+			.ColorAndOpacity(FLinearColor(0.15f, 0.15f, 0.15f, 1.0f))
+		]
+
+		// Overlay slot for the main window area
+		+ SOverlay::Slot()
+		.HAlign(HAlign_Fill)
+		.VAlign(VAlign_Fill)
+		[
+			SNew(SVerticalBox)
+
+			+ SVerticalBox::Slot()
+			.AutoHeight()
+			.Padding(FMargin(0.0f, 0.0f, 0.0f, 5.0f))
+			[
+				SNew(SLoadingProfilerToolbar)
+			]
+
+			+ SVerticalBox::Slot()
+			.FillHeight(1.0f)
+			[
+				TabManager->RestoreFrom(Layout, ConstructUnderWindow).ToSharedRef()
+			]
+		]
+
+		// Session hint overlay
+		+ SOverlay::Slot()
+		.HAlign(HAlign_Center)
+		.VAlign(VAlign_Center)
+		[
+			SNew(SBorder)
+			.Visibility(this, &SLoadingProfilerWindow::IsSessionOverlayVisible)
+			.BorderImage(FCoreStyle::Get().GetBrush("PopupText.Background"))
+			.Padding(8.0f)
+			[
+				SNew(STextBlock)
+				.Text(LOCTEXT("SelectTraceOverlayText", "Please select a trace."))
+			]
+		]
+	];
+
+#if !WITH_EDITOR
 	// Tell tab-manager about the global menu bar.
 	TabManager->SetMenuMultiBox(MenuBarBuilder.GetMultiBox(), MenuWidget);
+#endif
 }
-
 END_SLATE_FUNCTION_BUILD_OPTIMIZATION
 
 ////////////////////////////////////////////////////////////////////////////////////////////////////
diff --git a/Engine/Source/Developer/TraceInsights/Private/Insights/LoadingProfiler/Widgets/SLoadingProfilerWindow.h b/Engine/Source/Developer/TraceInsights/Private/Insights/LoadingProfiler/Widgets/SLoadingProfilerWindow.h
--- a/Engine/Source/Developer/TraceInsights/Private/Insights/LoadingProfiler/Widgets/SLoadingProfilerWindow.h
+++ b/Engine/Source/Developer/TraceInsights/Private/Insights/LoadingProfiler/Widgets/SLoadingProfilerWindow.h
@@ -32,7 +32,6 @@ namespace Insights
 struct FLoadingProfilerTabs
 {
 	// Tab identifiers
-	static const FName ToolbarID;
 	static const FName TimingViewID;
 	static const FName EventAggregationTreeViewID;
 	static const FName ObjectTypeAggregationTreeViewID;
@@ -81,9 +80,6 @@ public:
 	TSharedPtr<Insights::SUntypedTableTreeView> GetRequestsTreeView() const { return RequestsTreeView; }
 
 private:
-	TSharedRef<SDockTab> SpawnTab_Toolbar(const FSpawnTabArgs& Args);
-	void OnToolbarTabClosed(TSharedRef<SDockTab> TabBeingClosed);
-
 	TSharedRef<SDockTab> SpawnTab_TimingView(const FSpawnTabArgs& Args);
 	void OnTimingViewTabClosed(TSharedRef<SDockTab> TabBeingClosed);
 
diff --git a/Engine/Source/Developer/TraceInsights/Private/Insights/MemoryProfiler/Widgets/SMemoryProfilerToolbar.cpp b/Engine/Source/Developer/TraceInsights/Private/Insights/MemoryProfiler/Widgets/SMemoryProfilerToolbar.cpp
--- a/Engine/Source/Developer/TraceInsights/Private/Insights/MemoryProfiler/Widgets/SMemoryProfilerToolbar.cpp
+++ b/Engine/Source/Developer/TraceInsights/Private/Insights/MemoryProfiler/Widgets/SMemoryProfilerToolbar.cpp
@@ -2,22 +2,14 @@
 
 #include "SMemoryProfilerToolbar.h"
 
-#include "EditorStyleSet.h"
-#include "Framework/Application/SlateApplication.h"
-#include "Framework/Docking/TabManager.h"
-#include "Framework/MultiBox/MultiBoxDefs.h"
 #include "Framework/MultiBox/MultiBoxBuilder.h"
-#include "Widgets/Docking/SDockTab.h"
-#include "Widgets/Input/SCheckBox.h"
-#include "Widgets/Layout/SBorder.h"
 #include "Widgets/SBoxPanel.h"
-#include "Widgets/SToolTip.h"
 
 // Insights
 #include "Insights/InsightsCommands.h"
 #include "Insights/InsightsManager.h"
+#include "Insights/InsightsStyle.h"
 #include "Insights/MemoryProfiler/MemoryProfilerCommands.h"
-#include "Insights/MemoryProfiler/MemoryProfilerManager.h"
 
 ////////////////////////////////////////////////////////////////////////////////////////////////////
 
@@ -45,9 +37,15 @@ void SMemoryProfilerToolbar::Construct(const FArguments& InArgs)
 		{
 			ToolbarBuilder.BeginSection("View");
 			{
-				ToolbarBuilder.AddToolBarButton(FMemoryProfilerCommands::Get().ToggleTimingViewVisibility);
-				ToolbarBuilder.AddToolBarButton(FMemoryProfilerCommands::Get().ToggleMemInvestigationViewVisibility);
-				ToolbarBuilder.AddToolBarButton(FMemoryProfilerCommands::Get().ToggleMemTagTreeViewVisibility);
+				ToolbarBuilder.AddToolBarButton(FMemoryProfilerCommands::Get().ToggleTimingViewVisibility,
+					NAME_None, TAttribute<FText>(), TAttribute<FText>(),
+					FSlateIcon(FInsightsStyle::GetStyleSetName(), "TimingView.Icon.Large"));
+				ToolbarBuilder.AddToolBarButton(FMemoryProfilerCommands::Get().ToggleMemInvestigationViewVisibility,
+					NAME_None, TAttribute<FText>(), TAttribute<FText>(),
+					FSlateIcon(FInsightsStyle::GetStyleSetName(), "MemInvestigationView.Icon.Large"));
+				ToolbarBuilder.AddToolBarButton(FMemoryProfilerCommands::Get().ToggleMemTagTreeViewVisibility,
+					NAME_None, TAttribute<FText>(), TAttribute<FText>(),
+					FSlateIcon(FInsightsStyle::GetStyleSetName(), "MemTagTreeView.Icon.Large"));
 			}
 			ToolbarBuilder.EndSection();
 		}
@@ -56,7 +54,9 @@ void SMemoryProfilerToolbar::Construct(const FArguments& InArgs)
 		{
 			ToolbarBuilder.BeginSection("Debug");
 			{
-				ToolbarBuilder.AddToolBarButton(FInsightsCommands::Get().ToggleDebugInfo);
+				ToolbarBuilder.AddToolBarButton(FInsightsCommands::Get().ToggleDebugInfo,
+					NAME_None, FText(), TAttribute<FText>(),
+					FSlateIcon(FInsightsStyle::GetStyleSetName(), "Icon.Bug"));
 			}
 			ToolbarBuilder.EndSection();
 		}
@@ -64,45 +64,32 @@ void SMemoryProfilerToolbar::Construct(const FArguments& InArgs)
 
 	TSharedPtr<FUICommandList> CommandList = FInsightsManager::Get()->GetCommandList();
 
-	FToolBarBuilder ToolbarBuilder(CommandList.ToSharedRef(), FMultiBoxCustomization::None);
+	FSlimHorizontalToolBarBuilder ToolbarBuilder(CommandList.ToSharedRef(), FMultiBoxCustomization::None);
 	Local::FillViewToolbar(ToolbarBuilder);
 
-	FToolBarBuilder RightSideToolbarBuilder(CommandList.ToSharedRef(), FMultiBoxCustomization::None);
+	FSlimHorizontalToolBarBuilder RightSideToolbarBuilder(CommandList.ToSharedRef(), FMultiBoxCustomization::None);
 	Local::FillRightSideToolbar(RightSideToolbarBuilder);
 
-	// Create the tool bar!
 	ChildSlot
 	[
 		SNew(SHorizontalBox)
 
-		+SHorizontalBox::Slot()
-		.HAlign(HAlign_Left)
+		+ SHorizontalBox::Slot()
+		.HAlign(HAlign_Fill)
 		.VAlign(VAlign_Center)
 		.FillWidth(1.0)
 		.Padding(0.0f)
 		[
-			SNew(SBorder)
-			.Padding(0)
-			.BorderImage(FEditorStyle::GetBrush("NoBorder"))
-			.IsEnabled(FSlateApplication::Get().GetNormalExecutionAttribute())
-			[
-				ToolbarBuilder.MakeWidget()
-			]
+			ToolbarBuilder.MakeWidget()
 		]
 
-		+SHorizontalBox::Slot()
+		+ SHorizontalBox::Slot()
 		.HAlign(HAlign_Right)
 		.VAlign(VAlign_Center)
 		.AutoWidth()
 		.Padding(0.0f)
 		[
-			SNew(SBorder)
-			.Padding(0)
-			.BorderImage(FEditorStyle::GetBrush("NoBorder"))
-			.IsEnabled(FSlateApplication::Get().GetNormalExecutionAttribute())
-			[
-				RightSideToolbarBuilder.MakeWidget()
-			]
+			RightSideToolbarBuilder.MakeWidget()
 		]
 	];
 }
diff --git a/Engine/Source/Developer/TraceInsights/Private/Insights/MemoryProfiler/Widgets/SMemoryProfilerToolbar.h b/Engine/Source/Developer/TraceInsights/Private/Insights/MemoryProfiler/Widgets/SMemoryProfilerToolbar.h
--- a/Engine/Source/Developer/TraceInsights/Private/Insights/MemoryProfiler/Widgets/SMemoryProfilerToolbar.h
+++ b/Engine/Source/Developer/TraceInsights/Private/Insights/MemoryProfiler/Widgets/SMemoryProfilerToolbar.h
@@ -1,16 +1,10 @@
 // Copyright Epic Games, Inc. All Rights Reserved.
 
-
 #pragma once
 
-#include "CoreMinimal.h"
-#include "Styling/SlateTypes.h"
 #include "Widgets/DeclarativeSyntaxSupport.h"
 #include "Widgets/SCompoundWidget.h"
 
-class SBorder;
-class FToolBarBuilder;
-
 /** Ribbon based toolbar used as a main menu in the Profiler window. */
 class SMemoryProfilerToolbar : public SCompoundWidget
 {
@@ -27,9 +21,9 @@ public:
 	/**
 	 * Construct this widget
 	 *
-	 * @param	InArgs	The declaration data for this widget
+	 * @param InArgs The declaration data for this widget
 	 */
-	void Construct( const FArguments& InArgs );
+	void Construct(const FArguments& InArgs);
 
 private:
 	/** Create the UI commands for the toolbar */
diff --git a/Engine/Source/Developer/TraceInsights/Private/Insights/MemoryProfiler/Widgets/SMemoryProfilerWindow.cpp b/Engine/Source/Developer/TraceInsights/Private/Insights/MemoryProfiler/Widgets/SMemoryProfilerWindow.cpp
--- a/Engine/Source/Developer/TraceInsights/Private/Insights/MemoryProfiler/Widgets/SMemoryProfilerWindow.cpp
+++ b/Engine/Source/Developer/TraceInsights/Private/Insights/MemoryProfiler/Widgets/SMemoryProfilerWindow.cpp
@@ -2,18 +2,18 @@
 
 #include "SMemoryProfilerWindow.h"
 
-#include "EditorStyleSet.h"
 #include "Features/IModularFeatures.h"
+#include "Framework/Commands/UICommandList.h"
 #include "Framework/Docking/LayoutService.h"
+#include "Framework/Docking/TabManager.h"
 #include "Framework/Docking/WorkspaceItem.h"
 #include "Framework/MultiBox/MultiBoxBuilder.h"
 #include "SlateOptMacros.h"
+#include "Styling/CoreStyle.h"
 #include "Widgets/Docking/SDockTab.h"
 #include "Widgets/Images/SImage.h"
 #include "Widgets/Layout/SBorder.h"
 #include "Widgets/Layout/SBox.h"
-#include "Widgets/Layout/SSpacer.h"
-#include "Widgets/Notifications/SNotificationList.h"
 #include "Widgets/SBoxPanel.h"
 #include "Widgets/Text/STextBlock.h"
 
@@ -27,7 +27,6 @@
 #include "Insights/Common/InsightsMenuBuilder.h"
 #include "Insights/InsightsManager.h"
 #include "Insights/InsightsStyle.h"
-#include "Insights/ITimingViewExtender.h"
 #include "Insights/MemoryProfiler/MemoryProfilerManager.h"
 #include "Insights/MemoryProfiler/ViewModels/MemorySharedState.h"
 #include "Insights/MemoryProfiler/Widgets/SMemAllocTableTreeView.h"
@@ -46,7 +45,6 @@
 
 ////////////////////////////////////////////////////////////////////////////////////////////////////
 
-const FName FMemoryProfilerTabs::ToolbarID(TEXT("Toolbar"));
 const FName FMemoryProfilerTabs::TimingViewID(TEXT("TimingView"));
 const FName FMemoryProfilerTabs::MemInvestigationViewID(TEXT("MemInvestigation"));
 const FName FMemoryProfilerTabs::MemTagTreeViewID(TEXT("LowLevelMemTags"));
@@ -55,12 +53,7 @@ const FName FMemoryProfilerTabs::MemAllocTableTreeViewID(TEXT("MemAllocTableTree
 ////////////////////////////////////////////////////////////////////////////////////////////////////
 
 SMemoryProfilerWindow::SMemoryProfilerWindow()
-	: TimingView()
-	, SharedState(MakeShared<FMemorySharedState>())
-	, MemInvestigationView()
-	, MemTagTreeView()
-	, TabManager()
-	, ActiveTimerHandle()
+	: SharedState(MakeShared<FMemorySharedState>())
 	, DurationActive(0.0f)
 {
 }
@@ -87,8 +80,6 @@ SMemoryProfilerWindow::~SMemoryProfilerWindow()
 		check(TimingView == nullptr);
 	}
 
-	HideTab(FMemoryProfilerTabs::ToolbarID);
-
 #if WITH_EDITOR
 	if (DurationActive > 0.0f && FEngineAnalytics::IsAvailable())
 	{
@@ -245,29 +236,6 @@ void SMemoryProfilerWindow::UpdateMemAllocTableTreeView(TSharedPtr<Insights::SMe
 ////////////////////////////////////////////////////////////////////////////////////////////////////
 
 BEGIN_SLATE_FUNCTION_BUILD_OPTIMIZATION
-
-TSharedRef<SDockTab> SMemoryProfilerWindow::SpawnTab_Toolbar(const FSpawnTabArgs& Args)
-{
-	const TSharedRef<SDockTab> DockTab = SNew(SDockTab)
-		.ShouldAutosize(true)
-		.TabRole(ETabRole::PanelTab)
-		[
-			SNew(SMemoryProfilerToolbar)
-		];
-
-	DockTab->SetOnTabClosed(SDockTab::FOnTabClosedCallback::CreateRaw(this, &SMemoryProfilerWindow::OnToolbarTabClosed));
-
-	return DockTab;
-}
-
-////////////////////////////////////////////////////////////////////////////////////////////////////
-
-void SMemoryProfilerWindow::OnToolbarTabClosed(TSharedRef<SDockTab> TabBeingClosed)
-{
-}
-
-////////////////////////////////////////////////////////////////////////////////////////////////////
-
 TSharedRef<SDockTab> SMemoryProfilerWindow::SpawnTab_TimingView(const FSpawnTabArgs& Args)
 {
 	FMemoryProfilerManager::Get()->SetTimingViewVisible(true);
@@ -280,7 +248,7 @@ TSharedRef<SDockTab> SMemoryProfilerWindow::SpawnTab_TimingView(const FSpawnTabA
 		];
 
 	SharedState->SetTimingView(TimingView);
-	IModularFeatures::Get().RegisterModularFeature(Insights::TimingViewExtenderFeatureName, SharedState.Get());
+	IModularFeatures::Get().RegisterModularFeature(Insights::TimingViewExtenderFeatureName, &SharedState.Get());
 
 	TimingView->Reset(true);
 	ResetTimingViewMarkers();
@@ -290,12 +258,13 @@ TSharedRef<SDockTab> SMemoryProfilerWindow::SpawnTab_TimingView(const FSpawnTabA
 
 	return DockTab;
 }
+END_SLATE_FUNCTION_BUILD_OPTIMIZATION
 
 ////////////////////////////////////////////////////////////////////////////////////////////////////
 
 void SMemoryProfilerWindow::OnTimingViewTabClosed(TSharedRef<SDockTab> TabBeingClosed)
 {
-	IModularFeatures::Get().UnregisterModularFeature(Insights::TimingViewExtenderFeatureName, SharedState.Get());
+	IModularFeatures::Get().UnregisterModularFeature(Insights::TimingViewExtenderFeatureName, &SharedState.Get());
 	SharedState->SetTimingView(nullptr);
 
 	TimingView = nullptr;
@@ -305,6 +274,7 @@ void SMemoryProfilerWindow::OnTimingViewTabClosed(TSharedRef<SDockTab> TabBeingC
 
 ////////////////////////////////////////////////////////////////////////////////////////////////////
 
+BEGIN_SLATE_FUNCTION_BUILD_OPTIMIZATION
 TSharedRef<SDockTab> SMemoryProfilerWindow::SpawnTab_MemInvestigationView(const FSpawnTabArgs& Args)
 {
 	FMemoryProfilerManager::Get()->SetMemInvestigationViewVisible(true);
@@ -322,6 +292,7 @@ TSharedRef<SDockTab> SMemoryProfilerWindow::SpawnTab_MemInvestigationView(const
 
 	return DockTab;
 }
+END_SLATE_FUNCTION_BUILD_OPTIMIZATION
 
 ////////////////////////////////////////////////////////////////////////////////////////////////////
 
@@ -333,6 +304,7 @@ void SMemoryProfilerWindow::OnMemInvestigationViewTabClosed(TSharedRef<SDockTab>
 
 ////////////////////////////////////////////////////////////////////////////////////////////////////
 
+BEGIN_SLATE_FUNCTION_BUILD_OPTIMIZATION
 TSharedRef<SDockTab> SMemoryProfilerWindow::SpawnTab_MemTagTreeView(const FSpawnTabArgs& Args)
 {
 	FMemoryProfilerManager::Get()->SetMemTagTreeViewVisible(true);
@@ -350,6 +322,7 @@ TSharedRef<SDockTab> SMemoryProfilerWindow::SpawnTab_MemTagTreeView(const FSpawn
 
 	return DockTab;
 }
+END_SLATE_FUNCTION_BUILD_OPTIMIZATION
 
 ////////////////////////////////////////////////////////////////////////////////////////////////////
 
@@ -361,6 +334,7 @@ void SMemoryProfilerWindow::OnMemTagTreeViewTabClosed(TSharedRef<SDockTab> TabBe
 
 ////////////////////////////////////////////////////////////////////////////////////////////////////
 
+BEGIN_SLATE_FUNCTION_BUILD_OPTIMIZATION
 TSharedRef<SDockTab> SMemoryProfilerWindow::SpawnTab_MemAllocTableTreeView(const FSpawnTabArgs& Args, int32 TabIndex)
 {
 	//FMemoryProfilerManager::Get()->SetMemAllocTableTreeViewVisible(TabIndex, true);
@@ -386,6 +360,7 @@ TSharedRef<SDockTab> SMemoryProfilerWindow::SpawnTab_MemAllocTableTreeView(const
 
 	return DockTab;
 }
+END_SLATE_FUNCTION_BUILD_OPTIMIZATION
 
 ////////////////////////////////////////////////////////////////////////////////////////////////////
 
@@ -447,6 +422,7 @@ void SMemoryProfilerWindow::CloseMemAllocTableTreeTabs()
 
 ////////////////////////////////////////////////////////////////////////////////////////////////////
 
+BEGIN_SLATE_FUNCTION_BUILD_OPTIMIZATION
 TSharedPtr<Insights::SMemAllocTableTreeView> SMemoryProfilerWindow::ShowMemAllocTableTreeViewTab()
 {
 	if (SharedState->GetCurrentQueryTarget()->GetName() == Insights::FQueryTargetWindowSpec::NewWindow)
@@ -461,7 +437,6 @@ TSharedPtr<Insights::SMemAllocTableTreeView> SMemoryProfilerWindow::ShowMemAlloc
 			.SetIcon(FSlateIcon(FInsightsStyle::GetStyleSetName(), "MemAllocTableTreeView.Icon.Small"))
 			.SetGroup(AppMenuGroup.ToSharedRef());
 
-
 		TSharedPtr<Insights::FQueryTargetWindowSpec> NewTarget = MakeShared<Insights::FQueryTargetWindowSpec>(TabId, MemAllocTableTreeViewTabDisplayName);
 		SharedState->AddQueryTarget(NewTarget);
 		SharedState->SetCurrentQueryTarget(NewTarget);
@@ -486,9 +461,11 @@ TSharedPtr<Insights::SMemAllocTableTreeView> SMemoryProfilerWindow::ShowMemAlloc
 
 	return nullptr;
 }
+END_SLATE_FUNCTION_BUILD_OPTIMIZATION
 
 ////////////////////////////////////////////////////////////////////////////////////////////////////
 
+BEGIN_SLATE_FUNCTION_BUILD_OPTIMIZATION
 void SMemoryProfilerWindow::Construct(const FArguments& InArgs, const TSharedRef<SDockTab>& ConstructUnderMajorTab, const TSharedPtr<SWindow>& ConstructUnderWindow)
 {
 	// Create & initialize tab manager.
@@ -501,23 +478,18 @@ void SMemoryProfilerWindow::Construct(const FArguments& InArgs, const TSharedRef
 
 	AppMenuGroup = TabManager->AddLocalWorkspaceMenuCategory(LOCTEXT("MemoryProfilerMenuGroupName", "Memory Insights"));
 
-	TabManager->RegisterTabSpawner(FMemoryProfilerTabs::ToolbarID, FOnSpawnTab::CreateRaw(this, &SMemoryProfilerWindow::SpawnTab_Toolbar))
-		.SetDisplayName(LOCTEXT("DeviceToolbarTabTitle", "Toolbar"))
-		.SetIcon(FSlateIcon(FInsightsStyle::GetStyleSetName(), "Toolbar.Icon.Small"))
-		.SetGroup(AppMenuGroup.ToSharedRef());
-
 	TabManager->RegisterTabSpawner(FMemoryProfilerTabs::TimingViewID, FOnSpawnTab::CreateRaw(this, &SMemoryProfilerWindow::SpawnTab_TimingView))
-		.SetDisplayName(LOCTEXT("MemoryProfiler.TimingViewTabTitle", "Timing View"))
+		.SetDisplayName(LOCTEXT("TimingViewTabTitle", "Timing View"))
 		.SetIcon(FSlateIcon(FInsightsStyle::GetStyleSetName(), "TimingView.Icon.Small"))
 		.SetGroup(AppMenuGroup.ToSharedRef());
 
 	TabManager->RegisterTabSpawner(FMemoryProfilerTabs::MemInvestigationViewID, FOnSpawnTab::CreateRaw(this, &SMemoryProfilerWindow::SpawnTab_MemInvestigationView))
-		.SetDisplayName(LOCTEXT("MemoryProfiler.MemInvestigationViewTabTitle", "Investigation"))
+		.SetDisplayName(LOCTEXT("MemInvestigationViewTabTitle", "Investigation"))
 		.SetIcon(FSlateIcon(FInsightsStyle::GetStyleSetName(), "MemInvestigationView.Icon.Small"))
 		.SetGroup(AppMenuGroup.ToSharedRef());
 
 	TabManager->RegisterTabSpawner(FMemoryProfilerTabs::MemTagTreeViewID, FOnSpawnTab::CreateRaw(this, &SMemoryProfilerWindow::SpawnTab_MemTagTreeView))
-		.SetDisplayName(LOCTEXT("MemoryProfiler.MemTagTreeViewTabTitle", "LLM Tags"))
+		.SetDisplayName(LOCTEXT("MemTagTreeViewTabTitle", "LLM Tags"))
 		.SetIcon(FSlateIcon(FInsightsStyle::GetStyleSetName(), "MemTagTreeView.Icon.Small"))
 		.SetGroup(AppMenuGroup.ToSharedRef());
 
@@ -525,37 +497,25 @@ void SMemoryProfilerWindow::Construct(const FArguments& InArgs, const TSharedRef
 	ensure(MemoryProfilerManager.IsValid());
 
 	// Create tab layout.
-	TSharedRef<FTabManager::FLayout> Layout = FTabManager::NewLayout("InsightsMemoryProfilerLayout_v1.0")
+	TSharedRef<FTabManager::FLayout> Layout = FTabManager::NewLayout("InsightsMemoryProfilerLayout_v1.2")
 		->AddArea
 		(
 			FTabManager::NewPrimaryArea()
-			->SetOrientation(Orient_Vertical)
+			->SetOrientation(Orient_Horizontal)
 			->Split
 			(
 				FTabManager::NewStack()
-				->AddTab(FMemoryProfilerTabs::ToolbarID, ETabState::OpenedTab)
+				->SetSizeCoefficient(0.77f)
 				->SetHideTabWell(true)
+				->AddTab(FMemoryProfilerTabs::TimingViewID, ETabState::OpenedTab)
 			)
 			->Split
 			(
-				FTabManager::NewSplitter()
-				->SetOrientation(Orient_Horizontal)
-				->SetSizeCoefficient(1.0f)
-				->Split
-				(
-					FTabManager::NewStack()
-					->SetSizeCoefficient(0.77f)
-					->SetHideTabWell(true)
-					->AddTab(FMemoryProfilerTabs::TimingViewID, ETabState::OpenedTab)
-				)
-				->Split
-				(
-					FTabManager::NewStack()
-					->SetSizeCoefficient(0.23f)
-					->AddTab(FMemoryProfilerTabs::MemInvestigationViewID, ETabState::OpenedTab)
-					->AddTab(FMemoryProfilerTabs::MemTagTreeViewID, ETabState::OpenedTab)
-					->SetForegroundTab(FMemoryProfilerTabs::MemInvestigationViewID)
-				)
+				FTabManager::NewStack()
+				->SetSizeCoefficient(0.23f)
+				->AddTab(FMemoryProfilerTabs::MemInvestigationViewID, ETabState::OpenedTab)
+				->AddTab(FMemoryProfilerTabs::MemTagTreeViewID, ETabState::OpenedTab)
+				->SetForegroundTab(FMemoryProfilerTabs::MemInvestigationViewID)
 			)
 		);
 
@@ -563,7 +523,6 @@ void SMemoryProfilerWindow::Construct(const FArguments& InArgs, const TSharedRef
 
 	// Create & initialize main menu.
 	FMenuBarBuilder MenuBarBuilder = FMenuBarBuilder(TSharedPtr<FUICommandList>());
-
 	MenuBarBuilder.AddPullDownMenu(
 		LOCTEXT("MenuLabel", "Menu"),
 		FText::GetEmpty(),
@@ -571,64 +530,80 @@ void SMemoryProfilerWindow::Construct(const FArguments& InArgs, const TSharedRef
 		FName(TEXT("Menu"))
 	);
 
+#if !WITH_EDITOR
 	TSharedRef<SWidget> MenuWidget = MenuBarBuilder.MakeWidget();
+	MenuWidget->SetClipping(EWidgetClipping::ClipToBoundsWithoutIntersecting);
+#endif
 
 	ChildSlot
+	[
+		SNew(SOverlay)
+
+#if !WITH_EDITOR
+		// Menu
+		+ SOverlay::Slot()
+		.HAlign(HAlign_Left)
+		.VAlign(VAlign_Top)
+		.Padding(34.0f, -60.0f, 0.0f, 0.0f)
 		[
-			SNew(SOverlay)
-
-			// Version
-			+ SOverlay::Slot()
-				.HAlign(HAlign_Right)
-				.VAlign(VAlign_Top)
-				.Padding(0.0f, -16.0f, 0.0f, 0.0f)
-				[
-					SNew(STextBlock)
-						.Clipping(EWidgetClipping::ClipToBoundsWithoutIntersecting)
-						.Text(LOCTEXT("UnrealInsightsVersion", UNREAL_INSIGHTS_VERSION_STRING_EX))
-						.ColorAndOpacity(FLinearColor(0.15f, 0.15f, 0.15f, 1.0f))
-				]
-
-			// Overlay slot for the main window area
-			+ SOverlay::Slot()
-				.HAlign(HAlign_Fill)
-				.VAlign(VAlign_Fill)
-				[
-					SNew(SVerticalBox)
-
-					+ SVerticalBox::Slot()
-						.AutoHeight()
-						[
-							MenuWidget
-						]
-
-					+ SVerticalBox::Slot()
-						.FillHeight(1.0f)
-						[
-							TabManager->RestoreFrom(Layout, ConstructUnderWindow).ToSharedRef()
-						]
-				]
-
-			// Session hint overlay
-			+ SOverlay::Slot()
-				.HAlign(HAlign_Center)
-				.VAlign(VAlign_Center)
-				[
-					SNew(SBorder)
-						.Visibility(this, &SMemoryProfilerWindow::IsSessionOverlayVisible)
-						.BorderImage(FEditorStyle::GetBrush("NotificationList.ItemBackground"))
-						.Padding(8.0f)
-						[
-							SNew(STextBlock)
-								.Text(LOCTEXT("SelectTraceOverlayText", "Please select a trace."))
-						]
-				]
-		];
-
+			MenuWidget
+		]
+#endif
+
+		// Version
+		+ SOverlay::Slot()
+		.HAlign(HAlign_Right)
+		.VAlign(VAlign_Top)
+		.Padding(0.0f, -16.0f, 0.0f, 0.0f)
+		[
+			SNew(STextBlock)
+			.Clipping(EWidgetClipping::ClipToBoundsWithoutIntersecting)
+			.Text(LOCTEXT("UnrealInsightsVersion", UNREAL_INSIGHTS_VERSION_STRING_EX))
+			.ColorAndOpacity(FLinearColor(0.15f, 0.15f, 0.15f, 1.0f))
+		]
+
+		// Overlay slot for the main window area
+		+ SOverlay::Slot()
+		.HAlign(HAlign_Fill)
+		.VAlign(VAlign_Fill)
+		[
+			SNew(SVerticalBox)
+
+			+ SVerticalBox::Slot()
+			.AutoHeight()
+			.Padding(FMargin(0.0f, 0.0f, 0.0f, 5.0f))
+			[
+				SNew(SMemoryProfilerToolbar)
+			]
+
+			+ SVerticalBox::Slot()
+			.FillHeight(1.0f)
+			[
+				TabManager->RestoreFrom(Layout, ConstructUnderWindow).ToSharedRef()
+			]
+		]
+
+		// Session hint overlay
+		+ SOverlay::Slot()
+		.HAlign(HAlign_Center)
+		.VAlign(VAlign_Center)
+		[
+			SNew(SBorder)
+			.Visibility(this, &SMemoryProfilerWindow::IsSessionOverlayVisible)
+			.BorderImage(FCoreStyle::Get().GetBrush("PopupText.Background"))
+			.Padding(8.0f)
+			[
+				SNew(STextBlock)
+				.Text(LOCTEXT("SelectTraceOverlayText", "Please select a trace."))
+			]
+		]
+	];
+
+#if !WITH_EDITOR
 	// Tell tab-manager about the global menu bar.
 	TabManager->SetMenuMultiBox(MenuBarBuilder.GetMultiBox(), MenuWidget);
+#endif
 }
-
 END_SLATE_FUNCTION_BUILD_OPTIMIZATION
 
 ////////////////////////////////////////////////////////////////////////////////////////////////////
diff --git a/Engine/Source/Developer/TraceInsights/Private/Insights/MemoryProfiler/Widgets/SMemoryProfilerWindow.h b/Engine/Source/Developer/TraceInsights/Private/Insights/MemoryProfiler/Widgets/SMemoryProfilerWindow.h
--- a/Engine/Source/Developer/TraceInsights/Private/Insights/MemoryProfiler/Widgets/SMemoryProfilerWindow.h
+++ b/Engine/Source/Developer/TraceInsights/Private/Insights/MemoryProfiler/Widgets/SMemoryProfilerWindow.h
@@ -38,7 +38,6 @@ namespace Insights
 struct FMemoryProfilerTabs
 {
 	// Tab identifiers
-	static const FName ToolbarID;
 	static const FName TimingViewID;
 	static const FName MemInvestigationViewID;
 	static const FName MemTagTreeViewID;
@@ -93,9 +92,6 @@ public:
 	void CloseMemAllocTableTreeTabs();
 
 private:
-	TSharedRef<SDockTab> SpawnTab_Toolbar(const FSpawnTabArgs& Args);
-	void OnToolbarTabClosed(TSharedRef<SDockTab> TabBeingClosed);
-
 	TSharedRef<SDockTab> SpawnTab_TimingView(const FSpawnTabArgs& Args);
 	void OnTimingViewTabClosed(TSharedRef<SDockTab> TabBeingClosed);
 
@@ -180,10 +176,12 @@ private:
 	virtual FReply OnDragOver(const FGeometry& MyGeometry, const FDragDropEvent& DragDropEvent)  override;
 
 private:
+	TSharedRef<FMemorySharedState> SharedState;
+
 	/** The Timing view (multi-track) widget */
 	TSharedPtr<STimingView> TimingView;
+
 	TArray<TSharedRef<Insights::FTimeMarker>> CustomTimeMarkers;
-	TSharedPtr<FMemorySharedState> SharedState;
 
 	/** The Memory Investigation (Allocation Queries) view widget */
 	TSharedPtr<SMemInvestigationView> MemInvestigationView;
diff --git a/Engine/Source/Developer/TraceInsights/Private/Insights/NetworkingProfiler/Widgets/SNetworkingProfilerWindow.cpp b/Engine/Source/Developer/TraceInsights/Private/Insights/NetworkingProfiler/Widgets/SNetworkingProfilerWindow.cpp
--- a/Engine/Source/Developer/TraceInsights/Private/Insights/NetworkingProfiler/Widgets/SNetworkingProfilerWindow.cpp
+++ b/Engine/Source/Developer/TraceInsights/Private/Insights/NetworkingProfiler/Widgets/SNetworkingProfilerWindow.cpp
@@ -3,6 +3,8 @@
 #include "SNetworkingProfilerWindow.h"
 
 #include "Framework/Commands/UICommandList.h"
+#include "Framework/Docking/LayoutService.h"
+#include "Framework/Docking/TabManager.h"
 #include "Framework/Docking/WorkspaceItem.h"
 #include "Framework/MultiBox/MultiBoxBuilder.h"
 #include "SlateOptMacros.h"
@@ -11,8 +13,6 @@
 #include "Widgets/Images/SImage.h"
 #include "Widgets/Layout/SBorder.h"
 #include "Widgets/Layout/SBox.h"
-#include "Widgets/Layout/SSpacer.h"
-#include "Widgets/Notifications/SNotificationList.h"
 #include "Widgets/SBoxPanel.h"
 #include "Widgets/Text/STextBlock.h"
 
@@ -31,6 +31,7 @@
 #include "Insights/NetworkingProfiler/Widgets/SNetworkingProfilerToolbar.h"
 #include "Insights/NetworkingProfiler/Widgets/SPacketContentView.h"
 #include "Insights/NetworkingProfiler/Widgets/SPacketView.h"
+#include "Insights/TraceInsightsModule.h"
 #include "Insights/Version.h"
 #include "Insights/Widgets/SInsightsSettings.h"
 
@@ -40,7 +41,6 @@
 
 ////////////////////////////////////////////////////////////////////////////////////////////////////
 
-const FName FNetworkingProfilerTabs::ToolbarID(TEXT("Toolbar"));
 const FName FNetworkingProfilerTabs::PacketViewID(TEXT("PacketView"));
 const FName FNetworkingProfilerTabs::PacketContentViewID(TEXT("PacketContent"));
 const FName FNetworkingProfilerTabs::NetStatsViewID(TEXT("NetStats"));
@@ -48,12 +48,11 @@ const FName FNetworkingProfilerTabs::NetStatsViewID(TEXT("NetStats"));
 ////////////////////////////////////////////////////////////////////////////////////////////////////
 
 SNetworkingProfilerWindow::SNetworkingProfilerWindow()
-	: AvailableGameInstances()
-	, SelectedGameInstance(nullptr)
-	, AvailableConnections()
-	, SelectedConnection(nullptr)
-	, AvailableConnectionModes()
-	, SelectedConnectionMode(nullptr)
+	: SelectedPacketStartIndex(0)
+	, SelectedPacketEndIndex(0)
+	, SelectionStartPosition(0)
+	, SelectionEndPosition(0)
+	, SelectedEventTypeIndex(InvalidEventTypeIndex)
 	, DurationActive(0.0f)
 {
 }
@@ -80,8 +79,6 @@ SNetworkingProfilerWindow::~SNetworkingProfilerWindow()
 		check(PacketView == nullptr);
 	}
 
-	HideTab(FNetworkingProfilerTabs::ToolbarID);
-
 #if WITH_EDITOR
 	if (DurationActive > 0.0f && FEngineAnalytics::IsAvailable())
 	{
@@ -120,34 +117,13 @@ void SNetworkingProfilerWindow::Reset()
 	SelectedPacketEndIndex = 0;
 	SelectionStartPosition = 0;
 	SelectionEndPosition = 0;
-}
-
-////////////////////////////////////////////////////////////////////////////////////////////////////
-BEGIN_SLATE_FUNCTION_BUILD_OPTIMIZATION
-////////////////////////////////////////////////////////////////////////////////////////////////////
-
-TSharedRef<SDockTab> SNetworkingProfilerWindow::SpawnTab_Toolbar(const FSpawnTabArgs& Args)
-{
-	const TSharedRef<SDockTab> DockTab = SNew(SDockTab)
-		.ShouldAutosize(true)
-		.TabRole(ETabRole::PanelTab)
-		[
-			SNew(SNetworkingProfilerToolbar, SharedThis(this))
-		];
-
-	DockTab->SetOnTabClosed(SDockTab::FOnTabClosedCallback::CreateRaw(this, &SNetworkingProfilerWindow::OnToolbarTabClosed));
-
-	return DockTab;
-}
-
-////////////////////////////////////////////////////////////////////////////////////////////////////
 
-void SNetworkingProfilerWindow::OnToolbarTabClosed(TSharedRef<SDockTab> TabBeingClosed)
-{
+	SelectedEventTypeIndex = InvalidEventTypeIndex;
 }
 
 ////////////////////////////////////////////////////////////////////////////////////////////////////
 
+BEGIN_SLATE_FUNCTION_BUILD_OPTIMIZATION
 TSharedRef<SDockTab> SNetworkingProfilerWindow::SpawnTab_PacketView(const FSpawnTabArgs& Args)
 {
 	const TSharedRef<SDockTab> DockTab = SNew(SDockTab)
@@ -161,6 +137,7 @@ TSharedRef<SDockTab> SNetworkingProfilerWindow::SpawnTab_PacketView(const FSpawn
 
 	return DockTab;
 }
+END_SLATE_FUNCTION_BUILD_OPTIMIZATION
 
 ////////////////////////////////////////////////////////////////////////////////////////////////////
 
@@ -171,6 +148,7 @@ void SNetworkingProfilerWindow::OnPacketViewTabClosed(TSharedRef<SDockTab> TabBe
 
 ////////////////////////////////////////////////////////////////////////////////////////////////////
 
+BEGIN_SLATE_FUNCTION_BUILD_OPTIMIZATION
 TSharedRef<SDockTab> SNetworkingProfilerWindow::SpawnTab_PacketContentView(const FSpawnTabArgs& Args)
 {
 	const TSharedRef<SDockTab> DockTab = SNew(SDockTab)
@@ -184,6 +162,7 @@ TSharedRef<SDockTab> SNetworkingProfilerWindow::SpawnTab_PacketContentView(const
 
 	return DockTab;
 }
+END_SLATE_FUNCTION_BUILD_OPTIMIZATION
 
 ////////////////////////////////////////////////////////////////////////////////////////////////////
 
@@ -194,6 +173,7 @@ void SNetworkingProfilerWindow::OnPacketContentViewTabClosed(TSharedRef<SDockTab
 
 ////////////////////////////////////////////////////////////////////////////////////////////////////
 
+BEGIN_SLATE_FUNCTION_BUILD_OPTIMIZATION
 TSharedRef<SDockTab> SNetworkingProfilerWindow::SpawnTab_NetStatsView(const FSpawnTabArgs& Args)
 {
 	const TSharedRef<SDockTab> DockTab = SNew(SDockTab)
@@ -207,6 +187,7 @@ TSharedRef<SDockTab> SNetworkingProfilerWindow::SpawnTab_NetStatsView(const FSpa
 
 	return DockTab;
 }
+END_SLATE_FUNCTION_BUILD_OPTIMIZATION
 
 ////////////////////////////////////////////////////////////////////////////////////////////////////
 
@@ -217,31 +198,33 @@ void SNetworkingProfilerWindow::OnNetStatsViewTabClosed(TSharedRef<SDockTab> Tab
 
 ////////////////////////////////////////////////////////////////////////////////////////////////////
 
+BEGIN_SLATE_FUNCTION_BUILD_OPTIMIZATION
 void SNetworkingProfilerWindow::Construct(const FArguments& InArgs, const TSharedRef<SDockTab>& ConstructUnderMajorTab, const TSharedPtr<SWindow>& ConstructUnderWindow)
 {
 	CommandList = MakeShared<FUICommandList>();
 
 	// Create & initialize tab manager.
 	TabManager = FGlobalTabmanager::Get()->NewTabManager(ConstructUnderMajorTab);
-	TSharedRef<FWorkspaceItem> AppMenuGroup = TabManager->AddLocalWorkspaceMenuCategory(LOCTEXT("NetworkingProfilerMenuGroupName", "Networking Insights"));
+	const auto& PersistLayout = [](const TSharedRef<FTabManager::FLayout>& LayoutToSave)
+	{
+		FLayoutSaveRestore::SaveToConfig(FTraceInsightsModule::GetUnrealInsightsLayoutIni(), LayoutToSave);
+	};
+	TabManager->SetOnPersistLayout(FTabManager::FOnPersistLayout::CreateLambda(PersistLayout));
 
-	TabManager->RegisterTabSpawner(FNetworkingProfilerTabs::ToolbarID, FOnSpawnTab::CreateRaw(this, &SNetworkingProfilerWindow::SpawnTab_Toolbar))
-		.SetDisplayName(LOCTEXT("DeviceToolbarTabTitle", "Toolbar"))
-		.SetIcon(FSlateIcon(FInsightsStyle::GetStyleSetName(), "Toolbar.Icon.Small"))
-		.SetGroup(AppMenuGroup);
+	TSharedRef<FWorkspaceItem> AppMenuGroup = TabManager->AddLocalWorkspaceMenuCategory(LOCTEXT("NetworkingProfilerMenuGroupName", "Networking Insights"));
 
 	TabManager->RegisterTabSpawner(FNetworkingProfilerTabs::PacketViewID, FOnSpawnTab::CreateRaw(this, &SNetworkingProfilerWindow::SpawnTab_PacketView))
-		.SetDisplayName(LOCTEXT("NetworkingProfiler.PacketViewTabTitle", "Packet View"))
+		.SetDisplayName(LOCTEXT("PacketViewTabTitle", "Packet View"))
 		.SetIcon(FSlateIcon(FInsightsStyle::GetStyleSetName(), "PacketView.Icon.Small"))
 		.SetGroup(AppMenuGroup);
 
 	TabManager->RegisterTabSpawner(FNetworkingProfilerTabs::PacketContentViewID, FOnSpawnTab::CreateRaw(this, &SNetworkingProfilerWindow::SpawnTab_PacketContentView))
-		.SetDisplayName(LOCTEXT("NetworkingProfiler.PacketContentViewTabTitle", "Packet Content"))
+		.SetDisplayName(LOCTEXT("PacketContentViewTabTitle", "Packet Content"))
 		.SetIcon(FSlateIcon(FInsightsStyle::GetStyleSetName(), "PacketContentView.Icon.Small"))
 		.SetGroup(AppMenuGroup);
 
 	TabManager->RegisterTabSpawner(FNetworkingProfilerTabs::NetStatsViewID, FOnSpawnTab::CreateRaw(this, &SNetworkingProfilerWindow::SpawnTab_NetStatsView))
-		.SetDisplayName(LOCTEXT("NetworkingProfiler.NetStatsViewTabTitle", "Net Stats"))
+		.SetDisplayName(LOCTEXT("NetStatsViewTabTitle", "Net Stats"))
 		.SetIcon(FSlateIcon(FInsightsStyle::GetStyleSetName(), "NetStatsView.Icon.Small"))
 		.SetGroup(AppMenuGroup);
 
@@ -249,55 +232,44 @@ void SNetworkingProfilerWindow::Construct(const FArguments& InArgs, const TShare
 	ensure(NetworkingProfilerManager.IsValid());
 
 	// Create tab layout.
-	const TSharedRef<FTabManager::FLayout> Layout = FTabManager::NewLayout("InsightsNetworkingProfilerLayout_v1.0")
+	TSharedRef<FTabManager::FLayout> Layout = FTabManager::NewLayout("InsightsNetworkingProfilerLayout_v1.2")
 		->AddArea
 		(
 			FTabManager::NewPrimaryArea()
-			->SetOrientation(Orient_Vertical)
-			->Split
-			(
-				FTabManager::NewStack()
-				->AddTab(FNetworkingProfilerTabs::ToolbarID, ETabState::OpenedTab)
-				->SetHideTabWell(true)
-			)
+			->SetOrientation(Orient_Horizontal)
 			->Split
 			(
 				FTabManager::NewSplitter()
-				->SetOrientation(Orient_Horizontal)
-				->SetSizeCoefficient(1.0f)
+				->SetOrientation(Orient_Vertical)
+				->SetSizeCoefficient(0.65f)
 				->Split
 				(
-					FTabManager::NewSplitter()
-					->SetOrientation(Orient_Vertical)
-					->SetSizeCoefficient(0.65f)
-					->Split
-					(
-						FTabManager::NewStack()
-						->SetSizeCoefficient(0.35f)
-						->SetHideTabWell(true)
-						->AddTab(FNetworkingProfilerTabs::PacketViewID, ETabState::OpenedTab)
-					)
-					->Split
-					(
-						FTabManager::NewStack()
-						->SetSizeCoefficient(0.65f)
-						->SetHideTabWell(true)
-						->AddTab(FNetworkingProfilerTabs::PacketContentViewID, ETabState::OpenedTab)
-					)
+					FTabManager::NewStack()
+					->SetSizeCoefficient(0.35f)
+					->SetHideTabWell(true)
+					->AddTab(FNetworkingProfilerTabs::PacketViewID, ETabState::OpenedTab)
 				)
 				->Split
 				(
 					FTabManager::NewStack()
-					->SetSizeCoefficient(0.35f)
-					->AddTab(FNetworkingProfilerTabs::NetStatsViewID, ETabState::OpenedTab)
-					//->SetForegroundTab(FNetworkingProfilerTabs::NetStatsViewID)
+					->SetSizeCoefficient(0.65f)
+					->SetHideTabWell(true)
+					->AddTab(FNetworkingProfilerTabs::PacketContentViewID, ETabState::OpenedTab)
 				)
 			)
+			->Split
+			(
+				FTabManager::NewStack()
+				->SetSizeCoefficient(0.35f)
+				->AddTab(FNetworkingProfilerTabs::NetStatsViewID, ETabState::OpenedTab)
+				->SetForegroundTab(FNetworkingProfilerTabs::NetStatsViewID)
+			)
 		);
 
+	Layout = FLayoutSaveRestore::LoadFromConfig(FTraceInsightsModule::GetUnrealInsightsLayoutIni(), Layout);
+
 	// Create & initialize main menu.
 	FMenuBarBuilder MenuBarBuilder = FMenuBarBuilder(TSharedPtr<FUICommandList>());
-
 	MenuBarBuilder.AddPullDownMenu(
 		LOCTEXT("MenuLabel", "Menu"),
 		FText::GetEmpty(),
@@ -305,67 +277,84 @@ void SNetworkingProfilerWindow::Construct(const FArguments& InArgs, const TShare
 		FName(TEXT("Menu"))
 	);
 
+#if !WITH_EDITOR
 	TSharedRef<SWidget> MenuWidget = MenuBarBuilder.MakeWidget();
+	MenuWidget->SetClipping(EWidgetClipping::ClipToBoundsWithoutIntersecting);
+#endif
+
 	ChildSlot
+	[
+		SNew(SOverlay)
+
+#if !WITH_EDITOR
+		// Menu
+		+ SOverlay::Slot()
+		.HAlign(HAlign_Left)
+		.VAlign(VAlign_Top)
+		.Padding(34.0f, -60.0f, 0.0f, 0.0f)
 		[
-			SNew(SOverlay)
-
-			// Version
-			+ SOverlay::Slot()
-				.HAlign(HAlign_Right)
-				.VAlign(VAlign_Top)
-				.Padding(0.0f, -16.0f, 0.0f, 0.0f)
-				[
-					SNew(STextBlock)
-						.Clipping(EWidgetClipping::ClipToBoundsWithoutIntersecting)
-						.Text(LOCTEXT("UnrealInsightsVersion", UNREAL_INSIGHTS_VERSION_STRING_EX))
-						.ColorAndOpacity(FLinearColor(0.15f, 0.15f, 0.15f, 1.0f))
-				]
-
-			// Overlay slot for the main window area
-			+ SOverlay::Slot()
-				.HAlign(HAlign_Fill)
-				.VAlign(VAlign_Fill)
-				[
-					SNew(SVerticalBox)
-
-					+ SVerticalBox::Slot()
-						.AutoHeight()
-						[
-							MenuWidget
-						]
-
-					+ SVerticalBox::Slot()
-						.FillHeight(1.0f)
-						[
-							TabManager->RestoreFrom(Layout, ConstructUnderWindow).ToSharedRef()
-						]
-				]
-
-			// Session hint overlay
-			+ SOverlay::Slot()
-				.HAlign(HAlign_Center)
-				.VAlign(VAlign_Center)
-				[
-					SNew(SBorder)
-						.Visibility(this, &SNetworkingProfilerWindow::IsSessionOverlayVisible)
-						.BorderImage(FCoreStyle::Get().GetBrush("PopupText.Background"))
-						.Padding(8.0f)
-						[
-							SNew(STextBlock)
-								.Text(LOCTEXT("SelectTraceOverlayText", "Please select a trace."))
-						]
-				]
-		];
-
+			MenuWidget
+		]
+#endif
+
+		// Version
+		+ SOverlay::Slot()
+		.HAlign(HAlign_Right)
+		.VAlign(VAlign_Top)
+		.Padding(0.0f, -16.0f, 0.0f, 0.0f)
+		[
+			SNew(STextBlock)
+			.Clipping(EWidgetClipping::ClipToBoundsWithoutIntersecting)
+			.Text(LOCTEXT("UnrealInsightsVersion", UNREAL_INSIGHTS_VERSION_STRING_EX))
+			.ColorAndOpacity(FLinearColor(0.15f, 0.15f, 0.15f, 1.0f))
+		]
+
+		// Overlay slot for the main window area
+		+ SOverlay::Slot()
+		.HAlign(HAlign_Fill)
+		.VAlign(VAlign_Fill)
+		[
+			SNew(SVerticalBox)
+
+			+ SVerticalBox::Slot()
+			.AutoHeight()
+			.Padding(FMargin(0.0f, 0.0f, 0.0f, 5.0f))
+			[
+				SNew(SNetworkingProfilerToolbar, SharedThis(this))
+			]
+
+			+ SVerticalBox::Slot()
+			.FillHeight(1.0f)
+			[
+				TabManager->RestoreFrom(Layout, ConstructUnderWindow).ToSharedRef()
+			]
+		]
+
+		// Session hint overlay
+		+ SOverlay::Slot()
+		.HAlign(HAlign_Center)
+		.VAlign(VAlign_Center)
+		[
+			SNew(SBorder)
+			.Visibility(this, &SNetworkingProfilerWindow::IsSessionOverlayVisible)
+			.BorderImage(FCoreStyle::Get().GetBrush("PopupText.Background"))
+			.Padding(8.0f)
+			[
+				SNew(STextBlock)
+				.Text(LOCTEXT("SelectTraceOverlayText", "Please select a trace."))
+			]
+		]
+	];
+
+#if !WITH_EDITOR
 	// Tell tab-manager about the global menu bar.
 	TabManager->SetMenuMultiBox(MenuBarBuilder.GetMultiBox(), MenuWidget);
+#endif
 
 	BindCommands();
 }
-
-////////////////////////////////////////////////////////////////////////////////////////////////////
 END_SLATE_FUNCTION_BUILD_OPTIMIZATION
+
 ////////////////////////////////////////////////////////////////////////////////////////////////////
 
 void SNetworkingProfilerWindow::BindCommands()
@@ -538,7 +527,7 @@ void SNetworkingProfilerWindow::UpdateAvailableGameInstances()
 			bIsInstanceSelected = true;
 		}
 	}
-	
+
 	AvailableGameInstances.Reset();
 
 	TSharedPtr<const TraceServices::IAnalysisSession> Session = FInsightsManager::Get()->GetSession();
diff --git a/Engine/Source/Developer/TraceInsights/Private/Insights/NetworkingProfiler/Widgets/SNetworkingProfilerWindow.h b/Engine/Source/Developer/TraceInsights/Private/Insights/NetworkingProfiler/Widgets/SNetworkingProfilerWindow.h
--- a/Engine/Source/Developer/TraceInsights/Private/Insights/NetworkingProfiler/Widgets/SNetworkingProfilerWindow.h
+++ b/Engine/Source/Developer/TraceInsights/Private/Insights/NetworkingProfiler/Widgets/SNetworkingProfilerWindow.h
@@ -32,7 +32,6 @@ class SNetStatsView;
 struct FNetworkingProfilerTabs
 {
 	// Tab identifiers
-	static const FName ToolbarID;
 	static const FName PacketViewID;
 	static const FName PacketContentViewID;
 	static const FName NetStatsViewID;
@@ -144,9 +143,6 @@ public:
 	void SetSelectedEventTypeIndex(uint32 InEventTypeIndex);
 
 private:
-	TSharedRef<SDockTab> SpawnTab_Toolbar(const FSpawnTabArgs& Args);
-	void OnToolbarTabClosed(TSharedRef<SDockTab> TabBeingClosed);
-
 	TSharedRef<SDockTab> SpawnTab_PacketView(const FSpawnTabArgs& Args);
 	void OnPacketViewTabClosed(TSharedRef<SDockTab> TabBeingClosed);
 
diff --git a/Engine/Source/Developer/TraceInsights/Private/Insights/Widgets/SSessionInfoWindow.cpp b/Engine/Source/Developer/TraceInsights/Private/Insights/Widgets/SSessionInfoWindow.cpp
--- a/Engine/Source/Developer/TraceInsights/Private/Insights/Widgets/SSessionInfoWindow.cpp
+++ b/Engine/Source/Developer/TraceInsights/Private/Insights/Widgets/SSessionInfoWindow.cpp
@@ -39,7 +39,6 @@ const FName FSessionInfoTabs::SessionInfoID(TEXT("SessionInfo"));
 
 SSessionInfoWindow::SSessionInfoWindow()
 	: DurationActive(0.0f)
-	, TabManager()
 {
 }
 
@@ -68,7 +67,7 @@ void SSessionInfoWindow::Construct(const FArguments& InArgs, const TSharedRef<SD
 
 	TabManager->RegisterTabSpawner(FSessionInfoTabs::SessionInfoID, FOnSpawnTab::CreateRaw(this, &SSessionInfoWindow::SpawnTab_SessionInfo))
 		.SetDisplayName(LOCTEXT("SessionInfo", "Session Info"))
-		.SetIcon(FSlateIcon(FInsightsStyle::GetStyleSetName(), "Toolbar.Icon.Small"))
+		.SetIcon(FSlateIcon(FInsightsStyle::GetStyleSetName(), "SessionInfo.Icon.Small"))
 		.SetGroup(AppMenuGroup);
 
 	TSharedRef<FTabManager::FLayout> Layout = []() -> TSharedRef<FTabManager::FLayout>
@@ -90,7 +89,6 @@ void SSessionInfoWindow::Construct(const FArguments& InArgs, const TSharedRef<SD
 
 	// Create & initialize main menu.
 	FMenuBarBuilder MenuBarBuilder = FMenuBarBuilder(TSharedPtr<FUICommandList>());
-
 	MenuBarBuilder.AddPullDownMenu(
 		LOCTEXT("MenuLabel", "Menu"),
 		FText::GetEmpty(),
@@ -98,10 +96,26 @@ void SSessionInfoWindow::Construct(const FArguments& InArgs, const TSharedRef<SD
 		FName(TEXT("Menu"))
 	);
 
+#if !WITH_EDITOR
+	TSharedRef<SWidget> MenuWidget = MenuBarBuilder.MakeWidget();
+	MenuWidget->SetClipping(EWidgetClipping::ClipToBoundsWithoutIntersecting);
+#endif
+
 	ChildSlot
 	[
 		SNew(SOverlay)
 
+#if !WITH_EDITOR
+		// Menu
+		+ SOverlay::Slot()
+		.HAlign(HAlign_Left)
+		.VAlign(VAlign_Top)
+		.Padding(34.0f, -60.0f, 0.0f, 0.0f)
+		[
+			MenuWidget
+		]
+#endif
+
 		// Version
 		+ SOverlay::Slot()
 		.HAlign(HAlign_Right)
@@ -113,24 +127,13 @@ void SSessionInfoWindow::Construct(const FArguments& InArgs, const TSharedRef<SD
 			.Text(LOCTEXT("UnrealInsightsVersion", UNREAL_INSIGHTS_VERSION_STRING_EX))
 			.ColorAndOpacity(FLinearColor(0.15f, 0.15f, 0.15f, 1.0f))
 		]
+
 		// Overlay slot for the main window area
 		+ SOverlay::Slot()
 		.HAlign(HAlign_Fill)
 		.VAlign(VAlign_Fill)
 		[
-			SNew(SVerticalBox)
-
-			+ SVerticalBox::Slot()
-			.AutoHeight()
-			[
-				MenuBarBuilder.MakeWidget()
-			]
-
-			+ SVerticalBox::Slot()
-			.FillHeight(1.0f)
-			[
-				TabManager->RestoreFrom(Layout, ConstructUnderWindow).ToSharedRef()
-			]
+			TabManager->RestoreFrom(Layout, ConstructUnderWindow).ToSharedRef()
 		]
 	];
 }
@@ -309,15 +312,14 @@ TSharedRef<SDockTab> SSessionInfoWindow::SpawnTab_SessionInfo(const FSpawnTabArg
 	return DockTab;
 }
 
+END_SLATE_FUNCTION_BUILD_OPTIMIZATION
+
 ////////////////////////////////////////////////////////////////////////////////////////////////////
 
 void SSessionInfoWindow::OnSessionInfoTabClosed(TSharedRef<SDockTab> TabBeingClosed)
 {
-
 }
 
-END_SLATE_FUNCTION_BUILD_OPTIMIZATION
-
 ////////////////////////////////////////////////////////////////////////////////////////////////////
 
 void SSessionInfoWindow::FillMenu(FMenuBuilder& MenuBuilder, const TSharedPtr<FTabManager> TabManager)
diff --git a/Engine/Source/Developer/TraceInsights/Private/Insights/Widgets/STimingProfilerToolbar.cpp b/Engine/Source/Developer/TraceInsights/Private/Insights/Widgets/STimingProfilerToolbar.cpp
--- a/Engine/Source/Developer/TraceInsights/Private/Insights/Widgets/STimingProfilerToolbar.cpp
+++ b/Engine/Source/Developer/TraceInsights/Private/Insights/Widgets/STimingProfilerToolbar.cpp
@@ -2,20 +2,14 @@
 
 #include "STimingProfilerToolbar.h"
 
-#include "EditorStyleSet.h"
-#include "Framework/Application/SlateApplication.h"
-#include "Framework/Docking/TabManager.h"
-#include "Framework/MultiBox/MultiBoxDefs.h"
 #include "Framework/MultiBox/MultiBoxBuilder.h"
-#include "Widgets/Docking/SDockTab.h"
-#include "Widgets/Layout/SBorder.h"
+#include "Widgets/SBoxPanel.h"
 
 // Insights
 #include "Insights/InsightsCommands.h"
 #include "Insights/InsightsManager.h"
 #include "Insights/InsightsStyle.h"
 #include "Insights/TimingProfilerCommands.h"
-#include "Insights/TimingProfilerManager.h"
 
 ////////////////////////////////////////////////////////////////////////////////////////////////////
 
@@ -88,24 +82,17 @@ void STimingProfilerToolbar::Construct(const FArguments& InArgs)
 	FSlimHorizontalToolBarBuilder RightSideToolbarBuilder(CommandList.ToSharedRef(), FMultiBoxCustomization::None);
 	Local::FillRightSideToolbar(RightSideToolbarBuilder);
 
-	// Create the tool bar!
 	ChildSlot
 	[
 		SNew(SHorizontalBox)
 
 		+ SHorizontalBox::Slot()
-		.HAlign(HAlign_Left)
+		.HAlign(HAlign_Fill)
 		.VAlign(VAlign_Center)
 		.FillWidth(1.0)
 		.Padding(0.0f)
 		[
-			SNew(SBorder)
-			.Padding(0)
-			.BorderImage(FEditorStyle::GetBrush("NoBorder"))
-			.IsEnabled(FSlateApplication::Get().GetNormalExecutionAttribute())
-			[
-				ToolbarBuilder.MakeWidget()
-			]
+			ToolbarBuilder.MakeWidget()
 		]
 
 		+ SHorizontalBox::Slot()
@@ -114,13 +101,7 @@ void STimingProfilerToolbar::Construct(const FArguments& InArgs)
 		.AutoWidth()
 		.Padding(0.0f)
 		[
-			SNew(SBorder)
-			.Padding(0)
-			.BorderImage(FEditorStyle::GetBrush("NoBorder"))
-			.IsEnabled(FSlateApplication::Get().GetNormalExecutionAttribute())
-			[
-				RightSideToolbarBuilder.MakeWidget()
-			]
+			RightSideToolbarBuilder.MakeWidget()
 		]
 	];
 }
diff --git a/Engine/Source/Developer/TraceInsights/Private/Insights/Widgets/STimingProfilerToolbar.h b/Engine/Source/Developer/TraceInsights/Private/Insights/Widgets/STimingProfilerToolbar.h
--- a/Engine/Source/Developer/TraceInsights/Private/Insights/Widgets/STimingProfilerToolbar.h
+++ b/Engine/Source/Developer/TraceInsights/Private/Insights/Widgets/STimingProfilerToolbar.h
@@ -2,8 +2,6 @@
 
 #pragma once
 
-#include "CoreMinimal.h"
-#include "Styling/SlateTypes.h"
 #include "Widgets/DeclarativeSyntaxSupport.h"
 #include "Widgets/SCompoundWidget.h"
 
diff --git a/Engine/Source/Developer/TraceInsights/Private/Insights/Widgets/STimingProfilerWindow.cpp b/Engine/Source/Developer/TraceInsights/Private/Insights/Widgets/STimingProfilerWindow.cpp
--- a/Engine/Source/Developer/TraceInsights/Private/Insights/Widgets/STimingProfilerWindow.cpp
+++ b/Engine/Source/Developer/TraceInsights/Private/Insights/Widgets/STimingProfilerWindow.cpp
@@ -2,18 +2,18 @@
 
 #include "STimingProfilerWindow.h"
 
+#include "Framework/Commands/UICommandList.h"
 #include "Framework/Docking/LayoutExtender.h"
+#include "Framework/Docking/LayoutService.h"
+#include "Framework/Docking/TabManager.h"
 #include "Framework/Docking/WorkspaceItem.h"
 #include "Framework/MultiBox/MultiBoxBuilder.h"
-#include "Framework/Docking/LayoutService.h"
 #include "SlateOptMacros.h"
 #include "Styling/CoreStyle.h"
 #include "Widgets/Docking/SDockTab.h"
 #include "Widgets/Images/SImage.h"
 #include "Widgets/Layout/SBorder.h"
 #include "Widgets/Layout/SBox.h"
-#include "Widgets/Layout/SSpacer.h"
-#include "Widgets/Notifications/SNotificationList.h"
 #include "Widgets/SBoxPanel.h"
 #include "Widgets/Text/STextBlock.h"
 #include "Modules/ModuleManager.h"
@@ -39,7 +39,6 @@
 #include "Insights/Widgets/STimerTreeView.h"
 #include "Insights/Widgets/STimingProfilerToolbar.h"
 #include "Insights/Widgets/STimingView.h"
-#include "Framework/Docking/TabManager.h"
 
 ////////////////////////////////////////////////////////////////////////////////////////////////////
 
@@ -47,7 +46,7 @@
 
 ////////////////////////////////////////////////////////////////////////////////////////////////////
 
-const FName FTimingProfilerTabs::ToolbarID(TEXT("Toolbar"));
+const FName FTimingProfilerTabs::ToolbarID(TEXT("Toolbar")); // DEPRECATED
 const FName FTimingProfilerTabs::FramesTrackID(TEXT("Frames"));
 const FName FTimingProfilerTabs::TimingViewID(TEXT("TimingView"));
 const FName FTimingProfilerTabs::TimersID(TEXT("Timers"));
@@ -109,8 +108,6 @@ STimingProfilerWindow::~STimingProfilerWindow()
 		check(FrameTrack == nullptr);
 	}
 
-	HideTab(FTimingProfilerTabs::ToolbarID);
-
 #if WITH_EDITOR
 	if (DurationActive > 0.0f && FEngineAnalytics::IsAvailable())
 	{
@@ -162,29 +159,6 @@ void STimingProfilerWindow::Reset()
 ////////////////////////////////////////////////////////////////////////////////////////////////////
 
 BEGIN_SLATE_FUNCTION_BUILD_OPTIMIZATION
-
-TSharedRef<SDockTab> STimingProfilerWindow::SpawnTab_Toolbar(const FSpawnTabArgs& Args)
-{
-	const TSharedRef<SDockTab> DockTab = SNew(SDockTab)
-		.ShouldAutosize(true)
-		.TabRole(ETabRole::PanelTab)
-		[
-			SNew(STimingProfilerToolbar)
-		];
-
-	DockTab->SetOnTabClosed(SDockTab::FOnTabClosedCallback::CreateRaw(this, &STimingProfilerWindow::OnToolbarTabClosed));
-
-	return DockTab;
-}
-
-////////////////////////////////////////////////////////////////////////////////////////////////////
-
-void STimingProfilerWindow::OnToolbarTabClosed(TSharedRef<SDockTab> TabBeingClosed)
-{
-}
-
-////////////////////////////////////////////////////////////////////////////////////////////////////
-
 TSharedRef<SDockTab> STimingProfilerWindow::SpawnTab_FramesTrack(const FSpawnTabArgs& Args)
 {
 	FTimingProfilerManager::Get()->SetFramesTrackVisible(true);
@@ -200,6 +174,7 @@ TSharedRef<SDockTab> STimingProfilerWindow::SpawnTab_FramesTrack(const FSpawnTab
 
 	return DockTab;
 }
+END_SLATE_FUNCTION_BUILD_OPTIMIZATION
 
 ////////////////////////////////////////////////////////////////////////////////////////////////////
 
@@ -211,6 +186,7 @@ void STimingProfilerWindow::OnFramesTrackTabClosed(TSharedRef<SDockTab> TabBeing
 
 ////////////////////////////////////////////////////////////////////////////////////////////////////
 
+BEGIN_SLATE_FUNCTION_BUILD_OPTIMIZATION
 TSharedRef<SDockTab> STimingProfilerWindow::SpawnTab_TimingView(const FSpawnTabArgs& Args)
 {
 	FTimingProfilerManager::Get()->SetTimingViewVisible(true);
@@ -228,6 +204,7 @@ TSharedRef<SDockTab> STimingProfilerWindow::SpawnTab_TimingView(const FSpawnTabA
 
 	return DockTab;
 }
+END_SLATE_FUNCTION_BUILD_OPTIMIZATION
 
 ////////////////////////////////////////////////////////////////////////////////////////////////////
 
@@ -239,6 +216,7 @@ void STimingProfilerWindow::OnTimingViewTabClosed(TSharedRef<SDockTab> TabBeingC
 
 ////////////////////////////////////////////////////////////////////////////////////////////////////
 
+BEGIN_SLATE_FUNCTION_BUILD_OPTIMIZATION
 TSharedRef<SDockTab> STimingProfilerWindow::SpawnTab_Timers(const FSpawnTabArgs& Args)
 {
 	FTimingProfilerManager::Get()->SetTimersViewVisible(true);
@@ -254,6 +232,7 @@ TSharedRef<SDockTab> STimingProfilerWindow::SpawnTab_Timers(const FSpawnTabArgs&
 
 	return DockTab;
 }
+END_SLATE_FUNCTION_BUILD_OPTIMIZATION
 
 ////////////////////////////////////////////////////////////////////////////////////////////////////
 
@@ -265,6 +244,7 @@ void STimingProfilerWindow::OnTimersTabClosed(TSharedRef<SDockTab> TabBeingClose
 
 ////////////////////////////////////////////////////////////////////////////////////////////////////
 
+BEGIN_SLATE_FUNCTION_BUILD_OPTIMIZATION
 TSharedRef<SDockTab> STimingProfilerWindow::SpawnTab_Callers(const FSpawnTabArgs& Args)
 {
 	FTimingProfilerManager::Get()->SetCallersTreeViewVisible(true);
@@ -280,6 +260,7 @@ TSharedRef<SDockTab> STimingProfilerWindow::SpawnTab_Callers(const FSpawnTabArgs
 
 	return DockTab;
 }
+END_SLATE_FUNCTION_BUILD_OPTIMIZATION
 
 ////////////////////////////////////////////////////////////////////////////////////////////////////
 
@@ -291,6 +272,7 @@ void STimingProfilerWindow::OnCallersTabClosed(TSharedRef<SDockTab> TabBeingClos
 
 ////////////////////////////////////////////////////////////////////////////////////////////////////
 
+BEGIN_SLATE_FUNCTION_BUILD_OPTIMIZATION
 TSharedRef<SDockTab> STimingProfilerWindow::SpawnTab_Callees(const FSpawnTabArgs& Args)
 {
 	FTimingProfilerManager::Get()->SetCalleesTreeViewVisible(true);
@@ -306,6 +288,7 @@ TSharedRef<SDockTab> STimingProfilerWindow::SpawnTab_Callees(const FSpawnTabArgs
 
 	return DockTab;
 }
+END_SLATE_FUNCTION_BUILD_OPTIMIZATION
 
 ////////////////////////////////////////////////////////////////////////////////////////////////////
 
@@ -317,6 +300,7 @@ void STimingProfilerWindow::OnCalleesTabClosed(TSharedRef<SDockTab> TabBeingClos
 
 ////////////////////////////////////////////////////////////////////////////////////////////////////
 
+BEGIN_SLATE_FUNCTION_BUILD_OPTIMIZATION
 TSharedRef<SDockTab> STimingProfilerWindow::SpawnTab_StatsCounters(const FSpawnTabArgs& Args)
 {
 	FTimingProfilerManager::Get()->SetStatsCountersViewVisible(true);
@@ -332,6 +316,7 @@ TSharedRef<SDockTab> STimingProfilerWindow::SpawnTab_StatsCounters(const FSpawnT
 
 	return DockTab;
 }
+END_SLATE_FUNCTION_BUILD_OPTIMIZATION
 
 ////////////////////////////////////////////////////////////////////////////////////////////////////
 
@@ -343,6 +328,7 @@ void STimingProfilerWindow::OnStatsCountersTabClosed(TSharedRef<SDockTab> TabBei
 
 ////////////////////////////////////////////////////////////////////////////////////////////////////
 
+BEGIN_SLATE_FUNCTION_BUILD_OPTIMIZATION
 TSharedRef<SDockTab> STimingProfilerWindow::SpawnTab_LogView(const FSpawnTabArgs& Args)
 {
 	FTimingProfilerManager::Get()->SetLogViewVisible(true);
@@ -358,6 +344,7 @@ TSharedRef<SDockTab> STimingProfilerWindow::SpawnTab_LogView(const FSpawnTabArgs
 
 	return DockTab;
 }
+END_SLATE_FUNCTION_BUILD_OPTIMIZATION
 
 ////////////////////////////////////////////////////////////////////////////////////////////////////
 
@@ -369,10 +356,13 @@ void STimingProfilerWindow::OnLogViewTabClosed(TSharedRef<SDockTab> TabBeingClos
 
 ////////////////////////////////////////////////////////////////////////////////////////////////////
 
+BEGIN_SLATE_FUNCTION_BUILD_OPTIMIZATION
 void STimingProfilerWindow::Construct(const FArguments& InArgs, const TSharedRef<SDockTab>& ConstructUnderMajorTab, const TSharedPtr<SWindow>& ConstructUnderWindow)
 {
 	// Create & initialize tab manager.
 	TabManager = FGlobalTabmanager::Get()->NewTabManager(ConstructUnderMajorTab);
+	//TabManager->SetAllowWindowMenuBar(true);
+
 	const auto& PersistLayout = [](const TSharedRef<FTabManager::FLayout>& LayoutToSave)
 	{
 		FLayoutSaveRestore::SaveToConfig(FTraceInsightsModule::GetUnrealInsightsLayoutIni(), LayoutToSave);
@@ -383,11 +373,6 @@ void STimingProfilerWindow::Construct(const FArguments& InArgs, const TSharedRef
 
 	Extension = MakeShared<FInsightsMajorTabExtender>(TabManager, AppMenuGroup);
 
-	TabManager->RegisterTabSpawner(FTimingProfilerTabs::ToolbarID, FOnSpawnTab::CreateRaw(this, &STimingProfilerWindow::SpawnTab_Toolbar))
-		.SetDisplayName(LOCTEXT("DeviceToolbarTabTitle", "Toolbar"))
-		.SetIcon(FSlateIcon(FInsightsStyle::GetStyleSetName(), "Toolbar.Icon.Small"))
-		.SetGroup(AppMenuGroup);
-
 	TabManager->RegisterTabSpawner(FTimingProfilerTabs::FramesTrackID, FOnSpawnTab::CreateRaw(this, &STimingProfilerWindow::SpawnTab_FramesTrack))
 		.SetDisplayName(LOCTEXT("FramesTrackTabTitle", "Frames"))
 		.SetIcon(FSlateIcon(FInsightsStyle::GetStyleSetName(), "FramesTrack.Icon.Small"))
@@ -462,88 +447,75 @@ void STimingProfilerWindow::Construct(const FArguments& InArgs, const TSharedRef
 		else
 		{
 			// Create tab layout.
-			return FTabManager::NewLayout("InsightsTimingProfilerLayout_v1.1")
+			return FTabManager::NewLayout("InsightsTimingProfilerLayout_v1.2")
 				->AddArea
 				(
 					FTabManager::NewPrimaryArea()
-					->SetOrientation(Orient_Vertical)
+					->SetOrientation(Orient_Horizontal)
 					->Split
 					(
-						FTabManager::NewStack()
-						->AddTab(FTimingProfilerTabs::ToolbarID, ETabState::OpenedTab)
-						->SetHideTabWell(true)
+						FTabManager::NewSplitter()
+						->SetSizeCoefficient(0.65f)
+						->SetOrientation(Orient_Vertical)
+						->Split
+						(
+							FTabManager::NewStack()
+							->SetSizeCoefficient(0.1f)
+							->SetHideTabWell(true)
+							->AddTab(FTimingProfilerTabs::FramesTrackID, ETabState::OpenedTab)
+						)
+						->Split
+						(
+							FTabManager::NewStack()
+							->SetSizeCoefficient(0.5f)
+							->SetHideTabWell(true)
+							->AddTab(FTimingProfilerTabs::TimingViewID, ETabState::OpenedTab)
+						)
+						->Split
+						(
+							FTabManager::NewStack()
+							->SetSizeCoefficient(0.2f)
+							->SetHideTabWell(true)
+							->AddTab(FTimingProfilerTabs::LogViewID, ETabState::OpenedTab)
+						)
 					)
 					->Split
 					(
 						FTabManager::NewSplitter()
-						->SetOrientation(Orient_Horizontal)
-						->SetSizeCoefficient(1.0f)
+						->SetSizeCoefficient(0.35f)
+						->SetOrientation(Orient_Vertical)
 						->Split
 						(
-							FTabManager::NewSplitter()
-							->SetOrientation(Orient_Vertical)
-							->SetSizeCoefficient(0.65f)
-							->Split
-							(
-								FTabManager::NewStack()
-								->SetSizeCoefficient(0.1f)
-								->SetHideTabWell(true)
-								->AddTab(FTimingProfilerTabs::FramesTrackID, ETabState::OpenedTab)
-							)
-							->Split
-							(
-								FTabManager::NewStack()
-								->SetSizeCoefficient(0.5f)
-								->SetHideTabWell(true)
-								->AddTab(FTimingProfilerTabs::TimingViewID, ETabState::OpenedTab)
-							)
-							->Split
-							(
-								FTabManager::NewStack()
-								->SetSizeCoefficient(0.2f)
-								->SetHideTabWell(true)
-								->AddTab(FTimingProfilerTabs::LogViewID, ETabState::OpenedTab)
-							)
+							FTabManager::NewStack()
+							->SetSizeCoefficient(0.67f)
+							->AddTab(FTimingProfilerTabs::TimersID, ETabState::OpenedTab)
+							->AddTab(FTimingProfilerTabs::StatsCountersID, ETabState::OpenedTab)
+							->SetForegroundTab(FTimingProfilerTabs::TimersID)
 						)
 						->Split
 						(
-							FTabManager::NewSplitter()
-							->SetOrientation(Orient_Vertical)
-							->SetSizeCoefficient(0.35f)
-							->Split
-							(
-								FTabManager::NewStack()
-								->SetSizeCoefficient(0.67f)
-								->AddTab(FTimingProfilerTabs::TimersID, ETabState::OpenedTab)
-								->AddTab(FTimingProfilerTabs::StatsCountersID, ETabState::OpenedTab)
-								->SetForegroundTab(FTimingProfilerTabs::TimersID)
-							)
-							->Split
-							(
-								FTabManager::NewStack()
-								->SetSizeCoefficient(0.165f)
-								->SetHideTabWell(true)
-								->AddTab(FTimingProfilerTabs::CallersID, ETabState::OpenedTab)
-							)
-							->Split
-							(
-								FTabManager::NewStack()
-								->SetSizeCoefficient(0.165f)
-								->SetHideTabWell(true)
-								->AddTab(FTimingProfilerTabs::CalleesID, ETabState::OpenedTab)
-							)
+							FTabManager::NewStack()
+							->SetSizeCoefficient(0.165f)
+							->SetHideTabWell(true)
+							->AddTab(FTimingProfilerTabs::CallersID, ETabState::OpenedTab)
+						)
+						->Split
+						(
+							FTabManager::NewStack()
+							->SetSizeCoefficient(0.165f)
+							->SetHideTabWell(true)
+							->AddTab(FTimingProfilerTabs::CalleesID, ETabState::OpenedTab)
 						)
 					)
 				);
 		}
 	}();
-	
+
 	Layout->ProcessExtensions(Extension->GetLayoutExtender());
 	Layout = FLayoutSaveRestore::LoadFromConfig(FTraceInsightsModule::GetUnrealInsightsLayoutIni(), Layout);
 
 	// Create & initialize main menu.
 	FMenuBarBuilder MenuBarBuilder = FMenuBarBuilder(TSharedPtr<FUICommandList>(), Extension->GetMenuExtender());
-
 	MenuBarBuilder.AddPullDownMenu(
 		LOCTEXT("MenuLabel", "Menu"),
 		FText::GetEmpty(),
@@ -551,67 +523,83 @@ void STimingProfilerWindow::Construct(const FArguments& InArgs, const TSharedRef
 		FName(TEXT("Menu"))
 	);
 
+#if !WITH_EDITOR
 	TSharedRef<SWidget> MenuWidget = MenuBarBuilder.MakeWidget();
+	MenuWidget->SetClipping(EWidgetClipping::ClipToBoundsWithoutIntersecting);
+#endif
 
 	ChildSlot
+	[
+		SNew(SOverlay)
+
+#if !WITH_EDITOR
+		// Menu
+		+ SOverlay::Slot()
+		.HAlign(HAlign_Left)
+		.VAlign(VAlign_Top)
+		.Padding(34.0f, -60.0f, 0.0f, 0.0f)
 		[
-			SNew(SOverlay)
-
-			// Version
-			+ SOverlay::Slot()
-				.HAlign(HAlign_Right)
-				.VAlign(VAlign_Top)
-				.Padding(0.0f, -16.0f, 0.0f, 0.0f)
-				[
-					SNew(STextBlock)
-						.Clipping(EWidgetClipping::ClipToBoundsWithoutIntersecting)
-						.Text(LOCTEXT("UnrealInsightsVersion", UNREAL_INSIGHTS_VERSION_STRING_EX))
-						.ColorAndOpacity(FLinearColor(0.15f, 0.15f, 0.15f, 1.0f))
-				]
-
-			// Overlay slot for the main window area
-			+ SOverlay::Slot()
-				.HAlign(HAlign_Fill)
-				.VAlign(VAlign_Fill)
-				[
-					SNew(SVerticalBox)
-
-					+ SVerticalBox::Slot()
-						.AutoHeight()
-						[
-							MenuWidget
-						]
-
-					+ SVerticalBox::Slot()
-						.FillHeight(1.0f)
-						[
-							TabManager->RestoreFrom(Layout, ConstructUnderWindow).ToSharedRef()
-						]
-				]
-
-			// Session hint overlay
-			+ SOverlay::Slot()
-				.HAlign(HAlign_Center)
-				.VAlign(VAlign_Center)
-				[
-					SNew(SBorder)
-						.Visibility(this, &STimingProfilerWindow::IsSessionOverlayVisible)
-						.BorderImage(FCoreStyle::Get().GetBrush("PopupText.Background"))
-						.Padding(8.0f)
-						[
-							SNew(STextBlock)
-								.Text(LOCTEXT("SelectTraceOverlayText", "Please select a trace."))
-						]
-				]
-		];
-
+			MenuWidget
+		]
+#endif
+
+		// Version
+		+ SOverlay::Slot()
+		.HAlign(HAlign_Right)
+		.VAlign(VAlign_Top)
+		.Padding(0.0f, -16.0f, 0.0f, 0.0f)
+		[
+			SNew(STextBlock)
+			.Clipping(EWidgetClipping::ClipToBoundsWithoutIntersecting)
+			.Text(LOCTEXT("UnrealInsightsVersion", UNREAL_INSIGHTS_VERSION_STRING_EX))
+			.ColorAndOpacity(FLinearColor(0.15f, 0.15f, 0.15f, 1.0f))
+		]
+
+		// Overlay slot for the main window area
+		+ SOverlay::Slot()
+		.HAlign(HAlign_Fill)
+		.VAlign(VAlign_Fill)
+		[
+			SNew(SVerticalBox)
+
+			+ SVerticalBox::Slot()
+			.AutoHeight()
+			.Padding(FMargin(0.0f, 0.0f, 0.0f, 5.0f))
+			[
+				SNew(STimingProfilerToolbar)
+			]
+
+			+ SVerticalBox::Slot()
+			.FillHeight(1.0f)
+			[
+				TabManager->RestoreFrom(Layout, ConstructUnderWindow).ToSharedRef()
+			]
+		]
+
+		// Session hint overlay
+		+ SOverlay::Slot()
+		.HAlign(HAlign_Center)
+		.VAlign(VAlign_Center)
+		[
+			SNew(SBorder)
+			.Visibility(this, &STimingProfilerWindow::IsSessionOverlayVisible)
+			.BorderImage(FCoreStyle::Get().GetBrush("PopupText.Background"))
+			.Padding(8.0f)
+			[
+				SNew(STextBlock)
+				.Text(LOCTEXT("SelectTraceOverlayText", "Please select a trace."))
+			]
+		]
+	];
+
+#if !WITH_EDITOR
 	// Tell tab-manager about the global menu bar.
 	TabManager->SetMenuMultiBox(MenuBarBuilder.GetMultiBox(), MenuWidget);
+#endif
 
 	// Tell clients about creation
 	TraceInsightsModule.OnMajorTabCreated().Broadcast(FInsightsManagerTabs::TimingProfilerTabId, TabManager.ToSharedRef());
 }
-
 END_SLATE_FUNCTION_BUILD_OPTIMIZATION
 
 ////////////////////////////////////////////////////////////////////////////////////////////////////
diff --git a/Engine/Source/Developer/TraceInsights/Private/Insights/Widgets/STimingProfilerWindow.h b/Engine/Source/Developer/TraceInsights/Private/Insights/Widgets/STimingProfilerWindow.h
--- a/Engine/Source/Developer/TraceInsights/Private/Insights/Widgets/STimingProfilerWindow.h
+++ b/Engine/Source/Developer/TraceInsights/Private/Insights/Widgets/STimingProfilerWindow.h
@@ -63,9 +63,6 @@ public:
 	TSharedPtr<SLogView> GetLogView() const { return LogView; }
 
 private:
-	TSharedRef<SDockTab> SpawnTab_Toolbar(const FSpawnTabArgs& Args);
-	void OnToolbarTabClosed(TSharedRef<SDockTab> TabBeingClosed);
-
 	TSharedRef<SDockTab> SpawnTab_FramesTrack(const FSpawnTabArgs& Args);
 	void OnFramesTrackTabClosed(TSharedRef<SDockTab> TabBeingClosed);
 
diff --git a/Engine/Source/Developer/TraceInsights/Public/Insights/IUnrealInsightsModule.h b/Engine/Source/Developer/TraceInsights/Public/Insights/IUnrealInsightsModule.h
--- a/Engine/Source/Developer/TraceInsights/Public/Insights/IUnrealInsightsModule.h
+++ b/Engine/Source/Developer/TraceInsights/Public/Insights/IUnrealInsightsModule.h
@@ -44,7 +44,7 @@ struct TRACEINSIGHTS_API FInsightsManagerTabs
 struct TRACEINSIGHTS_API FTimingProfilerTabs
 {
 	// Tab identifiers
-	static const FName ToolbarID;
+	static const FName ToolbarID; // DEPRECATED
 	static const FName FramesTrackID;
 	static const FName TimingViewID;
 	static const FName TimersID;
