[Insights] Fix relation Y position when the depth is 0.

diff --git a/Engine/Source/Developer/TraceInsights/Private/Insights/TaskGraphProfiler/TaskGraphProfilerManager.cpp b/Engine/Source/Developer/TraceInsights/Private/Insights/TaskGraphProfiler/TaskGraphProfilerManager.cpp
--- a/Engine/Source/Developer/TraceInsights/Private/Insights/TaskGraphProfiler/TaskGraphProfilerManager.cpp
+++ b/Engine/Source/Developer/TraceInsights/Private/Insights/TaskGraphProfiler/TaskGraphProfilerManager.cpp
@@ -477,43 +477,24 @@ void FTaskGraphProfilerManager::AddRelation(const FThreadTrackEvent* InSelectedE
 	TUniquePtr<ITimingEventRelation> Relation = MakeUnique<FTaskGraphRelation>(SourceTimestamp, SourceThreadId, TargetTimestamp, TargetThreadId, Type);
 	FTaskGraphRelation* TaskRelationPtr = StaticCast<FTaskGraphRelation*>(Relation.Get());
 
-	// If we have a valid event, we can try getting the tracks and depths using this faster approach.
-	if (InSelectedEvent)
-	{
-		TSharedPtr<const FThreadTimingTrack> EventTrack = StaticCastSharedRef<const FThreadTimingTrack>(InSelectedEvent->GetTrack());
-		uint32 ThreadId = EventTrack->GetThreadId();
-
-		if (TaskRelationPtr->GetSourceThreadId() == ThreadId)
-		{
-			TaskRelationPtr->SetSourceTrack(EventTrack);
-			TaskRelationPtr->SetSourceDepth(SourceDepth > 0 ? SourceDepth : EventTrack->GetDepthAt(TaskRelationPtr->GetSourceTime()) - 1);
-		}
-
-		if (TaskRelationPtr->GetTargetThreadId() == ThreadId)
-		{
-			TaskRelationPtr->SetTargetTrack(EventTrack);
-			TaskRelationPtr->SetTargetDepth(TargetDepth > 0 ? TargetDepth : EventTrack->GetDepthAt(TaskRelationPtr->GetTargetTime()) - 1);
-		}
-	}
-
 	if(!TaskRelationPtr->GetSourceTrack().IsValid())
 	{
-		TSharedPtr<FCpuTimingTrack> Track = ThreadSharedState->GetCpuTrack(TaskRelationPtr->GetSourceThreadId());
+		TSharedPtr<const FCpuTimingTrack> Track = ThreadSharedState->GetCpuTrack(TaskRelationPtr->GetSourceThreadId());
 		if (Track.IsValid())
 		{
 			TaskRelationPtr->SetSourceTrack(Track);
-			TaskRelationPtr->SetSourceDepth(SourceDepth > 0 ? SourceDepth : Track->GetDepthAt(TaskRelationPtr->GetSourceTime()) - 1);
+			TaskRelationPtr->SetSourceDepth(GetRelationDisplayDepth(Track, TaskRelationPtr->GetSourceTime(), SourceDepth));
 		}
 	}
 
 	if (!TaskRelationPtr->GetTargetTrack().IsValid())
 
 	{
-		TSharedPtr<FCpuTimingTrack> Track = ThreadSharedState->GetCpuTrack(TaskRelationPtr->GetTargetThreadId());
+		TSharedPtr<const FCpuTimingTrack> Track = ThreadSharedState->GetCpuTrack(TaskRelationPtr->GetTargetThreadId());
 		if (Track.IsValid())
 		{
 			TaskRelationPtr->SetTargetTrack(Track);
-			TaskRelationPtr->SetTargetDepth(TargetDepth > 0 ? TargetDepth : Track->GetDepthAt(TaskRelationPtr->GetTargetTime()) - 1);
+			TaskRelationPtr->SetTargetDepth(GetRelationDisplayDepth(Track, TaskRelationPtr->GetTargetTime(), TargetDepth));
 		}
 	}
 
@@ -525,6 +506,21 @@ void FTaskGraphProfilerManager::AddRelation(const FThreadTrackEvent* InSelectedE
 
 ////////////////////////////////////////////////////////////////////////////////////////////////////
 
+int32 FTaskGraphProfilerManager::GetRelationDisplayDepth(TSharedPtr<const FThreadTimingTrack> Track, double Time, int32 KnownDepth)
+{
+	if (KnownDepth >= 0)
+	{
+		return KnownDepth;
+	}
+
+	KnownDepth = Track->GetDepthAt(Time) - 1;
+
+	// We return the depth 0 based, so the depth will be 0 when the event is on the first row or when there is no event on the track at the time.
+	return FMath::Max(KnownDepth, 0);
+}
+
+////////////////////////////////////////////////////////////////////////////////////////////////////
+
 void FTaskGraphProfilerManager::ClearTaskRelations()
 {
 	TSharedPtr<class STimingProfilerWindow> TimingWindow = FTimingProfilerManager::Get()->GetProfilerWindow();
diff --git a/Engine/Source/Developer/TraceInsights/Private/Insights/TaskGraphProfiler/TaskGraphProfilerManager.h b/Engine/Source/Developer/TraceInsights/Private/Insights/TaskGraphProfiler/TaskGraphProfilerManager.h
--- a/Engine/Source/Developer/TraceInsights/Private/Insights/TaskGraphProfiler/TaskGraphProfilerManager.h
+++ b/Engine/Source/Developer/TraceInsights/Private/Insights/TaskGraphProfiler/TaskGraphProfilerManager.h
@@ -16,6 +16,7 @@ namespace TraceServices
 	class ITasksProvider;
 }
 
+class FThreadTimingTrack;
 class FThreadTrackEvent;
 
 namespace Insights
@@ -125,6 +126,7 @@ private:
 	void GetSingleTaskRelations(const TraceServices::FTaskInfo* Task, const TraceServices::ITasksProvider* TasksProvider, const FThreadTrackEvent* InSelectedEvent);
 
 	void InitializeColorCode();
+	int32 GetRelationDisplayDepth(TSharedPtr<const FThreadTimingTrack> Track, double Time, int32 KnownDepth);
 
 private:
 	bool bIsInitialized;
