[Insights]
  - Browser: Updated layout of tooltip for a trace session in the Trace Store page.
  - Browser/Connection: Moved the "Connect" functionality to a separate Connection tab.
  - Browser/Connection: Added "Trace recorder IP address" editable text box (prefilled with local host ip).
  - Updated icons for Automation Testing and debug tools.

diff --git a/Engine/Source/Developer/TraceInsights/Private/Insights/Common/InsightsMenuBuilder.cpp b/Engine/Source/Developer/TraceInsights/Private/Insights/Common/InsightsMenuBuilder.cpp
--- a/Engine/Source/Developer/TraceInsights/Private/Insights/Common/InsightsMenuBuilder.cpp
+++ b/Engine/Source/Developer/TraceInsights/Private/Insights/Common/InsightsMenuBuilder.cpp
@@ -76,7 +76,7 @@ void FInsightsMenuBuilder::PopulateMenu(FMenuBuilder& MenuBuilder)
 		MenuBuilder.AddMenuEntry(
 			LOCTEXT("OpenStarshipSuite", "Starship Test Suite"),
 			LOCTEXT("OpenStarshipSuiteDesc", "Opens the Starship UX test suite."),
-			FSlateIcon(FInsightsStyle::GetStyleSetName(), "Icon.Bug"),
+			FSlateIcon(FInsightsStyle::GetStyleSetName(), "Icon.Test"),
 			OpenStarshipSuiteAction,
 			NAME_None,
 			EUserInterfaceActionType::Button
diff --git a/Engine/Source/Developer/TraceInsights/Private/Insights/InsightsManager.cpp b/Engine/Source/Developer/TraceInsights/Private/Insights/InsightsManager.cpp
--- a/Engine/Source/Developer/TraceInsights/Private/Insights/InsightsManager.cpp
+++ b/Engine/Source/Developer/TraceInsights/Private/Insights/InsightsManager.cpp
@@ -32,9 +32,25 @@
 
 ////////////////////////////////////////////////////////////////////////////////////////////////////
 
+/**
+ * Special tab type, that cannot be dragged/undocked from the tab bar
+ */
+class SLockedTab : public SDockTab
+{
+	virtual FReply OnDragDetected(const FGeometry& MyGeometry, const FPointerEvent& MouseEvent) override
+	{
+		return FReply::Handled();
+	}
+};
+
+////////////////////////////////////////////////////////////////////////////////////////////////////
+
 #define LOCTEXT_NAMESPACE "InsightsManager"
 
-const FName FInsightsManagerTabs::StartPageTabId(TEXT("StartPage"));
+const FName FInsightsManagerTabs::StartPageTabId(TEXT("TraceStore")); // DEPRECATED
+const FName FInsightsManagerTabs::TraceStoreTabId(TEXT("TraceStore"));
+const FName FInsightsManagerTabs::ConnectionTabId(TEXT("Connection"));
+const FName FInsightsManagerTabs::LauncherTabId(TEXT("Launcher"));
 const FName FInsightsManagerTabs::SessionInfoTabId(TEXT("SessionInfo"));
 const FName FInsightsManagerTabs::TimingProfilerTabId(TEXT("TimingProfiler"));
 const FName FInsightsManagerTabs::LoadingProfilerTabId(TEXT("LoadingProfiler"));
@@ -214,17 +230,45 @@ void FInsightsManager::BindCommands()
 
 void FInsightsManager::RegisterMajorTabs(IUnrealInsightsModule& InsightsModule)
 {
-	const FInsightsMajorTabConfig& StartPageConfig = InsightsModule.FindMajorTabConfig(FInsightsManagerTabs::StartPageTabId);
-	if (StartPageConfig.bIsAvailable)
+	const FInsightsMajorTabConfig& TraceStoreConfig = InsightsModule.FindMajorTabConfig(FInsightsManagerTabs::TraceStoreTabId);
+	if (TraceStoreConfig.bIsAvailable)
+	{
+		// Register tab spawner for the Trace Store tab.
+		FTabSpawnerEntry& TabSpawnerEntry = FGlobalTabmanager::Get()->RegisterNomadTabSpawner(FInsightsManagerTabs::TraceStoreTabId,
+			FOnSpawnTab::CreateRaw(this, &FInsightsManager::SpawnTraceStoreTab))
+			.SetDisplayName(TraceStoreConfig.TabLabel.IsSet() ? TraceStoreConfig.TabLabel.GetValue() : LOCTEXT("TraceStoreTabTitle", "Trace Store"))
+			.SetTooltipText(TraceStoreConfig.TabTooltip.IsSet() ? TraceStoreConfig.TabTooltip.GetValue() : LOCTEXT("TraceStoreTooltipText", "Open the Trace Store Browser."))
+			.SetIcon(TraceStoreConfig.TabIcon.IsSet() ? TraceStoreConfig.TabIcon.GetValue() : FSlateIcon(FInsightsStyle::GetStyleSetName(), "TraceStore.Icon.Small"));
+
+		TSharedRef<FWorkspaceItem> Group = TraceStoreConfig.WorkspaceGroup.IsValid() ? TraceStoreConfig.WorkspaceGroup.ToSharedRef() : WorkspaceMenu::GetMenuStructure().GetDeveloperToolsProfilingCategory();
+		TabSpawnerEntry.SetGroup(Group);
+	}
+
+	const FInsightsMajorTabConfig& ConnectionConfig = InsightsModule.FindMajorTabConfig(FInsightsManagerTabs::ConnectionTabId);
+	if (ConnectionConfig.bIsAvailable)
 	{
-		// Register tab spawner for the Start Page.
-		FTabSpawnerEntry& TabSpawnerEntry = FGlobalTabmanager::Get()->RegisterNomadTabSpawner(FInsightsManagerTabs::StartPageTabId,
-			FOnSpawnTab::CreateRaw(this, &FInsightsManager::SpawnStartPageTab))
-			.SetDisplayName(StartPageConfig.TabLabel.IsSet() ? StartPageConfig.TabLabel.GetValue() : LOCTEXT("StartPageTabTitle", "Trace Store"))
-			.SetTooltipText(StartPageConfig.TabTooltip.IsSet() ? StartPageConfig.TabTooltip.GetValue() : LOCTEXT("StartPageTooltipText", "Open the Trace Store Browser."))
-			.SetIcon(StartPageConfig.TabIcon.IsSet() ? StartPageConfig.TabIcon.GetValue() : FSlateIcon(FInsightsStyle::GetStyleSetName(), "StartPage.Icon.Small"));
+		// Register tab spawner for the Connection tab.
+		FTabSpawnerEntry& TabSpawnerEntry = FGlobalTabmanager::Get()->RegisterNomadTabSpawner(FInsightsManagerTabs::ConnectionTabId,
+			FOnSpawnTab::CreateRaw(this, &FInsightsManager::SpawnConnectionTab))
+			.SetDisplayName(ConnectionConfig.TabLabel.IsSet() ? ConnectionConfig.TabLabel.GetValue() : LOCTEXT("ConnectionTabTitle", "Connection"))
+			.SetTooltipText(ConnectionConfig.TabTooltip.IsSet() ? ConnectionConfig.TabTooltip.GetValue() : LOCTEXT("ConnectionTooltipText", "Open the Connection tab."))
+			.SetIcon(ConnectionConfig.TabIcon.IsSet() ? ConnectionConfig.TabIcon.GetValue() : FSlateIcon(FInsightsStyle::GetStyleSetName(), "Connection.Icon.Small"));
+
+		TSharedRef<FWorkspaceItem> Group = ConnectionConfig.WorkspaceGroup.IsValid() ? ConnectionConfig.WorkspaceGroup.ToSharedRef() : WorkspaceMenu::GetMenuStructure().GetDeveloperToolsProfilingCategory();
+		TabSpawnerEntry.SetGroup(Group);
+	}
 
-		TSharedRef<FWorkspaceItem> Group = StartPageConfig.WorkspaceGroup.IsValid() ? StartPageConfig.WorkspaceGroup.ToSharedRef() : WorkspaceMenu::GetMenuStructure().GetDeveloperToolsProfilingCategory();
+	const FInsightsMajorTabConfig& LauncherConfig = InsightsModule.FindMajorTabConfig(FInsightsManagerTabs::LauncherTabId);
+	if (LauncherConfig.bIsAvailable)
+	{
+		// Register tab spawner for the Launcher tab.
+		FTabSpawnerEntry& TabSpawnerEntry = FGlobalTabmanager::Get()->RegisterNomadTabSpawner(FInsightsManagerTabs::LauncherTabId,
+			FOnSpawnTab::CreateRaw(this, &FInsightsManager::SpawnLauncherTab))
+			.SetDisplayName(LauncherConfig.TabLabel.IsSet() ? LauncherConfig.TabLabel.GetValue() : LOCTEXT("LauncherTabTitle", "Launcher"))
+			.SetTooltipText(LauncherConfig.TabTooltip.IsSet() ? LauncherConfig.TabTooltip.GetValue() : LOCTEXT("LauncherTooltipText", "Open the Launcher tab."))
+			.SetIcon(LauncherConfig.TabIcon.IsSet() ? LauncherConfig.TabIcon.GetValue() : FSlateIcon(FInsightsStyle::GetStyleSetName(), "Launcher.Icon.Small"));
+
+		TSharedRef<FWorkspaceItem> Group = LauncherConfig.WorkspaceGroup.IsValid() ? LauncherConfig.WorkspaceGroup.ToSharedRef() : WorkspaceMenu::GetMenuStructure().GetDeveloperToolsProfilingCategory();
 		TabSpawnerEntry.SetGroup(Group);
 	}
 
@@ -253,25 +297,25 @@ void FInsightsManager::RegisterMajorTabs(IUnrealInsightsModule& InsightsModule)
 void FInsightsManager::UnregisterMajorTabs()
 {
 	FGlobalTabmanager::Get()->UnregisterNomadTabSpawner(FInsightsManagerTabs::SessionInfoTabId);
-	FGlobalTabmanager::Get()->UnregisterNomadTabSpawner(FInsightsManagerTabs::StartPageTabId);
+
+	FGlobalTabmanager::Get()->UnregisterNomadTabSpawner(FInsightsManagerTabs::TraceStoreTabId);
+	FGlobalTabmanager::Get()->UnregisterNomadTabSpawner(FInsightsManagerTabs::ConnectionTabId);
+	FGlobalTabmanager::Get()->UnregisterNomadTabSpawner(FInsightsManagerTabs::LauncherTabId);
 }
 
 ////////////////////////////////////////////////////////////////////////////////////////////////////
 
-TSharedRef<SDockTab> FInsightsManager::SpawnStartPageTab(const FSpawnTabArgs& Args)
+TSharedRef<SDockTab> FInsightsManager::SpawnTraceStoreTab(const FSpawnTabArgs& Args)
 {
-	const TSharedRef<SDockTab> DockTab = SNew(SDockTab)
-		.TabRole(ETabRole::NomadTab);
-		//.OnCanCloseTab_Lambda([]() { return false; })
-		//.ContentPadding(FMargin(2.0f, 20.0f, 2.0f, 2.0f));
+	const TSharedRef<SDockTab> DockTab = SNew(SLockedTab)
+		.TabRole(ETabRole::MajorTab);
 
-	DockTab->SetOnTabClosed(SDockTab::FOnTabClosedCallback::CreateRaw(this, &FInsightsManager::OnStartPageTabClosed));
+	DockTab->SetOnTabClosed(SDockTab::FOnTabClosedCallback::CreateRaw(this, &FInsightsManager::OnTraceStoreTabClosed));
 
-	// Create the SStartPageWindow widget.
-	TSharedRef<SStartPageWindow> Window = SNew(SStartPageWindow);
+	TSharedRef<STraceStoreWindow> Window = SNew(STraceStoreWindow);
 	DockTab->SetContent(Window);
 
-	AssignStartPageWindow(Window);
+	AssignTraceStoreWindow(Window);
 
 	if (!bIsMainTabSet)
 	{
@@ -284,9 +328,77 @@ TSharedRef<SDockTab> FInsightsManager::SpawnStartPageTab(const FSpawnTabArgs& Ar
 
 ////////////////////////////////////////////////////////////////////////////////////////////////////
 
-void FInsightsManager::OnStartPageTabClosed(TSharedRef<SDockTab> TabBeingClosed)
+void FInsightsManager::OnTraceStoreTabClosed(TSharedRef<SDockTab> TabBeingClosed)
+{
+	RemoveTraceStoreWindow();
+
+	// Disable TabClosed delegate.
+	TabBeingClosed->SetOnTabClosed(SDockTab::FOnTabClosedCallback());
+}
+
+////////////////////////////////////////////////////////////////////////////////////////////////////
+
+TSharedRef<SDockTab> FInsightsManager::SpawnConnectionTab(const FSpawnTabArgs& Args)
+{
+	const TSharedRef<SDockTab> DockTab = SNew(SLockedTab)
+		.TabRole(ETabRole::MajorTab)
+		.OnCanCloseTab_Lambda([]() { return false; }); // can't close this tab
+
+	DockTab->SetOnTabClosed(SDockTab::FOnTabClosedCallback::CreateRaw(this, &FInsightsManager::OnConnectionTabClosed));
+
+	TSharedRef<SConnectionWindow> Window = SNew(SConnectionWindow);
+	DockTab->SetContent(Window);
+
+	AssignConnectionWindow(Window);
+
+	//if (!bIsMainTabSet)
+	//{
+	//	FGlobalTabmanager::Get()->SetMainTab(DockTab);
+	//	bIsMainTabSet = true;
+	//}
+
+	return DockTab;
+}
+
+////////////////////////////////////////////////////////////////////////////////////////////////////
+
+void FInsightsManager::OnConnectionTabClosed(TSharedRef<SDockTab> TabBeingClosed)
+{
+	RemoveConnectionWindow();
+
+	// Disable TabClosed delegate.
+	TabBeingClosed->SetOnTabClosed(SDockTab::FOnTabClosedCallback());
+}
+
+////////////////////////////////////////////////////////////////////////////////////////////////////
+
+TSharedRef<SDockTab> FInsightsManager::SpawnLauncherTab(const FSpawnTabArgs& Args)
+{
+	const TSharedRef<SDockTab> DockTab = SNew(SLockedTab)
+		.TabRole(ETabRole::NomadTab)
+		.OnCanCloseTab_Lambda([]() { return false; }); // can't close this tab
+
+	DockTab->SetOnTabClosed(SDockTab::FOnTabClosedCallback::CreateRaw(this, &FInsightsManager::OnLauncherTabClosed));
+
+	TSharedRef<SLauncherWindow> Window = SNew(SLauncherWindow);
+	DockTab->SetContent(Window);
+
+	AssignLauncherWindow(Window);
+
+	//if (!bIsMainTabSet)
+	//{
+	//	FGlobalTabmanager::Get()->SetMainTab(DockTab);
+	//	bIsMainTabSet = true;
+	//}
+
+	return DockTab;
+}
+
+////////////////////////////////////////////////////////////////////////////////////////////////////
+
+void FInsightsManager::OnLauncherTabClosed(TSharedRef<SDockTab> TabBeingClosed)
 {
-	RemoveStartPageWindow();
+	RemoveLauncherWindow();
 
 	// Disable TabClosed delegate.
 	TabBeingClosed->SetOnTabClosed(SDockTab::FOnTabClosedCallback());
@@ -688,7 +800,7 @@ void FInsightsManager::LoadTraceFile(const FString& InTraceFilename, bool InAuto
 
 void FInsightsManager::OpenSettings()
 {
-	TSharedPtr<SStartPageWindow> Wnd = GetStartPageWindow();
+	TSharedPtr<STraceStoreWindow> Wnd = GetTraceStoreWindow();
 	if (Wnd.IsValid())
 	{
 		Wnd->OpenSettings();
diff --git a/Engine/Source/Developer/TraceInsights/Private/Insights/InsightsManager.h b/Engine/Source/Developer/TraceInsights/Private/Insights/InsightsManager.h
--- a/Engine/Source/Developer/TraceInsights/Private/Insights/InsightsManager.h
+++ b/Engine/Source/Developer/TraceInsights/Private/Insights/InsightsManager.h
@@ -30,7 +30,9 @@ namespace TraceServices
 	class IModuleService;
 }
 
-class SStartPageWindow;
+class STraceStoreWindow;
+class SConnectionWindow;
+class SLauncherWindow;
 class SSessionInfoWindow;
 class FInsightsTestRunner;
 class FInsightsMenuBuilder;
@@ -125,25 +127,69 @@ public:
 	static FInsightsSettings& GetSettings();
 
 	//////////////////////////////////////////////////
-	// StartPage (Trace Store Browser)
+	// Trace Store
 
-	void AssignStartPageWindow(const TSharedRef<SStartPageWindow>& InStartPageWindow)
+	void AssignTraceStoreWindow(const TSharedRef<STraceStoreWindow>& InTraceStoreWindow)
 	{
-		StartPageWindow = InStartPageWindow;
+		TraceStoreWindow = InTraceStoreWindow;
 	}
 
-	void RemoveStartPageWindow()
+	void RemoveTraceStoreWindow()
 	{
-		StartPageWindow.Reset();
+		TraceStoreWindow.Reset();
 	}
 
 	/**
-	 * Converts profiler window weak pointer to a shared pointer and returns it.
+	 * Converts Trace Store window weak pointer to a shared pointer and returns it.
+	 * Make sure the returned pointer is valid before trying to dereference it.
+	 */
+	TSharedPtr<class STraceStoreWindow> GetTraceStoreWindow() const
+	{
+		return TraceStoreWindow.Pin();
+	}
+
+	//////////////////////////////////////////////////
+	// Connection
+
+	void AssignConnectionWindow(const TSharedRef<SConnectionWindow>& InConnectionWindow)
+	{
+		ConnectionWindow = InConnectionWindow;
+	}
+
+	void RemoveConnectionWindow()
+	{
+		ConnectionWindow.Reset();
+	}
+
+	/**
+	 * Converts Connection window weak pointer to a shared pointer and returns it.
 	 * Make sure the returned pointer is valid before trying to dereference it.
 	 */
-	TSharedPtr<class SStartPageWindow> GetStartPageWindow() const
+	TSharedPtr<class SConnectionWindow> GetConnectionWindow() const
 	{
-		return StartPageWindow.Pin();
+		return ConnectionWindow.Pin();
+	}
+
+	//////////////////////////////////////////////////
+	// Launcher
+
+	void AssignLauncherWindow(const TSharedRef<SLauncherWindow>& InLauncherWindow)
+	{
+		LauncherWindow = InLauncherWindow;
+	}
+
+	void RemoveLauncherWindow()
+	{
+		LauncherWindow.Reset();
+	}
+
+	/**
+	 * Converts Launcher window weak pointer to a shared pointer and returns it.
+	 * Make sure the returned pointer is valid before trying to dereference it.
+	 */
+	TSharedPtr<class SLauncherWindow> GetLauncherWindow() const
+	{
+		return LauncherWindow.Pin();
 	}
 
 	//////////////////////////////////////////////////
@@ -237,11 +283,23 @@ private:
 	/** Binds our UI commands to delegates. */
 	void BindCommands();
 
-	/** Called to spawn the Start Page major tab. */
-	TSharedRef<SDockTab> SpawnStartPageTab(const FSpawnTabArgs& Args);
+	/** Called to spawn the Trace Store major tab. */
+	TSharedRef<SDockTab> SpawnTraceStoreTab(const FSpawnTabArgs& Args);
+
+	/** Callback called when the Trace Store major tab is closed. */
+	void OnTraceStoreTabClosed(TSharedRef<SDockTab> TabBeingClosed);
+
+	/** Called to spawn the Connection major tab. */
+	TSharedRef<SDockTab> SpawnConnectionTab(const FSpawnTabArgs& Args);
 
-	/** Callback called when the Start Page major tab is closed. */
-	void OnStartPageTabClosed(TSharedRef<SDockTab> TabBeingClosed);
+	/** Callback called when the Connection major tab is closed. */
+	void OnConnectionTabClosed(TSharedRef<SDockTab> TabBeingClosed);
+
+	/** Called to spawn the Launcher major tab. */
+	TSharedRef<SDockTab> SpawnLauncherTab(const FSpawnTabArgs& Args);
+
+	/** Callback called when the Launcher major tab is closed. */
+	void OnLauncherTabClosed(TSharedRef<SDockTab> TabBeingClosed);
 
 	/** Called to spawn the Session Info major tab. */
 	TSharedRef<SDockTab> SpawnSessionInfoTab(const FSpawnTabArgs& Args);
@@ -306,8 +364,14 @@ private:
 	/** An instance of the main settings. */
 	FInsightsSettings Settings;
 
-	/** A weak pointer to the Start Page window. */
-	TWeakPtr<class SStartPageWindow> StartPageWindow;
+	/** A weak pointer to the Trace Store window. */
+	TWeakPtr<class STraceStoreWindow> TraceStoreWindow;
+
+	/** A weak pointer to the Connection window. */
+	TWeakPtr<class SConnectionWindow> ConnectionWindow;
+
+	/** A weak pointer to the Launcher window. */
+	TWeakPtr<class SLauncherWindow> LauncherWindow;
 
 	/** A weak pointer to the Session Info window. */
 	TWeakPtr<class SSessionInfoWindow> SessionInfoWindow;
diff --git a/Engine/Source/Developer/TraceInsights/Private/Insights/InsightsStyle.cpp b/Engine/Source/Developer/TraceInsights/Private/Insights/InsightsStyle.cpp
--- a/Engine/Source/Developer/TraceInsights/Private/Insights/InsightsStyle.cpp
+++ b/Engine/Source/Developer/TraceInsights/Private/Insights/InsightsStyle.cpp
@@ -83,17 +83,20 @@ TSharedRef<FSlateStyleSet> FInsightsStyle::Create()
 	Style.Set("Toolbar.Icon.Small", new IMAGE_BRUSH("Icons/icon_tab_Tools_16x", Icon16x16));
 
 	//////////////////////////////////////////////////
-	// Start Page buttons
+	// Trace Store, Connection, Launcher
 
-	Style.Set("StartPage.Icon.Large", new IMAGE_BRUSH_SVG("../../Editor/Slate/Starship/Common/Home", Icon32x32));
-	Style.Set("StartPage.Icon.Small", new IMAGE_BRUSH_SVG("../../Editor/Slate/Starship/Common/Home", Icon16x16));
+	Style.Set("TraceStore.Icon.Large", new IMAGE_BRUSH_SVG("../../Editor/Slate/Starship/Common/Home", Icon32x32));
+	Style.Set("TraceStore.Icon.Small", new IMAGE_BRUSH_SVG("../../Editor/Slate/Starship/Common/Home", Icon16x16));
 
-	Style.Set("Open.Icon.Large", new IMAGE_BRUSH("Icons/Profiler/Profiler_LoadMultiple_Profiler_40x", Icon32x32));
-	Style.Set("Open.Icon.Small", new IMAGE_BRUSH("Icons/Profiler/Profiler_Load_Profiler_40x", Icon16x16));
+	Style.Set("Connection.Icon.Large", new TODO_IMAGE_BRUSH(Icon32x32));
+	Style.Set("Connection.Icon.Small", new TODO_IMAGE_BRUSH(Icon16x16));
+
+	Style.Set("Launcher.Icon.Large", new TODO_IMAGE_BRUSH(Icon32x32));
+	Style.Set("Launcher.Icon.Small", new TODO_IMAGE_BRUSH(Icon16x16));
+
+	Style.Set("FolderExplore.Icon.Large", new IMAGE_BRUSH_SVG("../../Editor/Slate/Starship/Common/ContentBrowser", Icon32x32));
+	Style.Set("FolderExplore.Icon.Small", new IMAGE_BRUSH_SVG("../../Editor/Slate/Starship/Common/ContentBrowser", Icon16x16));
 
-	Style.Set("OpenFile.Icon.Large", new IMAGE_BRUSH("Icons/LV_Load", Icon32x32));
-	Style.Set("OpenFile.Icon.Small", new IMAGE_BRUSH("Icons/LV_Load", Icon16x16));
-	
 	//////////////////////////////////////////////////
 	// Timing Insights
 
@@ -157,8 +160,9 @@ TSharedRef<FSlateStyleSet> FInsightsStyle::Create()
 	Style.Set("MemTagTreeView.Icon.Large", new TODO_IMAGE_BRUSH(Icon32x32));
 	Style.Set("MemTagTreeView.Icon.Small", new TODO_IMAGE_BRUSH(Icon16x16));
 
-	Style.Set("Mem.Add.Small", new IMAGE_BRUSH("Icons/icon_Cascade_AddLOD2_40x", Icon20x20));
-	Style.Set("Mem.Remove.Small", new IMAGE_BRUSH("Icons/icon_Cascade_DeleteLOD_40x", Icon20x20));
+	Style.Set("Mem.Add.Small", new CORE_IMAGE_BRUSH_SVG("Starship/Common/plus-circle", Icon16x16));
+	Style.Set("Mem.Remove.Small", new CORE_IMAGE_BRUSH_SVG("Starship/Common/Delete", Icon16x16));
+	Style.Set("Mem.LoadXML.Small", new CORE_IMAGE_BRUSH_SVG("Starship/Common/folder-open", Icon16x16));
 
 	//////////////////////////////////////////////////
 
@@ -177,8 +181,11 @@ TSharedRef<FSlateStyleSet> FInsightsStyle::Create()
 	Style.Set("Icon.Filter", new CORE_IMAGE_BRUSH_SVG("Starship/Common/filter", Icon24x24));
 	Style.Set("Icon.Find", new IMAGE_BRUSH_SVG("../../Editor/Slate/Starship/Common/TraceDataFiltering", Icon24x24));
 
-	Style.Set("FolderExplore.Icon.Large", new IMAGE_BRUSH_SVG("../../Editor/Slate/Starship/Common/ContentBrowser", Icon32x32));
-	Style.Set("FolderExplore.Icon.Small", new IMAGE_BRUSH_SVG("../../Editor/Slate/Starship/Common/ContentBrowser", Icon16x16));
+	Style.Set("Icon.FolderExplore", new IMAGE_BRUSH_SVG("../../Editor/Slate/Starship/Common/ContentBrowser", Icon24x24));
+	Style.Set("Icon.FolderOpen", new CORE_IMAGE_BRUSH_SVG("Starship/Common/folder-open", Icon24x24));
+
+	Style.Set("Icon.TestAutomation", new IMAGE_BRUSH_SVG("../../Editor/Slate/Starship/Common/TestAutomation", Icon24x24));
+	Style.Set("Icon.Test", new IMAGE_BRUSH_SVG("../../Editor/Slate/Starship/Common/Test", Icon24x24));
 
 	Style.Set("Icon.Bug", new IMAGE_BRUSH_SVG("../../Editor/Slate/Starship/Common/bug", Icon24x24));
 
diff --git a/Engine/Source/Developer/TraceInsights/Private/Insights/TraceInsightsModule.cpp b/Engine/Source/Developer/TraceInsights/Private/Insights/TraceInsightsModule.cpp
--- a/Engine/Source/Developer/TraceInsights/Private/Insights/TraceInsightsModule.cpp
+++ b/Engine/Source/Developer/TraceInsights/Private/Insights/TraceInsightsModule.cpp
@@ -203,14 +203,17 @@ void FTraceInsightsModule::CreateSessionBrowser(const FCreateSessionBrowserParam
 	//////////////////////////////////////////////////
 	// Setup the window's content.
 
-	TSharedRef<FTabManager::FLayout> DefaultLayout = FTabManager::NewLayout("TraceSessionBrowserLayout_v1.0");
+	TSharedRef<FTabManager::FLayout> DefaultLayout = FTabManager::NewLayout("TraceSessionBrowserLayout_v1.1");
 	DefaultLayout->AddArea
 	(
 		FTabManager::NewPrimaryArea()
 		->Split
 		(
 			FTabManager::NewStack()
-			->AddTab(FInsightsManagerTabs::StartPageTabId, ETabState::OpenedTab)
+			->AddTab(FInsightsManagerTabs::TraceStoreTabId, ETabState::OpenedTab)
+			->AddTab(FInsightsManagerTabs::ConnectionTabId, ETabState::OpenedTab)
+			//->AddTab(FInsightsManagerTabs::LauncherTabId, ETabState::ClosedTab)
+			->SetForegroundTab(FInsightsManagerTabs::TraceStoreTabId)
 		)
 	);
 
@@ -234,12 +237,12 @@ void FTraceInsightsModule::CreateSessionBrowser(const FCreateSessionBrowserParam
 	//////////////////////////////////////////////////
 	// Set up command line parameter forwarding.
 
-	TSharedPtr<class SStartPageWindow> StartPageWnd = FInsightsManager::Get()->GetStartPageWindow();
-	if (StartPageWnd.IsValid())
+	TSharedPtr<class STraceStoreWindow> TraceStoreWnd = FInsightsManager::Get()->GetTraceStoreWindow();
+	if (TraceStoreWnd.IsValid())
 	{
-		StartPageWnd->SetEnableAutomaticTesting(Params.bInitializeTesting);
-		StartPageWnd->SetEnableDebugTools(Params.bAllowDebugTools);
-		StartPageWnd->SetStartProcessWithStompMalloc(Params.bStartProcessWithStompMalloc);
+		TraceStoreWnd->SetEnableAutomaticTesting(Params.bInitializeTesting);
+		TraceStoreWnd->SetEnableDebugTools(Params.bAllowDebugTools);
+		TraceStoreWnd->SetStartProcessWithStompMalloc(Params.bStartProcessWithStompMalloc);
 	}
 }
 
@@ -324,7 +327,6 @@ void FTraceInsightsModule::AddAreaForSessionViewer(TSharedRef<FTabManager::FLayo
 
 #if WITH_EDITOR
 	// In editor, we default to all tabs closed.
-	Stack->AddTab(FInsightsManagerTabs::StartPageTabId, ETabState::ClosedTab);
 	Stack->AddTab(FInsightsManagerTabs::SessionInfoTabId, ETabState::ClosedTab);
 	Stack->AddTab(FInsightsManagerTabs::TimingProfilerTabId, ETabState::ClosedTab);
 	Stack->AddTab(FInsightsManagerTabs::LoadingProfilerTabId, ETabState::ClosedTab);
@@ -340,7 +342,6 @@ void FTraceInsightsModule::AddAreaForSessionViewer(TSharedRef<FTabManager::FLayo
 		->Split(Stack)
 	);
 #else
-	Stack->AddTab(FInsightsManagerTabs::StartPageTabId, ETabState::ClosedTab);
 	Stack->AddTab(FInsightsManagerTabs::SessionInfoTabId, ETabState::OpenedTab);
 	Stack->AddTab(FInsightsManagerTabs::TimingProfilerTabId, ETabState::OpenedTab);
 	Stack->AddTab(FInsightsManagerTabs::LoadingProfilerTabId, ETabState::ClosedTab);
diff --git a/Engine/Source/Developer/TraceInsights/Private/Insights/TraceInsightsModule.h b/Engine/Source/Developer/TraceInsights/Private/Insights/TraceInsightsModule.h
--- a/Engine/Source/Developer/TraceInsights/Private/Insights/TraceInsightsModule.h
+++ b/Engine/Source/Developer/TraceInsights/Private/Insights/TraceInsightsModule.h
@@ -90,8 +90,14 @@ protected:
 	/** Callback called when a major tab is closed. */
 	void OnTabBeingClosed(TSharedRef<SDockTab> TabBeingClosed);
 
-	/** Start Page */
-	TSharedRef<SDockTab> SpawnStartPageTab(const FSpawnTabArgs& Args);
+	/** Trace Store */
+	TSharedRef<SDockTab> SpawnTraceStoreTab(const FSpawnTabArgs& Args);
+
+	/** Connection */
+	TSharedRef<SDockTab> SpawnConnectionTab(const FSpawnTabArgs& Args);
+
+	/** Launcher */
+	TSharedRef<SDockTab> SpawnLauncherTab(const FSpawnTabArgs& Args);
 
 	/** Session Info */
 	TSharedRef<SDockTab> SpawnSessionInfoTab(const FSpawnTabArgs& Args);
diff --git a/Engine/Source/Developer/TraceInsights/Private/Insights/Widgets/SStartPageWindow.cpp b/Engine/Source/Developer/TraceInsights/Private/Insights/Widgets/SStartPageWindow.cpp
--- a/Engine/Source/Developer/TraceInsights/Private/Insights/Widgets/SStartPageWindow.cpp
+++ b/Engine/Source/Developer/TraceInsights/Private/Insights/Widgets/SStartPageWindow.cpp
@@ -2,42 +2,31 @@
 
 #include "SStartPageWindow.h"
 
-#include "Containers/StringView.h"
 #include "DesktopPlatformModule.h"
-#include "EditorStyleSet.h"
 #include "Framework/Application/SlateApplication.h"
-#include "Framework/Docking/WorkspaceItem.h"
 #include "Framework/MultiBox/MultiBoxBuilder.h"
-#include "HAL/FileManagerGeneric.h"
-#include "Input/Events.h"
 #include "Internationalization/Text.h"
 #include "IPAddress.h"
 #include "SlateOptMacros.h"
 #include "SocketSubsystem.h"
 #include "Styling/CoreStyle.h"
 #include "Styling/StyleColors.h"
-#include "Trace/Analysis.h"
-#include "Trace/Analyzer.h"
 #include "Trace/ControlClient.h"
 #include "Trace/StoreClient.h"
-#include "Widgets/Docking/SDockTab.h"
 #include "Widgets/Images/SImage.h"
+#include "Widgets/Input/SButton.h"
+#include "Widgets/Input/SEditableTextBox.h"
+#include "Widgets/Input/SComboButton.h"
 #include "Widgets/Input/SCheckBox.h"
 #include "Widgets/Input/SSearchBox.h"
 #include "Widgets/Layout/SBorder.h"
 #include "Widgets/Layout/SBox.h"
 #include "Widgets/Layout/SGridPanel.h"
-#include "Widgets/Layout/SScrollBox.h"
-#include "Widgets/Layout/SSeparator.h"
-#include "Widgets/Layout/SSpacer.h"
 #include "Widgets/Notifications/SNotificationList.h"
 #include "Widgets/SBoxPanel.h"
 #include "Widgets/SToolTip.h"
-#include "Widgets/Text/STextBlock.h"
-#include "Widgets/Input/SButton.h"
-#include "Widgets/Input/SEditableTextBox.h"
-#include "Widgets/Input/SComboButton.h"
 #include "Widgets/Testing/SStarshipSuite.h"
+#include "Widgets/Text/STextBlock.h"
 
 #if WITH_EDITOR
 	#include "EngineAnalytics.h"
@@ -50,7 +39,6 @@
 #include "Insights/InsightsManager.h"
 #include "Insights/Log.h"
 #include "Insights/StoreService/StoreBrowser.h"
-#include "Insights/TimingProfilerManager.h"
 #include "Insights/InsightsStyle.h"
 #include "Insights/Version.h"
 #include "Insights/Widgets/SInsightsSettings.h"
@@ -58,7 +46,7 @@
 
 ////////////////////////////////////////////////////////////////////////////////////////////////////
 
-#define LOCTEXT_NAMESPACE "SStartPageWindow"
+#define LOCTEXT_NAMESPACE "STraceStoreWindow"
 
 ////////////////////////////////////////////////////////////////////////////////////////////////////
 // STraceListRow
@@ -77,7 +65,7 @@ public:
 	 * @param InTrace The trace displayed by this row.
 	 * @param InOwnerTableView The table to which the row must be added.
 	 */
-	void Construct(const FArguments& InArgs, TSharedPtr<FTraceViewModel> InTrace, TSharedRef<SStartPageWindow> InParentWidget, const TSharedRef<STableViewBase>& InOwnerTableView)
+	void Construct(const FArguments& InArgs, TSharedPtr<FTraceViewModel> InTrace, TSharedRef<STraceStoreWindow> InParentWidget, const TSharedRef<STableViewBase>& InOwnerTableView)
 	{
 		WeakTrace = MoveTemp(InTrace);
 		WeakParentWidget = InParentWidget;
@@ -484,43 +472,65 @@ public:
 				SNew(SVerticalBox)
 
 				+ SVerticalBox::Slot()
+				.Padding(FMargin(-7.0f, -7.0f, -7.0f, 0.0f))
 				.AutoHeight()
-				.Padding(2.0f)
 				[
-					SNew(SHorizontalBox)
-
-					+ SHorizontalBox::Slot()
-					.FillWidth(1.0f)
-					[
-						SNew(STextBlock)
-						.Text(this, &STraceListRow::GetTraceName)
-						.Font(FCoreStyle::GetDefaultFontStyle("Bold", 10))
-						.ColorAndOpacity(FLinearColor::White)
-					]
-
-					+ SHorizontalBox::Slot()
-					.AutoWidth()
+					SNew(SBorder)
+					.Padding(FMargin(6.0f, 6.0f, 6.0f, 6.0f))
+					.BorderImage(FInsightsStyle::Get().GetBrush("WhiteBrush"))
+					.BorderBackgroundColor(FSlateColor(EStyleColor::Panel))
 					[
-						SNew(STextBlock)
-						.Text(this, &STraceListRow::GetTraceIndexAndId)
-						.ColorAndOpacity(FLinearColor::Gray)
+						SNew(SHorizontalBox)
+
+						+ SHorizontalBox::Slot()
+						.Padding(FMargin(2.0f, 2.0f, 2.0f, 2.0f))
+						.FillWidth(1.0f)
+						[
+							SNew(STextBlock)
+							.Text(this, &STraceListRow::GetTraceName)
+							.Font(FCoreStyle::GetDefaultFontStyle("Regular", 12))
+							.ColorAndOpacity(FSlateColor(EStyleColor::Foreground))
+						]
+
+						+ SHorizontalBox::Slot()
+						.Padding(FMargin(2.0f, 2.0f, 2.0f, 2.0f))
+						.AutoWidth()
+						[
+							SNew(STextBlock)
+							.Font(FCoreStyle::GetDefaultFontStyle("Regular", 12))
+							.Text(this, &STraceListRow::GetTraceIndexAndId)
+							.ColorAndOpacity(FSlateColor(EStyleColor::White25))
+						]
 					]
 				]
 
 				+ SVerticalBox::Slot()
+				.Padding(FMargin(-7.0f, 1.0f, -7.0f, 0.0f))
 				.AutoHeight()
-				.Padding(2.0f)
 				[
-					SNew(STextBlock)
-					.Text(this, &STraceListRow::GetTraceUri)
-					.TextStyle(FEditorStyle::Get(), TEXT("Profiler.Tooltip"))
+					SNew(SBorder)
+					.Padding(FMargin(6.0f, 6.0f, 6.0f, 4.0f))
+					.BorderImage(FInsightsStyle::Get().GetBrush("WhiteBrush"))
+					.BorderBackgroundColor(FSlateColor(EStyleColor::Panel))
+					[
+						SNew(STextBlock)
+						.Text(this, &STraceListRow::GetTraceUri)
+						.Font(FCoreStyle::GetDefaultFontStyle("Regular", 8))
+						.ColorAndOpacity(FSlateColor(EStyleColor::Foreground))
+					]
 				]
 
 				+ SVerticalBox::Slot()
 				.AutoHeight()
-				.Padding(2.0f)
+				.Padding(FMargin(-7.0f, 0.0f, -7.0f, -7.0f))
 				[
-					SAssignNew(GridPanel, SGridPanel)
+					SNew(SBorder)
+					.Padding(FMargin(6.0f, 0.0f, 6.0f, 4.0f))
+					.BorderImage(FInsightsStyle::Get().GetBrush("WhiteBrush"))
+					.BorderBackgroundColor(FSlateColor(EStyleColor::Panel))
+					[
+						SAssignNew(GridPanel, SGridPanel)
+					]
 				]
 			];
 
@@ -565,8 +575,7 @@ private:
 			[
 				SNew(STextBlock)
 				.Text(InHeaderText)
-				.Font(FCoreStyle::GetDefaultFontStyle("Bold", 10))
-				.ColorAndOpacity(FLinearColor::Gray)
+				.ColorAndOpacity(FSlateColor(EStyleColor::White25))
 			];
 
 		SGridPanel::FSlot* Slot1 = nullptr;
@@ -579,7 +588,7 @@ private:
 				.Text(this, InValueTextFn)
 				.WrapTextAt(512.0f)
 				.WrappingPolicy(ETextWrappingPolicy::AllowPerCharacterWrapping)
-				.ColorAndOpacity(FLinearColor::White)
+				.ColorAndOpacity(FSlateColor(EStyleColor::Foreground))
 			];
 
 		if (InVisibilityFn)
@@ -597,14 +606,14 @@ private:
 
 private:
 	TWeakPtr<FTraceViewModel> WeakTrace;
-	TWeakPtr<SStartPageWindow> WeakParentWidget;
+	TWeakPtr<STraceStoreWindow> WeakParentWidget;
 };
 
 ////////////////////////////////////////////////////////////////////////////////////////////////////
-// SStartPageWindow
+// STraceStoreWindow
 ////////////////////////////////////////////////////////////////////////////////////////////////////
 
-SStartPageWindow::SStartPageWindow()
+STraceStoreWindow::STraceStoreWindow()
 	: NotificationList()
 	, ActiveNotifications()
 	, OverlaySettingsSlot(nullptr)
@@ -624,14 +633,13 @@ SStartPageWindow::SStartPageWindow()
 	, TraceViewModelMap()
 	, TraceListView()
 	, SelectedTrace()
-	, HostTextBox()
 	, SplashScreenOverlayFadeTime(0.0f)
 {
 }
 
 ////////////////////////////////////////////////////////////////////////////////////////////////////
 
-SStartPageWindow::~SStartPageWindow()
+STraceStoreWindow::~STraceStoreWindow()
 {
 #if WITH_EDITOR
 	if (DurationActive > 0.0f && FEngineAnalytics::IsAvailable())
@@ -645,7 +653,7 @@ SStartPageWindow::~SStartPageWindow()
 
 BEGIN_SLATE_FUNCTION_BUILD_OPTIMIZATION
 
-void SStartPageWindow::Construct(const FArguments& InArgs)
+void STraceStoreWindow::Construct(const FArguments& InArgs)
 {
 	ChildSlot
 	[
@@ -707,27 +715,10 @@ void SStartPageWindow::Construct(const FArguments& InArgs)
 			.HAlign(HAlign_Fill)
 			.VAlign(VAlign_Fill)
 			.AutoHeight()
-			.Padding(12.0f, 4.0f)
+			.Padding(12.0f, 4.0f, 12.0f, 8.0f)
 			[
 				ConstructLoadPanel()
 			]
-
-			+ SVerticalBox::Slot()
-			.HAlign(HAlign_Fill)
-			.AutoHeight()
-			.Padding(3.0f, 3.0f)
-			[
-				SNew(SBox)
-				.Visibility(this, &SStartPageWindow::StopTraceRecorder_Visibility)
-				[
-					SNew(SBorder)
-					.BorderImage(FCoreStyle::Get().GetBrush("PopupText.Background"))
-					.Padding(8.0f)
-					[
-						ConstructConnectPanel()
-					]
-				]
-			]
 		]
 
 		// Overlay for fake splashscreen.
@@ -737,11 +728,11 @@ void SStartPageWindow::Construct(const FArguments& InArgs)
 		.Padding(0.0f)
 		[
 			SNew(SBox)
-			.Visibility(this, &SStartPageWindow::SplashScreenOverlay_Visibility)
+			.Visibility(this, &STraceStoreWindow::SplashScreenOverlay_Visibility)
 			[
 				SNew(SBorder)
 				.BorderImage(FCoreStyle::Get().GetBrush("PopupText.Background"))
-				.BorderBackgroundColor(this, &SStartPageWindow::SplashScreenOverlay_ColorAndOpacity)
+				.BorderBackgroundColor(this, &STraceStoreWindow::SplashScreenOverlay_ColorAndOpacity)
 				.Padding(0.0f)
 				.HAlign(HAlign_Fill)
 				.VAlign(VAlign_Fill)
@@ -751,9 +742,9 @@ void SStartPageWindow::Construct(const FArguments& InArgs)
 					.VAlign(VAlign_Center)
 					[
 						SNew(STextBlock)
-						.Text(this, &SStartPageWindow::GetSplashScreenOverlayText)
+						.Text(this, &STraceStoreWindow::GetSplashScreenOverlayText)
 						.Font(FCoreStyle::GetDefaultFontStyle("Bold", 11))
-						.ColorAndOpacity(this, &SStartPageWindow::SplashScreenOverlay_TextColorAndOpacity)
+						.ColorAndOpacity(this, &STraceStoreWindow::SplashScreenOverlay_TextColorAndOpacity)
 					]
 				]
 			]
@@ -783,18 +774,18 @@ void SStartPageWindow::Construct(const FArguments& InArgs)
 
 ////////////////////////////////////////////////////////////////////////////////////////////////////
 
-TSharedRef<SWidget> SStartPageWindow::ConstructSessionsPanel()
+TSharedRef<SWidget> STraceStoreWindow::ConstructSessionsPanel()
 {
 	TSharedRef<SWidget> Widget = SAssignNew(TraceListView, SListView<TSharedPtr<FTraceViewModel>>)
 		.IsFocusable(true)
 		.ItemHeight(20.0f)
 		.SelectionMode(ESelectionMode::Single)
-		.OnSelectionChanged(this, &SStartPageWindow::TraceList_OnSelectionChanged)
-		.OnMouseButtonDoubleClick(this, &SStartPageWindow::TraceList_OnMouseButtonDoubleClick)
+		.OnSelectionChanged(this, &STraceStoreWindow::TraceList_OnSelectionChanged)
+		.OnMouseButtonDoubleClick(this, &STraceStoreWindow::TraceList_OnMouseButtonDoubleClick)
 		.ListItemsSource(&TraceViewModels)
-		.OnGenerateRow(this, &SStartPageWindow::TraceList_OnGenerateRow)
+		.OnGenerateRow(this, &STraceStoreWindow::TraceList_OnGenerateRow)
 		.ConsumeMouseWheel(EConsumeMouseWheel::Always)
-		//.OnContextMenuOpening(FOnContextMenuOpening::CreateSP(this, &SStartPageWindow::TraceList_GetContextMenu))
+		//.OnContextMenuOpening(FOnContextMenuOpening::CreateSP(this, &STraceStoreWindow::TraceList_GetContextMenu))
 		.HeaderRow
 		(
 			SNew(SHeaderRow)
@@ -837,7 +828,7 @@ TSharedRef<SWidget> SStartPageWindow::ConstructSessionsPanel()
 
 ////////////////////////////////////////////////////////////////////////////////////////////////////
 
-TSharedRef<SWidget> SStartPageWindow::ConstructLoadPanel()
+TSharedRef<SWidget> STraceStoreWindow::ConstructLoadPanel()
 {
 	TSharedRef<SWidget> Widget = SNew(SHorizontalBox)
 
@@ -856,8 +847,8 @@ TSharedRef<SWidget> SStartPageWindow::ConstructLoadPanel()
 	[
 		SNew(SButton)
 		.ButtonStyle(&FAppStyle::Get().GetWidgetStyle<FButtonStyle>("PrimaryButton"))
-		.IsEnabled(this, &SStartPageWindow::Open_IsEnabled)
-		.OnClicked(this, &SStartPageWindow::Open_OnClicked)
+		.IsEnabled(this, &STraceStoreWindow::Open_IsEnabled)
+		.OnClicked(this, &STraceStoreWindow::Open_OnClicked)
 		.ToolTipText(LOCTEXT("OpenButtonTooltip", "Start analysis for selected trace session."))
 		.ContentPadding(FMargin(0.0f, 0.0f, 0.0f, 0.0f))
 		.Content()
@@ -879,7 +870,7 @@ TSharedRef<SWidget> SStartPageWindow::ConstructLoadPanel()
 	[
 		SNew(SComboButton)
 		.ToolTipText(LOCTEXT("MRU_Tooltip", "Open a trace file or choose a trace session."))
-		.OnGetMenuContent(this, &SStartPageWindow::MakeTraceListMenu)
+		.OnGetMenuContent(this, &STraceStoreWindow::MakeTraceListMenu)
 		.HasDownArrow(true)
 	];
 
@@ -888,7 +879,7 @@ TSharedRef<SWidget> SStartPageWindow::ConstructLoadPanel()
 
 ////////////////////////////////////////////////////////////////////////////////////////////////////
 
-TSharedRef<SWidget> SStartPageWindow::ConstructTraceStoreDirectoryPanel()
+TSharedRef<SWidget> STraceStoreWindow::ConstructTraceStoreDirectoryPanel()
 {
 	TSharedRef<SWidget> Widget =
 
@@ -911,7 +902,7 @@ TSharedRef<SWidget> SStartPageWindow::ConstructTraceStoreDirectoryPanel()
 			SNew(SEditableTextBox)
 			.IsReadOnly(true)
 			.BackgroundColor(FSlateColor(EStyleColor::Background))
-			.Text(this, &SStartPageWindow::GetTraceStoreDirectory)
+			.Text(this, &STraceStoreWindow::GetTraceStoreDirectory)
 		]
 
 		+ SHorizontalBox::Slot()
@@ -922,7 +913,7 @@ TSharedRef<SWidget> SStartPageWindow::ConstructTraceStoreDirectoryPanel()
 			SNew(SButton)
 			.ButtonStyle(&FAppStyle::Get().GetWidgetStyle<FButtonStyle>("SimpleButton"))
 			.ToolTipText(LOCTEXT("ExploreTraceStoreDirButtonToolTip", "Explore the Trace Store Directory"))
-			.OnClicked(this, &SStartPageWindow::ExploreTraceStoreDirectory_OnClicked)
+			.OnClicked(this, &STraceStoreWindow::ExploreTraceStoreDirectory_OnClicked)
 			[
 				SNew(SImage)
 				.Image(FInsightsStyle::Get().GetBrush("FolderExplore.Icon.Small"))
@@ -935,7 +926,7 @@ TSharedRef<SWidget> SStartPageWindow::ConstructTraceStoreDirectoryPanel()
 
 ////////////////////////////////////////////////////////////////////////////////////////////////////
 
-TSharedRef<SWidget> SStartPageWindow::ConstructAutoStartPanel()
+TSharedRef<SWidget> STraceStoreWindow::ConstructAutoStartPanel()
 {
 	TSharedRef<SWidget> Widget = SNew(SHorizontalBox)
 
@@ -947,8 +938,8 @@ TSharedRef<SWidget> SStartPageWindow::ConstructAutoStartPanel()
 	[
 		SNew(SCheckBox)
 		.ToolTipText(LOCTEXT("AutoStart_Tooltip", "Enable auto-start analysis for LIVE trace sessions."))
-		.IsChecked(this, &SStartPageWindow::AutoStart_IsChecked)
-		.OnCheckStateChanged(this, &SStartPageWindow::AutoStart_OnCheckStateChanged)
+		.IsChecked(this, &STraceStoreWindow::AutoStart_IsChecked)
+		.OnCheckStateChanged(this, &STraceStoreWindow::AutoStart_OnCheckStateChanged)
 		[
 			SNew(STextBlock)
 			.Text(LOCTEXT("AutoStart_Text", "Auto-start analysis for LIVE trace sessions"))
@@ -957,7 +948,7 @@ TSharedRef<SWidget> SStartPageWindow::ConstructAutoStartPanel()
 
 	+ SHorizontalBox::Slot()
 	.AutoWidth()
-	.Padding(4.0f, 0.0f, 0.0f, 0.0f)
+	.Padding(6.0f, 0.0f, 0.0f, 0.0f)
 	.HAlign(HAlign_Left)
 	.VAlign(VAlign_Center)
 	[
@@ -968,7 +959,7 @@ TSharedRef<SWidget> SStartPageWindow::ConstructAutoStartPanel()
 
 	+ SHorizontalBox::Slot()
 	.AutoWidth()
-	.Padding(4.0f, 0.0f, 0.0f, 0.0f)
+	.Padding(6.0f, 0.0f, 0.0f, 0.0f)
 	.HAlign(HAlign_Left)
 	.VAlign(VAlign_Center)
 	[
@@ -982,160 +973,7 @@ TSharedRef<SWidget> SStartPageWindow::ConstructAutoStartPanel()
 
 ////////////////////////////////////////////////////////////////////////////////////////////////////
 
-TSharedRef<SWidget> SStartPageWindow::ConstructRecorderPanel()
-{
-	TSharedRef<SWidget> Widget = SNew(SVerticalBox)
-
-	+ SVerticalBox::Slot()
-	.AutoHeight()
-	.HAlign(HAlign_Center)
-	.Padding(0.0f, 2.0f)
-	[
-		SNew(STextBlock)
-		.Text(LOCTEXT("RecorderPanelTitle", "Trace Recorder"))
-		.Font(FCoreStyle::GetDefaultFontStyle("Bold", 11))
-		.ColorAndOpacity(FLinearColor::Gray)
-	]
-
-	+ SVerticalBox::Slot()
-	.AutoHeight()
-	.HAlign(HAlign_Left)
-	.Padding(0.0f, 2.0f)
-	[
-		SNew(SHorizontalBox)
-
-		+ SHorizontalBox::Slot()
-		.AutoWidth()
-		.Padding(0.0f, 0.0f, 2.0f, 0.0f)
-		.VAlign(VAlign_Center)
-		[
-			SNew(STextBlock)
-			.Text(LOCTEXT("RecorderStatusTitle", "Status:"))
-			.ColorAndOpacity(FLinearColor::Gray)
-		]
-
-		+ SHorizontalBox::Slot()
-		.AutoWidth()
-		.Padding(2.0f, 0.0f)
-		.VAlign(VAlign_Center)
-		[
-			SNew(STextBlock)
-			.Text(this, &SStartPageWindow::GetRecorderStatusText)
-		]
-
-		+ SHorizontalBox::Slot()
-		.AutoWidth()
-		.Padding(2.0f, 0.0f)
-		.VAlign(VAlign_Center)
-		[
-			SNew(SButton)
-			.Text(LOCTEXT("StartRecorder", "Start"))
-			.ToolTipText(LOCTEXT("StartRecorderToolTip", "Start the Trace Recorder"))
-			.OnClicked(this, &SStartPageWindow::StartTraceRecorder_OnClicked)
-			.Visibility(this, &SStartPageWindow::StartTraceRecorder_Visibility)
-		]
-
-		+ SHorizontalBox::Slot()
-		.AutoWidth()
-		.Padding(2.0f, 0.0f)
-		.VAlign(VAlign_Center)
-		[
-			SNew(SButton)
-			.Text(LOCTEXT("StopRecorder", "Stop"))
-			.ToolTipText(LOCTEXT("StopRecorderToolTip", "Stop the Trace Recorder"))
-			.OnClicked(this, &SStartPageWindow::StopTraceRecorder_OnClicked)
-			.Visibility(this, &SStartPageWindow::StopTraceRecorder_Visibility)
-		]
-	]
-
-	+ SVerticalBox::Slot()
-	.AutoHeight()
-	.HAlign(HAlign_Left)
-	.Padding(0.0f, 2.0f, 0.0f, 1.0f)
-	[
-		SNew(SHorizontalBox)
-		.Visibility(this, &SStartPageWindow::StopTraceRecorder_Visibility)
-
-		+ SHorizontalBox::Slot()
-		.AutoWidth()
-		.Padding(0.0f, 0.0f, 2.0f, 0.0f)
-		.VAlign(VAlign_Center)
-		[
-			SNew(STextBlock)
-			.Text_Lambda([this]() -> FText
-			{
-				return FText::Format(LOCTEXT("ConnectionCountFormat", "Connections / live sessions: {0}"), FText::AsNumber(LiveSessionCount));
-			})
-			.ColorAndOpacity(FLinearColor::Gray)
-		]
-	];
-
-	return Widget;
-}
-
-////////////////////////////////////////////////////////////////////////////////////////////////////
-
-TSharedRef<SWidget> SStartPageWindow::ConstructConnectPanel()
-{
-	TSharedRef<SWidget> Widget = SNew(SVerticalBox)
-
-	+ SVerticalBox::Slot()
-	.AutoHeight()
-	.HAlign(HAlign_Left)
-	.Padding(0.0f, 2.0f)
-	[
-		SNew(STextBlock)
-		.Text(LOCTEXT("ConnectPanelTitle", "New Connection"))
-		.Font(FCoreStyle::GetDefaultFontStyle("Bold", 11))
-		.ColorAndOpacity(FLinearColor::Gray)
-	]
-
-	+ SVerticalBox::Slot()
-	.AutoHeight()
-	.HAlign(HAlign_Fill)
-	.Padding(0.0f, 2.0f)
-	[
-		SNew(SHorizontalBox)
-
-		+ SHorizontalBox::Slot()
-		.AutoWidth()
-		.Padding(0.0f, 0.0f, 4.0f, 0.0f)
-		.VAlign(VAlign_Center)
-		[
-			SNew(STextBlock)
-			.Text(LOCTEXT("HostTitle", "Running instance IP"))
-			.ColorAndOpacity(FLinearColor::Gray)
-		]
-
-		+ SHorizontalBox::Slot()
-		.AutoWidth()
-		.VAlign(VAlign_Center)
-		[
-			SNew(SBox)
-			.MinDesiredWidth(120.0f)
-			[
-				SAssignNew(HostTextBox, SEditableTextBox)
-			]
-		]
-
-		+ SHorizontalBox::Slot()
-		.AutoWidth()
-		.Padding(2.0f, 0.0f)
-		.VAlign(VAlign_Center)
-		[
-			SNew(SButton)
-			.Text(LOCTEXT("Connect", "Connect"))
-			.ToolTipText(LOCTEXT("ConnectToolTip", "Connect the running instance at specified ip with the local trace recorder."))
-			.OnClicked(this, &SStartPageWindow::Connect_OnClicked)
-		]
-	];
-
-	return Widget;
-}
-
-////////////////////////////////////////////////////////////////////////////////////////////////////
-
-TSharedRef<ITableRow> SStartPageWindow::TraceList_OnGenerateRow(TSharedPtr<FTraceViewModel> InTrace, const TSharedRef<STableViewBase>& OwnerTable)
+TSharedRef<ITableRow> STraceStoreWindow::TraceList_OnGenerateRow(TSharedPtr<FTraceViewModel> InTrace, const TSharedRef<STableViewBase>& OwnerTable)
 {
 	return SNew(STraceListRow, InTrace, SharedThis(this), OwnerTable);
 }
@@ -1144,14 +982,14 @@ END_SLATE_FUNCTION_BUILD_OPTIMIZATION
 
 ////////////////////////////////////////////////////////////////////////////////////////////////////
 
-void SStartPageWindow::ShowSplashScreenOverlay()
+void STraceStoreWindow::ShowSplashScreenOverlay()
 {
 	SplashScreenOverlayFadeTime = 3.5f;
 }
 
 ////////////////////////////////////////////////////////////////////////////////////////////////////
 
-void SStartPageWindow::TickSplashScreenOverlay(const float InDeltaTime)
+void STraceStoreWindow::TickSplashScreenOverlay(const float InDeltaTime)
 {
 	if (SplashScreenOverlayFadeTime > 0.0f)
 	{
@@ -1161,7 +999,7 @@ void SStartPageWindow::TickSplashScreenOverlay(const float InDeltaTime)
 
 ////////////////////////////////////////////////////////////////////////////////////////////////////
 
-float SStartPageWindow::SplashScreenOverlayOpacity() const
+float STraceStoreWindow::SplashScreenOverlayOpacity() const
 {
 	constexpr float FadeInStartTime = 3.5f;
 	constexpr float FadeInEndTime = 3.0f;
@@ -1179,35 +1017,35 @@ float SStartPageWindow::SplashScreenOverlayOpacity() const
 
 ////////////////////////////////////////////////////////////////////////////////////////////////////
 
-EVisibility SStartPageWindow::SplashScreenOverlay_Visibility() const
+EVisibility STraceStoreWindow::SplashScreenOverlay_Visibility() const
 {
 	return SplashScreenOverlayFadeTime > 0.0f ? EVisibility::Visible : EVisibility::Collapsed;
 }
 
 ////////////////////////////////////////////////////////////////////////////////////////////////////
 
-FSlateColor SStartPageWindow::SplashScreenOverlay_ColorAndOpacity() const
+FSlateColor STraceStoreWindow::SplashScreenOverlay_ColorAndOpacity() const
 {
 	return FSlateColor(FLinearColor(0.7f, 0.7f, 0.7f, SplashScreenOverlayOpacity()));
 }
 
 ////////////////////////////////////////////////////////////////////////////////////////////////////
 
-FSlateColor SStartPageWindow::SplashScreenOverlay_TextColorAndOpacity() const
+FSlateColor STraceStoreWindow::SplashScreenOverlay_TextColorAndOpacity() const
 {
 	return FSlateColor(FLinearColor(0.8f, 0.8f, 0.8f, SplashScreenOverlayOpacity()));
 }
 
 ////////////////////////////////////////////////////////////////////////////////////////////////////
 
-FText SStartPageWindow::GetSplashScreenOverlayText() const
+FText STraceStoreWindow::GetSplashScreenOverlayText() const
 {
 	return FText::Format(LOCTEXT("StartAnalysis", "Starting analysis...\n{0}"), FText::FromString(SplashScreenOverlayTraceFile));
 }
 
 ////////////////////////////////////////////////////////////////////////////////////////////////////
 
-FReply SStartPageWindow::RefreshTraces_OnClicked()
+FReply STraceStoreWindow::RefreshTraces_OnClicked()
 {
 	RefreshTraceList();
 	return FReply::Handled();
@@ -1215,71 +1053,7 @@ FReply SStartPageWindow::RefreshTraces_OnClicked()
 
 ////////////////////////////////////////////////////////////////////////////////////////////////////
 
-FReply SStartPageWindow::Connect_OnClicked()
-{
-	FText HostText = HostTextBox->GetText();
-	if (HostText.IsEmptyOrWhitespace())
-	{
-		// nothing to do
-		return FReply::Handled();
-	}
-
-	const FString& HostStr = HostText.ToString();
-	UE_LOG(TraceInsights, Log, TEXT("[StartPage] Try connecting to \"%s\"..."), *HostStr);
-
-	bool bConnectedSuccessfully = false;
-	UE::Trace::FControlClient ControlClient;
-	if (ControlClient.Connect(*HostStr))
-	{
-		TSharedPtr<FInternetAddr> RecorderAddr;
-		if (ISocketSubsystem* Sockets = ISocketSubsystem::Get())
-		{
-			bool bCanBindAll = false;
-			RecorderAddr = Sockets->GetLocalHostAddr(*GLog, bCanBindAll);
-			if (RecorderAddr.IsValid())
-			{
-				const FString RecorderAddrStr = RecorderAddr->ToString(false);
-				UE_LOG(TraceInsights, Log, TEXT("[StartPage] SendSendTo(\"%s\")..."), *RecorderAddrStr);
-				ControlClient.SendSendTo(*RecorderAddrStr);
-				bConnectedSuccessfully = true;
-			}
-		}
-	}
-
-	if (bConnectedSuccessfully)
-	{
-		UE_LOG(TraceInsights, Log, TEXT("[StartPage] Successfully connected."));
-		FNotificationInfo NotificationInfo(FText::Format(LOCTEXT("ConnectSuccess", "Successfully connected to \"{0}\"!"), HostText));
-		NotificationInfo.bFireAndForget = false;
-		NotificationInfo.bUseLargeFont = false;
-		NotificationInfo.bUseSuccessFailIcons = true;
-		NotificationInfo.ExpireDuration = 10.0f;
-		SNotificationItemWeak NotificationItem = NotificationList->AddNotification(NotificationInfo);
-		NotificationItem.Pin()->SetCompletionState(SNotificationItem::CS_Success);
-		NotificationItem.Pin()->ExpireAndFadeout();
-		ActiveNotifications.Add(TEXT("ConnectSuccess"), NotificationItem);
-	}
-	else
-	{
-		UE_LOG(TraceInsights, Warning, TEXT("[StartPage] Failed to connect to \"%s\"!"), *HostStr);
-		FNotificationInfo NotificationInfo(FText::Format(LOCTEXT("ConnectFailed", "Failed to connect to \"{0}\"!"), HostText));
-		NotificationInfo.bFireAndForget = false;
-		NotificationInfo.bUseLargeFont = false;
-		NotificationInfo.bUseSuccessFailIcons = true;
-		NotificationInfo.ExpireDuration = 10.0f;
-		SNotificationItemWeak NotificationItem = NotificationList->AddNotification(NotificationInfo);
-		NotificationItem.Pin()->SetCompletionState(SNotificationItem::CS_Fail);
-		NotificationItem.Pin()->ExpireAndFadeout();
-		ActiveNotifications.Add(TEXT("ConnectFailed"), NotificationItem);
-	}
-
-	RefreshTraceList();
-	return FReply::Handled();
-}
-
-////////////////////////////////////////////////////////////////////////////////////////////////////
-
-void SStartPageWindow::RefreshTraceList()
+void STraceStoreWindow::RefreshTraceList()
 {
 	UE::Trace::FStoreClient* StoreClient = FInsightsManager::Get()->GetStoreClient();
 	if (StoreClient == nullptr)
@@ -1302,7 +1076,7 @@ void SStartPageWindow::RefreshTraceList()
 		{
 			TracesChangeSerial = NewChangeSerial;
 
-			//UE_LOG(TraceInsights, Log, TEXT("[StartPage] Synching the trace list with StoreBrowser..."));
+			//UE_LOG(TraceInsights, Log, TEXT("[TraceStore] Synching the trace list with StoreBrowser..."));
 
 			const TArray<TSharedPtr<Insights::FStoreBrowserTraceInfo>>& InTraces = StoreBrowser->GetLockedTraces();
 			const TMap<uint32, TSharedPtr<Insights::FStoreBrowserTraceInfo>>& InTraceMap = StoreBrowser->GetLockedTraceMap();
@@ -1370,14 +1144,14 @@ void SStartPageWindow::RefreshTraceList()
 	const double Duration = StopwatchTotal.GetAccumulatedTime();
 	if ((Duration > 0.0001) && (UpdatedTraces > 0 || AddedTraces > 0 || RemovedTraces > 0))
 	{
-		UE_LOG(TraceInsights, Log, TEXT("[StartPage] The trace list refreshed in %.0f ms (%d traces : %d updated, %d added, %d removed)."),
+		UE_LOG(TraceInsights, Log, TEXT("[TraceStore] The trace list refreshed in %.0f ms (%d traces : %d updated, %d added, %d removed)."),
 			Duration * 1000.0, TraceViewModels.Num(), UpdatedTraces, AddedTraces, RemovedTraces);
 	}
 }
 
 ////////////////////////////////////////////////////////////////////////////////////////////////////
 
-void SStartPageWindow::UpdateTrace(FTraceViewModel& InOutTrace, const Insights::FStoreBrowserTraceInfo& InSourceTrace)
+void STraceStoreWindow::UpdateTrace(FTraceViewModel& InOutTrace, const Insights::FStoreBrowserTraceInfo& InSourceTrace)
 {
 	//TraceId -- no need to update
 
@@ -1422,7 +1196,7 @@ void SStartPageWindow::UpdateTrace(FTraceViewModel& InOutTrace, const Insights::
 			(AutoStartConfigurationTypeFilter == EBuildConfiguration::Unknown || AutoStartConfigurationTypeFilter == InOutTrace.ConfigurationType) &&
 			(AutoStartTargetTypeFilter == EBuildTargetType::Unknown || AutoStartTargetTypeFilter == InOutTrace.TargetType))
 		{
-			UE_LOG(TraceInsights, Log, TEXT("[StartPage] Auto starting analysis for trace with id 0x%08X..."), InOutTrace.TraceId);
+			UE_LOG(TraceInsights, Log, TEXT("[TraceStore] Auto starting analysis for trace with id 0x%08X..."), InOutTrace.TraceId);
 			AutoStartedSessions.Add(InOutTrace.TraceId);
 			LoadTrace(InOutTrace.TraceId);
 		}
@@ -1431,7 +1205,7 @@ void SStartPageWindow::UpdateTrace(FTraceViewModel& InOutTrace, const Insights::
 
 ////////////////////////////////////////////////////////////////////////////////////////////////////
 
-void SStartPageWindow::OnTraceListChanged()
+void STraceStoreWindow::OnTraceListChanged()
 {
 	Algo::SortBy(TraceViewModels, &FTraceViewModel::Timestamp);
 
@@ -1467,35 +1241,35 @@ void SStartPageWindow::OnTraceListChanged()
 
 ////////////////////////////////////////////////////////////////////////////////////////////////////
 
-void SStartPageWindow::TraceList_OnSelectionChanged(TSharedPtr<FTraceViewModel> TraceSession, ESelectInfo::Type SelectInfo)
+void STraceStoreWindow::TraceList_OnSelectionChanged(TSharedPtr<FTraceViewModel> TraceSession, ESelectInfo::Type SelectInfo)
 {
 	SelectedTrace = TraceSession;
 }
 
 ////////////////////////////////////////////////////////////////////////////////////////////////////
 
-void SStartPageWindow::TraceList_OnMouseButtonDoubleClick(TSharedPtr<FTraceViewModel> TraceSession)
+void STraceStoreWindow::TraceList_OnMouseButtonDoubleClick(TSharedPtr<FTraceViewModel> TraceSession)
 {
 	LoadTraceSession(TraceSession);
 }
 
 ////////////////////////////////////////////////////////////////////////////////////////////////////
 
-ECheckBoxState SStartPageWindow::AutoStart_IsChecked() const
+ECheckBoxState STraceStoreWindow::AutoStart_IsChecked() const
 {
 	return bAutoStartAnalysisForLiveSessions ? ECheckBoxState::Checked : ECheckBoxState::Unchecked;
 }
 
 ////////////////////////////////////////////////////////////////////////////////////////////////////
 
-void SStartPageWindow::AutoStart_OnCheckStateChanged(ECheckBoxState NewState)
+void STraceStoreWindow::AutoStart_OnCheckStateChanged(ECheckBoxState NewState)
 {
 	bAutoStartAnalysisForLiveSessions = (NewState == ECheckBoxState::Checked);
 }
 
 ////////////////////////////////////////////////////////////////////////////////////////////////////
 
-void SStartPageWindow::Tick(const FGeometry& AllottedGeometry, const double InCurrentTime, const float InDeltaTime)
+void STraceStoreWindow::Tick(const FGeometry& AllottedGeometry, const double InCurrentTime, const float InDeltaTime)
 {
 	// We need to update the trace list, but not too often.
 	static uint64 NextTimestamp = 0;
@@ -1512,7 +1286,7 @@ void SStartPageWindow::Tick(const FGeometry& AllottedGeometry, const double InCu
 
 ////////////////////////////////////////////////////////////////////////////////////////////////////
 
-EActiveTimerReturnType SStartPageWindow::UpdateActiveDuration(double InCurrentTime, float InDeltaTime)
+EActiveTimerReturnType STraceStoreWindow::UpdateActiveDuration(double InCurrentTime, float InDeltaTime)
 {
 	DurationActive += InDeltaTime;
 
@@ -1522,19 +1296,19 @@ EActiveTimerReturnType SStartPageWindow::UpdateActiveDuration(double InCurrentTi
 
 ////////////////////////////////////////////////////////////////////////////////////////////////////
 
-void SStartPageWindow::OnMouseEnter(const FGeometry& MyGeometry, const FPointerEvent& MouseEvent)
+void STraceStoreWindow::OnMouseEnter(const FGeometry& MyGeometry, const FPointerEvent& MouseEvent)
 {
 	SCompoundWidget::OnMouseEnter(MyGeometry, MouseEvent);
 
 	if (!ActiveTimerHandle.IsValid())
 	{
-		ActiveTimerHandle = RegisterActiveTimer(0.f, FWidgetActiveTimerDelegate::CreateSP(this, &SStartPageWindow::UpdateActiveDuration));
+		ActiveTimerHandle = RegisterActiveTimer(0.f, FWidgetActiveTimerDelegate::CreateSP(this, &STraceStoreWindow::UpdateActiveDuration));
 	}
 }
 
 ////////////////////////////////////////////////////////////////////////////////////////////////////
 
-void SStartPageWindow::OnMouseLeave(const FPointerEvent& MouseEvent)
+void STraceStoreWindow::OnMouseLeave(const FPointerEvent& MouseEvent)
 {
 	SCompoundWidget::OnMouseLeave(MouseEvent);
 
@@ -1547,14 +1321,14 @@ void SStartPageWindow::OnMouseLeave(const FPointerEvent& MouseEvent)
 
 ////////////////////////////////////////////////////////////////////////////////////////////////////
 
-FReply SStartPageWindow::OnKeyDown(const FGeometry& MyGeometry, const FKeyEvent& InKeyEvent)
+FReply STraceStoreWindow::OnKeyDown(const FGeometry& MyGeometry, const FKeyEvent& InKeyEvent)
 {
 	return FReply::Unhandled();
 }
 
 ////////////////////////////////////////////////////////////////////////////////////////////////////
 
-FReply SStartPageWindow::OnDragOver(const FGeometry& MyGeometry, const FDragDropEvent& DragDropEvent)
+FReply STraceStoreWindow::OnDragOver(const FGeometry& MyGeometry, const FDragDropEvent& DragDropEvent)
 {
 	TSharedPtr<FExternalDragOperation> DragDropOp = DragDropEvent.GetOperationAs<FExternalDragOperation>();
 	if (DragDropOp.IsValid())
@@ -1578,7 +1352,7 @@ FReply SStartPageWindow::OnDragOver(const FGeometry& MyGeometry, const FDragDrop
 
 ////////////////////////////////////////////////////////////////////////////////////////////////////
 
-FReply SStartPageWindow::OnDrop(const FGeometry& MyGeometry, const FDragDropEvent& DragDropEvent)
+FReply STraceStoreWindow::OnDrop(const FGeometry& MyGeometry, const FDragDropEvent& DragDropEvent)
 {
 	TSharedPtr<FExternalDragOperation> DragDropOp = DragDropEvent.GetOperationAs<FExternalDragOperation>();
 	if (DragDropOp.IsValid())
@@ -1604,14 +1378,14 @@ FReply SStartPageWindow::OnDrop(const FGeometry& MyGeometry, const FDragDropEven
 
 ////////////////////////////////////////////////////////////////////////////////////////////////////
 
-bool SStartPageWindow::Open_IsEnabled() const
+bool STraceStoreWindow::Open_IsEnabled() const
 {
 	return TraceViewModels.Num() > 0;
 }
 
 ////////////////////////////////////////////////////////////////////////////////////////////////////
 
-FReply SStartPageWindow::Open_OnClicked()
+FReply STraceStoreWindow::Open_OnClicked()
 {
 	LoadTraceSession(SelectedTrace);
 	return FReply::Handled();
@@ -1619,7 +1393,7 @@ FReply SStartPageWindow::Open_OnClicked()
 
 ////////////////////////////////////////////////////////////////////////////////////////////////////
 
-void SStartPageWindow::OpenFileDialog()
+void STraceStoreWindow::OpenFileDialog()
 {
 	const FString ProfilingDirectory(FPaths::ConvertRelativePathToFull(FInsightsManager::Get()->GetStoreDir()));
 
@@ -1654,7 +1428,7 @@ void SStartPageWindow::OpenFileDialog()
 
 ////////////////////////////////////////////////////////////////////////////////////////////////////
 
-void SStartPageWindow::LoadTraceSession(TSharedPtr<FTraceViewModel> InTraceSession)
+void STraceStoreWindow::LoadTraceSession(TSharedPtr<FTraceViewModel> InTraceSession)
 {
 	if (InTraceSession.IsValid())
 	{
@@ -1664,9 +1438,9 @@ void SStartPageWindow::LoadTraceSession(TSharedPtr<FTraceViewModel> InTraceSessi
 
 ////////////////////////////////////////////////////////////////////////////////////////////////////
 
-void SStartPageWindow::LoadTraceFile(const FString& InTraceFile)
+void STraceStoreWindow::LoadTraceFile(const FString& InTraceFile)
 {
-	UE_LOG(TraceInsights, Log, TEXT("[StartPage] Start analysis (in separate process) for trace file: \"%s\""), *InTraceFile);
+	UE_LOG(TraceInsights, Log, TEXT("[TraceStore] Start analysis (in separate process) for trace file: \"%s\""), *InTraceFile);
 
 	const TCHAR* ExecutablePath = FPlatformProcess::ExecutablePath();
 
@@ -1699,9 +1473,9 @@ void SStartPageWindow::LoadTraceFile(const FString& InTraceFile)
 
 ////////////////////////////////////////////////////////////////////////////////////////////////////
 
-void SStartPageWindow::LoadTrace(uint32 InTraceId)
+void STraceStoreWindow::LoadTrace(uint32 InTraceId)
 {
-	UE_LOG(TraceInsights, Log, TEXT("[StartPage] Start analysis (in separate process) for trace id: 0x%08X"), InTraceId);
+	UE_LOG(TraceInsights, Log, TEXT("[TraceStore] Start analysis (in separate process) for trace id: 0x%08X"), InTraceId);
 
 	const TCHAR* ExecutablePath = FPlatformProcess::ExecutablePath();
 
@@ -1740,7 +1514,7 @@ void SStartPageWindow::LoadTrace(uint32 InTraceId)
 
 ////////////////////////////////////////////////////////////////////////////////////////////////////
 
-TSharedRef<SWidget> SStartPageWindow::MakeTraceListMenu()
+TSharedRef<SWidget> STraceStoreWindow::MakeTraceListMenu()
 {
 	RefreshTraceList();
 
@@ -1751,8 +1525,8 @@ TSharedRef<SWidget> SStartPageWindow::MakeTraceListMenu()
 		MenuBuilder.AddMenuEntry(
 			LOCTEXT("OpenFileButtonLabel", "Open File..."),
 			LOCTEXT("OpenFileButtonTooltip", "Start analysis for a specified trace file."),
-			FSlateIcon(FInsightsStyle::GetStyleSetName(), "OpenFile.Icon.Small"),
-			FUIAction(FExecuteAction::CreateSP(this, &SStartPageWindow::OpenFileDialog)),
+			FSlateIcon(FInsightsStyle::GetStyleSetName(), "Icon.FolderOpen"),
+			FUIAction(FExecuteAction::CreateSP(this, &STraceStoreWindow::OpenFileDialog)),
 			NAME_None,
 			EUserInterfaceActionType::Button
 		);
@@ -1785,7 +1559,7 @@ TSharedRef<SWidget> SStartPageWindow::MakeTraceListMenu()
 					Label,
 					TAttribute<FText>(), // no tooltip
 					FSlateIcon(),
-					FUIAction(FExecuteAction::CreateSP(this, &SStartPageWindow::LoadTrace, Trace.TraceId)),
+					FUIAction(FExecuteAction::CreateSP(this, &STraceStoreWindow::LoadTrace, Trace.TraceId)),
 					NAME_None,
 					EUserInterfaceActionType::Button
 				);
@@ -1811,7 +1585,7 @@ TSharedRef<SWidget> SStartPageWindow::MakeTraceListMenu()
 		MenuBuilder.AddMenuEntry(
 			LOCTEXT("EnableAutomatedTesting", "Enable Automation Testing"),
 			LOCTEXT("EnableAutomatedTestingDesc", "Activates the automatic test system for new sessions opened from this window."),
-			FSlateIcon(),
+			FSlateIcon(FInsightsStyle::GetStyleSetName(), "Icon.TestAutomation"),
 			ToogleAutomationTestsAction,
 			NAME_None,
 			EUserInterfaceActionType::ToggleButton
@@ -1833,7 +1607,7 @@ TSharedRef<SWidget> SStartPageWindow::MakeTraceListMenu()
 		MenuBuilder.AddMenuEntry(
 			LOCTEXT("EnableDebugTools", "Enable Debug Tools"),
 			LOCTEXT("EnableDebugToolsDesc", "Enables debug tools for new sessions opened from this window."),
-			FSlateIcon(),
+			FSlateIcon(FInsightsStyle::GetStyleSetName(), "Icon.Bug"),
 			ToogleDebugToolsAction,
 			NAME_None,
 			EUserInterfaceActionType::ToggleButton
@@ -1852,7 +1626,7 @@ TSharedRef<SWidget> SStartPageWindow::MakeTraceListMenu()
 		MenuBuilder.AddMenuEntry(
 			LOCTEXT("OpenStarshipSuite", "Starship Test Suite"),
 			LOCTEXT("OpenStarshipSuiteDesc", "Opens the Starship UX test suite."),
-			FSlateIcon(FInsightsStyle::GetStyleSetName(), "Icon.Bug"),
+			FSlateIcon(FInsightsStyle::GetStyleSetName(), "Icon.Test"),
 			OpenStarshipSuiteAction,
 			NAME_None,
 			EUserInterfaceActionType::Button
@@ -1867,14 +1641,14 @@ TSharedRef<SWidget> SStartPageWindow::MakeTraceListMenu()
 
 ////////////////////////////////////////////////////////////////////////////////////////////////////
 
-FText SStartPageWindow::GetTraceStoreDirectory() const
+FText STraceStoreWindow::GetTraceStoreDirectory() const
 {
 	return FText::FromString(FPaths::ConvertRelativePathToFull(FInsightsManager::Get()->GetStoreDir()));
 }
 
 ////////////////////////////////////////////////////////////////////////////////////////////////////
 
-FReply SStartPageWindow::ExploreTraceStoreDirectory_OnClicked()
+FReply STraceStoreWindow::ExploreTraceStoreDirectory_OnClicked()
 {
 	FString FullPath(FPaths::ConvertRelativePathToFull(FInsightsManager::Get()->GetStoreDir()));
 	FPlatformProcess::ExploreFolder(*FullPath);
@@ -1883,107 +1657,381 @@ FReply SStartPageWindow::ExploreTraceStoreDirectory_OnClicked()
 
 ////////////////////////////////////////////////////////////////////////////////////////////////////
 
-FText SStartPageWindow::GetRecorderStatusText() const
+void STraceStoreWindow::OpenSettings()
 {
-	UE::Trace::FStoreClient* StoreClient = FInsightsManager::Get()->GetStoreClient();
-	const bool bIsRecorderServerRunning = (StoreClient != nullptr); // TODO: StoreClient->IsRecorderServerRunning();
-
-	if (bIsRecorderServerRunning)
-	{
-		return FText(LOCTEXT("RecorderServerRunning", "Running"));
-	}
-	else
-	{
-		return FText(LOCTEXT("RecorderServerStopped", "Stopped"));
-	}
+	MainContentPanel->SetEnabled(false);
+	(*OverlaySettingsSlot)
+	[
+		SNew(SBorder)
+		.BorderImage(FCoreStyle::Get().GetBrush("PopupText.Background"))
+		.Padding(8.0f)
+		[
+			SNew(SInsightsSettings)
+			.OnClose(this, &STraceStoreWindow::CloseSettings)
+			.SettingPtr(&FInsightsManager::GetSettings())
+		]
+	];
 }
 
 ////////////////////////////////////////////////////////////////////////////////////////////////////
 
-EVisibility SStartPageWindow::StartTraceRecorder_Visibility() const
+void STraceStoreWindow::CloseSettings()
 {
-	UE::Trace::FStoreClient* StoreClient = FInsightsManager::Get()->GetStoreClient();
-	const bool bIsRecorderServerRunning = (StoreClient != nullptr); // TODO: StoreClient->IsRecorderServerRunning();
-
-	return bIsRecorderServerRunning ? EVisibility::Collapsed : EVisibility::Visible;
+	// Close the profiler settings by simply replacing widget with a null one.
+	(*OverlaySettingsSlot)
+	[
+		SNullWidget::NullWidget
+	];
+	MainContentPanel->SetEnabled(true);
 }
 
 ////////////////////////////////////////////////////////////////////////////////////////////////////
 
-EVisibility SStartPageWindow::StopTraceRecorder_Visibility() const
+void STraceStoreWindow::GetExtraCommandLineParams(FString& OutParams) const
 {
-	UE::Trace::FStoreClient* StoreClient = FInsightsManager::Get()->GetStoreClient();
-	const bool bIsRecorderServerRunning = (StoreClient != nullptr); // TODO: StoreClient->IsRecorderServerRunning();
-
-	return bIsRecorderServerRunning ? EVisibility::Visible : EVisibility::Collapsed;
+	if (bEnableAutomaticTesting)
+	{
+		OutParams.Append(TEXT(" -InsightsTest"));
+	}
+	if (bEnableDebugTools)
+	{
+		OutParams.Append(TEXT(" -DebugTools"));
+	}
+	if (bStartProcessWithStompMalloc)
+	{
+		OutParams.Append(TEXT(" -stompmalloc"));
+	}
 }
 
 ////////////////////////////////////////////////////////////////////////////////////////////////////
 
-FReply SStartPageWindow::StartTraceRecorder_OnClicked()
+#undef LOCTEXT_NAMESPACE
+#define LOCTEXT_NAMESPACE "SConnectionWindow"
+
+////////////////////////////////////////////////////////////////////////////////////////////////////
+// SConnectionWindow
+////////////////////////////////////////////////////////////////////////////////////////////////////
+
+SConnectionWindow::SConnectionWindow()
+: NotificationList()
+, ActiveNotifications()
+, ConnectTask()
+, bIsConnecting(false)
+, bIsConnectedSuccessfully(false)
 {
-	//TODO: StoreClient->StartRecorderServer();
-	RefreshTraceList();
-	return FReply::Handled();
 }
 
 ////////////////////////////////////////////////////////////////////////////////////////////////////
 
-FReply SStartPageWindow::StopTraceRecorder_OnClicked()
+SConnectionWindow::~SConnectionWindow()
 {
-	//TODO: StoreClient->StopRecorderServer();
-	RefreshTraceList();
-	return FReply::Handled();
+	if (ConnectTask && !ConnectTask->IsComplete())
+	{
+		FTaskGraphInterface::Get().WaitUntilTaskCompletes(ConnectTask);
+		ConnectTask = nullptr;
+	}
 }
 
 ////////////////////////////////////////////////////////////////////////////////////////////////////
 
-void SStartPageWindow::OpenSettings()
+BEGIN_SLATE_FUNCTION_BUILD_OPTIMIZATION
+
+void SConnectionWindow::Construct(const FArguments& InArgs)
 {
-	MainContentPanel->SetEnabled(false);
-	(*OverlaySettingsSlot)
+	ChildSlot
 	[
-		SNew(SBorder)
-		.BorderImage(FCoreStyle::Get().GetBrush("PopupText.Background"))
-		.Padding(8.0f)
+		SNew(SOverlay)
+
+		// Version
+		+ SOverlay::Slot()
+		.HAlign(HAlign_Right)
+		.VAlign(VAlign_Top)
+		.Padding(0.0f, -16.0f, 4.0f, 0.0f)
 		[
-			SNew(SInsightsSettings)
-			.OnClose(this, &SStartPageWindow::CloseSettings)
-			.SettingPtr(&FInsightsManager::GetSettings())
+			SNew(STextBlock)
+			.Clipping(EWidgetClipping::ClipToBoundsWithoutIntersecting)
+			.Text(LOCTEXT("UnrealInsightsVersion", UNREAL_INSIGHTS_VERSION_STRING_EX))
+			.ColorAndOpacity(FLinearColor(0.15f, 0.15f, 0.15f, 1.0f))
+		]
+
+		+ SOverlay::Slot()
+		.HAlign(HAlign_Fill)
+		.VAlign(VAlign_Fill)
+		.Padding(0.0f, 0.0f, 0.0f, 0.0f)
+		[
+			SNew(SBox)
+			[
+				SNew(SBorder)
+				.HAlign(HAlign_Fill)
+				.VAlign(VAlign_Fill)
+				.Padding(0.0f)
+				.BorderImage(FCoreStyle::Get().GetBrush("WhiteBrush"))
+				.BorderBackgroundColor(FSlateColor(EStyleColor::Panel))
+			]
+		]
+
+		// Overlay slot for the main window area
+		+ SOverlay::Slot()
+		.HAlign(HAlign_Fill)
+		.VAlign(VAlign_Fill)
+		[
+			SAssignNew(MainContentPanel, SVerticalBox)
+
+			+ SVerticalBox::Slot()
+			.HAlign(HAlign_Fill)
+			.AutoHeight()
+			.Padding(3.0f, 3.0f)
+			[
+				ConstructConnectPanel()
+			]
 		]
 	];
 }
 
 ////////////////////////////////////////////////////////////////////////////////////////////////////
 
-void SStartPageWindow::CloseSettings()
+TSharedRef<SWidget> SConnectionWindow::ConstructConnectPanel()
 {
-	// Close the profiler settings by simply replacing widget with a null one.
-	(*OverlaySettingsSlot)
-	[
-		SNullWidget::NullWidget
-	];
-	MainContentPanel->SetEnabled(true);
+	TSharedRef<SWidget> Widget = SNew(SVerticalBox)
+
+		+ SVerticalBox::Slot()
+		.AutoHeight()
+		.HAlign(HAlign_Fill)
+		.Padding(12.0f, 12.0f, 12.0f, 0.0f)
+		[
+			SNew(SHorizontalBox)
+
+			+ SHorizontalBox::Slot()
+			.AutoWidth()
+			.VAlign(VAlign_Center)
+			[
+				SNew(SBox)
+				.MinDesiredWidth(180.0f)
+				.HAlign(HAlign_Right)
+				[
+					SNew(STextBlock)
+					.Text(LOCTEXT("TraceRecorderAddressText", "Trace recorder IP address"))
+				]
+			]
+
+			+ SHorizontalBox::Slot()
+			.FillWidth(1.0f)
+			.VAlign(VAlign_Center)
+			.Padding(6.0f, 0.0f, 0.0f, 0.0f)
+			[
+				SNew(SBox)
+				.MinDesiredWidth(120.0f)
+				[
+					SAssignNew(TraceRecorderAddressTextBox, SEditableTextBox)
+				]
+			]
+		]
+
+		+ SVerticalBox::Slot()
+		.AutoHeight()
+		.HAlign(HAlign_Fill)
+		.Padding(12.0f, 8.0f, 12.0f, 0.0f)
+		[
+			SNew(SHorizontalBox)
+
+			+ SHorizontalBox::Slot()
+			.AutoWidth()
+			.VAlign(VAlign_Center)
+			[
+				SNew(SBox)
+				.MinDesiredWidth(180.0f)
+				.HAlign(HAlign_Right)
+				[
+					SNew(STextBlock)
+					.Text(LOCTEXT("RunningInstanceAddressText", "Running instance IP address"))
+				]
+			]
+
+			+ SHorizontalBox::Slot()
+			.FillWidth(1.0f)
+			.VAlign(VAlign_Center)
+			.Padding(6.0f, 0.0f, 0.0f, 0.0f)
+			[
+				SNew(SBox)
+				.MinDesiredWidth(120.0f)
+				[
+					SAssignNew(RunningInstanceAddressTextBox, SEditableTextBox)
+				]
+			]
+		]
+
+		+ SVerticalBox::Slot()
+		.AutoHeight()
+		.HAlign(HAlign_Fill)
+		.Padding(12.0f, 8.0f, 12.0f, 12.0f)
+		[
+			SNew(SBox)
+			.HAlign(HAlign_Right)
+			[
+				SNew(SButton)
+				.ButtonStyle(&FAppStyle::Get().GetWidgetStyle<FButtonStyle>("PrimaryButton"))
+				.Text(LOCTEXT("Connect", "Connect"))
+				.ToolTipText(LOCTEXT("ConnectToolTip", "Connect the running instance at specified ip with the local trace recorder."))
+				.OnClicked(this, &SConnectionWindow::Connect_OnClicked)
+				.IsEnabled_Lambda([this]() { return !bIsConnecting; })
+			]
+		]
+
+		// Notification area overlay
+		+ SVerticalBox::Slot()
+		.FillHeight(1.0f)
+		.HAlign(HAlign_Right)
+		.VAlign(VAlign_Bottom)
+		.Padding(16.0f)
+		[
+			SAssignNew(NotificationList, SNotificationList)
+		];
+
+	const FText LocalHost = FText::FromString(TEXT("127.0.0.1"));
+
+	TSharedPtr<FInternetAddr> RecorderAddr;
+	if (ISocketSubsystem* Sockets = ISocketSubsystem::Get())
+	{
+		bool bCanBindAll = false;
+		RecorderAddr = Sockets->GetLocalHostAddr(*GLog, bCanBindAll);
+	}
+	if (RecorderAddr.IsValid())
+	{
+		const FString RecorderAddrStr = RecorderAddr->ToString(false);
+		TraceRecorderAddressTextBox->SetText(FText::FromString(RecorderAddrStr));
+	}
+	else
+	{
+		TraceRecorderAddressTextBox->SetText(LocalHost);
+	}
+
+	RunningInstanceAddressTextBox->SetText(LocalHost);
+
+	return Widget;
 }
 
+END_SLATE_FUNCTION_BUILD_OPTIMIZATION
+
 ////////////////////////////////////////////////////////////////////////////////////////////////////
 
-void SStartPageWindow::GetExtraCommandLineParams(FString& OutParams) const
+FReply SConnectionWindow::Connect_OnClicked()
 {
-	if (bEnableAutomaticTesting)
+	FText TraceRecorderAddressText = TraceRecorderAddressTextBox->GetText();
+	if (TraceRecorderAddressText.IsEmptyOrWhitespace())
 	{
-		OutParams.Append(TEXT(" -InsightsTest"));
+		// nothing to do
+		return FReply::Handled();
 	}
-	if (bEnableDebugTools)
+	const FString& TraceRecorderAddressStr = TraceRecorderAddressText.ToString();
+
+	FText RunningInstanceAddressText = RunningInstanceAddressTextBox->GetText();
+	if (RunningInstanceAddressText.IsEmptyOrWhitespace())
 	{
-		OutParams.Append(TEXT(" -DebugTools"));
+		// nothing to do
+		return FReply::Handled();
 	}
-	if (bStartProcessWithStompMalloc)
+	const FString& RunningInstanceAddressStr = RunningInstanceAddressText.ToString();
+
+	FGraphEventArray Prerequisites;
+	FGraphEventArray* PrerequisitesPtr = nullptr;
+	if (ConnectTask.IsValid())
 	{
-		OutParams.Append(TEXT(" -stompmalloc"));
+		Prerequisites.Add(ConnectTask);
+		PrerequisitesPtr = &Prerequisites;
 	}
+
+	FGraphEventRef PreConnectTask = FFunctionGraphTask::CreateAndDispatchWhenReady(
+		[this, TraceRecorderAddressStr, RunningInstanceAddressStr]
+		{
+			bIsConnecting = true;
+
+			UE_LOG(TraceInsights, Log, TEXT("[Connection] Try connecting to \"%s\"..."), *RunningInstanceAddressStr);
+
+			UE::Trace::FControlClient ControlClient;
+			if (ControlClient.Connect(*RunningInstanceAddressStr))
+			{
+				UE_LOG(TraceInsights, Log, TEXT("[Connection] SendSendTo(\"%s\")..."), *TraceRecorderAddressStr);
+				ControlClient.SendSendTo(*TraceRecorderAddressStr);
+				bIsConnectedSuccessfully = true;
+			}
+			else
+			{
+				bIsConnectedSuccessfully = false;
+			}
+		},
+		TStatId{}, PrerequisitesPtr, ENamedThreads::AnyBackgroundThreadNormalTask);
+
+	ConnectTask = FFunctionGraphTask::CreateAndDispatchWhenReady(
+		[this, RunningInstanceAddressStr]
+		{
+			if (bIsConnectedSuccessfully)
+			{
+				UE_LOG(TraceInsights, Log, TEXT("[Connection] Successfully connected."));
+
+				FNotificationInfo NotificationInfo(FText::Format(LOCTEXT("ConnectSuccess", "Successfully connected to \"{0}\"!"), FText::FromString(RunningInstanceAddressStr)));
+				NotificationInfo.bFireAndForget = false;
+				NotificationInfo.bUseLargeFont = false;
+				NotificationInfo.bUseSuccessFailIcons = true;
+				NotificationInfo.ExpireDuration = 10.0f;
+				SNotificationItemWeak NotificationItem = NotificationList->AddNotification(NotificationInfo);
+				NotificationItem.Pin()->SetCompletionState(SNotificationItem::CS_Success);
+				NotificationItem.Pin()->ExpireAndFadeout();
+				ActiveNotifications.Add(TEXT("ConnectSuccess"), NotificationItem);
+			}
+			else
+			{
+				UE_LOG(TraceInsights, Warning, TEXT("[Connection] Failed to connect to \"%s\"!"), *RunningInstanceAddressStr);
+
+				FNotificationInfo NotificationInfo(FText::Format(LOCTEXT("ConnectFailed", "Failed to connect to \"{0}\"!"), FText::FromString(RunningInstanceAddressStr)));
+				NotificationInfo.bFireAndForget = false;
+				NotificationInfo.bUseLargeFont = false;
+				NotificationInfo.bUseSuccessFailIcons = true;
+				NotificationInfo.ExpireDuration = 10.0f;
+				SNotificationItemWeak NotificationItem = NotificationList->AddNotification(NotificationInfo);
+				NotificationItem.Pin()->SetCompletionState(SNotificationItem::CS_Fail);
+				NotificationItem.Pin()->ExpireAndFadeout();
+				ActiveNotifications.Add(TEXT("ConnectFailed"), NotificationItem);
+			}
+
+			bIsConnecting = false;
+		},
+		TStatId{}, PreConnectTask, ENamedThreads::GameThread);
+
+	return FReply::Handled();
 }
 
 ////////////////////////////////////////////////////////////////////////////////////////////////////
 
 #undef LOCTEXT_NAMESPACE
+#define LOCTEXT_NAMESPACE "SLauncherWindow"
+
+////////////////////////////////////////////////////////////////////////////////////////////////////
+// SLauncherWindow
+////////////////////////////////////////////////////////////////////////////////////////////////////
+
+SLauncherWindow::SLauncherWindow()
+{
+}
+
+////////////////////////////////////////////////////////////////////////////////////////////////////
+
+SLauncherWindow::~SLauncherWindow()
+{
+}
+
+////////////////////////////////////////////////////////////////////////////////////////////////////
+
+BEGIN_SLATE_FUNCTION_BUILD_OPTIMIZATION
+void SLauncherWindow::Construct(const FArguments& InArgs)
+{
+	ChildSlot
+	[
+		SNew(STextBlock)
+		.Text(LOCTEXT("Launcher", "Launcher"))
+	];
+}
+END_SLATE_FUNCTION_BUILD_OPTIMIZATION
+
+////////////////////////////////////////////////////////////////////////////////////////////////////
+
+#undef LOCTEXT_NAMESPACE
diff --git a/Engine/Source/Developer/TraceInsights/Private/Insights/Widgets/SStartPageWindow.h b/Engine/Source/Developer/TraceInsights/Private/Insights/Widgets/SStartPageWindow.h
--- a/Engine/Source/Developer/TraceInsights/Private/Insights/Widgets/SStartPageWindow.h
+++ b/Engine/Source/Developer/TraceInsights/Private/Insights/Widgets/SStartPageWindow.h
@@ -2,30 +2,22 @@
 
 #pragma once
 
-#include "CoreMinimal.h"
-#include "Framework/Docking/TabManager.h"
-#include "GenericPlatform/GenericPlatformMisc.h"
-#include "Input/Reply.h"
-#include "Layout/Visibility.h"
-#include "Misc/Guid.h"
-#include "SlateFwd.h"
-#include "Trace/StoreClient.h"
-#include "TraceServices/ModuleService.h"
+#include "Async/TaskGraphInterfaces.h"
+#include "Styling/SlateTypes.h"
 #include "Widgets/DeclarativeSyntaxSupport.h"
-#include "Widgets/Docking/SDockTab.h"
-#include "Widgets/Layout/SSplitter.h"
 #include "Widgets/SCompoundWidget.h"
 #include "Widgets/SOverlay.h"
 #include "Widgets/Views/SListView.h"
 
-// Insights
-#include "Insights/InsightsManager.h"
-
 ////////////////////////////////////////////////////////////////////////////////////////////////////
 
 class FActiveTimerHandle;
-class SVerticalBox;
+class ITableRow;
 class SEditableTextBox;
+class SNotificationList;
+class SSearchBox;
+class STableViewBase;
+class SVerticalBox;
 
 namespace Insights
 {
@@ -86,17 +78,19 @@ struct FTraceViewModel
 	}
 };
 
-/** Implements the Start Page window. */
-class SStartPageWindow : public SCompoundWidget
+////////////////////////////////////////////////////////////////////////////////////////////////////
+
+/** Implements the Trace Store window. */
+class STraceStoreWindow : public SCompoundWidget
 {
 public:
 	/** Default constructor. */
-	SStartPageWindow();
+	STraceStoreWindow();
 
 	/** Virtual destructor. */
-	virtual ~SStartPageWindow();
+	virtual ~STraceStoreWindow();
 
-	SLATE_BEGIN_ARGS(SStartPageWindow) {}
+	SLATE_BEGIN_ARGS(STraceStoreWindow) {}
 	SLATE_END_ARGS()
 
 	/** Constructs this widget. */
@@ -121,8 +115,6 @@ private:
 	TSharedRef<SWidget> ConstructLoadPanel();
 	TSharedRef<SWidget> ConstructTraceStoreDirectoryPanel();
 	TSharedRef<SWidget> ConstructAutoStartPanel();
-	TSharedRef<SWidget> ConstructRecorderPanel();
-	TSharedRef<SWidget> ConstructConnectPanel();
 
 	/** Generate a new row for the Traces list view. */
 	TSharedRef<ITableRow> TraceList_OnGenerateRow(TSharedPtr<FTraceViewModel> InTrace, const TSharedRef<STableViewBase>& OwnerTable);
@@ -176,20 +168,6 @@ private:
 	FText GetTraceStoreDirectory() const;
 	FReply ExploreTraceStoreDirectory_OnClicked();
 
-	//////////////////////////////////////////////////
-	// Recorder
-
-	FText GetRecorderStatusText() const;
-	EVisibility StartTraceRecorder_Visibility() const;
-	EVisibility StopTraceRecorder_Visibility() const;
-	FReply StartTraceRecorder_OnClicked();
-	FReply StopTraceRecorder_OnClicked();
-
-	//////////////////////////////////////////////////
-	// New Connection
-
-	FReply Connect_OnClicked();
-
 	//////////////////////////////////////////////////
 
 	/** Updates the amount of time the profiler has been active. */
@@ -249,7 +227,7 @@ private:
 	 */
 	virtual FReply OnDragOver(const FGeometry& MyGeometry, const FDragDropEvent& DragDropEvent)  override;
 
-public:
+private:
 	/** Widget for the non-intrusive notifications. */
 	TSharedPtr<SNotificationList> NotificationList;
 
@@ -262,11 +240,9 @@ public:
 	/** The number of seconds the profiler has been active */
 	float DurationActive;
 
-private:
 	/** The handle to the active update duration tick */
 	TWeakPtr<FActiveTimerHandle> ActiveTimerHandle;
 
-	/** Holds all widgets for the profiler window like menu bar, toolbar and tabs. */
 	TSharedPtr<SVerticalBox> MainContentPanel;
 
 	int32 LiveSessionCount;
@@ -288,8 +264,6 @@ private:
 	TSharedPtr<SListView<TSharedPtr<FTraceViewModel>>> TraceListView;
 	TSharedPtr<FTraceViewModel> SelectedTrace;
 
-	TSharedPtr<SEditableTextBox> HostTextBox;
-
 	FString SplashScreenOverlayTraceFile;
 	float SplashScreenOverlayFadeTime;
 
@@ -297,3 +271,63 @@ private:
 	bool bEnableDebugTools = false;
 	bool bStartProcessWithStompMalloc = false;
 };
+
+////////////////////////////////////////////////////////////////////////////////////////////////////
+
+/** Implements the Connection window. */
+class SConnectionWindow : public SCompoundWidget
+{
+public:
+	/** Default constructor. */
+	SConnectionWindow();
+
+	/** Virtual destructor. */
+	virtual ~SConnectionWindow();
+
+	SLATE_BEGIN_ARGS(SConnectionWindow) {}
+	SLATE_END_ARGS()
+
+	/** Constructs this widget. */
+	void Construct(const FArguments& InArgs);
+
+private:
+	TSharedRef<SWidget> ConstructConnectPanel();
+	FReply Connect_OnClicked();
+
+private:
+	TSharedPtr<SVerticalBox> MainContentPanel;
+	TSharedPtr<SEditableTextBox> TraceRecorderAddressTextBox;
+	TSharedPtr<SEditableTextBox> RunningInstanceAddressTextBox;
+
+	/** Widget for the non-intrusive notifications. */
+	TSharedPtr<SNotificationList> NotificationList;
+
+	/** Holds all active and visible notifications, stored as FGuid -> SNotificationItemWeak. */
+	TMap<FString, SNotificationItemWeak> ActiveNotifications;
+
+	FGraphEventRef ConnectTask;
+
+	std::atomic<bool> bIsConnecting;
+	std::atomic<bool> bIsConnectedSuccessfully;
+};
+
+////////////////////////////////////////////////////////////////////////////////////////////////////
+
+/** Implements the Launcher window. */
+class SLauncherWindow : public SCompoundWidget
+{
+public:
+	/** Default constructor. */
+	SLauncherWindow();
+
+	/** Virtual destructor. */
+	virtual ~SLauncherWindow();
+
+	SLATE_BEGIN_ARGS(SLauncherWindow) {}
+	SLATE_END_ARGS()
+
+	/** Constructs this widget. */
+	void Construct(const FArguments& InArgs);
+};
+
+////////////////////////////////////////////////////////////////////////////////////////////////////
diff --git a/Engine/Source/Developer/TraceInsights/Public/Insights/IUnrealInsightsModule.h b/Engine/Source/Developer/TraceInsights/Public/Insights/IUnrealInsightsModule.h
--- a/Engine/Source/Developer/TraceInsights/Public/Insights/IUnrealInsightsModule.h
+++ b/Engine/Source/Developer/TraceInsights/Public/Insights/IUnrealInsightsModule.h
@@ -29,7 +29,10 @@ namespace TraceServices
 /** Major tab IDs for Insights tools */
 struct TRACEINSIGHTS_API FInsightsManagerTabs
 {
-	static const FName StartPageTabId;
+	static const FName StartPageTabId; // DEPRECATED
+	static const FName TraceStoreTabId;
+	static const FName ConnectionTabId;
+	static const FName LauncherTabId;
 	static const FName SessionInfoTabId;
 	static const FName TimingProfilerTabId;
 	static const FName LoadingProfilerTabId;
