[Insights] Memory: Updated icons / UX for the LLM Tags view.

diff --git a/Engine/Source/Developer/TraceInsights/Private/Insights/MemoryProfiler/Widgets/SMemTagTreeView.cpp b/Engine/Source/Developer/TraceInsights/Private/Insights/MemoryProfiler/Widgets/SMemTagTreeView.cpp
--- a/Engine/Source/Developer/TraceInsights/Private/Insights/MemoryProfiler/Widgets/SMemTagTreeView.cpp
+++ b/Engine/Source/Developer/TraceInsights/Private/Insights/MemoryProfiler/Widgets/SMemTagTreeView.cpp
@@ -12,6 +12,7 @@
 #include "Widgets/Layout/SScrollBox.h"
 #include "Widgets/Input/SButton.h"
 #include "Widgets/Input/SSearchBox.h"
+#include "Widgets/Input/SSegmentedControl.h"
 #include "Widgets/Images/SImage.h"
 #include "Widgets/SToolTip.h"
 #include "Widgets/Views/STableViewBase.h"
@@ -306,20 +307,22 @@ TSharedRef<SWidget> SMemTagTreeView::ConstructTracksMiniToolbar()
 
 	+ SHorizontalBox::Slot()
 	.VAlign(VAlign_Center)
-	.Padding(FMargin(2.0f, 0.0f, 2.0f, 0.0f))
+	.Padding(FMargin(0.0f, 0.0f, 0.0f, 0.0f))
 	.AutoWidth()
 	[
 		SNew(SButton)
 		.ToolTipText(LOCTEXT("HideAll_ToolTip", "Remove memory graph tracks for all LLM tags."))
-		.ContentPadding(FMargin(-14.0f, 0.0f, -17.0f, -1.0f))
 		.OnClicked(this, &SMemTagTreeView::HideAllTracks_OnClicked)
-		.Content()
+		.ButtonStyle(&FAppStyle::Get().GetWidgetStyle<FButtonStyle>("SimpleButton"))
+		.ContentPadding(FMargin(2.0f, 2.0f, 2.0f, 2.0f))
 		[
 			SNew(SBox)
 			.HAlign(HAlign_Center)
+			.VAlign(VAlign_Center)
 			[
 				SNew(SImage)
-				.ColorAndOpacity(FSlateColor(FLinearColor(0.5f, 0.0f, 0.0f, 1.0f)))
+				//.ColorAndOpacity(FSlateColor::UseForeground())
+				.ColorAndOpacity(FSlateColor(FLinearColor(1.0f, 0.5f, 0.5f, 1.0f)))
 				.Image(FInsightsStyle::Get().GetBrush("Mem.Remove.Small"))
 			]
 		]
@@ -327,20 +330,22 @@ TSharedRef<SWidget> SMemTagTreeView::ConstructTracksMiniToolbar()
 
 	+ SHorizontalBox::Slot()
 	.VAlign(VAlign_Center)
-	.Padding(FMargin(2.0f, 0.0f, 2.0f, 0.0f))
+	.Padding(FMargin(0.0f, 0.0f, 0.0f, 0.0f))
 	.AutoWidth()
 	[
 		SNew(SButton)
 		.ToolTipText(LOCTEXT("ShowAll_ToolTip", "Create memory graph tracks for all visible (filtered) LLM tags."))
-		.ContentPadding(FMargin(-14.0f, 0.0f, -17.0f, -1.0f))
 		.OnClicked(this, &SMemTagTreeView::ShowAllTracks_OnClicked)
-		.Content()
+		.ButtonStyle(&FAppStyle::Get().GetWidgetStyle<FButtonStyle>("SimpleButton"))
+		.ContentPadding(FMargin(2.0f, 2.0f, 2.0f, 2.0f))
 		[
 			SNew(SBox)
 			.HAlign(HAlign_Center)
+			.VAlign(VAlign_Center)
 			[
 				SNew(SImage)
-				.ColorAndOpacity(FSlateColor(FLinearColor(0.0f, 0.5f, 0.0f, 1.0f)))
+				//.ColorAndOpacity(FSlateColor::UseForeground())
+				.ColorAndOpacity(FSlateColor(FLinearColor(0.5f, 1.0f, 0.5f, 1.0f)))
 				.Image(FInsightsStyle::Get().GetBrush("Mem.Add.Small"))
 			]
 		]
@@ -348,79 +353,71 @@ TSharedRef<SWidget> SMemTagTreeView::ConstructTracksMiniToolbar()
 
 	+ SHorizontalBox::Slot()
 	.VAlign(VAlign_Center)
-	.Padding(2.0f)
+	.Padding(FMargin(2.0f, 0.0f, 2.0f, 0.0f))
 	.AutoWidth()
 	[
 		SNew(SButton)
 		.ToolTipText(LOCTEXT("LoadReportXML_ToolTip", "Load LLMReportTypes.xml"))
-		.ContentPadding(FMargin(-14.0f, 2.0f, -14.0f, 2.0f))
 		.OnClicked(this, &SMemTagTreeView::LoadReportXML_OnClicked)
-		.Content()
+		.ButtonStyle(&FAppStyle::Get().GetWidgetStyle<FButtonStyle>("SimpleButton"))
+		.ContentPadding(FMargin(4.0f, 2.0f, 4.0f, 2.0f))
 		[
-			SNew(SBox)
-			.HAlign(HAlign_Center)
-			[
-				SNew(STextBlock)
-				.Text(LOCTEXT("LoadReportXML_Text", "Load XML..."))
-			]
-		]
-	]
-
-	+ SHorizontalBox::Slot()
-	.VAlign(VAlign_Center)
-	.Padding(2.0f)
-	[
-		SNew(SButton)
-		.ToolTipText(LOCTEXT("SmallHeight_ToolTip", "Change height of LLM Tag Graph tracks to Small."))
-		.ContentPadding(FMargin(-18.0f, 2.0f, -18.0f, 2.0f))
-		.OnClicked(this, &SMemTagTreeView::AllTracksSmallHeight_OnClicked)
-		.Content()
-		[
-			SNew(SBox)
-			.HAlign(HAlign_Center)
+			SNew(SHorizontalBox)
+			+ SHorizontalBox::Slot()
+			.AutoWidth()
+			.VAlign(VAlign_Center)
 			[
-				SNew(STextBlock)
-				.Text(LOCTEXT("SmallHeight_Text", "\u2195S"))
+				SNew(SImage)
+				.ColorAndOpacity(FSlateColor::UseForeground())
+				.Image(FInsightsStyle::Get().GetBrush("Mem.LoadXML.Small"))
 			]
-		]
-	]
-
-	+ SHorizontalBox::Slot()
-	.VAlign(VAlign_Center)
-	.Padding(2.0f)
-	[
-		SNew(SButton)
-		.ToolTipText(LOCTEXT("MediumHeight_ToolTip", "Change height of LLM Tag Graph tracks to Medium."))
-		.ContentPadding(FMargin(-18.0f, 2.0f, -18.0f, 2.0f))
-		.OnClicked(this, &SMemTagTreeView::AllTracksMediumHeight_OnClicked)
-		.Content()
-		[
-			SNew(SBox)
-			.HAlign(HAlign_Center)
+			+ SHorizontalBox::Slot()
+			.VAlign(VAlign_Center)
+			.Padding(4.0f, 0.0f, 0.0f, 0.0f)
+			.AutoWidth()
 			[
 				SNew(STextBlock)
-				.Text(LOCTEXT("MediumHeight_Text", "\u2195M"))
+				.TextStyle(&FAppStyle::Get().GetWidgetStyle< FTextBlockStyle >(TEXT("ButtonText")))
+				.Justification(ETextJustify::Center)
+				.Text(LOCTEXT("LoadReportXML_Text", "Load XML..."))
 			]
 		]
 	]
 
 	+ SHorizontalBox::Slot()
 	.VAlign(VAlign_Center)
-	.Padding(2.0f)
+	.HAlign(HAlign_Right)
+	.Padding(FMargin(0.0f, 0.0f, 0.0f, 0.0f))
 	[
-		SNew(SButton)
-		.ToolTipText(LOCTEXT("LargeHeight_ToolTip", "Change height of LLM Tag Graph tracks to Large."))
-		.ContentPadding(FMargin(-18.0f, 2.0f, -18.0f, 2.0f))
-		.OnClicked(this, &SMemTagTreeView::AllTracksLargeHeight_OnClicked)
-		.Content()
-		[
-			SNew(SBox)
-			.HAlign(HAlign_Center)
-			[
-				SNew(STextBlock)
-				.Text(LOCTEXT("LargeHeight_Text", "\u2195L"))
-			]
-		]
+		SNew(SSegmentedControl<uint32>)
+		.OnValueChanged_Lambda([this](uint32 InValue)
+			{
+				TSharedPtr<SMemoryProfilerWindow> ProfilerWindow = ProfilerWindowWeakPtr.Pin();
+				if (ProfilerWindow.IsValid())
+				{
+					FMemorySharedState& SharedState = ProfilerWindow->GetSharedState();
+					SharedState.SetTrackHeightMode((EMemoryTrackHeightMode)InValue);
+				}
+			})
+		.Value_Lambda([this]
+			{
+				TSharedPtr<SMemoryProfilerWindow> ProfilerWindow = ProfilerWindowWeakPtr.Pin();
+				if (ProfilerWindow.IsValid())
+				{
+					FMemorySharedState& SharedState = ProfilerWindow->GetSharedState();
+					return (uint32)SharedState.GetTrackHeightMode();
+				}
+				return (uint32)EMemoryTrackHeightMode::Medium;
+			})
+		+ SSegmentedControl<uint32>::Slot((uint32)EMemoryTrackHeightMode::Small)
+		.Text(LOCTEXT("SmallHeight_Text", "\u2195S"))
+		.ToolTip(LOCTEXT("SmallHeight_ToolTip", "Change height of LLM Tag Graph tracks to Small."))
+		+ SSegmentedControl<uint32>::Slot((uint32)EMemoryTrackHeightMode::Medium)
+		.Text(LOCTEXT("MediumHeight_Text", "\u2195M"))
+		.ToolTip(LOCTEXT("MediumHeight_ToolTip", "Change height of LLM Tag Graph tracks to Medium."))
+		+ SSegmentedControl<uint32>::Slot((uint32)EMemoryTrackHeightMode::Large)
+		.Text(LOCTEXT("LargeHeight_Text", "\u2195L"))
+		.ToolTip(LOCTEXT("LargeHeight_ToolTip", "Change height of LLM Tag Graph tracks to Large."))
 	];
 
 	return Widget;
@@ -2263,51 +2260,6 @@ FReply SMemTagTreeView::HideAllTracks_OnClicked()
 	return FReply::Handled();
 }
 
-////////////////////////////////////////////////////////////////////////////////////////////////////
-
-FReply SMemTagTreeView::AllTracksSmallHeight_OnClicked()
-{
-	TSharedPtr<SMemoryProfilerWindow> ProfilerWindow = ProfilerWindowWeakPtr.Pin();
-
-	if (ProfilerWindow.IsValid())
-	{
-		FMemorySharedState& SharedState = ProfilerWindow->GetSharedState();
-		SharedState.SetTrackHeightMode(EMemoryTrackHeightMode::Small);
-	}
-
-	return FReply::Handled();
-}
-
-////////////////////////////////////////////////////////////////////////////////////////////////////
-
-FReply SMemTagTreeView::AllTracksMediumHeight_OnClicked()
-{
-	TSharedPtr<SMemoryProfilerWindow> ProfilerWindow = ProfilerWindowWeakPtr.Pin();
-
-	if (ProfilerWindow.IsValid())
-	{
-		FMemorySharedState& SharedState = ProfilerWindow->GetSharedState();
-		SharedState.SetTrackHeightMode(EMemoryTrackHeightMode::Medium);
-	}
-
-	return FReply::Handled();
-}
-
-////////////////////////////////////////////////////////////////////////////////////////////////////
-
-FReply SMemTagTreeView::AllTracksLargeHeight_OnClicked()
-{
-	TSharedPtr<SMemoryProfilerWindow> ProfilerWindow = ProfilerWindowWeakPtr.Pin();
-
-	if (ProfilerWindow.IsValid())
-	{
-		FMemorySharedState& SharedState = ProfilerWindow->GetSharedState();
-		SharedState.SetTrackHeightMode(EMemoryTrackHeightMode::Large);
-	}
-
-	return FReply::Handled();
-}
-
 ////////////////////////////////////////////////////////////////////////////////////////////////////
 // Create memory graph tracks for selected LLM tag(s)
 ////////////////////////////////////////////////////////////////////////////////////////////////////
diff --git a/Engine/Source/Developer/TraceInsights/Private/Insights/MemoryProfiler/Widgets/SMemTagTreeView.h b/Engine/Source/Developer/TraceInsights/Private/Insights/MemoryProfiler/Widgets/SMemTagTreeView.h
--- a/Engine/Source/Developer/TraceInsights/Private/Insights/MemoryProfiler/Widgets/SMemTagTreeView.h
+++ b/Engine/Source/Developer/TraceInsights/Private/Insights/MemoryProfiler/Widgets/SMemTagTreeView.h
@@ -264,10 +264,6 @@ private:
 
 	FReply LoadReportXML_OnClicked();
 
-	FReply AllTracksSmallHeight_OnClicked();
-	FReply AllTracksMediumHeight_OnClicked();
-	FReply AllTracksLargeHeight_OnClicked();
-
 	////////////////////////////////////////////////////////////////////////////////////////////////////
 
 	// Create memory graph tracks for selected LLM tag(s)
