[Insights] Removed the registration of tab spawner for Trace Store, Connection and Launcher when in editor.

diff --git a/Engine/Source/Developer/TraceInsights/Private/Insights/InsightsManager.cpp b/Engine/Source/Developer/TraceInsights/Private/Insights/InsightsManager.cpp
--- a/Engine/Source/Developer/TraceInsights/Private/Insights/InsightsManager.cpp
+++ b/Engine/Source/Developer/TraceInsights/Private/Insights/InsightsManager.cpp
@@ -230,6 +230,7 @@ void FInsightsManager::BindCommands()
 
 void FInsightsManager::RegisterMajorTabs(IUnrealInsightsModule& InsightsModule)
 {
+#if !WITH_EDITOR
 	const FInsightsMajorTabConfig& TraceStoreConfig = InsightsModule.FindMajorTabConfig(FInsightsManagerTabs::TraceStoreTabId);
 	if (TraceStoreConfig.bIsAvailable)
 	{
@@ -271,6 +272,7 @@ void FInsightsManager::RegisterMajorTabs(IUnrealInsightsModule& InsightsModule)
 		TSharedRef<FWorkspaceItem> Group = LauncherConfig.WorkspaceGroup.IsValid() ? LauncherConfig.WorkspaceGroup.ToSharedRef() : WorkspaceMenu::GetMenuStructure().GetDeveloperToolsProfilingCategory();
 		TabSpawnerEntry.SetGroup(Group);
 	}
+#endif // !WITH_EDITOR
 
 	const FInsightsMajorTabConfig& SessionInfoConfig = InsightsModule.FindMajorTabConfig(FInsightsManagerTabs::SessionInfoTabId);
 	if (SessionInfoConfig.bIsAvailable)
