UE5_MAIN:  LLM error due to missing LLM_DEFINE_TAG(Insights), leading to cook failure.

#rnx

diff --git a/Engine/Source/Developer/TraceInsights/Private/Insights/InsightsStyle.cpp b/Engine/Source/Developer/TraceInsights/Private/Insights/InsightsStyle.cpp
--- a/Engine/Source/Developer/TraceInsights/Private/Insights/InsightsStyle.cpp
+++ b/Engine/Source/Developer/TraceInsights/Private/Insights/InsightsStyle.cpp
@@ -7,6 +7,8 @@
 #include "Styling/SlateStyleRegistry.h"
 #include "Styling/ToolBarStyle.h"
 
+LLM_DEFINE_TAG(Insights);
+
 ////////////////////////////////////////////////////////////////////////////////////////////////////
 // FInsightsStyle
 ////////////////////////////////////////////////////////////////////////////////////////////////////
