[Insights]
  - Timing view: Added support for "bool" metadata for cpu timing events.
  - Timing view: Changed cpu metadata for large uint64 values to be displayed as hex.

diff --git a/Engine/Source/Developer/TraceInsights/Private/Insights/ViewModels/ThreadTimingTrack.cpp b/Engine/Source/Developer/TraceInsights/Private/Insights/ViewModels/ThreadTimingTrack.cpp
--- a/Engine/Source/Developer/TraceInsights/Private/Insights/ViewModels/ThreadTimingTrack.cpp
+++ b/Engine/Source/Developer/TraceInsights/Private/Insights/ViewModels/ThreadTimingTrack.cpp
@@ -69,7 +69,15 @@ static void AppendMetadataToTooltip(FTooltipDrawState& Tooltip, TArrayView<const
 		case ECborCode::Uint:
 			{
 				uint64 Value = Context.AsUInt();
-				FString ValueStr = FString::Printf(TEXT("%llu"), Value);
+				FString ValueStr;
+				if (Value > 999'999'999ULL)
+				{
+					ValueStr = FString::Printf(TEXT("0x%llX"), Value);
+				}
+				else
+				{
+					ValueStr = FString::Printf(TEXT("%llu"), Value);
+				}
 				Tooltip.AddNameValueTextLine(Key, ValueStr);
 				continue;
 			}
@@ -106,6 +114,18 @@ static void AppendMetadataToTooltip(FTooltipDrawState& Tooltip, TArrayView<const
 			continue;
 		}
 
+		if (Context.RawCode() == (ECborCode::Prim | ECborCode::False))
+		{
+			Tooltip.AddNameValueTextLine(Key, FString(TEXT("false")));
+			continue;
+		}
+
+		if (Context.RawCode() == (ECborCode::Prim | ECborCode::True))
+		{
+			Tooltip.AddNameValueTextLine(Key, FString(TEXT("true")));
+			continue;
+		}
+
 		if (Context.IsFiniteContainer())
 		{
 			CborReader.SkipContainer(ECborCode::Array);
@@ -161,7 +181,14 @@ static void AppendMetadataToString(FString& Str, TArrayView<const uint8>& Metada
 		case ECborCode::Uint:
 			{
 				uint64 Value = Context.AsUInt();
-				Str += FString::Printf(TEXT("%llu"), Value);
+				if (Value > 999'999'999ULL)
+				{
+					Str += FString::Printf(TEXT("0x%llX"), Value);
+				}
+				else
+				{
+					Str += FString::Printf(TEXT("%llu"), Value);
+				}
 				continue;
 			}
 
@@ -192,6 +219,18 @@ static void AppendMetadataToString(FString& Str, TArrayView<const uint8>& Metada
 			continue;
 		}
 
+		if (Context.RawCode() == (ECborCode::Prim | ECborCode::False))
+		{
+			Str += TEXT("false");
+			continue;
+		}
+
+		if (Context.RawCode() == (ECborCode::Prim | ECborCode::True))
+		{
+			Str += TEXT("true");
+			continue;
+		}
+
 		if (Context.IsFiniteContainer())
 		{
 			CborReader.SkipContainer(ECborCode::Array);
