[Insights]
  - Fixed session info and application title after drag&drop of a new utrace file over an instance of Unreal Insights (viewer mode). Also added drag&drop support for the Session Info tab.

diff --git a/Engine/Source/Developer/TraceInsights/Private/Insights/InsightsManager.cpp b/Engine/Source/Developer/TraceInsights/Private/Insights/InsightsManager.cpp
--- a/Engine/Source/Developer/TraceInsights/Private/Insights/InsightsManager.cpp
+++ b/Engine/Source/Developer/TraceInsights/Private/Insights/InsightsManager.cpp
@@ -800,6 +800,80 @@ void FInsightsManager::LoadTraceFile(const FString& InTraceFilename, bool InAuto
 
 ////////////////////////////////////////////////////////////////////////////////////////////////////
 
+bool FInsightsManager::OnDragOver(const FDragDropEvent& DragDropEvent)
+{
+	TSharedPtr<FExternalDragOperation> DragDropOp = DragDropEvent.GetOperationAs<FExternalDragOperation>();
+	if (DragDropOp.IsValid())
+	{
+		if (DragDropOp->HasFiles())
+		{
+			const TArray<FString>& Files = DragDropOp->GetFiles();
+			if (Files.Num() == 1)
+			{
+				const FString DraggedFileExtension = FPaths::GetExtension(Files[0], true);
+				if (DraggedFileExtension == TEXT(".utrace"))
+				{
+					return true;
+				}
+			}
+		}
+	}
+
+	return false;
+}
+
+////////////////////////////////////////////////////////////////////////////////////////////////////
+
+bool FInsightsManager::OnDrop(const FDragDropEvent& DragDropEvent)
+{
+	TSharedPtr<FExternalDragOperation> DragDropOp = DragDropEvent.GetOperationAs<FExternalDragOperation>();
+	if (DragDropOp.IsValid())
+	{
+		if (DragDropOp->HasFiles())
+		{
+			// For now, only allow a single file.
+			const TArray<FString>& Files = DragDropOp->GetFiles();
+			if (Files.Num() == 1)
+			{
+				const FString DraggedFileExtension = FPaths::GetExtension(Files[0], true);
+				if (DraggedFileExtension == TEXT(".utrace"))
+				{
+					LoadTraceFile(Files[0]);
+					UpdateAppTitle();
+					return true;
+				}
+			}
+		}
+	}
+
+	return false;
+}
+
+////////////////////////////////////////////////////////////////////////////////////////////////////
+
+void FInsightsManager::UpdateAppTitle()
+{
+#if !WITH_EDITOR
+	TSharedPtr<SWindow> RootWindow = FGlobalTabmanager::Get()->GetRootWindow();
+	if (RootWindow)
+	{
+		if (CurrentTraceFilename.IsEmpty())
+		{
+			const FText AppTitle = LOCTEXT("UnrealInsightsAppName", "Unreal Insights");
+			RootWindow->SetTitle(AppTitle);
+		}
+		else
+		{
+			const FString SessionName = FPaths::GetBaseFilename(CurrentTraceFilename);
+			const FText AppTitle = FText::Format(LOCTEXT("UnrealInsightsAppNameFmt", "{0} - Unreal Insights"), FText::FromString(SessionName));
+			RootWindow->SetTitle(AppTitle);
+		}
+	}
+#endif
+}
+
+////////////////////////////////////////////////////////////////////////////////////////////////////
+
 void FInsightsManager::OpenSettings()
 {
 	TSharedPtr<STraceStoreWindow> Wnd = GetTraceStoreWindow();
diff --git a/Engine/Source/Developer/TraceInsights/Private/Insights/InsightsManager.h b/Engine/Source/Developer/TraceInsights/Private/Insights/InsightsManager.h
--- a/Engine/Source/Developer/TraceInsights/Private/Insights/InsightsManager.h
+++ b/Engine/Source/Developer/TraceInsights/Private/Insights/InsightsManager.h
@@ -5,6 +5,7 @@
 #include "Containers/Ticker.h"
 #include "CoreMinimal.h"
 #include "Framework/Commands/UICommandList.h"
+#include "Input/DragAndDrop.h"
 #include "TraceServices/AnalysisService.h"
 #include "TraceServices/ModuleService.h"
 
@@ -240,6 +241,11 @@ public:
 	 */
 	void LoadTraceFile(const FString& TraceFilename, bool InAutoQuit = false);
 
+	bool OnDragOver(const FDragDropEvent& DragDropEvent);
+	bool OnDrop(const FDragDropEvent& DragDropEvent);
+
+	void UpdateAppTitle();
+
 	/** Opens the Settings dialog. */
 	void OpenSettings();
 
diff --git a/Engine/Source/Developer/TraceInsights/Private/Insights/LoadingProfiler/Widgets/SLoadingProfilerWindow.cpp b/Engine/Source/Developer/TraceInsights/Private/Insights/LoadingProfiler/Widgets/SLoadingProfilerWindow.cpp
--- a/Engine/Source/Developer/TraceInsights/Private/Insights/LoadingProfiler/Widgets/SLoadingProfilerWindow.cpp
+++ b/Engine/Source/Developer/TraceInsights/Private/Insights/LoadingProfiler/Widgets/SLoadingProfilerWindow.cpp
@@ -750,51 +750,24 @@ FReply SLoadingProfilerWindow::OnKeyDown(const FGeometry& MyGeometry, const FKey
 
 FReply SLoadingProfilerWindow::OnDragOver(const FGeometry& MyGeometry, const FDragDropEvent& DragDropEvent)
 {
-	TSharedPtr<FExternalDragOperation> DragDropOp = DragDropEvent.GetOperationAs<FExternalDragOperation>();
-	if (DragDropOp.IsValid())
+	if (FInsightsManager::Get()->OnDragOver(DragDropEvent))
 	{
-		if (DragDropOp->HasFiles())
-		{
-			const TArray<FString>& Files = DragDropOp->GetFiles();
-			if (Files.Num() == 1)
-			{
-				const FString DraggedFileExtension = FPaths::GetExtension(Files[0], true);
-				if (DraggedFileExtension == TEXT(".utrace"))
-				{
-					return FReply::Handled();
-				}
-			}
-		}
+		return FReply::Handled();
 	}
 
-	return SCompoundWidget::OnDragOver(MyGeometry,DragDropEvent);
+	return SCompoundWidget::OnDragOver(MyGeometry, DragDropEvent);
 }
 
 ////////////////////////////////////////////////////////////////////////////////////////////////////
 
 FReply SLoadingProfilerWindow::OnDrop(const FGeometry& MyGeometry, const FDragDropEvent& DragDropEvent)
 {
-	TSharedPtr<FExternalDragOperation> DragDropOp = DragDropEvent.GetOperationAs<FExternalDragOperation>();
-	if (DragDropOp.IsValid())
+	if (FInsightsManager::Get()->OnDrop(DragDropEvent))
 	{
-		if (DragDropOp->HasFiles())
-		{
-			// For now, only allow a single file.
-			const TArray<FString>& Files = DragDropOp->GetFiles();
-			if (Files.Num() == 1)
-			{
-				const FString DraggedFileExtension = FPaths::GetExtension(Files[0], true);
-				if (DraggedFileExtension == TEXT(".utrace"))
-				{
-					// Enqueue load operation.
-					FInsightsManager::Get()->LoadTraceFile(Files[0]);
-					return FReply::Handled();
-				}
-			}
-		}
+		return FReply::Handled();
 	}
 
-	return SCompoundWidget::OnDrop(MyGeometry,DragDropEvent);
+	return SCompoundWidget::OnDrop(MyGeometry, DragDropEvent);
 }
 
 ////////////////////////////////////////////////////////////////////////////////////////////////////
diff --git a/Engine/Source/Developer/TraceInsights/Private/Insights/MemoryProfiler/Widgets/SMemoryProfilerWindow.cpp b/Engine/Source/Developer/TraceInsights/Private/Insights/MemoryProfiler/Widgets/SMemoryProfilerWindow.cpp
--- a/Engine/Source/Developer/TraceInsights/Private/Insights/MemoryProfiler/Widgets/SMemoryProfilerWindow.cpp
+++ b/Engine/Source/Developer/TraceInsights/Private/Insights/MemoryProfiler/Widgets/SMemoryProfilerWindow.cpp
@@ -708,51 +708,24 @@ FReply SMemoryProfilerWindow::OnKeyDown(const FGeometry& MyGeometry, const FKeyE
 
 FReply SMemoryProfilerWindow::OnDragOver(const FGeometry& MyGeometry, const FDragDropEvent& DragDropEvent)
 {
-	TSharedPtr<FExternalDragOperation> DragDropOp = DragDropEvent.GetOperationAs<FExternalDragOperation>();
-	if (DragDropOp.IsValid())
+	if (FInsightsManager::Get()->OnDragOver(DragDropEvent))
 	{
-		if (DragDropOp->HasFiles())
-		{
-			const TArray<FString>& Files = DragDropOp->GetFiles();
-			if (Files.Num() == 1)
-			{
-				const FString DraggedFileExtension = FPaths::GetExtension(Files[0], true);
-				if (DraggedFileExtension == TEXT(".utrace"))
-				{
-					return FReply::Handled();
-				}
-			}
-		}
+		return FReply::Handled();
 	}
 
-	return SCompoundWidget::OnDragOver(MyGeometry,DragDropEvent);
+	return SCompoundWidget::OnDragOver(MyGeometry, DragDropEvent);
 }
 
 ////////////////////////////////////////////////////////////////////////////////////////////////////
 
 FReply SMemoryProfilerWindow::OnDrop(const FGeometry& MyGeometry, const FDragDropEvent& DragDropEvent)
 {
-	TSharedPtr<FExternalDragOperation> DragDropOp = DragDropEvent.GetOperationAs<FExternalDragOperation>();
-	if (DragDropOp.IsValid())
+	if (FInsightsManager::Get()->OnDrop(DragDropEvent))
 	{
-		if (DragDropOp->HasFiles())
-		{
-			// For now, only allow a single file.
-			const TArray<FString>& Files = DragDropOp->GetFiles();
-			if (Files.Num() == 1)
-			{
-				const FString DraggedFileExtension = FPaths::GetExtension(Files[0], true);
-				if (DraggedFileExtension == TEXT(".utrace"))
-				{
-					// Enqueue load operation.
-					FInsightsManager::Get()->LoadTraceFile(Files[0]);
-					return FReply::Handled();
-				}
-			}
-		}
+		return FReply::Handled();
 	}
 
-	return SCompoundWidget::OnDrop(MyGeometry,DragDropEvent);
+	return SCompoundWidget::OnDrop(MyGeometry, DragDropEvent);
 }
 
 ////////////////////////////////////////////////////////////////////////////////////////////////////
diff --git a/Engine/Source/Developer/TraceInsights/Private/Insights/NetworkingProfiler/Widgets/SNetworkingProfilerWindow.cpp b/Engine/Source/Developer/TraceInsights/Private/Insights/NetworkingProfiler/Widgets/SNetworkingProfilerWindow.cpp
--- a/Engine/Source/Developer/TraceInsights/Private/Insights/NetworkingProfiler/Widgets/SNetworkingProfilerWindow.cpp
+++ b/Engine/Source/Developer/TraceInsights/Private/Insights/NetworkingProfiler/Widgets/SNetworkingProfilerWindow.cpp
@@ -695,51 +695,24 @@ FReply SNetworkingProfilerWindow::OnKeyDown(const FGeometry& MyGeometry, const F
 
 FReply SNetworkingProfilerWindow::OnDragOver(const FGeometry& MyGeometry, const FDragDropEvent& DragDropEvent)
 {
-	TSharedPtr<FExternalDragOperation> DragDropOp = DragDropEvent.GetOperationAs<FExternalDragOperation>();
-	if (DragDropOp.IsValid())
+	if (FInsightsManager::Get()->OnDragOver(DragDropEvent))
 	{
-		if (DragDropOp->HasFiles())
-		{
-			const TArray<FString>& Files = DragDropOp->GetFiles();
-			if (Files.Num() == 1)
-			{
-				const FString DraggedFileExtension = FPaths::GetExtension(Files[0], true);
-				if (DraggedFileExtension == TEXT(".utrace"))
-				{
-					return FReply::Handled();
-				}
-			}
-		}
+		return FReply::Handled();
 	}
 
-	return SCompoundWidget::OnDragOver(MyGeometry,DragDropEvent);
+	return SCompoundWidget::OnDragOver(MyGeometry, DragDropEvent);
 }
 
 ////////////////////////////////////////////////////////////////////////////////////////////////////
 
 FReply SNetworkingProfilerWindow::OnDrop(const FGeometry& MyGeometry, const FDragDropEvent& DragDropEvent)
 {
-	TSharedPtr<FExternalDragOperation> DragDropOp = DragDropEvent.GetOperationAs<FExternalDragOperation>();
-	if (DragDropOp.IsValid())
+	if (FInsightsManager::Get()->OnDrop(DragDropEvent))
 	{
-		if (DragDropOp->HasFiles())
-		{
-			// For now, only allow a single file.
-			const TArray<FString>& Files = DragDropOp->GetFiles();
-			if (Files.Num() == 1)
-			{
-				const FString DraggedFileExtension = FPaths::GetExtension(Files[0], true);
-				if (DraggedFileExtension == TEXT(".utrace"))
-				{
-					// Enqueue load operation.
-					FInsightsManager::Get()->LoadTraceFile(Files[0]);
-					return FReply::Handled();
-				}
-			}
-		}
+		return FReply::Handled();
 	}
 
-	return SCompoundWidget::OnDrop(MyGeometry,DragDropEvent);
+	return SCompoundWidget::OnDrop(MyGeometry, DragDropEvent);
 }
 
 ////////////////////////////////////////////////////////////////////////////////////////////////////
diff --git a/Engine/Source/Developer/TraceInsights/Private/Insights/TraceInsightsModule.cpp b/Engine/Source/Developer/TraceInsights/Private/Insights/TraceInsightsModule.cpp
--- a/Engine/Source/Developer/TraceInsights/Private/Insights/TraceInsightsModule.cpp
+++ b/Engine/Source/Developer/TraceInsights/Private/Insights/TraceInsightsModule.cpp
@@ -527,24 +527,7 @@ void FTraceInsightsModule::InitializeTesting(bool InInitAutomationModules, bool
 
 void FTraceInsightsModule::UpdateAppTitle()
 {
-#if !WITH_EDITOR
-	TSharedPtr<SWindow> RootWindow = FGlobalTabmanager::Get()->GetRootWindow();
-	if (RootWindow)
-	{
-		const FString TraceFilename = FInsightsManager::Get()->GetTraceFilename();
-		if (TraceFilename.IsEmpty())
-		{
-			const FText AppTitle = NSLOCTEXT("TraceInsightsModule", "UnrealInsightsAppName", "Unreal Insights");
-			RootWindow->SetTitle(AppTitle);
-		}
-		else
-		{
-			const FString SessionName = FPaths::GetBaseFilename(TraceFilename);
-			const FText AppTitle = FText::Format(NSLOCTEXT("TraceInsightsModule", "UnrealInsightsAppNameFmt", "{0} - Unreal Insights"), FText::FromString(SessionName));
-			RootWindow->SetTitle(AppTitle);
-		}
-	}
-#endif
+	FInsightsManager::Get()->UpdateAppTitle();
 }
 
 ////////////////////////////////////////////////////////////////////////////////////////////////////
diff --git a/Engine/Source/Developer/TraceInsights/Private/Insights/Widgets/SSessionInfoWindow.cpp b/Engine/Source/Developer/TraceInsights/Private/Insights/Widgets/SSessionInfoWindow.cpp
--- a/Engine/Source/Developer/TraceInsights/Private/Insights/Widgets/SSessionInfoWindow.cpp
+++ b/Engine/Source/Developer/TraceInsights/Private/Insights/Widgets/SSessionInfoWindow.cpp
@@ -40,6 +40,8 @@ const FName FSessionInfoTabs::SessionInfoID(TEXT("SessionInfo"));
 
 SSessionInfoWindow::SSessionInfoWindow()
 	: DurationActive(0.0f)
+	, AnalysisSession(nullptr)
+	, bIsSessionInfoSet(false)
 {
 }
 
@@ -339,23 +341,39 @@ void SSessionInfoWindow::FillMenu(FMenuBuilder& MenuBuilder, const TSharedPtr<FT
 
 void SSessionInfoWindow::Tick(const FGeometry& AllottedGeometry, const double InCurrentTime, const float InDeltaTime)
 {
-	// If we already have the session info data we no longer poll for it.
-	if (bIsSessionInfoSet)
-	{
-		return;
-	}
-
 	TSharedPtr<const TraceServices::IAnalysisSession> Session = FInsightsManager::Get()->GetSession();
 
-	if (Session.IsValid())
+	if (Session != AnalysisSession)
 	{
-		TraceServices::FAnalysisSessionReadScope SessionReadScope(*Session.Get());
+		// The session has changed. We need new info.
+		AnalysisSession = Session;
+		bIsSessionInfoSet = false;
+
+		// We can quickly get the session name and uri.
+		{
+			TraceServices::FAnalysisSessionReadScope SessionReadScope(*Session.Get());
+
+			SessionNameText = FText::FromString(FPaths::GetBaseFilename(Session->GetName()));
+
+			FString Uri(Session->GetName());
+			FPaths::NormalizeFilename(Uri);
+			UriText = FText::FromString(Uri);
+		}
 
-		SessionNameText = FText::FromString(FPaths::GetBaseFilename(Session->GetName()));
+		PlatformText = FText::GetEmpty();
+		AppNameText = FText::GetEmpty();
+		BranchText = FText::GetEmpty();
+		BuildVersionText = FText::GetEmpty();
+		ChangelistText = FText::GetEmpty();
+		BuildConfigurationTypeText = FText::GetEmpty();
+		BuildTargetTypeText = FText::GetEmpty();
+		CommandLineText = FText::GetEmpty();
+	}
 
-		FString Uri(Session->GetName());
-		FPaths::NormalizeFilename(Uri);
-		UriText = FText::FromString(Uri);
+	// If we already have the session info data, we no longer poll for it.
+	if (!bIsSessionInfoSet && Session.IsValid())
+	{
+		TraceServices::FAnalysisSessionReadScope SessionReadScope(*Session.Get());
 
 		const TraceServices::IDiagnosticsProvider* DiagnosticsProvider = TraceServices::ReadDiagnosticsProvider(*Session.Get());
 
@@ -370,6 +388,7 @@ void SSessionInfoWindow::Tick(const FGeometry& AllottedGeometry, const double In
 			BuildConfigurationTypeText = FText::FromString(LexToString(SessionInfo.ConfigurationType));
 			BuildTargetTypeText = FText::FromString(LexToString(SessionInfo.TargetType));
 			CommandLineText = FText::FromString(SessionInfo.CommandLine);
+
 			bIsSessionInfoSet = true;
 		}
 	}
@@ -421,6 +440,11 @@ FReply SSessionInfoWindow::OnKeyDown(const FGeometry& MyGeometry, const FKeyEven
 
 FReply SSessionInfoWindow::OnDragOver(const FGeometry& MyGeometry, const FDragDropEvent& DragDropEvent)
 {
+	if (FInsightsManager::Get()->OnDragOver(DragDropEvent))
+	{
+		return FReply::Handled();
+	}
+
 	return SCompoundWidget::OnDragOver(MyGeometry, DragDropEvent);
 }
 
@@ -428,6 +452,11 @@ FReply SSessionInfoWindow::OnDragOver(const FGeometry& MyGeometry, const FDragDr
 
 FReply SSessionInfoWindow::OnDrop(const FGeometry& MyGeometry, const FDragDropEvent& DragDropEvent)
 {
+	if (FInsightsManager::Get()->OnDrop(DragDropEvent))
+	{
+		return FReply::Handled();
+	}
+
 	return SCompoundWidget::OnDrop(MyGeometry, DragDropEvent);
 }
 
diff --git a/Engine/Source/Developer/TraceInsights/Private/Insights/Widgets/SSessionInfoWindow.h b/Engine/Source/Developer/TraceInsights/Private/Insights/Widgets/SSessionInfoWindow.h
--- a/Engine/Source/Developer/TraceInsights/Private/Insights/Widgets/SSessionInfoWindow.h
+++ b/Engine/Source/Developer/TraceInsights/Private/Insights/Widgets/SSessionInfoWindow.h
@@ -166,5 +166,7 @@ private:
 	FText BuildConfigurationTypeText;
 	FText BuildTargetTypeText;
 	FText CommandLineText;
+
+	TSharedPtr<const TraceServices::IAnalysisSession> AnalysisSession;
 	bool bIsSessionInfoSet = false;
 };
diff --git a/Engine/Source/Developer/TraceInsights/Private/Insights/Widgets/STimingProfilerWindow.cpp b/Engine/Source/Developer/TraceInsights/Private/Insights/Widgets/STimingProfilerWindow.cpp
--- a/Engine/Source/Developer/TraceInsights/Private/Insights/Widgets/STimingProfilerWindow.cpp
+++ b/Engine/Source/Developer/TraceInsights/Private/Insights/Widgets/STimingProfilerWindow.cpp
@@ -710,51 +710,24 @@ FReply STimingProfilerWindow::OnKeyDown(const FGeometry& MyGeometry, const FKeyE
 
 FReply STimingProfilerWindow::OnDragOver(const FGeometry& MyGeometry, const FDragDropEvent& DragDropEvent)
 {
-	TSharedPtr<FExternalDragOperation> DragDropOp = DragDropEvent.GetOperationAs<FExternalDragOperation>();
-	if (DragDropOp.IsValid())
+	if (FInsightsManager::Get()->OnDragOver(DragDropEvent))
 	{
-		if (DragDropOp->HasFiles())
-		{
-			const TArray<FString>& Files = DragDropOp->GetFiles();
-			if (Files.Num() == 1)
-			{
-				const FString DraggedFileExtension = FPaths::GetExtension(Files[0], true);
-				if (DraggedFileExtension == TEXT(".utrace"))
-				{
-					return FReply::Handled();
-				}
-			}
-		}
+		return FReply::Handled();
 	}
 
-	return SCompoundWidget::OnDragOver(MyGeometry,DragDropEvent);
+	return SCompoundWidget::OnDragOver(MyGeometry, DragDropEvent);
 }
 
 ////////////////////////////////////////////////////////////////////////////////////////////////////
 
 FReply STimingProfilerWindow::OnDrop(const FGeometry& MyGeometry, const FDragDropEvent& DragDropEvent)
 {
-	TSharedPtr<FExternalDragOperation> DragDropOp = DragDropEvent.GetOperationAs<FExternalDragOperation>();
-	if (DragDropOp.IsValid())
+	if (FInsightsManager::Get()->OnDrop(DragDropEvent))
 	{
-		if (DragDropOp->HasFiles())
-		{
-			// For now, only allow a single file.
-			const TArray<FString>& Files = DragDropOp->GetFiles();
-			if (Files.Num() == 1)
-			{
-				const FString DraggedFileExtension = FPaths::GetExtension(Files[0], true);
-				if (DraggedFileExtension == TEXT(".utrace"))
-				{
-					// Enqueue load operation.
-					FInsightsManager::Get()->LoadTraceFile(Files[0]);
-					return FReply::Handled();
-				}
-			}
-		}
+		return FReply::Handled();
 	}
 
-	return SCompoundWidget::OnDrop(MyGeometry,DragDropEvent);
+	return SCompoundWidget::OnDrop(MyGeometry, DragDropEvent);
 }
 
 ////////////////////////////////////////////////////////////////////////////////////////////////////
