Avoid computing fields for filtering unless the filter contains them
  - 434s s to 20s to apply an Address range filter to 7 minutes of memory profiling

diff --git a/Engine/Source/Developer/TraceInsights/Private/Insights/MemoryProfiler/Widgets/SMemAllocTableTreeView.cpp b/Engine/Source/Developer/TraceInsights/Private/Insights/MemoryProfiler/Widgets/SMemAllocTableTreeView.cpp
--- a/Engine/Source/Developer/TraceInsights/Private/Insights/MemoryProfiler/Widgets/SMemAllocTableTreeView.cpp
+++ b/Engine/Source/Developer/TraceInsights/Private/Insights/MemoryProfiler/Widgets/SMemAllocTableTreeView.cpp
@@ -1129,8 +1129,12 @@ void SMemAllocTableTreeView::UpdateQueryInfo()
 
 bool SMemAllocTableTreeView::ApplyCustomAdvancedFilters(const FTableTreeNodePtr& NodePtr)
 {
-	FMemAllocNodePtr MemNodePtr = StaticCastSharedPtr<FMemAllocNode>(NodePtr);
-	Context.SetFilterData<FString>(FullCallStackIndex, MemNodePtr->GetFullCallstack().ToString());
+	// Super heavy to compute, validate that the filter has a use for this key before computing it
+	if (FilterConfigurator && FilterConfigurator->IsKeyUsed(FullCallStackIndex))
+	{
+		FMemAllocNodePtr MemNodePtr = StaticCastSharedPtr<FMemAllocNode>(NodePtr);
+		Context.SetFilterData<FString>(FullCallStackIndex, MemNodePtr->GetFullCallstack().ToString());
+	}
 
 	return true;
 }
diff --git a/Engine/Source/Developer/TraceInsights/Private/Insights/Table/Widgets/STableTreeView.cpp b/Engine/Source/Developer/TraceInsights/Private/Insights/Table/Widgets/STableTreeView.cpp
--- a/Engine/Source/Developer/TraceInsights/Private/Insights/Table/Widgets/STableTreeView.cpp
+++ b/Engine/Source/Developer/TraceInsights/Private/Insights/Table/Widgets/STableTreeView.cpp
@@ -2934,7 +2934,11 @@ bool STableTreeView::ApplyAdvancedFilters(const FTableTreeNodePtr& NodePtr)
 			{
 				if (!Column->IsHierarchy())
 				{
-					Context.SetFilterData<FString>(Column->GetIndex(), Column->GetValue(*NodePtr)->AsString());
+					// Converting string is heavy, check if the key is used for filtering before doing any conversion
+					if (FilterConfiguratorToUse->IsKeyUsed(Column->GetIndex()))
+					{
+						Context.SetFilterData<FString>(Column->GetIndex(), Column->GetValue(*NodePtr)->AsString());
+					}
 				}
 				break;
 			}
diff --git a/Engine/Source/Developer/TraceInsights/Private/Insights/ViewModels/FilterConfigurator.cpp b/Engine/Source/Developer/TraceInsights/Private/Insights/ViewModels/FilterConfigurator.cpp
--- a/Engine/Source/Developer/TraceInsights/Private/Insights/ViewModels/FilterConfigurator.cpp
+++ b/Engine/Source/Developer/TraceInsights/Private/Insights/ViewModels/FilterConfigurator.cpp
@@ -44,6 +44,7 @@ FFilterConfigurator& FFilterConfigurator::operator=(const FFilterConfigurator& O
 	AvailableFilters = Other.AvailableFilters;
 	OnDestroyedEvent = Other.OnDestroyedEvent;
 
+	ComputeUsedKeys();
 	RootNode->ProcessFilter();
 
 	OnChangesCommitedEvent.Broadcast();
@@ -72,6 +73,17 @@ bool FFilterConfigurator::ApplyFilters(const FFilterContext& Context) const
 	return RootNode->ApplyFilters(Context);
 }
 
+bool FFilterConfigurator::IsKeyUsed(int32 Key) const
+{
+	return KeysUsed.Contains(Key);
+}
+
+void FFilterConfigurator::ComputeUsedKeys()
+{
+	KeysUsed.Reset();
+	RootNode->GetUsedKeys(KeysUsed);
+}
+
 ////////////////////////////////////////////////////////////////////////////////////////////////////
 
 } // namespace Insights
\ No newline at end of file
diff --git a/Engine/Source/Developer/TraceInsights/Private/Insights/ViewModels/FilterConfigurator.h b/Engine/Source/Developer/TraceInsights/Private/Insights/ViewModels/FilterConfigurator.h
--- a/Engine/Source/Developer/TraceInsights/Private/Insights/ViewModels/FilterConfigurator.h
+++ b/Engine/Source/Developer/TraceInsights/Private/Insights/ViewModels/FilterConfigurator.h
@@ -30,6 +30,8 @@ public:
 
 	virtual bool ApplyFilters(const FFilterContext& Context) const override;
 
+	bool IsKeyUsed(int32 Key) const;
+
 	TSharedPtr<TArray<TSharedPtr<struct FFilter>>>& GetAvailableFilters() { return AvailableFilters; }
 
 	////////////////////////////////////////////////////////////////////////////////////////////////////
@@ -56,9 +58,13 @@ private:
 
 private:
 
+	void ComputeUsedKeys();
+
 	FFilterConfiguratorNodePtr RootNode;
 
 	TSharedPtr<TArray<TSharedPtr<struct FFilter>>> AvailableFilters;
+
+	TSet<int32> KeysUsed;
 };
 
 ////////////////////////////////////////////////////////////////////////////////////////////////////
diff --git a/Engine/Source/Developer/TraceInsights/Private/Insights/ViewModels/FilterConfiguratorNode.cpp b/Engine/Source/Developer/TraceInsights/Private/Insights/ViewModels/FilterConfiguratorNode.cpp
--- a/Engine/Source/Developer/TraceInsights/Private/Insights/ViewModels/FilterConfiguratorNode.cpp
+++ b/Engine/Source/Developer/TraceInsights/Private/Insights/ViewModels/FilterConfiguratorNode.cpp
@@ -210,6 +210,22 @@ void FFilterConfiguratorNode::ProcessFilter()
 	}
 }
 
+void FFilterConfiguratorNode::GetUsedKeys(TSet<int32>& KeysUsed) const
+{
+	if (IsGroup())
+	{
+		auto& ChildrenArr = GetChildren();
+		for (int Index = 0; Index < ChildrenArr.Num(); ++Index)
+		{
+			((FFilterConfiguratorNode*)ChildrenArr[Index].Get())->GetUsedKeys(KeysUsed);
+		}
+	}
+	else
+	{
+		KeysUsed.Add(SelectedFilter->Key);
+	}
+}
+
 ////////////////////////////////////////////////////////////////////////////////////////////////////
 
 bool FFilterConfiguratorNode::ApplyFilters(const FFilterContext& Context) const
diff --git a/Engine/Source/Developer/TraceInsights/Private/Insights/ViewModels/FilterConfiguratorNode.h b/Engine/Source/Developer/TraceInsights/Private/Insights/ViewModels/FilterConfiguratorNode.h
--- a/Engine/Source/Developer/TraceInsights/Private/Insights/ViewModels/FilterConfiguratorNode.h
+++ b/Engine/Source/Developer/TraceInsights/Private/Insights/ViewModels/FilterConfiguratorNode.h
@@ -83,6 +83,8 @@ public:
 
 	bool ApplyFilters(const class FFilterContext& Context) const;
 
+	void GetUsedKeys(TSet<int32>& GetUsedKeys) const;
+
 	void ProcessFilter();
 
 private:
