[Insights] Added regex filter to ignore top callstack frames by symbol file path (now ignoring callstack frames with /Containers/ file path).

diff --git a/Engine/Source/Developer/TraceInsights/Private/Insights/MemoryProfiler/ViewModels/MemAllocTable.cpp b/Engine/Source/Developer/TraceInsights/Private/Insights/MemoryProfiler/ViewModels/MemAllocTable.cpp
--- a/Engine/Source/Developer/TraceInsights/Private/Insights/MemoryProfiler/ViewModels/MemAllocTable.cpp
+++ b/Engine/Source/Developer/TraceInsights/Private/Insights/MemoryProfiler/ViewModels/MemAllocTable.cpp
@@ -3,7 +3,7 @@
 #include "MemAllocTable.h"
 #include "CallstackFormatting.h"
 #include "Containers/StringView.h"
-
+#include "Internationalization/Regex.h"
 #include "Widgets/SBoxPanel.h"
 #include "Widgets/SToolTip.h"
 #include "Widgets/Text/STextBlock.h"
@@ -587,6 +587,7 @@ void FMemAllocTable::AddDefaultColumns()
 							check(Callstack->Num() > 0);
 
 							const TraceServices::FStackFrame* Frame = nullptr;
+							const TraceServices::FStackFrame* FirstUnknownSymbol = nullptr;
 							for (uint32 FrameIndex = 0; FrameIndex < Callstack->Num(); ++FrameIndex)
 							{
 								Frame = Callstack->Frame(FrameIndex);
@@ -594,10 +595,17 @@ void FMemAllocTable::AddDefaultColumns()
 
 								if (!Frame->Symbol || !Frame->Symbol->Name)
 								{
-									break;
+									if (FirstUnknownSymbol == nullptr)
+									{
+										FirstUnknownSymbol = Frame;
+									}
+									continue;
 								}
 
-								FStringView IgnoreSymbols[] =
+								bool bIgnoreSymbol = false;
+
+								// Ignore symbols by function prefix.
+								FStringView IgnoreSymbolsByFunctionName[] =
 								{
 									TEXT("FMemory::"_SV),
 									TEXT("FMallocWrapper::"_SV),
@@ -605,23 +613,49 @@ void FMemAllocTable::AddDefaultColumns()
 									TEXT("Malloc"_SV),
 									TEXT("Realloc"_SV),
 								};
-
-								bool bMatchIgnoreSymbol = false;
-								for (uint32 IgnoreSymbolsIndex = 0; IgnoreSymbolsIndex < UE_ARRAY_COUNT(IgnoreSymbols); ++IgnoreSymbolsIndex)
+								for (uint32 StringIndex = 0; StringIndex < UE_ARRAY_COUNT(IgnoreSymbolsByFunctionName); ++StringIndex)
 								{
-									if (FCString::Strnicmp(Frame->Symbol->Name, IgnoreSymbols[IgnoreSymbolsIndex].GetData(), IgnoreSymbols[IgnoreSymbolsIndex].Len()) == 0)
+									if (FCString::Strnicmp(Frame->Symbol->Name, IgnoreSymbolsByFunctionName[StringIndex].GetData(), IgnoreSymbolsByFunctionName[StringIndex].Len()) == 0)
 									{
-										bMatchIgnoreSymbol = true;
+										bIgnoreSymbol = true;
 										break;
 									}
 								}
 
-								if (!bMatchIgnoreSymbol)
+#if 1 // TODO: check perf impact
+								if (!bIgnoreSymbol && Frame->Symbol->File)
+								{
+									// Ignore symbols by file, specified as RegexPattern strings.
+									FStringView IgnoreSymbolsByFilePath[] =
+									{
+										TEXT(".*/Containers/.*"_SV),
+									};
+									for (uint32 StringIndex = 0; StringIndex < UE_ARRAY_COUNT(IgnoreSymbolsByFilePath); ++StringIndex)
+									{
+										const FString Pattern(IgnoreSymbolsByFilePath[StringIndex]);
+										const FRegexPattern RegexPattern(Pattern);
+										FString File(Frame->Symbol->File);
+										File.ReplaceCharInline(TEXT('\\'), TEXT('/'), ESearchCase::CaseSensitive);
+										FRegexMatcher RegexMatcher(RegexPattern, File);
+										if (RegexMatcher.FindNext())
+										{
+											bIgnoreSymbol = true;
+											break;
+										}
+									}
+								}
+#endif
+
+								if (!bIgnoreSymbol)
 								{
 									break;
 								}
 							}
-							check(Frame != nullptr);
+							if (!Frame)
+							{
+								check(FirstUnknownSymbol != nullptr);
+								Frame = FirstUnknownSymbol;
+							}
 
 							TStringBuilder<1024> Str;
 							FormatStackFrame(*Frame, Str, EStackFrameFormatFlags::Module);
