[Insights] Added pre-validation (if source file exists) before opening a file with ISourceCodeAccessor.

diff --git a/Engine/Source/Developer/TraceInsights/Private/Insights/TraceInsightsModule.cpp b/Engine/Source/Developer/TraceInsights/Private/Insights/TraceInsightsModule.cpp
--- a/Engine/Source/Developer/TraceInsights/Private/Insights/TraceInsightsModule.cpp
+++ b/Engine/Source/Developer/TraceInsights/Private/Insights/TraceInsightsModule.cpp
@@ -5,6 +5,9 @@
 #include "Framework/Docking/LayoutService.h"
 #include "Framework/Notifications/NotificationManager.h"
 #include "HAL/PlatformApplicationMisc.h"
+#include "ISourceCodeAccessModule.h"
+#include "ISourceCodeAccessor.h"
+#include "Logging/MessageLog.h"
 #include "Misc/ConfigCacheIni.h"
 #include "Misc/Paths.h"
 #include "Modules/ModuleManager.h"
@@ -60,6 +63,11 @@ void FTraceInsightsModule::StartupModule()
 	RegisterComponent(Insights::FTaskGraphProfilerManager::CreateInstance());
 	RegisterComponent(Insights::FContextSwitchesProfilerManager::CreateInstance());
 
+#if !WITH_EDITOR
+	ISourceCodeAccessModule& SourceCodeAccessModule = FModuleManager::LoadModuleChecked<ISourceCodeAccessModule>("SourceCodeAccess");
+	SourceCodeAccessModule.OnOpenFileFailed().AddRaw(this, &FTraceInsightsModule::HandleCodeAccessorOpenFileFailed);
+#endif
+
 	UnrealInsightsLayoutIni = GConfig->GetConfigFilename(TEXT("UnrealInsightsLayout"));
 }
 
@@ -67,6 +75,14 @@ void FTraceInsightsModule::StartupModule()
 
 void FTraceInsightsModule::ShutdownModule()
 {
+#if !WITH_EDITOR
+	if (FModuleManager::Get().IsModuleLoaded("SourceCodeAccess"))
+	{
+		ISourceCodeAccessModule& SourceCodeAccessModule = FModuleManager::GetModuleChecked<ISourceCodeAccessModule>("SourceCodeAccess");
+		SourceCodeAccessModule.OnOpenFileFailed().RemoveAll(this);
+	}
+#endif
+
 	if (PersistentLayout.IsValid())
 	{
 		// Save application layout.
@@ -531,3 +547,14 @@ void FTraceInsightsModule::UpdateAppTitle()
 }
 
 ////////////////////////////////////////////////////////////////////////////////////////////////////
+
+void FTraceInsightsModule::HandleCodeAccessorOpenFileFailed(const FString& Filename)
+{
+	FMessageLog ReportMessageLog(FInsightsManager::Get()->GetLogListingName());
+	FText Text = FText::Format(NSLOCTEXT("TraceInsightsModule", "FailedToOpenSourceFile", "Failed to open source file!\n\"{0}\""), FText::FromString(Filename));
+	TSharedRef<FTokenizedMessage> Message = FTokenizedMessage::Create(EMessageSeverity::Warning, Text);
+	ReportMessageLog.AddMessage(Message);
+	ReportMessageLog.Notify();
+}
+
+////////////////////////////////////////////////////////////////////////////////////////////////////
diff --git a/Engine/Source/Developer/TraceInsights/Private/Insights/TraceInsightsModule.h b/Engine/Source/Developer/TraceInsights/Private/Insights/TraceInsightsModule.h
--- a/Engine/Source/Developer/TraceInsights/Private/Insights/TraceInsightsModule.h
+++ b/Engine/Source/Developer/TraceInsights/Private/Insights/TraceInsightsModule.h
@@ -106,6 +106,8 @@ protected:
 
 	void UpdateAppTitle();
 
+	void HandleCodeAccessorOpenFileFailed(const FString& Filename);
+
 protected:
 #if WITH_TRACE_STORE
 	TUniquePtr<UE::Trace::FStoreService> StoreService;
diff --git a/Engine/Source/Developer/TraceInsights/Private/Insights/Widgets/SLogView.cpp b/Engine/Source/Developer/TraceInsights/Private/Insights/Widgets/SLogView.cpp
--- a/Engine/Source/Developer/TraceInsights/Private/Insights/Widgets/SLogView.cpp
+++ b/Engine/Source/Developer/TraceInsights/Private/Insights/Widgets/SLogView.cpp
@@ -1858,9 +1858,19 @@ void SLogView::OpenSource() const
 	if (SelectedLogMessage.IsValid())
 	{
 		FLogMessageRecord& Record = Cache.Get(SelectedLogMessage->GetIndex());
+		const FString File = Record.GetFile();
+		const uint32 Line = Record.GetLine();
+
 		ISourceCodeAccessModule& SourceCodeAccessModule = FModuleManager::LoadModuleChecked<ISourceCodeAccessModule>("SourceCodeAccess");
-		ISourceCodeAccessor& SourceCodeAccessor = SourceCodeAccessModule.GetAccessor();
-		SourceCodeAccessor.OpenFileAtLine(Record.GetFile(), Record.GetLine());
+		if (FPaths::FileExists(File))
+		{
+			ISourceCodeAccessor& SourceCodeAccessor = SourceCodeAccessModule.GetAccessor();
+			SourceCodeAccessor.OpenFileAtLine(File, Line);
+		}
+		else
+		{
+			SourceCodeAccessModule.OnOpenFileFailed().Broadcast(File);
+		}
 	}
 }
 
diff --git a/Engine/Source/Developer/TraceInsights/Private/Insights/Widgets/STimersView.cpp b/Engine/Source/Developer/TraceInsights/Private/Insights/Widgets/STimersView.cpp
--- a/Engine/Source/Developer/TraceInsights/Private/Insights/Widgets/STimersView.cpp
+++ b/Engine/Source/Developer/TraceInsights/Private/Insights/Widgets/STimersView.cpp
@@ -3060,8 +3060,15 @@ void STimersView::OpenSourceFileInIDE(FTimerNodePtr InNode) const
 		if (bIsValidSource)
 		{
 			ISourceCodeAccessModule& SourceCodeAccessModule = FModuleManager::LoadModuleChecked<ISourceCodeAccessModule>("SourceCodeAccess");
-			ISourceCodeAccessor& SourceCodeAccessor = SourceCodeAccessModule.GetAccessor();
-			SourceCodeAccessor.OpenFileAtLine(File, Line);
+			if (FPaths::FileExists(File))
+			{
+				ISourceCodeAccessor& SourceCodeAccessor = SourceCodeAccessModule.GetAccessor();
+				SourceCodeAccessor.OpenFileAtLine(File, Line);
+			}
+			else
+			{
+				SourceCodeAccessModule.OnOpenFileFailed().Broadcast(File);
+			}
 		}
 	}
 }
