[Insights] Memory: Added Start Event Index, End Event Index and Event Distance columns to Mem Allocs table.

diff --git a/Engine/Source/Developer/TraceInsights/Private/Insights/MemoryProfiler/ViewModels/MemAllocTable.cpp b/Engine/Source/Developer/TraceInsights/Private/Insights/MemoryProfiler/ViewModels/MemAllocTable.cpp
--- a/Engine/Source/Developer/TraceInsights/Private/Insights/MemoryProfiler/ViewModels/MemAllocTable.cpp
+++ b/Engine/Source/Developer/TraceInsights/Private/Insights/MemoryProfiler/ViewModels/MemAllocTable.cpp
@@ -25,6 +25,9 @@ namespace Insights
 ////////////////////////////////////////////////////////////////////////////////////////////////////
 // Column identifiers
 
+const FName FMemAllocTableColumns::StartEventIndexColumnId(TEXT("StartEventIndex"));
+const FName FMemAllocTableColumns::EndEventIndexColumnId(TEXT("EndEventIndex"));
+const FName FMemAllocTableColumns::EventDistanceColumnId(TEXT("EventDistance"));
 const FName FMemAllocTableColumns::StartTimeColumnId(TEXT("StartTime"));
 const FName FMemAllocTableColumns::EndTimeColumnId(TEXT("EndTime"));
 const FName FMemAllocTableColumns::DurationColumnId(TEXT("Duration"));
@@ -80,6 +83,180 @@ void FMemAllocTable::AddDefaultColumns()
 
 	int32 ColumnIndex = 0;
 
+	//////////////////////////////////////////////////
+	// Start Event Index Column
+	{
+		TSharedRef<FTableColumn> ColumnRef = MakeShared<FTableColumn>(FMemAllocTableColumns::StartEventIndexColumnId);
+		FTableColumn& Column = *ColumnRef;
+
+		Column.SetIndex(ColumnIndex++);
+
+		Column.SetShortName(LOCTEXT("StartEventIndexColumnName", "Start Index"));
+		Column.SetTitleName(LOCTEXT("StartEventIndexColumnTitle", "Start Event Index"));
+		Column.SetDescription(LOCTEXT("StartEventIndexColumnDesc", "The event index when the allocation was allocated."));
+
+		Column.SetFlags(ETableColumnFlags::CanBeHidden | ETableColumnFlags::CanBeFiltered);
+
+		Column.SetHorizontalAlignment(HAlign_Left);
+		Column.SetInitialWidth(100.0f);
+
+		Column.SetDataType(ETableCellDataType::Int64);
+
+		class FMemAllocStartEventIndexValueGetter : public FTableCellValueGetter
+		{
+		public:
+			virtual const TOptional<FTableCellValue> GetValue(const FTableColumn& Column, const FBaseTreeNode& Node) const override
+			{
+				if (Node.IsGroup())
+				{
+					const FTableTreeNode& NodePtr = static_cast<const FTableTreeNode&>(Node);
+					if (NodePtr.HasAggregatedValue(Column.GetId()))
+					{
+						return NodePtr.GetAggregatedValue(Column.GetId());
+					}
+				}
+				else //if (Node->Is<FMemAllocNode>())
+				{
+					const FMemAllocNode& MemAllocNode = static_cast<const FMemAllocNode&>(Node);
+					const FMemoryAlloc* Alloc = MemAllocNode.GetMemAlloc();
+					if (Alloc)
+					{
+						return FTableCellValue(Alloc->GetStartEventIndex());
+					}
+				}
+
+				return TOptional<FTableCellValue>();
+			}
+		};
+		TSharedRef<ITableCellValueGetter> Getter = MakeShared<FMemAllocStartEventIndexValueGetter>();
+		Column.SetValueGetter(Getter);
+
+		TSharedRef<ITableCellValueFormatter> Formatter = MakeShared<FInt64ValueFormatterAsNumber>();
+		Column.SetValueFormatter(Formatter);
+
+		TSharedRef<ITableCellValueSorter> Sorter = MakeShared<FSorterByInt64Value>(ColumnRef);
+		Column.SetValueSorter(Sorter);
+
+		Column.SetAggregation(ETableColumnAggregation::Min);
+
+		AddColumn(ColumnRef);
+	}
+	//////////////////////////////////////////////////
+	// End Event Index Column
+	{
+		TSharedRef<FTableColumn> ColumnRef = MakeShared<FTableColumn>(FMemAllocTableColumns::EndEventIndexColumnId);
+		FTableColumn& Column = *ColumnRef;
+
+		Column.SetIndex(ColumnIndex++);
+
+		Column.SetShortName(LOCTEXT("EndEventIndexColumnName", "End Index"));
+		Column.SetTitleName(LOCTEXT("EndEventIndexColumnTitle", "End Event Index"));
+		Column.SetDescription(LOCTEXT("EndEventIndexColumnDesc", "The event index when the allocation was freed."));
+
+		Column.SetFlags(ETableColumnFlags::CanBeHidden | ETableColumnFlags::CanBeFiltered);
+
+		Column.SetHorizontalAlignment(HAlign_Left);
+		Column.SetInitialWidth(100.0f);
+
+		Column.SetDataType(ETableCellDataType::Int64);
+
+		class FMemAllocEndEventIndexValueGetter : public FTableCellValueGetter
+		{
+		public:
+			virtual const TOptional<FTableCellValue> GetValue(const FTableColumn& Column, const FBaseTreeNode& Node) const override
+			{
+				if (Node.IsGroup())
+				{
+					const FTableTreeNode& NodePtr = static_cast<const FTableTreeNode&>(Node);
+					if (NodePtr.HasAggregatedValue(Column.GetId()))
+					{
+						return NodePtr.GetAggregatedValue(Column.GetId());
+					}
+				}
+				else //if (Node->Is<FMemAllocNode>())
+				{
+					const FMemAllocNode& MemAllocNode = static_cast<const FMemAllocNode&>(Node);
+					const FMemoryAlloc* Alloc = MemAllocNode.GetMemAlloc();
+					if (Alloc)
+					{
+						return FTableCellValue(Alloc->GetEndEventIndex());
+					}
+				}
+
+				return TOptional<FTableCellValue>();
+			}
+		};
+		TSharedRef<ITableCellValueGetter> Getter = MakeShared<FMemAllocEndEventIndexValueGetter>();
+		Column.SetValueGetter(Getter);
+
+		TSharedRef<ITableCellValueFormatter> Formatter = MakeShared<FInt64ValueFormatterAsNumber>();
+		Column.SetValueFormatter(Formatter);
+
+		TSharedRef<ITableCellValueSorter> Sorter = MakeShared<FSorterByInt64Value>(ColumnRef);
+		Column.SetValueSorter(Sorter);
+
+		Column.SetAggregation(ETableColumnAggregation::Max);
+
+		AddColumn(ColumnRef);
+	}
+	//////////////////////////////////////////////////
+	// Event Distance Column
+	{
+		TSharedRef<FTableColumn> ColumnRef = MakeShared<FTableColumn>(FMemAllocTableColumns::EventDistanceColumnId);
+		FTableColumn& Column = *ColumnRef;
+
+		Column.SetIndex(ColumnIndex++);
+
+		Column.SetShortName(LOCTEXT("EventDistanceColumnName", "Event Distance"));
+		Column.SetTitleName(LOCTEXT("EventDistanceColumnTitle", "Event Distance"));
+		Column.SetDescription(LOCTEXT("EventDistanceColumnDesc", "The event distance (index of free event minus index of alloc event)."));
+
+		Column.SetFlags(ETableColumnFlags::CanBeHidden | ETableColumnFlags::CanBeFiltered);
+
+		Column.SetHorizontalAlignment(HAlign_Left);
+		Column.SetInitialWidth(100.0f);
+
+		Column.SetDataType(ETableCellDataType::Int64);
+
+		class FMemAllocEventDistanceValueGetter : public FTableCellValueGetter
+		{
+		public:
+			virtual const TOptional<FTableCellValue> GetValue(const FTableColumn& Column, const FBaseTreeNode& Node) const override
+			{
+				if (Node.IsGroup())
+				{
+					const FTableTreeNode& NodePtr = static_cast<const FTableTreeNode&>(Node);
+					if (NodePtr.HasAggregatedValue(Column.GetId()))
+					{
+						return NodePtr.GetAggregatedValue(Column.GetId());
+					}
+				}
+				else //if (Node->Is<FMemAllocNode>())
+				{
+					const FMemAllocNode& MemAllocNode = static_cast<const FMemAllocNode&>(Node);
+					const FMemoryAlloc* Alloc = MemAllocNode.GetMemAlloc();
+					if (Alloc)
+					{
+						return FTableCellValue(Alloc->GetEventDistance());
+					}
+				}
+
+				return TOptional<FTableCellValue>();
+			}
+		};
+		TSharedRef<ITableCellValueGetter> Getter = MakeShared<FMemAllocEventDistanceValueGetter>();
+		Column.SetValueGetter(Getter);
+
+		TSharedRef<ITableCellValueFormatter> Formatter = MakeShared<FInt64ValueFormatterAsNumber>();
+		Column.SetValueFormatter(Formatter);
+
+		TSharedRef<ITableCellValueSorter> Sorter = MakeShared<FSorterByInt64Value>(ColumnRef);
+		Column.SetValueSorter(Sorter);
+
+		Column.SetAggregation(ETableColumnAggregation::Max);
+
+		AddColumn(ColumnRef);
+	}
 	//////////////////////////////////////////////////
 	// Start Time Column
 	{
@@ -157,7 +334,7 @@ void FMemAllocTable::AddDefaultColumns()
 
 		Column.SetDataType(ETableCellDataType::Double);
 
-		class FMemAllocStartTimeValueGetter : public FTableCellValueGetter
+		class FMemAllocEndTimeValueGetter : public FTableCellValueGetter
 		{
 		public:
 			virtual const TOptional<FTableCellValue> GetValue(const FTableColumn& Column, const FBaseTreeNode& Node) const override
@@ -183,7 +360,7 @@ void FMemAllocTable::AddDefaultColumns()
 				return TOptional<FTableCellValue>();
 			}
 		};
-		TSharedRef<ITableCellValueGetter> Getter = MakeShared<FMemAllocStartTimeValueGetter>();
+		TSharedRef<ITableCellValueGetter> Getter = MakeShared<FMemAllocEndTimeValueGetter>();
 		Column.SetValueGetter(Getter);
 
 		TSharedRef<ITableCellValueFormatter> Formatter = MakeShared<FDoubleValueFormatterAsTimeAuto>();
@@ -215,7 +392,7 @@ void FMemAllocTable::AddDefaultColumns()
 
 		Column.SetDataType(ETableCellDataType::Double);
 
-		class FMemAllocStartTimeValueGetter : public FTableCellValueGetter
+		class FMemAllocDurationValueGetter : public FTableCellValueGetter
 		{
 		public:
 			virtual const TOptional<FTableCellValue> GetValue(const FTableColumn& Column, const FBaseTreeNode& Node) const override
@@ -241,7 +418,7 @@ void FMemAllocTable::AddDefaultColumns()
 				return TOptional<FTableCellValue>();
 			}
 		};
-		TSharedRef<ITableCellValueGetter> Getter = MakeShared<FMemAllocStartTimeValueGetter>();
+		TSharedRef<ITableCellValueGetter> Getter = MakeShared<FMemAllocDurationValueGetter>();
 		Column.SetValueGetter(Getter);
 
 		TSharedRef<ITableCellValueFormatter> Formatter = MakeShared<FDoubleValueFormatterAsTimeAuto>();
diff --git a/Engine/Source/Developer/TraceInsights/Private/Insights/MemoryProfiler/ViewModels/MemAllocTable.h b/Engine/Source/Developer/TraceInsights/Private/Insights/MemoryProfiler/ViewModels/MemAllocTable.h
--- a/Engine/Source/Developer/TraceInsights/Private/Insights/MemoryProfiler/ViewModels/MemAllocTable.h
+++ b/Engine/Source/Developer/TraceInsights/Private/Insights/MemoryProfiler/ViewModels/MemAllocTable.h
@@ -18,6 +18,9 @@ class FTableColumn;
 // Column identifiers
 struct FMemAllocTableColumns
 {
+	static const FName StartEventIndexColumnId;
+	static const FName EndEventIndexColumnId;
+	static const FName EventDistanceColumnId;
 	static const FName StartTimeColumnId;
 	static const FName EndTimeColumnId;
 	static const FName DurationColumnId;
diff --git a/Engine/Source/Developer/TraceInsights/Private/Insights/MemoryProfiler/ViewModels/MemoryAlloc.cpp b/Engine/Source/Developer/TraceInsights/Private/Insights/MemoryProfiler/ViewModels/MemoryAlloc.cpp
--- a/Engine/Source/Developer/TraceInsights/Private/Insights/MemoryProfiler/ViewModels/MemoryAlloc.cpp
+++ b/Engine/Source/Developer/TraceInsights/Private/Insights/MemoryProfiler/ViewModels/MemoryAlloc.cpp
@@ -8,14 +8,16 @@ namespace Insights
 
 ////////////////////////////////////////////////////////////////////////////////////////////////////
 
-FMemoryAlloc::FMemoryAlloc(double InStartTime, double InEndTime, uint64 InAddress, uint64 InSize, const TCHAR* InTag, const TraceServices::FCallstack* InCallstack, HeapId InRootHeap)
-: StartTime(InStartTime)
-, EndTime(InEndTime)
-, Address(InAddress)
-, Size(InSize)
-, Tag(InTag)
-, Callstack(InCallstack)
-, RootHeap(InRootHeap)
+FMemoryAlloc::FMemoryAlloc()
+: StartEventIndex(0)
+, EndEventIndex(0)
+, StartTime(0.0)
+, EndTime(0.0)
+, Address(0)
+, Size(0)
+, Tag(nullptr)
+, Callstack(nullptr)
+, RootHeap(0)
 , bIsBlock(false)
 {
 }
diff --git a/Engine/Source/Developer/TraceInsights/Private/Insights/MemoryProfiler/ViewModels/MemoryAlloc.h b/Engine/Source/Developer/TraceInsights/Private/Insights/MemoryProfiler/ViewModels/MemoryAlloc.h
--- a/Engine/Source/Developer/TraceInsights/Private/Insights/MemoryProfiler/ViewModels/MemoryAlloc.h
+++ b/Engine/Source/Developer/TraceInsights/Private/Insights/MemoryProfiler/ViewModels/MemoryAlloc.h
@@ -17,9 +17,12 @@ class FMemoryAlloc
 	friend class SMemAllocTableTreeView;
 
 public:
-	FMemoryAlloc(double InStartTime, double InEndTime, uint64 InAddress, uint64 InSize, const TCHAR* InTag, const TraceServices::FCallstack* InCallstack, HeapId InRootHeap);
+	FMemoryAlloc();
 	~FMemoryAlloc();
 
+	int64 GetStartEventIndex() const { return int64(StartEventIndex); }
+	int64 GetEndEventIndex() const { return int64(EndEventIndex); }
+	int64 GetEventDistance() const { return int64(EndEventIndex) - int64(StartEventIndex); }
 	double GetStartTime() const { return StartTime; }
 	double GetEndTime() const { return EndTime; }
 	double GetDuration() const { return EndTime - StartTime; }
@@ -32,6 +35,8 @@ public:
 	HeapId GetRootHeap() const { return RootHeap; };
 
 private:
+	uint32 StartEventIndex;
+	uint32 EndEventIndex;
 	double StartTime;
 	double EndTime;
 	uint64 Address;
diff --git a/Engine/Source/Developer/TraceInsights/Private/Insights/MemoryProfiler/Widgets/SMemAllocTableTreeView.cpp b/Engine/Source/Developer/TraceInsights/Private/Insights/MemoryProfiler/Widgets/SMemAllocTableTreeView.cpp
--- a/Engine/Source/Developer/TraceInsights/Private/Insights/MemoryProfiler/Widgets/SMemAllocTableTreeView.cpp
+++ b/Engine/Source/Developer/TraceInsights/Private/Insights/MemoryProfiler/Widgets/SMemAllocTableTreeView.cpp
@@ -149,7 +149,7 @@ void SMemAllocTableTreeView::RebuildTree(bool bResync)
 				for (int32 AllocIndex = TableTreeNodes.Num(); AllocIndex < TotalAllocCount; ++AllocIndex)
 				{
 					const FMemoryAlloc* Alloc = MemAllocTable->GetMemAlloc(AllocIndex);
-					FName NodeName(Alloc->bIsBlock ? BaseHeapName : BaseNodeName, AllocIndex + 1);
+					FName NodeName(Alloc->bIsBlock ? BaseHeapName : BaseNodeName, Alloc->GetStartEventIndex() + 1);
 					FMemAllocNodePtr NodePtr = MakeShared<FMemAllocNode>(NodeName, MemAllocTable, AllocIndex);
 
 					// Until we have an UX story around heap allocations
@@ -352,6 +352,8 @@ void SMemAllocTableTreeView::UpdateQuery(TraceServices::IAllocationsProvider::EQ
 				{
 					const TraceServices::IAllocationsProvider::FAllocation* Allocation = Result->Get(AllocIndex);
 					FMemoryAlloc& Alloc = Allocs[AllocsDestIndex];
+					Alloc.StartEventIndex = Allocation->GetStartEventIndex();
+					Alloc.EndEventIndex = Allocation->GetEndEventIndex();
 					Alloc.StartTime = Allocation->GetStartTime();
 					Alloc.EndTime = Allocation->GetEndTime();
 					Alloc.Address = Allocation->GetAddress();
@@ -534,16 +536,11 @@ void SMemAllocTableTreeView::InitAvailableViewPresets()
 		}
 		virtual void GetColumnConfigSet(TArray<FColumnConfig>& InOutConfigSet) const override
 		{
-			InOutConfigSet.Add({ FTable::GetHierarchyColumnId(),                true,  200.0f });
-			InOutConfigSet.Add({ FMemAllocTableColumns::StartTimeColumnId,      false, 100.0f });
-			InOutConfigSet.Add({ FMemAllocTableColumns::EndTimeColumnId,        false, 100.0f });
-			InOutConfigSet.Add({ FMemAllocTableColumns::DurationColumnId,       false, 100.0f });
-			InOutConfigSet.Add({ FMemAllocTableColumns::AddressColumnId,        false, 120.0f });
-			InOutConfigSet.Add({ FMemAllocTableColumns::MemoryPageColumnId,     false, 120.0f });
-			InOutConfigSet.Add({ FMemAllocTableColumns::CountColumnId,          true,  100.0f });
-			InOutConfigSet.Add({ FMemAllocTableColumns::SizeColumnId,           true,  100.0f });
-			InOutConfigSet.Add({ FMemAllocTableColumns::TagColumnId,            true,  120.0f });
-			InOutConfigSet.Add({ FMemAllocTableColumns::FunctionColumnId,       true,  550.0f });
+			InOutConfigSet.Add({ FTable::GetHierarchyColumnId(),          true, 200.0f });
+			InOutConfigSet.Add({ FMemAllocTableColumns::CountColumnId,    true, 100.0f });
+			InOutConfigSet.Add({ FMemAllocTableColumns::SizeColumnId,     true, 100.0f });
+			InOutConfigSet.Add({ FMemAllocTableColumns::TagColumnId,      true, 120.0f });
+			InOutConfigSet.Add({ FMemAllocTableColumns::FunctionColumnId, true, 550.0f });
 		}
 	};
 	AvailableViewPresets.Add(MakeShared<FDefaultViewPreset>());
@@ -579,16 +576,19 @@ void SMemAllocTableTreeView::InitAvailableViewPresets()
 		}
 		virtual void GetColumnConfigSet(TArray<FColumnConfig>& InOutConfigSet) const override
 		{
-			InOutConfigSet.Add({ FTable::GetHierarchyColumnId(),                true,  200.0f });
-			InOutConfigSet.Add({ FMemAllocTableColumns::StartTimeColumnId,      true,  100.0f });
-			InOutConfigSet.Add({ FMemAllocTableColumns::EndTimeColumnId,        true,  100.0f });
-			InOutConfigSet.Add({ FMemAllocTableColumns::DurationColumnId,       true,  100.0f });
-			InOutConfigSet.Add({ FMemAllocTableColumns::AddressColumnId,        true,  120.0f });
-			InOutConfigSet.Add({ FMemAllocTableColumns::MemoryPageColumnId,     true,  120.0f });
-			InOutConfigSet.Add({ FMemAllocTableColumns::CountColumnId,          true,  100.0f });
-			InOutConfigSet.Add({ FMemAllocTableColumns::SizeColumnId,           true,  100.0f });
-			InOutConfigSet.Add({ FMemAllocTableColumns::TagColumnId,            true,  120.0f });
-			InOutConfigSet.Add({ FMemAllocTableColumns::FunctionColumnId,       true,  550.0f });
+			InOutConfigSet.Add({ FTable::GetHierarchyColumnId(),                 true, 200.0f });
+			InOutConfigSet.Add({ FMemAllocTableColumns::StartEventIndexColumnId, true, 100.0f });
+			InOutConfigSet.Add({ FMemAllocTableColumns::EndEventIndexColumnId,   true, 100.0f });
+			InOutConfigSet.Add({ FMemAllocTableColumns::EventDistanceColumnId,   true, 100.0f });
+			InOutConfigSet.Add({ FMemAllocTableColumns::StartTimeColumnId,       true, 100.0f });
+			InOutConfigSet.Add({ FMemAllocTableColumns::EndTimeColumnId,         true, 100.0f });
+			InOutConfigSet.Add({ FMemAllocTableColumns::DurationColumnId,        true, 100.0f });
+			InOutConfigSet.Add({ FMemAllocTableColumns::AddressColumnId,         true, 120.0f });
+			InOutConfigSet.Add({ FMemAllocTableColumns::MemoryPageColumnId,      true, 120.0f });
+			InOutConfigSet.Add({ FMemAllocTableColumns::CountColumnId,           true, 100.0f });
+			InOutConfigSet.Add({ FMemAllocTableColumns::SizeColumnId,            true, 100.0f });
+			InOutConfigSet.Add({ FMemAllocTableColumns::TagColumnId,             true, 120.0f });
+			InOutConfigSet.Add({ FMemAllocTableColumns::FunctionColumnId,        true, 550.0f });
 		}
 	};
 	AvailableViewPresets.Add(MakeShared<FDetailedViewPreset>());
@@ -634,15 +634,11 @@ void SMemAllocTableTreeView::InitAvailableViewPresets()
 		}
 		virtual void GetColumnConfigSet(TArray<FColumnConfig>& InOutConfigSet) const override
 		{
-			InOutConfigSet.Add({ FTable::GetHierarchyColumnId(),           true,  400.0f });
-			InOutConfigSet.Add({ FMemAllocTableColumns::StartTimeColumnId, false, 0.0f });
-			InOutConfigSet.Add({ FMemAllocTableColumns::EndTimeColumnId,   false, 0.0f });
-			InOutConfigSet.Add({ FMemAllocTableColumns::DurationColumnId,  false, 0.0f });
-			InOutConfigSet.Add({ FMemAllocTableColumns::AddressColumnId,   false, 0.0f });
-			InOutConfigSet.Add({ FMemAllocTableColumns::CountColumnId,     true,  100.0f });
-			InOutConfigSet.Add({ FMemAllocTableColumns::SizeColumnId,      true,  100.0f });
-			InOutConfigSet.Add({ FMemAllocTableColumns::TagColumnId,       true,  200.0f });
-			InOutConfigSet.Add({ FMemAllocTableColumns::FunctionColumnId,  true,  200.0f });
+			InOutConfigSet.Add({ FTable::GetHierarchyColumnId(),          true, 400.0f });
+			InOutConfigSet.Add({ FMemAllocTableColumns::CountColumnId,    true, 100.0f });
+			InOutConfigSet.Add({ FMemAllocTableColumns::SizeColumnId,     true, 100.0f });
+			InOutConfigSet.Add({ FMemAllocTableColumns::TagColumnId,      true, 200.0f });
+			InOutConfigSet.Add({ FMemAllocTableColumns::FunctionColumnId, true, 200.0f });
 		}
 	};
 	AvailableViewPresets.Add(MakeShared<FHeapViewPreset>());
@@ -688,16 +684,12 @@ void SMemAllocTableTreeView::InitAvailableViewPresets()
 		}
 		virtual void GetColumnConfigSet(TArray<FColumnConfig>& InOutConfigSet) const override
 		{
-			InOutConfigSet.Add({ FTable::GetHierarchyColumnId(),              true,  200.0f });
-			InOutConfigSet.Add({ FMemAllocTableColumns::StartTimeColumnId,    false, 0.0f   });
-			InOutConfigSet.Add({ FMemAllocTableColumns::EndTimeColumnId,      false, 0.0f   });
-			InOutConfigSet.Add({ FMemAllocTableColumns::DurationColumnId,     false, 0.0f   });
-			InOutConfigSet.Add({ FMemAllocTableColumns::AddressColumnId,      true,  120.0f });
-			InOutConfigSet.Add({ FMemAllocTableColumns::MemoryPageColumnId,   false, 0.0f   });
-			InOutConfigSet.Add({ FMemAllocTableColumns::CountColumnId,        true,  100.0f });
-			InOutConfigSet.Add({ FMemAllocTableColumns::SizeColumnId,         true,  100.0f });
-			InOutConfigSet.Add({ FMemAllocTableColumns::TagColumnId,          true,  120.0f });
-			InOutConfigSet.Add({ FMemAllocTableColumns::FunctionColumnId,     true,  400.0f });
+			InOutConfigSet.Add({ FTable::GetHierarchyColumnId(),          true, 200.0f });
+			InOutConfigSet.Add({ FMemAllocTableColumns::AddressColumnId,  true, 120.0f });
+			InOutConfigSet.Add({ FMemAllocTableColumns::CountColumnId,    true, 100.0f });
+			InOutConfigSet.Add({ FMemAllocTableColumns::SizeColumnId,     true, 100.0f });
+			InOutConfigSet.Add({ FMemAllocTableColumns::TagColumnId,      true, 120.0f });
+			InOutConfigSet.Add({ FMemAllocTableColumns::FunctionColumnId, true, 400.0f });
 		}
 	};
 	AvailableViewPresets.Add(MakeShared<FSizeViewPreset>());
@@ -744,16 +736,10 @@ void SMemAllocTableTreeView::InitAvailableViewPresets()
 		}
 		virtual void GetColumnConfigSet(TArray<FColumnConfig>& InOutConfigSet) const override
 		{
-			InOutConfigSet.Add({ FTable::GetHierarchyColumnId(),           true,  200.0f });
-			InOutConfigSet.Add({ FMemAllocTableColumns::StartTimeColumnId, false, 0.0f   });
-			InOutConfigSet.Add({ FMemAllocTableColumns::EndTimeColumnId,   false, 0.0f   });
-			InOutConfigSet.Add({ FMemAllocTableColumns::DurationColumnId,  false, 0.0f   });
-			InOutConfigSet.Add({ FMemAllocTableColumns::AddressColumnId,   false, 0.0f   });
-			InOutConfigSet.Add({ FMemAllocTableColumns::MemoryPageColumnId,false, 0.0f   });
-			InOutConfigSet.Add({ FMemAllocTableColumns::CountColumnId,     true,  100.0f });
-			InOutConfigSet.Add({ FMemAllocTableColumns::SizeColumnId,      true,  100.0f });
-			InOutConfigSet.Add({ FMemAllocTableColumns::TagColumnId,       false, 0.0f   });
-			InOutConfigSet.Add({ FMemAllocTableColumns::FunctionColumnId,  true,  400.0f });
+			InOutConfigSet.Add({ FTable::GetHierarchyColumnId(),          true, 200.0f });
+			InOutConfigSet.Add({ FMemAllocTableColumns::CountColumnId,    true, 100.0f });
+			InOutConfigSet.Add({ FMemAllocTableColumns::SizeColumnId,     true, 100.0f });
+			InOutConfigSet.Add({ FMemAllocTableColumns::FunctionColumnId, true, 400.0f });
 		}
 	};
 	AvailableViewPresets.Add(MakeShared<FTagViewPreset>());
@@ -810,15 +796,11 @@ void SMemAllocTableTreeView::InitAvailableViewPresets()
 		}
 		virtual void GetColumnConfigSet(TArray<FColumnConfig>& InOutConfigSet) const override
 		{
-			InOutConfigSet.Add({ FTable::GetHierarchyColumnId(),           true,  400.0f });
-			InOutConfigSet.Add({ FMemAllocTableColumns::StartTimeColumnId, false, 0.0f   });
-			InOutConfigSet.Add({ FMemAllocTableColumns::EndTimeColumnId,   false, 0.0f   });
-			InOutConfigSet.Add({ FMemAllocTableColumns::DurationColumnId,  false, 0.0f   });
-			InOutConfigSet.Add({ FMemAllocTableColumns::AddressColumnId,   false, 0.0f   });
-			InOutConfigSet.Add({ FMemAllocTableColumns::CountColumnId,     true,  100.0f });
-			InOutConfigSet.Add({ FMemAllocTableColumns::SizeColumnId,      true,  100.0f });
-			InOutConfigSet.Add({ FMemAllocTableColumns::TagColumnId,       true,  200.0f });
-			InOutConfigSet.Add({ FMemAllocTableColumns::FunctionColumnId,  true,  200.0f });
+			InOutConfigSet.Add({ FTable::GetHierarchyColumnId(),          true, 400.0f });
+			InOutConfigSet.Add({ FMemAllocTableColumns::CountColumnId,    true, 100.0f });
+			InOutConfigSet.Add({ FMemAllocTableColumns::SizeColumnId,     true, 100.0f });
+			InOutConfigSet.Add({ FMemAllocTableColumns::TagColumnId,      true, 200.0f });
+			InOutConfigSet.Add({ FMemAllocTableColumns::FunctionColumnId, true, 200.0f });
 		}
 
 	private:
@@ -869,17 +851,11 @@ void SMemAllocTableTreeView::InitAvailableViewPresets()
 		}
 		virtual void GetColumnConfigSet(TArray<FColumnConfig>& InOutConfigSet) const override
 		{
-			InOutConfigSet.Add({ FTable::GetHierarchyColumnId(),           true,  200.0f });
-			InOutConfigSet.Add({ FMemAllocTableColumns::StartTimeColumnId, false, 0.0f   });
-			InOutConfigSet.Add({ FMemAllocTableColumns::EndTimeColumnId,   false, 0.0f   });
-			InOutConfigSet.Add({ FMemAllocTableColumns::DurationColumnId,  false, 0.0f   });
-			InOutConfigSet.Add({ FMemAllocTableColumns::AddressColumnId,   true,  120.0f });
-			InOutConfigSet.Add({ FMemAllocTableColumns::MemoryPageColumnId,false, 0.0f   });
-			InOutConfigSet.Add({ FMemAllocTableColumns::CountColumnId,     true,  100.0f });
-			InOutConfigSet.Add({ FMemAllocTableColumns::SizeColumnId,      true,  100.0f });
-			InOutConfigSet.Add({ FMemAllocTableColumns::TagColumnId,       false, 0.0f   });
-			InOutConfigSet.Add({ FMemAllocTableColumns::FunctionColumnId,  true,  400.0f });
-
+			InOutConfigSet.Add({ FTable::GetHierarchyColumnId(),          true, 200.0f });
+			InOutConfigSet.Add({ FMemAllocTableColumns::AddressColumnId,  true, 120.0f });
+			InOutConfigSet.Add({ FMemAllocTableColumns::CountColumnId,    true, 100.0f });
+			InOutConfigSet.Add({ FMemAllocTableColumns::SizeColumnId,     true, 100.0f });
+			InOutConfigSet.Add({ FMemAllocTableColumns::FunctionColumnId, true, 400.0f });
 		}
 	};
 	AvailableViewPresets.Add(MakeShared<FPageViewPreset>());
@@ -918,27 +894,25 @@ void SMemAllocTableTreeView::ApplyViewPreset(const IViewPreset& InPreset)
 
 void SMemAllocTableTreeView::ApplyColumnConfig(const TArrayView<FColumnConfig>& InColumnConfigSet)
 {
+	// TODO: Reorder columns as in the config set.
+	// Currenly we only apply visibility and column width.
 	for (const TSharedRef<FTableColumn>& ColumnRef : Table->GetColumns())
 	{
 		FTableColumn& Column = ColumnRef.Get();
-		for (const FColumnConfig& Config : InColumnConfigSet)
+		const FName ColumnId = Column.GetId();
+		const FColumnConfig* ConfigPtr = InColumnConfigSet.FindByPredicate([ColumnId](const FColumnConfig& Config) { return ColumnId == Config.ColumnId; });
+		if (ConfigPtr && ConfigPtr->bIsVisible)
 		{
-			if (Column.GetId() == Config.ColumnId)
+			ShowColumn(Column);
+			if (ConfigPtr->Width > 0.0f)
 			{
-				if (Config.bIsVisible)
-				{
-					ShowColumn(Column);
-					if (Config.Width > 0.0f)
-					{
-						TreeViewHeaderRow->SetColumnWidth(Column.GetId(), Config.Width);
-					}
-				}
-				else
-				{
-					HideColumn(Column);
-				}
+				TreeViewHeaderRow->SetColumnWidth(ColumnId, ConfigPtr->Width);
 			}
 		}
+		else
+		{
+			HideColumn(Column);
+		}
 	}
 }
 
@@ -1303,11 +1277,18 @@ void SMemAllocTableTreeView::OpenCallstackFrameSourceFileInIDE()
 
 			if (Frame->Symbol && Frame->Symbol->File)
 			{
-				//if (FPlatformFileManager::Get().GetPlatformFile().FileExists(Frame->Symbol->File))
+				const FString File = Frame->Symbol->File;
+				const uint32 Line = Frame->Symbol->Line;
+
+				ISourceCodeAccessModule& SourceCodeAccessModule = FModuleManager::LoadModuleChecked<ISourceCodeAccessModule>("SourceCodeAccess");
+				if (FPaths::FileExists(File))
 				{
-					ISourceCodeAccessModule& SourceCodeAccessModule = FModuleManager::LoadModuleChecked<ISourceCodeAccessModule>("SourceCodeAccess");
 					ISourceCodeAccessor& SourceCodeAccessor = SourceCodeAccessModule.GetAccessor();
-					SourceCodeAccessor.OpenFileAtLine(Frame->Symbol->File, Frame->Symbol->Line);
+					SourceCodeAccessor.OpenFileAtLine(File, Line);
+				}
+				else
+				{
+					SourceCodeAccessModule.OnOpenFileFailed().Broadcast(File);
 				}
 			}
 		}
diff --git a/Engine/Source/Developer/TraceInsights/Private/Insights/Table/Widgets/STableTreeView.cpp b/Engine/Source/Developer/TraceInsights/Private/Insights/Table/Widgets/STableTreeView.cpp
--- a/Engine/Source/Developer/TraceInsights/Private/Insights/Table/Widgets/STableTreeView.cpp
+++ b/Engine/Source/Developer/TraceInsights/Private/Insights/Table/Widgets/STableTreeView.cpp
@@ -1486,17 +1486,37 @@ void STableTreeView::UpdateAggregatedValues(FTableTreeNode& GroupNode)
 				break;
 
 			case ETableColumnAggregation::Min:
-				STableTreeView::UpdateAggregationRec<double>(Column, GroupNode, std::numeric_limits<double>::max(), false, [](double InValue, TOptional<FTableCellValue> InTableCellValue)
+				if (Column.GetDataType() == ETableCellDataType::Float || Column.GetDataType() == ETableCellDataType::Double)
 				{
-					return FMath::Min(InValue, InTableCellValue->AsDouble());
-				});
+					STableTreeView::UpdateAggregationRec<double>(Column, GroupNode, std::numeric_limits<double>::max(), false, [](double InValue, TOptional<FTableCellValue> InTableCellValue)
+					{
+						return FMath::Min(InValue, InTableCellValue->AsDouble());
+					});
+				}
+				else
+				{
+					STableTreeView::UpdateAggregationRec<int64>(Column, GroupNode, std::numeric_limits<int64>::max(), false, [](int64 InValue, TOptional<FTableCellValue> InTableCellValue)
+					{
+						return FMath::Min(InValue, InTableCellValue->AsInt64());
+					});
+				}
 				break;
 
 			case ETableColumnAggregation::Max:
-				STableTreeView::UpdateAggregationRec<double>(Column, GroupNode, std::numeric_limits<double>::lowest(), false, [](double InValue, TOptional<FTableCellValue> InTableCellValue)
+				if (Column.GetDataType() == ETableCellDataType::Float || Column.GetDataType() == ETableCellDataType::Double)
 				{
-					return FMath::Max(InValue, InTableCellValue->AsDouble());
-				});
+					STableTreeView::UpdateAggregationRec<double>(Column, GroupNode, std::numeric_limits<double>::lowest(), false, [](double InValue, TOptional<FTableCellValue> InTableCellValue)
+					{
+						return FMath::Max(InValue, InTableCellValue->AsDouble());
+					});
+				}
+				else
+				{
+					STableTreeView::UpdateAggregationRec<int64>(Column, GroupNode, std::numeric_limits<int64>::min(), false, [](int64 InValue, TOptional<FTableCellValue> InTableCellValue)
+					{
+						return FMath::Max(InValue, InTableCellValue->AsInt64());
+					});
+				}
 				break;
 		}
 	}
diff --git a/Engine/Source/Developer/TraceServices/Private/Model/AllocationsProvider.cpp b/Engine/Source/Developer/TraceServices/Private/Model/AllocationsProvider.cpp
--- a/Engine/Source/Developer/TraceServices/Private/Model/AllocationsProvider.cpp
+++ b/Engine/Source/Developer/TraceServices/Private/Model/AllocationsProvider.cpp
@@ -223,6 +223,22 @@ bool FTagTracker::HasTagFromPtrScope(uint32 InThreadId, uint8 InTracker) const
 // IAllocationsProvider::FAllocation
 ////////////////////////////////////////////////////////////////////////////////////////////////////
 
+uint32 IAllocationsProvider::FAllocation::GetStartEventIndex() const
+{
+	const auto* Inner = (const FAllocationItem*)this;
+	return Inner->StartEventIndex;
+}
+
+////////////////////////////////////////////////////////////////////////////////////////////////////
+
+uint32 IAllocationsProvider::FAllocation::GetEndEventIndex() const
+{
+	const auto* Inner = (const FAllocationItem*)this;
+	return Inner->EndEventIndex;
+}
+
+////////////////////////////////////////////////////////////////////////////////////////////////////
+
 double IAllocationsProvider::FAllocation::GetStartTime() const
 {
 	const auto* Inner = (const FAllocationItem*)this;
diff --git a/Engine/Source/Developer/TraceServices/Public/TraceServices/Model/AllocationsProvider.h b/Engine/Source/Developer/TraceServices/Public/TraceServices/Model/AllocationsProvider.h
--- a/Engine/Source/Developer/TraceServices/Public/TraceServices/Model/AllocationsProvider.h
+++ b/Engine/Source/Developer/TraceServices/Public/TraceServices/Model/AllocationsProvider.h
@@ -86,6 +86,8 @@ public:
 
 	struct TRACESERVICES_API FAllocation
 	{
+		uint32 GetStartEventIndex() const;
+		uint32 GetEndEventIndex() const;
 		double GetStartTime() const;
 		double GetEndTime() const;
 		uint64 GetAddress() const;
