[Insights] Fixed icon for "Open Source File in [Editor]" menu options.

diff --git a/Engine/Source/Developer/TraceInsights/Private/Insights/InsightsStyle.cpp b/Engine/Source/Developer/TraceInsights/Private/Insights/InsightsStyle.cpp
--- a/Engine/Source/Developer/TraceInsights/Private/Insights/InsightsStyle.cpp
+++ b/Engine/Source/Developer/TraceInsights/Private/Insights/InsightsStyle.cpp
@@ -242,6 +242,11 @@ void FInsightsStyle::FStyle::Initialize()
 
 	//////////////////////////////////////////////////
 
+	Set("MainFrame.OpenVisualStudio", new EDITOR_IMAGE_BRUSH_SVG("Starship/Common/VisualStudio", Icon16x16));
+	Set("MainFrame.OpenSourceCodeEditor", new EDITOR_IMAGE_BRUSH_SVG("Starship/Common/SourceCodeEditor", Icon16x16));
+
+	//////////////////////////////////////////////////
+
 	Set("Icons.ResetToDefault", new EDITOR_IMAGE_BRUSH("Icons/Profiler/profiler_ResetToDefault_32x", Icon16x16));
 	Set("Icons.DiffersFromDefault", new EDITOR_IMAGE_BRUSH_SVG("Starship/Common/ResetToDefault", Icon16x16));
 
diff --git a/Engine/Source/Developer/TraceInsights/Private/Insights/MemoryProfiler/Widgets/SMemAllocTableTreeView.cpp b/Engine/Source/Developer/TraceInsights/Private/Insights/MemoryProfiler/Widgets/SMemAllocTableTreeView.cpp
--- a/Engine/Source/Developer/TraceInsights/Private/Insights/MemoryProfiler/Widgets/SMemAllocTableTreeView.cpp
+++ b/Engine/Source/Developer/TraceInsights/Private/Insights/MemoryProfiler/Widgets/SMemAllocTableTreeView.cpp
@@ -1238,11 +1238,7 @@ void SMemAllocTableTreeView::ExtendMenu(FMenuBuilder& MenuBuilder)
 		(
 			ItemLabel,
 			ItemToolTip,
-#if PLATFORM_WINDOWS
-			FSlateIcon(FAppStyle::GetAppStyleSetName(), "MainFrame.OpenVisualStudio"),
-#else
-			FSlateIcon(),
-#endif
+			FSlateIcon(FAppStyle::GetAppStyleSetName(), SourceCodeAccessor.GetOpenIconName()),
 			Action_OpenIDE,
 			NAME_None,
 			EUserInterfaceActionType::Button
diff --git a/Engine/Source/Developer/TraceInsights/Private/Insights/Widgets/SLogView.cpp b/Engine/Source/Developer/TraceInsights/Private/Insights/Widgets/SLogView.cpp
--- a/Engine/Source/Developer/TraceInsights/Private/Insights/Widgets/SLogView.cpp
+++ b/Engine/Source/Developer/TraceInsights/Private/Insights/Widgets/SLogView.cpp
@@ -1165,11 +1165,7 @@ TSharedPtr<SWidget> SLogView::ListView_GetContextMenu()
 				NAME_None,
 				ItemLabel,
 				ItemToolTip,
-#if PLATFORM_WINDOWS
-				FSlateIcon(FAppStyle::GetAppStyleSetName(), "MainFrame.OpenVisualStudio")
-#else
-				FSlateIcon()
-#endif
+				FSlateIcon(FAppStyle::GetAppStyleSetName(), SourceCodeAccessor.GetOpenIconName())
 			);
 		}
 	}
diff --git a/Engine/Source/Developer/TraceInsights/Private/Insights/Widgets/STimersView.cpp b/Engine/Source/Developer/TraceInsights/Private/Insights/Widgets/STimersView.cpp
--- a/Engine/Source/Developer/TraceInsights/Private/Insights/Widgets/STimersView.cpp
+++ b/Engine/Source/Developer/TraceInsights/Private/Insights/Widgets/STimersView.cpp
@@ -560,11 +560,7 @@ TSharedPtr<SWidget> STimersView::TreeView_GetMenuContent()
 				NAME_None,
 				ItemLabel,
 				ItemToolTip,
-#if PLATFORM_WINDOWS
-				FSlateIcon(FAppStyle::GetAppStyleSetName(), "MainFrame.OpenVisualStudio")
-#else
-				FSlateIcon()
-#endif
+				FSlateIcon(FAppStyle::GetAppStyleSetName(), SourceCodeAccessor.GetOpenIconName())
 			);
 		}
 	}
