typo

diff --git a/Engine/Source/Developer/TraceInsights/Private/Insights/ViewModels/ThreadTimingTrack.cpp b/Engine/Source/Developer/TraceInsights/Private/Insights/ViewModels/ThreadTimingTrack.cpp
--- a/Engine/Source/Developer/TraceInsights/Private/Insights/ViewModels/ThreadTimingTrack.cpp
+++ b/Engine/Source/Developer/TraceInsights/Private/Insights/ViewModels/ThreadTimingTrack.cpp
@@ -1210,7 +1210,7 @@ void FThreadTimingTrack::InitTooltip(FTooltipDrawState& InOutTooltip, const ITim
 					break;
 				}
 
-				InOutTooltip.AddTextLine(TEXT("-------- Wating for tasks --------"), FLinearColor::Red);
+				InOutTooltip.AddTextLine(TEXT("-------- Waiting for tasks --------"), FLinearColor::Red);
 				const int32 MaxWaitedTasksToList = 5;
 				FString TaskIdsStr = Waiting->Tasks.Num() <= MaxWaitedTasksToList ?
 					FString::JoinBy(Waiting->Tasks, TEXT(", "), [](TaskTrace::FId TaskId) { return FString::FromInt(TaskId); }) :
