[Insights] Timing View: Added "Allow Panning on Screen Edges" toggle option in the View Mode menu. If enabled, the panning is allowed to continue when mouse cursor reaches the edges of the screen. This option is persistent to UnrealInsightsSettings.ini file.

diff --git a/Engine/Source/Developer/TraceInsights/Private/Insights/InsightsSettings.h b/Engine/Source/Developer/TraceInsights/Private/Insights/InsightsSettings.h
--- a/Engine/Source/Developer/TraceInsights/Private/Insights/InsightsSettings.h
+++ b/Engine/Source/Developer/TraceInsights/Private/Insights/InsightsSettings.h
@@ -17,6 +17,7 @@ public:
 		, bIsDefault(bInIsDefault)
 		, DefaultZoomLevel(5.0) // 5 seconds between major tick marks
 		, bAutoHideEmptyTracks(true)
+		, bAllowPanningOnScreenEdges(false)
 		, bAutoZoomOnFrameSelection(false)
 		, AutoScrollFrameAlignment((int32)TraceFrameType_Game) // -1 = none, 0 = game, 1 = rendering
 		, AutoScrollViewportOffsetPercent(0.1) // scrolls forward 10% of viewport's width
@@ -41,6 +42,7 @@ public:
 
 		GConfig->GetDouble(TEXT("Insights.TimingProfiler"), TEXT("DefaultZoomLevel"), DefaultZoomLevel, SettingsIni);
 		GConfig->GetBool(TEXT("Insights.TimingProfiler"), TEXT("bAutoHideEmptyTracks"), bAutoHideEmptyTracks, SettingsIni);
+		GConfig->GetBool(TEXT("Insights.TimingProfiler"), TEXT("bAllowPanningOnScreenEdges"), bAllowPanningOnScreenEdges, SettingsIni);
 		GConfig->GetBool(TEXT("Insights.TimingProfiler"), TEXT("bAutoZoomOnFrameSelection"), bAutoZoomOnFrameSelection, SettingsIni);
 
 		// Auto-scroll options
@@ -73,6 +75,7 @@ public:
 	{
 		GConfig->SetDouble(TEXT("Insights.TimingProfiler"), TEXT("DefaultZoomLevel"), DefaultZoomLevel, SettingsIni);
 		GConfig->SetBool(TEXT("Insights.TimingProfiler"), TEXT("bAutoHideEmptyTracks"), bAutoHideEmptyTracks, SettingsIni);
+		GConfig->SetBool(TEXT("Insights.TimingProfiler"), TEXT("bAllowPanningOnScreenEdges"), bAllowPanningOnScreenEdges, SettingsIni);
 		GConfig->SetBool(TEXT("Insights.TimingProfiler"), TEXT("bAutoZoomOnFrameSelection"), bAutoZoomOnFrameSelection, SettingsIni);
 
 		// Auto-scroll options
@@ -110,6 +113,7 @@ public:
 	{
 		DefaultZoomLevel = Defaults.DefaultZoomLevel;
 		bAutoHideEmptyTracks = Defaults.bAutoHideEmptyTracks;
+		bAllowPanningOnScreenEdges = Defaults.bAllowPanningOnScreenEdges;
 		bAutoZoomOnFrameSelection = Defaults.bAutoZoomOnFrameSelection;
 		AutoScrollFrameAlignment = Defaults.AutoScrollFrameAlignment;
 		AutoScrollViewportOffsetPercent = Defaults.AutoScrollViewportOffsetPercent;
@@ -126,6 +130,10 @@ public:
 	void SetAutoHideEmptyTracks(bool bOnOff) { bAutoHideEmptyTracks = bOnOff; }
 	void SetAndSaveAutoHideEmptyTracks(bool bOnOff) { SET_AND_SAVE(bAutoHideEmptyTracks, bOnOff); }
 
+	bool IsPanningOnScreenEdgesEnabled() const { return bAllowPanningOnScreenEdges; }
+	void SetPanningOnScreenEdges(bool bOnOff) { bAllowPanningOnScreenEdges = bOnOff; }
+	void SetAndSavePanningOnScreenEdges(bool bOnOff) { SET_AND_SAVE(bAllowPanningOnScreenEdges, bOnOff); }
+
 	bool IsAutoZoomOnFrameSelectionEnabled() const { return bAutoZoomOnFrameSelection; }
 	void SetAutoZoomOnFrameSelection(bool bOnOff) { bAutoZoomOnFrameSelection = bOnOff; }
 	void SetAndSaveAutoZoomOnFrameSelection(bool bOnOff) { SET_AND_SAVE(bAutoZoomOnFrameSelection, bOnOff); }
@@ -170,6 +178,9 @@ private:
 	/** Auto hide empty tracks (ex.: ones without timing events in the current viewport). */
 	bool bAutoHideEmptyTracks;
 
+	/** If enabled, the panning is allowed to continue when mouse cursor reaches the edges of the screen. */
+	bool bAllowPanningOnScreenEdges;
+
 	/** If enabled, the Timing View will also be zoomed when a new frame is selected in the Frames track. */
 	bool bAutoZoomOnFrameSelection;
 
diff --git a/Engine/Source/Developer/TraceInsights/Private/Insights/TimingProfilerCommands.cpp b/Engine/Source/Developer/TraceInsights/Private/Insights/TimingProfilerCommands.cpp
--- a/Engine/Source/Developer/TraceInsights/Private/Insights/TimingProfilerCommands.cpp
+++ b/Engine/Source/Developer/TraceInsights/Private/Insights/TimingProfilerCommands.cpp
@@ -89,6 +89,12 @@ void FTimingViewCommands::RegisterCommands()
 		EUserInterfaceActionType::ToggleButton,
 		FInputChord(EKeys::V));
 
+	UI_COMMAND(PanningOnScreenEdges,
+		"Allow Panning on Screen Edges",
+		"If enabled, the panning is allowed to continue when mouse cursor reaches the edges of the screen.\nThis option is persistent to UnrealInsightsSettings.ini file.",
+		EUserInterfaceActionType::ToggleButton,
+		FInputChord());
+
 	UI_COMMAND(ToggleCompactMode,
 		"Compact Mode",
 		"Toggle compact mode for supporting tracks.\n(ex.: the timing tracks will be displayed with reduced height)",
diff --git a/Engine/Source/Developer/TraceInsights/Private/Insights/TimingProfilerCommands.h b/Engine/Source/Developer/TraceInsights/Private/Insights/TimingProfilerCommands.h
--- a/Engine/Source/Developer/TraceInsights/Private/Insights/TimingProfilerCommands.h
+++ b/Engine/Source/Developer/TraceInsights/Private/Insights/TimingProfilerCommands.h
@@ -69,6 +69,9 @@ public:
 	/** Toggles visibility of empty tracks. */
 	TSharedPtr<FUICommandInfo> AutoHideEmptyTracks;
 
+	/** Toggles "panning on screen edges". */
+	TSharedPtr<FUICommandInfo> PanningOnScreenEdges;
+
 	/** Toggles 'compact mode' for timing tracks. */
 	TSharedPtr<FUICommandInfo> ToggleCompactMode;
 
diff --git a/Engine/Source/Developer/TraceInsights/Private/Insights/Widgets/STimingView.cpp b/Engine/Source/Developer/TraceInsights/Private/Insights/Widgets/STimingView.cpp
--- a/Engine/Source/Developer/TraceInsights/Private/Insights/Widgets/STimingView.cpp
+++ b/Engine/Source/Developer/TraceInsights/Private/Insights/Widgets/STimingView.cpp
@@ -102,6 +102,10 @@ STimingView::STimingView()
 	, DefaultTimeMarker(MakeShared<Insights::FTimeMarker>())
 	, MarkersTrack(MakeShared<FMarkersTimingTrack>())
 	, GraphTrack(MakeShared<FTimingGraphTrack>())
+	, bAllowPanningOnScreenEdges(false)
+	, DPIScaleFactor(1.0f)
+	, EdgeFrameCountX(0)
+	, EdgeFrameCountY(0)
 	, WhiteBrush(FInsightsStyle::Get().GetBrush("WhiteBrush"))
 	, MainFont(FAppStyle::Get().GetFontStyle("SmallFont"))
 	, QuickFindTabId(TEXT("QuickFind"), TimingViewId++)
@@ -408,6 +412,10 @@ void STimingView::Reset(bool bIsFirstReset)
 	LastAutoScrollTime = 0;
 
 	bIsPanning = false;
+	bAllowPanningOnScreenEdges = Settings.IsPanningOnScreenEdgesEnabled();
+	DPIScaleFactor = 1.0f;
+	EdgeFrameCountX = 0;
+	EdgeFrameCountY = 0;
 	PanningMode = EPanningMode::None;
 
 	OverscrollLeft = 0.0f;
@@ -1886,6 +1894,15 @@ FReply STimingView::OnMouseButtonDown(const FGeometry& MyGeometry, const FPointe
 	MousePositionOnButtonDown = MyGeometry.AbsoluteToLocal(MouseEvent.GetScreenSpacePosition());
 	MousePosition = MousePositionOnButtonDown;
 
+	if (bAllowPanningOnScreenEdges)
+	{
+		const FVector2D ScreenSpacePosition = MouseEvent.GetScreenSpacePosition();
+		DPIScaleFactor = FPlatformApplicationMisc::GetDPIScaleFactorAtPoint(ScreenSpacePosition.X, ScreenSpacePosition.Y);
+
+		EdgeFrameCountX = 0;
+		EdgeFrameCountY = 0;
+	}
+
 	bool bStartPanningSelectingOrScrubbing = false;
 	bool bStartPanning = false;
 	bool bStartSelecting = false;
@@ -1950,7 +1967,14 @@ FReply STimingView::OnMouseButtonDown(const FGeometry& MyGeometry, const FPointe
 		}
 
 		// Capture mouse, so we can drag outside this widget.
-		Reply = FReply::Handled().CaptureMouse(SharedThis(this));
+		if (bAllowPanningOnScreenEdges)
+		{
+			Reply = FReply::Handled().CaptureMouse(SharedThis(this)).UseHighPrecisionMouseMovement(SharedThis(this)).SetUserFocus(SharedThis(this), EFocusCause::Mouse);
+		}
+		else
+		{
+			Reply = FReply::Handled().CaptureMouse(SharedThis(this));
+		}
 	}
 
 	if (bPreventThrottling)
@@ -2259,7 +2283,40 @@ FReply STimingView::OnMouseMove(const FGeometry& MyGeometry, const FPointerEvent
 
 	MousePosition = MyGeometry.AbsoluteToLocal(MouseEvent.GetScreenSpacePosition());
 
-	if (!MouseEvent.GetCursorDelta().IsZero())
+	const FVector2D& CursorDelta = MouseEvent.GetCursorDelta();
+
+	if (bIsPanning && bAllowPanningOnScreenEdges)
+	{
+		if (MouseEvent.GetScreenSpacePosition().X == MouseEvent.GetLastScreenSpacePosition().X)
+		{
+			++EdgeFrameCountX;
+		}
+		else
+		{
+			EdgeFrameCountX = 0;
+		}
+
+		if (EdgeFrameCountX > 10) // handle delay between high precision mouse movement and update of the actual cursor position
+		{
+			MousePositionOnButtonDown.X -= CursorDelta.X / DPIScaleFactor;
+		}
+
+		if (MouseEvent.GetScreenSpacePosition().Y == MouseEvent.GetLastScreenSpacePosition().Y)
+		{
+			++EdgeFrameCountY;
+		}
+		else
+		{
+			EdgeFrameCountY = 0;
+		}
+
+		if (EdgeFrameCountY > 10) // handle delay between high precision mouse movement and update of the actual cursor position
+		{
+			MousePositionOnButtonDown.Y -= CursorDelta.Y / DPIScaleFactor;
+		}
+	}
+
+	if (!CursorDelta.IsZero())
 	{
 		if (bIsPanning)
 		{
@@ -2371,7 +2428,6 @@ FReply STimingView::OnMouseWheel(const FGeometry& MyGeometry, const FPointerEven
 {
 	if (MouseEvent.GetModifierKeys().IsShiftDown())
 	{
-		//MousePosition = MyGeometry.AbsoluteToLocal(MouseEvent.GetScreenSpacePosition());
 		if (GraphTrack->IsVisible() &&
 			MousePosition.Y >= GraphTrack->GetPosY() &&
 			MousePosition.Y < GraphTrack->GetPosY() + GraphTrack->GetHeight())
@@ -2428,7 +2484,6 @@ FReply STimingView::OnMouseWheel(const FGeometry& MyGeometry, const FPointerEven
 	{
 		// Zoom in/out horizontally.
 		const double Delta = MouseEvent.GetWheelDelta();
-		//MousePosition = MyGeometry.AbsoluteToLocal(MouseEvent.GetScreenSpacePosition());
 		if (Viewport.RelativeZoomWithFixedX(Delta, MousePosition.X))
 		{
 			UpdateHorizontalScrollBar();
@@ -2765,11 +2820,6 @@ FReply STimingView::OnKeyDown(const FGeometry& MyGeometry, const FKeyEvent& InKe
 		Viewport.AddDirtyFlags(ETimingTrackViewportDirtyFlags::HInvalidated);
 		return FReply::Handled();
 	}
-	else if (InKeyEvent.GetKey() == EKeys::Zero)
-	{
-		ChooseNextCpuThreadTrackColoringMode();
-		return FReply::Handled();
-	}
 	else if (InKeyEvent.GetKey() == EKeys::X)
 	{
 		ChooseNextEventDepthLimit();
@@ -3045,6 +3095,12 @@ void STimingView::BindCommands()
 		FCanExecuteAction(),
 		FIsActionChecked::CreateSP(this, &STimingView::IsAutoHideEmptyTracksEnabled));
 
+	CommandList->MapAction(
+		Commands.PanningOnScreenEdges,
+		FExecuteAction::CreateSP(this, &STimingView::TogglePanningOnScreenEdges),
+		FCanExecuteAction(),
+		FIsActionChecked::CreateSP(this, &STimingView::IsPanningOnScreenEdgesEnabled));
+
 	CommandList->MapAction(
 		Commands.ToggleCompactMode,
 		FExecuteAction::CreateSP(this, &STimingView::ToggleCompactMode),
@@ -4465,6 +4521,12 @@ TSharedRef<SWidget> STimingView::MakeViewModeMenu()
 
 	CreateCpuThreadTrackColoringModeMenu(MenuBuilder);
 
+	MenuBuilder.BeginSection("Misc", LOCTEXT("MiscHeading", "Misc Settings"));
+	{
+		MenuBuilder.AddMenuEntry(Commands.PanningOnScreenEdges);
+	}
+	MenuBuilder.EndSection();
+
 	return MenuBuilder.MakeWidget();
 }
 
@@ -4511,6 +4573,24 @@ void STimingView::ToggleAutoHideEmptyTracks()
 
 ////////////////////////////////////////////////////////////////////////////////////////////////////
 
+bool STimingView::IsPanningOnScreenEdgesEnabled() const
+{
+	return bAllowPanningOnScreenEdges;
+}
+
+////////////////////////////////////////////////////////////////////////////////////////////////////
+
+void STimingView::TogglePanningOnScreenEdges()
+{
+	bAllowPanningOnScreenEdges = !bAllowPanningOnScreenEdges;
+
+	// Persistent option. Save it to the config file.
+	FInsightsSettings& Settings = FInsightsManager::Get()->GetSettings();
+	Settings.SetAndSavePanningOnScreenEdges(bAllowPanningOnScreenEdges);
+}
+
+////////////////////////////////////////////////////////////////////////////////////////////////////
+
 void AddMenuEntryRadioButton(
 	FMenuBuilder& InOutMenuBuilder,
 	const FUIAction& InAction,
diff --git a/Engine/Source/Developer/TraceInsights/Private/Insights/Widgets/STimingView.h b/Engine/Source/Developer/TraceInsights/Private/Insights/Widgets/STimingView.h
--- a/Engine/Source/Developer/TraceInsights/Private/Insights/Widgets/STimingView.h
+++ b/Engine/Source/Developer/TraceInsights/Private/Insights/Widgets/STimingView.h
@@ -78,6 +78,9 @@ public:
 	bool IsAutoHideEmptyTracksEnabled() const;
 	void ToggleAutoHideEmptyTracks();
 
+	bool IsPanningOnScreenEdgesEnabled() const;
+	void TogglePanningOnScreenEdges();
+
 	bool QuickFind_CanExecute() const;
 	void QuickFind_Execute();
 
@@ -555,7 +558,6 @@ protected:
 
 	/** The current mouse position. */
 	FVector2D MousePosition;
-	float DPIScaleFactor;
 
 	/** Mouse position during the call on mouse button down. */
 	FVector2D MousePositionOnButtonDown;
@@ -604,6 +606,14 @@ protected:
 	/** True, if the user is currently interactively panning the view (horizontally and/or vertically). */
 	bool bIsPanning;
 
+	/** If enabled, the panning is allowed to continue when mouse cursor reaches the edges of the screen. */
+	bool bAllowPanningOnScreenEdges;
+
+	float DPIScaleFactor;
+
+	uint32 EdgeFrameCountX;
+	uint32 EdgeFrameCountY;
+
 	/** How to pan. */
 	enum class EPanningMode : uint8
 	{
