Make the TaskId from TaskTrace 64 bit. Breaks compatibility with older traces containing Task events.

diff --git a/Engine/Source/Developer/TraceInsights/Private/Insights/TaskGraphProfiler/TaskGraphProfilerManager.cpp b/Engine/Source/Developer/TraceInsights/Private/Insights/TaskGraphProfiler/TaskGraphProfilerManager.cpp
--- a/Engine/Source/Developer/TraceInsights/Private/Insights/TaskGraphProfiler/TaskGraphProfilerManager.cpp
+++ b/Engine/Source/Developer/TraceInsights/Private/Insights/TaskGraphProfiler/TaskGraphProfilerManager.cpp
@@ -429,7 +429,7 @@ void FTaskGraphProfilerManager::ShowTaskRelations(const FThreadTrackEvent* InSel
 
 ////////////////////////////////////////////////////////////////////////////////////////////////////
 
-void FTaskGraphProfilerManager::ShowTaskRelations(uint32 TaskId)
+void FTaskGraphProfilerManager::ShowTaskRelations(TaskTrace::FId TaskId)
 {
 	TSharedPtr<const TraceServices::IAnalysisSession> Session = FInsightsManager::Get()->GetSession();
 	if (!Session.IsValid())
@@ -600,7 +600,7 @@ void FTaskGraphProfilerManager::SetShowRelations(bool bInValue)
 	bShowRelations = bInValue;
 	if (!bShowRelations)
 	{
-		TaskTimingSharedState->SetTaskId(FTaskTimingTrack::InvalidTaskId);
+		TaskTimingSharedState->SetTaskId(TaskTrace::InvalidId);
 	}
 }
 
diff --git a/Engine/Source/Developer/TraceInsights/Private/Insights/TaskGraphProfiler/TaskGraphProfilerManager.h b/Engine/Source/Developer/TraceInsights/Private/Insights/TaskGraphProfiler/TaskGraphProfilerManager.h
--- a/Engine/Source/Developer/TraceInsights/Private/Insights/TaskGraphProfiler/TaskGraphProfilerManager.h
+++ b/Engine/Source/Developer/TraceInsights/Private/Insights/TaskGraphProfiler/TaskGraphProfilerManager.h
@@ -2,6 +2,7 @@
 
 #pragma once
 
+#include "Async/TaskTrace.h"
 #include "Containers/Ticker.h"
 #include "CoreMinimal.h"
 #include "Framework/Commands/UICommandList.h"
@@ -96,7 +97,7 @@ public:
 	void OnSessionChanged();
 
 	void ShowTaskRelations(const FThreadTrackEvent* InSelectedEvent, uint32 ThreadId);
-	void ShowTaskRelations(uint32 TaskId);
+	void ShowTaskRelations(TaskTrace::FId TaskId);
 	void AddRelation(const FThreadTrackEvent* InSelectedEvent, double SourceTimestamp, uint32 SourceThreadId, double TargetTimestamp, uint32 TargetThreadId, ETaskEventType Type);
 	void AddRelation(const FThreadTrackEvent* InSelectedEvent, double SourceTimestamp, uint32 SourceThreadId, int32 SourceDepth, double TargetTimestamp, uint32 TargetThreadId, int32 TargetDepth, ETaskEventType Type);
 	void ClearTaskRelations();
diff --git a/Engine/Source/Developer/TraceInsights/Private/Insights/TaskGraphProfiler/ViewModels/TaskEntry.h b/Engine/Source/Developer/TraceInsights/Private/Insights/TaskGraphProfiler/ViewModels/TaskEntry.h
--- a/Engine/Source/Developer/TraceInsights/Private/Insights/TaskGraphProfiler/ViewModels/TaskEntry.h
+++ b/Engine/Source/Developer/TraceInsights/Private/Insights/TaskGraphProfiler/ViewModels/TaskEntry.h
@@ -2,6 +2,7 @@
 
 #pragma once
 
+#include "Async/TaskTrace.h"
 #include "CoreMinimal.h"
 
 #include "TraceServices/Model/TasksProfiler.h"
@@ -40,10 +41,10 @@ public:
 	uint32 GetNumPrerequisites() const { return NumPrerequisites; }
 	uint32 GetNumSubsequents() const { return NumSubsequents; }
 
-	uint32 GetId() const { return Id; }
+	TaskTrace::FId GetId() const { return Id; }
 
 private:
-	uint32 Id;
+	TaskTrace::FId Id;
 
 	const TCHAR* DebugName; 
 	bool bTracked; 
diff --git a/Engine/Source/Developer/TraceInsights/Private/Insights/TaskGraphProfiler/ViewModels/TaskTimingTrack.cpp b/Engine/Source/Developer/TraceInsights/Private/Insights/TaskGraphProfiler/ViewModels/TaskTimingTrack.cpp
--- a/Engine/Source/Developer/TraceInsights/Private/Insights/TaskGraphProfiler/ViewModels/TaskTimingTrack.cpp
+++ b/Engine/Source/Developer/TraceInsights/Private/Insights/TaskGraphProfiler/ViewModels/TaskTimingTrack.cpp
@@ -145,7 +145,7 @@ void FTaskTimingSharedState::Tick(Insights::ITimingViewSession& InSession, const
 		if (!TimingView->GetSelectedEvent().IsValid() && 
 			(!TimingView->GetSelectedTrack().IsValid() || TimingView->GetSelectedTrack().Get() != TaskTrack.Get()))
 		{
-			SetTaskId(FTaskTimingTrack::InvalidTaskId);
+			SetTaskId(TaskTrace::InvalidId);
 			FTaskGraphProfilerManager::Get()->ClearTaskRelations();
 		}
 	}
@@ -163,7 +163,7 @@ void FTaskTimingSharedState::ExtendFilterMenu(Insights::ITimingViewSession& InSe
 
 ////////////////////////////////////////////////////////////////////////////////////////////////////
 
-void FTaskTimingSharedState::SetTaskId(uint32 InTaskId)
+void FTaskTimingSharedState::SetTaskId(TaskTrace::FId InTaskId)
 {
 	if (TaskTrack.IsValid())
 	{
@@ -426,7 +426,7 @@ void FTaskTimingSharedState::OnTaskSettingsChanged()
 		return;
 	}
 
-	if (TaskTrack->GetTaskId() == FTaskTimingTrack::InvalidTaskId)
+	if (TaskTrack->GetTaskId() == TaskTrace::InvalidId)
 	{
 		FTaskGraphProfilerManager::Get()->ClearTaskRelations();
 	}
@@ -440,7 +440,7 @@ void FTaskTimingSharedState::OnTaskSettingsChanged()
 		return;
 	}
 
-	if(TaskTrack->GetTaskId() != FTaskTimingTrack::InvalidTaskId)
+	if(TaskTrack->GetTaskId() != TaskTrace::InvalidId)
 	{
 		FTaskGraphProfilerManager::Get()->ShowTaskRelations(TaskTrack->GetTaskId());
 	}
@@ -454,12 +454,9 @@ INSIGHTS_IMPLEMENT_RTTI(FTaskTimingTrack)
 
 ////////////////////////////////////////////////////////////////////////////////////////////////////
 
-const uint32 FTaskTimingTrack::InvalidTaskId = ~0;
-
-
 void FTaskTimingTrack::BuildDrawState(ITimingEventsTrackDrawStateBuilder& Builder, const ITimingTrackUpdateContext& Context)
 {
-	if (TaskId == InvalidTaskId)
+	if (TaskId == TaskTrace::InvalidId)
 	{
 		return;
 	}
@@ -527,9 +524,9 @@ void FTaskTimingTrack::OnTimingEventSelected(TSharedPtr<const ITimingEvent> InSe
 
 	if (!InSelectedEvent.IsValid() || !InSelectedEvent->Is<FThreadTrackEvent>())
 	{
-		if (TaskId != InvalidTaskId)
+		if (TaskId != TaskTrace::InvalidId)
 		{
-			TaskId = InvalidTaskId;
+			TaskId = TaskTrace::InvalidId;
 			FTaskGraphProfilerManager::Get()->ClearTaskRelations();
 			SetDirtyFlag();
 		}
@@ -563,7 +560,7 @@ void FTaskTimingTrack::OnTimingEventSelected(TSharedPtr<const ITimingEvent> InSe
 	}
 	else
 	{
-		TaskId = InvalidTaskId;
+		TaskId = TaskTrace::InvalidId;
 	}
 
 	SetDirtyFlag();
@@ -575,7 +572,7 @@ const TSharedPtr<const ITimingEvent> FTaskTimingTrack::GetEvent(float InPosX, fl
 {
 	TSharedPtr<FTaskTrackEvent> TimingEvent;
 
-	if (TaskId == InvalidTaskId)
+	if (TaskId == TaskTrace::InvalidId)
 	{
 		return TimingEvent;
 	}
diff --git a/Engine/Source/Developer/TraceInsights/Private/Insights/TaskGraphProfiler/ViewModels/TaskTimingTrack.h b/Engine/Source/Developer/TraceInsights/Private/Insights/TaskGraphProfiler/ViewModels/TaskTimingTrack.h
--- a/Engine/Source/Developer/TraceInsights/Private/Insights/TaskGraphProfiler/ViewModels/TaskTimingTrack.h
+++ b/Engine/Source/Developer/TraceInsights/Private/Insights/TaskGraphProfiler/ViewModels/TaskTimingTrack.h
@@ -2,6 +2,7 @@
 
 #pragma once
 
+#include "Async/TaskTrace.h"
 #include "CoreMinimal.h"
 
 // Insights
@@ -46,7 +47,7 @@ public:
 	void ShowTaskTrack() { SetTaskTrackToggle(true); }
 	void HideTaskTrack() { SetTaskTrackToggle(false); }
 
-	void SetTaskId(uint32 InTaskId);
+	void SetTaskId(TaskTrace::FId InTaskId);
 
 	void SetResetOnNextTick(bool bInValue) { bResetOnNextTick = bInValue; }
 
@@ -95,8 +96,6 @@ class FTaskTimingTrack : public FTimingEventsTrack
 	INSIGHTS_DECLARE_RTTI(FTaskTimingTrack, FTimingEventsTrack)
 
 public:
-	static const uint32 InvalidTaskId;
-
 	explicit FTaskTimingTrack(FTaskTimingSharedState& InSharedState, const FString& InName, uint32 InTimelineIndex)
 		: FTimingEventsTrack(InName)
 		, TimelineIndex(InTimelineIndex)
@@ -119,8 +118,8 @@ public:
 
 	virtual const TSharedPtr<const ITimingEvent> GetEvent(float InPosX, float InPosY, const FTimingTrackViewport& Viewport) const override;
 
-	void SetTaskId(uint32 InTaskId) { TaskId = InTaskId; SetDirtyFlag(); }
-	uint32 GetTaskId() { return TaskId; }
+	void SetTaskId(TaskTrace::FId InTaskId) { TaskId = InTaskId; SetDirtyFlag(); }
+	TaskTrace::FId GetTaskId() { return TaskId; }
 	
 	void OnTimingEventSelected(TSharedPtr<const ITimingEvent> InSelectedEvent);
 	void GetEventRelations(const FThreadTrackEvent& InSelectedEvent);
@@ -130,7 +129,7 @@ private:
 
 	FTaskTimingSharedState& SharedState;
 
-	uint32 TaskId = InvalidTaskId;
+	TaskTrace::FId TaskId = TaskTrace::InvalidId;
 
 	FVector2D MousePositionOnButtonDown;
 };
diff --git a/Engine/Source/Runtime/Core/Private/Async/TaskTrace.cpp b/Engine/Source/Runtime/Core/Private/Async/TaskTrace.cpp
--- a/Engine/Source/Runtime/Core/Private/Async/TaskTrace.cpp
+++ b/Engine/Source/Runtime/Core/Private/Async/TaskTrace.cpp
@@ -19,12 +19,12 @@ namespace TaskTrace
 
 	UE_TRACE_EVENT_BEGIN(TaskTrace, Created)
 		UE_TRACE_EVENT_FIELD(uint64, Timestamp)
-		UE_TRACE_EVENT_FIELD(uint32, TaskId)
+		UE_TRACE_EVENT_FIELD(uint64, TaskId)
 	UE_TRACE_EVENT_END()
 
 	UE_TRACE_EVENT_BEGIN(TaskTrace, Launched)
 		UE_TRACE_EVENT_FIELD(uint64, Timestamp)
-		UE_TRACE_EVENT_FIELD(uint32, TaskId)
+		UE_TRACE_EVENT_FIELD(uint64, TaskId)
 		UE_TRACE_EVENT_FIELD(UE::Trace::WideString, DebugName)
 		UE_TRACE_EVENT_FIELD(bool, Tracked)
 		UE_TRACE_EVENT_FIELD(int32, ThreadToExecuteOn)
@@ -32,39 +32,39 @@ namespace TaskTrace
 
 	UE_TRACE_EVENT_BEGIN(TaskTrace, Scheduled)
 		UE_TRACE_EVENT_FIELD(uint64, Timestamp)
-		UE_TRACE_EVENT_FIELD(uint32, TaskId)
+		UE_TRACE_EVENT_FIELD(uint64, TaskId)
 	UE_TRACE_EVENT_END()
 
 	UE_TRACE_EVENT_BEGIN(TaskTrace, SubsequentAdded)
 		UE_TRACE_EVENT_FIELD(uint64, Timestamp)
-		UE_TRACE_EVENT_FIELD(uint32, TaskId)
-		UE_TRACE_EVENT_FIELD(uint32, SubsequentId)
+		UE_TRACE_EVENT_FIELD(uint64, TaskId)
+		UE_TRACE_EVENT_FIELD(uint64, SubsequentId)
 	UE_TRACE_EVENT_END()
 
 	UE_TRACE_EVENT_BEGIN(TaskTrace, Started)
 		UE_TRACE_EVENT_FIELD(uint64, Timestamp)
-		UE_TRACE_EVENT_FIELD(uint32, TaskId)
+		UE_TRACE_EVENT_FIELD(uint64, TaskId)
 	UE_TRACE_EVENT_END()
 
 	UE_TRACE_EVENT_BEGIN(TaskTrace, NestedAdded)
 		UE_TRACE_EVENT_FIELD(uint64, Timestamp)
-		UE_TRACE_EVENT_FIELD(uint32, TaskId)
-		UE_TRACE_EVENT_FIELD(uint32, NestedId)
+		UE_TRACE_EVENT_FIELD(uint64, TaskId)
+		UE_TRACE_EVENT_FIELD(uint64, NestedId)
 	UE_TRACE_EVENT_END()
 
 	UE_TRACE_EVENT_BEGIN(TaskTrace, Finished)
 		UE_TRACE_EVENT_FIELD(uint64, Timestamp)
-		UE_TRACE_EVENT_FIELD(uint32, TaskId)
+		UE_TRACE_EVENT_FIELD(uint64, TaskId)
 	UE_TRACE_EVENT_END()
 
 	UE_TRACE_EVENT_BEGIN(TaskTrace, Completed)
 		UE_TRACE_EVENT_FIELD(uint64, Timestamp)
-		UE_TRACE_EVENT_FIELD(uint32, TaskId)
+		UE_TRACE_EVENT_FIELD(uint64, TaskId)
 	UE_TRACE_EVENT_END()
 
 	UE_TRACE_EVENT_BEGIN(TaskTrace, WaitingStarted)
 		UE_TRACE_EVENT_FIELD(uint64, Timestamp)
-		UE_TRACE_EVENT_FIELD(uint32[], Tasks)
+		UE_TRACE_EVENT_FIELD(uint64[], Tasks)
 	UE_TRACE_EVENT_END()
 
 	UE_TRACE_EVENT_BEGIN(TaskTrace, WaitingFinished)
diff --git a/Engine/Source/Runtime/Core/Public/Async/TaskTrace.h b/Engine/Source/Runtime/Core/Public/Async/TaskTrace.h
--- a/Engine/Source/Runtime/Core/Public/Async/TaskTrace.h
+++ b/Engine/Source/Runtime/Core/Public/Async/TaskTrace.h
@@ -23,7 +23,7 @@ namespace TaskTrace
 {
 	UE_TRACE_CHANNEL_EXTERN(TaskChannel);
 
-	using FId = uint32;
+	using FId = uint64;
 
 	const FId InvalidId = ~FId(0);
 
