[Insights] Minor changes / fixes to various tooltip texts.

diff --git a/Engine/Source/Developer/TraceInsights/Private/Insights/ViewModels/StatsViewColumnFactory.cpp b/Engine/Source/Developer/TraceInsights/Private/Insights/ViewModels/StatsViewColumnFactory.cpp
--- a/Engine/Source/Developer/TraceInsights/Private/Insights/ViewModels/StatsViewColumnFactory.cpp
+++ b/Engine/Source/Developer/TraceInsights/Private/Insights/ViewModels/StatsViewColumnFactory.cpp
@@ -60,8 +60,8 @@ TSharedRef<Insights::FTableColumn> FStatsViewColumnFactory::CreateNameColumn()
 	FTableColumn& Column = *ColumnRef;
 
 	Column.SetShortName(LOCTEXT("Name_ColumnName", "Name"));
-	Column.SetTitleName(LOCTEXT("Name_ColumnTitle", "Stats or Group Name"));
-	Column.SetDescription(LOCTEXT("Name_ColumnDesc", "Name of stats or group"));
+	Column.SetTitleName(LOCTEXT("Name_ColumnTitle", "Counter or Group Name"));
+	Column.SetDescription(LOCTEXT("Name_ColumnDesc", "Name of counter or group"));
 
 	Column.SetFlags(ETableColumnFlags::ShouldBeVisible |
 					ETableColumnFlags::CanBeFiltered |
@@ -140,7 +140,7 @@ TSharedRef<Insights::FTableColumn> FStatsViewColumnFactory::CreateTypeColumn()
 
 	Column.SetShortName(LOCTEXT("Type_ColumnName", "Type"));
 	Column.SetTitleName(LOCTEXT("Type_ColumnTitle", "Type"));
-	Column.SetDescription(LOCTEXT("Type_ColumnDesc", "Type of stats counter or group"));
+	Column.SetDescription(LOCTEXT("Type_ColumnDesc", "Type of counter or group"));
 
 	Column.SetFlags(ETableColumnFlags::CanBeHidden |
 					//ETableColumnFlags::ShouldBeVisible |
diff --git a/Engine/Source/Developer/TraceInsights/Private/Insights/Widgets/SFrameTrack.cpp b/Engine/Source/Developer/TraceInsights/Private/Insights/Widgets/SFrameTrack.cpp
--- a/Engine/Source/Developer/TraceInsights/Private/Insights/Widgets/SFrameTrack.cpp
+++ b/Engine/Source/Developer/TraceInsights/Private/Insights/Widgets/SFrameTrack.cpp
@@ -1165,7 +1165,7 @@ void SFrameTrack::ShowContextMenu(const FPointerEvent& MouseEvent)
 		MenuBuilder.AddMenuEntry
 		(
 			LOCTEXT("ContextMenu_AutoZoom", "Auto Zoom"),
-			LOCTEXT("ContextMenu_AutoZoom_Desc", "Enable auto zoom. Makes the entire session time range to fit into the Frames track's view."),
+			LOCTEXT("ContextMenu_AutoZoom_Desc", "Enables auto zoom. Makes the entire session time range to fit into the Frames track's view."),
 			FSlateIcon(),
 			Action_AutoZoom,
 			NAME_None,
@@ -1181,7 +1181,7 @@ void SFrameTrack::ShowContextMenu(const FPointerEvent& MouseEvent)
 		MenuBuilder.AddMenuEntry
 		(
 			LOCTEXT("ContextMenu_ZoomTimingViewOnFrameSelection", "Zoom Timing View on Frame Selection"),
-			LOCTEXT("ContextMenu_ZoomTimingViewOnFrameSelection_Desc", "If enabled, the Timing view will also be zoomed when a frame is selected.\n(This option is persistent to UnrealInsightsSettings.ini file.)"),
+			LOCTEXT("ContextMenu_ZoomTimingViewOnFrameSelection_Desc", "If enabled, the Timing view will also be zoomed when a frame is selected.\n(This option is persistent to the UnrealInsightsSettings.ini file.)"),
 			FSlateIcon(),
 			Action_ZoomTimingViewOnFrameSelection,
 			NAME_None,
diff --git a/Engine/Source/Developer/TraceInsights/Private/Insights/Widgets/SLogView.cpp b/Engine/Source/Developer/TraceInsights/Private/Insights/Widgets/SLogView.cpp
--- a/Engine/Source/Developer/TraceInsights/Private/Insights/Widgets/SLogView.cpp
+++ b/Engine/Source/Developer/TraceInsights/Private/Insights/Widgets/SLogView.cpp
@@ -47,8 +47,8 @@ class FLogViewCommands : public TCommands<FLogViewCommands>
 public:
 	FLogViewCommands()
 	: TCommands<FLogViewCommands>(
-		TEXT("FLogViewCommands"),
-		NSLOCTEXT("FLogViewCommands", "Log View Commands", "Log View Commands"),
+		TEXT("LogViewCommands"),
+		NSLOCTEXT("Contexts", "LogViewCommands", "Insights - Log View"),
 		NAME_None,
 		FInsightsStyle::GetStyleSetName())
 	{
@@ -62,16 +62,61 @@ public:
 	PRAGMA_DISABLE_OPTIMIZATION
 	virtual void RegisterCommands() override
 	{
-		UI_COMMAND(Command_HideSelectedCategory, "Hide Category", "Hides the selected log category.", EUserInterfaceActionType::Button, FInputChord());
-		UI_COMMAND(Command_ShowOnlySelectedCategory, "Show Only Selected Category", "Shows only the selected log category (hides all other log categories).", EUserInterfaceActionType::Button, FInputChord());
-		UI_COMMAND(Command_ShowAllCategories, "Show All Categories", "Resets the category filter (shows all log categories).", EUserInterfaceActionType::Button, FInputChord());
-		UI_COMMAND(Command_CopySelected, "Copy", "Copies the selected log (with all its properties) to clipboard.", EUserInterfaceActionType::Button, FInputChord(EModifierKey::Control, EKeys::C));
-		UI_COMMAND(Command_CopyMessage, "Copy Message", "Copies the message text of the selected log to clipboard.", EUserInterfaceActionType::Button, FInputChord(EModifierKey::Shift, EKeys::C));
-		UI_COMMAND(Command_CopyRange, "Copy Range", "Copies all the logs in the selected time range (highlighted in blue) to clipboard.", EUserInterfaceActionType::Button, FInputChord(EModifierKey::Control | EModifierKey::Shift, EKeys::C));
-		UI_COMMAND(Command_CopyAll, "Copy All", "Copies all the (filtered) logs to clipboard.", EUserInterfaceActionType::Button, FInputChord());
-		UI_COMMAND(Command_SaveRange, "Save Range As...", "Saves all the logs in the selected time range (highlighted in blue) to a text file (tab-separated values or comma-separated values).", EUserInterfaceActionType::Button, FInputChord(EModifierKey::Control, EKeys::S));
-		UI_COMMAND(Command_SaveAll, "Save All As...", "Saves all the (filtered) logs to a text file (tab-separated values or comma-separated values).", EUserInterfaceActionType::Button, FInputChord());
-		UI_COMMAND(Command_OpenSource, "Open Source", "Opens the source file of the selected message in the registered IDE.", EUserInterfaceActionType::Button, FInputChord());
+		UI_COMMAND(Command_HideSelectedCategory,
+			"Hide Category",
+			"Hides the selected log category.",
+			EUserInterfaceActionType::Button,
+			FInputChord());
+
+		UI_COMMAND(Command_ShowOnlySelectedCategory,
+			"Show Only Selected Category",
+			"Shows only the selected log category (hides all other log categories).",
+			EUserInterfaceActionType::Button,
+			FInputChord());
+
+		UI_COMMAND(Command_ShowAllCategories,
+			"Show All Categories",
+			"Resets the category filter (shows all log categories).",
+			EUserInterfaceActionType::Button,
+			FInputChord());
+
+		UI_COMMAND(Command_CopySelected,
+			"Copy",
+			"Copies the selected log (with all its properties) to clipboard.",
+			EUserInterfaceActionType::Button,
+			FInputChord(EModifierKey::Control, EKeys::C));
+
+		UI_COMMAND(Command_CopyMessage,
+			"Copy Message",
+			"Copies the message text of the selected log to clipboard.",
+			EUserInterfaceActionType::Button,
+			FInputChord(EModifierKey::Shift, EKeys::C));
+
+		UI_COMMAND(Command_CopyRange,
+			"Copy Range",
+			"Copies all the logs in the selected time range (highlighted in blue) to clipboard.",
+			EUserInterfaceActionType::Button,
+			FInputChord(EModifierKey::Control | EModifierKey::Shift, EKeys::C));
+
+		UI_COMMAND(Command_CopyAll,
+			"Copy All",
+			"Copies all the (filtered) logs to clipboard.",
+			EUserInterfaceActionType::Button, FInputChord());
+
+		UI_COMMAND(Command_SaveRange,
+			"Save Range As...",
+			"Saves all the logs in the selected time range (highlighted in blue) to a text file (tab-separated values or comma-separated values).",
+			EUserInterfaceActionType::Button, FInputChord(EModifierKey::Control, EKeys::S));
+
+		UI_COMMAND(Command_SaveAll,
+			"Save All As...",
+			"Saves all the (filtered) logs to a text file (tab-separated values or comma-separated values).",
+			EUserInterfaceActionType::Button, FInputChord());
+
+		UI_COMMAND(Command_OpenSource,
+			"Open Source",
+			"Opens the source file of the selected message in the registered IDE.",
+			EUserInterfaceActionType::Button, FInputChord());
 	}
 	PRAGMA_ENABLE_OPTIMIZATION
 
