[Insights] Updated context name for TCommands classes.

diff --git a/Engine/Source/Developer/TraceInsights/Private/Insights/ContextSwitches/ViewModels/ContextSwitchesSharedState.cpp b/Engine/Source/Developer/TraceInsights/Private/Insights/ContextSwitches/ViewModels/ContextSwitchesSharedState.cpp
--- a/Engine/Source/Developer/TraceInsights/Private/Insights/ContextSwitches/ViewModels/ContextSwitchesSharedState.cpp
+++ b/Engine/Source/Developer/TraceInsights/Private/Insights/ContextSwitches/ViewModels/ContextSwitchesSharedState.cpp
@@ -30,11 +30,11 @@ namespace Insights
 ////////////////////////////////////////////////////////////////////////////////////////////////////
 
 FContextSwitchesStateCommands::FContextSwitchesStateCommands()
-	: TCommands<FContextSwitchesStateCommands>(
-		TEXT("FContextSwitchesStateCommands"),
-		NSLOCTEXT("FContextSwitchesStateCommands", "Context Switches State Commands", "Context Switches Commands"),
-		NAME_None,
-		FInsightsStyle::GetStyleSetName())
+: TCommands<FContextSwitchesStateCommands>(
+	TEXT("ContextSwitchesStateCommands"),
+	NSLOCTEXT("Contexts", "ContextSwitchesStateCommands", "Insights - Context Switches"),
+	NAME_None,
+	FInsightsStyle::GetStyleSetName())
 {
 }
 
diff --git a/Engine/Source/Developer/TraceInsights/Private/Insights/InsightsCommands.cpp b/Engine/Source/Developer/TraceInsights/Private/Insights/InsightsCommands.cpp
--- a/Engine/Source/Developer/TraceInsights/Private/Insights/InsightsCommands.cpp
+++ b/Engine/Source/Developer/TraceInsights/Private/Insights/InsightsCommands.cpp
@@ -20,12 +20,11 @@
 ////////////////////////////////////////////////////////////////////////////////////////////////////
 
 FInsightsCommands::FInsightsCommands()
-	: TCommands<FInsightsCommands>(
-		TEXT("InsightsCommands"), // Context name for fast lookup
-		NSLOCTEXT("Contexts", "InsightsCommand", "Insights Command"), // Localized context name for displaying
-		NAME_None, // Parent
-		FInsightsStyle::GetStyleSetName() // Icon Style Set
-	)
+: TCommands<FInsightsCommands>(
+	TEXT("InsightsCommands"), // Context name for fast lookup
+	NSLOCTEXT("Contexts", "InsightsCommands", "Insights"), // Localized context name for displaying
+	NAME_None, // Parent
+	FInsightsStyle::GetStyleSetName()) // Icon Style Set
 {
 }
 
@@ -35,9 +34,23 @@ FInsightsCommands::FInsightsCommands()
 PRAGMA_DISABLE_OPTIMIZATION
 void FInsightsCommands::RegisterCommands()
 {
-	UI_COMMAND(InsightsManager_Load, "Load...", "Loads profiler data from a trace file", EUserInterfaceActionType::Button, FInputChord(EModifierKey::Control, EKeys::L));
-	UI_COMMAND(ToggleDebugInfo, "Debug", "Toggles the display of debug info", EUserInterfaceActionType::ToggleButton, FInputChord(EModifierKey::Control, EKeys::D));
-	UI_COMMAND(OpenSettings, "Settings", "Opens the Unreal Insights settings", EUserInterfaceActionType::Button, FInputChord(EModifierKey::Control, EKeys::O));
+	UI_COMMAND(InsightsManager_Load,
+		"Load...",
+		"Loads profiler data from a trace file.",
+		EUserInterfaceActionType::Button,
+		FInputChord(EModifierKey::Control, EKeys::L));
+
+	UI_COMMAND(ToggleDebugInfo,
+		"Debug",
+		"Toggles the display of debug info.",
+		EUserInterfaceActionType::ToggleButton,
+		FInputChord(EModifierKey::Control, EKeys::D));
+
+	UI_COMMAND(OpenSettings,
+		"Settings",
+		"Opens the Unreal Insights settings.",
+		EUserInterfaceActionType::Button,
+		FInputChord(EModifierKey::Control, EKeys::O));
 }
 PRAGMA_ENABLE_OPTIMIZATION
 
diff --git a/Engine/Source/Developer/TraceInsights/Private/Insights/LoadingProfiler/LoadingProfilerCommands.cpp b/Engine/Source/Developer/TraceInsights/Private/Insights/LoadingProfiler/LoadingProfilerCommands.cpp
--- a/Engine/Source/Developer/TraceInsights/Private/Insights/LoadingProfiler/LoadingProfilerCommands.cpp
+++ b/Engine/Source/Developer/TraceInsights/Private/Insights/LoadingProfiler/LoadingProfilerCommands.cpp
@@ -37,12 +37,11 @@ void FLoadingProfilerMenuBuilder::AddMenuEntry(FMenuBuilder& MenuBuilder, const
 ////////////////////////////////////////////////////////////////////////////////////////////////////
 
 FLoadingProfilerCommands::FLoadingProfilerCommands()
-	: TCommands<FLoadingProfilerCommands>(
-		TEXT("LoadingProfilerCommand"), // Context name for fast lookup
-		NSLOCTEXT("Contexts", "LoadingProfilerCommand", "Asset Loading Insights"), // Localized context name for displaying
-		NAME_None, // Parent
-		FInsightsStyle::GetStyleSetName() // Icon Style Set
-	)
+: TCommands<FLoadingProfilerCommands>(
+	TEXT("LoadingProfilerCommands"),
+	NSLOCTEXT("Contexts", "LoadingProfilerCommands", "Insights - Asset Loading Insights"),
+	NAME_None,
+	FInsightsStyle::GetStyleSetName())
 {
 }
 
diff --git a/Engine/Source/Developer/TraceInsights/Private/Insights/MemoryProfiler/MemoryProfilerCommands.cpp b/Engine/Source/Developer/TraceInsights/Private/Insights/MemoryProfiler/MemoryProfilerCommands.cpp
--- a/Engine/Source/Developer/TraceInsights/Private/Insights/MemoryProfiler/MemoryProfilerCommands.cpp
+++ b/Engine/Source/Developer/TraceInsights/Private/Insights/MemoryProfiler/MemoryProfilerCommands.cpp
@@ -37,12 +37,11 @@ void FMemoryProfilerMenuBuilder::AddMenuEntry(FMenuBuilder& MenuBuilder, const T
 ////////////////////////////////////////////////////////////////////////////////////////////////////
 
 FMemoryProfilerCommands::FMemoryProfilerCommands()
-	: TCommands<FMemoryProfilerCommands>(
-		TEXT("MemoryProfilerCommand"), // Context name for fast lookup
-		NSLOCTEXT("Contexts", "MemoryProfilerCommand", "Memory Insights"), // Localized context name for displaying
-		NAME_None, // Parent
-		FInsightsStyle::GetStyleSetName() // Icon Style Set
-	)
+: TCommands<FMemoryProfilerCommands>(
+	TEXT("MemoryProfilerCommands"),
+	NSLOCTEXT("Contexts", "MemoryProfilerCommands", "Insights - Memory Insights"),
+	NAME_None,
+	FInsightsStyle::GetStyleSetName())
 {
 }
 
diff --git a/Engine/Source/Developer/TraceInsights/Private/Insights/NetworkingProfiler/NetworkingProfilerCommands.cpp b/Engine/Source/Developer/TraceInsights/Private/Insights/NetworkingProfiler/NetworkingProfilerCommands.cpp
--- a/Engine/Source/Developer/TraceInsights/Private/Insights/NetworkingProfiler/NetworkingProfilerCommands.cpp
+++ b/Engine/Source/Developer/TraceInsights/Private/Insights/NetworkingProfiler/NetworkingProfilerCommands.cpp
@@ -37,12 +37,11 @@ void FNetworkingProfilerMenuBuilder::AddMenuEntry(FMenuBuilder& MenuBuilder, con
 ////////////////////////////////////////////////////////////////////////////////////////////////////
 
 FNetworkingProfilerCommands::FNetworkingProfilerCommands()
-	: TCommands<FNetworkingProfilerCommands>(
-		TEXT("NetworkingProfilerCommand"), // Context name for fast lookup
-		NSLOCTEXT("Contexts", "NetworkingProfilerCommand", "Networking Insights"), // Localized context name for displaying
-		NAME_None, // Parent
-		FInsightsStyle::GetStyleSetName() // Icon Style Set
-	)
+: TCommands<FNetworkingProfilerCommands>(
+	TEXT("NetworkingProfilerCommands"),
+	NSLOCTEXT("Contexts", "NetworkingProfilerCommands", "Insights - Networking Insights"),
+	NAME_None,
+	FInsightsStyle::GetStyleSetName())
 {
 }
 
diff --git a/Engine/Source/Developer/TraceInsights/Private/Insights/TaskGraphProfiler/ViewModels/TaskTimingTrack.cpp b/Engine/Source/Developer/TraceInsights/Private/Insights/TaskGraphProfiler/ViewModels/TaskTimingTrack.cpp
--- a/Engine/Source/Developer/TraceInsights/Private/Insights/TaskGraphProfiler/ViewModels/TaskTimingTrack.cpp
+++ b/Engine/Source/Developer/TraceInsights/Private/Insights/TaskGraphProfiler/ViewModels/TaskTimingTrack.cpp
@@ -35,11 +35,11 @@ class FTaskTimingStateCommands : public TCommands<FTaskTimingStateCommands>
 {
 public:
 	FTaskTimingStateCommands()
-		: TCommands<FTaskTimingStateCommands>(
-			TEXT("FTaskTimingStateCommands"),
-			NSLOCTEXT("FTaskTimingStateCommands", "Task Timing State Commands", "Task Table Tree View Commands"),
-			NAME_None,
-			FInsightsStyle::GetStyleSetName())
+	: TCommands<FTaskTimingStateCommands>(
+		TEXT("TaskTimingStateCommands"),
+		NSLOCTEXT("Contexts", "TaskTimingStateCommands", "Insights - Task Timing View"),
+		NAME_None,
+		FInsightsStyle::GetStyleSetName())
 	{
 	}
 
diff --git a/Engine/Source/Developer/TraceInsights/Private/Insights/TaskGraphProfiler/Widgets/STaskTableTreeView.cpp b/Engine/Source/Developer/TraceInsights/Private/Insights/TaskGraphProfiler/Widgets/STaskTableTreeView.cpp
--- a/Engine/Source/Developer/TraceInsights/Private/Insights/TaskGraphProfiler/Widgets/STaskTableTreeView.cpp
+++ b/Engine/Source/Developer/TraceInsights/Private/Insights/TaskGraphProfiler/Widgets/STaskTableTreeView.cpp
@@ -38,11 +38,11 @@ class FTaskTableTreeViewCommands : public TCommands<FTaskTableTreeViewCommands>
 {
 public:
 	FTaskTableTreeViewCommands()
-		: TCommands<FTaskTableTreeViewCommands>(
-			TEXT("FTaskTableTreeViewCommands"),
-			NSLOCTEXT("FTaskTableTreeViewCommands", "Task Table Tree View Commands", "Task Table Tree View Commands"),
-			NAME_None,
-			FInsightsStyle::GetStyleSetName())
+	: TCommands<FTaskTableTreeViewCommands>(
+		TEXT("TaskTableTreeViewCommands"),
+		NSLOCTEXT("Contexts", "TaskTableTreeViewCommands", "Insights - Task Table Tree View"),
+		NAME_None,
+		FInsightsStyle::GetStyleSetName())
 	{
 	}
 
diff --git a/Engine/Source/Developer/TraceInsights/Private/Insights/TimingProfilerCommands.cpp b/Engine/Source/Developer/TraceInsights/Private/Insights/TimingProfilerCommands.cpp
--- a/Engine/Source/Developer/TraceInsights/Private/Insights/TimingProfilerCommands.cpp
+++ b/Engine/Source/Developer/TraceInsights/Private/Insights/TimingProfilerCommands.cpp
@@ -37,12 +37,11 @@ void FTimingProfilerMenuBuilder::AddMenuEntry(FMenuBuilder& MenuBuilder, const T
 ////////////////////////////////////////////////////////////////////////////////////////////////////
 
 FTimingProfilerCommands::FTimingProfilerCommands()
-	: TCommands<FTimingProfilerCommands>(
-		TEXT("TimingProfilerCommand"), // Context name for fast lookup
-		NSLOCTEXT("Contexts", "TimingProfilerCommand", "Timing Insights"), // Localized context name for displaying
-		NAME_None, // Parent
-		FInsightsStyle::GetStyleSetName() // Icon Style Set
-	)
+: TCommands<FTimingProfilerCommands>(
+	TEXT("TimingProfilerCommands"),
+	NSLOCTEXT("Contexts", "TimingProfilerCommands", "Insights - Timing Insights"),
+	NAME_None,
+	FInsightsStyle::GetStyleSetName())
 {
 }
 
@@ -68,12 +67,11 @@ PRAGMA_ENABLE_OPTIMIZATION
 ////////////////////////////////////////////////////////////////////////////////////////////////////
 
 FTimingViewCommands::FTimingViewCommands()
-	: TCommands<FTimingViewCommands>(
-		TEXT("TimingViewCommand"), // Context name for fast lookup
-		NSLOCTEXT("Contexts", "TimingViewCommand", "Timing Insights"), // Localized context name for displaying
-		NAME_None, // Parent
-		FInsightsStyle::GetStyleSetName() // Icon Style Set
-	)
+: TCommands<FTimingViewCommands>(
+	TEXT("TimingViewCommands"),
+	NSLOCTEXT("Contexts", "TimingViewCommands", "Insights - Timing View"),
+	NAME_None,
+	FInsightsStyle::GetStyleSetName())
 {
 }
 
@@ -85,13 +83,13 @@ void FTimingViewCommands::RegisterCommands()
 {
 	UI_COMMAND(AutoHideEmptyTracks,
 		"Auto Hide Empty Tracks",
-		"Auto hide empty tracks (ex.: ones without timing events in the current viewport).\nThis option is persistent to UnrealInsightsSettings.ini file.",
+		"Auto hide empty tracks (ex.: ones without timing events in the current viewport).\nThis option is persistent to the UnrealInsightsSettings.ini file.",
 		EUserInterfaceActionType::ToggleButton,
 		FInputChord(EKeys::V));
 
 	UI_COMMAND(PanningOnScreenEdges,
 		"Allow Panning on Screen Edges",
-		"If enabled, the panning is allowed to continue when mouse cursor reaches the edges of the screen.\nThis option is persistent to UnrealInsightsSettings.ini file.",
+		"If enabled, the panning is allowed to continue when the mouse cursor reaches the edges of the screen.\nThis option is persistent to the UnrealInsightsSettings.ini file.",
 		EUserInterfaceActionType::ToggleButton,
 		FInputChord());
 
