[Insights] More minor changes to UI (sections in context menus, tooltip texts, session info).

diff --git a/Engine/Source/Developer/TraceInsights/Private/Insights/ContextSwitches/ViewModels/ContextSwitchesSharedState.cpp b/Engine/Source/Developer/TraceInsights/Private/Insights/ContextSwitches/ViewModels/ContextSwitchesSharedState.cpp
--- a/Engine/Source/Developer/TraceInsights/Private/Insights/ContextSwitches/ViewModels/ContextSwitchesSharedState.cpp
+++ b/Engine/Source/Developer/TraceInsights/Private/Insights/ContextSwitches/ViewModels/ContextSwitchesSharedState.cpp
@@ -243,7 +243,7 @@ void FContextSwitchesSharedState::ExtendCpuTracksFilterMenu(Insights::ITimingVie
 
 void FContextSwitchesSharedState::BuildSubMenu(FMenuBuilder& InMenuBuilder)
 {
-	InMenuBuilder.BeginSection("ContextSwitches", LOCTEXT("ContextSwitchesHeading", "Context Switches"));
+	InMenuBuilder.BeginSection("ContextSwitches", LOCTEXT("ContextMenu_Section_ContextSwitches", "Context Switches"));
 	{
 		InMenuBuilder.AddMenuEntry(FContextSwitchesStateCommands::Get().Command_ShowCoreTracks);
 		InMenuBuilder.AddMenuEntry(FContextSwitchesStateCommands::Get().Command_ShowContextSwitches);
diff --git a/Engine/Source/Developer/TraceInsights/Private/Insights/MemoryProfiler/ViewModels/MemorySharedState.cpp b/Engine/Source/Developer/TraceInsights/Private/Insights/MemoryProfiler/ViewModels/MemorySharedState.cpp
--- a/Engine/Source/Developer/TraceInsights/Private/Insights/MemoryProfiler/ViewModels/MemorySharedState.cpp
+++ b/Engine/Source/Developer/TraceInsights/Private/Insights/MemoryProfiler/ViewModels/MemorySharedState.cpp
@@ -353,11 +353,11 @@ void FMemorySharedState::ExtendFilterMenu(Insights::ITimingViewSession& InSessio
 		return;
 	}
 
-	InOutMenuBuilder.BeginSection("Memory", LOCTEXT("MemoryHeading", "Memory"));
+	InOutMenuBuilder.BeginSection("Memory", LOCTEXT("ContextMenu_Section_Memory", "Memory"));
 	{
 		InOutMenuBuilder.AddMenuEntry(
 			LOCTEXT("AllMemoryTracks", "Memory Tracks - M"),
-			LOCTEXT("AllMemoryTracks_Tooltip", "Show/hide the Memory tracks"),
+			LOCTEXT("AllMemoryTracks_Tooltip", "Shows or hides the Memory tracks."),
 			FSlateIcon(),
 			FUIAction(FExecuteAction::CreateSP(this, &FMemorySharedState::ShowHideAllMemoryTracks),
 					  FCanExecuteAction(),
diff --git a/Engine/Source/Developer/TraceInsights/Private/Insights/MemoryProfiler/Widgets/SMemAllocTableTreeView.cpp b/Engine/Source/Developer/TraceInsights/Private/Insights/MemoryProfiler/Widgets/SMemAllocTableTreeView.cpp
--- a/Engine/Source/Developer/TraceInsights/Private/Insights/MemoryProfiler/Widgets/SMemAllocTableTreeView.cpp
+++ b/Engine/Source/Developer/TraceInsights/Private/Insights/MemoryProfiler/Widgets/SMemAllocTableTreeView.cpp
@@ -1220,14 +1220,14 @@ ECheckBoxState SMemAllocTableTreeView::CallstackGroupingByFunction_IsChecked() c
 
 void SMemAllocTableTreeView::ExtendMenu(FMenuBuilder& MenuBuilder)
 {
-	MenuBuilder.BeginSection("CallstackFrame", LOCTEXT("ContextMenu_Header_CallstackFrame", "Callstack Frame"));
+	MenuBuilder.BeginSection("CallstackFrame", LOCTEXT("ContextMenu_Section_CallstackFrame", "Callstack Frame"));
 	{
 		ISourceCodeAccessModule& SourceCodeAccessModule = FModuleManager::LoadModuleChecked<ISourceCodeAccessModule>("SourceCodeAccess");
 		ISourceCodeAccessor& SourceCodeAccessor = SourceCodeAccessModule.GetAccessor();
 
-		FText ItemLabel = FText::Format(LOCTEXT("ContextMenu_Header_Open", "Open in {0}"), SourceCodeAccessor.GetNameText());
+		FText ItemLabel = FText::Format(LOCTEXT("ContextMenu_Open", "Open in {0}"), SourceCodeAccessor.GetNameText());
 		FText FileName = GetSelectedCallstackFrameFileName();
-		FText ItemToolTip = FText::Format(LOCTEXT("ContextMenu_Header_Open_Desc", "Open source file of selected callstack frame in {0}.\n{1}"), SourceCodeAccessor.GetNameText(), FileName);
+		FText ItemToolTip = FText::Format(LOCTEXT("ContextMenu_Open_Desc", "Open source file of selected callstack frame in {0}.\n{1}"), SourceCodeAccessor.GetNameText(), FileName);
 
 		FUIAction Action_OpenIDE
 		(
diff --git a/Engine/Source/Developer/TraceInsights/Private/Insights/TaskGraphProfiler/ViewModels/TaskTimingTrack.cpp b/Engine/Source/Developer/TraceInsights/Private/Insights/TaskGraphProfiler/ViewModels/TaskTimingTrack.cpp
--- a/Engine/Source/Developer/TraceInsights/Private/Insights/TaskGraphProfiler/ViewModels/TaskTimingTrack.cpp
+++ b/Engine/Source/Developer/TraceInsights/Private/Insights/TaskGraphProfiler/ViewModels/TaskTimingTrack.cpp
@@ -180,7 +180,12 @@ bool FTaskTimingSharedState::ExtendGlobalContextMenu(ITimingViewSession& InSessi
 		return false;
 	}
 
-	InMenuBuilder.AddSubMenu(LOCTEXT("Tasks", "Tasks"), LOCTEXT("Task", "Task Graph Insights settings"), FNewMenuDelegate::CreateSP(this, &FTaskTimingSharedState::BuildTasksSubMenu));
+	InMenuBuilder.AddSubMenu
+	(
+		LOCTEXT("ContextMenu_Tasks_SubMenu", "Tasks"),
+		LOCTEXT("ContextMenu_Tasks_SubMenu_Desc", "Task Graph Insights settings"),
+		FNewMenuDelegate::CreateSP(this, &FTaskTimingSharedState::BuildTasksSubMenu)
+	);
 
 	return true;
 }
diff --git a/Engine/Source/Developer/TraceInsights/Private/Insights/TaskGraphProfiler/Widgets/STaskTableTreeView.cpp b/Engine/Source/Developer/TraceInsights/Private/Insights/TaskGraphProfiler/Widgets/STaskTableTreeView.cpp
--- a/Engine/Source/Developer/TraceInsights/Private/Insights/TaskGraphProfiler/Widgets/STaskTableTreeView.cpp
+++ b/Engine/Source/Developer/TraceInsights/Private/Insights/TaskGraphProfiler/Widgets/STaskTableTreeView.cpp
@@ -90,7 +90,7 @@ void STaskTableTreeView::Construct(const FArguments& InArgs, TSharedPtr<Insights
 
 void STaskTableTreeView::ExtendMenu(FMenuBuilder& MenuBuilder)
 {
-	MenuBuilder.BeginSection("Node", LOCTEXT("ContextMenu_Header_Task", "Task"));
+	MenuBuilder.BeginSection("Node", LOCTEXT("ContextMenu_Section_Task", "Task"));
 
 	{
 		MenuBuilder.AddMenuEntry
diff --git a/Engine/Source/Developer/TraceInsights/Private/Insights/ViewModels/FileActivityTimingTrack.cpp b/Engine/Source/Developer/TraceInsights/Private/Insights/ViewModels/FileActivityTimingTrack.cpp
--- a/Engine/Source/Developer/TraceInsights/Private/Insights/ViewModels/FileActivityTimingTrack.cpp
+++ b/Engine/Source/Developer/TraceInsights/Private/Insights/ViewModels/FileActivityTimingTrack.cpp
@@ -368,11 +368,11 @@ void FFileActivitySharedState::ExtendOtherTracksFilterMenu(Insights::ITimingView
 
 void FFileActivitySharedState::BuildSubMenu(FMenuBuilder& InOutMenuBuilder)
 {
-	InOutMenuBuilder.BeginSection("File Activity", LOCTEXT("FileActivityHeading", "File Activity"));
+	InOutMenuBuilder.BeginSection("File Activity", LOCTEXT("ContextMenu_Section_FileActivity", "File Activity"));
 	{
 		InOutMenuBuilder.AddMenuEntry(
 			LOCTEXT("ShowIoOverviewTrack", "I/O Overview Track - I"),
-			LOCTEXT("ShowIoOverviewTrack_Tooltip", "Show/hide the I/O Overview track."),
+			LOCTEXT("ShowIoOverviewTrack_Tooltip", "Shows or hides the I/O Overview track."),
 			FSlateIcon(),
 			FUIAction(FExecuteAction::CreateSP(this, &FFileActivitySharedState::ShowHideIoOverviewTrack),
 				FCanExecuteAction(),
@@ -383,7 +383,7 @@ void FFileActivitySharedState::BuildSubMenu(FMenuBuilder& InOutMenuBuilder)
 
 		InOutMenuBuilder.AddMenuEntry(
 			LOCTEXT("ShowOnlyErrors", "Only Errors (I/O Overview Track)"),
-			LOCTEXT("ShowOnlyErrors_Tooltip", "Show only the events with errors, in the I/O Overview track."),
+			LOCTEXT("ShowOnlyErrors_Tooltip", "Shows only the events with errors, in the I/O Overview track."),
 			FSlateIcon(),
 			FUIAction(FExecuteAction::CreateSP(this, &FFileActivitySharedState::ToggleOnlyErrors),
 				FCanExecuteAction(),
@@ -394,7 +394,7 @@ void FFileActivitySharedState::BuildSubMenu(FMenuBuilder& InOutMenuBuilder)
 
 		InOutMenuBuilder.AddMenuEntry(
 			LOCTEXT("ShowIoActivityTrack", "I/O Activity Track - I"),
-			LOCTEXT("ShowIoActivityTrack_Tooltip", "Show/hide the I/O Activity track."),
+			LOCTEXT("ShowIoActivityTrack_Tooltip", "Shows or hides the I/O Activity track."),
 			FSlateIcon(),
 			FUIAction(FExecuteAction::CreateSP(this, &FFileActivitySharedState::ShowHideIoActivityTrack),
 				FCanExecuteAction(),
@@ -405,7 +405,7 @@ void FFileActivitySharedState::BuildSubMenu(FMenuBuilder& InOutMenuBuilder)
 
 		InOutMenuBuilder.AddMenuEntry(
 			LOCTEXT("ShowBackgroundEvents", "Background Events (I/O Activity Track) - O"),
-			LOCTEXT("ShowBackgroundEvents_Tooltip", "Show background events for file activities, in the I/O Activity track."),
+			LOCTEXT("ShowBackgroundEvents_Tooltip", "Shows background events for file activities, in the I/O Activity track."),
 			FSlateIcon(),
 			FUIAction(FExecuteAction::CreateSP(this, &FFileActivitySharedState::ToggleBackgroundEvents),
 				FCanExecuteAction(),
diff --git a/Engine/Source/Developer/TraceInsights/Private/Insights/ViewModels/FrameTimingTrack.cpp b/Engine/Source/Developer/TraceInsights/Private/Insights/ViewModels/FrameTimingTrack.cpp
--- a/Engine/Source/Developer/TraceInsights/Private/Insights/ViewModels/FrameTimingTrack.cpp
+++ b/Engine/Source/Developer/TraceInsights/Private/Insights/ViewModels/FrameTimingTrack.cpp
@@ -123,7 +123,7 @@ void FFrameSharedState::ExtendOtherTracksFilterMenu(Insights::ITimingViewSession
 		return;
 	}
 
-	InOutMenuBuilder.BeginSection("FrameTracks", LOCTEXT("FrameTracksHeading", "Frames"));
+	InOutMenuBuilder.BeginSection("FrameTracks", LOCTEXT("ContextMenu_Section_FrameTracks", "Frames"));
 	{
 		//TODO: MenuBuilder.AddMenuEntry(Commands.ShowAllFrameTracks);
 		InOutMenuBuilder.AddMenuEntry(
diff --git a/Engine/Source/Developer/TraceInsights/Private/Insights/ViewModels/GraphTrack.cpp b/Engine/Source/Developer/TraceInsights/Private/Insights/ViewModels/GraphTrack.cpp
--- a/Engine/Source/Developer/TraceInsights/Private/Insights/ViewModels/GraphTrack.cpp
+++ b/Engine/Source/Developer/TraceInsights/Private/Insights/ViewModels/GraphTrack.cpp
@@ -650,7 +650,7 @@ const TSharedPtr<const ITimingEvent> FGraphTrack::GetEvent(float InPosX, float I
 
 void FGraphTrack::BuildContextMenu(FMenuBuilder& MenuBuilder)
 {
-	MenuBuilder.BeginSection("Options", LOCTEXT("OptionsHeading", "Options"));
+	MenuBuilder.BeginSection("Options", LOCTEXT("ContextMenu_Section_Options", "Options"));
 	{
 		if (EnumHasAnyFlags(VisibleOptions, EGraphOptions::ShowDebugInfo)) // debug functionality
 		{
@@ -663,7 +663,7 @@ void FGraphTrack::BuildContextMenu(FMenuBuilder& MenuBuilder)
 			MenuBuilder.AddMenuEntry
 			(
 				LOCTEXT("ContextMenu_ShowDebugInfo", "Show Debug Info"),
-				LOCTEXT("ContextMenu_ShowDebugInfo_Desc", "Show Debug Info."),
+				LOCTEXT("ContextMenu_ShowDebugInfo_Desc", "Shows debug info."),
 				FSlateIcon(),
 				Action_ShowDebugInfo,
 				NAME_None,
@@ -682,7 +682,7 @@ void FGraphTrack::BuildContextMenu(FMenuBuilder& MenuBuilder)
 			MenuBuilder.AddMenuEntry
 			(
 				LOCTEXT("ContextMenu_ShowPoints", "Show Points"),
-				LOCTEXT("ContextMenu_ShowPoints_Desc", "Show points."),
+				LOCTEXT("ContextMenu_ShowPoints_Desc", "Shows points."),
 				FSlateIcon(),
 				Action_ShowPoints,
 				NAME_None,
@@ -701,7 +701,7 @@ void FGraphTrack::BuildContextMenu(FMenuBuilder& MenuBuilder)
 			MenuBuilder.AddMenuEntry
 			(
 				LOCTEXT("ContextMenu_ShowPointsWithBorder", "Show Points with Border"),
-				LOCTEXT("ContextMenu_ShowPointsWithBorder_Desc", "Show border around points."),
+				LOCTEXT("ContextMenu_ShowPointsWithBorder_Desc", "Shows border around points."),
 				FSlateIcon(),
 				Action_ShowPointsWithBorder,
 				NAME_None,
@@ -720,7 +720,7 @@ void FGraphTrack::BuildContextMenu(FMenuBuilder& MenuBuilder)
 			MenuBuilder.AddMenuEntry
 			(
 				LOCTEXT("ContextMenu_ShowLines", "Show Connected Lines"),
-				LOCTEXT("ContextMenu_ShowLines_Desc", "Show connected lines. Each event is a single point in time."),
+				LOCTEXT("ContextMenu_ShowLines_Desc", "Shows connected lines. Each event is a single point in time."),
 				FSlateIcon(),
 				Action_ShowLines,
 				NAME_None,
@@ -739,7 +739,7 @@ void FGraphTrack::BuildContextMenu(FMenuBuilder& MenuBuilder)
 			MenuBuilder.AddMenuEntry
 			(
 				LOCTEXT("ContextMenu_ShowPolygon", "Show Polygon"),
-				LOCTEXT("ContextMenu_ShowPolygon_Desc", "Show filled polygon under the graph series."),
+				LOCTEXT("ContextMenu_ShowPolygon_Desc", "Shows filled polygon under the graph series."),
 				FSlateIcon(),
 				Action_ShowPolygon,
 				NAME_None,
@@ -758,7 +758,7 @@ void FGraphTrack::BuildContextMenu(FMenuBuilder& MenuBuilder)
 			MenuBuilder.AddMenuEntry
 			(
 				LOCTEXT("ContextMenu_UseEventDuration", "Use Event Duration"),
-				LOCTEXT("ContextMenu_UseEventDuration_Desc", "Use duration of timing events (for Connected Lines and Polygon)."),
+				LOCTEXT("ContextMenu_UseEventDuration_Desc", "Uses duration of timing events (for Connected Lines and Polygon)."),
 				FSlateIcon(),
 				Action_UseEventDuration,
 				NAME_None,
@@ -777,7 +777,7 @@ void FGraphTrack::BuildContextMenu(FMenuBuilder& MenuBuilder)
 			MenuBuilder.AddMenuEntry
 			(
 				LOCTEXT("ContextMenu_ShowBars", "Show Bars"),
-				LOCTEXT("ContextMenu_ShowBars_Desc", "Show bars. Width of bars corresponds to duration of timing events."),
+				LOCTEXT("ContextMenu_ShowBars_Desc", "Shows bars. Width of bars corresponds to duration of timing events."),
 				FSlateIcon(),
 				Action_ShowBars,
 				NAME_None,
@@ -787,7 +787,7 @@ void FGraphTrack::BuildContextMenu(FMenuBuilder& MenuBuilder)
 	}
 	MenuBuilder.EndSection();
 
-	MenuBuilder.BeginSection("Series", LOCTEXT("ContextMenu_Header_Series", "Series"));
+	MenuBuilder.BeginSection("Series", LOCTEXT("ContextMenu_Section_Series", "Series"));
 	{
 		MenuBuilder.AddWidget(
 			SNew(SBox)
diff --git a/Engine/Source/Developer/TraceInsights/Private/Insights/ViewModels/LoadingTimingTrack.cpp b/Engine/Source/Developer/TraceInsights/Private/Insights/ViewModels/LoadingTimingTrack.cpp
--- a/Engine/Source/Developer/TraceInsights/Private/Insights/ViewModels/LoadingTimingTrack.cpp
+++ b/Engine/Source/Developer/TraceInsights/Private/Insights/ViewModels/LoadingTimingTrack.cpp
@@ -118,11 +118,11 @@ void FLoadingSharedState::ExtendOtherTracksFilterMenu(Insights::ITimingViewSessi
 		return;
 	}
 
-	InOutMenuBuilder.BeginSection("Asset Loading", LOCTEXT("AssetLoadingHeading", "Asset Loading"));
+	InOutMenuBuilder.BeginSection("Asset Loading", LOCTEXT("ContextMenu_Section_AssetLoading", "Asset Loading"));
 	{
 		InOutMenuBuilder.AddMenuEntry(
 			LOCTEXT("ShowAllLoadingTracks", "Asset Loading Tracks - L"),
-			LOCTEXT("ShowAllLoadingTracks_Tooltip", "Show/hide the Asset Loading tracks"),
+			LOCTEXT("ShowAllLoadingTracks_Tooltip", "Shows or hides the Asset Loading tracks."),
 			FSlateIcon(),
 			FUIAction(FExecuteAction::CreateSP(this, &FLoadingSharedState::ShowHideAllLoadingTracks),
 					  FCanExecuteAction(),
diff --git a/Engine/Source/Developer/TraceInsights/Private/Insights/ViewModels/MarkersTimingTrack.cpp b/Engine/Source/Developer/TraceInsights/Private/Insights/ViewModels/MarkersTimingTrack.cpp
--- a/Engine/Source/Developer/TraceInsights/Private/Insights/ViewModels/MarkersTimingTrack.cpp
+++ b/Engine/Source/Developer/TraceInsights/Private/Insights/ViewModels/MarkersTimingTrack.cpp
@@ -258,11 +258,11 @@ FReply FMarkersTimingTrack::OnMouseButtonDoubleClick(const FGeometry& MyGeometry
 
 void FMarkersTimingTrack::BuildContextMenu(FMenuBuilder& MenuBuilder)
 {
-	MenuBuilder.BeginSection(TEXT("Content"), LOCTEXT("ContentHeading", "Content"));
+	MenuBuilder.BeginSection(TEXT("Content"), LOCTEXT("ContextMenu_Section_Content", "Content"));
 	{
 		MenuBuilder.AddMenuEntry(
 			LOCTEXT("ContextMenu_Bookmarks", "Bookmarks"),
-			LOCTEXT("ContextMenu_Bookmarks_Desc", "Change this track to show only the bookmarks."),
+			LOCTEXT("ContextMenu_Bookmarks_Desc", "Changes this track to show only the bookmarks."),
 			FSlateIcon(),
 			FUIAction(FExecuteAction::CreateSP(this, &FMarkersTimingTrack::SetBookmarksTrack),
 						FCanExecuteAction(),
@@ -273,7 +273,7 @@ void FMarkersTimingTrack::BuildContextMenu(FMenuBuilder& MenuBuilder)
 
 		MenuBuilder.AddMenuEntry(
 			LOCTEXT("ContextMenu_Logs", "Logs"),
-			LOCTEXT("ContextMenu_Logs_Desc", "Change this track to show all logs."),
+			LOCTEXT("ContextMenu_Logs_Desc", "Changes this track to show all logs."),
 			FSlateIcon(),
 			FUIAction(FExecuteAction::CreateSP(this, &FMarkersTimingTrack::SetLogsTrack),
 					  FCanExecuteAction(),
@@ -284,7 +284,7 @@ void FMarkersTimingTrack::BuildContextMenu(FMenuBuilder& MenuBuilder)
 	}
 	MenuBuilder.EndSection();
 
-	MenuBuilder.BeginSection(TEXT("MarkerLines"), LOCTEXT("MarkerLinesHeading", "Marker Lines"));
+	MenuBuilder.BeginSection(TEXT("MarkerLines"), LOCTEXT("ContextMenu_Section_MarkerLines", "Marker Lines"));
 	{
 		MenuBuilder.AddMenuEntry(
 			LOCTEXT("ContextMenu_ToggleCollapsed", "Collapsed"),
diff --git a/Engine/Source/Developer/TraceInsights/Private/Insights/ViewModels/ThreadTimingTrack.cpp b/Engine/Source/Developer/TraceInsights/Private/Insights/ViewModels/ThreadTimingTrack.cpp
--- a/Engine/Source/Developer/TraceInsights/Private/Insights/ViewModels/ThreadTimingTrack.cpp
+++ b/Engine/Source/Developer/TraceInsights/Private/Insights/ViewModels/ThreadTimingTrack.cpp
@@ -545,7 +545,7 @@ void FThreadTimingSharedState::ExtendGpuTracksFilterMenu(Insights::ITimingViewSe
 		return;
 	}
 
-	InOutMenuBuilder.BeginSection("GpuTracks", LOCTEXT("GpuTracksHeading", "GPU Tracks"));
+	InOutMenuBuilder.BeginSection("GpuTracks", LOCTEXT("ContextMenu_Section_GpuTracks", "GPU Tracks"));
 	{
 		//TODO: MenuBuilder.AddMenuEntry(Commands.ShowAllGpuTracks);
 		InOutMenuBuilder.AddMenuEntry(
@@ -571,7 +571,7 @@ void FThreadTimingSharedState::ExtendCpuTracksFilterMenu(Insights::ITimingViewSe
 		return;
 	}
 
-	InOutMenuBuilder.BeginSection("CpuTracks", LOCTEXT("CpuTracksHeading", "CPU Tracks"));
+	InOutMenuBuilder.BeginSection("CpuTracks", LOCTEXT("ContextMenu_Section_CpuTracks", "CPU Tracks"));
 	{
 		//TODO: MenuBuilder.AddMenuEntry(Commands.ShowAllCpuTracks);
 		InOutMenuBuilder.AddMenuEntry(
@@ -587,7 +587,7 @@ void FThreadTimingSharedState::ExtendCpuTracksFilterMenu(Insights::ITimingViewSe
 	}
 	InOutMenuBuilder.EndSection();
 
-	InOutMenuBuilder.BeginSection("ThreadGroups", LOCTEXT("ThreadGroupsHeading", "CPU Thread Groups"));
+	InOutMenuBuilder.BeginSection("CpuThreadGroups", LOCTEXT("ContextMenu_Section_CpuThreadGroups", "CPU Thread Groups"));
 	CreateThreadGroupsMenu(InOutMenuBuilder);
 	InOutMenuBuilder.EndSection();
 }
@@ -1447,7 +1447,7 @@ void FThreadTimingTrack::BuildContextMenu(FMenuBuilder& MenuBuilder)
 {
 	if (GetGroupName() != nullptr)
 	{
-		MenuBuilder.BeginSection("CpuThread", LOCTEXT("CpuThreadSection", "CPU Thread"));
+		MenuBuilder.BeginSection("CpuThread", LOCTEXT("ContextMenu_Section_CpuThread", "CPU Thread"));
 		{
 			MenuBuilder.AddMenuEntry(
 				FText::Format(LOCTEXT("CpuThreadGroupFmt", "Group: {0}"), FText::FromString(GetGroupName())),
@@ -1476,7 +1476,7 @@ void FThreadTimingTrack::BuildContextMenu(FMenuBuilder& MenuBuilder)
 		ChildTrack->BuildContextMenu(MenuBuilder);
 	}
 
-	MenuBuilder.BeginSection("TimingEvents", LOCTEXT("TimingEventsSection", "Timing Events"));
+	MenuBuilder.BeginSection("TimingEvents", LOCTEXT("ContextMenu_Section_TimingEvents", "Timing Events"));
 	{
 		FExecuteAction FilterTrackAction;
 		FilterTrackAction.BindSP(this, &FThreadTimingTrack::OnFilterTrackClicked);
diff --git a/Engine/Source/Developer/TraceInsights/Private/Insights/ViewModels/TimeRulerTrack.cpp b/Engine/Source/Developer/TraceInsights/Private/Insights/ViewModels/TimeRulerTrack.cpp
--- a/Engine/Source/Developer/TraceInsights/Private/Insights/ViewModels/TimeRulerTrack.cpp
+++ b/Engine/Source/Developer/TraceInsights/Private/Insights/ViewModels/TimeRulerTrack.cpp
@@ -440,7 +440,7 @@ void FTimeRulerTrack::BuildContextMenu(FMenuBuilder& MenuBuilder)
 		// Sort TimeMarkers by name.
 		VisibleTimeMarkers.Sort([](const TSharedRef<Insights::FTimeMarker>& A, const TSharedRef<Insights::FTimeMarker>& B) -> bool { return A->GetName().Compare(B->GetName()) <= 0; });
 
-		MenuBuilder.BeginSection("TimeMarkers", LOCTEXT("TimeMarkersSection", "Time Markers"));
+		MenuBuilder.BeginSection("TimeMarkers", LOCTEXT("ContextMenu_Section_TimeMarkers", "Time Markers"));
 
 		for (TSharedRef<Insights::FTimeMarker>& TimeMarker : VisibleTimeMarkers)
 		{
diff --git a/Engine/Source/Developer/TraceInsights/Private/Insights/Widgets/SFrameTrack.cpp b/Engine/Source/Developer/TraceInsights/Private/Insights/Widgets/SFrameTrack.cpp
--- a/Engine/Source/Developer/TraceInsights/Private/Insights/Widgets/SFrameTrack.cpp
+++ b/Engine/Source/Developer/TraceInsights/Private/Insights/Widgets/SFrameTrack.cpp
@@ -1110,7 +1110,7 @@ void SFrameTrack::ShowContextMenu(const FPointerEvent& MouseEvent)
 	const bool bShouldCloseWindowAfterMenuSelection = true;
 	FMenuBuilder MenuBuilder(bShouldCloseWindowAfterMenuSelection, NULL);
 
-	MenuBuilder.BeginSection("Frames", LOCTEXT("FramesHeading", "Frames"));
+	MenuBuilder.BeginSection("Frames", LOCTEXT("ContextMenu_Section_Frames", "Frames"));
 	{
 		struct FLocal
 		{
@@ -1129,7 +1129,7 @@ void SFrameTrack::ShowContextMenu(const FPointerEvent& MouseEvent)
 		MenuBuilder.AddMenuEntry
 		(
 			LOCTEXT("ContextMenu_ShowGameFrames", "Game Frames"),
-			LOCTEXT("ContextMenu_ShowGameFrames_Desc", "Show/hide the Game frames."),
+			LOCTEXT("ContextMenu_ShowGameFrames_Desc", "Shows or hides the Game frames."),
 			FSlateIcon(),
 			Action_ShowGameFrames,
 			NAME_None,
@@ -1145,7 +1145,7 @@ void SFrameTrack::ShowContextMenu(const FPointerEvent& MouseEvent)
 		MenuBuilder.AddMenuEntry
 		(
 			LOCTEXT("ContextMenu_ShowRenderingFrames", "Rendering Frames"),
-			LOCTEXT("ContextMenu_ShowRenderingFrames_Desc", "Show/hide the Rendering frames."),
+			LOCTEXT("ContextMenu_ShowRenderingFrames_Desc", "Shows or hides the Rendering frames."),
 			FSlateIcon(),
 			Action_ShowRenderingFrames,
 			NAME_None,
@@ -1154,7 +1154,7 @@ void SFrameTrack::ShowContextMenu(const FPointerEvent& MouseEvent)
 	}
 	MenuBuilder.EndSection();
 
-	MenuBuilder.BeginSection("Zoom", LOCTEXT("ZoomHeading", "Zoom"));
+	MenuBuilder.BeginSection("Zoom", LOCTEXT("ContextMenu_Section_Zoom", "Zoom"));
 	{
 		FUIAction Action_AutoZoom
 		(
diff --git a/Engine/Source/Developer/TraceInsights/Private/Insights/Widgets/SLogView.cpp b/Engine/Source/Developer/TraceInsights/Private/Insights/Widgets/SLogView.cpp
--- a/Engine/Source/Developer/TraceInsights/Private/Insights/Widgets/SLogView.cpp
+++ b/Engine/Source/Developer/TraceInsights/Private/Insights/Widgets/SLogView.cpp
@@ -1191,17 +1191,17 @@ TSharedPtr<SWidget> SLogView::ListView_GetContextMenu()
 			ISourceCodeAccessModule& SourceCodeAccessModule = FModuleManager::LoadModuleChecked<ISourceCodeAccessModule>("SourceCodeAccess");
 			ISourceCodeAccessor& SourceCodeAccessor = SourceCodeAccessModule.GetAccessor();
 
-			const FText ItemLabel = FText::Format(LOCTEXT("ContextMenu_Header_OpenSource", "Open Source in {0}"), SourceCodeAccessor.GetNameText());
+			const FText ItemLabel = FText::Format(LOCTEXT("ContextMenu_OpenSource", "Open Source in {0}"), SourceCodeAccessor.GetNameText());
 
 			FText ItemToolTip;
 			if (!File.IsEmpty())
 			{
-				ItemToolTip = FText::Format(LOCTEXT("ContextMenu_Header_OpenSource_Desc1", "Opens the source file of the selected message in {0}.\n{1} ({2})"),
+				ItemToolTip = FText::Format(LOCTEXT("ContextMenu_OpenSource_Desc1", "Opens the source file of the selected message in {0}.\n{1} ({2})"),
 					SourceCodeAccessor.GetNameText(), FText::FromString(File), FText::AsNumber(Line, &FNumberFormattingOptions::DefaultNoGrouping()));
 			}
 			else
 			{
-				ItemToolTip = FText::Format(LOCTEXT("ContextMenu_Header_OpenSource_Desc2", "Opens the source file of the selected message in {0}."),
+				ItemToolTip = FText::Format(LOCTEXT("ContextMenu_OpenSource_Desc2", "Opens the source file of the selected message in {0}."),
 					SourceCodeAccessor.GetNameText());
 			}
 
@@ -1225,7 +1225,7 @@ TSharedRef<SWidget> SLogView::MakeVerbosityThresholdMenu()
 {
 	FMenuBuilder MenuBuilder(/*bInShouldCloseWindowAfterMenuSelection=*/true, nullptr);
 
-	MenuBuilder.BeginSection("LogViewVerbosityThreshold"/*, LOCTEXT("VerbosityThresholdHeading", "Verbosity Threshold")*/);
+	MenuBuilder.BeginSection("VerbosityThreshold"/*, LOCTEXT("LogVerbosityThresholdMenu_Section_VerbosityThreshold", "Verbosity Threshold")*/);
 	CreateVerbosityThresholdMenuSection(MenuBuilder);
 	MenuBuilder.EndSection();
 
@@ -1322,7 +1322,7 @@ TSharedRef<SWidget> SLogView::MakeCategoryFilterMenu()
 {
 	FMenuBuilder MenuBuilder(/*bInShouldCloseWindowAfterMenuSelection=*/true, nullptr);
 
-	MenuBuilder.BeginSection("LogViewShowAllCategories", LOCTEXT("QuickFilterHeading", "Quick Filter"));
+	MenuBuilder.BeginSection("QuickFilter", LOCTEXT("CategoryFilterMenu_Section_QuickFilter", "Quick Filter"));
 	{
 		MenuBuilder.AddMenuEntry(
 			LOCTEXT("ShowAllCategories", "Show/Hide All"),
@@ -1337,7 +1337,7 @@ TSharedRef<SWidget> SLogView::MakeCategoryFilterMenu()
 	}
 	MenuBuilder.EndSection();
 
-	MenuBuilder.BeginSection("LogViewCategoriesEntries", LOCTEXT("CategoriesHeading", "Categories"));
+	MenuBuilder.BeginSection("Categories", LOCTEXT("CategoryFilterMenu_Section_Categories", "Categories"));
 	CreateCategoriesFilterMenuSection(MenuBuilder);
 	MenuBuilder.EndSection();
 
diff --git a/Engine/Source/Developer/TraceInsights/Private/Insights/Widgets/SSessionInfoWindow.cpp b/Engine/Source/Developer/TraceInsights/Private/Insights/Widgets/SSessionInfoWindow.cpp
--- a/Engine/Source/Developer/TraceInsights/Private/Insights/Widgets/SSessionInfoWindow.cpp
+++ b/Engine/Source/Developer/TraceInsights/Private/Insights/Widgets/SSessionInfoWindow.cpp
@@ -143,11 +143,11 @@ void SSessionInfoWindow::Construct(const FArguments& InArgs, const TSharedRef<SD
 
 ////////////////////////////////////////////////////////////////////////////////////////////////////
 
-void SSessionInfoWindow::AddSectionLine(TSharedPtr<SVerticalBox> InVerticalBox, const FText& InSectionName) const
+void SSessionInfoWindow::BeginSection(TSharedPtr<SVerticalBox> InVerticalBox, const FText& InSectionName) const
 {
 	InVerticalBox->AddSlot()
 		.AutoHeight()
-		.Padding(0.0f, 4.0f, 0.0f, 4.0f)
+		.Padding(0.0f, 0.0f, 0.0f, 4.0f)
 		[
 			SNew(SBorder)
 			.BorderBackgroundColor(FLinearColor(0.03f, 0.03f, 0.03f, 1.0f))
@@ -167,6 +167,20 @@ void SSessionInfoWindow::AddSectionLine(TSharedPtr<SVerticalBox> InVerticalBox,
 
 ////////////////////////////////////////////////////////////////////////////////////////////////////
 
+void SSessionInfoWindow::EndSection(TSharedPtr<SVerticalBox> InVerticalBox) const
+{
+	InVerticalBox->AddSlot()
+		.AutoHeight()
+		.Padding(0.0f, 0.0f, 0.0f, 0.0f)
+		[
+			SNew(SBox)
+			.Padding(FMargin(0.0f))
+			.HeightOverride(4.0f)
+		];
+}
+
+////////////////////////////////////////////////////////////////////////////////////////////////////
+
 TSharedRef<SWidget> SSessionInfoWindow::CreateTextBox(const TAttribute<FText>& InText, bool bMultiLine) const
 {
 	TSharedPtr<SWidget> TextBox;
@@ -291,7 +305,7 @@ TSharedRef<SDockTab> SSessionInfoWindow::SpawnTab_SessionInfo(const FSpawnTabArg
 
 	Image->SetImage(new FSlateColorBrush(FLinearColor(0.015f, 0.015f, 0.015f, 1.0f)));
 
-	AddSectionLine(VerticalBox, LOCTEXT("SessionInfo_SectionText", "Session Info"));
+	BeginSection(VerticalBox, LOCTEXT("SessionInfo_SectionText", "Session Info"));
 	AddInfoLine(VerticalBox, LOCTEXT("SessionName_HeaderText",	"Session Name"),		TAttribute<FText>(this, &SSessionInfoWindow::GetSessionNameText));
 	AddInfoLine(VerticalBox, LOCTEXT("Uri_HeaderText",			"URI"),					TAttribute<FText>(this, &SSessionInfoWindow::GetUriText));
 	//AddInfoLine(VerticalBox, LOCTEXT("FileSize_HeaderText",		"File Size"),			TAttribute<FText>(this, &SSessionInfoWindow::GetFileSizeText));
@@ -303,12 +317,15 @@ TSharedRef<SDockTab> SSessionInfoWindow::SpawnTab_SessionInfo(const FSpawnTabArg
 	AddInfoLine(VerticalBox, LOCTEXT("BuildConfig_HeaderText",	"Build Config"),		TAttribute<FText>(this, &SSessionInfoWindow::GetBuildConfigText));
 	AddInfoLine(VerticalBox, LOCTEXT("BuildTarget_HeaderText",	"Build Target"),		TAttribute<FText>(this, &SSessionInfoWindow::GetBuildTargetText));
 	AddInfoLine(VerticalBox, LOCTEXT("CommandLine_HeaderText",	"Command Line"),		TAttribute<FText>(this, &SSessionInfoWindow::GetCommandLineText), true);
+	EndSection(VerticalBox);
 
-	AddSectionLine(VerticalBox, LOCTEXT("AnalysisStatus_SectionText", "Analysis Status"));
+	BeginSection(VerticalBox, LOCTEXT("AnalysisStatus_SectionText", "Analysis Status"));
 	AddSimpleInfoLine(VerticalBox, TAttribute<FText>(this, &SSessionInfoWindow::GetStatusText), true);
+	EndSection(VerticalBox);
 
-	AddSectionLine(VerticalBox, LOCTEXT("AnalysisModules_SectionText", "Analysis Modules"));
+	BeginSection(VerticalBox, LOCTEXT("AnalysisModules_SectionText", "Analysis Modules"));
 	AddSimpleInfoLine(VerticalBox, TAttribute<FText>(this, &SSessionInfoWindow::GetModulesText), true);
+	EndSection(VerticalBox);
 
 	DockTab->SetOnTabClosed(SDockTab::FOnTabClosedCallback::CreateRaw(this, &SSessionInfoWindow::OnSessionInfoTabClosed));
 
diff --git a/Engine/Source/Developer/TraceInsights/Private/Insights/Widgets/SSessionInfoWindow.h b/Engine/Source/Developer/TraceInsights/Private/Insights/Widgets/SSessionInfoWindow.h
--- a/Engine/Source/Developer/TraceInsights/Private/Insights/Widgets/SSessionInfoWindow.h
+++ b/Engine/Source/Developer/TraceInsights/Private/Insights/Widgets/SSessionInfoWindow.h
@@ -119,7 +119,8 @@ private:
 	static void FillMenu(FMenuBuilder& MenuBuilder, const TSharedPtr<FTabManager> TabManager);
 
 private:
-	void AddSectionLine(TSharedPtr<SVerticalBox> InVerticalBox, const FText& InSectionName) const;
+	void BeginSection(TSharedPtr<SVerticalBox> InVerticalBox, const FText& InSectionName) const;
+	void EndSection(TSharedPtr<SVerticalBox> InVerticalBox) const;
 	TSharedRef<SWidget> CreateTextBox(const TAttribute<FText>& InText, bool bMultiLine) const;
 	void AddInfoLine(TSharedPtr<SVerticalBox> InVerticalBox, const FText& InHeader, const TAttribute<FText>& InValue, bool bMultiLine = false) const;
 	void AddSimpleInfoLine(TSharedPtr<SVerticalBox> InVerticalBox, const TAttribute<FText>& InValue, bool bMultiLine = false) const;
diff --git a/Engine/Source/Developer/TraceInsights/Private/Insights/Widgets/SStartPageWindow.cpp b/Engine/Source/Developer/TraceInsights/Private/Insights/Widgets/SStartPageWindow.cpp
--- a/Engine/Source/Developer/TraceInsights/Private/Insights/Widgets/SStartPageWindow.cpp
+++ b/Engine/Source/Developer/TraceInsights/Private/Insights/Widgets/SStartPageWindow.cpp
@@ -1523,7 +1523,7 @@ TSharedRef<SWidget> STraceStoreWindow::MakeTraceListMenu()
 
 	FMenuBuilder MenuBuilder(/*bInShouldCloseWindowAfterMenuSelection=*/true, nullptr);
 
-	MenuBuilder.BeginSection("Misc", LOCTEXT("MiscHeading", "Misc"));
+	MenuBuilder.BeginSection("Misc", LOCTEXT("TraceListMenu_Section_Misc", "Misc"));
 	{
 		MenuBuilder.AddMenuEntry(
 			LOCTEXT("OpenFileButtonLabel", "Open File..."),
@@ -1536,7 +1536,7 @@ TSharedRef<SWidget> STraceStoreWindow::MakeTraceListMenu()
 	}
 	MenuBuilder.EndSection();
 
-	MenuBuilder.BeginSection("AvailableTraces", LOCTEXT("AvailableTracesHeading", "Top Most Recently Created Traces"));
+	MenuBuilder.BeginSection("AvailableTraces", LOCTEXT("TraceListMenu_Section_AvailableTraces", "Top Most Recently Created Traces"));
 	{
 		UE::Trace::FStoreClient* StoreClient = FInsightsManager::Get()->GetStoreClient();
 		if (StoreClient != nullptr)
@@ -1571,7 +1571,7 @@ TSharedRef<SWidget> STraceStoreWindow::MakeTraceListMenu()
 	}
 	MenuBuilder.EndSection();
 
-	MenuBuilder.BeginSection("DebugOptions", LOCTEXT("OptionsHeading", "Debug Options"));
+	MenuBuilder.BeginSection("DebugOptions", LOCTEXT("TraceListMenu_Section_DebugOptions", "Debug Options"));
 
 	// Enable Automation Tests Option.
 	{
diff --git a/Engine/Source/Developer/TraceInsights/Private/Insights/Widgets/STimingView.cpp b/Engine/Source/Developer/TraceInsights/Private/Insights/Widgets/STimingView.cpp
--- a/Engine/Source/Developer/TraceInsights/Private/Insights/Widgets/STimingView.cpp
+++ b/Engine/Source/Developer/TraceInsights/Private/Insights/Widgets/STimingView.cpp
@@ -4047,7 +4047,7 @@ TSharedRef<SWidget> STimingView::MakeCompactAutoScrollOptionsMenu()
 {
 	FMenuBuilder MenuBuilder(/*bInShouldCloseWindowAfterMenuSelection=*/true, CommandList);
 
-	MenuBuilder.BeginSection("AutoScrollOptions", LOCTEXT("AutoScrollOptionsHeading", "Auto-Scroll Options"));
+	MenuBuilder.BeginSection("AutoScrollOptions", LOCTEXT("CompactAutoScrollOptionsMenu_Section", "Auto-Scroll Options"));
 	{
 		CreateCompactMenuLine(MenuBuilder,
 			LOCTEXT("FrameAlignment", "Frame Alignment:"),
@@ -4147,7 +4147,7 @@ TSharedRef<SWidget> STimingView::MakeAutoScrollOptionsMenu()
 {
 	FMenuBuilder MenuBuilder(/*bInShouldCloseWindowAfterMenuSelection=*/true, CommandList);
 
-	MenuBuilder.BeginSection("Alignment", LOCTEXT("AlignmentHeading", "Alignment"));
+	MenuBuilder.BeginSection("Alignment", LOCTEXT("AutoScrollOptionsMenu_Section_Alignment", "Alignment"));
 	{
 		MenuBuilder.AddMenuEntry(
 			LOCTEXT("AutoScrollNoFrameAlignment", "None"),
@@ -4184,7 +4184,7 @@ TSharedRef<SWidget> STimingView::MakeAutoScrollOptionsMenu()
 	}
 	MenuBuilder.EndSection();
 
-	MenuBuilder.BeginSection("ViewportOffset", LOCTEXT("ViewportOffsetHeading", "Viewport Offset"));
+	MenuBuilder.BeginSection("ViewportOffset", LOCTEXT("AutoScrollOptionsMenu_Section_ViewportOffset", "Viewport Offset"));
 	{
 		MenuBuilder.AddMenuEntry(
 			LOCTEXT("AutoScrollViewportOffset-10", "-10%"),
@@ -4246,7 +4246,7 @@ TSharedRef<SWidget> STimingView::MakeAutoScrollOptionsMenu()
 	}
 	MenuBuilder.EndSection();
 
-	MenuBuilder.BeginSection("Delay", LOCTEXT("DelayHeading", "Delay"));
+	MenuBuilder.BeginSection("Delay", LOCTEXT("AutoScrollOptionsMenu_Section_Delay", "Delay"));
 	{
 		MenuBuilder.AddMenuEntry(
 			LOCTEXT("AutoScrollDelay0", "0"),
@@ -4345,7 +4345,7 @@ void STimingView::CreateAllTracksMenu(FMenuBuilder& MenuBuilder)
 
 	if (TopDockedTracks.Num() > 0)
 	{
-		MenuBuilder.BeginSection("TopDockedTracks", LOCTEXT("TopDockedTracksHeading", "Top Docked Tracks"));
+		MenuBuilder.BeginSection("TopDockedTracks", LOCTEXT("ContextMenu_Section_TopDockedTracks", "Top Docked Tracks"));
 		//MenuBuilder.AddWidget(
 		//	SNew(STextBlock)
 		//	.Text(LOCTEXT("TopDockedTracks", "Top Docked Tracks")),
@@ -4366,7 +4366,7 @@ void STimingView::CreateAllTracksMenu(FMenuBuilder& MenuBuilder)
 
 	if (ScrollableTracks.Num() > 0)
 	{
-		MenuBuilder.BeginSection("ScrollableTracks", LOCTEXT("ScrollableTracksHeading", "Scrollable Tracks"));
+		MenuBuilder.BeginSection("ScrollableTracks", LOCTEXT("ContextMenu_Section_ScrollableTracks", "Scrollable Tracks"));
 		//MenuBuilder.AddWidget(
 		//	SNew(STextBlock)
 		//	.Text(LOCTEXT("ScrollableTracks", "Scrollable Tracks")),
@@ -4388,7 +4388,7 @@ void STimingView::CreateAllTracksMenu(FMenuBuilder& MenuBuilder)
 
 	if (BottomDockedTracks.Num() > 0)
 	{
-		MenuBuilder.BeginSection("BottomDockedTracks", LOCTEXT("BottomDockedTracksHeading", "Bottom Docked Tracks"));
+		MenuBuilder.BeginSection("BottomDockedTracks", LOCTEXT("ContextMenu_Section_BottomDockedTracks", "Bottom Docked Tracks"));
 		//MenuBuilder.AddWidget(
 		//	SNew(STextBlock)
 		//	.Text(LOCTEXT("BottomDockedTracks", "Bottom Docked Tracks")),
@@ -4409,7 +4409,7 @@ void STimingView::CreateAllTracksMenu(FMenuBuilder& MenuBuilder)
 
 	if (ForegroundTracks.Num() > 0)
 	{
-		MenuBuilder.BeginSection("ForegroundTracks", LOCTEXT("ForegroundTracksHeading", "Foreground Tracks"));
+		MenuBuilder.BeginSection("ForegroundTracks", LOCTEXT("ContextMenu_Section_ForegroundTracks", "Foreground Tracks"));
 		//MenuBuilder.AddWidget(
 		//	SNew(STextBlock)
 		//	.Text(LOCTEXT("ForegroundTracks", "Foreground Tracks")),
@@ -4458,7 +4458,7 @@ TSharedRef<SWidget> STimingView::MakeOtherTracksFilterMenu()
 
 	const FTimingViewCommands& Commands = FTimingViewCommands::Get();
 
-	MenuBuilder.BeginSection("GraphTracks", LOCTEXT("GraphTracksHeading", "Main Graph Track"));
+	MenuBuilder.BeginSection("GraphTracks", LOCTEXT("ContextMenu_Section_GraphTracks", "Main Graph Track"));
 	{
 		MenuBuilder.AddMenuEntry(Commands.ShowMainGraphTrack);
 	}
@@ -4510,7 +4510,7 @@ TSharedRef<SWidget> STimingView::MakeViewModeMenu()
 
 	const FTimingViewCommands& Commands = FTimingViewCommands::Get();
 
-	MenuBuilder.BeginSection("ViewMode", LOCTEXT("ViewModeHeading", "View Mode"));
+	MenuBuilder.BeginSection("ViewMode", LOCTEXT("ContextMenu_Section_ViewMode", "View Mode"));
 	{
 		MenuBuilder.AddMenuEntry(Commands.ToggleCompactMode);
 		MenuBuilder.AddMenuEntry(Commands.AutoHideEmptyTracks);
@@ -4521,7 +4521,7 @@ TSharedRef<SWidget> STimingView::MakeViewModeMenu()
 
 	CreateCpuThreadTrackColoringModeMenu(MenuBuilder);
 
-	MenuBuilder.BeginSection("Misc", LOCTEXT("MiscHeading", "Misc Settings"));
+	MenuBuilder.BeginSection("Misc", LOCTEXT("ContextMenu_Section_Misc", "Misc Settings"));
 	{
 		MenuBuilder.AddMenuEntry(Commands.PanningOnScreenEdges);
 	}
@@ -4634,7 +4634,7 @@ void AddMenuEntryRadioButton(
 
 void STimingView::CreateDepthLimitMenu(FMenuBuilder& MenuBuilder)
 {
-	MenuBuilder.BeginSection("DepthLimit", LOCTEXT("DepthLimitHeading", "Depth Limit"));
+	MenuBuilder.BeginSection("DepthLimit", LOCTEXT("ContextMenu_Section_DepthLimit", "Depth Limit"));
 	{
 		AddMenuEntryRadioButton(MenuBuilder,
 			FUIAction(
@@ -4722,7 +4722,7 @@ bool STimingView::CheckEventDepthLimit(uint32 DepthLimit) const
 
 void STimingView::CreateCpuThreadTrackColoringModeMenu(FMenuBuilder& MenuBuilder)
 {
-	MenuBuilder.BeginSection("CpuThreadTrackColoringMode", LOCTEXT("CpuThreadTrackColoringModeHeading", "Coloring Mode (CPU Thread Tracks)"));
+	MenuBuilder.BeginSection("CpuThreadTrackColoringMode", LOCTEXT("ContextMenu_Section_CpuThreadTrackColoringMode", "Coloring Mode (CPU Thread Tracks)"));
 	{
 		MenuBuilder.AddMenuEntry(
 			LOCTEXT("CpuThreadTrackColoringMode_ByTimerName", "By Timer Name"),
