[Insights] Moved the Sorting and Column Visibility menu actions from the tree context menu to the column header context menu.

diff --git a/Engine/Source/Developer/TraceInsights/Private/Insights/MemoryProfiler/Widgets/SMemTagTreeView.cpp b/Engine/Source/Developer/TraceInsights/Private/Insights/MemoryProfiler/Widgets/SMemTagTreeView.cpp
--- a/Engine/Source/Developer/TraceInsights/Private/Insights/MemoryProfiler/Widgets/SMemTagTreeView.cpp
+++ b/Engine/Source/Developer/TraceInsights/Private/Insights/MemoryProfiler/Widgets/SMemTagTreeView.cpp
@@ -177,7 +177,7 @@ TSharedRef<SWidget> SMemTagTreeView::MakeTrackersMenu()
 {
 	FMenuBuilder MenuBuilder(/*bInShouldCloseWindowAfterMenuSelection=*/true, nullptr);
 
-	MenuBuilder.BeginSection("Trackers"/*, LOCTEXT("TrackersHeading", "Trackers")*/);
+	MenuBuilder.BeginSection("Trackers"/*, LOCTEXT("ContextMenu_Section_Trackers", "Trackers")*/);
 	CreateTrackersMenuSection(MenuBuilder);
 	MenuBuilder.EndSection();
 
@@ -467,6 +467,7 @@ TSharedRef<SWidget> SMemTagTreeView::ConstructTracksMiniToolbar()
 }
 
 END_SLATE_FUNCTION_BUILD_OPTIMIZATION
+
 ////////////////////////////////////////////////////////////////////////////////////////////////////
 
 TSharedPtr<SWidget> SMemTagTreeView::TreeView_GetMenuContent()
@@ -499,7 +500,7 @@ TSharedPtr<SWidget> SMemTagTreeView::TreeView_GetMenuContent()
 	FMenuBuilder MenuBuilder(bShouldCloseWindowAfterMenuSelection, NULL);
 
 	// Selection menu
-	MenuBuilder.BeginSection("Selection", LOCTEXT("ContextMenu_Header_Selection", "Selection"));
+	MenuBuilder.BeginSection("Selection", LOCTEXT("ContextMenu_Section_Selection", "Selection"));
 	{
 		struct FLocal
 		{
@@ -523,29 +524,12 @@ TSharedPtr<SWidget> SMemTagTreeView::TreeView_GetMenuContent()
 	}
 	MenuBuilder.EndSection();
 
-	MenuBuilder.BeginSection("Misc", LOCTEXT("ContextMenu_Header_Misc", "Miscellaneous"));
+	MenuBuilder.BeginSection("Misc", LOCTEXT("ContextMenu_Section_Misc", "Miscellaneous"));
 	{
-		/*TODO
-		FUIAction Action_CopySelectedToClipboard
-		(
-			FExecuteAction::CreateSP(this, &SMemTagTreeView::ContextMenu_CopySelectedToClipboard_Execute),
-			FCanExecuteAction::CreateSP(this, &SMemTagTreeView::ContextMenu_CopySelectedToClipboard_CanExecute)
-		);
-		MenuBuilder.AddMenuEntry
-		(
-			LOCTEXT("ContextMenu_Header_Misc_CopySelectedToClipboard", "Copy To Clipboard"),
-			LOCTEXT("ContextMenu_Header_Misc_CopySelectedToClipboard_Desc", "Copies selection to clipboard"),
-			FSlateIcon(FAppStyle::Get().GetStyleSetName(), "GenericCommands.Copy"),
-			Action_CopySelectedToClipboard,
-			NAME_None,
-			EUserInterfaceActionType::Button
-		);
-		*/
-
 		MenuBuilder.AddSubMenu
 		(
-			LOCTEXT("ContextMenu_Misc_CreateGraphTracks", "Create Graph Tracks"),
-			LOCTEXT("ContextMenu_Misc_CreateGraphTracks_Desc", "Create memory graph tracks."),
+			LOCTEXT("ContextMenu_CreateGraphTracks_SubMenu", "Create Graph Tracks"),
+			LOCTEXT("ContextMenu_CreateGraphTracks_SubMenu_Desc", "Create memory graph tracks."),
 			FNewMenuDelegate::CreateSP(this, &SMemTagTreeView::TreeView_BuildCreateGraphTracksMenu),
 			false,
 			FSlateIcon()
@@ -553,8 +537,8 @@ TSharedPtr<SWidget> SMemTagTreeView::TreeView_GetMenuContent()
 
 		MenuBuilder.AddSubMenu
 		(
-			LOCTEXT("ContextMenu_Misc_RemoveGraphTracks", "Remove Graph Tracks"),
-			LOCTEXT("ContextMenu_Misc_RemoveGraphTracks_Desc", "Remove memory graph tracks."),
+			LOCTEXT("ContextMenu_RemoveGraphTracks_SubMenu", "Remove Graph Tracks"),
+			LOCTEXT("ContextMenu_RemoveGraphTracks_SubMenu_Desc", "Remove memory graph tracks."),
 			FNewMenuDelegate::CreateSP(this, &SMemTagTreeView::TreeView_BuildRemoveGraphTracksMenu),
 			false,
 			FSlateIcon()
@@ -567,8 +551,8 @@ TSharedPtr<SWidget> SMemTagTreeView::TreeView_GetMenuContent()
 		);
 		MenuBuilder.AddMenuEntry
 		(
-			LOCTEXT("ContextMenu_Misc_GenerateColorForSelectedMemTags", "Generate New Color"),
-			LOCTEXT("ContextMenu_Misc_GenerateColorForSelectedMemTags_Desc", "Generate new color for selected LLM tag(s)."),
+			LOCTEXT("ContextMenu_GenerateColorForSelectedMemTags", "Generate New Color"),
+			LOCTEXT("ContextMenu_GenerateColorForSelectedMemTags_Desc", "Generate new color for selected LLM tag(s)."),
 			FSlateIcon(),
 			Action_GenerateColorForSelectedMemTags,
 			NAME_None,
@@ -582,80 +566,13 @@ TSharedPtr<SWidget> SMemTagTreeView::TreeView_GetMenuContent()
 		);
 		MenuBuilder.AddMenuEntry
 		(
-			LOCTEXT("ContextMenu_Misc_EditColorForSelectedMemTags", "Edit Color..."),
-			LOCTEXT("ContextMenu_Misc_EditColorForSelectedMemTags_Desc", "Edit color for selected LLM tag(s)."),
+			LOCTEXT("ContextMenu_EditColorForSelectedMemTags", "Edit Color..."),
+			LOCTEXT("ContextMenu_EditColorForSelectedMemTags_Desc", "Edit color for selected LLM tag(s)."),
 			FSlateIcon(),
 			Action_EditColorForSelectedMemTags,
 			NAME_None,
 			EUserInterfaceActionType::Button
 		);
-
-		MenuBuilder.AddSubMenu
-		(
-			LOCTEXT("ContextMenu_Header_Misc_Sort", "Sort By"),
-			LOCTEXT("ContextMenu_Header_Misc_Sort_Desc", "Sort by column"),
-			FNewMenuDelegate::CreateSP(this, &SMemTagTreeView::TreeView_BuildSortByMenu),
-			false,
-			FSlateIcon(FInsightsStyle::GetStyleSetName(), "Icons.SortBy")
-		);
-	}
-	MenuBuilder.EndSection();
-
-	MenuBuilder.BeginSection("Columns", LOCTEXT("ContextMenu_Header_Columns", "Columns"));
-	{
-		MenuBuilder.AddSubMenu
-		(
-			LOCTEXT("ContextMenu_Header_Columns_View", "View Column"),
-			LOCTEXT("ContextMenu_Header_Columns_View_Desc", "Hides or shows columns"),
-			FNewMenuDelegate::CreateSP(this, &SMemTagTreeView::TreeView_BuildViewColumnMenu),
-			false,
-			FSlateIcon(FInsightsStyle::GetStyleSetName(), "Icons.ViewColumn")
-		);
-
-		FUIAction Action_ShowAllColumns
-		(
-			FExecuteAction::CreateSP(this, &SMemTagTreeView::ContextMenu_ShowAllColumns_Execute),
-			FCanExecuteAction::CreateSP(this, &SMemTagTreeView::ContextMenu_ShowAllColumns_CanExecute)
-		);
-		MenuBuilder.AddMenuEntry
-		(
-			LOCTEXT("ContextMenu_Header_Columns_ShowAllColumns", "Show All Columns"),
-			LOCTEXT("ContextMenu_Header_Columns_ShowAllColumns_Desc", "Resets tree view to show all columns"),
-			FSlateIcon(FInsightsStyle::GetStyleSetName(), "Icons.ResetColumn"),
-			Action_ShowAllColumns,
-			NAME_None,
-			EUserInterfaceActionType::Button
-		);
-
-		//FUIAction Action_ShowMinMaxMedColumns
-		//(
-		//	FExecuteAction::CreateSP(this, &SMemTagTreeView::ContextMenu_ShowMinMaxMedColumns_Execute),
-		//	FCanExecuteAction::CreateSP(this, &SMemTagTreeView::ContextMenu_ShowMinMaxMedColumns_CanExecute)
-		//);
-		//MenuBuilder.AddMenuEntry
-		//(
-		//	LOCTEXT("ContextMenu_Header_Columns_ShowMinMaxMedColumns", "Reset Columns to Min/Max/Median Preset"),
-		//	LOCTEXT("ContextMenu_Header_Columns_ShowMinMaxMedColumns_Desc", "Resets columns to Min/Max/Median preset"),
-		//	FSlateIcon(FInsightsStyle::GetStyleSetName(), "Icons.ResetColumn"),
-		//	Action_ShowMinMaxMedColumns,
-		//	NAME_None,
-		//	EUserInterfaceActionType::Button
-		//);
-
-		FUIAction Action_ResetColumns
-		(
-			FExecuteAction::CreateSP(this, &SMemTagTreeView::ContextMenu_ResetColumns_Execute),
-			FCanExecuteAction::CreateSP(this, &SMemTagTreeView::ContextMenu_ResetColumns_CanExecute)
-		);
-		MenuBuilder.AddMenuEntry
-		(
-			LOCTEXT("ContextMenu_Header_Columns_ResetColumns", "Reset Columns to Default"),
-			LOCTEXT("ContextMenu_Header_Columns_ResetColumns_Desc", "Resets columns to default"),
-			FSlateIcon(FInsightsStyle::GetStyleSetName(), "Icons.ResetColumn"),
-			Action_ResetColumns,
-			NAME_None,
-			EUserInterfaceActionType::Button
-		);
 	}
 	MenuBuilder.EndSection();
 
@@ -676,8 +593,8 @@ void SMemTagTreeView::TreeView_BuildCreateGraphTracksMenu(FMenuBuilder& MenuBuil
 		);
 		MenuBuilder.AddMenuEntry
 		(
-			LOCTEXT("ContextMenu_Misc_CreateGraphTracksForSelectedMemTags", "Selected"),
-			LOCTEXT("ContextMenu_Misc_CreateGraphTracksForSelectedMemTags_Desc", "Create memory graph tracks for selected LLM tag(s)."),
+			LOCTEXT("ContextMenu_CreateGraphTracksForSelectedMemTags", "Selected"),
+			LOCTEXT("ContextMenu_CreateGraphTracksForSelectedMemTags_Desc", "Create memory graph tracks for selected LLM tag(s)."),
 			FSlateIcon(),
 			Action_CreateGraphTracksForSelectedMemTags,
 			NAME_None,
@@ -692,8 +609,8 @@ void SMemTagTreeView::TreeView_BuildCreateGraphTracksMenu(FMenuBuilder& MenuBuil
 		);
 		MenuBuilder.AddMenuEntry
 		(
-			LOCTEXT("ContextMenu_Misc_CreateGraphTracksForFilteredMemTags", "Filtered"),
-			LOCTEXT("ContextMenu_Misc_CreateGraphTracksForFilteredMemTags_Desc", "Create memory graph tracks for all visible (filtered) LLM tags."),
+			LOCTEXT("ContextMenu_CreateGraphTracksForFilteredMemTags", "Filtered"),
+			LOCTEXT("ContextMenu_CreateGraphTracksForFilteredMemTags_Desc", "Create memory graph tracks for all visible (filtered) LLM tags."),
 			FSlateIcon(),
 			Action_CreateGraphTracksForFilteredMemTags,
 			NAME_None,
@@ -708,8 +625,8 @@ void SMemTagTreeView::TreeView_BuildCreateGraphTracksMenu(FMenuBuilder& MenuBuil
 		);
 		MenuBuilder.AddMenuEntry
 		(
-			LOCTEXT("ContextMenu_Misc_CreateAllGraphTracks", "All"),
-			LOCTEXT("ContextMenu_Misc_CreateAllGraphTracks_Desc", "Create memory graph tracks for all LLM tags."),
+			LOCTEXT("ContextMenu_CreateAllGraphTracks", "All"),
+			LOCTEXT("ContextMenu_CreateAllGraphTracks_Desc", "Create memory graph tracks for all LLM tags."),
 			FSlateIcon(),
 			Action_CreateAllGraphTracks,
 			NAME_None,
@@ -733,8 +650,8 @@ void SMemTagTreeView::TreeView_BuildRemoveGraphTracksMenu(FMenuBuilder& MenuBuil
 		);
 		MenuBuilder.AddMenuEntry
 		(
-			LOCTEXT("ContextMenu_Misc_RemoveGraphTracksForSelectedMemTags", "Selected"),
-			LOCTEXT("ContextMenu_Misc_RemoveGraphTracksForSelectedMemTags_Desc", "Remove memory graph tracks for selected LLM tag(s)."),
+			LOCTEXT("ContextMenu_RemoveGraphTracksForSelectedMemTags", "Selected"),
+			LOCTEXT("ContextMenu_RemoveGraphTracksForSelectedMemTags_Desc", "Remove memory graph tracks for selected LLM tag(s)."),
 			FSlateIcon(),
 			Action_RemoveGraphTracksForSelectedMemTags,
 			NAME_None,
@@ -749,8 +666,8 @@ void SMemTagTreeView::TreeView_BuildRemoveGraphTracksMenu(FMenuBuilder& MenuBuil
 		);
 		MenuBuilder.AddMenuEntry
 		(
-			LOCTEXT("ContextMenu_Misc_RemoveAllGraphTracks", "All"),
-			LOCTEXT("ContextMenu_Misc_RemoveAllGraphTracks_Desc", "Remove memory graph tracks for all LLM tags."),
+			LOCTEXT("ContextMenu_RemoveAllGraphTracks", "All"),
+			LOCTEXT("ContextMenu_RemoveAllGraphTracks_Desc", "Remove memory graph tracks for all LLM tags."),
 			FSlateIcon(),
 			Action_RemoveAllGraphTracks,
 			NAME_None,
@@ -764,9 +681,7 @@ void SMemTagTreeView::TreeView_BuildRemoveGraphTracksMenu(FMenuBuilder& MenuBuil
 
 void SMemTagTreeView::TreeView_BuildSortByMenu(FMenuBuilder& MenuBuilder)
 {
-	// TODO: Refactor later @see TSharedPtr<SWidget> SCascadePreviewViewportToolBar::GenerateViewMenu() const
-
-	MenuBuilder.BeginSection("ColumnName", LOCTEXT("ContextMenu_Header_Misc_ColumnName", "Column Name"));
+	MenuBuilder.BeginSection("SortColumn", LOCTEXT("ContextMenu_Section_SortColumn", "Sort Column"));
 
 	for (const TSharedRef<Insights::FTableColumn>& ColumnRef : Table->GetColumns())
 	{
@@ -794,9 +709,7 @@ void SMemTagTreeView::TreeView_BuildSortByMenu(FMenuBuilder& MenuBuilder)
 
 	MenuBuilder.EndSection();
 
-	//-----------------------------------------------------------------------------
-
-	MenuBuilder.BeginSection("SortMode", LOCTEXT("ContextMenu_Header_Misc_Sort_SortMode", "Sort Mode"));
+	MenuBuilder.BeginSection("SortMode", LOCTEXT("ContextMenu_Section_SortMode", "Sort Mode"));
 	{
 		FUIAction Action_SortAscending
 		(
@@ -806,8 +719,8 @@ void SMemTagTreeView::TreeView_BuildSortByMenu(FMenuBuilder& MenuBuilder)
 		);
 		MenuBuilder.AddMenuEntry
 		(
-			LOCTEXT("ContextMenu_Header_Misc_Sort_SortAscending", "Sort Ascending"),
-			LOCTEXT("ContextMenu_Header_Misc_Sort_SortAscending_Desc", "Sorts ascending"),
+			LOCTEXT("ContextMenu_SortAscending", "Sort Ascending"),
+			LOCTEXT("ContextMenu_SortAscending_Desc", "Sorts ascending."),
 			FSlateIcon(FAppStyle::Get().GetStyleSetName(), "Icons.SortUp"),
 			Action_SortAscending,
 			NAME_None,
@@ -822,8 +735,8 @@ void SMemTagTreeView::TreeView_BuildSortByMenu(FMenuBuilder& MenuBuilder)
 		);
 		MenuBuilder.AddMenuEntry
 		(
-			LOCTEXT("ContextMenu_Header_Misc_Sort_SortDescending", "Sort Descending"),
-			LOCTEXT("ContextMenu_Header_Misc_Sort_SortDescending_Desc", "Sorts descending"),
+			LOCTEXT("ContextMenu_SortDescending", "Sort Descending"),
+			LOCTEXT("ContextMenu_SortDescending_Desc", "Sorts descending."),
 			FSlateIcon(FAppStyle::Get().GetStyleSetName(), "Icons.SortDown"),
 			Action_SortDescending,
 			NAME_None,
@@ -837,7 +750,7 @@ void SMemTagTreeView::TreeView_BuildSortByMenu(FMenuBuilder& MenuBuilder)
 
 void SMemTagTreeView::TreeView_BuildViewColumnMenu(FMenuBuilder& MenuBuilder)
 {
-	MenuBuilder.BeginSection("ViewColumn", LOCTEXT("ContextMenu_Header_Columns_View", "View Column"));
+	MenuBuilder.BeginSection("Columns", LOCTEXT("ContextMenu_Section_Columns", "Columns"));
 
 	for (const TSharedRef<Insights::FTableColumn>& ColumnRef : Table->GetColumns())
 	{
@@ -894,38 +807,13 @@ FText SMemTagTreeView::GetColumnHeaderText(const FName ColumnId) const
 
 TSharedRef<SWidget> SMemTagTreeView::TreeViewHeaderRow_GenerateColumnMenu(const Insights::FTableColumn& Column)
 {
-	bool bIsMenuVisible = false;
-
 	const bool bShouldCloseWindowAfterMenuSelection = true;
 	FMenuBuilder MenuBuilder(bShouldCloseWindowAfterMenuSelection, NULL);
-	{
-		if (Column.CanBeHidden())
-		{
-			MenuBuilder.BeginSection("Column", LOCTEXT("TreeViewHeaderRow_Header_Column", "Column"));
-
-			FUIAction Action_HideColumn
-			(
-				FExecuteAction::CreateSP(this, &SMemTagTreeView::HideColumn, Column.GetId()),
-				FCanExecuteAction::CreateSP(this, &SMemTagTreeView::CanHideColumn, Column.GetId())
-			);
-			MenuBuilder.AddMenuEntry
-			(
-				LOCTEXT("TreeViewHeaderRow_HideColumn", "Hide"),
-				LOCTEXT("TreeViewHeaderRow_HideColumn_Desc", "Hides the selected column"),
-				FSlateIcon(),
-				Action_HideColumn,
-				NAME_None,
-				EUserInterfaceActionType::Button
-			);
-
-			bIsMenuVisible = true;
-			MenuBuilder.EndSection();
-		}
 
+	MenuBuilder.BeginSection("Sorting", LOCTEXT("ContextMenu_Section_Sorting", "Sorting"));
+	{
 		if (Column.CanBeSorted())
 		{
-			MenuBuilder.BeginSection("SortMode", LOCTEXT("ContextMenu_Header_Misc_Sort_SortMode", "Sort Mode"));
-
 			FUIAction Action_SortAscending
 			(
 				FExecuteAction::CreateSP(this, &SMemTagTreeView::HeaderMenu_SortMode_Execute, Column.GetId(), EColumnSortMode::Ascending),
@@ -934,52 +822,119 @@ TSharedRef<SWidget> SMemTagTreeView::TreeViewHeaderRow_GenerateColumnMenu(const
 			);
 			MenuBuilder.AddMenuEntry
 			(
-				LOCTEXT("ContextMenu_Header_Misc_Sort_SortAscending", "Sort Ascending"),
-				LOCTEXT("ContextMenu_Header_Misc_Sort_SortAscending_Desc", "Sorts ascending"),
+				FText::Format(LOCTEXT("ContextMenu_SortAscending_Fmt", "Sort Ascending (by {0})"), Column.GetTitleName()),
+				FText::Format(LOCTEXT("ContextMenu_SortAscending_Desc_Fmt", "Sorts ascending by {0}."), Column.GetTitleName()),
 				FSlateIcon(FAppStyle::Get().GetStyleSetName(), "Icons.SortUp"),
 				Action_SortAscending,
 				NAME_None,
 				EUserInterfaceActionType::RadioButton
 			);
 
-			FUIAction Action_SortDescending
+				FUIAction Action_SortDescending
+				(
+					FExecuteAction::CreateSP(this, &SMemTagTreeView::HeaderMenu_SortMode_Execute, Column.GetId(), EColumnSortMode::Descending),
+					FCanExecuteAction::CreateSP(this, &SMemTagTreeView::HeaderMenu_SortMode_CanExecute, Column.GetId(), EColumnSortMode::Descending),
+					FIsActionChecked::CreateSP(this, &SMemTagTreeView::HeaderMenu_SortMode_IsChecked, Column.GetId(), EColumnSortMode::Descending)
+				);
+				MenuBuilder.AddMenuEntry
+				(
+				FText::Format(LOCTEXT("ContextMenu_SortDescending_Fmt", "Sort Descending (by {0})"), Column.GetTitleName()),
+				FText::Format(LOCTEXT("ContextMenu_SortDescending_Desc_Fmt", "Sorts descending by {0}."), Column.GetTitleName()),
+					FSlateIcon(FAppStyle::Get().GetStyleSetName(), "Icons.SortDown"),
+					Action_SortDescending,
+					NAME_None,
+					EUserInterfaceActionType::RadioButton
+				);
+			}
+
+			MenuBuilder.AddSubMenu
 			(
-				FExecuteAction::CreateSP(this, &SMemTagTreeView::HeaderMenu_SortMode_Execute, Column.GetId(), EColumnSortMode::Descending),
-				FCanExecuteAction::CreateSP(this, &SMemTagTreeView::HeaderMenu_SortMode_CanExecute, Column.GetId(), EColumnSortMode::Descending),
-				FIsActionChecked::CreateSP(this, &SMemTagTreeView::HeaderMenu_SortMode_IsChecked, Column.GetId(), EColumnSortMode::Descending)
+				LOCTEXT("ContextMenu_SortBy_SubMenu", "Sort By"),
+				LOCTEXT("ContextMenu_SortBy_SubMenu_Desc", "Sorts by a column."),
+				FNewMenuDelegate::CreateSP(this, &SMemTagTreeView::TreeView_BuildSortByMenu),
+				false,
+				FSlateIcon(FInsightsStyle::GetStyleSetName(), "Icons.SortBy")
+			);
+		}
+		MenuBuilder.EndSection();
+
+	MenuBuilder.BeginSection("ColumnVisibility", LOCTEXT("ContextMenu_Section_ColumnVisibility", "Column Visibility"));
+	{
+		if (Column.CanBeHidden())
+		{
+			FUIAction Action_HideColumn
+			(
+				FExecuteAction::CreateSP(this, &SMemTagTreeView::HideColumn, Column.GetId()),
+				FCanExecuteAction::CreateSP(this, &SMemTagTreeView::CanHideColumn, Column.GetId())
 			);
 			MenuBuilder.AddMenuEntry
 			(
-				LOCTEXT("ContextMenu_Header_Misc_Sort_SortDescending", "Sort Descending"),
-				LOCTEXT("ContextMenu_Header_Misc_Sort_SortDescending_Desc", "Sorts descending"),
-				FSlateIcon(FAppStyle::Get().GetStyleSetName(), "Icons.SortDown"),
-				Action_SortDescending,
+				LOCTEXT("ContextMenu_HideColumn", "Hide"),
+				LOCTEXT("ContextMenu_HideColumn_Desc", "Hides the selected column."),
+				FSlateIcon(),
+				Action_HideColumn,
 				NAME_None,
-				EUserInterfaceActionType::RadioButton
+				EUserInterfaceActionType::Button
 			);
-
-			bIsMenuVisible = true;
-			MenuBuilder.EndSection();
 		}
 
-		//if (Column.CanBeFiltered())
-		//{
-		//	MenuBuilder.BeginSection("FilterMode", LOCTEXT("ContextMenu_Header_Misc_Filter_FilterMode", "Filter Mode"));
-		//	bIsMenuVisible = true;
-		//	MenuBuilder.EndSection();
-		//}
-	}
+		MenuBuilder.AddSubMenu
+		(
+			LOCTEXT("ContextMenu_ViewColumn_SubMenu", "View Column"),
+			LOCTEXT("ContextMenu_ViewColumn_SubMenu_Desc", "Hides or shows columns."),
+			FNewMenuDelegate::CreateSP(this, &SMemTagTreeView::TreeView_BuildViewColumnMenu),
+			false,
+			FSlateIcon(FInsightsStyle::GetStyleSetName(), "Icons.ViewColumn")
+		);
 
-	/*
-	TODO:
-	- Show top ten
-	- Show top bottom
-	- Filter by list (avg, median, 10%, 90%, etc.)
-	- Text box for filtering for each column instead of one text box used for filtering
-	- Grouping button for flat view modes (show at most X groups, show all groups for names)
-	*/
+		FUIAction Action_ShowAllColumns
+		(
+			FExecuteAction::CreateSP(this, &SMemTagTreeView::ContextMenu_ShowAllColumns_Execute),
+			FCanExecuteAction::CreateSP(this, &SMemTagTreeView::ContextMenu_ShowAllColumns_CanExecute)
+		);
+		MenuBuilder.AddMenuEntry
+		(
+			LOCTEXT("ContextMenu_ShowAllColumns", "Show All Columns"),
+			LOCTEXT("ContextMenu_ShowAllColumns_Desc", "Resets tree view to show all columns."),
+			FSlateIcon(FInsightsStyle::GetStyleSetName(), "Icons.ResetColumn"),
+			Action_ShowAllColumns,
+			NAME_None,
+			EUserInterfaceActionType::Button
+		);
+
+		//FUIAction Action_ShowMinMaxMedColumns
+		//(
+		//	FExecuteAction::CreateSP(this, &SMemTagTreeView::ContextMenu_ShowMinMaxMedColumns_Execute),
+		//	FCanExecuteAction::CreateSP(this, &SMemTagTreeView::ContextMenu_ShowMinMaxMedColumns_CanExecute)
+		//);
+		//MenuBuilder.AddMenuEntry
+		//(
+		//	LOCTEXT("ContextMenu_ShowMinMaxMedColumns", "Reset Columns to Min/Max/Median Preset"),
+		//	LOCTEXT("ContextMenu_ShowMinMaxMedColumns_Desc", "Resets columns to Min/Max/Median preset."),
+		//	FSlateIcon(FInsightsStyle::GetStyleSetName(), "Icons.ResetColumn"),
+		//	Action_ShowMinMaxMedColumns,
+		//	NAME_None,
+		//	EUserInterfaceActionType::Button
+		//);
 
-	return bIsMenuVisible ? MenuBuilder.MakeWidget() : (TSharedRef<SWidget>)SNullWidget::NullWidget;
+		FUIAction Action_ResetColumns
+		(
+			FExecuteAction::CreateSP(this, &SMemTagTreeView::ContextMenu_ResetColumns_Execute),
+			FCanExecuteAction::CreateSP(this, &SMemTagTreeView::ContextMenu_ResetColumns_CanExecute)
+		);
+		MenuBuilder.AddMenuEntry
+		(
+			LOCTEXT("ContextMenu_ResetColumns", "Reset Columns to Default"),
+			LOCTEXT("ContextMenu_ResetColumns_Desc", "Resets columns to default."),
+			FSlateIcon(FInsightsStyle::GetStyleSetName(), "Icons.ResetColumn"),
+			Action_ResetColumns,
+			NAME_None,
+			EUserInterfaceActionType::Button
+		);
+	}
+	MenuBuilder.EndSection();
+
+	return MenuBuilder.MakeWidget();
 }
 
 ////////////////////////////////////////////////////////////////////////////////////////////////////
diff --git a/Engine/Source/Developer/TraceInsights/Private/Insights/NetworkingProfiler/Widgets/SNetStatsView.cpp b/Engine/Source/Developer/TraceInsights/Private/Insights/NetworkingProfiler/Widgets/SNetStatsView.cpp
--- a/Engine/Source/Developer/TraceInsights/Private/Insights/NetworkingProfiler/Widgets/SNetStatsView.cpp
+++ b/Engine/Source/Developer/TraceInsights/Private/Insights/NetworkingProfiler/Widgets/SNetStatsView.cpp
@@ -256,7 +256,7 @@ TSharedPtr<SWidget> SNetStatsView::TreeView_GetMenuContent()
 	FMenuBuilder MenuBuilder(bShouldCloseWindowAfterMenuSelection, NULL);
 
 	// Selection menu
-	MenuBuilder.BeginSection("Selection", LOCTEXT("ContextMenu_Header_Selection", "Selection"));
+	MenuBuilder.BeginSection("Selection", LOCTEXT("ContextMenu_Section_Selection", "Selection"));
 	{
 		struct FLocal
 		{
@@ -280,88 +280,6 @@ TSharedPtr<SWidget> SNetStatsView::TreeView_GetMenuContent()
 	}
 	MenuBuilder.EndSection();
 
-	MenuBuilder.BeginSection("Misc", LOCTEXT("ContextMenu_Header_Misc", "Miscellaneous"));
-	{
-		/*TODO
-		MenuBuilder.AddMenuEntry
-		(
-			FTimersViewCommands::Get().Command_CopyToClipboard,
-			NAME_None,
-			TAttribute<FText>(),
-			TAttribute<FText>(),
-			FSlateIcon(FAppStyle::Get().GetStyleSetName(), "GenericCommands.Copy")
-		);
-		*/
-
-		MenuBuilder.AddSubMenu
-		(
-			LOCTEXT("ContextMenu_Header_Misc_Sort", "Sort By"),
-			LOCTEXT("ContextMenu_Header_Misc_Sort_Desc", "Sort by column"),
-			FNewMenuDelegate::CreateSP(this, &SNetStatsView::TreeView_BuildSortByMenu),
-			false,
-			FSlateIcon(FInsightsStyle::GetStyleSetName(), "Icons.SortBy")
-		);
-	}
-	MenuBuilder.EndSection();
-
-	MenuBuilder.BeginSection("Columns", LOCTEXT("ContextMenu_Header_Columns", "Columns"));
-	{
-		MenuBuilder.AddSubMenu
-		(
-			LOCTEXT("ContextMenu_Header_Columns_View", "View Column"),
-			LOCTEXT("ContextMenu_Header_Columns_View_Desc", "Hides or shows columns"),
-			FNewMenuDelegate::CreateSP(this, &SNetStatsView::TreeView_BuildViewColumnMenu),
-			false,
-			FSlateIcon(FInsightsStyle::GetStyleSetName(), "Icons.ViewColumn")
-		);
-
-		FUIAction Action_ShowAllColumns
-		(
-			FExecuteAction::CreateSP(this, &SNetStatsView::ContextMenu_ShowAllColumns_Execute),
-			FCanExecuteAction::CreateSP(this, &SNetStatsView::ContextMenu_ShowAllColumns_CanExecute)
-		);
-		MenuBuilder.AddMenuEntry
-		(
-			LOCTEXT("ContextMenu_Header_Columns_ShowAllColumns", "Show All Columns"),
-			LOCTEXT("ContextMenu_Header_Columns_ShowAllColumns_Desc", "Resets tree view to show all columns."),
-			FSlateIcon(FInsightsStyle::GetStyleSetName(), "Icons.ResetColumn"),
-			Action_ShowAllColumns,
-			NAME_None,
-			EUserInterfaceActionType::Button
-		);
-
-		FUIAction Action_ShowMinMaxMedColumns
-		(
-			FExecuteAction::CreateSP(this, &SNetStatsView::ContextMenu_ShowMinMaxMedColumns_Execute),
-			FCanExecuteAction::CreateSP(this, &SNetStatsView::ContextMenu_ShowMinMaxMedColumns_CanExecute)
-		);
-		MenuBuilder.AddMenuEntry
-		(
-			LOCTEXT("ContextMenu_Header_Columns_ShowMinMaxMedColumns", "Reset Columns to Min/Max/Median Preset"),
-			LOCTEXT("ContextMenu_Header_Columns_ShowMinMaxMedColumns_Desc", "Resets columns to Min/Max/Median preset."),
-			FSlateIcon(FInsightsStyle::GetStyleSetName(), "Icons.ResetColumn"),
-			Action_ShowMinMaxMedColumns,
-			NAME_None,
-			EUserInterfaceActionType::Button
-		);
-
-		FUIAction Action_ResetColumns
-		(
-			FExecuteAction::CreateSP(this, &SNetStatsView::ContextMenu_ResetColumns_Execute),
-			FCanExecuteAction::CreateSP(this, &SNetStatsView::ContextMenu_ResetColumns_CanExecute)
-		);
-		MenuBuilder.AddMenuEntry
-		(
-			LOCTEXT("ContextMenu_Header_Columns_ResetColumns", "Reset Columns to Default"),
-			LOCTEXT("ContextMenu_Header_Columns_ResetColumns_Desc", "Resets columns to default."),
-			FSlateIcon(FInsightsStyle::GetStyleSetName(), "Icons.ResetColumn"),
-			Action_ResetColumns,
-			NAME_None,
-			EUserInterfaceActionType::Button
-		);
-	}
-	MenuBuilder.EndSection();
-
 	return MenuBuilder.MakeWidget();
 }
 
@@ -369,9 +287,7 @@ TSharedPtr<SWidget> SNetStatsView::TreeView_GetMenuContent()
 
 void SNetStatsView::TreeView_BuildSortByMenu(FMenuBuilder& MenuBuilder)
 {
-	// TODO: Refactor later @see TSharedPtr<SWidget> SCascadePreviewViewportToolBar::GenerateViewMenu() const
-
-	MenuBuilder.BeginSection("ColumnName", LOCTEXT("ContextMenu_Header_Misc_ColumnName", "Column Name"));
+	MenuBuilder.BeginSection("SortColumn", LOCTEXT("ContextMenu_Section_SortColumn", "Sort Column"));
 
 	for (const TSharedRef<Insights::FTableColumn>& ColumnRef : Table->GetColumns())
 	{
@@ -399,9 +315,7 @@ void SNetStatsView::TreeView_BuildSortByMenu(FMenuBuilder& MenuBuilder)
 
 	MenuBuilder.EndSection();
 
-	//-----------------------------------------------------------------------------
-
-	MenuBuilder.BeginSection("SortMode", LOCTEXT("ContextMenu_Header_Misc_Sort_SortMode", "Sort Mode"));
+	MenuBuilder.BeginSection("SortMode", LOCTEXT("ContextMenu_Section_SortMode", "Sort Mode"));
 	{
 		FUIAction Action_SortAscending
 		(
@@ -411,8 +325,8 @@ void SNetStatsView::TreeView_BuildSortByMenu(FMenuBuilder& MenuBuilder)
 		);
 		MenuBuilder.AddMenuEntry
 		(
-			LOCTEXT("ContextMenu_Header_Misc_Sort_SortAscending", "Sort Ascending"),
-			LOCTEXT("ContextMenu_Header_Misc_Sort_SortAscending_Desc", "Sorts ascending"),
+			LOCTEXT("ContextMenu_SortAscending", "Sort Ascending"),
+			LOCTEXT("ContextMenu_SortAscending_Desc", "Sorts ascending."),
 			FSlateIcon(FAppStyle::Get().GetStyleSetName(), "Icons.SortUp"),
 			Action_SortAscending,
 			NAME_None,
@@ -427,8 +341,8 @@ void SNetStatsView::TreeView_BuildSortByMenu(FMenuBuilder& MenuBuilder)
 		);
 		MenuBuilder.AddMenuEntry
 		(
-			LOCTEXT("ContextMenu_Header_Misc_Sort_SortDescending", "Sort Descending"),
-			LOCTEXT("ContextMenu_Header_Misc_Sort_SortDescending_Desc", "Sorts descending"),
+			LOCTEXT("ContextMenu_SortDescending", "Sort Descending"),
+			LOCTEXT("ContextMenu_SortDescending_Desc", "Sorts descending."),
 			FSlateIcon(FAppStyle::Get().GetStyleSetName(), "Icons.SortDown"),
 			Action_SortDescending,
 			NAME_None,
@@ -442,7 +356,7 @@ void SNetStatsView::TreeView_BuildSortByMenu(FMenuBuilder& MenuBuilder)
 
 void SNetStatsView::TreeView_BuildViewColumnMenu(FMenuBuilder& MenuBuilder)
 {
-	MenuBuilder.BeginSection("ViewColumn", LOCTEXT("ContextMenu_Header_Columns_View", "View Column"));
+	MenuBuilder.BeginSection("Columns", LOCTEXT("ContextMenu_Section_Columns", "Columns"));
 
 	for (const TSharedRef<Insights::FTableColumn>& ColumnRef : Table->GetColumns())
 	{
@@ -499,38 +413,13 @@ FText SNetStatsView::GetColumnHeaderText(const FName ColumnId) const
 
 TSharedRef<SWidget> SNetStatsView::TreeViewHeaderRow_GenerateColumnMenu(const Insights::FTableColumn& Column)
 {
-	bool bIsMenuVisible = false;
-
 	const bool bShouldCloseWindowAfterMenuSelection = true;
 	FMenuBuilder MenuBuilder(bShouldCloseWindowAfterMenuSelection, NULL);
-	{
-		if (Column.CanBeHidden())
-		{
-			MenuBuilder.BeginSection("Column", LOCTEXT("TreeViewHeaderRow_Header_Column", "Column"));
-
-			FUIAction Action_HideColumn
-			(
-				FExecuteAction::CreateSP(this, &SNetStatsView::HideColumn, Column.GetId()),
-				FCanExecuteAction::CreateSP(this, &SNetStatsView::CanHideColumn, Column.GetId())
-			);
-			MenuBuilder.AddMenuEntry
-			(
-				LOCTEXT("TreeViewHeaderRow_HideColumn", "Hide"),
-				LOCTEXT("TreeViewHeaderRow_HideColumn_Desc", "Hides the selected column"),
-				FSlateIcon(),
-				Action_HideColumn,
-				NAME_None,
-				EUserInterfaceActionType::Button
-			);
-
-			bIsMenuVisible = true;
-			MenuBuilder.EndSection();
-		}
 
+	MenuBuilder.BeginSection("Sorting", LOCTEXT("ContextMenu_Section_Sorting", "Sorting"));
+	{
 		if (Column.CanBeSorted())
 		{
-			MenuBuilder.BeginSection("SortMode", LOCTEXT("ContextMenu_Header_Misc_Sort_SortMode", "Sort Mode"));
-
 			FUIAction Action_SortAscending
 			(
 				FExecuteAction::CreateSP(this, &SNetStatsView::HeaderMenu_SortMode_Execute, Column.GetId(), EColumnSortMode::Ascending),
@@ -539,8 +428,8 @@ TSharedRef<SWidget> SNetStatsView::TreeViewHeaderRow_GenerateColumnMenu(const In
 			);
 			MenuBuilder.AddMenuEntry
 			(
-				LOCTEXT("ContextMenu_Header_Misc_Sort_SortAscending", "Sort Ascending"),
-				LOCTEXT("ContextMenu_Header_Misc_Sort_SortAscending_Desc", "Sorts ascending"),
+				FText::Format(LOCTEXT("ContextMenu_SortAscending_Fmt", "Sort Ascending (by {0})"), Column.GetTitleName()),
+				FText::Format(LOCTEXT("ContextMenu_SortAscending_Desc_Fmt", "Sorts ascending by {0}."), Column.GetTitleName()),
 				FSlateIcon(FAppStyle::Get().GetStyleSetName(), "Icons.SortUp"),
 				Action_SortAscending,
 				NAME_None,
@@ -555,36 +444,104 @@ TSharedRef<SWidget> SNetStatsView::TreeViewHeaderRow_GenerateColumnMenu(const In
 			);
 			MenuBuilder.AddMenuEntry
 			(
-				LOCTEXT("ContextMenu_Header_Misc_Sort_SortDescending", "Sort Descending"),
-				LOCTEXT("ContextMenu_Header_Misc_Sort_SortDescending_Desc", "Sorts descending"),
+				FText::Format(LOCTEXT("ContextMenu_SortDescending_Fmt", "Sort Descending (by {0})"), Column.GetTitleName()),
+				FText::Format(LOCTEXT("ContextMenu_SortDescending_Desc_Fmt", "Sorts descending by {0}."), Column.GetTitleName()),
 				FSlateIcon(FAppStyle::Get().GetStyleSetName(), "Icons.SortDown"),
 				Action_SortDescending,
 				NAME_None,
 				EUserInterfaceActionType::RadioButton
 			);
-
-			bIsMenuVisible = true;
-			MenuBuilder.EndSection();
 		}
 
-		//if (Column.CanBeFiltered())
-		//{
-		//	MenuBuilder.BeginSection("FilterMode", LOCTEXT("ContextMenu_Header_Misc_Filter_FilterMode", "Filter Mode"));
-		//	bIsMenuVisible = true;
-		//	MenuBuilder.EndSection();
-		//}
+		MenuBuilder.AddSubMenu
+		(
+			LOCTEXT("ContextMenu_SortBy_SubMenu", "Sort By"),
+			LOCTEXT("ContextMenu_SortBy_SubMenu_Desc", "Sorts by a column."),
+			FNewMenuDelegate::CreateSP(this, &SNetStatsView::TreeView_BuildSortByMenu),
+			false,
+			FSlateIcon(FInsightsStyle::GetStyleSetName(), "Icons.SortBy")
+		);
 	}
+	MenuBuilder.EndSection();
+
+	MenuBuilder.BeginSection("ColumnVisibility", LOCTEXT("ContextMenu_Section_ColumnVisibility", "Column Visibility"));
+	{
+		if (Column.CanBeHidden())
+		{
+
+			FUIAction Action_HideColumn
+			(
+				FExecuteAction::CreateSP(this, &SNetStatsView::HideColumn, Column.GetId()),
+				FCanExecuteAction::CreateSP(this, &SNetStatsView::CanHideColumn, Column.GetId())
+			);
+			MenuBuilder.AddMenuEntry
+			(
+				LOCTEXT("ContextMenu_HideColumn", "Hide"),
+				LOCTEXT("ContextMenu_HideColumn_Desc", "Hides the selected column."),
+				FSlateIcon(),
+				Action_HideColumn,
+				NAME_None,
+				EUserInterfaceActionType::Button
+			);
+		}
+
+		MenuBuilder.AddSubMenu
+		(
+			LOCTEXT("ContextMenu_ViewColumn_SubMenu", "View Column"),
+			LOCTEXT("ContextMenu_ViewColumn_SubMenu_Desc", "Hides or shows columns."),
+			FNewMenuDelegate::CreateSP(this, &SNetStatsView::TreeView_BuildViewColumnMenu),
+			false,
+			FSlateIcon(FInsightsStyle::GetStyleSetName(), "Icons.ViewColumn")
+		);
+
+		FUIAction Action_ShowAllColumns
+		(
+			FExecuteAction::CreateSP(this, &SNetStatsView::ContextMenu_ShowAllColumns_Execute),
+			FCanExecuteAction::CreateSP(this, &SNetStatsView::ContextMenu_ShowAllColumns_CanExecute)
+		);
+		MenuBuilder.AddMenuEntry
+		(
+			LOCTEXT("ContextMenu_ShowAllColumns", "Show All Columns"),
+			LOCTEXT("ContextMenu_ShowAllColumns_Desc", "Resets tree view to show all columns."),
+			FSlateIcon(FInsightsStyle::GetStyleSetName(), "Icons.ResetColumn"),
+			Action_ShowAllColumns,
+			NAME_None,
+			EUserInterfaceActionType::Button
+		);
+
+		FUIAction Action_ShowMinMaxMedColumns
+		(
+			FExecuteAction::CreateSP(this, &SNetStatsView::ContextMenu_ShowMinMaxMedColumns_Execute),
+			FCanExecuteAction::CreateSP(this, &SNetStatsView::ContextMenu_ShowMinMaxMedColumns_CanExecute)
+		);
+		MenuBuilder.AddMenuEntry
+		(
+			LOCTEXT("ContextMenu_ShowMinMaxMedColumns", "Reset Columns to Min/Max/Median Preset"),
+			LOCTEXT("ContextMenu_ShowMinMaxMedColumns_Desc", "Resets columns to Min/Max/Median preset."),
+			FSlateIcon(FInsightsStyle::GetStyleSetName(), "Icons.ResetColumn"),
+			Action_ShowMinMaxMedColumns,
+			NAME_None,
+			EUserInterfaceActionType::Button
+		);
 
-	/*
-	TODO:
-	- Show top ten
-	- Show top bottom
-	- Filter by list (avg, median, 10%, 90%, etc.)
-	- Text box for filtering for each column instead of one text box used for filtering
-	- Grouping button for flat view modes (show at most X groups, show all groups for names)
-	*/
+		FUIAction Action_ResetColumns
+		(
+			FExecuteAction::CreateSP(this, &SNetStatsView::ContextMenu_ResetColumns_Execute),
+			FCanExecuteAction::CreateSP(this, &SNetStatsView::ContextMenu_ResetColumns_CanExecute)
+		);
+		MenuBuilder.AddMenuEntry
+		(
+			LOCTEXT("ContextMenu_ResetColumns", "Reset Columns to Default"),
+			LOCTEXT("ContextMenu_ResetColumns_Desc", "Resets columns to default."),
+			FSlateIcon(FInsightsStyle::GetStyleSetName(), "Icons.ResetColumn"),
+			Action_ResetColumns,
+			NAME_None,
+			EUserInterfaceActionType::Button
+		);
+	}
+	MenuBuilder.EndSection();
 
-	return bIsMenuVisible ? MenuBuilder.MakeWidget() : (TSharedRef<SWidget>)SNullWidget::NullWidget;
+	return MenuBuilder.MakeWidget();
 }
 
 ////////////////////////////////////////////////////////////////////////////////////////////////////
diff --git a/Engine/Source/Developer/TraceInsights/Private/Insights/Table/Widgets/STableTreeView.cpp b/Engine/Source/Developer/TraceInsights/Private/Insights/Table/Widgets/STableTreeView.cpp
--- a/Engine/Source/Developer/TraceInsights/Private/Insights/Table/Widgets/STableTreeView.cpp
+++ b/Engine/Source/Developer/TraceInsights/Private/Insights/Table/Widgets/STableTreeView.cpp
@@ -52,8 +52,8 @@ class FTableTreeViewCommands : public TCommands<FTableTreeViewCommands>
 public:
 	FTableTreeViewCommands()
 	: TCommands<FTableTreeViewCommands>(
-		TEXT("FTableTreeViewCommands"),
-		NSLOCTEXT("FTableTreeViewCommands", "Table Tree View Commands", "Table Tree View Commands"),
+		TEXT("TableTreeViewCommands"),
+		NSLOCTEXT("Contexts", "TableTreeViewCommands", "Insights - Table Tree View"),
 		NAME_None,
 		FInsightsStyle::GetStyleSetName())
 	{
@@ -495,7 +495,7 @@ TSharedPtr<SWidget> STableTreeView::TreeView_GetMenuContent()
 	FMenuBuilder MenuBuilder(bShouldCloseWindowAfterMenuSelection, CommandList.ToSharedRef());
 
 	// Selection menu
-	MenuBuilder.BeginSection("Selection", LOCTEXT("ContextMenu_Header_Selection", "Selection"));
+	MenuBuilder.BeginSection("Selection", LOCTEXT("ContextMenu_Section_Selection", "Selection"));
 	{
 		struct FLocal
 		{
@@ -519,7 +519,7 @@ TSharedPtr<SWidget> STableTreeView::TreeView_GetMenuContent()
 	}
 	MenuBuilder.EndSection();
 
-	MenuBuilder.BeginSection("Node", LOCTEXT("ContextMenu_Header_Node", "Node"));
+	MenuBuilder.BeginSection("Node", LOCTEXT("ContextMenu_Section_Node", "Node"));
 	{
 		MenuBuilder.AddMenuEntry
 		(
@@ -553,7 +553,7 @@ TSharedPtr<SWidget> STableTreeView::TreeView_GetMenuContent()
 
 	ExtendMenu(MenuBuilder);
 
-	MenuBuilder.BeginSection("Misc", LOCTEXT("ContextMenu_Header_Misc", "Miscellaneous"));
+	MenuBuilder.BeginSection("Misc", LOCTEXT("ContextMenu_Section_Misc", "Miscellaneous"));
 	{
 		MenuBuilder.AddMenuEntry
 		(
@@ -584,17 +584,8 @@ TSharedPtr<SWidget> STableTreeView::TreeView_GetMenuContent()
 
 		MenuBuilder.AddSubMenu
 		(
-			LOCTEXT("ContextMenu_Header_Misc_Sort", "Sort By"),
-			LOCTEXT("ContextMenu_Header_Misc_Sort_Desc", "Sort by column"),
-			FNewMenuDelegate::CreateSP(this, &STableTreeView::TreeView_BuildSortByMenu),
-			false,
-			FSlateIcon(FInsightsStyle::GetStyleSetName(), "Icons.SortBy")
-		);
-
-		MenuBuilder.AddSubMenu
-		(
-			LOCTEXT("ContextMenu_Header_Misc_Export", "Export"),
-			LOCTEXT("ContextMenu_Header_Misc_Sort_Desc", "Export to file"),
+			LOCTEXT("ContextMenu_Export_SubMenu", "Export"),
+			LOCTEXT("ContextMenu_Export_SubMenu_Desc", "Exports to a file."),
 			FNewMenuDelegate::CreateSP(this, &STableTreeView::TreeView_BuildExportMenu),
 			false,
 			FSlateIcon(FAppStyle::Get().GetStyleSetName(), "Icons.Save")
@@ -602,49 +593,6 @@ TSharedPtr<SWidget> STableTreeView::TreeView_GetMenuContent()
 	}
 	MenuBuilder.EndSection();
 
-	MenuBuilder.BeginSection("Columns", LOCTEXT("ContextMenu_Header_Columns", "Columns"));
-	{
-		MenuBuilder.AddSubMenu
-		(
-			LOCTEXT("ContextMenu_Header_Columns_View", "View Column"),
-			LOCTEXT("ContextMenu_Header_Columns_View_Desc", "Hides or shows columns"),
-			FNewMenuDelegate::CreateSP(this, &STableTreeView::TreeView_BuildViewColumnMenu),
-			false,
-			FSlateIcon(FInsightsStyle::GetStyleSetName(), "Icons.ViewColumn")
-		);
-
-		FUIAction Action_ShowAllColumns
-		(
-			FExecuteAction::CreateSP(this, &STableTreeView::ContextMenu_ShowAllColumns_Execute),
-			FCanExecuteAction::CreateSP(this, &STableTreeView::ContextMenu_ShowAllColumns_CanExecute)
-		);
-		MenuBuilder.AddMenuEntry
-		(
-			LOCTEXT("ContextMenu_Header_Columns_ShowAllColumns", "Show All Columns"),
-			LOCTEXT("ContextMenu_Header_Columns_ShowAllColumns_Desc", "Resets tree view to show all columns"),
-			FSlateIcon(FInsightsStyle::GetStyleSetName(), "Icons.ResetColumn"),
-			Action_ShowAllColumns,
-			NAME_None,
-			EUserInterfaceActionType::Button
-		);
-
-		FUIAction Action_ResetColumns
-		(
-			FExecuteAction::CreateSP(this, &STableTreeView::ContextMenu_ResetColumns_Execute),
-			FCanExecuteAction::CreateSP(this, &STableTreeView::ContextMenu_ResetColumns_CanExecute)
-		);
-		MenuBuilder.AddMenuEntry
-		(
-			LOCTEXT("ContextMenu_Header_Columns_ResetColumns", "Reset Columns to Default"),
-			LOCTEXT("ContextMenu_Header_Columns_ResetColumns_Desc", "Resets columns to default"),
-			FSlateIcon(FInsightsStyle::GetStyleSetName(), "Icons.ResetColumn"),
-			Action_ResetColumns,
-			NAME_None,
-			EUserInterfaceActionType::Button
-		);
-	}
-	MenuBuilder.EndSection();
-
 	return MenuBuilder.MakeWidget();
 }
 
@@ -652,9 +600,7 @@ TSharedPtr<SWidget> STableTreeView::TreeView_GetMenuContent()
 
 void STableTreeView::TreeView_BuildSortByMenu(FMenuBuilder& MenuBuilder)
 {
-	// TODO: Refactor later @see TSharedPtr<SWidget> SCascadePreviewViewportToolBar::GenerateViewMenu() const
-
-	MenuBuilder.BeginSection("ColumnName", LOCTEXT("ContextMenu_Header_Misc_ColumnName", "Column Name"));
+	MenuBuilder.BeginSection("SortColumn", LOCTEXT("ContextMenu_Section_SortColumn", "Sort Column"));
 
 	//TODO: for (Sorting : AvailableSortings)
 	for (const TSharedRef<FTableColumn>& ColumnRef : Table->GetColumns())
@@ -682,9 +628,7 @@ void STableTreeView::TreeView_BuildSortByMenu(FMenuBuilder& MenuBuilder)
 
 	MenuBuilder.EndSection();
 
-	//-----------------------------------------------------------------------------
-
-	MenuBuilder.BeginSection("SortMode", LOCTEXT("ContextMenu_Header_Misc_Sort_SortMode", "Sort Mode"));
+	MenuBuilder.BeginSection("SortMode", LOCTEXT("ContextMenu_Section_SortMode", "Sort Mode"));
 	{
 		FUIAction Action_SortAscending
 		(
@@ -694,8 +638,8 @@ void STableTreeView::TreeView_BuildSortByMenu(FMenuBuilder& MenuBuilder)
 		);
 		MenuBuilder.AddMenuEntry
 		(
-			LOCTEXT("ContextMenu_Header_Misc_Sort_SortAscending", "Sort Ascending"),
-			LOCTEXT("ContextMenu_Header_Misc_Sort_SortAscending_Desc", "Sorts ascending"),
+			LOCTEXT("ContextMenu_SortAscending", "Sort Ascending"),
+			LOCTEXT("ContextMenu_SortAscending_Desc", "Sorts ascending."),
 			FSlateIcon(FAppStyle::Get().GetStyleSetName(), "Icons.SortUp"),
 			Action_SortAscending,
 			NAME_None,
@@ -710,8 +654,8 @@ void STableTreeView::TreeView_BuildSortByMenu(FMenuBuilder& MenuBuilder)
 		);
 		MenuBuilder.AddMenuEntry
 		(
-			LOCTEXT("ContextMenu_Header_Misc_Sort_SortDescending", "Sort Descending"),
-			LOCTEXT("ContextMenu_Header_Misc_Sort_SortDescending_Desc", "Sorts descending"),
+			LOCTEXT("ContextMenu_SortDescending", "Sort Descending"),
+			LOCTEXT("ContextMenu_SortDescending_Desc", "Sorts descending."),
 			FSlateIcon(FAppStyle::Get().GetStyleSetName(), "Icons.SortDown"),
 			Action_SortDescending,
 			NAME_None,
@@ -725,11 +669,12 @@ void STableTreeView::TreeView_BuildSortByMenu(FMenuBuilder& MenuBuilder)
 
 void STableTreeView::TreeView_BuildViewColumnMenu(FMenuBuilder& MenuBuilder)
 {
-	MenuBuilder.BeginSection("ViewColumn", LOCTEXT("ContextMenu_Header_Columns_View", "View Column"));
+	MenuBuilder.BeginSection("Columns", LOCTEXT("ContextMenu_Section_Columns", "Columns"));
 
 	for (const TSharedRef<FTableColumn>& ColumnRef : Table->GetColumns())
 	{
 		const FTableColumn& Column = ColumnRef.Get();
+
 		FUIAction Action_ToggleColumn
 		(
 			FExecuteAction::CreateSP(this, &STableTreeView::ToggleColumnVisibility, Column.GetId()),
@@ -754,35 +699,35 @@ void STableTreeView::TreeView_BuildViewColumnMenu(FMenuBuilder& MenuBuilder)
 
 void STableTreeView::TreeView_BuildExportMenu(FMenuBuilder& MenuBuilder)
 {
-	MenuBuilder.BeginSection("Export", LOCTEXT("Export", "Export"));
+	MenuBuilder.BeginSection("Export", LOCTEXT("ContextMenu_Section_Export", "Export"));
+	{
+		MenuBuilder.AddMenuEntry
+		(
+			FTableTreeViewCommands::Get().Command_ExportToFile,
+			NAME_None,
+			TAttribute<FText>(),
+			TAttribute<FText>(),
+			FSlateIcon(FAppStyle::Get().GetStyleSetName(), "Icons.Save")
+		);
 
-	MenuBuilder.AddMenuEntry
-	(
-		FTableTreeViewCommands::Get().Command_ExportToFile,
-		NAME_None,
-		TAttribute<FText>(),
-		TAttribute<FText>(),
-		FSlateIcon(FAppStyle::Get().GetStyleSetName(), "Icons.Save")
-	);
-
-	MenuBuilder.AddMenuEntry
-	(
-		FTableTreeViewCommands::Get().Command_ExportEntireTreeToFile,
-		NAME_None,
-		TAttribute<FText>(),
-		TAttribute<FText>(),
-		FSlateIcon(FAppStyle::Get().GetStyleSetName(), "Icons.Save")
-	);
-
-	MenuBuilder.AddMenuEntry
-	(
-		FTableTreeViewCommands::Get().Command_ExportEntireTreeToFileNoLeafs,
-		NAME_None,
-		TAttribute<FText>(),
-		TAttribute<FText>(),
-		FSlateIcon(FAppStyle::Get().GetStyleSetName(), "Icons.Save")
-	);
+		MenuBuilder.AddMenuEntry
+		(
+			FTableTreeViewCommands::Get().Command_ExportEntireTreeToFile,
+			NAME_None,
+			TAttribute<FText>(),
+			TAttribute<FText>(),
+			FSlateIcon(FAppStyle::Get().GetStyleSetName(), "Icons.Save")
+		);
 
+		MenuBuilder.AddMenuEntry
+		(
+			FTableTreeViewCommands::Get().Command_ExportEntireTreeToFileNoLeafs,
+			NAME_None,
+			TAttribute<FText>(),
+			TAttribute<FText>(),
+			FSlateIcon(FAppStyle::Get().GetStyleSetName(), "Icons.Save")
+		);
+	}
 	MenuBuilder.EndSection();
 }
 
@@ -826,39 +771,13 @@ FText STableTreeView::GetColumnHeaderText(const FName ColumnId) const
 
 TSharedRef<SWidget> STableTreeView::TreeViewHeaderRow_GenerateColumnMenu(const FTableColumn& Column)
 {
-	bool bIsMenuVisible = false;
-
 	const bool bShouldCloseWindowAfterMenuSelection = true;
 	FMenuBuilder MenuBuilder(bShouldCloseWindowAfterMenuSelection, NULL);
-	{
-		if (Column.CanBeHidden())
-		{
-			MenuBuilder.BeginSection("Column", LOCTEXT("TreeViewHeaderRow_Header_Column", "Column"));
-
-			FUIAction Action_HideColumn
-			(
-				FExecuteAction::CreateSP(this, &STableTreeView::HideColumn, Column.GetId()),
-				FCanExecuteAction::CreateSP(this, &STableTreeView::CanHideColumn, Column.GetId())
-			);
-
-			MenuBuilder.AddMenuEntry
-			(
-				LOCTEXT("TreeViewHeaderRow_HideColumn", "Hide"),
-				LOCTEXT("TreeViewHeaderRow_HideColumn_Desc", "Hides the selected column"),
-				FSlateIcon(),
-				Action_HideColumn,
-				NAME_None,
-				EUserInterfaceActionType::Button
-			);
-			bIsMenuVisible = true;
-
-			MenuBuilder.EndSection();
-		}
 
+	MenuBuilder.BeginSection("Sorting", LOCTEXT("ContextMenu_Section_Sorting", "Sorting"));
+	{
 		if (Column.CanBeSorted())
 		{
-			MenuBuilder.BeginSection("SortMode", LOCTEXT("ContextMenu_Header_Misc_Sort_SortMode", "Sort Mode"));
-
 			FUIAction Action_SortAscending
 			(
 				FExecuteAction::CreateSP(this, &STableTreeView::HeaderMenu_SortMode_Execute, Column.GetId(), EColumnSortMode::Ascending),
@@ -867,8 +786,8 @@ TSharedRef<SWidget> STableTreeView::TreeViewHeaderRow_GenerateColumnMenu(const F
 			);
 			MenuBuilder.AddMenuEntry
 			(
-				LOCTEXT("ContextMenu_Header_Misc_Sort_SortAscending", "Sort Ascending"),
-				LOCTEXT("ContextMenu_Header_Misc_Sort_SortAscending_Desc", "Sorts ascending"),
+				FText::Format(LOCTEXT("ContextMenu_SortAscending_Fmt", "Sort Ascending (by {0})"), Column.GetTitleName()),
+				FText::Format(LOCTEXT("ContextMenu_SortAscending_Desc_Fmt", "Sorts ascending by {0}."), Column.GetTitleName()),
 				FSlateIcon(FAppStyle::Get().GetStyleSetName(), "Icons.SortUp"),
 				Action_SortAscending,
 				NAME_None,
@@ -883,27 +802,88 @@ TSharedRef<SWidget> STableTreeView::TreeViewHeaderRow_GenerateColumnMenu(const F
 			);
 			MenuBuilder.AddMenuEntry
 			(
-				LOCTEXT("ContextMenu_Header_Misc_Sort_SortDescending", "Sort Descending"),
-				LOCTEXT("ContextMenu_Header_Misc_Sort_SortDescending_Desc", "Sorts descending"),
+				FText::Format(LOCTEXT("ContextMenu_SortDescending_Fmt", "Sort Descending (by {0})"), Column.GetTitleName()),
+				FText::Format(LOCTEXT("ContextMenu_SortDescending_Desc_Fmt", "Sorts descending by {0}."), Column.GetTitleName()),
 				FSlateIcon(FAppStyle::Get().GetStyleSetName(), "Icons.SortDown"),
 				Action_SortDescending,
 				NAME_None,
 				EUserInterfaceActionType::RadioButton
 			);
-			bIsMenuVisible = true;
+		}
+
+		MenuBuilder.AddSubMenu
+		(
+			LOCTEXT("ContextMenu_SortBy_SubMenu", "Sort By"),
+			LOCTEXT("ContextMenu_SortBy_SubMenu_Desc", "Sorts by a column."),
+			FNewMenuDelegate::CreateSP(this, &STableTreeView::TreeView_BuildSortByMenu),
+			false,
+			FSlateIcon(FInsightsStyle::GetStyleSetName(), "Icons.SortBy")
+		);
+	}
+	MenuBuilder.EndSection();
 
-			MenuBuilder.EndSection();
+	MenuBuilder.BeginSection("ColumnVisibility", LOCTEXT("ContextMenu_Section_ColumnVisibility", "Column Visibility"));
+	{
+		if (Column.CanBeHidden())
+		{
+			FUIAction Action_HideColumn
+			(
+				FExecuteAction::CreateSP(this, &STableTreeView::HideColumn, Column.GetId()),
+				FCanExecuteAction::CreateSP(this, &STableTreeView::CanHideColumn, Column.GetId())
+			);
+			MenuBuilder.AddMenuEntry
+			(
+				LOCTEXT("ContextMenu_HideColumn", "Hide"),
+				LOCTEXT("ContextMenu_HideColumn_Desc", "Hides the selected column."),
+				FSlateIcon(),
+				Action_HideColumn,
+				NAME_None,
+				EUserInterfaceActionType::Button
+			);
 		}
 
-		//if (Column.CanBeFiltered())
-		//{
-		//	MenuBuilder.BeginSection("FilterMode", LOCTEXT("ContextMenu_Header_Misc_Filter_FilterMode", "Filter Mode"));
-		//	bIsMenuVisible = true;
-		//	MenuBuilder.EndSection();
-		//}
+		MenuBuilder.AddSubMenu
+		(
+			LOCTEXT("ContextMenu_ViewColumn_SubMenu", "View Column"),
+			LOCTEXT("ContextMenu_ViewColumn_SubMenu_Desc", "Hides or shows columns."),
+			FNewMenuDelegate::CreateSP(this, &STableTreeView::TreeView_BuildViewColumnMenu),
+			false,
+			FSlateIcon(FInsightsStyle::GetStyleSetName(), "Icons.ViewColumn")
+		);
+
+		FUIAction Action_ShowAllColumns
+		(
+			FExecuteAction::CreateSP(this, &STableTreeView::ContextMenu_ShowAllColumns_Execute),
+			FCanExecuteAction::CreateSP(this, &STableTreeView::ContextMenu_ShowAllColumns_CanExecute)
+		);
+		MenuBuilder.AddMenuEntry
+		(
+			LOCTEXT("ContextMenu_ShowAllColumns", "Show All Columns"),
+			LOCTEXT("ContextMenu_ShowAllColumns_Desc", "Resets tree view to show all columns."),
+			FSlateIcon(FInsightsStyle::GetStyleSetName(), "Icons.ResetColumn"),
+			Action_ShowAllColumns,
+			NAME_None,
+			EUserInterfaceActionType::Button
+		);
+
+		FUIAction Action_ResetColumns
+		(
+			FExecuteAction::CreateSP(this, &STableTreeView::ContextMenu_ResetColumns_Execute),
+			FCanExecuteAction::CreateSP(this, &STableTreeView::ContextMenu_ResetColumns_CanExecute)
+		);
+		MenuBuilder.AddMenuEntry
+		(
+			LOCTEXT("ContextMenu_ResetColumns", "Reset Columns to Default"),
+			LOCTEXT("ContextMenu_ResetColumns_Desc", "Resets columns to default."),
+			FSlateIcon(FInsightsStyle::GetStyleSetName(), "Icons.ResetColumn"),
+			Action_ResetColumns,
+			NAME_None,
+			EUserInterfaceActionType::Button
+		);
 	}
+	MenuBuilder.EndSection();
 
-	return bIsMenuVisible ? MenuBuilder.MakeWidget() : (TSharedRef<SWidget>)SNullWidget::NullWidget;
+	return MenuBuilder.MakeWidget();
 }
 
 ////////////////////////////////////////////////////////////////////////////////////////////////////
@@ -1821,12 +1801,12 @@ TSharedRef<SWidget> STableTreeView::GetGroupingCrumbMenuContent(const TSharedPtr
 	MenuBuilder.BeginSection("InsertOrAdd");
 	{
 		const FText AddGroupingText = (CrumbGroupingDepth == CurrentGroupings.Num() - 1) ? // after last one
-			LOCTEXT("GroupingMenu_Add", "Add Grouping...") :
-			LOCTEXT("GroupingMenu_Insert", "Insert Grouping...");
+			LOCTEXT("GroupingMenu_Add_SubMenu", "Add Grouping...") :
+			LOCTEXT("GroupingMenu_Insert_SubMenu", "Insert Grouping...");
 		MenuBuilder.AddSubMenu
 		(
 			AddGroupingText,
-			LOCTEXT("GroupingMenu_AddOrInsert_Desc", "Add or insert new grouping."),
+			LOCTEXT("GroupingMenu_AddOrInsert_SubMenu_Desc", "Adds or inserts a new grouping."),
 			FNewMenuDelegate::CreateSP(this, &STableTreeView::BuildGroupingSubMenu_Add, CrumbGrouping),
 			false,
 			FSlateIcon()
@@ -1845,8 +1825,8 @@ TSharedRef<SWidget> STableTreeView::GetGroupingCrumbMenuContent(const TSharedPtr
 		{
 			MenuBuilder.AddSubMenu
 			(
-				LOCTEXT("GroupingMenu_Change", "Change To..."),
-				LOCTEXT("GroupingMenu_Change_Desc", "Change selected grouping."),
+				LOCTEXT("GroupingMenu_Change_SubMenu", "Change To..."),
+				LOCTEXT("GroupingMenu_Change_SubMenu_Desc", "Changes the selected grouping."),
 				FNewMenuDelegate::CreateSP(this, &STableTreeView::BuildGroupingSubMenu_Change, CrumbGrouping),
 				false,
 				FSlateIcon()
@@ -1862,7 +1842,7 @@ TSharedRef<SWidget> STableTreeView::GetGroupingCrumbMenuContent(const TSharedPtr
 				MenuBuilder.AddMenuEntry
 				(
 					LOCTEXT("GroupingMenu_MoveLeft", "Move Left"),
-					LOCTEXT("GroupingMenu_MoveLeft_Desc", "Move selected grouping to the left."),
+					LOCTEXT("GroupingMenu_MoveLeft_Desc", "Moves the selected grouping to the left."),
 					FSlateIcon(),
 					Action_MoveLeft,
 					NAME_None,
@@ -1880,7 +1860,7 @@ TSharedRef<SWidget> STableTreeView::GetGroupingCrumbMenuContent(const TSharedPtr
 				MenuBuilder.AddMenuEntry
 				(
 					LOCTEXT("GroupingMenu_MoveRight", "Move Right"),
-					LOCTEXT("GroupingMenu_MoveRight_Desc", "Move selected grouping to the right."),
+					LOCTEXT("GroupingMenu_MoveRight_Desc", "Moves the selected grouping to the right."),
 					FSlateIcon(),
 					Action_MoveRight,
 					NAME_None,
@@ -1898,7 +1878,7 @@ TSharedRef<SWidget> STableTreeView::GetGroupingCrumbMenuContent(const TSharedPtr
 				MenuBuilder.AddMenuEntry
 				(
 					LOCTEXT("GroupingMenu_Remove", "Remove"),
-					LOCTEXT("GroupingMenu_Remove_Desc", "Remove selected grouping."),
+					LOCTEXT("GroupingMenu_Remove_Desc", "Removes the selected grouping."),
 					FSlateIcon(),
 					Action_Remove,
 					NAME_None,
@@ -1921,7 +1901,7 @@ TSharedRef<SWidget> STableTreeView::GetGroupingCrumbMenuContent(const TSharedPtr
 			MenuBuilder.AddMenuEntry
 			(
 				LOCTEXT("GroupingMenu_Reset", "Reset"),
-				LOCTEXT("GroupingMenu_Reset_Desc", "Reset groupings to default."),
+				LOCTEXT("GroupingMenu_Reset_Desc", "Resets groupings to default."),
 				FSlateIcon(),
 				Action_Reset,
 				NAME_None,
diff --git a/Engine/Source/Developer/TraceInsights/Private/Insights/Widgets/SStatsView.cpp b/Engine/Source/Developer/TraceInsights/Private/Insights/Widgets/SStatsView.cpp
--- a/Engine/Source/Developer/TraceInsights/Private/Insights/Widgets/SStatsView.cpp
+++ b/Engine/Source/Developer/TraceInsights/Private/Insights/Widgets/SStatsView.cpp
@@ -48,11 +48,11 @@ class FStatsViewCommands : public TCommands<FStatsViewCommands>
 {
 public:
 	FStatsViewCommands()
-		: TCommands<FStatsViewCommands>(
-			TEXT("FStatsViewCommands"),
-			NSLOCTEXT("FStatsViewCommands", "Stats View Commands", "Stats View Commands"),
-			NAME_None,
-			FInsightsStyle::GetStyleSetName())
+	: TCommands<FStatsViewCommands>(
+		TEXT("CountersViewCommands"),
+		NSLOCTEXT("Contexts", "CountersViewCommands", "Insights - Counters View"),
+		NAME_None,
+		FInsightsStyle::GetStyleSetName())
 	{
 	}
 
@@ -361,7 +361,7 @@ TSharedPtr<SWidget> SStatsView::TreeView_GetMenuContent()
 	FMenuBuilder MenuBuilder(bShouldCloseWindowAfterMenuSelection, CommandList.ToSharedRef());
 
 	// Selection menu
-	MenuBuilder.BeginSection("Selection", LOCTEXT("ContextMenu_Header_Selection", "Selection"));
+	MenuBuilder.BeginSection("Selection", LOCTEXT("ContextMenu_Section_Selection", "Selection"));
 	{
 		struct FLocal
 		{
@@ -386,7 +386,7 @@ TSharedPtr<SWidget> SStatsView::TreeView_GetMenuContent()
 	MenuBuilder.EndSection();
 
 	// Counter options section
-	MenuBuilder.BeginSection("CounterOptions", LOCTEXT("ContextMenu_Header_CounterOptions", "Counter Options"));
+	MenuBuilder.BeginSection("CounterOptions", LOCTEXT("ContextMenu_Section_CounterOptions", "Counter Options"));
 	{
 		auto CanExecute = [NumSelectedNodes, SelectedNode]()
 		{
@@ -407,8 +407,8 @@ TSharedPtr<SWidget> SStatsView::TreeView_GetMenuContent()
 			{
 				MenuBuilder.AddMenuEntry
 				(
-					LOCTEXT("ContextMenu_Header_CounterOptions_RemoveFromGraphTrack", "Remove series from graph track"),
-					LOCTEXT("ContextMenu_Header_CounterOptions_RemoveFromGraphTrack_Desc", "Remove the series containing event instances of this counter from the timing graph track."),
+					LOCTEXT("ContextMenu_RemoveFromGraphTrack", "Remove series from graph track"),
+					LOCTEXT("ContextMenu_RemoveFromGraphTrack_Desc", "Remove the series containing event instances of the selected counter from the timing graph track."),
 					FSlateIcon(FInsightsStyle::GetStyleSetName(), "Icons.ToggleShowGraphSeries"),
 					Action_ToggleCounterInGraphTrack,
 					NAME_None,
@@ -419,8 +419,8 @@ TSharedPtr<SWidget> SStatsView::TreeView_GetMenuContent()
 			{
 				MenuBuilder.AddMenuEntry
 				(
-					LOCTEXT("ContextMenu_Header_CounterOptions_AddToGraphTrack", "Add series to graph track"),
-					LOCTEXT("ContextMenu_Header_CounterOptions_AddToGraphTrack_Desc", "Add a series containing event instances of this counter to the timing graph track."),
+					LOCTEXT("ContextMenu_AddToGraphTrack", "Add series to graph track"),
+					LOCTEXT("ContextMenu_AddToGraphTrack_Desc", "Add a series containing event instances of the selected counter to the timing graph track."),
 					FSlateIcon(FInsightsStyle::GetStyleSetName(), "Icons.ToggleShowGraphSeries"),
 					Action_ToggleCounterInGraphTrack,
 					NAME_None,
@@ -431,7 +431,7 @@ TSharedPtr<SWidget> SStatsView::TreeView_GetMenuContent()
 	}
 	MenuBuilder.EndSection();
 
-	MenuBuilder.BeginSection("Misc", LOCTEXT("ContextMenu_Header_Misc", "Miscellaneous"));
+	MenuBuilder.BeginSection("Misc", LOCTEXT("ContextMenu_Section_Misc", "Miscellaneous"));
 	{
 		MenuBuilder.AddMenuEntry
 		(
@@ -444,77 +444,6 @@ TSharedPtr<SWidget> SStatsView::TreeView_GetMenuContent()
 	}
 	MenuBuilder.EndSection();
 
-	MenuBuilder.BeginSection("Sorting", LOCTEXT("ContextMenu_Header_Sorting", "Sorting"));
-	{
-		MenuBuilder.AddSubMenu
-		(
-			LOCTEXT("ContextMenu_Header_Misc_Sort", "Sort By"),
-			LOCTEXT("ContextMenu_Header_Misc_Sort_Desc", "Sort by column."),
-			FNewMenuDelegate::CreateSP(this, &SStatsView::TreeView_BuildSortByMenu),
-			false,
-			FSlateIcon(FInsightsStyle::GetStyleSetName(), "Icons.SortBy")
-		);
-	}
-	MenuBuilder.EndSection();
-
-	MenuBuilder.BeginSection("Columns", LOCTEXT("ContextMenu_Header_Columns", "Columns"));
-	{
-		MenuBuilder.AddSubMenu
-		(
-			LOCTEXT("ContextMenu_Header_Columns_View", "View Column"),
-			LOCTEXT("ContextMenu_Header_Columns_View_Desc", "Hides or shows columns."),
-			FNewMenuDelegate::CreateSP(this, &SStatsView::TreeView_BuildViewColumnMenu),
-			false,
-			FSlateIcon(FInsightsStyle::GetStyleSetName(), "Icons.ViewColumn")
-		);
-
-		FUIAction Action_ShowAllColumns
-		(
-			FExecuteAction::CreateSP(this, &SStatsView::ContextMenu_ShowAllColumns_Execute),
-			FCanExecuteAction::CreateSP(this, &SStatsView::ContextMenu_ShowAllColumns_CanExecute)
-		);
-		MenuBuilder.AddMenuEntry
-		(
-			LOCTEXT("ContextMenu_Header_Columns_ShowAllColumns", "Show All Columns"),
-			LOCTEXT("ContextMenu_Header_Columns_ShowAllColumns_Desc", "Resets tree view to show all columns."),
-			FSlateIcon(FInsightsStyle::GetStyleSetName(), "Icons.ResetColumn"),
-			Action_ShowAllColumns,
-			NAME_None,
-			EUserInterfaceActionType::Button
-		);
-
-		FUIAction Action_ShowMinMaxMedColumns
-		(
-			FExecuteAction::CreateSP(this, &SStatsView::ContextMenu_ShowMinMaxMedColumns_Execute),
-			FCanExecuteAction::CreateSP(this, &SStatsView::ContextMenu_ShowMinMaxMedColumns_CanExecute)
-		);
-		MenuBuilder.AddMenuEntry
-		(
-			LOCTEXT("ContextMenu_Header_Columns_ShowMinMaxMedColumns", "Reset Columns to Min/Max/Median Preset"),
-			LOCTEXT("ContextMenu_Header_Columns_ShowMinMaxMedColumns_Desc", "Resets columns to Min/Max/Median preset."),
-			FSlateIcon(FInsightsStyle::GetStyleSetName(), "Icons.ResetColumn"),
-			Action_ShowMinMaxMedColumns,
-			NAME_None,
-			EUserInterfaceActionType::Button
-		);
-
-		FUIAction Action_ResetColumns
-		(
-			FExecuteAction::CreateSP(this, &SStatsView::ContextMenu_ResetColumns_Execute),
-			FCanExecuteAction::CreateSP(this, &SStatsView::ContextMenu_ResetColumns_CanExecute)
-		);
-		MenuBuilder.AddMenuEntry
-		(
-			LOCTEXT("ContextMenu_Header_Columns_ResetColumns", "Reset Columns to Default"),
-			LOCTEXT("ContextMenu_Header_Columns_ResetColumns_Desc", "Resets columns to default."),
-			FSlateIcon(FInsightsStyle::GetStyleSetName(), "Icons.ResetColumn"),
-			Action_ResetColumns,
-			NAME_None,
-			EUserInterfaceActionType::Button
-		);
-	}
-	MenuBuilder.EndSection();
-
 	return MenuBuilder.MakeWidget();
 }
 
@@ -522,9 +451,7 @@ TSharedPtr<SWidget> SStatsView::TreeView_GetMenuContent()
 
 void SStatsView::TreeView_BuildSortByMenu(FMenuBuilder& MenuBuilder)
 {
-	// TODO: Refactor later @see TSharedPtr<SWidget> SCascadePreviewViewportToolBar::GenerateViewMenu() const
-
-	MenuBuilder.BeginSection("ColumnName", LOCTEXT("ContextMenu_Header_Misc_ColumnName", "Column Name"));
+	MenuBuilder.BeginSection("SortColumn", LOCTEXT("ContextMenu_Section_SortColumn", "Sort Column"));
 
 	for (const TSharedRef<Insights::FTableColumn>& ColumnRef : Table->GetColumns())
 	{
@@ -552,9 +479,7 @@ void SStatsView::TreeView_BuildSortByMenu(FMenuBuilder& MenuBuilder)
 
 	MenuBuilder.EndSection();
 
-	//-----------------------------------------------------------------------------
-
-	MenuBuilder.BeginSection("SortMode", LOCTEXT("ContextMenu_Header_Misc_Sort_SortMode", "Sort Mode"));
+	MenuBuilder.BeginSection("SortMode", LOCTEXT("ContextMenu_Section_SortMode", "Sort Mode"));
 	{
 		FUIAction Action_SortAscending
 		(
@@ -564,8 +489,8 @@ void SStatsView::TreeView_BuildSortByMenu(FMenuBuilder& MenuBuilder)
 		);
 		MenuBuilder.AddMenuEntry
 		(
-			LOCTEXT("ContextMenu_Header_Misc_Sort_SortAscending", "Sort Ascending"),
-			LOCTEXT("ContextMenu_Header_Misc_Sort_SortAscending_Desc", "Sorts ascending"),
+			LOCTEXT("ContextMenu_SortAscending", "Sort Ascending"),
+			LOCTEXT("ContextMenu_SortAscending_Desc", "Sorts ascending."),
 			FSlateIcon(FAppStyle::Get().GetStyleSetName(), "Icons.SortUp"),
 			Action_SortAscending,
 			NAME_None,
@@ -580,8 +505,8 @@ void SStatsView::TreeView_BuildSortByMenu(FMenuBuilder& MenuBuilder)
 		);
 		MenuBuilder.AddMenuEntry
 		(
-			LOCTEXT("ContextMenu_Header_Misc_Sort_SortDescending", "Sort Descending"),
-			LOCTEXT("ContextMenu_Header_Misc_Sort_SortDescending_Desc", "Sorts descending"),
+			LOCTEXT("ContextMenu_SortDescending", "Sort Descending"),
+			LOCTEXT("ContextMenu_SortDescending_Desc", "Sorts descending."),
 			FSlateIcon(FAppStyle::Get().GetStyleSetName(), "Icons.SortDown"),
 			Action_SortDescending,
 			NAME_None,
@@ -595,7 +520,7 @@ void SStatsView::TreeView_BuildSortByMenu(FMenuBuilder& MenuBuilder)
 
 void SStatsView::TreeView_BuildViewColumnMenu(FMenuBuilder& MenuBuilder)
 {
-	MenuBuilder.BeginSection("ViewColumn", LOCTEXT("ContextMenu_Header_Columns_View", "View Column"));
+	MenuBuilder.BeginSection("Columns", LOCTEXT("ContextMenu_Section_Columns", "Columns"));
 
 	for (const TSharedRef<Insights::FTableColumn>& ColumnRef : Table->GetColumns())
 	{
@@ -652,38 +577,13 @@ FText SStatsView::GetColumnHeaderText(const FName ColumnId) const
 
 TSharedRef<SWidget> SStatsView::TreeViewHeaderRow_GenerateColumnMenu(const Insights::FTableColumn& Column)
 {
-	bool bIsMenuVisible = false;
-
 	const bool bShouldCloseWindowAfterMenuSelection = true;
 	FMenuBuilder MenuBuilder(bShouldCloseWindowAfterMenuSelection, NULL);
-	{
-		if (Column.CanBeHidden())
-		{
-			MenuBuilder.BeginSection("Column", LOCTEXT("TreeViewHeaderRow_Header_Column", "Column"));
-
-			FUIAction Action_HideColumn
-			(
-				FExecuteAction::CreateSP(this, &SStatsView::HideColumn, Column.GetId()),
-				FCanExecuteAction::CreateSP(this, &SStatsView::CanHideColumn, Column.GetId())
-			);
-			MenuBuilder.AddMenuEntry
-			(
-				LOCTEXT("TreeViewHeaderRow_HideColumn", "Hide"),
-				LOCTEXT("TreeViewHeaderRow_HideColumn_Desc", "Hides the selected column"),
-				FSlateIcon(),
-				Action_HideColumn,
-				NAME_None,
-				EUserInterfaceActionType::Button
-			);
-
-			bIsMenuVisible = true;
-			MenuBuilder.EndSection();
-		}
 
+	MenuBuilder.BeginSection("Sorting", LOCTEXT("ContextMenu_Section_Sorting", "Sorting"));
+	{
 		if (Column.CanBeSorted())
 		{
-			MenuBuilder.BeginSection("SortMode", LOCTEXT("ContextMenu_Header_Misc_Sort_SortMode", "Sort Mode"));
-
 			FUIAction Action_SortAscending
 			(
 				FExecuteAction::CreateSP(this, &SStatsView::HeaderMenu_SortMode_Execute, Column.GetId(), EColumnSortMode::Ascending),
@@ -692,8 +592,8 @@ TSharedRef<SWidget> SStatsView::TreeViewHeaderRow_GenerateColumnMenu(const Insig
 			);
 			MenuBuilder.AddMenuEntry
 			(
-				LOCTEXT("ContextMenu_Header_Misc_Sort_SortAscending", "Sort Ascending"),
-				LOCTEXT("ContextMenu_Header_Misc_Sort_SortAscending_Desc", "Sorts ascending"),
+				FText::Format(LOCTEXT("ContextMenu_SortAscending_Fmt", "Sort Ascending (by {0})"), Column.GetTitleName()),
+				FText::Format(LOCTEXT("ContextMenu_SortAscending_Desc_Fmt", "Sorts ascending by {0}."), Column.GetTitleName()),
 				FSlateIcon(FAppStyle::Get().GetStyleSetName(), "Icons.SortUp"),
 				Action_SortAscending,
 				NAME_None,
@@ -708,36 +608,103 @@ TSharedRef<SWidget> SStatsView::TreeViewHeaderRow_GenerateColumnMenu(const Insig
 			);
 			MenuBuilder.AddMenuEntry
 			(
-				LOCTEXT("ContextMenu_Header_Misc_Sort_SortDescending", "Sort Descending"),
-				LOCTEXT("ContextMenu_Header_Misc_Sort_SortDescending_Desc", "Sorts descending"),
+				FText::Format(LOCTEXT("ContextMenu_SortDescending_Fmt", "Sort Descending (by {0})"), Column.GetTitleName()),
+				FText::Format(LOCTEXT("ContextMenu_SortDescending_Desc_Fmt", "Sorts descending by {0}."), Column.GetTitleName()),
 				FSlateIcon(FAppStyle::Get().GetStyleSetName(), "Icons.SortDown"),
 				Action_SortDescending,
 				NAME_None,
 				EUserInterfaceActionType::RadioButton
 			);
-
-			bIsMenuVisible = true;
-			MenuBuilder.EndSection();
 		}
 
-		//if (Column.CanBeFiltered())
-		//{
-		//	MenuBuilder.BeginSection("FilterMode", LOCTEXT("ContextMenu_Header_Misc_Filter_FilterMode", "Filter Mode"));
-		//	bIsMenuVisible = true;
-		//	MenuBuilder.EndSection();
-		//}
+		MenuBuilder.AddSubMenu
+		(
+			LOCTEXT("ContextMenu_SortBy_SubMenu", "Sort By"),
+			LOCTEXT("ContextMenu_SortBy_SubMenu_Desc", "Sorts by a column."),
+			FNewMenuDelegate::CreateSP(this, &SStatsView::TreeView_BuildSortByMenu),
+			false,
+			FSlateIcon(FInsightsStyle::GetStyleSetName(), "Icons.SortBy")
+		);
 	}
+	MenuBuilder.EndSection();
+
+	MenuBuilder.BeginSection("ColumnVisibility", LOCTEXT("ContextMenu_Section_ColumnVisibility", "Column Visibility"));
+	{
+		if (Column.CanBeHidden())
+		{
+			FUIAction Action_HideColumn
+			(
+				FExecuteAction::CreateSP(this, &SStatsView::HideColumn, Column.GetId()),
+				FCanExecuteAction::CreateSP(this, &SStatsView::CanHideColumn, Column.GetId())
+			);
+			MenuBuilder.AddMenuEntry
+			(
+				LOCTEXT("ContextMenu_HideColumn", "Hide"),
+				LOCTEXT("ContextMenu_HideColumn_Desc", "Hides the selected column."),
+				FSlateIcon(),
+				Action_HideColumn,
+				NAME_None,
+				EUserInterfaceActionType::Button
+			);
+		}
 
-	/*
-	TODO:
-	- Show top ten
-	- Show top bottom
-	- Filter by list (avg, median, 10%, 90%, etc.)
-	- Text box for filtering for each column instead of one text box used for filtering
-	- Grouping button for flat view modes (show at most X groups, show all groups for names)
-	*/
+		MenuBuilder.AddSubMenu
+		(
+			LOCTEXT("ContextMenu_ViewColumn_SubMenu", "View Column"),
+			LOCTEXT("ContextMenu_ViewColumn_SubMenu_Desc", "Hides or shows columns."),
+			FNewMenuDelegate::CreateSP(this, &SStatsView::TreeView_BuildViewColumnMenu),
+			false,
+			FSlateIcon(FInsightsStyle::GetStyleSetName(), "Icons.ViewColumn")
+		);
 
-	return bIsMenuVisible ? MenuBuilder.MakeWidget() : (TSharedRef<SWidget>)SNullWidget::NullWidget;
+		FUIAction Action_ShowAllColumns
+		(
+			FExecuteAction::CreateSP(this, &SStatsView::ContextMenu_ShowAllColumns_Execute),
+			FCanExecuteAction::CreateSP(this, &SStatsView::ContextMenu_ShowAllColumns_CanExecute)
+		);
+		MenuBuilder.AddMenuEntry
+		(
+			LOCTEXT("ContextMenu_ShowAllColumns", "Show All Columns"),
+			LOCTEXT("ContextMenu_ShowAllColumns_Desc", "Resets tree view to show all columns."),
+			FSlateIcon(FInsightsStyle::GetStyleSetName(), "Icons.ResetColumn"),
+			Action_ShowAllColumns,
+			NAME_None,
+			EUserInterfaceActionType::Button
+		);
+
+		FUIAction Action_ShowMinMaxMedColumns
+		(
+			FExecuteAction::CreateSP(this, &SStatsView::ContextMenu_ShowMinMaxMedColumns_Execute),
+			FCanExecuteAction::CreateSP(this, &SStatsView::ContextMenu_ShowMinMaxMedColumns_CanExecute)
+		);
+		MenuBuilder.AddMenuEntry
+		(
+			LOCTEXT("ContextMenu_ShowMinMaxMedColumns", "Reset Columns to Min/Max/Median Preset"),
+			LOCTEXT("ContextMenu_ShowMinMaxMedColumns_Desc", "Resets columns to Min/Max/Median preset."),
+			FSlateIcon(FInsightsStyle::GetStyleSetName(), "Icons.ResetColumn"),
+			Action_ShowMinMaxMedColumns,
+			NAME_None,
+			EUserInterfaceActionType::Button
+		);
+
+		FUIAction Action_ResetColumns
+		(
+			FExecuteAction::CreateSP(this, &SStatsView::ContextMenu_ResetColumns_Execute),
+			FCanExecuteAction::CreateSP(this, &SStatsView::ContextMenu_ResetColumns_CanExecute)
+		);
+		MenuBuilder.AddMenuEntry
+		(
+			LOCTEXT("ContextMenu_ResetColumns", "Reset Columns to Default"),
+			LOCTEXT("ContextMenu_ResetColumns_Desc", "Resets columns to default."),
+			FSlateIcon(FInsightsStyle::GetStyleSetName(), "Icons.ResetColumn"),
+			Action_ResetColumns,
+			NAME_None,
+			EUserInterfaceActionType::Button
+		);
+	}
+	MenuBuilder.EndSection();
+
+	return MenuBuilder.MakeWidget();
 }
 
 ////////////////////////////////////////////////////////////////////////////////////////////////////
diff --git a/Engine/Source/Developer/TraceInsights/Private/Insights/Widgets/STimerTreeView.cpp b/Engine/Source/Developer/TraceInsights/Private/Insights/Widgets/STimerTreeView.cpp
--- a/Engine/Source/Developer/TraceInsights/Private/Insights/Widgets/STimerTreeView.cpp
+++ b/Engine/Source/Developer/TraceInsights/Private/Insights/Widgets/STimerTreeView.cpp
@@ -34,11 +34,11 @@ class FTimerTreeViewCommands : public TCommands<FTimerTreeViewCommands>
 {
 public:
 	FTimerTreeViewCommands()
-		: TCommands<FTimerTreeViewCommands>(
-			TEXT("FTimerTreeViewCommands"),
-			NSLOCTEXT("FTimerTreeViewCommands", "Timer Tree View Commands", "Timer Tree View Commands"),
-			NAME_None,
-			FInsightsStyle::GetStyleSetName())
+	: TCommands<FTimerTreeViewCommands>(
+		TEXT("TimerTreeViewCommands"),
+		NSLOCTEXT("Contexts", "TimerTreeViewCommands", "Insights - Timer Tree View"),
+		NAME_None,
+		FInsightsStyle::GetStyleSetName())
 	{
 	}
 
@@ -50,7 +50,11 @@ public:
 	PRAGMA_DISABLE_OPTIMIZATION
 	virtual void RegisterCommands() override
 	{
-		UI_COMMAND(Command_CopyToClipboard, "Copy To Clipboard", "Copies selection to clipboard", EUserInterfaceActionType::Button, FInputChord(EModifierKey::Control, EKeys::C));
+		UI_COMMAND(Command_CopyToClipboard,
+			"Copy To Clipboard",
+			"Copies the selection to clipboard.",
+			EUserInterfaceActionType::Button,
+			FInputChord(EModifierKey::Control, EKeys::C));
 	}
 	PRAGMA_ENABLE_OPTIMIZATION
 
@@ -197,7 +201,7 @@ TSharedPtr<SWidget> STimerTreeView::TreeView_GetMenuContent()
 	FMenuBuilder MenuBuilder(bShouldCloseWindowAfterMenuSelection, CommandList.ToSharedRef());
 
 	// Selection menu
-	MenuBuilder.BeginSection("Selection", LOCTEXT("ContextMenu_Header_Selection", "Selection"));
+	MenuBuilder.BeginSection("Selection", LOCTEXT("ContextMenu_Section_Selection", "Selection"));
 	{
 		struct FLocal
 		{
@@ -221,7 +225,7 @@ TSharedPtr<SWidget> STimerTreeView::TreeView_GetMenuContent()
 	}
 	MenuBuilder.EndSection();
 
-	MenuBuilder.BeginSection("Misc", LOCTEXT("ContextMenu_Header_Misc", "Miscellaneous"));
+	MenuBuilder.BeginSection("Misc", LOCTEXT("ContextMenu_Section_Misc", "Miscellaneous"));
 	{
 		MenuBuilder.AddMenuEntry
 		(
@@ -231,58 +235,6 @@ TSharedPtr<SWidget> STimerTreeView::TreeView_GetMenuContent()
 			TAttribute<FText>(),
 			FSlateIcon(FAppStyle::Get().GetStyleSetName(), "GenericCommands.Copy")
 		);
-
-		MenuBuilder.AddSubMenu
-		(
-			LOCTEXT("ContextMenu_Header_Misc_Sort", "Sort By"),
-			LOCTEXT("ContextMenu_Header_Misc_Sort_Desc", "Sort by column"),
-			FNewMenuDelegate::CreateSP(this, &STimerTreeView::TreeView_BuildSortByMenu),
-			false,
-			FSlateIcon(FInsightsStyle::GetStyleSetName(), "Icons.SortBy")
-		);
-	}
-	MenuBuilder.EndSection();
-
-	MenuBuilder.BeginSection("Columns", LOCTEXT("ContextMenu_Header_Columns", "Columns"));
-	{
-		MenuBuilder.AddSubMenu
-		(
-			LOCTEXT("ContextMenu_Header_Columns_View", "View Column"),
-			LOCTEXT("ContextMenu_Header_Columns_View_Desc", "Hides or shows columns"),
-			FNewMenuDelegate::CreateSP(this, &STimerTreeView::TreeView_BuildViewColumnMenu),
-			false,
-			FSlateIcon(FInsightsStyle::GetStyleSetName(), "Icons.ViewColumn")
-		);
-
-		FUIAction Action_ShowAllColumns
-		(
-			FExecuteAction::CreateSP(this, &STimerTreeView::ContextMenu_ShowAllColumns_Execute),
-			FCanExecuteAction::CreateSP(this, &STimerTreeView::ContextMenu_ShowAllColumns_CanExecute)
-		);
-		MenuBuilder.AddMenuEntry
-		(
-			LOCTEXT("ContextMenu_Header_Columns_ShowAllColumns", "Show All Columns"),
-			LOCTEXT("ContextMenu_Header_Columns_ShowAllColumns_Desc", "Resets tree view to show all columns"),
-			FSlateIcon(FInsightsStyle::GetStyleSetName(), "Icons.ResetColumn"),
-			Action_ShowAllColumns,
-			NAME_None,
-			EUserInterfaceActionType::Button
-		);
-
-		FUIAction Action_ResetColumns
-		(
-			FExecuteAction::CreateSP(this, &STimerTreeView::ContextMenu_ResetColumns_Execute),
-			FCanExecuteAction::CreateSP(this, &STimerTreeView::ContextMenu_ResetColumns_CanExecute)
-		);
-		MenuBuilder.AddMenuEntry
-		(
-			LOCTEXT("ContextMenu_Header_Columns_ResetColumns", "Reset Columns to Default"),
-			LOCTEXT("ContextMenu_Header_Columns_ResetColumns_Desc", "Resets columns to default"),
-			FSlateIcon(FInsightsStyle::GetStyleSetName(), "Icons.ResetColumn"),
-			Action_ResetColumns,
-			NAME_None,
-			EUserInterfaceActionType::Button
-		);
 	}
 	MenuBuilder.EndSection();
 
@@ -293,7 +245,7 @@ TSharedPtr<SWidget> STimerTreeView::TreeView_GetMenuContent()
 
 void STimerTreeView::TreeView_BuildSortByMenu(FMenuBuilder& MenuBuilder)
 {
-	MenuBuilder.BeginSection("ColumnName", LOCTEXT("ContextMenu_Header_Misc_ColumnName", "Column Name"));
+	MenuBuilder.BeginSection("SortColumn", LOCTEXT("ContextMenu_Section_SortColumn", "Sort Column"));
 
 	for (const TSharedRef<Insights::FTableColumn>& ColumnRef : Table->GetColumns())
 	{
@@ -321,9 +273,7 @@ void STimerTreeView::TreeView_BuildSortByMenu(FMenuBuilder& MenuBuilder)
 
 	MenuBuilder.EndSection();
 
-	//-----------------------------------------------------------------------------
-
-	MenuBuilder.BeginSection("SortMode", LOCTEXT("ContextMenu_Header_Misc_Sort_SortMode", "Sort Mode"));
+	MenuBuilder.BeginSection("SortMode", LOCTEXT("ContextMenu_Section_SortMode", "Sort Mode"));
 	{
 		FUIAction Action_SortAscending
 		(
@@ -333,8 +283,8 @@ void STimerTreeView::TreeView_BuildSortByMenu(FMenuBuilder& MenuBuilder)
 		);
 		MenuBuilder.AddMenuEntry
 		(
-			LOCTEXT("ContextMenu_Header_Misc_Sort_SortAscending", "Sort Ascending"),
-			LOCTEXT("ContextMenu_Header_Misc_Sort_SortAscending_Desc", "Sorts ascending"),
+			LOCTEXT("ContextMenu_SortAscending", "Sort Ascending"),
+			LOCTEXT("ContextMenu_SortAscending_Desc", "Sorts ascending."),
 			FSlateIcon(FAppStyle::Get().GetStyleSetName(), "Icons.SortUp"),
 			Action_SortAscending,
 			NAME_None,
@@ -349,8 +299,8 @@ void STimerTreeView::TreeView_BuildSortByMenu(FMenuBuilder& MenuBuilder)
 		);
 		MenuBuilder.AddMenuEntry
 		(
-			LOCTEXT("ContextMenu_Header_Misc_Sort_SortDescending", "Sort Descending"),
-			LOCTEXT("ContextMenu_Header_Misc_Sort_SortDescending_Desc", "Sorts descending"),
+			LOCTEXT("ContextMenu_SortDescending", "Sort Descending"),
+			LOCTEXT("ContextMenu_SortDescending_Desc", "Sorts descending."),
 			FSlateIcon(FAppStyle::Get().GetStyleSetName(), "Icons.SortDown"),
 			Action_SortDescending,
 			NAME_None,
@@ -364,7 +314,7 @@ void STimerTreeView::TreeView_BuildSortByMenu(FMenuBuilder& MenuBuilder)
 
 void STimerTreeView::TreeView_BuildViewColumnMenu(FMenuBuilder& MenuBuilder)
 {
-	MenuBuilder.BeginSection("ViewColumn", LOCTEXT("ContextMenu_Header_Columns_View", "View Column"));
+	MenuBuilder.BeginSection("Columns", LOCTEXT("ContextMenu_Section_Columns", "Columns"));
 
 	for (const TSharedRef<Insights::FTableColumn>& ColumnRef : Table->GetColumns())
 	{
@@ -426,38 +376,13 @@ FText STimerTreeView::GetColumnHeaderText(const FName ColumnId) const
 
 TSharedRef<SWidget> STimerTreeView::TreeViewHeaderRow_GenerateColumnMenu(const Insights::FTableColumn& Column)
 {
-	bool bIsMenuVisible = false;
-
 	const bool bShouldCloseWindowAfterMenuSelection = true;
 	FMenuBuilder MenuBuilder(bShouldCloseWindowAfterMenuSelection, NULL);
-	{
-		if (Column.CanBeHidden())
-		{
-			MenuBuilder.BeginSection("Column", LOCTEXT("TreeViewHeaderRow_Header_Column", "Column"));
-
-			FUIAction Action_HideColumn
-			(
-				FExecuteAction::CreateSP(this, &STimerTreeView::HideColumn, Column.GetId()),
-				FCanExecuteAction::CreateSP(this, &STimerTreeView::CanHideColumn, Column.GetId())
-			);
-			MenuBuilder.AddMenuEntry
-			(
-				LOCTEXT("TreeViewHeaderRow_HideColumn", "Hide"),
-				LOCTEXT("TreeViewHeaderRow_HideColumn_Desc", "Hides the selected column"),
-				FSlateIcon(),
-				Action_HideColumn,
-				NAME_None,
-				EUserInterfaceActionType::Button
-			);
-
-			bIsMenuVisible = true;
-			MenuBuilder.EndSection();
-		}
 
+	MenuBuilder.BeginSection("Sorting", LOCTEXT("ContextMenu_Section_Sorting", "Sorting"));
+	{
 		if (Column.CanBeSorted())
 		{
-			MenuBuilder.BeginSection("SortMode", LOCTEXT("ContextMenu_Header_Misc_Sort_SortMode", "Sort Mode"));
-
 			FUIAction Action_SortAscending
 			(
 				FExecuteAction::CreateSP(this, &STimerTreeView::HeaderMenu_SortMode_Execute, Column.GetId(), EColumnSortMode::Ascending),
@@ -466,8 +391,8 @@ TSharedRef<SWidget> STimerTreeView::TreeViewHeaderRow_GenerateColumnMenu(const I
 			);
 			MenuBuilder.AddMenuEntry
 			(
-				LOCTEXT("ContextMenu_Header_Misc_Sort_SortAscending", "Sort Ascending"),
-				LOCTEXT("ContextMenu_Header_Misc_Sort_SortAscending_Desc", "Sorts ascending"),
+				FText::Format(LOCTEXT("ContextMenu_SortAscending_Fmt", "Sort Ascending (by {0})"), Column.GetTitleName()),
+				FText::Format(LOCTEXT("ContextMenu_SortAscending_Desc_Fmt", "Sorts ascending by {0}."), Column.GetTitleName()),
 				FSlateIcon(FAppStyle::Get().GetStyleSetName(), "Icons.SortUp"),
 				Action_SortAscending,
 				NAME_None,
@@ -482,20 +407,88 @@ TSharedRef<SWidget> STimerTreeView::TreeViewHeaderRow_GenerateColumnMenu(const I
 			);
 			MenuBuilder.AddMenuEntry
 			(
-				LOCTEXT("ContextMenu_Header_Misc_Sort_SortDescending", "Sort Descending"),
-				LOCTEXT("ContextMenu_Header_Misc_Sort_SortDescending_Desc", "Sorts descending"),
+				FText::Format(LOCTEXT("ContextMenu_SortDescending_Fmt", "Sort Descending (by {0})"), Column.GetTitleName()),
+				FText::Format(LOCTEXT("ContextMenu_SortDescending_Desc_Fmt", "Sorts descending by {0}."), Column.GetTitleName()),
 				FSlateIcon(FAppStyle::Get().GetStyleSetName(), "Icons.SortDown"),
 				Action_SortDescending,
 				NAME_None,
 				EUserInterfaceActionType::RadioButton
 			);
+		}
+
+		MenuBuilder.AddSubMenu
+		(
+			LOCTEXT("ContextMenu_SortBy_SubMenu", "Sort By"),
+			LOCTEXT("ContextMenu_SortBy_SubMenu_Desc", "Sorts by a column."),
+			FNewMenuDelegate::CreateSP(this, &STimerTreeView::TreeView_BuildSortByMenu),
+			false,
+			FSlateIcon(FInsightsStyle::GetStyleSetName(), "Icons.SortBy")
+		);
+	}
+	MenuBuilder.EndSection();
 
-			bIsMenuVisible = true;
-			MenuBuilder.EndSection();
+	MenuBuilder.BeginSection("ColumnVisibility", LOCTEXT("ContextMenu_Section_ColumnVisibility", "Column Visibility"));
+	{
+		if (Column.CanBeHidden())
+		{
+			FUIAction Action_HideColumn
+			(
+				FExecuteAction::CreateSP(this, &STimerTreeView::HideColumn, Column.GetId()),
+				FCanExecuteAction::CreateSP(this, &STimerTreeView::CanHideColumn, Column.GetId())
+			);
+			MenuBuilder.AddMenuEntry
+			(
+				LOCTEXT("ContextMenu_HideColumn", "Hide"),
+				LOCTEXT("ContextMenu_HideColumn_Desc", "Hides the selected column."),
+				FSlateIcon(),
+				Action_HideColumn,
+				NAME_None,
+				EUserInterfaceActionType::Button
+			);
 		}
+
+		MenuBuilder.AddSubMenu
+		(
+			LOCTEXT("ContextMenu_ViewColumn_SubMenu", "View Column"),
+			LOCTEXT("ContextMenu_ViewColumn_SubMenu_Desc", "Hides or shows columns."),
+			FNewMenuDelegate::CreateSP(this, &STimerTreeView::TreeView_BuildViewColumnMenu),
+			false,
+			FSlateIcon(FInsightsStyle::GetStyleSetName(), "Icons.ViewColumn")
+		);
+
+		FUIAction Action_ShowAllColumns
+		(
+			FExecuteAction::CreateSP(this, &STimerTreeView::ContextMenu_ShowAllColumns_Execute),
+			FCanExecuteAction::CreateSP(this, &STimerTreeView::ContextMenu_ShowAllColumns_CanExecute)
+		);
+		MenuBuilder.AddMenuEntry
+		(
+			LOCTEXT("ContextMenu_ShowAllColumns", "Show All Columns"),
+			LOCTEXT("ContextMenu_ShowAllColumns_Desc", "Resets tree view to show all columns."),
+			FSlateIcon(FInsightsStyle::GetStyleSetName(), "Icons.ResetColumn"),
+			Action_ShowAllColumns,
+			NAME_None,
+			EUserInterfaceActionType::Button
+		);
+
+		FUIAction Action_ResetColumns
+		(
+			FExecuteAction::CreateSP(this, &STimerTreeView::ContextMenu_ResetColumns_Execute),
+			FCanExecuteAction::CreateSP(this, &STimerTreeView::ContextMenu_ResetColumns_CanExecute)
+		);
+		MenuBuilder.AddMenuEntry
+		(
+			LOCTEXT("ContextMenu_ResetColumns", "Reset Columns to Default"),
+			LOCTEXT("ContextMenu_ResetColumns_Desc", "Resets columns to default."),
+			FSlateIcon(FInsightsStyle::GetStyleSetName(), "Icons.ResetColumn"),
+			Action_ResetColumns,
+			NAME_None,
+			EUserInterfaceActionType::Button
+		);
 	}
+	MenuBuilder.EndSection();
 
-	return bIsMenuVisible ? MenuBuilder.MakeWidget() : (TSharedRef<SWidget>)SNullWidget::NullWidget;
+	return MenuBuilder.MakeWidget();
 }
 
 ////////////////////////////////////////////////////////////////////////////////////////////////////
diff --git a/Engine/Source/Developer/TraceInsights/Private/Insights/Widgets/STimersView.cpp b/Engine/Source/Developer/TraceInsights/Private/Insights/Widgets/STimersView.cpp
--- a/Engine/Source/Developer/TraceInsights/Private/Insights/Widgets/STimersView.cpp
+++ b/Engine/Source/Developer/TraceInsights/Private/Insights/Widgets/STimersView.cpp
@@ -57,7 +57,11 @@ class FTimersViewCommands : public TCommands<FTimersViewCommands>
 {
 public:
 	FTimersViewCommands()
-		: TCommands<FTimersViewCommands>(TEXT("FTimersViewCommands"), NSLOCTEXT("FTimersViewCommands", "Timer View Commands", "Timer View Commands"), NAME_None, FInsightsStyle::Get().GetStyleSetName())
+	: TCommands<FTimersViewCommands>(
+		TEXT("TimersViewCommands"),
+		NSLOCTEXT("Contexts", "TimersViewCommands", "Insights - Timers View"),
+		NAME_None,
+		FInsightsStyle::Get().GetStyleSetName())
 	{
 	}
 
@@ -422,7 +426,7 @@ TSharedPtr<SWidget> STimersView::TreeView_GetMenuContent()
 	FMenuBuilder MenuBuilder(bShouldCloseWindowAfterMenuSelection, CommandList.ToSharedRef());
 
 	// Selection menu
-	MenuBuilder.BeginSection("Selection", LOCTEXT("ContextMenu_Header_Selection", "Selection"));
+	MenuBuilder.BeginSection("Selection", LOCTEXT("ContextMenu_Section_Selection", "Selection"));
 	{
 		struct FLocal
 		{
@@ -447,7 +451,7 @@ TSharedPtr<SWidget> STimersView::TreeView_GetMenuContent()
 	MenuBuilder.EndSection();
 
 	// Timer options section
-	MenuBuilder.BeginSection("TimerOptions", LOCTEXT("ContextMenu_Header_TimerOptions", "Timer Options"));
+	MenuBuilder.BeginSection("TimerOptions", LOCTEXT("ContextMenu_Section_TimerOptions", "Timer Options"));
 	{
 		auto CanExecute = [NumSelectedNodes, SelectedNode]()
 		{
@@ -472,8 +476,8 @@ TSharedPtr<SWidget> STimersView::TreeView_GetMenuContent()
 			{
 				MenuBuilder.AddMenuEntry
 				(
-					LOCTEXT("ContextMenu_Header_TimerOptions_StopHighlightEvent", "Stop Highlighting Event"),
-					LOCTEXT("ContextMenu_Header_TimerOptions_StopHighlightEvent_Desc", "Stops highlighting timing event instances for the selected timer."),
+					LOCTEXT("ContextMenu_StopHighlightEvent", "Stop Highlighting Event"),
+					LOCTEXT("ContextMenu_StopHighlightEvent_Desc", "Stops highlighting timing event instances for the selected timer."),
 					FSlateIcon(FAppStyle::Get().GetStyleSetName(), "Icons.Visible"),
 					Action_ToggleHighlight,
 					NAME_None,
@@ -484,8 +488,8 @@ TSharedPtr<SWidget> STimersView::TreeView_GetMenuContent()
 			{
 				MenuBuilder.AddMenuEntry
 				(
-					LOCTEXT("ContextMenu_Header_TimerOptions_HighlightEvent", "Highlight Event"),
-					LOCTEXT("ContextMenu_Header_TimerOptions_HighlightEvent_Desc", "Highlights all timing event instances for the selected timer."),
+					LOCTEXT("ContextMenu_HighlightEvent", "Highlight Event"),
+					LOCTEXT("ContextMenu_HighlightEvent_Desc", "Highlights all timing event instances for the selected timer."),
 					FSlateIcon(FAppStyle::Get().GetStyleSetName(), "Icons.Visible"),
 					Action_ToggleHighlight,
 					NAME_None,
@@ -506,8 +510,8 @@ TSharedPtr<SWidget> STimersView::TreeView_GetMenuContent()
 			{
 				MenuBuilder.AddMenuEntry
 				(
-					LOCTEXT("ContextMenu_Header_TimerOptions_RemoveFromGraphTrack", "Remove series from graph track"),
-					LOCTEXT("ContextMenu_Header_TimerOptions_RemoveFromGraphTrack_Desc", "Remove the series containing event instances of the selected timer from the timing graph track."),
+					LOCTEXT("ContextMenu_RemoveFromGraphTrack", "Remove series from graph track"),
+					LOCTEXT("ContextMenu_RemoveFromGraphTrack_Desc", "Remove the series containing event instances of the selected timer from the timing graph track."),
 					FSlateIcon(FInsightsStyle::GetStyleSetName(), "Icons.ToggleShowGraphSeries"),
 					Action_ToggleTimerInGraphTrack,
 					NAME_None,
@@ -518,8 +522,8 @@ TSharedPtr<SWidget> STimersView::TreeView_GetMenuContent()
 			{
 				MenuBuilder.AddMenuEntry
 				(
-					LOCTEXT("ContextMenu_Header_TimerOptions_AddToGraphTrack", "Add series to graph track"),
-					LOCTEXT("ContextMenu_Header_TimerOptions_AddToGraphTrack_Desc", "Add a series containing event instances of the selected timer to the timing graph track."),
+					LOCTEXT("ContextMenu_AddToGraphTrack", "Add series to graph track"),
+					LOCTEXT("ContextMenu_AddToGraphTrack_Desc", "Add a series containing event instances of the selected timer to the timing graph track."),
 					FSlateIcon(FInsightsStyle::GetStyleSetName(), "Icons.ToggleShowGraphSeries"),
 					Action_ToggleTimerInGraphTrack,
 					NAME_None,
@@ -541,17 +545,17 @@ TSharedPtr<SWidget> STimersView::TreeView_GetMenuContent()
 				bIsValidSource = SelectedNode->GetSourceFileAndLine(File, Line);
 			}
 
-			FText ItemLabel = FText::Format(LOCTEXT("ContextMenu_Header_OpenSource", "Open Source in {0}"), SourceCodeAccessor.GetNameText());
+			FText ItemLabel = FText::Format(LOCTEXT("ContextMenu_OpenSource", "Open Source in {0}"), SourceCodeAccessor.GetNameText());
 
 			FText ItemToolTip;
 			if (bIsValidSource)
 			{
-				ItemToolTip = FText::Format(LOCTEXT("ContextMenu_Header_OpenSource_Desc1", "Opens the source file of the selected timer in {0}.\n{1} ({2})"),
+				ItemToolTip = FText::Format(LOCTEXT("ContextMenu_OpenSource_Desc1", "Opens the source file of the selected timer in {0}.\n{1} ({2})"),
 					SourceCodeAccessor.GetNameText(), FText::FromString(File), FText::AsNumber(Line, &FNumberFormattingOptions::DefaultNoGrouping()));
 			}
 			else
 			{
-				ItemToolTip = FText::Format(LOCTEXT("ContextMenu_Header_OpenSource_Desc2", "Opens the source file of the selected timer in {0}."),
+				ItemToolTip = FText::Format(LOCTEXT("ContextMenu_OpenSource_Desc2", "Opens the source file of the selected timer in {0}."),
 					SourceCodeAccessor.GetNameText());
 			}
 
@@ -566,7 +570,7 @@ TSharedPtr<SWidget> STimersView::TreeView_GetMenuContent()
 	}
 	MenuBuilder.EndSection();
 
-	MenuBuilder.BeginSection("Misc", LOCTEXT("ContextMenu_Header_Misc", "Miscellaneous"));
+	MenuBuilder.BeginSection("Misc", LOCTEXT("ContextMenu_Section_Misc", "Miscellaneous"));
 	{
 		MenuBuilder.AddMenuEntry
 		(
@@ -597,8 +601,8 @@ TSharedPtr<SWidget> STimersView::TreeView_GetMenuContent()
 
 		MenuBuilder.AddSubMenu
 		(
-			LOCTEXT("ContextMenu_Header_Misc_Export", "More Export Options"),
-			LOCTEXT("ContextMenu_Header_Misc_Export_Desc", "Exports threads, timers and timing events to text files."),
+			LOCTEXT("ContextMenu_ExportOptions_SubMenu", "More Export Options"),
+			LOCTEXT("ContextMenu_ExportOptions_SubMenu_Desc", "Exports threads, timers and timing events to text files."),
 			FNewMenuDelegate::CreateSP(this, &STimersView::TreeView_BuildExportMenu),
 			false,
 			FSlateIcon(FAppStyle::Get().GetStyleSetName(), "Icons.Save")
@@ -606,77 +610,6 @@ TSharedPtr<SWidget> STimersView::TreeView_GetMenuContent()
 	}
 	MenuBuilder.EndSection();
 
-	MenuBuilder.BeginSection("Sorting", LOCTEXT("ContextMenu_Header_Sorting", "Sorting"));
-	{
-		MenuBuilder.AddSubMenu
-		(
-			LOCTEXT("ContextMenu_Header_Misc_Sort", "Sort By"),
-			LOCTEXT("ContextMenu_Header_Misc_Sort_Desc", "Sort by column."),
-			FNewMenuDelegate::CreateSP(this, &STimersView::TreeView_BuildSortByMenu),
-			false,
-			FSlateIcon(FInsightsStyle::GetStyleSetName(), "Icons.SortBy")
-		);
-	}
-	MenuBuilder.EndSection();
-
-	MenuBuilder.BeginSection("Columns", LOCTEXT("ContextMenu_Header_Columns", "Columns"));
-	{
-		MenuBuilder.AddSubMenu
-		(
-			LOCTEXT("ContextMenu_Header_Columns_View", "View Column"),
-			LOCTEXT("ContextMenu_Header_Columns_View_Desc", "Hides or shows columns."),
-			FNewMenuDelegate::CreateSP(this, &STimersView::TreeView_BuildViewColumnMenu),
-			false,
-			FSlateIcon(FInsightsStyle::GetStyleSetName(), "Icons.ViewColumn")
-		);
-
-		FUIAction Action_ShowAllColumns
-		(
-			FExecuteAction::CreateSP(this, &STimersView::ContextMenu_ShowAllColumns_Execute),
-			FCanExecuteAction::CreateSP(this, &STimersView::ContextMenu_ShowAllColumns_CanExecute)
-		);
-		MenuBuilder.AddMenuEntry
-		(
-			LOCTEXT("ContextMenu_Header_Columns_ShowAllColumns", "Show All Columns"),
-			LOCTEXT("ContextMenu_Header_Columns_ShowAllColumns_Desc", "Resets tree view to show all columns."),
-			FSlateIcon(FInsightsStyle::GetStyleSetName(), "Icons.ResetColumn"),
-			Action_ShowAllColumns,
-			NAME_None,
-			EUserInterfaceActionType::Button
-		);
-
-		FUIAction Action_ShowMinMaxMedColumns
-		(
-			FExecuteAction::CreateSP(this, &STimersView::ContextMenu_ShowMinMaxMedColumns_Execute),
-			FCanExecuteAction::CreateSP(this, &STimersView::ContextMenu_ShowMinMaxMedColumns_CanExecute)
-		);
-		MenuBuilder.AddMenuEntry
-		(
-			LOCTEXT("ContextMenu_Header_Columns_ShowMinMaxMedColumns", "Reset Columns to Min/Max/Median Preset"),
-			LOCTEXT("ContextMenu_Header_Columns_ShowMinMaxMedColumns_Desc", "Resets columns to Min/Max/Median preset."),
-			FSlateIcon(FInsightsStyle::GetStyleSetName(), "Icons.ResetColumn"),
-			Action_ShowMinMaxMedColumns,
-			NAME_None,
-			EUserInterfaceActionType::Button
-		);
-
-		FUIAction Action_ResetColumns
-		(
-			FExecuteAction::CreateSP(this, &STimersView::ContextMenu_ResetColumns_Execute),
-			FCanExecuteAction::CreateSP(this, &STimersView::ContextMenu_ResetColumns_CanExecute)
-		);
-		MenuBuilder.AddMenuEntry
-		(
-			LOCTEXT("ContextMenu_Header_Columns_ResetColumns", "Reset Columns to Default"),
-			LOCTEXT("ContextMenu_Header_Columns_ResetColumns_Desc", "Resets columns to default."),
-			FSlateIcon(FInsightsStyle::GetStyleSetName(), "Icons.ResetColumn"),
-			Action_ResetColumns,
-			NAME_None,
-			EUserInterfaceActionType::Button
-		);
-	}
-	MenuBuilder.EndSection();
-
 	return MenuBuilder.MakeWidget();
 }
 
@@ -684,9 +617,7 @@ TSharedPtr<SWidget> STimersView::TreeView_GetMenuContent()
 
 void STimersView::TreeView_BuildSortByMenu(FMenuBuilder& MenuBuilder)
 {
-	// TODO: Refactor later @see TSharedPtr<SWidget> SCascadePreviewViewportToolBar::GenerateViewMenu() const
-
-	MenuBuilder.BeginSection("ColumnName", LOCTEXT("ContextMenu_Header_Misc_ColumnName", "Column Name"));
+	MenuBuilder.BeginSection("SortColumn", LOCTEXT("ContextMenu_Section_SortColumn", "Sort Column"));
 
 	for (const TSharedRef<Insights::FTableColumn>& ColumnRef : Table->GetColumns())
 	{
@@ -714,9 +645,7 @@ void STimersView::TreeView_BuildSortByMenu(FMenuBuilder& MenuBuilder)
 
 	MenuBuilder.EndSection();
 
-	//-----------------------------------------------------------------------------
-
-	MenuBuilder.BeginSection("SortMode", LOCTEXT("ContextMenu_Header_Misc_Sort_SortMode", "Sort Mode"));
+	MenuBuilder.BeginSection("SortMode", LOCTEXT("ContextMenu_Section_SortMode", "Sort Mode"));
 	{
 		FUIAction Action_SortAscending
 		(
@@ -726,8 +655,8 @@ void STimersView::TreeView_BuildSortByMenu(FMenuBuilder& MenuBuilder)
 		);
 		MenuBuilder.AddMenuEntry
 		(
-			LOCTEXT("ContextMenu_Header_Misc_Sort_SortAscending", "Sort Ascending"),
-			LOCTEXT("ContextMenu_Header_Misc_Sort_SortAscending_Desc", "Sorts ascending"),
+			LOCTEXT("ContextMenu_SortAscending", "Sort Ascending"),
+			LOCTEXT("ContextMenu_SortAscending_Desc", "Sorts ascending."),
 			FSlateIcon(FAppStyle::Get().GetStyleSetName(), "Icons.SortUp"),
 			Action_SortAscending,
 			NAME_None,
@@ -742,8 +671,8 @@ void STimersView::TreeView_BuildSortByMenu(FMenuBuilder& MenuBuilder)
 		);
 		MenuBuilder.AddMenuEntry
 		(
-			LOCTEXT("ContextMenu_Header_Misc_Sort_SortDescending", "Sort Descending"),
-			LOCTEXT("ContextMenu_Header_Misc_Sort_SortDescending_Desc", "Sorts descending"),
+			LOCTEXT("ContextMenu_SortDescending", "Sort Descending"),
+			LOCTEXT("ContextMenu_SortDescending_Desc", "Sorts descending."),
 			FSlateIcon(FAppStyle::Get().GetStyleSetName(), "Icons.SortDown"),
 			Action_SortDescending,
 			NAME_None,
@@ -757,7 +686,7 @@ void STimersView::TreeView_BuildSortByMenu(FMenuBuilder& MenuBuilder)
 
 void STimersView::TreeView_BuildViewColumnMenu(FMenuBuilder& MenuBuilder)
 {
-	MenuBuilder.BeginSection("ViewColumn", LOCTEXT("ContextMenu_Header_Columns_View", "View Column"));
+	MenuBuilder.BeginSection("Columns", LOCTEXT("ContextMenu_Section_Columns", "Columns"));
 
 	for (const TSharedRef<Insights::FTableColumn>& ColumnRef : Table->GetColumns())
 	{
@@ -787,7 +716,7 @@ void STimersView::TreeView_BuildViewColumnMenu(FMenuBuilder& MenuBuilder)
 
 void STimersView::TreeView_BuildExportMenu(FMenuBuilder& MenuBuilder)
 {
-	MenuBuilder.BeginSection("Export", LOCTEXT("ContextMenu_Header_Export", "Export"));
+	MenuBuilder.BeginSection("Export", LOCTEXT("ContextMenu_Section_Export", "Export"));
 	{
 		MenuBuilder.AddMenuEntry
 		(
@@ -850,38 +779,13 @@ FText STimersView::GetColumnHeaderText(const FName ColumnId) const
 
 TSharedRef<SWidget> STimersView::TreeViewHeaderRow_GenerateColumnMenu(const Insights::FTableColumn& Column)
 {
-	bool bIsMenuVisible = false;
-
 	const bool bShouldCloseWindowAfterMenuSelection = true;
 	FMenuBuilder MenuBuilder(bShouldCloseWindowAfterMenuSelection, NULL);
-	{
-		if (Column.CanBeHidden())
-		{
-			MenuBuilder.BeginSection("Column", LOCTEXT("TreeViewHeaderRow_Header_Column", "Column"));
-
-			FUIAction Action_HideColumn
-			(
-				FExecuteAction::CreateSP(this, &STimersView::HideColumn, Column.GetId()),
-				FCanExecuteAction::CreateSP(this, &STimersView::CanHideColumn, Column.GetId())
-			);
-			MenuBuilder.AddMenuEntry
-			(
-				LOCTEXT("TreeViewHeaderRow_HideColumn", "Hide"),
-				LOCTEXT("TreeViewHeaderRow_HideColumn_Desc", "Hides the selected column"),
-				FSlateIcon(),
-				Action_HideColumn,
-				NAME_None,
-				EUserInterfaceActionType::Button
-			);
-
-			bIsMenuVisible = true;
-			MenuBuilder.EndSection();
-		}
 
+	MenuBuilder.BeginSection("Sorting", LOCTEXT("ContextMenu_Section_Sorting", "Sorting"));
+	{
 		if (Column.CanBeSorted())
 		{
-			MenuBuilder.BeginSection("SortMode", LOCTEXT("ContextMenu_Header_Misc_Sort_SortMode", "Sort Mode"));
-
 			FUIAction Action_SortAscending
 			(
 				FExecuteAction::CreateSP(this, &STimersView::HeaderMenu_SortMode_Execute, Column.GetId(), EColumnSortMode::Ascending),
@@ -890,8 +794,8 @@ TSharedRef<SWidget> STimersView::TreeViewHeaderRow_GenerateColumnMenu(const Insi
 			);
 			MenuBuilder.AddMenuEntry
 			(
-				LOCTEXT("ContextMenu_Header_Misc_Sort_SortAscending", "Sort Ascending"),
-				LOCTEXT("ContextMenu_Header_Misc_Sort_SortAscending_Desc", "Sorts ascending"),
+				FText::Format(LOCTEXT("ContextMenu_SortAscending_Fmt", "Sort Ascending (by {0})"), Column.GetTitleName()),
+				FText::Format(LOCTEXT("ContextMenu_SortAscending_Desc_Fmt", "Sorts ascending by {0}."), Column.GetTitleName()),
 				FSlateIcon(FAppStyle::Get().GetStyleSetName(), "Icons.SortUp"),
 				Action_SortAscending,
 				NAME_None,
@@ -906,36 +810,103 @@ TSharedRef<SWidget> STimersView::TreeViewHeaderRow_GenerateColumnMenu(const Insi
 			);
 			MenuBuilder.AddMenuEntry
 			(
-				LOCTEXT("ContextMenu_Header_Misc_Sort_SortDescending", "Sort Descending"),
-				LOCTEXT("ContextMenu_Header_Misc_Sort_SortDescending_Desc", "Sorts descending"),
+				FText::Format(LOCTEXT("ContextMenu_SortDescending_Fmt", "Sort Descending (by {0})"), Column.GetTitleName()),
+				FText::Format(LOCTEXT("ContextMenu_SortDescending_Desc_Fmt", "Sorts descending by {0}."), Column.GetTitleName()),
 				FSlateIcon(FAppStyle::Get().GetStyleSetName(), "Icons.SortDown"),
 				Action_SortDescending,
 				NAME_None,
 				EUserInterfaceActionType::RadioButton
 			);
-
-			bIsMenuVisible = true;
-			MenuBuilder.EndSection();
 		}
 
-		//if (Column.CanBeFiltered())
-		//{
-		//	MenuBuilder.BeginSection("FilterMode", LOCTEXT("ContextMenu_Header_Misc_Filter_FilterMode", "Filter Mode"));
-		//	bIsMenuVisible = true;
-		//	MenuBuilder.EndSection();
-		//}
+		MenuBuilder.AddSubMenu
+		(
+			LOCTEXT("ContextMenu_SortBy_SubMenu", "Sort By"),
+			LOCTEXT("ContextMenu_SortBy_SubMenu_Desc", "Sorts by a column."),
+			FNewMenuDelegate::CreateSP(this, &STimersView::TreeView_BuildSortByMenu),
+			false,
+			FSlateIcon(FInsightsStyle::GetStyleSetName(), "Icons.SortBy")
+		);
 	}
+	MenuBuilder.EndSection();
 
-	/*
-	TODO:
-	- Show top ten
-	- Show top bottom
-	- Filter by list (avg, median, 10%, 90%, etc.)
-	- Text box for filtering for each column instead of one text box used for filtering
-	- Grouping button for flat view modes (show at most X groups, show all groups for names)
-	*/
+	MenuBuilder.BeginSection("ColumnVisibility", LOCTEXT("ContextMenu_Section_ColumnVisibility", "Column Visibility"));
+	{
+		if (Column.CanBeHidden())
+		{
+			FUIAction Action_HideColumn
+			(
+				FExecuteAction::CreateSP(this, &STimersView::HideColumn, Column.GetId()),
+				FCanExecuteAction::CreateSP(this, &STimersView::CanHideColumn, Column.GetId())
+			);
+			MenuBuilder.AddMenuEntry
+			(
+				LOCTEXT("ContextMenu_HideColumn", "Hide"),
+				LOCTEXT("ContextMenu_HideColumn_Desc", "Hides the selected column."),
+				FSlateIcon(),
+				Action_HideColumn,
+				NAME_None,
+				EUserInterfaceActionType::Button
+			);
+		}
 
-	return bIsMenuVisible ? MenuBuilder.MakeWidget() : (TSharedRef<SWidget>)SNullWidget::NullWidget;
+		MenuBuilder.AddSubMenu
+		(
+			LOCTEXT("ContextMenu_ViewColumn_SubMenu", "View Column"),
+			LOCTEXT("ContextMenu_ViewColumn_SubMenu_Desc", "Hides or shows columns."),
+			FNewMenuDelegate::CreateSP(this, &STimersView::TreeView_BuildViewColumnMenu),
+			false,
+			FSlateIcon(FInsightsStyle::GetStyleSetName(), "Icons.ViewColumn")
+		);
+
+		FUIAction Action_ShowAllColumns
+		(
+			FExecuteAction::CreateSP(this, &STimersView::ContextMenu_ShowAllColumns_Execute),
+			FCanExecuteAction::CreateSP(this, &STimersView::ContextMenu_ShowAllColumns_CanExecute)
+		);
+		MenuBuilder.AddMenuEntry
+		(
+			LOCTEXT("ContextMenu_ShowAllColumns", "Show All Columns"),
+			LOCTEXT("ContextMenu_ShowAllColumns_Desc", "Resets tree view to show all columns."),
+			FSlateIcon(FInsightsStyle::GetStyleSetName(), "Icons.ResetColumn"),
+			Action_ShowAllColumns,
+			NAME_None,
+			EUserInterfaceActionType::Button
+		);
+
+		FUIAction Action_ShowMinMaxMedColumns
+		(
+			FExecuteAction::CreateSP(this, &STimersView::ContextMenu_ShowMinMaxMedColumns_Execute),
+			FCanExecuteAction::CreateSP(this, &STimersView::ContextMenu_ShowMinMaxMedColumns_CanExecute)
+		);
+		MenuBuilder.AddMenuEntry
+		(
+			LOCTEXT("ContextMenu_ShowMinMaxMedColumns", "Reset Columns to Min/Max/Median Preset"),
+			LOCTEXT("ContextMenu_ShowMinMaxMedColumns_Desc", "Resets columns to Min/Max/Median preset."),
+			FSlateIcon(FInsightsStyle::GetStyleSetName(), "Icons.ResetColumn"),
+			Action_ShowMinMaxMedColumns,
+			NAME_None,
+			EUserInterfaceActionType::Button
+		);
+
+		FUIAction Action_ResetColumns
+		(
+			FExecuteAction::CreateSP(this, &STimersView::ContextMenu_ResetColumns_Execute),
+			FCanExecuteAction::CreateSP(this, &STimersView::ContextMenu_ResetColumns_CanExecute)
+		);
+		MenuBuilder.AddMenuEntry
+		(
+			LOCTEXT("ContextMenu_ResetColumns", "Reset Columns to Default"),
+			LOCTEXT("ContextMenu_ResetColumns_Desc", "Resets columns to default."),
+			FSlateIcon(FInsightsStyle::GetStyleSetName(), "Icons.ResetColumn"),
+			Action_ResetColumns,
+			NAME_None,
+			EUserInterfaceActionType::Button
+		);
+	}
+	MenuBuilder.EndSection();
+
+	return MenuBuilder.MakeWidget();
 }
 
 ////////////////////////////////////////////////////////////////////////////////////////////////////
