[Insights] Increased the precision of the time value displayed in the tooltip of a Graph track. The precision now depends on the viewport scaling (similar with values on the Time Ruler track).

diff --git a/Engine/Source/Developer/TraceInsights/Private/Insights/MemoryProfiler/ViewModels/MemoryGraphTrack.cpp b/Engine/Source/Developer/TraceInsights/Private/Insights/MemoryProfiler/ViewModels/MemoryGraphTrack.cpp
--- a/Engine/Source/Developer/TraceInsights/Private/Insights/MemoryProfiler/ViewModels/MemoryGraphTrack.cpp
+++ b/Engine/Source/Developer/TraceInsights/Private/Insights/MemoryProfiler/ViewModels/MemoryGraphTrack.cpp
@@ -1126,7 +1126,8 @@ void FMemoryGraphTrack::InitTooltip(FTooltipDrawState& InOutTooltip, const ITimi
 			InOutTooltip.AddTitle(SubTitle, Series->GetColor());
 		}
 
-		InOutTooltip.AddNameValueTextLine(TEXT("Time:"), TimeUtils::FormatTimeAuto(TooltipEvent.GetStartTime()));
+		const double Precision = FMath::Max(1.0 / TimeScaleX, TimeUtils::Nanosecond);
+		InOutTooltip.AddNameValueTextLine(TEXT("Time:"), TimeUtils::FormatTime(TooltipEvent.GetStartTime(), Precision));
 		if (Series->HasEventDuration())
 		{
 			InOutTooltip.AddNameValueTextLine(TEXT("Duration:"), TimeUtils::FormatTimeAuto(TooltipEvent.GetDuration()));
diff --git a/Engine/Source/Developer/TraceInsights/Private/Insights/ViewModels/GraphTrack.cpp b/Engine/Source/Developer/TraceInsights/Private/Insights/ViewModels/GraphTrack.cpp
--- a/Engine/Source/Developer/TraceInsights/Private/Insights/ViewModels/GraphTrack.cpp
+++ b/Engine/Source/Developer/TraceInsights/Private/Insights/ViewModels/GraphTrack.cpp
@@ -63,6 +63,7 @@ FGraphTrack::FGraphTrack(const FString& InName)
 	, VisibleOptions(EGraphOptions::DefaultVisibleOptions)
 	, EditableOptions(EGraphOptions::DefaultEditableOptions)
 	, SharedValueViewport()
+	, TimeScaleX(1.0)
 	, NumAddedEvents(0)
 	, NumDrawPoints(0)
 	, NumDrawLines(0)
@@ -103,6 +104,8 @@ void FGraphTrack::PostUpdate(const ITimingTrackUpdateContext& Context)
 	{
 		SetHoveredState(false);
 	}
+
+	TimeScaleX = Context.GetViewport().GetScaleX();
 }
 
 ////////////////////////////////////////////////////////////////////////////////////////////////////
@@ -607,7 +610,8 @@ void FGraphTrack::InitTooltip(FTooltipDrawState& InOutTooltip, const ITimingEven
 
 		InOutTooltip.ResetContent();
 		InOutTooltip.AddTitle(Series->GetName().ToString(), Series->GetColor());
-		InOutTooltip.AddNameValueTextLine(TEXT("Time:"), TimeUtils::FormatTimeAuto(TooltipEvent.GetStartTime()));
+		const double Precision = FMath::Max(1.0 / TimeScaleX, TimeUtils::Nanosecond);
+		InOutTooltip.AddNameValueTextLine(TEXT("Time:"), TimeUtils::FormatTime(TooltipEvent.GetStartTime(), Precision));
 		if (Series->HasEventDuration())
 		{
 			InOutTooltip.AddNameValueTextLine(TEXT("Duration:"), TimeUtils::FormatTimeAuto(TooltipEvent.GetDuration()));
diff --git a/Engine/Source/Developer/TraceInsights/Public/Insights/ViewModels/GraphTrack.h b/Engine/Source/Developer/TraceInsights/Public/Insights/ViewModels/GraphTrack.h
--- a/Engine/Source/Developer/TraceInsights/Public/Insights/ViewModels/GraphTrack.h
+++ b/Engine/Source/Developer/TraceInsights/Public/Insights/ViewModels/GraphTrack.h
@@ -146,6 +146,8 @@ protected:
 
 	FGraphValueViewport SharedValueViewport;
 
+	double TimeScaleX;
+
 	// Stats
 	int32 NumAddedEvents; // total event count
 	int32 NumDrawPoints;
