Quick Find filter functionality implemented for core tracks

diff --git a/Engine/Source/Developer/TraceInsights/Private/Insights/ContextSwitches/ViewModels/ContextSwitchesFilterConverters.cpp b/Engine/Source/Developer/TraceInsights/Private/Insights/ContextSwitches/ViewModels/ContextSwitchesFilterConverters.cpp
new file mode 100644
--- /dev/null
+++ b/Engine/Source/Developer/TraceInsights/Private/Insights/ContextSwitches/ViewModels/ContextSwitchesFilterConverters.cpp
@@ -0,0 +1,70 @@
+// Copyright Epic Games, Inc. All Rights Reserved.
+
+#include "ContextSwitchesFilterConverters.h"
+
+#include "TraceServices/Model/ContextSwitches.h"
+
+#include "Insights/InsightsManager.h"
+
+#define LOCTEXT_NAMESPACE "Insights::ContextSwitchesFilterValueConverter"
+
+namespace Insights
+{
+
+////////////////////////////////////////////////////////////////////////////////////////////////////
+
+bool FCoreEventNameFilterValueConverter::Convert(const FString& Input, int64& Output, FText& OutError) const
+{
+	TSharedPtr<const TraceServices::IAnalysisSession> Session = FInsightsManager::Get()->GetSession();
+	bool bFound = false;
+	if (Session.IsValid())
+	{
+		TraceServices::FAnalysisSessionReadScope SessionReadScope(*Session.Get());
+		const TraceServices::IThreadProvider& ThreadProvider = TraceServices::ReadThreadProvider(*Session.Get());
+		ThreadProvider.EnumerateThreads([&Input, &Output, &Session, &bFound](const TraceServices::FThreadInfo& ThreadInfo)
+			{
+				if (!bFound && FCString::Stricmp(*Input, ThreadInfo.Name) == 0)
+				{
+					const TraceServices::IContextSwitchesProvider* ContextSwitchesProvider = TraceServices::ReadContextSwitchesProvider(*Session.Get());
+					if (ContextSwitchesProvider)
+					{
+						uint32 SystemThreadId;
+						if (ContextSwitchesProvider->GetSystemThreadId(ThreadInfo.Id, SystemThreadId))
+						{
+							Output = SystemThreadId;
+							bFound = true;
+						}
+					}
+				}
+			});
+	}
+
+	if (!bFound)
+	{
+		OutError = LOCTEXT("NoCoreEventFound", "No core event with this name was found!");
+		return false;
+	}
+	
+	return true;
+}
+
+////////////////////////////////////////////////////////////////////////////////////////////////////
+
+FText FCoreEventNameFilterValueConverter::GetTooltipText() const
+{
+	return LOCTEXT("FCoreEventNameFilterValueConverterTooltip", "Enter the exact name of the core event.");
+}
+
+////////////////////////////////////////////////////////////////////////////////////////////////////
+
+FText FCoreEventNameFilterValueConverter::GetHintText() const
+{
+	// Use the name of a well known event as the hint to show the user what kind of name we are expecting.
+	return FText::FromString(TEXT("GameThread"));
+}
+
+////////////////////////////////////////////////////////////////////////////////////////////////////
+
+} // namespace Insights
+
+#undef LOCTEXT_NAMESPACE
\ No newline at end of file
diff --git a/Engine/Source/Developer/TraceInsights/Private/Insights/ContextSwitches/ViewModels/ContextSwitchesFilterConverters.h b/Engine/Source/Developer/TraceInsights/Private/Insights/ContextSwitches/ViewModels/ContextSwitchesFilterConverters.h
new file mode 100644
--- /dev/null
+++ b/Engine/Source/Developer/TraceInsights/Private/Insights/ContextSwitches/ViewModels/ContextSwitchesFilterConverters.h
@@ -0,0 +1,22 @@
+// Copyright Epic Games, Inc. All Rights Reserved.
+
+#pragma once
+
+#include "Insights/ViewModels/Filters.h"
+
+namespace Insights
+{
+
+////////////////////////////////////////////////////////////////////////////////////////////////////
+
+class FCoreEventNameFilterValueConverter : public IFilterValueConverter
+{
+public:
+	virtual bool Convert(const FString& Input, int64& Output, FText& OutError) const override;
+	virtual FText GetTooltipText() const override;
+	virtual FText GetHintText() const override;
+};
+
+////////////////////////////////////////////////////////////////////////////////////////////////////
+
+} // namespace Insights
\ No newline at end of file
diff --git a/Engine/Source/Developer/TraceInsights/Private/Insights/ContextSwitches/ViewModels/ContextSwitchesSharedState.cpp b/Engine/Source/Developer/TraceInsights/Private/Insights/ContextSwitches/ViewModels/ContextSwitchesSharedState.cpp
--- a/Engine/Source/Developer/TraceInsights/Private/Insights/ContextSwitches/ViewModels/ContextSwitchesSharedState.cpp
+++ b/Engine/Source/Developer/TraceInsights/Private/Insights/ContextSwitches/ViewModels/ContextSwitchesSharedState.cpp
@@ -9,6 +9,7 @@
 
 // Insights
 #include "Insights/ContextSwitches/ContextSwitchesProfilerManager.h"
+#include "Insights/ContextSwitches/ViewModels/ContextSwitchesFilterConverters.h"
 #include "Insights/ContextSwitches/ViewModels/ContextSwitchesTimingTrack.h"
 #include "Insights/ContextSwitches/ViewModels/ContextSwitchTimingEvent.h"
 #include "Insights/ContextSwitches/ViewModels/CpuCoreTimingTrack.h"
@@ -16,6 +17,7 @@
 #include "Insights/InsightsManager.h"
 #include "Insights/InsightsStyle.h"
 #include "Insights/TimingProfilerManager.h"
+#include "Insights/ViewModels/FilterConfigurator.h"
 #include "Insights/ViewModels/ThreadTimingTrack.h"
 #include "Insights/Widgets/STimingProfilerWindow.h"
 #include "Insights/Widgets/STimingView.h"
@@ -255,6 +257,45 @@ void FContextSwitchesSharedState::BuildSubMenu(FMenuBuilder& InMenuBuilder)
 
 ////////////////////////////////////////////////////////////////////////////////////////////////////
 
+void FContextSwitchesSharedState::AddQuickFindFilters(TSharedPtr<class FFilterConfigurator> FilterConfigurator)
+{
+	TSharedPtr<TArray<TSharedPtr<struct FFilter>>>& AvailableFilters = FilterConfigurator->GetAvailableFilters();
+
+	TSharedPtr<TArray<TSharedPtr<IFilterOperator>>> CoreEventNameFilterOperators = MakeShared<TArray<TSharedPtr<IFilterOperator>>>();
+	CoreEventNameFilterOperators->Add(StaticCastSharedRef<IFilterOperator>(MakeShared<FFilterOperator<int64>>(EFilterOperator::Eq, TEXT("Is"), [](int64 lhs, int64 rhs) { return lhs == rhs; })));
+
+	TSharedPtr<FFilterWithSuggestions> TimerNameFilter = MakeShared<FFilterWithSuggestions>(static_cast<int32>(EFilterField::CoreEventName), LOCTEXT("CoreEventName", "Core Event Name"), LOCTEXT("CoreEventName", "Core Event Name"), EFilterDataType::StringInt64Pair, CoreEventNameFilterOperators);
+	TimerNameFilter->Callback = [this](const FString& Text, TArray<FString>& OutSuggestions)
+	{
+		this->PopulateCoreEventNameSuggestionList(Text, OutSuggestions);
+	};
+
+	TimerNameFilter->Converter = MakeShared<FCoreEventNameFilterValueConverter>();
+
+	AvailableFilters->Add(TimerNameFilter);
+}
+
+////////////////////////////////////////////////////////////////////////////////////////////////////
+
+void FContextSwitchesSharedState::PopulateCoreEventNameSuggestionList(const FString& Text, TArray<FString>& OutSuggestions)
+{
+	TSharedPtr<const TraceServices::IAnalysisSession> Session = FInsightsManager::Get()->GetSession();
+	if (Session.IsValid())
+	{
+		TraceServices::FAnalysisSessionReadScope SessionReadScope(*Session.Get());
+		const TraceServices::IThreadProvider& ThreadProvider = TraceServices::ReadThreadProvider(*Session.Get());
+		ThreadProvider.EnumerateThreads([&Text, &OutSuggestions](const TraceServices::FThreadInfo& ThreadInfo)
+			{
+				if (FCString::Stristr(ThreadInfo.Name, *Text))
+				{
+					OutSuggestions.Add(ThreadInfo.Name);
+				}
+			});
+	}
+}
+
+////////////////////////////////////////////////////////////////////////////////////////////////////
+
 void FContextSwitchesSharedState::AddCommands()
 {
 	FContextSwitchesStateCommands::Register();
diff --git a/Engine/Source/Developer/TraceInsights/Private/Insights/ContextSwitches/ViewModels/ContextSwitchesSharedState.h b/Engine/Source/Developer/TraceInsights/Private/Insights/ContextSwitches/ViewModels/ContextSwitchesSharedState.h
--- a/Engine/Source/Developer/TraceInsights/Private/Insights/ContextSwitches/ViewModels/ContextSwitchesSharedState.h
+++ b/Engine/Source/Developer/TraceInsights/Private/Insights/ContextSwitches/ViewModels/ContextSwitchesSharedState.h
@@ -67,6 +67,7 @@ public:
 	virtual void OnEndSession(Insights::ITimingViewSession& InSession) override;
 	virtual void Tick(Insights::ITimingViewSession& InSession, const TraceServices::IAnalysisSession& InAnalysisSession) override;
 	virtual void ExtendCpuTracksFilterMenu(Insights::ITimingViewSession& InSession, FMenuBuilder& InMenuBuilder) override;
+	virtual void AddQuickFindFilters(TSharedPtr<class FFilterConfigurator> FilterConfigurator) override;
 
 	//////////////////////////////////////////////////
 
@@ -149,6 +150,8 @@ private:
 	void Command_DockCpuCoreTrackToTop_Execute();
 	bool Command_DockCpuCoreTrackToTop_CanExecute() const { return AreContextSwitchesAvailable() && AreCoreTracksVisible() && IsValidContextSwitchEventSelected(); }
 
+	void PopulateCoreEventNameSuggestionList(const FString& Text, TArray<FString>& OutSuggestions);
+
 private:
 	STimingView* TimingView;
 
diff --git a/Engine/Source/Developer/TraceInsights/Private/Insights/ContextSwitches/ViewModels/CpuCoreTimingTrack.cpp b/Engine/Source/Developer/TraceInsights/Private/Insights/ContextSwitches/ViewModels/CpuCoreTimingTrack.cpp
--- a/Engine/Source/Developer/TraceInsights/Private/Insights/ContextSwitches/ViewModels/CpuCoreTimingTrack.cpp
+++ b/Engine/Source/Developer/TraceInsights/Private/Insights/ContextSwitches/ViewModels/CpuCoreTimingTrack.cpp
@@ -11,6 +11,9 @@
 #include "Insights/ContextSwitches/ViewModels/ContextSwitchesSharedState.h"
 #include "Insights/ContextSwitches/ViewModels/ContextSwitchTimingEvent.h"
 #include "Insights/InsightsManager.h"
+#include "Insights/ViewModels/FilterConfigurator.h"
+#include "Insights/ViewModels/Filters.h"
+#include "Insights/ViewModels/ThreadTimingTrack.h"
 #include "Insights/ViewModels/TimingTrackViewport.h"
 #include "Insights/ViewModels/TimingViewDrawHelper.h"
 #include "Insights/ViewModels/TooltipDrawState.h"
@@ -45,52 +48,10 @@ void FCpuCoreTimingTrack::BuildDrawState(ITimingEventsTrackDrawStateBuilder& Bui
 
 	const FTimingTrackViewport& Viewport = Context.GetViewport();
 
-	ContextSwitchesProvider->EnumerateCpuCoreEvents(CoreNumber, Viewport.GetStartTime(), Viewport.GetEndTime(), [&Builder, ContextSwitchesProvider](const TraceServices::FCpuCoreEvent& CpuCoreEvent)
+	ContextSwitchesProvider->EnumerateCpuCoreEvents(CoreNumber, Viewport.GetStartTime(), Viewport.GetEndTime(), 
+		[&Builder, ContextSwitchesProvider, this](const TraceServices::FCpuCoreEvent& CpuCoreEvent)
 		{
-			constexpr uint32 EventDepth = 0;
-			const uint32 EventColor = FTimingEvent::ComputeEventColor(CpuCoreEvent.SystemThreadId);
-			const uint32 SystemThreadId = CpuCoreEvent.SystemThreadId;
-
-			Builder.AddEvent(CpuCoreEvent.Start, CpuCoreEvent.End, EventDepth, EventColor,
-				[SystemThreadId, &CpuCoreEvent](float Width) -> const FString
-				{
-					const TCHAR* ThreadName = nullptr;
-
-					TSharedPtr<const TraceServices::IAnalysisSession> Session = FInsightsManager::Get()->GetSession();
-					if (Session.IsValid())
-					{
-						TraceServices::FAnalysisSessionReadScope SessionReadScope(*Session.Get());
-						const TraceServices::IContextSwitchesProvider* ContextSwitchesProvider = TraceServices::ReadContextSwitchesProvider(*Session.Get());
-						if (ContextSwitchesProvider)
-						{
-							uint32 ThreadId;
-							if (ContextSwitchesProvider->GetThreadId(SystemThreadId, ThreadId))
-							{
-								const TraceServices::IThreadProvider& ThreadProvider = TraceServices::ReadThreadProvider(*Session.Get());
-								ThreadName = ThreadProvider.GetThreadName(ThreadId);
-							}
-						}
-					}
-
-					FString EventName;
-
-					if (ThreadName)
-					{
-						EventName = ThreadName;
-					}
-					else
-					{
-						EventName = FString::Printf(TEXT("Unknown %d"), SystemThreadId);
-					}
-
-					if (Width > EventName.Len() * 4.0f + 32.0f)
-					{
-						const double Duration = CpuCoreEvent.End - CpuCoreEvent.Start;
-						FTimingEventsTrackDrawStateBuilder::AppendDurationToEventName(EventName, Duration);
-					}
-
-					return EventName;
-				});
+			this->AddCoreTimingEvent(Builder, CpuCoreEvent);
 
 			return TraceServices::EContextSwitchEnumerationResult::Continue;
 		});
@@ -123,47 +84,15 @@ void FCpuCoreTimingTrack::BuildFilteredDrawState(ITimingEventsTrackDrawStateBuil
 
 			if (bFilterOnlyByEventType)
 			{
-				ContextSwitchesProvider->EnumerateCpuCoreEvents(CoreNumber, Viewport.GetStartTime(), Viewport.GetEndTime(), [&Builder, ContextSwitchesProvider, FilterEventType](const TraceServices::FCpuCoreEvent& CpuCoreEvent)
+				ContextSwitchesProvider->EnumerateCpuCoreEvents(CoreNumber, Viewport.GetStartTime(), Viewport.GetEndTime(), 
+					[&Builder, ContextSwitchesProvider, FilterEventType, this](const TraceServices::FCpuCoreEvent& CpuCoreEvent)
 					{
 						if (CpuCoreEvent.SystemThreadId != FilterEventType)
 						{
 							return TraceServices::EContextSwitchEnumerationResult::Continue;
 						}
 
-						constexpr uint32 EventDepth = 0;
-						const uint32 EventColor = FTimingEvent::ComputeEventColor(CpuCoreEvent.SystemThreadId);
-						const uint32 SystemThreadId = CpuCoreEvent.SystemThreadId;
-
-						Builder.AddEvent(CpuCoreEvent.Start, CpuCoreEvent.End, EventDepth, EventColor,
-							[SystemThreadId](float Width) -> const FString
-							{
-								const TCHAR* ThreadName = nullptr;
-
-								TSharedPtr<const TraceServices::IAnalysisSession> Session = FInsightsManager::Get()->GetSession();
-								if (Session.IsValid())
-								{
-									TraceServices::FAnalysisSessionReadScope SessionReadScope(*Session.Get());
-									const TraceServices::IContextSwitchesProvider* ContextSwitchesProvider = TraceServices::ReadContextSwitchesProvider(*Session.Get());
-									if (ContextSwitchesProvider)
-									{
-										uint32 ThreadId;
-										if (ContextSwitchesProvider->GetThreadId(SystemThreadId, ThreadId))
-										{
-											const TraceServices::IThreadProvider& ThreadProvider = TraceServices::ReadThreadProvider(*Session.Get());
-											ThreadName = ThreadProvider.GetThreadName(ThreadId);
-										}
-									}
-								}
-
-								if (ThreadName)
-								{
-									return FString(ThreadName);
-								}
-								else
-								{
-									return FString::Printf(TEXT("Unknown %d"), SystemThreadId);
-								}
-							});
+						this->AddCoreTimingEvent(Builder, CpuCoreEvent);
 
 						return TraceServices::EContextSwitchEnumerationResult::Continue;
 					});
@@ -174,6 +103,103 @@ void FCpuCoreTimingTrack::BuildFilteredDrawState(ITimingEventsTrackDrawStateBuil
 			}
 		}
 	}
+
+	if (HasCustomFilter())
+	{
+		STimingView* TimingView = SharedState.GetTimingView();
+		if (!TimingView)
+		{
+			return;
+		}
+
+		TSharedPtr<FFilterConfigurator> FilterConfigurator = TimingView->GetFilterConfigurator();
+		if (!FilterConfigurator.IsValid())
+		{
+			return;
+		}
+
+		FFilterContext FilterContext;
+		FilterContext.SetReturnValueForUnsetFilters(false);
+
+		FilterContext.AddFilterData<double>(static_cast<int32>(EFilterField::StartTime), 0.0f);
+		FilterContext.AddFilterData<double>(static_cast<int32>(EFilterField::EndTime), 0.0f);
+		FilterContext.AddFilterData<double>(static_cast<int32>(EFilterField::Duration), 0.0f);
+		FilterContext.AddFilterData<FString>(static_cast<int32>(EFilterField::TrackName), this->GetName());
+		FilterContext.AddFilterData<int64>(static_cast<int32>(EFilterField::CoreEventName), 0);
+
+		TSharedPtr<const TraceServices::IAnalysisSession> Session = FInsightsManager::Get()->GetSession();
+		TraceServices::FAnalysisSessionReadScope SessionReadScope(*Session.Get());
+		const TraceServices::IContextSwitchesProvider* ContextSwitchesProvider = TraceServices::ReadContextSwitchesProvider(*Session.Get());
+		const FTimingTrackViewport& Viewport = Context.GetViewport();
+
+		ContextSwitchesProvider->EnumerateCpuCoreEvents(CoreNumber, Viewport.GetStartTime(), Viewport.GetEndTime(), 
+			[&Builder, ContextSwitchesProvider, &FilterContext, &FilterConfigurator, this](const TraceServices::FCpuCoreEvent& CpuCoreEvent)
+			{
+				FilterContext.SetFilterData<double>(static_cast<int32>(EFilterField::StartTime), CpuCoreEvent.Start);
+				FilterContext.SetFilterData<double>(static_cast<int32>(EFilterField::EndTime), CpuCoreEvent.End);
+				FilterContext.SetFilterData<double>(static_cast<int32>(EFilterField::Duration), CpuCoreEvent.End - CpuCoreEvent.Start);
+				FilterContext.SetFilterData<int64>(static_cast<int32>(EFilterField::CoreEventName), CpuCoreEvent.SystemThreadId);
+
+				if (!FilterConfigurator->ApplyFilters(FilterContext))
+				{
+					return TraceServices::EContextSwitchEnumerationResult::Continue;
+				}
+
+				this->AddCoreTimingEvent(Builder, CpuCoreEvent);
+
+				return TraceServices::EContextSwitchEnumerationResult::Continue;
+			});
+	}
+}
+
+////////////////////////////////////////////////////////////////////////////////////////////////////
+
+void FCpuCoreTimingTrack::AddCoreTimingEvent(ITimingEventsTrackDrawStateBuilder& Builder, const TraceServices::FCpuCoreEvent& CpuCoreEvent)
+{
+	constexpr uint32 EventDepth = 0;
+	const uint32 EventColor = FTimingEvent::ComputeEventColor(CpuCoreEvent.SystemThreadId);
+	const uint32 SystemThreadId = CpuCoreEvent.SystemThreadId;
+
+	Builder.AddEvent(CpuCoreEvent.Start, CpuCoreEvent.End, EventDepth, EventColor,
+		[SystemThreadId, &CpuCoreEvent](float Width) -> const FString
+		{
+			const TCHAR* ThreadName = nullptr;
+
+			TSharedPtr<const TraceServices::IAnalysisSession> Session = FInsightsManager::Get()->GetSession();
+			if (Session.IsValid())
+			{
+				TraceServices::FAnalysisSessionReadScope SessionReadScope(*Session.Get());
+				const TraceServices::IContextSwitchesProvider* ContextSwitchesProvider = TraceServices::ReadContextSwitchesProvider(*Session.Get());
+				if (ContextSwitchesProvider)
+				{
+					uint32 ThreadId;
+					if (ContextSwitchesProvider->GetThreadId(SystemThreadId, ThreadId))
+					{
+						const TraceServices::IThreadProvider& ThreadProvider = TraceServices::ReadThreadProvider(*Session.Get());
+						ThreadName = ThreadProvider.GetThreadName(ThreadId);
+					}
+				}
+			}
+
+			FString EventName;
+
+			if (ThreadName)
+			{
+				EventName = ThreadName;
+			}
+			else
+			{
+				EventName = FString::Printf(TEXT("Unknown %d"), SystemThreadId);
+			}
+
+			if (Width > EventName.Len() * 4.0f + 32.0f)
+			{
+				const double Duration = CpuCoreEvent.End - CpuCoreEvent.Start;
+				FTimingEventsTrackDrawStateBuilder::AppendDurationToEventName(EventName, Duration);
+			}
+
+			return EventName;
+		});
 }
 
 ////////////////////////////////////////////////////////////////////////////////////////////////////
@@ -309,6 +335,20 @@ void FCpuCoreTimingTrack::BuildContextMenu(FMenuBuilder& InOutMenuBuilder)
 
 ////////////////////////////////////////////////////////////////////////////////////////////////////
 
+bool FCpuCoreTimingTrack::HasCustomFilter() const
+{
+	STimingView* TimingView = SharedState.GetTimingView();
+	if (!TimingView)
+	{
+		return false;
+	}
+
+	TSharedPtr<FFilterConfigurator> FilterConfigurator = TimingView->GetFilterConfigurator();
+	return FilterConfigurator.IsValid() && FilterConfigurator->GetRootNode().IsValid() && FilterConfigurator->GetRootNode()->GetChildren().Num() > 0;
+}
+
+////////////////////////////////////////////////////////////////////////////////////////////////////
+
 } // namespace Insights
 
 #undef LOCTEXT_NAMESPACE
diff --git a/Engine/Source/Developer/TraceInsights/Private/Insights/ContextSwitches/ViewModels/CpuCoreTimingTrack.h b/Engine/Source/Developer/TraceInsights/Private/Insights/ContextSwitches/ViewModels/CpuCoreTimingTrack.h
--- a/Engine/Source/Developer/TraceInsights/Private/Insights/ContextSwitches/ViewModels/CpuCoreTimingTrack.h
+++ b/Engine/Source/Developer/TraceInsights/Private/Insights/ContextSwitches/ViewModels/CpuCoreTimingTrack.h
@@ -11,6 +11,11 @@
 
 class FThreadTrackEvent;
 
+namespace TraceServices
+{
+	struct FCpuCoreEvent;
+};
+
 namespace Insights
 {
 
@@ -46,6 +51,10 @@ public:
 protected:
 	virtual const TSharedPtr<const ITimingEvent> GetEvent(double InTime, double SecondsPerPixel, int32 Depth) const override;
 
+	virtual bool HasCustomFilter() const override;
+
+	void AddCoreTimingEvent(ITimingEventsTrackDrawStateBuilder& Builder, const TraceServices::FCpuCoreEvent& CpuCoreEvent);
+
 private:
 	FContextSwitchesSharedState& SharedState;
 
diff --git a/Engine/Source/Developer/TraceInsights/Private/Insights/ViewModels/FilterConfiguratorNode.cpp b/Engine/Source/Developer/TraceInsights/Private/Insights/ViewModels/FilterConfiguratorNode.cpp
--- a/Engine/Source/Developer/TraceInsights/Private/Insights/ViewModels/FilterConfiguratorNode.cpp
+++ b/Engine/Source/Developer/TraceInsights/Private/Insights/ViewModels/FilterConfiguratorNode.cpp
@@ -264,8 +264,8 @@ bool FFilterConfiguratorNode::ApplyFilters(const FFilterContext& Context) const
 	{
 		if (!Context.HasFilterData(SelectedFilter->Key))
 		{
-			// If data is not set for this filter we consider it passed the test.
-			return true;
+			// If data is not set for this filter return the value specified in the Context.
+			return Context.GetReturnValueForUnsetFilters();
 		}
 
 		switch (SelectedFilter->DataType)
diff --git a/Engine/Source/Developer/TraceInsights/Private/Insights/ViewModels/Filters.h b/Engine/Source/Developer/TraceInsights/Private/Insights/ViewModels/Filters.h
--- a/Engine/Source/Developer/TraceInsights/Private/Insights/ViewModels/Filters.h
+++ b/Engine/Source/Developer/TraceInsights/Private/Insights/ViewModels/Filters.h
@@ -269,8 +269,12 @@ public:
 		return DataMap.Contains(Key);
 	}
 
+	bool GetReturnValueForUnsetFilters() const { return bReturnValueForUnsetFilters;	}
+	void SetReturnValueForUnsetFilters(bool InValue) { bReturnValueForUnsetFilters = InValue; }
+
 private:
 	TMap<int32, ContextData> DataMap;
+	bool bReturnValueForUnsetFilters = true;
 };
 
 ////////////////////////////////////////////////////////////////////////////////////////////////////
diff --git a/Engine/Source/Developer/TraceInsights/Private/Insights/ViewModels/ThreadTimingTrack.cpp b/Engine/Source/Developer/TraceInsights/Private/Insights/ViewModels/ThreadTimingTrack.cpp
--- a/Engine/Source/Developer/TraceInsights/Private/Insights/ViewModels/ThreadTimingTrack.cpp
+++ b/Engine/Source/Developer/TraceInsights/Private/Insights/ViewModels/ThreadTimingTrack.cpp
@@ -943,6 +943,7 @@ void FThreadTimingTrack::BuildFilteredDrawState(ITimingEventsTrackDrawStateBuild
 						FilterContexts.AddDefaulted(NumTasks);
 						for (FFilterContext& Context : FilterContexts)
 						{
+							Context.SetReturnValueForUnsetFilters(false);
 							Context.AddFilterData<double>(static_cast<int32>(EFilterField::StartTime), 0.0f);
 							Context.AddFilterData<double>(static_cast<int32>(EFilterField::EndTime), 0.0f);
 							Context.AddFilterData<double>(static_cast<int32>(EFilterField::Duration), 0.0f);
@@ -1515,6 +1516,7 @@ bool FThreadTimingTrack::FindTimingProfilerEvent(const FThreadTrackEvent& InTimi
 bool FThreadTimingTrack::FindTimingProfilerEvent(const FTimingEventSearchParameters& InParameters, TFunctionRef<void(double, double, uint32, const TraceServices::FTimingProfilerEvent&)> InFoundPredicate) const
 {
 	FFilterContext FilterConfiguratorContext;
+	FilterConfiguratorContext.SetReturnValueForUnsetFilters(false);
 	FilterConfiguratorContext.AddFilterData<double>(static_cast<int32>(EFilterField::StartTime), 0.0f);
 	FilterConfiguratorContext.AddFilterData<double>(static_cast<int32>(EFilterField::EndTime), 0.0f);
 	FilterConfiguratorContext.AddFilterData<double>(static_cast<int32>(EFilterField::Duration), 0.0f);
diff --git a/Engine/Source/Developer/TraceInsights/Private/Insights/ViewModels/ThreadTimingTrack.h b/Engine/Source/Developer/TraceInsights/Private/Insights/ViewModels/ThreadTimingTrack.h
--- a/Engine/Source/Developer/TraceInsights/Private/Insights/ViewModels/ThreadTimingTrack.h
+++ b/Engine/Source/Developer/TraceInsights/Private/Insights/ViewModels/ThreadTimingTrack.h
@@ -109,7 +109,8 @@ enum class EFilterField : int32
 	Duration = 2,
 	TrackName = 3,
 	TimerId = 4,
-	TimerName = 5
+	TimerName = 5,
+	CoreEventName = 6,
 };
 
 ////////////////////////////////////////////////////////////////////////////////////////////////////
diff --git a/Engine/Source/Developer/TraceInsights/Private/Insights/Widgets/STimingView.cpp b/Engine/Source/Developer/TraceInsights/Private/Insights/Widgets/STimingView.cpp
--- a/Engine/Source/Developer/TraceInsights/Private/Insights/Widgets/STimingView.cpp
+++ b/Engine/Source/Developer/TraceInsights/Private/Insights/Widgets/STimingView.cpp
@@ -4825,8 +4825,8 @@ void STimingView::QuickFind_Execute()
 
 	if (!QuickFindVm.IsValid())
 	{
-		TSharedPtr<FFilterConfigurator> FilterConfigurator = MakeShared<FFilterConfigurator>();
-		TSharedPtr<TArray<TSharedPtr<struct FFilter>>>& AvailableFilters = FilterConfigurator->GetAvailableFilters();
+		TSharedPtr<FFilterConfigurator> NewFilterConfigurator = MakeShared<FFilterConfigurator>();
+		TSharedPtr<TArray<TSharedPtr<struct FFilter>>>& AvailableFilters = NewFilterConfigurator->GetAvailableFilters();
 
 		AvailableFilters->Add(MakeShared<FFilter>(static_cast<int32>(EFilterField::StartTime), LOCTEXT("StartTime", "Start Time"), LOCTEXT("StartTime", "Start Time"), EFilterDataType::Double, FFilterService::Get()->GetDoubleOperators()));
 		AvailableFilters->Add(MakeShared<FFilter>(static_cast<int32>(EFilterField::EndTime), LOCTEXT("EndTime", "End Time"), LOCTEXT("EndTime", "End Time"), EFilterDataType::Double, FFilterService::Get()->GetDoubleOperators()));
@@ -4855,7 +4855,12 @@ void STimingView::QuickFind_Execute()
 
 		AvailableFilters->Add(TimerNameFilter);
 
-		QuickFindVm = MakeShared<FQuickFind>(FilterConfigurator);
+		for (Insights::ITimingViewExtender* Extender : GetExtenders())
+		{
+			Extender->AddQuickFindFilters(NewFilterConfigurator);
+		}
+
+		QuickFindVm = MakeShared<FQuickFind>(NewFilterConfigurator);
 		QuickFindVm->GetOnFindFirstEvent().AddSP(this, &STimingView::FindFirstEvent);
 		QuickFindVm->GetOnFindPreviousEvent().AddSP(this, &STimingView::FindPrevEvent);
 		QuickFindVm->GetOnFindNextEvent().AddSP(this, &STimingView::FindNextEvent);
@@ -5105,7 +5110,9 @@ void STimingView::FilterAllTracks()
 {
 	for (auto& Entry : AllTracks)
 	{
+		FilterConfigurator = MakeShared<Insights::FFilterConfigurator>(*QuickFindVm->GetFilterConfigurator());
 		Entry.Value->SetFilterConfigurator(QuickFindVm->GetFilterConfigurator());
+		Entry.Value->SetDirtyFlag();
 	}
 }
 
@@ -5115,7 +5122,9 @@ void STimingView::ClearFilters()
 {
 	for (auto& Entry : AllTracks)
 	{
+		FilterConfigurator.Reset();
 		Entry.Value->SetFilterConfigurator(nullptr);
+		Entry.Value->SetDirtyFlag();
 	}
 }
 
@@ -5164,14 +5173,14 @@ void STimingView::PopulateTimerNameSuggestionList(const FString& Text, TArray<FS
 
 ////////////////////////////////////////////////////////////////////////////////////////////////////
 
-void STimingView::EnumerateFilteredTracks(TSharedPtr<Insights::FFilterConfigurator> FilterConfigurator, EnumerateFilteredTracksCallback Callback)
+void STimingView::EnumerateFilteredTracks(TSharedPtr<Insights::FFilterConfigurator> InFilterConfigurator, EnumerateFilteredTracksCallback Callback)
 {
 	Insights::FFilterContext FilterContext;
 	FilterContext.AddFilterData(static_cast<int32>(EFilterField::TrackName), FString());
 	for (auto& Entry : AllTracks)
 	{
 		FilterContext.SetFilterData(static_cast<int32>(EFilterField::TrackName), Entry.Value->GetName());
-		if (FilterConfigurator->ApplyFilters(FilterContext))
+		if (InFilterConfigurator->ApplyFilters(FilterContext))
 		{
 			Callback(Entry.Value);
 		}
diff --git a/Engine/Source/Developer/TraceInsights/Private/Insights/Widgets/STimingView.h b/Engine/Source/Developer/TraceInsights/Private/Insights/Widgets/STimingView.h
--- a/Engine/Source/Developer/TraceInsights/Private/Insights/Widgets/STimingView.h
+++ b/Engine/Source/Developer/TraceInsights/Private/Insights/Widgets/STimingView.h
@@ -380,6 +380,8 @@ public:
 
 	void CloseQuickFindTab();
 
+	TSharedPtr<Insights::FFilterConfigurator> GetFilterConfigurator() { return FilterConfigurator; }
+
 protected:
 	virtual FVector2D ComputeDesiredSize(float) const override
 	{
@@ -709,6 +711,7 @@ protected:
 	TArray<TUniquePtr<ITimingEventRelation>> CurrentRelations;
 
 	TSharedPtr<Insights::FQuickFind> QuickFindVm;
+	TSharedPtr<Insights::FFilterConfigurator> FilterConfigurator;
 	static uint32 TimingViewId;
 	const FName QuickFindTabId;
 	
diff --git a/Engine/Source/Developer/TraceInsights/Public/Insights/ITimingViewExtender.h b/Engine/Source/Developer/TraceInsights/Public/Insights/ITimingViewExtender.h
--- a/Engine/Source/Developer/TraceInsights/Public/Insights/ITimingViewExtender.h
+++ b/Engine/Source/Developer/TraceInsights/Public/Insights/ITimingViewExtender.h
@@ -40,6 +40,9 @@ public:
 	/** Extension hook for the context menu for all tracks
 	@return True if any menu option was added and False if no option was added */ 
 	virtual bool ExtendGlobalContextMenu(ITimingViewSession& InSession, FMenuBuilder& InMenuBuilder) { return false; }
+
+	/** Allows extender to add filters to the Quick Find widget. */
+	virtual void AddQuickFindFilters(TSharedPtr<class FFilterConfigurator> FilterConfigurator) {}
 };
 
 } // namespace Insights
