[Insights] Fix localization namespace

diff --git a/Engine/Source/Developer/TraceInsights/Private/Insights/TimingProfiler/QuickFindFilterConverters.cpp b/Engine/Source/Developer/TraceInsights/Private/Insights/TimingProfiler/QuickFindFilterConverters.cpp
--- a/Engine/Source/Developer/TraceInsights/Private/Insights/TimingProfiler/QuickFindFilterConverters.cpp
+++ b/Engine/Source/Developer/TraceInsights/Private/Insights/TimingProfiler/QuickFindFilterConverters.cpp
@@ -4,7 +4,7 @@
 
 #include "Insights/InsightsManager.h"
 
-#define LOCTEXT_NAMESPACE "Insights::FMemoryFilterValueConverter"
+#define LOCTEXT_NAMESPACE "Insights::QuickFindFilterConverters"
 
 namespace Insights
 {
