[Insights]
  - Memory: Added a new query rule: **Growth vs. Decline**. Identifies \"growth\" allocations, allocated between time A and time B and freed after time B (A ? a ? B ? f) and \"decline\" allocations, allocated before time A and freed between time A and time B (a ? A ? f ? B). The \"decline\" allocations are changed to have negative size, so the size aggregation shows variation between A and B.
  - Memory: Fixed the iteration over SBIF offsetted cells resulting in some cells to be skipped when checked by queries.

diff --git a/Engine/Source/Developer/TraceInsights/Private/Insights/MemoryProfiler/ViewModels/MemAllocGroupingBySize.cpp b/Engine/Source/Developer/TraceInsights/Private/Insights/MemoryProfiler/ViewModels/MemAllocGroupingBySize.cpp
--- a/Engine/Source/Developer/TraceInsights/Private/Insights/MemoryProfiler/ViewModels/MemAllocGroupingBySize.cpp
+++ b/Engine/Source/Developer/TraceInsights/Private/Insights/MemoryProfiler/ViewModels/MemAllocGroupingBySize.cpp
@@ -101,7 +101,7 @@ void FMemAllocGroupingBySize::GroupNodes(const TArray<FTableTreeNodePtr>& Nodes,
 		const FMemoryAlloc* Alloc = MemAllocNode.GetMemAlloc();
 		if (Alloc)
 		{
-			const uint64 Size = Alloc->GetSize();
+			const uint64 Size = FMath::Abs(Alloc->GetSize());
 
 			uint32 ThresholdIndex;
 			if (bIsPow2)
diff --git a/Engine/Source/Developer/TraceInsights/Private/Insights/MemoryProfiler/ViewModels/MemoryAlloc.h b/Engine/Source/Developer/TraceInsights/Private/Insights/MemoryProfiler/ViewModels/MemoryAlloc.h
--- a/Engine/Source/Developer/TraceInsights/Private/Insights/MemoryProfiler/ViewModels/MemoryAlloc.h
+++ b/Engine/Source/Developer/TraceInsights/Private/Insights/MemoryProfiler/ViewModels/MemoryAlloc.h
@@ -28,7 +28,7 @@ public:
 	double GetDuration() const { return EndTime - StartTime; }
 	uint64 GetAddress() const { return Address; }
 	uint64 GetPage() const { return Address & ~(4llu*1024-1); }
-	uint64 GetSize() const { return Size; }
+	int64 GetSize() const { return Size; }
 	const TCHAR* GetTag() const { return Tag; }
 	const TraceServices::FCallstack* GetCallstack() const { return Callstack; }
 	FText GetFullCallstack() const;
@@ -40,11 +40,12 @@ private:
 	double StartTime;
 	double EndTime;
 	uint64 Address;
-	uint64 Size;
+	int64 Size;
 	const TCHAR* Tag;
 	const TraceServices::FCallstack* Callstack;
 	HeapId RootHeap;
 	bool bIsBlock;
+	bool bIsDecline;
 };
 
 ////////////////////////////////////////////////////////////////////////////////////////////////////
diff --git a/Engine/Source/Developer/TraceInsights/Private/Insights/MemoryProfiler/ViewModels/MemorySharedState.cpp b/Engine/Source/Developer/TraceInsights/Private/Insights/MemoryProfiler/ViewModels/MemorySharedState.cpp
--- a/Engine/Source/Developer/TraceInsights/Private/Insights/MemoryProfiler/ViewModels/MemorySharedState.cpp
+++ b/Engine/Source/Developer/TraceInsights/Private/Insights/MemoryProfiler/ViewModels/MemorySharedState.cpp
@@ -933,6 +933,12 @@ void FMemorySharedState::InitMemoryRules()
 		LOCTEXT("MemRule_AaBf_Verbose", "Growth"),
 		LOCTEXT("MemRule_AaBf_Desc", "Identifies allocations allocated between time A and time B and freed after time B.\n(A ≤ a ≤ B ≤ f)")));
 
+	MemoryRules.Add(MakeShared<Insights::FMemoryRuleSpec>(
+		ERule::aAfaBf, 2,
+		LOCTEXT("MemRule_aAfaBf_Short", "A*B*/*A*B"),
+		LOCTEXT("MemRule_aAfaBf_Verbose", "Growth vs. Decline"),
+		LOCTEXT("MemRule_aAfaBf_Desc", "Identifies \"growth\" allocations, allocated between time A and time B and freed after time B (A ≤ a ≤ B ≤ f)\nand \"decline\" allocations, allocated before time A and freed between time A and time B (a ≤ A ≤ f ≤ B).\nThe \"decline\" allocations are changed to have negative size, so the size aggregation shows variation between A and B.")));
+
 	MemoryRules.Add(MakeShared<Insights::FMemoryRuleSpec>(
 		ERule::AfB, 2,
 		LOCTEXT("MemRule_AfB_Short", "*A**B"),
diff --git a/Engine/Source/Developer/TraceInsights/Private/Insights/MemoryProfiler/Widgets/SMemAllocTableTreeView.cpp b/Engine/Source/Developer/TraceInsights/Private/Insights/MemoryProfiler/Widgets/SMemAllocTableTreeView.cpp
--- a/Engine/Source/Developer/TraceInsights/Private/Insights/MemoryProfiler/Widgets/SMemAllocTableTreeView.cpp
+++ b/Engine/Source/Developer/TraceInsights/Private/Insights/MemoryProfiler/Widgets/SMemAllocTableTreeView.cpp
@@ -137,6 +137,8 @@ void SMemAllocTableTreeView::RebuildTree(bool bResync)
 			const int32 TotalAllocCount = Allocs.Num();
 			if (TotalAllocCount != TableTreeNodes.Num())
 			{
+				UE_LOG(MemoryProfiler, Log, TEXT("[MemAlloc] Creating nodes (%d nodes --> %d allocs)..."), TableTreeNodes.Num(), TotalAllocCount);
+
 				if (TableTreeNodes.Num() > TotalAllocCount)
 				{
 					TableTreeNodes.Empty();
@@ -149,8 +151,6 @@ void SMemAllocTableTreeView::RebuildTree(bool bResync)
 				for (int32 AllocIndex = TableTreeNodes.Num(); AllocIndex < TotalAllocCount; ++AllocIndex)
 				{
 					const FMemoryAlloc* Alloc = MemAllocTable->GetMemAlloc(AllocIndex);
-					FName NodeName(Alloc->bIsBlock ? BaseHeapName : BaseNodeName, Alloc->GetStartEventIndex() + 1);
-					FMemAllocNodePtr NodePtr = MakeShared<FMemAllocNode>(NodeName, MemAllocTable, AllocIndex);
 
 					// Until we have an UX story around heap allocations
 					// remove them from the list
@@ -160,6 +160,8 @@ void SMemAllocTableTreeView::RebuildTree(bool bResync)
 						continue;
 					}
 
+					FName NodeName(Alloc->bIsBlock ? BaseHeapName : BaseNodeName, Alloc->GetStartEventIndex() + 1);
+					FMemAllocNodePtr NodePtr = MakeShared<FMemAllocNode>(NodeName, MemAllocTable, AllocIndex);
 					TableTreeNodes.Add(NodePtr);
 				}
 				ensure(TableTreeNodes.Num() == TotalAllocCount - HeapAllocCount);
@@ -200,7 +202,7 @@ void SMemAllocTableTreeView::RebuildTree(bool bResync)
 	if (TotalTime > 0.01)
 	{
 		const double SyncTime = SyncStopwatch.GetAccumulatedTime();
-		UE_LOG(MemoryProfiler, Log, TEXT("[MemAlloc] Tree view rebuilt in %.3fs (%.3fs + %.3fs) --> %d allocs (%d added)"),
+		UE_LOG(MemoryProfiler, Log, TEXT("[MemAlloc] Tree view rebuilt in %.3fs (%.3fs + %.3fs) --> %d nodes (%d added)"),
 			TotalTime, SyncTime, TotalTime - SyncTime, TableTreeNodes.Num(), TableTreeNodes.Num() - PreviousNodeCount);
 	}
 }
@@ -276,6 +278,8 @@ void SMemAllocTableTreeView::StartQuery()
 		QueryStopwatch.Reset();
 		QueryStopwatch.Start();
 	}
+
+	//TODO: update window title
 }
 
 ////////////////////////////////////////////////////////////////////////////////////////////////////
@@ -357,12 +361,21 @@ void SMemAllocTableTreeView::UpdateQuery(TraceServices::IAllocationsProvider::EQ
 					Alloc.StartTime = Allocation->GetStartTime();
 					Alloc.EndTime = Allocation->GetEndTime();
 					Alloc.Address = Allocation->GetAddress();
-					Alloc.Size = Allocation->GetSize();
+					Alloc.Size = int64(Allocation->GetSize());
 					Alloc.Tag = Provider.GetTagName(Allocation->GetTag());
 					Alloc.Callstack = Allocation->GetCallstack();
 					Alloc.RootHeap = Allocation->GetRootHeap();
 					Alloc.bIsBlock = Allocation->IsHeap();
 					check(Alloc.Callstack != nullptr);
+
+					if (Rule->GetValue() == TraceServices::IAllocationsProvider::EQueryRule::aAfaBf)
+					{
+						if (Alloc.StartTime <= TimeMarkers[0] && Alloc.EndTime <= TimeMarkers[1]) // decline
+						{
+							Alloc.Size = -Alloc.Size;
+							Alloc.bIsDecline = true;
+						}
+					}
 				}
 
 				PageStopwatch.Stop();
diff --git a/Engine/Source/Developer/TraceServices/Private/Model/AllocationsQuery.cpp b/Engine/Source/Developer/TraceServices/Private/Model/AllocationsQuery.cpp
--- a/Engine/Source/Developer/TraceServices/Private/Model/AllocationsQuery.cpp
+++ b/Engine/Source/Developer/TraceServices/Private/Model/AllocationsQuery.cpp
@@ -179,19 +179,14 @@ void FAllocationsQuery::Run()
 		}
 	}
 
-	UE_LOG(LogTraceServices, Log, TEXT("[MemAlloc] Done"));
-
 	AllocationsProvider.EndRead();
 
 	IsWorking = false;
 
 	uint64 EndTime = FPlatformTime::Cycles64();
 	const double TotalTime = static_cast<double>(EndTime - StartTime) * FPlatformTime::GetSecondsPerCycle64();
-	if (TotalTime > 0.1)
-	{
-		UE_LOG(LogTraceServices, Log, TEXT("[MemAlloc] Allocations query completed in %.3fs (%u cells, %u allocations)"),
-			TotalTime, CellCount, TotalAllocationCount);
-	}
+	UE_LOG(LogTraceServices, Log, TEXT("[MemAlloc] Allocations query completed in %.3fs (%u cells, %u allocations)"),
+		TotalTime, CellCount, TotalAllocationCount);
 }
 
 ////////////////////////////////////////////////////////////////////////////////////////////////////
@@ -271,6 +266,21 @@ void FAllocationsQuery::QueryLiveAllocs(TArray<const FAllocationItem*>& OutAlloc
 	}
 	break;
 
+	case IAllocationsProvider::EQueryRule::aAfaBf: // decline + growth
+	{
+		const double TimeA = Params.TimeA;
+		const double TimeB = Params.TimeB;
+		AllocationsProvider.EnumerateLiveAllocs([TimeA, TimeB, &OutAllocs](const FAllocationItem& Alloc)
+		{
+			if ((Alloc.StartTime <= TimeA && Alloc.EndTime >= TimeA && Alloc.EndTime <= TimeB) || // decline
+				(Alloc.StartTime >= TimeA && Alloc.StartTime <= TimeB && Alloc.EndTime >= TimeB)) // growth
+			{
+				OutAllocs.Add(&Alloc);
+			}
+		});
+	}
+	break;
+
 	case IAllocationsProvider::EQueryRule::AfB: // free events
 	{
 		const double TimeA = Params.TimeA;
diff --git a/Engine/Source/Developer/TraceServices/Private/Model/SbTree.cpp b/Engine/Source/Developer/TraceServices/Private/Model/SbTree.cpp
--- a/Engine/Source/Developer/TraceServices/Private/Model/SbTree.cpp
+++ b/Engine/Source/Developer/TraceServices/Private/Model/SbTree.cpp
@@ -168,6 +168,21 @@ void FSbTreeCell::Query(TArray<const FAllocationItem*>& OutAllocs, const IAlloca
 	}
 	break;
 
+	case IAllocationsProvider::EQueryRule::aAfaBf: // decline + growth
+	{
+		const double TimeA = Params.TimeA;
+		const double TimeB = Params.TimeB;
+		for (const FAllocationItem* Alloc : Allocs)
+		{
+			if ((Alloc->StartTime <= TimeA && Alloc->EndTime >= TimeA && Alloc->EndTime <= TimeB) || // decline
+				(Alloc->StartTime >= TimeA && Alloc->StartTime <= TimeB && Alloc->EndTime >= TimeB)) // growth
+			{
+				OutAllocs.Add(Alloc);
+			}
+		}
+	}
+	break;
+
 	case IAllocationsProvider::EQueryRule::AfB: // free events
 	{
 		const double TimeA = Params.TimeA;
@@ -520,6 +535,14 @@ void FSbTree::Query(TArray<const FSbTreeCell*>& OutCells, const IAllocationsProv
 		}
 		break;
 
+		case IAllocationsProvider::EQueryRule::aAfaBf: // decline + growth
+		{
+			int32 Column1;
+			GetColumnsAtTime(Params.TimeA, &Column1, nullptr);
+			IterateCells(OutCells, Column1, CurrentColumn);
+		}
+		break;
+
 		case IAllocationsProvider::EQueryRule::AfB: // free events
 		{
 			int32 Column1;
@@ -688,13 +711,11 @@ void FSbTree::IterateCells(TArray<const FSbTreeCell*> & OutCells, int32 StartCol
 	const uint32 NumOffsettedCells = static_cast<uint32>(OffsettedCells.Num());
 	for (uint32 Depth = 1; Depth <= MaxDepth; ++Depth) // offsetted cells doesn't exist on Depth 0
 	{
-		uint32 HalfCellWidth = (1 << Depth) >> 1;
-		if (HalfCellWidth > LocalStartColumn)
-		{
-			break;
-		}
-		uint32 FirstCellIndex = FSbTreeUtils::GetCellAtDepth(LocalStartColumn - HalfCellWidth, Depth);
-		uint32 LastCellIndex = FMath::Min(FSbTreeUtils::GetCellAtDepth(LocalEndColumn - HalfCellWidth, Depth) + 1, NumOffsettedCells);
+		const uint32 HalfCellWidth = (1 << Depth) >> 1;
+		const uint32 OffsettedStartColumn = (LocalStartColumn > HalfCellWidth) ? LocalStartColumn - HalfCellWidth : 0;
+		const uint32 FirstCellIndex = FSbTreeUtils::GetCellAtDepth(OffsettedStartColumn, Depth);
+		const uint32 OffsettedEndColumn = (LocalEndColumn > HalfCellWidth) ? LocalEndColumn - HalfCellWidth : 0;
+		const uint32 LastCellIndex = FMath::Min(FSbTreeUtils::GetCellAtDepth(OffsettedEndColumn, Depth) + 1, NumOffsettedCells);
 		const uint32 CellIncrement = 1 << (Depth + 1);
 		for (uint32 CellIndex = FirstCellIndex; CellIndex < LastCellIndex; CellIndex += CellIncrement)
 		{
diff --git a/Engine/Source/Developer/TraceServices/Public/TraceServices/Model/AllocationsProvider.h b/Engine/Source/Developer/TraceServices/Public/TraceServices/Model/AllocationsProvider.h
--- a/Engine/Source/Developer/TraceServices/Public/TraceServices/Model/AllocationsProvider.h
+++ b/Engine/Source/Developer/TraceServices/Public/TraceServices/Model/AllocationsProvider.h
@@ -67,6 +67,7 @@ public:
 		Aaf,     // after
 		aAfB,    // decline
 		AaBf,    // growth
+		aAfaBf,  // decline + growth (used as "growth vs. decline")
 		AfB,     // free events
 		AaB,     // alloc events
 		AafB,    // short living allocs
