[Insights] Implement Find functionality for CoreTimingTracks

diff --git a/Engine/Source/Developer/TraceInsights/Private/Insights/ContextSwitches/ViewModels/CpuCoreTimingTrack.cpp b/Engine/Source/Developer/TraceInsights/Private/Insights/ContextSwitches/ViewModels/CpuCoreTimingTrack.cpp
--- a/Engine/Source/Developer/TraceInsights/Private/Insights/ContextSwitches/ViewModels/CpuCoreTimingTrack.cpp
+++ b/Engine/Source/Developer/TraceInsights/Private/Insights/ContextSwitches/ViewModels/CpuCoreTimingTrack.cpp
@@ -269,6 +269,91 @@ const TSharedPtr<const ITimingEvent> FCpuCoreTimingTrack::GetEvent(double InTime
 
 ////////////////////////////////////////////////////////////////////////////////////////////////////
 
+const TSharedPtr<const ITimingEvent> FCpuCoreTimingTrack::SearchEvent(const FTimingEventSearchParameters& InSearchParameters) const
+{
+	TSharedPtr<const ITimingEvent> FoundEvent;
+
+	FFilterContext FilterConfiguratorContext;
+	FilterConfiguratorContext.SetReturnValueForUnsetFilters(false);
+	FilterConfiguratorContext.AddFilterData<double>(static_cast<int32>(EFilterField::StartTime), 0.0f);
+	FilterConfiguratorContext.AddFilterData<double>(static_cast<int32>(EFilterField::EndTime), 0.0f);
+	FilterConfiguratorContext.AddFilterData<double>(static_cast<int32>(EFilterField::Duration), 0.0f);
+	FilterConfiguratorContext.AddFilterData<FString>(static_cast<int32>(EFilterField::TrackName), this->GetName());
+	FilterConfiguratorContext.AddFilterData<int64>(static_cast<int32>(EFilterField::CoreEventName), 0);
+
+	TTimingEventSearch<TraceServices::FCpuCoreEvent>::Search(
+		InSearchParameters,
+
+		[this, &InSearchParameters](TTimingEventSearch<TraceServices::FCpuCoreEvent>::FContext& InContext)
+		{
+			TSharedPtr<const TraceServices::IAnalysisSession> Session = FInsightsManager::Get()->GetSession();
+			if (Session.IsValid())
+			{
+				TraceServices::FAnalysisSessionReadScope SessionReadScope(*Session.Get());
+				if (Session.IsValid() && TraceServices::ReadContextSwitchesProvider(*Session.Get()))
+				{
+					const TraceServices::IContextSwitchesProvider* ContextSwitchesProvider = TraceServices::ReadContextSwitchesProvider(*Session.Get());
+
+					auto Callback = [&InContext](double EventStartTime, double EventEndTime, uint32 EventDepth, const TraceServices::FCpuCoreEvent& Event)
+					{
+						InContext.Check(EventStartTime, EventEndTime, EventDepth, Event);
+						return InContext.ShouldContinueSearching() ? TraceServices::EContextSwitchEnumerationResult::Continue : TraceServices::EContextSwitchEnumerationResult::Stop;
+					};
+
+					if (InContext.GetParameters().SearchDirection == FTimingEventSearchParameters::ESearchDirection::Forward)
+					{
+						ContextSwitchesProvider->EnumerateCpuCoreEvents(CoreNumber, InSearchParameters.StartTime, InSearchParameters.EndTime,
+							[ContextSwitchesProvider, Callback](const TraceServices::FCpuCoreEvent& CpuCoreEvent)
+							{
+								return Callback(CpuCoreEvent.Start, CpuCoreEvent.End, 0, CpuCoreEvent);
+							});
+					}
+					else
+					{
+						ContextSwitchesProvider->EnumerateCpuCoreEventsBackwards(CoreNumber, InSearchParameters.EndTime, InSearchParameters.StartTime,
+							[ContextSwitchesProvider, Callback](const TraceServices::FCpuCoreEvent& CpuCoreEvent)
+							{
+								return Callback(CpuCoreEvent.Start, CpuCoreEvent.End, 0, CpuCoreEvent);
+							});
+					}
+				}
+			}
+		},
+
+		[&FilterConfiguratorContext, &InSearchParameters](double EventStartTime, double EventEndTime, uint32 EventDepth, const TraceServices::FCpuCoreEvent& Event)
+		{
+			if (!InSearchParameters.FilterExecutor.IsValid())
+			{
+				return true;
+			}
+
+			TSharedPtr<const TraceServices::IAnalysisSession> Session = FInsightsManager::Get()->GetSession();
+			if (Session.IsValid())
+			{
+				FilterConfiguratorContext.SetFilterData<double>(static_cast<int32>(EFilterField::StartTime), EventStartTime);
+				FilterConfiguratorContext.SetFilterData<double>(static_cast<int32>(EFilterField::EndTime), EventEndTime);
+				FilterConfiguratorContext.SetFilterData<double>(static_cast<int32>(EFilterField::Duration), EventEndTime - EventStartTime);
+				FilterConfiguratorContext.SetFilterData<int64>(static_cast<int32>(EFilterField::CoreEventName), Event.SystemThreadId);
+
+				return InSearchParameters.FilterExecutor->ApplyFilters(FilterConfiguratorContext);
+			}
+
+			return false;
+		},
+
+		[&FoundEvent, this](double InFoundStartTime, double InFoundEndTime, uint32 InFoundDepth, const TraceServices::FCpuCoreEvent& InEvent)
+		{
+			FoundEvent = MakeShared<FTimingEvent>(SharedThis(this), InFoundStartTime, InFoundEndTime, InFoundDepth);
+		},
+
+		TTimingEventSearch<TraceServices::FCpuCoreEvent>::NoMatch
+		);
+
+	return FoundEvent;
+}
+
+////////////////////////////////////////////////////////////////////////////////////////////////////
+
 void FCpuCoreTimingTrack::InitTooltip(FTooltipDrawState& InOutTooltip, const ITimingEvent& InTooltipEvent) const
 {
 	InOutTooltip.ResetContent();
diff --git a/Engine/Source/Developer/TraceInsights/Private/Insights/ContextSwitches/ViewModels/CpuCoreTimingTrack.h b/Engine/Source/Developer/TraceInsights/Private/Insights/ContextSwitches/ViewModels/CpuCoreTimingTrack.h
--- a/Engine/Source/Developer/TraceInsights/Private/Insights/ContextSwitches/ViewModels/CpuCoreTimingTrack.h
+++ b/Engine/Source/Developer/TraceInsights/Private/Insights/ContextSwitches/ViewModels/CpuCoreTimingTrack.h
@@ -48,6 +48,8 @@ public:
 
 	virtual void BuildContextMenu(FMenuBuilder& MenuBuilder) override;
 
+	virtual const TSharedPtr<const ITimingEvent> SearchEvent(const FTimingEventSearchParameters& InSearchParameters) const override;
+
 protected:
 	virtual const TSharedPtr<const ITimingEvent> GetEvent(double InTime, double SecondsPerPixel, int32 Depth) const override;
 
diff --git a/Engine/Source/Developer/TraceServices/Private/Model/ContextSwitches.cpp b/Engine/Source/Developer/TraceServices/Private/Model/ContextSwitches.cpp
--- a/Engine/Source/Developer/TraceServices/Private/Model/ContextSwitches.cpp
+++ b/Engine/Source/Developer/TraceServices/Private/Model/ContextSwitches.cpp
@@ -201,7 +201,7 @@ void FContextSwitchesProvider::EnumerateCpuCoreEvents(uint32 CoreNumber, double
 	}
 
 	auto Iterator = CpuCoreEvents->GetIteratorFromPage(PageIndex);
-	const FCpuCoreEvent* CurrentCpuCoreEvent = Iterator.NextItem();
+	const FCpuCoreEvent* CurrentCpuCoreEvent = Iterator.GetCurrentItem();
 	while (CurrentCpuCoreEvent && CurrentCpuCoreEvent->Start < EndTime)
 	{
 		if (CurrentCpuCoreEvent->End > StartTime)
@@ -216,6 +216,47 @@ void FContextSwitchesProvider::EnumerateCpuCoreEvents(uint32 CoreNumber, double
 	}
 }
 
+void FContextSwitchesProvider::EnumerateCpuCoreEventsBackwards(uint32 CoreNumber, double EndTime, double StartTime, CpuCoreEventCallback Callback) const
+{
+	Session.ReadAccessCheck();
+
+	if (CoreNumber >= (uint32)CpuCores.Num())
+	{
+		return;
+	}
+
+	const TPagedArray<FCpuCoreEvent>* CpuCoreEvents = CpuCores[CoreNumber];
+	if (CpuCoreEvents == nullptr)
+	{
+		return;
+	}
+
+	uint64 PageIndex = Algo::UpperBoundBy(*CpuCoreEvents, EndTime, [](const TPagedArrayPage<FCpuCoreEvent>& Page)
+		{
+			return Page.Items[0].Start;
+		});
+
+	if (PageIndex == 0)
+	{
+		return;
+	}
+
+	auto Iterator = PageIndex < CpuCoreEvents->NumPages() ? CpuCoreEvents->GetIteratorFromPage(PageIndex) : CpuCoreEvents->end();
+	const FCpuCoreEvent* CurrentCpuCoreEvent = Iterator.GetCurrentItem();
+	while (CurrentCpuCoreEvent && CurrentCpuCoreEvent->End > StartTime)
+	{
+		if (CurrentCpuCoreEvent->Start < EndTime)
+		{
+			if (Callback(*CurrentCpuCoreEvent) == EContextSwitchEnumerationResult::Stop)
+			{
+				break;
+			}
+		}
+
+		CurrentCpuCoreEvent = Iterator.PrevItem();
+	}
+}
+
 const TPagedArray<FContextSwitch>* FContextSwitchesProvider::GetContextSwitches(uint32 ThreadId) const
 {
 	Session.ReadAccessCheck();
diff --git a/Engine/Source/Developer/TraceServices/Private/Model/ContextSwitchesPrivate.h b/Engine/Source/Developer/TraceServices/Private/Model/ContextSwitchesPrivate.h
--- a/Engine/Source/Developer/TraceServices/Private/Model/ContextSwitchesPrivate.h
+++ b/Engine/Source/Developer/TraceServices/Private/Model/ContextSwitchesPrivate.h
@@ -31,6 +31,7 @@ public:
 	void EnumerateCpuCores(CpuCoreCallback Callback) const override;
 	void EnumerateContextSwitches(uint32 ThreadId, double StartTime, double EndTime, ContextSwitchCallback Callback) const override;
 	void EnumerateCpuCoreEvents(uint32 CoreNumber, double StartTime, double EndTime, CpuCoreEventCallback Callback) const override;
+	void EnumerateCpuCoreEventsBackwards(uint32 CoreNumber, double EndTime, double StartTime, CpuCoreEventCallback Callback) const override;
 
 	void Add(uint32 SystemThreadId, double Start, double End, uint32 CoreNumber);
 	void AddThreadInfo(uint32 ThreadId, uint32 SystemThreadId);
diff --git a/Engine/Source/Developer/TraceServices/Public/TraceServices/Model/ContextSwitches.h b/Engine/Source/Developer/TraceServices/Public/TraceServices/Model/ContextSwitches.h
--- a/Engine/Source/Developer/TraceServices/Public/TraceServices/Model/ContextSwitches.h
+++ b/Engine/Source/Developer/TraceServices/Public/TraceServices/Model/ContextSwitches.h
@@ -52,6 +52,7 @@ public:
 	virtual uint64 GetCpuCoresSerial() const = 0;
 	virtual void EnumerateCpuCores(CpuCoreCallback Callback) const = 0;
 	virtual void EnumerateCpuCoreEvents(uint32 CoreNumber, double StartTime, double EndTime, CpuCoreEventCallback Callback) const = 0;
+	virtual void EnumerateCpuCoreEventsBackwards(uint32 CoreNumber, double EndTime, double StartTime, CpuCoreEventCallback Callback) const = 0;
 	virtual void EnumerateContextSwitches(uint32 ThreadId, double StartTime, double EndTime, ContextSwitchCallback Callback) const = 0;
 };
 
