[Insights] Modules view tab

Adds a tab to display detected modules in a trace session as well as showing statistics about symbol loading. Also adds ability for the user to load symbols from a specified directory.

diff --git a/Engine/Source/Developer/TraceInsights/Private/Insights/MemoryProfiler/MemoryProfilerManager.h b/Engine/Source/Developer/TraceInsights/Private/Insights/MemoryProfiler/MemoryProfilerManager.h
--- a/Engine/Source/Developer/TraceInsights/Private/Insights/MemoryProfiler/MemoryProfilerManager.h
+++ b/Engine/Source/Developer/TraceInsights/Private/Insights/MemoryProfiler/MemoryProfilerManager.h
@@ -100,6 +100,10 @@ public:
 	void SetMemTagTreeViewVisible(const bool bIsVisible) { bIsMemTagTreeViewVisible = bIsVisible; }
 	void ShowHideMemTagTreeView(const bool bIsVisible);
 
+	const bool IsMemModuleViewVisible() const { return bIsMemModuleViewVisible; }
+	void SetMemModuleViewVisible(const bool bIsVisible) { bIsMemModuleViewVisible = bIsVisible; }
+	void ShowHideMemModuleView(const bool bIsVisible);
+	
 	////////////////////////////////////////////////////////////////////////////////////////////////////
 
 	const FName& GetLogListingName() const { return LogListingName; }
@@ -150,6 +154,9 @@ private:
 	/** If the LLM Tags tree view is visible or hidden. */
 	bool bIsMemTagTreeViewVisible;
 
+	/** If the modules view is visible or hidden. */
+	bool bIsMemModuleViewVisible;
+
 	/** The name of the Memory Insights log listing. */
 	FName LogListingName;
 
diff --git a/Engine/Source/Developer/TraceInsights/Private/Insights/MemoryProfiler/Widgets/SMemoryProfilerWindow.cpp b/Engine/Source/Developer/TraceInsights/Private/Insights/MemoryProfiler/Widgets/SMemoryProfilerWindow.cpp
--- a/Engine/Source/Developer/TraceInsights/Private/Insights/MemoryProfiler/Widgets/SMemoryProfilerWindow.cpp
+++ b/Engine/Source/Developer/TraceInsights/Private/Insights/MemoryProfiler/Widgets/SMemoryProfilerWindow.cpp
@@ -49,6 +49,7 @@ const FName FMemoryProfilerTabs::TimingViewID(TEXT("TimingView"));
 const FName FMemoryProfilerTabs::MemInvestigationViewID(TEXT("MemInvestigation"));
 const FName FMemoryProfilerTabs::MemTagTreeViewID(TEXT("LowLevelMemTags"));
 const FName FMemoryProfilerTabs::MemAllocTableTreeViewID(TEXT("MemAllocTableTreeView"));
+const FName FMemoryProfilerTabs::MemModulesViewID(TEXT("MemModules"));
 
 ////////////////////////////////////////////////////////////////////////////////////////////////////
 
@@ -465,6 +466,34 @@ END_SLATE_FUNCTION_BUILD_OPTIMIZATION
 
 ////////////////////////////////////////////////////////////////////////////////////////////////////
 
+BEGIN_SLATE_FUNCTION_BUILD_OPTIMIZATION
+TSharedRef<SDockTab> SMemoryProfilerWindow::SpawnTab_MemModulesView(const FSpawnTabArgs& Args)
+{
+	FMemoryProfilerManager::Get()->SetMemModuleViewVisible(true);
+
+	const TSharedRef<SDockTab> DockTab = SNew(SDockTab)
+		.ShouldAutosize(false)
+		.TabRole(ETabRole::PanelTab)
+		[
+			SAssignNew(MemModulesView, Insights::SModulesView)
+		];
+
+	DockTab->SetOnTabClosed(SDockTab::FOnTabClosedCallback::CreateRaw(this, &SMemoryProfilerWindow::OnMemModulesViewClosed));
+
+	return DockTab;
+}
+END_SLATE_FUNCTION_BUILD_OPTIMIZATION
+
+////////////////////////////////////////////////////////////////////////////////////////////////////
+
+void SMemoryProfilerWindow::OnMemModulesViewClosed(TSharedRef<SDockTab> TabBeingClosed)
+{
+	FMemoryProfilerManager::Get()->SetMemModuleViewVisible(false);
+	MemModulesView = nullptr;
+}
+
+////////////////////////////////////////////////////////////////////////////////////////////////////
+
 BEGIN_SLATE_FUNCTION_BUILD_OPTIMIZATION
 void SMemoryProfilerWindow::Construct(const FArguments& InArgs, const TSharedRef<SDockTab>& ConstructUnderMajorTab, const TSharedPtr<SWindow>& ConstructUnderWindow)
 {
@@ -493,6 +522,11 @@ void SMemoryProfilerWindow::Construct(const FArguments& InArgs, const TSharedRef
 		.SetIcon(FSlateIcon(FInsightsStyle::GetStyleSetName(), "Icons.MemTagTreeView"))
 		.SetGroup(AppMenuGroup.ToSharedRef());
 
+	TabManager->RegisterTabSpawner(FMemoryProfilerTabs::MemModulesViewID, FOnSpawnTab::CreateRaw(this, &SMemoryProfilerWindow::SpawnTab_MemModulesView))
+		.SetDisplayName(LOCTEXT("MemModulesViewTabTitle", "Modules"))
+		.SetIcon(FSlateIcon(FInsightsStyle::GetStyleSetName(), "Icons.MemTagTreeView")) //todo: Update icon
+		.SetGroup(AppMenuGroup.ToSharedRef());
+
 	TSharedPtr<FMemoryProfilerManager> MemoryProfilerManager = FMemoryProfilerManager::Get();
 	ensure(MemoryProfilerManager.IsValid());
 
@@ -515,6 +549,7 @@ void SMemoryProfilerWindow::Construct(const FArguments& InArgs, const TSharedRef
 				->SetSizeCoefficient(0.23f)
 				->AddTab(FMemoryProfilerTabs::MemInvestigationViewID, ETabState::OpenedTab)
 				->AddTab(FMemoryProfilerTabs::MemTagTreeViewID, ETabState::OpenedTab)
+				->AddTab(FMemoryProfilerTabs::MemModulesViewID, ETabState::OpenedTab)
 				->SetForegroundTab(FMemoryProfilerTabs::MemInvestigationViewID)
 			)
 		);
diff --git a/Engine/Source/Developer/TraceInsights/Private/Insights/MemoryProfiler/Widgets/SMemoryProfilerWindow.h b/Engine/Source/Developer/TraceInsights/Private/Insights/MemoryProfiler/Widgets/SMemoryProfilerWindow.h
--- a/Engine/Source/Developer/TraceInsights/Private/Insights/MemoryProfiler/Widgets/SMemoryProfilerWindow.h
+++ b/Engine/Source/Developer/TraceInsights/Private/Insights/MemoryProfiler/Widgets/SMemoryProfilerWindow.h
@@ -14,6 +14,7 @@
 
 // Insights
 #include "Insights/InsightsManager.h"
+#include "Insights/Widgets/SModulesView.h"
 
 ////////////////////////////////////////////////////////////////////////////////////////////////////
 
@@ -42,6 +43,7 @@ struct FMemoryProfilerTabs
 	static const FName MemInvestigationViewID;
 	static const FName MemTagTreeViewID;
 	static const FName MemAllocTableTreeViewID; // base id
+	static const FName MemModulesViewID;
 };
 
 ////////////////////////////////////////////////////////////////////////////////////////////////////
@@ -104,6 +106,9 @@ private:
 	TSharedRef<SDockTab> SpawnTab_MemAllocTableTreeView(const FSpawnTabArgs& Args, int32 TabIndex);
 	void OnMemAllocTableTreeViewTabClosed(TSharedRef<SDockTab> TabBeingClosed);
 
+	TSharedRef<SDockTab> SpawnTab_MemModulesView(const FSpawnTabArgs& Args);
+	void OnMemModulesViewClosed(TSharedRef<SDockTab> TabBeingClosed);
+	
 	/**
 	 * Fill the main menu with menu items.
 	 *
@@ -192,6 +197,9 @@ private:
 	/** The list of Allocations table tree view widgets */
 	TArray<TSharedPtr<Insights::SMemAllocTableTreeView>> MemAllocTableTreeViews;
 
+	/** The modules view widget. */
+	TSharedPtr<Insights::SModulesView> MemModulesView;
+
 	const int32 MaxMemAllocTableTreeViews = 4;
 	int32 LastMemAllocTableTreeViewIndex = -1;
 
diff --git a/Engine/Source/Developer/TraceInsights/Private/Insights/Widgets/SModulesView.cpp b/Engine/Source/Developer/TraceInsights/Private/Insights/Widgets/SModulesView.cpp
new file mode 100644
--- /dev/null
+++ b/Engine/Source/Developer/TraceInsights/Private/Insights/Widgets/SModulesView.cpp
@@ -0,0 +1,596 @@
+// Copyright Epic Games, Inc. All Rights Reserved.
+
+#include "SModulesView.h"
+
+//#include "Framework/Application/SlateApplication.h"
+#include "Framework/MultiBox/MultiBoxBuilder.h"
+#include "SlateOptMacros.h"
+#include "Textures/SlateIcon.h"
+#include "Widgets/Layout/SScrollBox.h"
+#include "Widgets/Text/STextBlock.h"
+
+// Insights
+//#include "Insights/InsightsManager.h"
+#include "DesktopPlatformModule.h"
+#include "Insights/InsightsManager.h"
+#include "Insights/InsightsStyle.h"
+#include "TraceServices/Private/Model/ModuleProvider.h"
+
+////////////////////////////////////////////////////////////////////////////////////////////////////
+
+#define LOCTEXT_NAMESPACE "SModulesView"
+
+namespace Insights
+{
+
+////////////////////////////////////////////////////////////////////////////////////////////////////
+// FModule
+// 
+// ViewModel for modules
+////////////////////////////////////////////////////////////////////////////////////////////////////
+
+class FModule
+{
+public:
+	FModule(const TraceServices::FModule* InModule)
+		: Module(InModule)
+	{
+		check(InModule);
+	}
+
+	FText GetAddressRangeText() const
+	{
+		TStringBuilder<64> Sb;
+		Sb.Appendf(TEXT("0x%016x - 0x%016x"), Module->Base, Module->Base + Module->Size);
+		return FText::FromStringView(Sb.ToView());
+	}
+	
+	FText GetModuleNameText() const
+	{
+		return FText::FromStringView(Module->Name);
+	}
+	
+	FText GetSymbolsFileText() const
+	{
+		return Module->StatusMessage ? FText::FromStringView(Module->StatusMessage) : FText();
+	}
+	
+	FText GetStatusText() const
+	{
+		return FText::FromStringView(ModuleStatusToString(Module->Status.load()));
+	}
+
+	FText GetStatsText() const
+	{
+		TStringBuilder<128> Sb;
+		Sb.Appendf(
+		TEXT("Discovered: %u, Cached: %u, Resolved: %u, Failed: %u"),
+			Module->Stats.Discovered.load(),
+			Module->Stats.Cached.load(),
+			Module->Stats.Resolved.load(),
+			Module->Stats.Failed.load()
+		);
+		return FText::FromStringView(Sb.ToView());
+	}
+
+	uint64 GetBaseAddress() const
+	{
+		return Module->Base;
+	}
+
+	FText GetToolTipText() const { return FText(); }
+
+private:
+	const TraceServices::FModule* Module;
+};
+	
+////////////////////////////////////////////////////////////////////////////////////////////////////
+// SModuleRow
+////////////////////////////////////////////////////////////////////////////////////////////////////
+
+class SModuleRow : public SMultiColumnTableRow<TSharedPtr<FModule>>
+{
+	SLATE_BEGIN_ARGS(SModuleRow) {}
+	SLATE_END_ARGS()
+
+public:
+	/**
+	 * Constructs the widget.
+	 *
+	 * @param InArgs The construction arguments.
+	 * @param InModule The log message displayed by this row.
+	 * @param InOwnerTableView The table to which the row must be added.
+	 */
+	void Construct(const FArguments& InArgs, TSharedPtr<FModule> InModule, TSharedRef<SModulesView> InParentWidget, const TSharedRef<STableViewBase>& InOwnerTableView)
+	{
+		WeakModule = MoveTemp(InModule);
+		WeakParentWidget = InParentWidget;
+
+		SMultiColumnTableRow<TSharedPtr<FModule>>::Construct(FSuperRowType::FArguments(), InOwnerTableView);
+
+		TSharedRef<SWidget> Row = ChildSlot.GetChildAt(0);
+
+		ChildSlot
+		[
+			SNew(SBorder)
+			//.BorderImage(FInsightsStyle::Get().GetBrush("WhiteBrush"))
+			//.BorderBackgroundColor(this, &SModuleRow::GetBackgroundColor)
+			[
+				Row
+			]
+		];
+	}
+
+public:
+	virtual TSharedRef<SWidget> GenerateWidgetForColumn(const FName& ColumnName) override
+	{
+		if (ColumnName == ModulesViewColumns::AddressRangeColumnName)
+		{
+			return SNew(SBox)
+				.Padding(FMargin(4.0, 0.0))
+				[
+					SNew(STextBlock)
+					.Text(this, &SModuleRow::GetAddressRange)
+				];
+		}
+		else if (ColumnName == ModulesViewColumns::ModuleNameColumnName)
+		{
+			return SNew(SBox)
+				.Padding(FMargin(4.0, 0.0))
+				[
+					SNew(STextBlock)
+					.Text(this, &SModuleRow::GetModuleName)
+				];
+		}
+		else if (ColumnName == ModulesViewColumns::SymbolsFileColumnName)
+		{
+			return SNew(SBox)
+				.Padding(FMargin(4.0, 0.0))
+				[
+					SNew(STextBlock)
+					.Text(this, &SModuleRow::GetSymbolsFile)
+				];
+		}
+		else if (ColumnName == ModulesViewColumns::StatusColumnName)
+		{
+			return SNew(SBox)
+				.Padding(FMargin(4.0, 0.0))
+				[
+					SNew(STextBlock)
+					.Text(this, &SModuleRow::GetStatus)
+				];
+		}
+		else if (ColumnName == ModulesViewColumns::StatsColumnName)
+		{
+			return SNew(SBox)
+				.Padding(FMargin(4.0, 0.0))
+				[
+					SNew(STextBlock)
+					.Text(this, &SModuleRow::GetStats)
+				];
+		}
+		else
+		{
+			return SNew(STextBlock).Text(LOCTEXT("UnknownColumn", "Unknown Column"));
+		}
+	}
+
+	FSlateColor GetBackgroundColor() const
+	{
+		return FSlateColor(FLinearColor(0.0f, 0.0f, 0.0f, 0.0f));
+	}
+
+	FText GetAddressRange() const
+	{
+		TSharedPtr<FModule> ModulePin = WeakModule.Pin();
+		if (ModulePin.IsValid())
+		{
+			return ModulePin->GetAddressRangeText();
+		}
+		else
+		{
+			return FText();
+		}
+	}
+
+	FText GetModuleName() const
+	{
+		TSharedPtr<FModule> ModulePin = WeakModule.Pin();
+		if (ModulePin.IsValid())
+		{
+			return ModulePin->GetModuleNameText();
+		}
+		else
+		{
+			return FText();
+		}
+	}
+
+	FText GetSymbolsFile() const
+	{
+		TSharedPtr<FModule> ModulePin = WeakModule.Pin();
+		if (ModulePin.IsValid())
+		{
+			return ModulePin->GetSymbolsFileText();
+		}
+		else
+		{
+			return FText();
+		}
+	}
+
+	FText GetStatus() const
+	{
+		TSharedPtr<FModule> ModulePin = WeakModule.Pin();
+		if (ModulePin.IsValid())
+		{
+			return ModulePin->GetStatusText();
+		}
+		else
+		{
+			return FText();
+		}
+	}
+
+	FText GetStats() const
+	{
+		const TSharedPtr<FModule> ModulePin = WeakModule.Pin();
+		if (ModulePin.IsValid())
+		{
+			return ModulePin->GetStatsText();
+		}
+		else
+		{
+			return FText();
+		}
+	}
+
+	FText GetRowToolTip() const
+	{
+		TSharedPtr<FModule> ModulePin = WeakModule.Pin();
+		if (ModulePin.IsValid())
+		{
+			return ModulePin->GetToolTipText();
+		}
+		else
+		{
+			return FText();
+		}
+	}
+
+private:
+	TWeakPtr<FModule> WeakModule;
+	TWeakPtr<SModulesView> WeakParentWidget;
+};
+
+////////////////////////////////////////////////////////////////////////////////////////////////////
+// SModulesView
+////////////////////////////////////////////////////////////////////////////////////////////////////
+
+SModulesView::SModulesView()
+	: LastUpdateTime(0.0)
+{
+}
+
+////////////////////////////////////////////////////////////////////////////////////////////////////
+
+SModulesView::~SModulesView()
+{
+	Reset();
+}
+
+////////////////////////////////////////////////////////////////////////////////////////////////////
+
+void SModulesView::Reset()
+{
+	//ListView
+	//ExternalScrollbar
+	Modules.Reset();
+
+	ListView->RebuildList();
+}
+
+////////////////////////////////////////////////////////////////////////////////////////////////////
+
+BEGIN_SLATE_FUNCTION_BUILD_OPTIMIZATION
+void SModulesView::Construct(const FArguments& InArgs)
+{
+	SAssignNew(ExternalScrollbar, SScrollBar)
+	.AlwaysShowScrollbar(true);
+
+	ChildSlot
+	[
+		SNew(SVerticalBox)
+
+		// Toolbar
+		+ SVerticalBox::Slot()
+		.AutoHeight()
+		[
+			SNew(SHorizontalBox)
+
+			// Stats Text
+			+ SHorizontalBox::Slot()
+			.VAlign(VAlign_Center)
+			.Padding(2.0f, 0.0f)
+			.AutoWidth()
+			[
+				SNew(STextBlock)
+				.Text(this, &SModulesView::GetStatsText)
+			]
+		]
+
+		+ SVerticalBox::Slot()
+		.FillHeight(1.0f)
+		[
+			SNew(SBox)
+			.VAlign(VAlign_Fill)
+			[
+				SNew(SHorizontalBox)
+
+				+ SHorizontalBox::Slot()
+				.FillWidth(1.0f)
+				.Padding(0.0f)
+				.VAlign(VAlign_Fill)
+				[
+					SNew(SScrollBox)
+					.Orientation(Orient_Horizontal)
+
+					+ SScrollBox::Slot()
+					.VAlign(VAlign_Fill)
+					[
+						SAssignNew(ListView, SListView<TSharedPtr<FModule>>)
+						.ExternalScrollbar(ExternalScrollbar)
+						.ItemHeight(20.0f)
+						.SelectionMode(ESelectionMode::Single)
+						.OnMouseButtonClick(this, &SModulesView::OnMouseButtonClick)
+						.OnSelectionChanged(this, &SModulesView::OnSelectionChanged)
+						.ListItemsSource(&Modules)
+						.OnGenerateRow(this, &SModulesView::OnGenerateRow)
+						.ConsumeMouseWheel(EConsumeMouseWheel::Always)
+						.OnContextMenuOpening(FOnContextMenuOpening::CreateSP(this, &SModulesView::ListView_GetContextMenu))
+						.HeaderRow
+						(
+							SNew(SHeaderRow)
+
+							+ SHeaderRow::Column(ModulesViewColumns::AddressRangeColumnName)
+							.ManualWidth(320.0f)
+							.DefaultLabel(LOCTEXT("AddressRangeColumn", "Address range"))
+
+							+ SHeaderRow::Column(ModulesViewColumns::ModuleNameColumnName)
+							.ManualWidth(300.0f)
+							.DefaultLabel(LOCTEXT("ModuleNameColumn", "Name"))
+
+							+ SHeaderRow::Column(ModulesViewColumns::StatsColumnName)
+							.ManualWidth(350.0f)
+							.DefaultLabel(LOCTEXT("StatsColumn", "Statistics"))
+
+							+ SHeaderRow::Column(ModulesViewColumns::StatusColumnName)
+							.ManualWidth(100.0f)
+							.DefaultLabel(LOCTEXT("StatusColumn", "Status"))
+							
+							+ SHeaderRow::Column(ModulesViewColumns::SymbolsFileColumnName)
+							.ManualWidth(800.0f)
+							.DefaultLabel(LOCTEXT("SymbolsFileColumn", "Symbols"))
+						)
+					]
+				]
+
+				+ SHorizontalBox::Slot()
+				.AutoWidth()
+				.Padding(0.0f)
+				[
+					SNew(SBox)
+					.WidthOverride(FOptionalSize(13.0f))
+					[
+						ExternalScrollbar.ToSharedRef()
+					]
+				]
+			]
+		]
+	];
+}
+END_SLATE_FUNCTION_BUILD_OPTIMIZATION
+
+////////////////////////////////////////////////////////////////////////////////////////////////////
+
+TSharedRef<ITableRow> SModulesView::OnGenerateRow(TSharedPtr<FModule> InModule, const TSharedRef<STableViewBase>& OwnerTable)
+{
+	// Generate a row for the log message corresponding to InModule.
+	return SNew(SModuleRow, InModule, SharedThis(this), OwnerTable);
+}
+
+////////////////////////////////////////////////////////////////////////////////////////////////////
+
+void SModulesView::Tick(const FGeometry& AllottedGeometry, const double InCurrentTime, const float InDeltaTime)
+{
+	using namespace TraceServices;
+	
+	if ((InCurrentTime - LastUpdateTime) < 1.0)
+	{
+		return;
+	}
+	LastUpdateTime = InCurrentTime;
+	
+	const TSharedPtr<const IAnalysisSession> Session = FInsightsManager::Get()->GetSession();
+	if (const IModuleProvider* ModuleProvider = Session->ReadProvider<IModuleProvider>("ModuleProvider"))
+	{
+		ModuleProvider->EnumerateModules(Modules.Num(), [this](const TraceServices::FModule& InModule)
+		{
+			Modules.Emplace(MakeShared<FModule>(&InModule));
+		});
+	}
+}
+
+////////////////////////////////////////////////////////////////////////////////////////////////////
+
+TSharedPtr<FModule> SModulesView::GetSelectedModule() const
+{
+	TArray<TSharedPtr<FModule>> SelectedItems = ListView->GetSelectedItems();
+	return (SelectedItems.Num() == 1) ? SelectedItems[0] : nullptr;
+}
+
+////////////////////////////////////////////////////////////////////////////////////////////////////
+
+void SModulesView::SelectModule(TSharedPtr<FModule> Module)
+{
+	if (Module.IsValid())
+	{
+	}
+}
+
+////////////////////////////////////////////////////////////////////////////////////////////////////
+
+void SModulesView::OnMouseButtonClick(TSharedPtr<FModule> Module)
+{
+	SelectModule(Module);
+}
+
+////////////////////////////////////////////////////////////////////////////////////////////////////
+
+void SModulesView::OnSelectionChanged(TSharedPtr<FModule> Module, ESelectInfo::Type SelectInfo)
+{
+	if (SelectInfo != ESelectInfo::Direct &&
+		SelectInfo != ESelectInfo::OnMouseClick)
+	{
+		SelectModule(Module);
+	}
+}
+
+////////////////////////////////////////////////////////////////////////////////////////////////////
+
+FText SModulesView::GetStatsText() const
+{
+	return FText();
+}
+
+////////////////////////////////////////////////////////////////////////////////////////////////////
+
+FSlateColor SModulesView::GetStatsTextColor() const
+{
+	return FSlateColor(FLinearColor(1.0f, 1.0f, 1.0f, 1.0f));
+}
+
+////////////////////////////////////////////////////////////////////////////////////////////////////
+
+TSharedPtr<SWidget> SModulesView::ListView_GetContextMenu()
+{
+	TSharedPtr<FModule> SelectedModule = GetSelectedModule();
+
+	FMenuBuilder MenuBuilder(/*bInShouldCloseWindowAfterMenuSelection=*/true, nullptr);
+
+	MenuBuilder.BeginSection("ModulesViewContextMenu");
+	{
+		if (SelectedModule.IsValid())
+		{
+			MenuBuilder.AddMenuEntry(
+				TAttribute<FText>(this, &SModulesView::LoadSymbolsFile_Label),
+				LOCTEXT("LoadSymbolsFile_TooltipReady", "Load symbols for a module by specifying a file. If succesful, tries to load other failed modules from the same directory."),
+				FSlateIcon(),
+				FUIAction(FExecuteAction::CreateSP(this, &SModulesView::LoadSymbols_Execute, SelectedModule, true),
+				          FCanExecuteAction::CreateSP(this, &SModulesView::LoadSymbols_CanExecute),
+				          EUIActionRepeatMode::RepeatDisabled),
+				NAME_None,
+				EUserInterfaceActionType::Button
+			);
+			MenuBuilder.AddMenuEntry(
+				TAttribute<FText>(this, &SModulesView::LoadSymbolsDirectory_Label),
+				LOCTEXT("LoadSymbolsDirectory_TooltipReady", "Load symbols for a module by specifying a directory. If succesful, tries to load other failed modules from the same directory."),
+				FSlateIcon(),
+				FUIAction(FExecuteAction::CreateSP(this, &SModulesView::LoadSymbols_Execute, SelectedModule, false),
+				          FCanExecuteAction::CreateSP(this, &SModulesView::LoadSymbols_CanExecute),
+				          EUIActionRepeatMode::RepeatDisabled),
+				NAME_None,
+				EUserInterfaceActionType::Button
+			);
+		}
+	}
+	MenuBuilder.EndSection();
+
+	return MenuBuilder.MakeWidget();
+}
+
+////////////////////////////////////////////////////////////////////////////////////////////////////
+
+FText SModulesView::LoadSymbolsFile_Label() const
+{
+	if (LoadSymbols_CanExecute())
+	{
+		return LOCTEXT("LoadSymbolsFile_Label", "Load symbols from file...");
+	}
+	return LOCTEXT("LoadSymbolsFile_LabelBusy", "Load symbols from file... (Operation in progress)");
+}
+
+////////////////////////////////////////////////////////////////////////////////////////////////////
+
+FText SModulesView::LoadSymbolsDirectory_Label() const
+{
+	if (LoadSymbols_CanExecute())
+	{
+		return LOCTEXT("LoadSymbolsDirectory_Label", "Load symbols from directory...");
+	}
+	return LOCTEXT("LoadSymbolsDirectory_LabelBusy", "Load symbols from directory... (Operation in progress)");
+}
+
+////////////////////////////////////////////////////////////////////////////////////////////////////
+
+bool SModulesView::LoadSymbols_CanExecute() const
+{
+	return !LoadSymbolsTask.IsValid() || LoadSymbolsTask->IsComplete();
+}
+	
+////////////////////////////////////////////////////////////////////////////////////////////////////
+
+void SModulesView::LoadSymbols_Execute(TSharedPtr<FModule> Module, bool bOpenFile)
+{
+	static FString DefaultPath;
+	bool bHasSelected = false;
+	FString SelectedDirectory;
+	TArray<FString> SelectedFiles;
+	if (IDesktopPlatform* DesktopPlatform = FDesktopPlatformModule::Get())
+	{
+		const FString Title = LOCTEXT("LoadSymbolsBrowseFile", "Try loading symbols from...").ToString();
+		if (bOpenFile)
+		{
+			bHasSelected = DesktopPlatform->OpenFileDialog(
+				FSlateApplication::Get().FindBestParentWindowHandleForDialogs(AsShared()),
+				Title,
+				DefaultPath,
+				Module->GetModuleNameText().ToString(),
+				TEXT(""),
+				0,
+				SelectedFiles
+			);
+		}
+		else
+		{
+			bHasSelected = DesktopPlatform->OpenDirectoryDialog(
+				FSlateApplication::Get().FindBestParentWindowHandleForDialogs(AsShared()),
+				Title,
+				DefaultPath,
+				SelectedDirectory
+				);
+		}
+	}
+
+	if (!bHasSelected)
+	{
+		return;
+	}
+	const FString& SelectedFile = bOpenFile ? SelectedFiles.Top() : SelectedDirectory;
+
+	// Make sure the last selected folder is default next time
+	DefaultPath = SelectedFile;
+
+	using namespace TraceServices;
+	const TSharedPtr<const IAnalysisSession> Session = FInsightsManager::Get()->GetSession();
+	if (IModuleProvider* ModuleProvider = const_cast<IAnalysisSession*>(Session.Get())->EditProvider<IModuleProvider>("ModuleProvider"))
+	{
+		LoadSymbolsTask = ModuleProvider->LoadSymbolsForModuleUsingPath(Module->GetBaseAddress(), *SelectedFile);
+	}
+}
+
+////////////////////////////////////////////////////////////////////////////////////////////////////
+
+} //namespace Insights
+
+#undef LOCTEXT_NAMESPACE
diff --git a/Engine/Source/Developer/TraceInsights/Private/Insights/Widgets/SModulesView.h b/Engine/Source/Developer/TraceInsights/Private/Insights/Widgets/SModulesView.h
new file mode 100644
--- /dev/null
+++ b/Engine/Source/Developer/TraceInsights/Private/Insights/Widgets/SModulesView.h
@@ -0,0 +1,103 @@
+// Copyright Epic Games, Inc. All Rights Reserved.
+
+#pragma once
+
+#include "SlateFwd.h"
+#include "Widgets/DeclarativeSyntaxSupport.h"
+#include "Widgets/SCompoundWidget.h"
+#include "Widgets/Views/SListView.h"
+
+class FMenuBuilder;
+
+namespace Insights
+{
+
+////////////////////////////////////////////////////////////////////////////////////////////////////
+
+class FModule;
+
+////////////////////////////////////////////////////////////////////////////////////////////////////
+
+namespace ModulesViewColumns
+{
+	static const FName AddressRangeColumnName(TEXT("BaseAddress"));
+	static const FName ModuleNameColumnName(TEXT("ModuleName"));
+	static const FName SymbolsFileColumnName(TEXT("SymbolsFile"));
+	static const FName StatusColumnName(TEXT("Status"));
+	static const FName StatsColumnName(TEXT("Stats"));
+}
+
+////////////////////////////////////////////////////////////////////////////////////////////////////
+
+/**
+ * Trace log window.
+ */
+class SModulesView : public SCompoundWidget
+{
+public:
+	/** Default constructor. */
+	SModulesView();
+
+	/** Virtual destructor. */
+	virtual ~SModulesView();
+
+	void Reset();
+
+	SLATE_BEGIN_ARGS(SModulesView) {}
+	SLATE_END_ARGS()
+
+	/**
+	 * Construct this widget.
+	 * @param InArgs - The declaration data for this widget
+	 */
+	void Construct(const FArguments& InArgs);
+
+	/**
+	 * Ticks this widget. Override in derived classes, but always call the parent implementation.
+	 *
+	 * @param AllottedGeometry - The space allotted for this widget
+	 * @param InCurrentTime - Current absolute real time
+	 * @param InDeltaTime - Real time passed since last tick
+	 */
+	virtual void Tick(const FGeometry& AllottedGeometry, const double InCurrentTime, const float InDeltaTime) override;
+
+	TSharedPtr<FModule> GetSelectedModule() const;
+
+protected:
+	/** Generate a new list view row. */
+	TSharedRef<ITableRow> OnGenerateRow(TSharedPtr<FModule> InModule, const TSharedRef<STableViewBase>& OwnerTable);
+
+	void SelectModule(TSharedPtr<FModule> Module);
+	void OnMouseButtonClick(TSharedPtr<FModule> Module);
+	void OnSelectionChanged(TSharedPtr<FModule> Module, ESelectInfo::Type SelectInfo);
+
+	FText GetStatsText() const;
+	FSlateColor GetStatsTextColor() const;
+
+	TSharedPtr<SWidget> ListView_GetContextMenu();
+	FText LoadSymbolsFile_Label() const;
+	FText LoadSymbolsDirectory_Label() const;
+	bool LoadSymbols_CanExecute() const;
+
+	void LoadSymbols_Execute(TSharedPtr<FModule> Module, bool bOpenFile);
+
+protected:
+	/** The list view widget. */
+	TSharedPtr<SListView<TSharedPtr<FModule>>> ListView;
+
+	/** External scrollbar used to synchronize tree view position. */
+	TSharedPtr<SScrollBar> ExternalScrollbar;
+
+	/** List of callstack frames (view model) to show in list view. */
+	TArray<TSharedPtr<FModule>> Modules;
+
+	/** Reference to ongoing load symbols task. */
+	FGraphEventRef LoadSymbolsTask;
+
+	/** Last time we queried about discovered modules. */
+	double LastUpdateTime;
+};
+
+////////////////////////////////////////////////////////////////////////////////////////////////////
+
+} //namespace Insights
