[Insights] Timing Insights: Removed Ctrl+T and L key shortcuts for toggling on/off the Timing view and the Log view (as those shortcuts are used instead by the Timing view itself).

diff --git a/Engine/Source/Developer/TraceInsights/Private/Insights/TimingProfilerCommands.cpp b/Engine/Source/Developer/TraceInsights/Private/Insights/TimingProfilerCommands.cpp
--- a/Engine/Source/Developer/TraceInsights/Private/Insights/TimingProfilerCommands.cpp
+++ b/Engine/Source/Developer/TraceInsights/Private/Insights/TimingProfilerCommands.cpp
@@ -52,13 +52,12 @@ PRAGMA_DISABLE_OPTIMIZATION
 void FTimingProfilerCommands::RegisterCommands()
 {
 	UI_COMMAND(ToggleFramesTrackVisibility, "Frames", "Toggles the visibility of the Frames track", EUserInterfaceActionType::ToggleButton, FInputChord());
-	UI_COMMAND(ToggleGraphTrackVisibility, "Graph", "Toggles the visibility of the Overview Graph track", EUserInterfaceActionType::ToggleButton, FInputChord(EModifierKey::Control, EKeys::G));
-	UI_COMMAND(ToggleTimingViewVisibility, "Timing", "Toggles the visibility of the main Timing view", EUserInterfaceActionType::ToggleButton, FInputChord(EModifierKey::Control, EKeys::T));
+	UI_COMMAND(ToggleTimingViewVisibility, "Timing", "Toggles the visibility of the main Timing view", EUserInterfaceActionType::ToggleButton, FInputChord());
 	UI_COMMAND(ToggleTimersViewVisibility, "Timers", "Toggles the visibility of the Timers view", EUserInterfaceActionType::ToggleButton, FInputChord());
 	UI_COMMAND(ToggleCallersTreeViewVisibility, "Callers", "Toggles the visibility of the Callers tree view", EUserInterfaceActionType::ToggleButton, FInputChord());
 	UI_COMMAND(ToggleCalleesTreeViewVisibility, "Callees", "Toggles the visibility of the Callees tree view", EUserInterfaceActionType::ToggleButton, FInputChord());
 	UI_COMMAND(ToggleStatsCountersViewVisibility, "Counters", "Toggles the visibility of the Counters view", EUserInterfaceActionType::ToggleButton, FInputChord());
-	UI_COMMAND(ToggleLogViewVisibility, "Log", "Toggles the visibility of the Log view", EUserInterfaceActionType::ToggleButton, FInputChord(EKeys::L));
+	UI_COMMAND(ToggleLogViewVisibility, "Log", "Toggles the visibility of the Log view", EUserInterfaceActionType::ToggleButton, FInputChord());
 }
 PRAGMA_ENABLE_OPTIMIZATION
 
diff --git a/Engine/Source/Developer/TraceInsights/Private/Insights/TimingProfilerCommands.h b/Engine/Source/Developer/TraceInsights/Private/Insights/TimingProfilerCommands.h
--- a/Engine/Source/Developer/TraceInsights/Private/Insights/TimingProfilerCommands.h
+++ b/Engine/Source/Developer/TraceInsights/Private/Insights/TimingProfilerCommands.h
@@ -32,9 +32,6 @@ public:
 	/** Toggles visibility for the Frames Track. Global and custom command. */
 	TSharedPtr<FUICommandInfo> ToggleFramesTrackVisibility;
 
-	/** Toggles visibility for the Graph Track. Global and custom command. */
-	TSharedPtr<FUICommandInfo> ToggleGraphTrackVisibility;
-
 	/** Toggles visibility for the Timing View. Global and custom command. */
 	TSharedPtr<FUICommandInfo> ToggleTimingViewVisibility;
 
@@ -78,7 +75,7 @@ public:
 	/** Toggles visibility for Main Graph track. */
 	TSharedPtr<FUICommandInfo> ShowMainGraphTrack;
 
-	/** Opens the quick find widget */
+	/** Opens the Quick Find widget. */
 	TSharedPtr<FUICommandInfo> QuickFind;
 
 	/** Toggles visibility for GPU thread track. */
@@ -138,7 +135,6 @@ protected:\
 	ECheckBoxState CmdName##_GetCheckState() const; /**< Handles FGetActionCheckState for CmdName. */
 
 	DECLARE_TOGGLE_COMMAND(ToggleFramesTrackVisibility)
-	DECLARE_TOGGLE_COMMAND(ToggleGraphTrackVisibility)
 	DECLARE_TOGGLE_COMMAND(ToggleTimingViewVisibility)
 	DECLARE_TOGGLE_COMMAND(ToggleTimersViewVisibility)
 	DECLARE_TOGGLE_COMMAND(ToggleCallersTreeViewVisibility)
