[Insights] Fixed various issues re alignment and clipping of text rendered in Timing view when windows scaling (high dpi) is other than 100% (ex.: alignment of tooltip for a cpu timing event; alignment of labels in graph tracks; clipping of text displayed in the bookmarks/logs track and in the timing events; etc.).

diff --git a/Engine/Source/Developer/TraceInsights/Private/Insights/MemoryProfiler/ViewModels/MemoryGraphTrack.cpp b/Engine/Source/Developer/TraceInsights/Private/Insights/MemoryProfiler/ViewModels/MemoryGraphTrack.cpp
--- a/Engine/Source/Developer/TraceInsights/Private/Insights/MemoryProfiler/ViewModels/MemoryGraphTrack.cpp
+++ b/Engine/Source/Developer/TraceInsights/Private/Insights/MemoryProfiler/ViewModels/MemoryGraphTrack.cpp
@@ -892,11 +892,12 @@ void FMemoryGraphTrack::DrawHorizontalAxisLabel(const TDrawHorizontalAxisLabelPa
 		}
 	}
 
-	const FVector2D TextSize = Params.FontMeasureService->Measure(LabelText, Font);
+	const float FontScale = Params.DrawContext.Geometry.Scale;
+	const FVector2D TextSize = Params.FontMeasureService->Measure(LabelText, Font, FontScale) / FontScale;
 	constexpr float TextH = 14.0f;
 
 	// Draw background for value text.
-	Params.DrawContext.DrawBox(Params.DrawContext.LayerId + 1, Params.X - TextSize.X - 4.0f, Params.Y, TextSize.X + 4.0f, TextH, Params.Brush, Params.TextBgColor);
+	Params.DrawContext.DrawBox(Params.DrawContext.LayerId + 1, Params.X - TextSize.X - 4.0f, Params.Y, TextSize.X + 5.0f, TextH, Params.Brush, Params.TextBgColor);
 
 	// Draw value text.
 	Params.DrawContext.DrawText(Params.DrawContext.LayerId + 2, Params.X - TextSize.X - 2.0f, Params.Y + 1.0f, LabelText, Font, Params.TextColor);
diff --git a/Engine/Source/Developer/TraceInsights/Private/Insights/NetworkingProfiler/ViewModels/PacketContentViewDrawHelper.cpp b/Engine/Source/Developer/TraceInsights/Private/Insights/NetworkingProfiler/ViewModels/PacketContentViewDrawHelper.cpp
--- a/Engine/Source/Developer/TraceInsights/Private/Insights/NetworkingProfiler/ViewModels/PacketContentViewDrawHelper.cpp
+++ b/Engine/Source/Developer/TraceInsights/Private/Insights/NetworkingProfiler/ViewModels/PacketContentViewDrawHelper.cpp
@@ -23,13 +23,14 @@
 // FPacketContentViewDrawStateBuilder
 ////////////////////////////////////////////////////////////////////////////////////////////////////
 
-FPacketContentViewDrawStateBuilder::FPacketContentViewDrawStateBuilder(FPacketContentViewDrawState& InDrawState, const FPacketContentViewport& InViewport)
+FPacketContentViewDrawStateBuilder::FPacketContentViewDrawStateBuilder(FPacketContentViewDrawState& InDrawState, const FPacketContentViewport& InViewport, float InFontScale)
 	: DrawState(InDrawState)
 	, Viewport(InViewport)
 	, MaxDepth(-1)
 	, LastEventX2()
 	, LastBox()
 	, EventFont(FAppStyle::Get().GetFontStyle("SmallFont"))
+	, FontScale(InFontScale)
 {
 	DrawState.Reset();
 }
@@ -251,7 +252,8 @@ void FPacketContentViewDrawStateBuilder::AddEvent(const TraceServices::FNetProfi
 		const FString Name(Builder);
 
 		const TSharedRef<FSlateFontMeasure> FontMeasureService = FSlateApplication::Get().GetRenderer()->GetFontMeasureService();
-		const int32 LastWholeCharacterIndex = FontMeasureService->FindLastWholeCharacterIndexBeforeOffset(Name, EventFont, FMath::RoundToInt(EventW - 2.0f));
+		const int32 HorizontalOffset = FMath::RoundToInt((EventW - 2.0f) * FontScale);
+		const int32 LastWholeCharacterIndex = FontMeasureService->FindLastWholeCharacterIndexBeforeOffset(Name, EventFont, HorizontalOffset, FontScale);
 
 		if (LastWholeCharacterIndex >= 0)
 		{
diff --git a/Engine/Source/Developer/TraceInsights/Private/Insights/NetworkingProfiler/ViewModels/PacketContentViewDrawHelper.h b/Engine/Source/Developer/TraceInsights/Private/Insights/NetworkingProfiler/ViewModels/PacketContentViewDrawHelper.h
--- a/Engine/Source/Developer/TraceInsights/Private/Insights/NetworkingProfiler/ViewModels/PacketContentViewDrawHelper.h
+++ b/Engine/Source/Developer/TraceInsights/Private/Insights/NetworkingProfiler/ViewModels/PacketContentViewDrawHelper.h
@@ -159,7 +159,7 @@ private:
 	};
 
 public:
-	explicit FPacketContentViewDrawStateBuilder(FPacketContentViewDrawState& InState, const FPacketContentViewport& InViewport);
+	explicit FPacketContentViewDrawStateBuilder(FPacketContentViewDrawState& InState, const FPacketContentViewport& InViewport, float InFontScale);
 
 	/**
 	 * Non-copyable
@@ -185,6 +185,7 @@ private:
 	TArray<FBoxData> LastBox;
 
 	const FSlateFontInfo EventFont;
+	float FontScale;
 };
 
 ////////////////////////////////////////////////////////////////////////////////////////////////////
diff --git a/Engine/Source/Developer/TraceInsights/Private/Insights/NetworkingProfiler/Widgets/SPacketContentView.cpp b/Engine/Source/Developer/TraceInsights/Private/Insights/NetworkingProfiler/Widgets/SPacketContentView.cpp
--- a/Engine/Source/Developer/TraceInsights/Private/Insights/NetworkingProfiler/Widgets/SPacketContentView.cpp
+++ b/Engine/Source/Developer/TraceInsights/Private/Insights/NetworkingProfiler/Widgets/SPacketContentView.cpp
@@ -6,6 +6,7 @@
 #include "Fonts/SlateFontInfo.h"
 #include "Framework/Application/SlateApplication.h"
 #include "Framework/MultiBox/MultiBoxBuilder.h"
+#include "HAL/PlatformApplicationMisc.h"
 #include "Rendering/DrawElements.h"
 #include "Styling/AppStyle.h"
 #include "Widgets/Images/SImage.h"
@@ -677,6 +678,9 @@ void SPacketContentView::Tick(const FGeometry& AllottedGeometry, const double In
 
 	ThisGeometry = AllottedGeometry;
 
+	const float FontScale = AllottedGeometry.Scale;
+	Tooltip.SetFontScale(FontScale);
+
 	FAxisViewportDouble& ViewportX = Viewport.GetHorizontalAxisViewport();
 
 	if (!bIsScrolling)
@@ -691,7 +695,7 @@ void SPacketContentView::Tick(const FGeometry& AllottedGeometry, const double In
 	if (bIsStateDirty)
 	{
 		bIsStateDirty = false;
-		UpdateState();
+		UpdateState(FontScale);
 		AdjustForSplitContent();
 	}
 
@@ -847,15 +851,15 @@ void SPacketContentView::DisableFilterEventType()
 
 ////////////////////////////////////////////////////////////////////////////////////////////////////
 
-void SPacketContentView::UpdateState()
+void SPacketContentView::UpdateState(float FontScale)
 {
 	FStopwatch Stopwatch;
 	Stopwatch.Start();
 
 	if (PacketBitSize > 0)
 	{
-		FPacketContentViewDrawStateBuilder Builder(*DrawState, Viewport);
-		FPacketContentViewDrawStateBuilder FilteredDrawStateBuilder(*FilteredDrawState, Viewport);
+		FPacketContentViewDrawStateBuilder Builder(*DrawState, Viewport, FontScale);
+		FPacketContentViewDrawStateBuilder FilteredDrawStateBuilder(*FilteredDrawState, Viewport, FontScale);
 
 		TSharedPtr<const TraceServices::IAnalysisSession> Session = FInsightsManager::Get()->GetSession();
 		if (Session.IsValid())
@@ -1105,7 +1109,8 @@ int32 SPacketContentView::OnPaint(const FPaintArgs& Args, const FGeometry& Allot
 		FSlateFontInfo SummaryFont = FAppStyle::Get().GetFontStyle("SmallFont");
 
 		const TSharedRef<FSlateFontMeasure> FontMeasureService = FSlateApplication::Get().GetRenderer()->GetFontMeasureService();
-		const float MaxFontCharHeight = FontMeasureService->Measure(TEXT("!"), SummaryFont).Y;
+		const float FontScale = DrawContext.Geometry.Scale;
+		const float MaxFontCharHeight = FontMeasureService->Measure(TEXT("!"), SummaryFont, FontScale).Y / FontScale;
 		const float DbgDY = MaxFontCharHeight;
 
 		const float DbgW = 280.0f;
diff --git a/Engine/Source/Developer/TraceInsights/Private/Insights/NetworkingProfiler/Widgets/SPacketContentView.h b/Engine/Source/Developer/TraceInsights/Private/Insights/NetworkingProfiler/Widgets/SPacketContentView.h
--- a/Engine/Source/Developer/TraceInsights/Private/Insights/NetworkingProfiler/Widgets/SPacketContentView.h
+++ b/Engine/Source/Developer/TraceInsights/Private/Insights/NetworkingProfiler/Widgets/SPacketContentView.h
@@ -200,7 +200,7 @@ private:
 
 	uint32 GetPacketSequence(int32 PacketIndex) const;
 
-	void UpdateState();
+	void UpdateState(float FontScale);
 
 	void UpdateHoveredEvent();
 	FNetworkPacketEventRef GetEventAtMousePosition(float X, float Y);
diff --git a/Engine/Source/Developer/TraceInsights/Private/Insights/NetworkingProfiler/Widgets/SPacketView.cpp b/Engine/Source/Developer/TraceInsights/Private/Insights/NetworkingProfiler/Widgets/SPacketView.cpp
--- a/Engine/Source/Developer/TraceInsights/Private/Insights/NetworkingProfiler/Widgets/SPacketView.cpp
+++ b/Engine/Source/Developer/TraceInsights/Private/Insights/NetworkingProfiler/Widgets/SPacketView.cpp
@@ -653,6 +653,7 @@ int32 SPacketView::OnPaint(const FPaintArgs& Args, const FGeometry& AllottedGeom
 
 	const TSharedRef<FSlateFontMeasure> FontMeasureService = FSlateApplication::Get().GetRenderer()->GetFontMeasureService();
 	FSlateFontInfo SummaryFont = FAppStyle::Get().GetFontStyle("SmallFont");
+	const float FontScale = AllottedGeometry.Scale;
 
 	const FSlateBrush* WhiteBrush = FInsightsStyle::Get().GetBrush("WhiteBrush");
 
@@ -775,7 +776,7 @@ int32 SPacketView::OnPaint(const FPaintArgs& Args, const FGeometry& AllottedGeom
 					});
 			}
 
-			FVector2D TextSize = FontMeasureService->Measure(Text, SummaryFont);
+			const FVector2D TextSize = FontMeasureService->Measure(Text, SummaryFont, FontScale) / FontScale;
 
 			const float DX = 2.0f;
 			const float W2 = TextSize.X / 2 + DX;
@@ -809,7 +810,7 @@ int32 SPacketView::OnPaint(const FPaintArgs& Args, const FGeometry& AllottedGeom
 	const bool bShouldDisplayDebugInfo = FInsightsManager::Get()->IsDebugInfoEnabled();
 	if (bShouldDisplayDebugInfo)
 	{
-		const float MaxFontCharHeight = FontMeasureService->Measure(TEXT("!"), SummaryFont).Y;
+		const float MaxFontCharHeight = FontMeasureService->Measure(TEXT("!"), SummaryFont, FontScale).Y / FontScale;
 		const float DbgDY = MaxFontCharHeight;
 
 		const float DbgW = 280.0f;
@@ -946,6 +947,7 @@ void SPacketView::DrawVerticalAxisGrid(FDrawContext& DrawContext, const FSlateBr
 		const FLinearColor TextColor(1.0f, 1.0f, 1.0f, 1.0f);
 
 		const TSharedRef<FSlateFontMeasure> FontMeasureService = FSlateApplication::Get().GetRenderer()->GetFontMeasureService();
+		const float FontScale = DrawContext.Geometry.Scale;
 
 		for (double Value = StartValue; Value < TopValue; Value += Grid)
 		{
@@ -956,7 +958,7 @@ void SPacketView::DrawVerticalAxisGrid(FDrawContext& DrawContext, const FSlateBr
 
 			const int64 ValueBits = static_cast<int64>(Value);
 			const FString Text = (ValueBits == 0) ? TEXT("0") : FString::Format(TEXT("{0} bits"), { FText::AsNumber(ValueBits).ToString() });
-			const FVector2D TextSize = FontMeasureService->Measure(Text, Font);
+			const FVector2D TextSize = FontMeasureService->Measure(Text, Font, FontScale) / FontScale;
 			constexpr float TextH = 14.0f;
 
 			// Draw background for value text.
@@ -1020,6 +1022,7 @@ void SPacketView::DrawHorizontalAxisGrid(FDrawContext& DrawContext, const FSlate
 		const FLinearColor TopTextColor(1.0f, 1.0f, 1.0f, 0.7f);
 
 		//const TSharedRef<FSlateFontMeasure> FontMeasureService = FSlateApplication::Get().GetRenderer()->GetFontMeasureService();
+		//const float FontScale = DrawContext.Geometry.Scale;
 
 		for (int32 Index = StartIndex; Index < RightIndex; Index += Grid)
 		{
@@ -1029,7 +1032,7 @@ void SPacketView::DrawHorizontalAxisGrid(FDrawContext& DrawContext, const FSlate
 			DrawContext.DrawBox(X, 0, 1, ViewHeight, Brush, GridColor);
 
 			const FString Text = FText::AsNumber(Index).ToString();
-			//const FVector2D TextSize = FontMeasureService->Measure(Text, Font);
+			//const FVector2D TextSize = FontMeasureService->Measure(Text, Font, FontScale) / FontScale;
 			//constexpr float TextH = 14.0f;
 
 			// Draw background for index text.
diff --git a/Engine/Source/Developer/TraceInsights/Private/Insights/ViewModels/DrawHelpers.cpp b/Engine/Source/Developer/TraceInsights/Private/Insights/ViewModels/DrawHelpers.cpp
--- a/Engine/Source/Developer/TraceInsights/Private/Insights/ViewModels/DrawHelpers.cpp
+++ b/Engine/Source/Developer/TraceInsights/Private/Insights/ViewModels/DrawHelpers.cpp
@@ -326,7 +326,8 @@ void FDrawHelpers::DrawSelection(
 	if (!Text.IsEmpty())
 	{
 		const TSharedRef<FSlateFontMeasure> FontMeasureService = FSlateApplication::Get().GetRenderer()->GetFontMeasureService();
-		const float TextWidth = FontMeasureService->Measure(Text, Font).X;
+		const float FontScale = DrawContext.Geometry.Scale;
+		const float TextWidth = FontMeasureService->Measure(Text, Font, FontScale).X / FontScale;
 
 		const float CenterX = (SelectionX1 + SelectionX2) / 2.0f;
 
diff --git a/Engine/Source/Developer/TraceInsights/Private/Insights/ViewModels/FrameTimingTrack.cpp b/Engine/Source/Developer/TraceInsights/Private/Insights/ViewModels/FrameTimingTrack.cpp
--- a/Engine/Source/Developer/TraceInsights/Private/Insights/ViewModels/FrameTimingTrack.cpp
+++ b/Engine/Source/Developer/TraceInsights/Private/Insights/ViewModels/FrameTimingTrack.cpp
@@ -254,6 +254,7 @@ void FFrameTimingTrack::Update(const ITimingTrackUpdateContext& Context)
 {
 	FTimingEventsTrack::Update(Context);
 
+	Header.SetFontScale(Context.GetGeometry().Scale);
 	Header.UpdateSize();
 	Header.Update(Context);
 }
@@ -292,6 +293,7 @@ void FFrameTimingTrack::PostDraw(const ITimingTrackDrawContext& Context) const
 	{
 		const FTimingEvent& SelectedEvent = SelectedEventPtr->As<FTimingEvent>();
 		const ITimingViewDrawHelper& Helper = Context.GetHelper();
+		const float FontScale = Context.GetDrawContext().Geometry.Scale;
 		DrawSelectedEventInfo(SelectedEvent, Context.GetViewport(), Context.GetDrawContext(), Helper.GetWhiteBrush(), Helper.GetEventFont());
 	}
 
@@ -309,7 +311,8 @@ void FFrameTimingTrack::DrawSelectedEventInfo(const FTimingEvent& SelectedEvent,
 		const FString Str = GetCompleteFrameName(InFoundFrame.Index, Duration);
 
 		const TSharedRef<FSlateFontMeasure> FontMeasureService = FSlateApplication::Get().GetRenderer()->GetFontMeasureService();
-		const FVector2D Size = FontMeasureService->Measure(Str, Font);
+		const float FontScale = DrawContext.Geometry.Scale;
+		const FVector2D Size = FontMeasureService->Measure(Str, Font, FontScale) / FontScale;
 		const float X = Viewport.GetWidth() - Size.X - 23.0f;
 		const float Y = Viewport.GetPosY() + Viewport.GetHeight() - Size.Y - 18.0f;
 
diff --git a/Engine/Source/Developer/TraceInsights/Private/Insights/ViewModels/MarkersTimingTrack.cpp b/Engine/Source/Developer/TraceInsights/Private/Insights/ViewModels/MarkersTimingTrack.cpp
--- a/Engine/Source/Developer/TraceInsights/Private/Insights/ViewModels/MarkersTimingTrack.cpp
+++ b/Engine/Source/Developer/TraceInsights/Private/Insights/ViewModels/MarkersTimingTrack.cpp
@@ -107,6 +107,7 @@ void FMarkersTimingTrack::PreUpdate(const ITimingTrackUpdateContext& Context)
 
 void FMarkersTimingTrack::Update(const ITimingTrackUpdateContext& Context)
 {
+	Header.SetFontScale(Context.GetGeometry().Scale);
 	Header.Update(Context);
 
 	const FTimingTrackViewport& Viewport = Context.GetViewport();
@@ -114,7 +115,7 @@ void FMarkersTimingTrack::Update(const ITimingTrackUpdateContext& Context)
 	{
 		ClearDirtyFlag();
 
-		UpdateDrawState(Viewport);
+		UpdateDrawState(Context);
 	}
 }
 
@@ -130,9 +131,9 @@ void FMarkersTimingTrack::PostUpdate(const ITimingTrackUpdateContext& Context)
 
 ////////////////////////////////////////////////////////////////////////////////////////////////////
 
-void FMarkersTimingTrack::UpdateDrawState(const FTimingTrackViewport& InViewport)
+void FMarkersTimingTrack::UpdateDrawState(const ITimingTrackUpdateContext& Context)
 {
-	FTimeMarkerTrackBuilder Builder(*this, InViewport);
+	FTimeMarkerTrackBuilder Builder(*this, Context.GetViewport(), Context.GetGeometry().Scale);
 
 	TSharedPtr<const TraceServices::IAnalysisSession> Session = FInsightsManager::Get()->GetSession();
 	if (Session.IsValid())
@@ -387,11 +388,12 @@ void FMarkersTimingTrack::UpdateBookmarkCategory()
 // FTimeMarkerTrackBuilder
 ////////////////////////////////////////////////////////////////////////////////////////////////////
 
-FTimeMarkerTrackBuilder::FTimeMarkerTrackBuilder(FMarkersTimingTrack& InTrack, const FTimingTrackViewport& InViewport)
+FTimeMarkerTrackBuilder::FTimeMarkerTrackBuilder(FMarkersTimingTrack& InTrack, const FTimingTrackViewport& InViewport, float InFontScale)
 	: Track(InTrack)
 	, Viewport(InViewport)
 	, FontMeasureService(FSlateApplication::Get().GetRenderer()->GetFontMeasureService())
 	, Font(FAppStyle::Get().GetFontStyle("SmallFont"))
+	, FontScale(InFontScale)
 {
 	Track.ResetCache();
 	Track.NumLogMessages = 0;
@@ -494,7 +496,7 @@ void FTimeMarkerTrackBuilder::Flush(float AvailableTextW)
 		bool bAddNewBox = true;
 		if (Track.TimeMarkerBoxes.Num() > 0)
 		{
-			FTimeMarkerBoxInfo& PrevBox = Track.TimeMarkerBoxes[Track.TimeMarkerBoxes.Num() - 1];
+			FTimeMarkerBoxInfo& PrevBox = Track.TimeMarkerBoxes.Last();
 			if (PrevBox.X + PrevBox.W == LastX1 &&
 				PrevBox.Color.R == Color.R &&
 				PrevBox.Color.G == Color.G &&
@@ -509,7 +511,7 @@ void FTimeMarkerTrackBuilder::Flush(float AvailableTextW)
 		if (bAddNewBox)
 		{
 			// Add new Box info to cache.
-			FTimeMarkerBoxInfo& Box = Track.TimeMarkerBoxes[Track.TimeMarkerBoxes.AddDefaulted()];
+			FTimeMarkerBoxInfo& Box = Track.TimeMarkerBoxes.AddDefaulted_GetRef();
 			Box.X = LastX1;
 			Box.W = LastX2 - LastX1;
 			Box.Color = Color;
@@ -525,19 +527,20 @@ void FTimeMarkerTrackBuilder::Flush(float AvailableTextW)
 				CategoryStr.RightChopInline(3, false);
 			}
 
-			const int32 LastWholeCharacterIndexCategory = FontMeasureService->FindLastWholeCharacterIndexBeforeOffset(CategoryStr, Font, FMath::RoundToInt(AvailableTextW - 2.0f));
-			const int32 LastWholeCharacterIndexMessage = FontMeasureService->FindLastWholeCharacterIndexBeforeOffset(LastMessage, Font, FMath::RoundToInt(AvailableTextW - 2.0f));
+			const int32 HorizontalOffset = FMath::RoundToInt((AvailableTextW - 2.0f) * FontScale);
+			const int32 LastWholeCharacterIndexCategory = FontMeasureService->FindLastWholeCharacterIndexBeforeOffset(CategoryStr, Font, HorizontalOffset, FontScale);
+			const int32 LastWholeCharacterIndexMessage = FontMeasureService->FindLastWholeCharacterIndexBeforeOffset(LastMessage, Font, HorizontalOffset, FontScale);
 
 			if (LastWholeCharacterIndexCategory >= 0 ||
 				LastWholeCharacterIndexMessage >= 0)
 			{
 				// Add new Text info to cache.
-				FTimeMarkerTextInfo& TextInfo = Track.TimeMarkerTexts[Track.TimeMarkerTexts.AddDefaulted()];
+				FTimeMarkerTextInfo& TextInfo = Track.TimeMarkerTexts.AddDefaulted_GetRef();
 				TextInfo.X = LastX2 + 2.0f;
 				TextInfo.Color = Color;
 				if (LastWholeCharacterIndexCategory >= 0)
 				{
-					TextInfo.Category = CategoryStr.Left(LastWholeCharacterIndexCategory + 1);
+					TextInfo.Category.AppendChars(*CategoryStr, LastWholeCharacterIndexCategory + 1);
 				}
 				if (LastWholeCharacterIndexMessage >= 0)
 				{
diff --git a/Engine/Source/Developer/TraceInsights/Private/Insights/ViewModels/MarkersTimingTrack.h b/Engine/Source/Developer/TraceInsights/Private/Insights/ViewModels/MarkersTimingTrack.h
--- a/Engine/Source/Developer/TraceInsights/Private/Insights/ViewModels/MarkersTimingTrack.h
+++ b/Engine/Source/Developer/TraceInsights/Private/Insights/ViewModels/MarkersTimingTrack.h
@@ -93,7 +93,7 @@ private:
 	}
 
 	void UpdateTrackNameAndHeight();
-	void UpdateDrawState(const FTimingTrackViewport& Viewport);
+	void UpdateDrawState(const ITimingTrackUpdateContext& Context);
 
 	void UpdateBookmarkCategory();
 
@@ -121,7 +121,7 @@ private:
 class FTimeMarkerTrackBuilder
 {
 public:
-	explicit FTimeMarkerTrackBuilder(FMarkersTimingTrack& InTrack, const FTimingTrackViewport& InViewport);
+	explicit FTimeMarkerTrackBuilder(FMarkersTimingTrack& InTrack, const FTimingTrackViewport& InViewport, float InFontScale);
 
 	/**
 	 * Non-copyable
@@ -148,6 +148,7 @@ private:
 
 	const TSharedRef<class FSlateFontMeasure> FontMeasureService;
 	const FSlateFontInfo Font;
+	float FontScale;
 
 	const TraceServices::ILogProvider* LogProviderPtr; // valid only between BeginLog() and EndLog()
 
diff --git a/Engine/Source/Developer/TraceInsights/Private/Insights/ViewModels/ThreadTimingTrack.cpp b/Engine/Source/Developer/TraceInsights/Private/Insights/ViewModels/ThreadTimingTrack.cpp
--- a/Engine/Source/Developer/TraceInsights/Private/Insights/ViewModels/ThreadTimingTrack.cpp
+++ b/Engine/Source/Developer/TraceInsights/Private/Insights/ViewModels/ThreadTimingTrack.cpp
@@ -1036,7 +1036,8 @@ void FThreadTimingTrack::DrawSelectedEventInfo(const FThreadTrackEvent& Selected
 			*TimeUtils::FormatTimeAuto(SelectedEvent.GetExclusiveTime()));
 
 		const TSharedRef<FSlateFontMeasure> FontMeasureService = FSlateApplication::Get().GetRenderer()->GetFontMeasureService();
-		const FVector2D Size = FontMeasureService->Measure(Str, Font);
+		const float FontScale = DrawContext.Geometry.Scale;
+		const FVector2D Size = FontMeasureService->Measure(Str, Font, FontScale) / FontScale;
 		const float X = Viewport.GetWidth() - Size.X - 23.0f;
 		const float Y = Viewport.GetPosY() + Viewport.GetHeight() - Size.Y - 18.0f;
 
diff --git a/Engine/Source/Developer/TraceInsights/Private/Insights/ViewModels/TimeRulerTrack.cpp b/Engine/Source/Developer/TraceInsights/Private/Insights/ViewModels/TimeRulerTrack.cpp
--- a/Engine/Source/Developer/TraceInsights/Private/Insights/ViewModels/TimeRulerTrack.cpp
+++ b/Engine/Source/Developer/TraceInsights/Private/Insights/ViewModels/TimeRulerTrack.cpp
@@ -218,13 +218,14 @@ void FTimeRulerTrack::Draw(const ITimingTrackDrawContext& Context) const
 	const double Precision = FMath::Max(DT / 10.0, TimeUtils::Nanosecond);
 
 	const TSharedRef<FSlateFontMeasure> FontMeasureService = FSlateApplication::Get().GetRenderer()->GetFontMeasureService();
+	const float FontScale = DrawContext.Geometry.Scale;
 
 	// Draw the time at major tick marks.
 	for (float X = MajorOX; X < Viewport.GetWidth() + MajorTickMark; X += MajorTickMark)
 	{
 		const double T = Viewport.SlateUnitsToTime(X);
 		FString Text = TimeUtils::FormatTime(T, Precision);
-		const float TextWidth = FontMeasureService->Measure(Text, Font).X;
+		const float TextWidth = FontMeasureService->Measure(Text, Font, FontScale).X / FontScale;
 		DrawContext.DrawText(X - TextWidth / 2, TextY, Text, Font,
 			(T < Viewport.GetMinValidTime() || T >= Viewport.GetMaxValidTime()) ? FLinearColor(0.7f, 0.5f, 0.5f, 1.0f) : FLinearColor(0.8f, 0.8f, 0.8f, 1.0f));
 	}
@@ -267,8 +268,9 @@ void FTimeRulerTrack::PostDraw(const ITimingTrackDrawContext& Context) const
 		}
 
 		const TSharedRef<FSlateFontMeasure> FontMeasureService = FSlateApplication::Get().GetRenderer()->GetFontMeasureService();
+		const float FontScale = DrawContext.Geometry.Scale;
 
-		const float MousePosTextWidth = FMath::RoundToFloat(FontMeasureService->Measure(MousePosText, Font).X);
+		const float MousePosTextWidth = FMath::RoundToFloat(FontMeasureService->Measure(MousePosText, Font, FontScale).X / FontScale);
 
 		if (!FMath::IsNearlyEqual(CrtMousePosTextWidth, MousePosTextWidth))
 		{
@@ -400,8 +402,9 @@ void FTimeRulerTrack::DrawTimeMarker(const ITimingTrackDrawContext& Context, con
 	}
 
 	const TSharedRef<FSlateFontMeasure> FontMeasureService = FSlateApplication::Get().GetRenderer()->GetFontMeasureService();
+	const float FontScale = DrawContext.Geometry.Scale;
 
-	const float TextWidth = FMath::RoundToFloat(FontMeasureService->Measure(TimeMarkerText, Font).X);
+	const float TextWidth = FMath::RoundToFloat(FontMeasureService->Measure(TimeMarkerText, Font, FontScale).X / FontScale);
 
 	if (!FMath::IsNearlyEqual(TimeMarker.GetCrtTextWidth(), TextWidth))
 	{
diff --git a/Engine/Source/Developer/TraceInsights/Private/Insights/ViewModels/TimingEventsTrack.cpp b/Engine/Source/Developer/TraceInsights/Private/Insights/ViewModels/TimingEventsTrack.cpp
--- a/Engine/Source/Developer/TraceInsights/Private/Insights/ViewModels/TimingEventsTrack.cpp
+++ b/Engine/Source/Developer/TraceInsights/Private/Insights/ViewModels/TimingEventsTrack.cpp
@@ -71,7 +71,7 @@ void FTimingEventsTrack::PreUpdate(const ITimingTrackUpdateContext& Context)
 		int32 MaxDepth = -1;
 
 		{
-			FTimingEventsTrackDrawStateBuilder Builder(*DrawState, Context.GetViewport());
+			FTimingEventsTrackDrawStateBuilder Builder(*DrawState, Context.GetViewport(), Context.GetGeometry().Scale);
 
 			BuildDrawState(Builder, Context);
 
@@ -123,7 +123,7 @@ void FTimingEventsTrack::PreUpdate(const ITimingTrackUpdateContext& Context)
 				FStopwatch Stopwatch;
 				Stopwatch.Start();
 				{
-					FTimingEventsTrackDrawStateBuilder Builder(*FilteredDrawState, Context.GetViewport());
+					FTimingEventsTrackDrawStateBuilder Builder(*FilteredDrawState, Context.GetViewport(), Context.GetGeometry().Scale);
 					BuildFilteredDrawState(Builder, Context);
 					Builder.Flush();
 				}
diff --git a/Engine/Source/Developer/TraceInsights/Private/Insights/ViewModels/TimingGraphTrack.cpp b/Engine/Source/Developer/TraceInsights/Private/Insights/ViewModels/TimingGraphTrack.cpp
--- a/Engine/Source/Developer/TraceInsights/Private/Insights/ViewModels/TimingGraphTrack.cpp
+++ b/Engine/Source/Developer/TraceInsights/Private/Insights/ViewModels/TimingGraphTrack.cpp
@@ -588,6 +588,7 @@ void FTimingGraphTrack::DrawVerticalAxisGrid(const ITimingTrackDrawContext& Cont
 	const FSlateBrush* Brush = Context.GetHelper().GetWhiteBrush();
 	//const FSlateFontInfo& Font = Context.GetHelper().GetEventFont();
 	const TSharedRef<FSlateFontMeasure> FontMeasureService = FSlateApplication::Get().GetRenderer()->GetFontMeasureService();
+	const float FontScale = DrawContext.Geometry.Scale;
 
 	const double TopValue = ViewportY.GetValueAtOffset(RoundedViewHeight);
 	const double GridValue = ViewportY.GetValueAtOffset(MinDY);
@@ -652,7 +653,7 @@ void FTimingGraphTrack::DrawVerticalAxisGrid(const ITimingTrackDrawContext& Cont
 			// Draw horizontal grid line.
 			DrawContext.DrawBox(0, Y, ViewWidth, 1, Brush, GridColor);
 
-			const FVector2D LabelTextSize = FontMeasureService->Measure(LabelText, Font);
+			const FVector2D LabelTextSize = FontMeasureService->Measure(LabelText, Font, FontScale) / FontScale;
 			const float LabelX = X0 - LabelTextSize.X - 4.0f;
 			const float LabelY = FMath::Min(Y0 + GetHeight() - TextH, FMath::Max(Y0, Y - TextH / 2));
 
diff --git a/Engine/Source/Developer/TraceInsights/Private/Insights/ViewModels/TimingViewDrawHelper.cpp b/Engine/Source/Developer/TraceInsights/Private/Insights/ViewModels/TimingViewDrawHelper.cpp
--- a/Engine/Source/Developer/TraceInsights/Private/Insights/ViewModels/TimingViewDrawHelper.cpp
+++ b/Engine/Source/Developer/TraceInsights/Private/Insights/ViewModels/TimingViewDrawHelper.cpp
@@ -24,13 +24,14 @@
 // FTimingEventsTrackDrawStateBuilder
 ////////////////////////////////////////////////////////////////////////////////////////////////////
 
-FTimingEventsTrackDrawStateBuilder::FTimingEventsTrackDrawStateBuilder(FTimingEventsTrackDrawState& InState, const FTimingTrackViewport& InViewport)
+FTimingEventsTrackDrawStateBuilder::FTimingEventsTrackDrawStateBuilder(FTimingEventsTrackDrawState& InState, const FTimingTrackViewport& InViewport, float InFontScale)
 	: DrawState(InState)
 	, Viewport(InViewport)
 	, MaxDepth(-1)
 	, LastEventX2()
 	, LastBox()
 	, EventFont(FAppStyle::Get().GetFontStyle("SmallFont"))
+	, FontScale(InFontScale)
 {
 	DrawState.Reset();
 }
@@ -213,7 +214,8 @@ void FTimingEventsTrackDrawStateBuilder::AddEvent(double EventStartTime, double
 		const FString Name = GetEventNameCallback(EventW);
 
 		const TSharedRef<FSlateFontMeasure> FontMeasureService = FSlateApplication::Get().GetRenderer()->GetFontMeasureService();
-		const int32 LastWholeCharacterIndex = FontMeasureService->FindLastWholeCharacterIndexBeforeOffset(Name, EventFont, FMath::RoundToInt(EventW - 2.0f));
+		const int32 HorizontalOffset = FMath::RoundToInt((EventW - 2.0f) * FontScale);
+		const int32 LastWholeCharacterIndex = FontMeasureService->FindLastWholeCharacterIndexBeforeOffset(Name, EventFont, HorizontalOffset, FontScale);
 
 		if (LastWholeCharacterIndex >= 0)
 		{
@@ -703,7 +705,8 @@ void FTimingViewDrawHelper::DrawTrackHeader(const FBaseTimingTrack& Track, const
 		if (TrackH > 4.0f)
 		{
 			const TSharedRef<FSlateFontMeasure> FontMeasureService = FSlateApplication::Get().GetRenderer()->GetFontMeasureService();
-			float TextWidth = FontMeasureService->Measure(Track.GetName(), EventFont).X;
+			const float FontScale = DrawContext.Geometry.Scale;
+			float TextWidth = FontMeasureService->Measure(Track.GetName(), EventFont, FontScale).X / FontScale;
 
 			constexpr float PinWidth = 8.0f;
 			if (Track.IsSelected())
diff --git a/Engine/Source/Developer/TraceInsights/Private/Insights/ViewModels/TimingViewDrawHelper.h b/Engine/Source/Developer/TraceInsights/Private/Insights/ViewModels/TimingViewDrawHelper.h
--- a/Engine/Source/Developer/TraceInsights/Private/Insights/ViewModels/TimingViewDrawHelper.h
+++ b/Engine/Source/Developer/TraceInsights/Private/Insights/ViewModels/TimingViewDrawHelper.h
@@ -102,7 +102,7 @@ private:
 	};
 
 public:
-	explicit FTimingEventsTrackDrawStateBuilder(FTimingEventsTrackDrawState& InState, const FTimingTrackViewport& InViewport);
+	explicit FTimingEventsTrackDrawStateBuilder(FTimingEventsTrackDrawState& InState, const FTimingTrackViewport& InViewport, float InFontScale);
 	virtual ~FTimingEventsTrackDrawStateBuilder() {}
 
 	/**
@@ -134,6 +134,7 @@ private:
 	TArray<FBoxData> LastBox;
 
 	const FSlateFontInfo EventFont;
+	float FontScale;
 };
 
 ////////////////////////////////////////////////////////////////////////////////////////////////////
diff --git a/Engine/Source/Developer/TraceInsights/Private/Insights/ViewModels/TooltipDrawState.cpp b/Engine/Source/Developer/TraceInsights/Private/Insights/ViewModels/TooltipDrawState.cpp
--- a/Engine/Source/Developer/TraceInsights/Private/Insights/ViewModels/TooltipDrawState.cpp
+++ b/Engine/Source/Developer/TraceInsights/Private/Insights/ViewModels/TooltipDrawState.cpp
@@ -29,6 +29,7 @@ FTooltipDrawState::FTooltipDrawState()
 	, NewLineY(0.0f)
 	, Opacity(0.0f)
 	, DesiredOpacity(0.0f)
+	, FontScale(1.0f)
 	, Texts()
 {
 }
@@ -65,7 +66,7 @@ void FTooltipDrawState::AddTitle(const FString& Title)
 {
 	const TSharedRef<FSlateFontMeasure> FontMeasureService = FSlateApplication::Get().GetRenderer()->GetFontMeasureService();
 
-	const FVector2D TextSize = FontMeasureService->Measure(Title, Font);
+	const FVector2D TextSize = FontMeasureService->Measure(Title, Font, FontScale) / FontScale;
 	Texts.Add({ BorderX, NewLineY, TextSize, Title, DefaultTitleColor, FDrawTextType::Title });
 
 	NewLineY += DefaultTitleHeight;
@@ -77,7 +78,7 @@ void FTooltipDrawState::AddTitle(const FString& Title, const FLinearColor& Color
 {
 	const TSharedRef<FSlateFontMeasure> FontMeasureService = FSlateApplication::Get().GetRenderer()->GetFontMeasureService();
 
-	const FVector2D TextSize = FontMeasureService->Measure(Title, Font);
+	const FVector2D TextSize = FontMeasureService->Measure(Title, Font, FontScale) / FontScale;
 	Texts.Add({ BorderX, NewLineY, TextSize, Title, Color, FDrawTextType::Misc });
 
 	NewLineY += DefaultTitleHeight;
@@ -89,10 +90,10 @@ void FTooltipDrawState::AddNameValueTextLine(const FString& Name, const FString&
 {
 	const TSharedRef<FSlateFontMeasure> FontMeasureService = FSlateApplication::Get().GetRenderer()->GetFontMeasureService();
 
-	const FVector2D NameTextSize = FontMeasureService->Measure(Name, Font);
+	const FVector2D NameTextSize = FontMeasureService->Measure(Name, Font, FontScale) / FontScale;
 	Texts.Add({ 0.0f, NewLineY, NameTextSize, Name, DefaultNameColor, FDrawTextType::Name });
 
-	const FVector2D ValueTextSize = FontMeasureService->Measure(Value, Font);
+	const FVector2D ValueTextSize = FontMeasureService->Measure(Value, Font, FontScale) / FontScale;
 	Texts.Add({ 0.0f, NewLineY, ValueTextSize, Value, DefaultValueColor, FDrawTextType::Value });
 
 	NewLineY += DefaultLineHeight;
@@ -104,7 +105,7 @@ void FTooltipDrawState::AddTextLine(const FString& Text, const FLinearColor& Col
 {
 	const TSharedRef<FSlateFontMeasure> FontMeasureService = FSlateApplication::Get().GetRenderer()->GetFontMeasureService();
 
-	const FVector2D TextSize = FontMeasureService->Measure(Text, Font);
+	const FVector2D TextSize = FontMeasureService->Measure(Text, Font, FontScale) / FontScale;
 	Texts.Add({ BorderX, NewLineY, TextSize, Text, Color, FDrawTextType::Misc });
 
 	NewLineY += DefaultLineHeight;
@@ -116,7 +117,7 @@ void FTooltipDrawState::AddTextLine(const float X, const float Y, const FString&
 {
 	const TSharedRef<FSlateFontMeasure> FontMeasureService = FSlateApplication::Get().GetRenderer()->GetFontMeasureService();
 
-	const FVector2D TextSize = FontMeasureService->Measure(Text, Font);
+	const FVector2D TextSize = FontMeasureService->Measure(Text, Font, FontScale) / FontScale;
 	Texts.Add({ X, Y, TextSize, Text, Color, FDrawTextType::Misc });
 }
 
diff --git a/Engine/Source/Developer/TraceInsights/Private/Insights/ViewModels/TrackHeader.cpp b/Engine/Source/Developer/TraceInsights/Private/Insights/ViewModels/TrackHeader.cpp
--- a/Engine/Source/Developer/TraceInsights/Private/Insights/ViewModels/TrackHeader.cpp
+++ b/Engine/Source/Developer/TraceInsights/Private/Insights/ViewModels/TrackHeader.cpp
@@ -24,6 +24,7 @@ FTrackHeader::FTrackHeader(FBaseTimingTrack& InParentTrack)
 	: ParentTrack(InParentTrack)
 	, WhiteBrush(FInsightsStyle::Get().GetBrush("WhiteBrush"))
 	, Font(FAppStyle::Get().GetFontStyle("SmallFont"))
+	, FontScale(1.0f)
 {
 	Reset();
 }
@@ -45,19 +46,14 @@ void FTrackHeader::Reset()
 
 void FTrackHeader::UpdateSize()
 {
-	//ParentTrack.GetName()
 	const TSharedRef<FSlateFontMeasure> FontMeasureService = FSlateApplication::Get().GetRenderer()->GetFontMeasureService();
-	const float TextWidth = FontMeasureService->Measure(ParentTrack.GetName(), Font).X;
+	const float TextWidth = FontMeasureService->Measure(ParentTrack.GetName(), Font, FontScale).X / FontScale;
 
 	Width = TextWidth + 4.0f;
 	if (bCanBeCollapsed)
 	{
 		Width += 9.0f;
 	}
-	//if (ParentTrack.IsSelected())
-	//{
-	//	Width += 9.0f;
-	//}
 }
 
 ////////////////////////////////////////////////////////////////////////////////////////////////////
diff --git a/Engine/Source/Developer/TraceInsights/Private/Insights/ViewModels/TrackHeader.h b/Engine/Source/Developer/TraceInsights/Private/Insights/ViewModels/TrackHeader.h
--- a/Engine/Source/Developer/TraceInsights/Private/Insights/ViewModels/TrackHeader.h
+++ b/Engine/Source/Developer/TraceInsights/Private/Insights/ViewModels/TrackHeader.h
@@ -36,6 +36,9 @@ public:
 	void Draw(const ITimingTrackDrawContext& Context) const;
 	void PostDraw(const ITimingTrackDrawContext& Context) const;
 
+	float GetFontScale() const { return FontScale; }
+	void SetFontScale(float InFontScale) { FontScale = InFontScale; }
+
 private:
 	void DrawInternal(const ITimingTrackDrawContext& Context, bool bDrawBackgroundLayer) const;
 
@@ -56,4 +59,5 @@ private:
 	// Slate resources
 	const FSlateBrush* WhiteBrush;
 	const FSlateFontInfo Font;
+	float FontScale;
 };
diff --git a/Engine/Source/Developer/TraceInsights/Private/Insights/Widgets/SFrameTrack.cpp b/Engine/Source/Developer/TraceInsights/Private/Insights/Widgets/SFrameTrack.cpp
--- a/Engine/Source/Developer/TraceInsights/Private/Insights/Widgets/SFrameTrack.cpp
+++ b/Engine/Source/Developer/TraceInsights/Private/Insights/Widgets/SFrameTrack.cpp
@@ -93,8 +93,8 @@ void SFrameTrack::Reset()
 	bDrawVerticalAxisLabelsOnLeftSide = false;
 
 	HoveredSample.Reset();
-	TooltipDesiredOpacity = 1.0f;
 	TooltipOpacity = 0.0f;
+	TooltipSizeX = 70.0f;
 
 	//ThisGeometry
 
@@ -489,13 +489,16 @@ int32 SFrameTrack::OnPaint(const FPaintArgs& Args, const FGeometry& AllottedGeom
 		// Draw tooltip for hovered sample (frame).
 		if (HoveredSample.IsValid())
 		{
+			constexpr float TooltipDesiredOpacity = 1.0f;
 			if (TooltipOpacity < TooltipDesiredOpacity)
 			{
+				// slow fade in
 				TooltipOpacity = TooltipOpacity * 0.9f + TooltipDesiredOpacity * 0.1f;
 			}
 			else
 			{
-				TooltipOpacity = TooltipDesiredOpacity;
+				// fast fade out
+				TooltipOpacity = TooltipOpacity * 0.75f + TooltipDesiredOpacity * 0.25f;
 			}
 
 			// First line: "Rendering Frame 1,234"
@@ -513,8 +516,9 @@ int32 SFrameTrack::OnPaint(const FPaintArgs& Args, const FGeometry& AllottedGeom
 			StringBuilder.Appendf(TEXT(" (%.1f fps)"), 1.0 / HoveredSample.Sample->LargestFrameDuration);
 			const FString Text2(StringBuilder);
 
-			const FVector2D TextSize1 = FontMeasureService->Measure(Text1, SummaryFont);
-			const FVector2D TextSize2 = FontMeasureService->Measure(Text2, SummaryFont);
+			const float FontScale = DrawContext.Geometry.Scale;
+			const FVector2D TextSize1 = FontMeasureService->Measure(Text1, SummaryFont, FontScale) / FontScale;
+			const FVector2D TextSize2 = FontMeasureService->Measure(Text2, SummaryFont, FontScale) / FontScale;
 
 			const FAxisViewportInt32& ViewportX = Viewport.GetHorizontalAxisViewport();
 
@@ -524,16 +528,26 @@ int32 SFrameTrack::OnPaint(const FPaintArgs& Args, const FGeometry& AllottedGeom
 			constexpr float DX = 3.0f;
 			const float DX1 = FMath::RoundToFloat(TextSize1.X / 2);
 			const float DX2 = FMath::RoundToFloat(TextSize2.X / 2);
-			const float BoxDX = FMath::Max(DX1, DX2) + DX;
+			const float TooltipDesiredSizeX = FMath::Max(DX1, DX2) + DX;
+
+			if (TooltipSizeX != TooltipDesiredSizeX)
+			{
+				TooltipSizeX = TooltipSizeX * 0.75f + TooltipDesiredSizeX * 0.25f;
+
+				if (FMath::IsNearlyEqual(TooltipSizeX, TooltipDesiredSizeX))
+				{
+					TooltipSizeX = TooltipDesiredSizeX;
+				}
+			}
 
 			float CX = CX0;
-			if (CX > ViewportX.GetSize() - BoxDX)
+			if (CX > ViewportX.GetSize() - TooltipSizeX)
 			{
-				CX = FMath::RoundToFloat(ViewportX.GetSize() - BoxDX);
+				CX = FMath::RoundToFloat(ViewportX.GetSize() - TooltipSizeX);
 			}
-			if (CX - BoxDX < 0)
+			if (CX - TooltipSizeX < 0)
 			{
-				CX = BoxDX;
+				CX = TooltipSizeX;
 			}
 
 			constexpr float BoxY = 11.0f;
@@ -541,7 +555,7 @@ int32 SFrameTrack::OnPaint(const FPaintArgs& Args, const FGeometry& AllottedGeom
 			constexpr float LineDY = 12.0f;
 
 			const FLinearColor BackgroundColor(0.9f, 0.9f, 0.9f, TooltipOpacity);
-			DrawContext.DrawBox(CX - BoxDX, BoxY, 2 * BoxDX, BoxH, WhiteBrush, BackgroundColor);
+			DrawContext.DrawBox(CX - TooltipSizeX, BoxY, 2 * TooltipSizeX, BoxH, WhiteBrush, BackgroundColor);
 			const int32 ArrowSize = 4;
 			for (int32 ArrowY = 0; ArrowY < ArrowSize; ++ArrowY)
 			{
@@ -561,6 +575,10 @@ int32 SFrameTrack::OnPaint(const FPaintArgs& Args, const FGeometry& AllottedGeom
 			DrawContext.DrawText(CX - DX2, BoxY + LineDY + 1.0f, Text2, SummaryFont, TextColor2);
 			DrawContext.LayerId++;
 		}
+		else
+		{
+			TooltipOpacity = 0.0f;
+		}
 
 		Stopwatch.Stop();
 		DrawDurationHistory.AddValue(Stopwatch.AccumulatedTime);
@@ -570,7 +588,8 @@ int32 SFrameTrack::OnPaint(const FPaintArgs& Args, const FGeometry& AllottedGeom
 	const bool bShouldDisplayDebugInfo = FInsightsManager::Get()->IsDebugInfoEnabled();
 	if (bShouldDisplayDebugInfo)
 	{
-		const float MaxFontCharHeight = FontMeasureService->Measure(TEXT("!"), SummaryFont).Y;
+		const float FontScale = DrawContext.Geometry.Scale;
+		const float MaxFontCharHeight = FontMeasureService->Measure(TEXT("!"), SummaryFont, FontScale).Y / FontScale;
 		const float DbgDY = MaxFontCharHeight;
 
 		const float DbgW = 280.0f;
@@ -657,6 +676,7 @@ void SFrameTrack::DrawVerticalAxisGrid(FDrawContext& DrawContext, const FSlateBr
 	const FLinearColor TextBgColor(0.05f, 0.05f, 0.05f, 1.0f);
 
 	const TSharedRef<FSlateFontMeasure> FontMeasureService = FSlateApplication::Get().GetRenderer()->GetFontMeasureService();
+	const float FontScale = DrawContext.Geometry.Scale;
 
 	// Available axis, pre-ordered by value.
 	struct FAxis
@@ -774,7 +794,7 @@ void SFrameTrack::DrawVerticalAxisGrid(FDrawContext& DrawContext, const FSlateBr
 								  (Axis.Value <= 1.0) ? FString::Printf(TEXT("%s (%.0f fps)"), *TimeUtils::FormatTimeAuto(Axis.Value), 1.0 / Axis.Value) :
 														TimeUtils::FormatTimeAuto(Axis.Value);
 
-		const FVector2D LabelTextSize = FontMeasureService->Measure(LabelText, Font);
+		const FVector2D LabelTextSize = FontMeasureService->Measure(LabelText, Font, FontScale) / FontScale;
 		const float LabelX = bDrawVerticalAxisLabelsOnLeftSide ? 0.0f : ViewWidth - LabelTextSize.X - 4.0f;
 
 		// Draw background for value text.
@@ -845,6 +865,7 @@ void SFrameTrack::DrawHorizontalAxisGrid(FDrawContext& DrawContext, const FSlate
 		else
 		{
 			const TSharedRef<FSlateFontMeasure> FontMeasureService = FSlateApplication::Get().GetRenderer()->GetFontMeasureService();
+			const float FontScale = DrawContext.Geometry.Scale;
 
 			// Draw labels.
 			const FLinearColor LabelBoxColor(0.05f, 0.05f, 0.05f, 1.0f);
@@ -853,7 +874,7 @@ void SFrameTrack::DrawHorizontalAxisGrid(FDrawContext& DrawContext, const FSlate
 			{
 				const float X = FMath::RoundToFloat(ViewportX.GetOffsetForValue(Index));
 				const FString LabelText = FText::AsNumber(Index).ToString();
-				const FVector2D LabelTextSize = FontMeasureService->Measure(LabelText, Font);
+				const FVector2D LabelTextSize = FontMeasureService->Measure(LabelText, Font, FontScale) / FontScale;
 				DrawContext.DrawBox(X, 10.0f, LabelTextSize.X + 4.0f, 12.0f, Brush, LabelBoxColor);
 				DrawContext.DrawText(X + 2.0f, 10.0f, LabelText, Font, LabelTextColor);
 			}
diff --git a/Engine/Source/Developer/TraceInsights/Private/Insights/Widgets/SFrameTrack.h b/Engine/Source/Developer/TraceInsights/Private/Insights/Widgets/SFrameTrack.h
--- a/Engine/Source/Developer/TraceInsights/Private/Insights/Widgets/SFrameTrack.h
+++ b/Engine/Source/Developer/TraceInsights/Private/Insights/Widgets/SFrameTrack.h
@@ -224,8 +224,8 @@ protected:
 
 	FFrameTrackSampleRef HoveredSample;
 
-	float TooltipDesiredOpacity;
 	mutable float TooltipOpacity;
+	mutable float TooltipSizeX;
 
 	//////////////////////////////////////////////////
 	// Misc
diff --git a/Engine/Source/Developer/TraceInsights/Private/Insights/Widgets/STimingView.cpp b/Engine/Source/Developer/TraceInsights/Private/Insights/Widgets/STimingView.cpp
--- a/Engine/Source/Developer/TraceInsights/Private/Insights/Widgets/STimingView.cpp
+++ b/Engine/Source/Developer/TraceInsights/Private/Insights/Widgets/STimingView.cpp
@@ -519,6 +519,8 @@ void STimingView::Tick(const FGeometry& AllottedGeometry, const double InCurrent
 
 	ThisGeometry = AllottedGeometry;
 
+	Tooltip.SetFontScale(AllottedGeometry.Scale);
+
 	bPreventThrottling = false;
 
 	constexpr float OverscrollFadeSpeed = 2.0f;
@@ -678,8 +680,14 @@ void STimingView::Tick(const FGeometry& AllottedGeometry, const double InCurrent
 	class FTimingTrackUpdateContext : public ITimingTrackUpdateContext
 	{
 	public:
-		explicit FTimingTrackUpdateContext(STimingView* InTimingView, double InCurrentTime, float InDeltaTime) : TimingView(InTimingView), CurrentTime(InCurrentTime), DeltaTime(InDeltaTime) {}
+		explicit FTimingTrackUpdateContext(STimingView* InTimingView, const FGeometry& InGeometry, double InCurrentTime, float InDeltaTime)
+			: TimingView(InTimingView)
+			, Geometry(InGeometry)
+			, CurrentTime(InCurrentTime)
+			, DeltaTime(InDeltaTime)		
+		{}
 
+		virtual const FGeometry& GetGeometry() const override { return Geometry; }
 		virtual const FTimingTrackViewport& GetViewport() const override { return TimingView->GetViewport(); }
 		virtual const FVector2D& GetMousePosition() const override { return TimingView->GetMousePosition(); }
 		virtual const TSharedPtr<const ITimingEvent> GetHoveredEvent() const override { return TimingView->GetHoveredEvent(); }
@@ -691,11 +699,12 @@ void STimingView::Tick(const FGeometry& AllottedGeometry, const double InCurrent
 
 	public:
 		STimingView* TimingView;
+		const FGeometry& Geometry;
 		double CurrentTime;
 		float DeltaTime;
 	};
 
-	FTimingTrackUpdateContext UpdateContext(this, InCurrentTime, InDeltaTime);
+	FTimingTrackUpdateContext UpdateContext(this, AllottedGeometry, InCurrentTime, InDeltaTime);
 
 	////////////////////////////////////////////////////////////////////////////////////////////////////
 	// Pre-Update.
@@ -1109,6 +1118,7 @@ int32 STimingView::OnPaint(const FPaintArgs& Args, const FGeometry& AllottedGeom
 #endif
 
 	const TSharedRef<FSlateFontMeasure> FontMeasureService = FSlateApplication::Get().GetRenderer()->GetFontMeasureService();
+	const float FontScale = AllottedGeometry.Scale;
 
 	const float ViewWidth = AllottedGeometry.GetLocalSize().X;
 	const float ViewHeight = AllottedGeometry.GetLocalSize().Y;
@@ -1508,7 +1518,7 @@ int32 STimingView::OnPaint(const FPaintArgs& Args, const FGeometry& AllottedGeom
 	{
 		const FSlateFontInfo& SummaryFont = MainFont;
 
-		const float MaxFontCharHeight = FontMeasureService->Measure(TEXT("!"), SummaryFont).Y;
+		const float MaxFontCharHeight = FontMeasureService->Measure(TEXT("!"), SummaryFont, FontScale).Y / FontScale;
 		const float DbgDY = MaxFontCharHeight;
 
 		const float DbgW = 320.0f;
@@ -4887,7 +4897,9 @@ TSharedRef<SDockTab> STimingView::SpawnQuickFindTab(const FSpawnTabArgs& Args)
 	const TSharedPtr<SWindow>& OwnerWindow = Args.GetOwnerWindow();
 	if (OwnerWindow.IsValid() && OwnerWindow != FSlateApplication::Get().FindWidgetWindow(SharedThis(this)))
 	{
-		const float LocalDPIScaleFactor = FPlatformApplicationMisc::GetDPIScaleFactorAtPoint(10.0f, 10.0f);
+		TSharedPtr<SWindow> TopmostAncestor = OwnerWindow->GetTopmostAncestor();
+		const FVector2D DPIProbePoint = TopmostAncestor->GetPositionInScreen() + FVector2D(10.0f, 10.0f);
+		const float LocalDPIScaleFactor = FPlatformApplicationMisc::GetDPIScaleFactorAtPoint(DPIProbePoint.X, DPIProbePoint.Y);
 		OwnerWindow->Resize(FVector2D(600 * LocalDPIScaleFactor, 400 * LocalDPIScaleFactor));
 	}
 
diff --git a/Engine/Source/Developer/TraceInsights/Public/Insights/Common/PaintUtils.h b/Engine/Source/Developer/TraceInsights/Public/Insights/Common/PaintUtils.h
--- a/Engine/Source/Developer/TraceInsights/Public/Insights/Common/PaintUtils.h
+++ b/Engine/Source/Developer/TraceInsights/Public/Insights/Common/PaintUtils.h
@@ -78,16 +78,20 @@ struct TRACEINSIGHTS_API FDrawContext
 
 	inline void DrawTextAligned(EHorizontalAlignment HAlign, const float X, const float Y, const FString& Text, const FSlateFontInfo& Font, const FLinearColor& Color) const
 	{
-		const TSharedRef<FSlateFontMeasure> FontMeasureService = FSlateApplication::Get().GetRenderer()->GetFontMeasureService();
-		const float TextWidth = FontMeasureService->Measure(Text, Font).X;
 		float TextX = X;
-		if (HAlign == HAlign_Right)
+		if (HAlign == HAlign_Right || HAlign == HAlign_Center)
 		{
-			TextX -= TextWidth;
-		}
-		else if (HAlign == HAlign_Center)
-		{
-			TextX -= TextWidth / 2;
+			const TSharedRef<FSlateFontMeasure> FontMeasureService = FSlateApplication::Get().GetRenderer()->GetFontMeasureService();
+			const float FontScale = Geometry.Scale;
+			const float TextWidth = FontMeasureService->Measure(Text, Font, FontScale).X / FontScale;
+			if (HAlign == HAlign_Right)
+			{
+				TextX -= TextWidth;
+			}
+			else
+			{
+				TextX -= TextWidth / 2;
+			}
 		}
 		FSlateDrawElement::MakeText(ElementList, LayerId, MAKE_PAINT_GEOMETRY_PT(Geometry, TextX, Y), Text, Font, DrawEffects, Color);
 	}
diff --git a/Engine/Source/Developer/TraceInsights/Public/Insights/ViewModels/BaseTimingTrack.h b/Engine/Source/Developer/TraceInsights/Public/Insights/ViewModels/BaseTimingTrack.h
--- a/Engine/Source/Developer/TraceInsights/Public/Insights/ViewModels/BaseTimingTrack.h
+++ b/Engine/Source/Developer/TraceInsights/Public/Insights/ViewModels/BaseTimingTrack.h
@@ -24,6 +24,7 @@ class ITimingEventRelation;
 class ITimingEventFilter;
 
 struct FDrawContext;
+struct FGeometry;
 
 ////////////////////////////////////////////////////////////////////////////////////////////////////
 
@@ -91,6 +92,7 @@ ENUM_CLASS_FLAGS(EDrawEventMode);
 class TRACEINSIGHTS_API ITimingTrackUpdateContext
 {
 public:
+	virtual const FGeometry& GetGeometry() const = 0;
 	virtual const FTimingTrackViewport& GetViewport() const = 0;
 	virtual const FVector2D& GetMousePosition() const = 0;
 	virtual const TSharedPtr<const ITimingEvent> GetHoveredEvent() const = 0;
diff --git a/Engine/Source/Developer/TraceInsights/Public/Insights/ViewModels/TooltipDrawState.h b/Engine/Source/Developer/TraceInsights/Public/Insights/ViewModels/TooltipDrawState.h
--- a/Engine/Source/Developer/TraceInsights/Public/Insights/ViewModels/TooltipDrawState.h
+++ b/Engine/Source/Developer/TraceInsights/Public/Insights/ViewModels/TooltipDrawState.h
@@ -82,6 +82,9 @@ public:
 
 	void Draw(const FDrawContext& DrawContext) const;
 
+	void SetFontScale(float InFontScale) { FontScale = InFontScale; }
+	float GetFontScale() const { return FontScale; }
+
 private:
 	const FSlateBrush* WhiteBrush;
 	const FSlateFontInfo Font;
@@ -99,6 +102,8 @@ private:
 	float Opacity;
 	float DesiredOpacity;
 
+	float FontScale;
+
 	TArray<FDrawTextInfo> Texts;
 };
 
