[Insights] Moved the TaskGraphRelation source file into TaskGraphProfiler folder and fixed includes.

diff --git a/Engine/Source/Developer/TraceInsights/Private/Insights/TaskGraphProfiler/TaskGraphProfilerManager.cpp b/Engine/Source/Developer/TraceInsights/Private/Insights/TaskGraphProfiler/TaskGraphProfilerManager.cpp
--- a/Engine/Source/Developer/TraceInsights/Private/Insights/TaskGraphProfiler/TaskGraphProfilerManager.cpp
+++ b/Engine/Source/Developer/TraceInsights/Private/Insights/TaskGraphProfiler/TaskGraphProfilerManager.cpp
@@ -2,23 +2,24 @@
 
 #include "TaskGraphProfilerManager.h"
 
+#include "Async/TaskGraphInterfaces.h"
 #include "Features/IModularFeatures.h"
+#include "Framework/Docking/TabManager.h"
 #include "Modules/ModuleManager.h"
 #include "TraceServices/AnalysisService.h"
 #include "TraceServices/Model/TasksProfiler.h"
-#include "Async/TaskGraphInterfaces.h"
 
+// Insights
 #include "Insights/InsightsStyle.h"
+#include "Insights/TaskGraphProfiler/ViewModels/TaskGraphRelation.h"
 #include "Insights/TaskGraphProfiler/ViewModels/TaskTable.h"
 #include "Insights/TaskGraphProfiler/ViewModels/TaskTimingTrack.h"
 #include "Insights/TaskGraphProfiler/Widgets/STaskTableTreeView.h"
 #include "Insights/TimingProfilerManager.h"
-#include "Insights/ViewModels/TaskGraphRelation.h"
 #include "Insights/ViewModels/ThreadTimingTrack.h"
 #include "Insights/ViewModels/ThreadTrackEvent.h"
 #include "Insights/Widgets/STimingProfilerWindow.h"
 #include "Insights/Widgets/STimingView.h"
-#include "Framework/Docking/TabManager.h"
 
 ////////////////////////////////////////////////////////////////////////////////////////////////////
 
diff --git a/Engine/Source/Developer/TraceInsights/Private/Insights/TaskGraphProfiler/ViewModels/TaskGraphRelation.cpp b/Engine/Source/Developer/TraceInsights/Private/Insights/TaskGraphProfiler/ViewModels/TaskGraphRelation.cpp
new file mode 100644
--- /dev/null
+++ b/Engine/Source/Developer/TraceInsights/Private/Insights/TaskGraphProfiler/ViewModels/TaskGraphRelation.cpp
@@ -0,0 +1,163 @@
+// Copyright Epic Games, Inc. All Rights Reserved.
+
+#include "TaskGraphRelation.h"
+
+// Insights
+#include "Insights/Common/PaintUtils.h"
+#include "Insights/TaskGraphProfiler/TaskGraphProfilerManager.h"
+#include "Insights/ViewModels/TimingTrackViewport.h"
+#include "Insights/ViewModels/TimingViewDrawHelper.h"
+
+////////////////////////////////////////////////////////////////////////////////////////////////////
+
+namespace Insights
+{
+
+INSIGHTS_IMPLEMENT_RTTI(FTaskGraphRelation)
+
+////////////////////////////////////////////////////////////////////////////////////////////////////
+
+FTaskGraphRelation::FTaskGraphRelation(double InSourceTime, int32 InSourceThreadId, double InTargetTime, int32 InTargetThreadId, ETaskEventType InType)
+{
+	SourceTime = InSourceTime;
+	SourceThreadId = InSourceThreadId;
+	TargetTime = InTargetTime;
+	TargetThreadId = InTargetThreadId;
+	Type = InType;
+}
+
+////////////////////////////////////////////////////////////////////////////////////////////////////
+
+void FTaskGraphRelation::Draw(const FDrawContext& DrawContext, const FTimingTrackViewport& Viewport, const ITimingViewDrawHelper& Helper, const ITimingEventRelation::EDrawFilter Filter)
+{
+	int32 LayerId = Helper.GetRelationLayerId();
+
+	TSharedPtr<const FBaseTimingTrack> SourceTrackShared = SourceTrack.Pin();
+	TSharedPtr<const FBaseTimingTrack> TargetTrackShared = TargetTrack.Pin();
+
+	if (!SourceTrackShared.IsValid() || !TargetTrackShared.IsValid())
+	{
+		return;
+	}
+
+	if (Filter == ITimingEventRelation::EDrawFilter::BetweenScrollableTracks)
+	{
+		if (SourceTrackShared->GetLocation() != ETimingTrackLocation::Scrollable ||
+			TargetTrackShared->GetLocation() != ETimingTrackLocation::Scrollable)
+		{
+			return;
+		}
+	}
+
+	if (Filter == ITimingEventRelation::EDrawFilter::BetweenDockedTracks)
+	{
+		if (SourceTrackShared->GetLocation() == ETimingTrackLocation::Scrollable &&
+			TargetTrackShared->GetLocation() == ETimingTrackLocation::Scrollable)
+		{
+			return;
+		}
+
+		LayerId = DrawContext.LayerId;
+	}
+
+	const int32 OutlineLayerId = LayerId - 1;
+
+	float X1 = Viewport.TimeToSlateUnitsRounded(SourceTime);
+	float X2 = Viewport.TimeToSlateUnitsRounded(TargetTime);
+	if (FMath::Max(X1, X2) < 0.0f || FMath::Min(X1, X2) > Viewport.GetWidth())
+	{
+		return;
+	}
+
+	float Y1 = SourceTrackShared->GetPosY();
+	Y1 += Viewport.GetLayout().GetLaneY(SourceDepth) + Viewport.GetLayout().EventH / 2.0f;
+	if (SourceTrackShared->GetChildTrack() && SourceTrackShared->GetChildTrack()->GetHeight() > 0.0f)
+	{
+		Y1 += SourceTrackShared->GetChildTrack()->GetHeight() + Viewport.GetLayout().ChildTimelineDY;
+	}
+
+	float Y2 = TargetTrackShared->GetPosY();
+	Y2 += Viewport.GetLayout().GetLaneY(TargetDepth) + Viewport.GetLayout().EventH / 2.0f;
+	if (TargetTrackShared->GetChildTrack() && TargetTrackShared->GetChildTrack()->GetHeight() > 0.0f)
+	{
+		Y2 += TargetTrackShared->GetChildTrack()->GetHeight() + Viewport.GetLayout().ChildTimelineDY;
+	}
+
+	const FVector2D StartPoint = FVector2D(X1, Y1);
+	const FVector2D EndPoint = FVector2D(X2, Y2);
+	const float Distance = FVector2D::Distance(StartPoint, EndPoint);
+
+	constexpr float LineHeightAtStart = 4.0f;
+	constexpr float LineLengthAtStart = 4.0f;
+	constexpr float LineLengthAtEnd = 12.0f;
+
+	const FVector2D StartDir(FMath::Max(X2 - X1, 4.0f * (LineLengthAtStart + LineLengthAtEnd)), 0.0f);
+
+	constexpr float OutlineThickness = 5.0f;
+	constexpr float LineThickness = 3.0f;
+
+	constexpr float ArrowDirectionLen = 10.0f;
+	constexpr float ArrowRotationAngle = 20.0f;
+	FVector2D ArrowDirection(-ArrowDirectionLen, 0.0f);
+
+	const FLinearColor OutlineColor(0.0f, 0.0f, 0.0f, 1.0f);
+	const FLinearColor Color = FTaskGraphProfilerManager::Get()->GetColorForTaskEvent(Type);
+
+	TArray<FVector2D> LinePoints;
+	LinePoints.Add(StartPoint + FVector2D(0.0f, -LineHeightAtStart / 2.0f));
+	LinePoints.Add(StartPoint + FVector2D(0.0f, +LineHeightAtStart / 2.0f));
+	DrawContext.DrawLines(OutlineLayerId, 0.0f, 0.0f, LinePoints, ESlateDrawEffect::None, OutlineColor, /*bAntialias=*/ true, OutlineThickness);
+	DrawContext.DrawLines(LayerId, 0.0f, 0.0f, LinePoints, ESlateDrawEffect::None, Color, /*bAntialias=*/ true, LineThickness);
+
+	constexpr float MinDistance = 1.5f * (LineLengthAtStart + LineLengthAtEnd);
+	constexpr float MaxDistance = 10000.0f; // arbitrary limit to avoid stack overflow in recursive FLineBuilder::Subdivide when rendering splines
+	if (Distance > MinDistance && Distance < MaxDistance && !FMath::IsNearlyEqual(StartPoint.Y, EndPoint.Y))
+	{
+		FVector2D SplineStart(StartPoint.X + LineLengthAtStart, StartPoint.Y);
+		FVector2D SplineEnd(EndPoint.X - LineLengthAtEnd, EndPoint.Y);
+		DrawContext.DrawSpline(OutlineLayerId, 0.0f, 0.0f, SplineStart, StartDir, SplineEnd, StartDir, OutlineThickness, OutlineColor);
+		DrawContext.DrawSpline(LayerId, 0.0f, 0.0f, SplineStart, StartDir, SplineEnd, StartDir, LineThickness, Color);
+
+		LinePoints.Empty();
+		LinePoints.Add(StartPoint);
+		LinePoints.Add(SplineStart);
+		DrawContext.DrawLines(OutlineLayerId, 0.0f, 0.0f, LinePoints, ESlateDrawEffect::None, OutlineColor, /*bAntialias=*/ true, OutlineThickness);
+		DrawContext.DrawLines(LayerId, 0.0f, 0.0f, LinePoints, ESlateDrawEffect::None, Color, /*bAntialias=*/ true, LineThickness);
+
+		LinePoints.Empty();
+		LinePoints.Add(SplineEnd);
+		LinePoints.Add(EndPoint);
+		DrawContext.DrawLines(OutlineLayerId, 0.0f, 0.0f, LinePoints, ESlateDrawEffect::None, OutlineColor, /*bAntialias=*/ true, OutlineThickness);
+		DrawContext.DrawLines(LayerId, 0.0f, 0.0f, LinePoints, ESlateDrawEffect::None, Color, /*bAntialias=*/ true, LineThickness);
+	}
+	else
+	{
+		LinePoints.Empty();
+		LinePoints.Add(StartPoint);
+		LinePoints.Add(EndPoint);
+		DrawContext.DrawLines(OutlineLayerId, 0.0f, 0.0f, LinePoints, ESlateDrawEffect::None, OutlineColor, /*bAntialias=*/ true, OutlineThickness);
+		DrawContext.DrawLines(LayerId, 0.0f, 0.0f, LinePoints, ESlateDrawEffect::None, Color, /*bAntialias=*/ true, LineThickness);
+
+		ArrowDirection = StartPoint - EndPoint;
+		ArrowDirection.Normalize();
+		ArrowDirection *= ArrowDirectionLen;
+	}
+
+	FVector2D ArrowOrigin = EndPoint;
+
+	LinePoints.Empty();
+	LinePoints.Add(ArrowOrigin);
+	LinePoints.Add(ArrowOrigin + ArrowDirection.GetRotated(-ArrowRotationAngle));
+	DrawContext.DrawLines(OutlineLayerId, 0.0f, 0.0f, LinePoints, ESlateDrawEffect::None, OutlineColor, /*bAntialias=*/ true, OutlineThickness);
+	DrawContext.DrawLines(LayerId, 0.0f, 0.0f, LinePoints, ESlateDrawEffect::None, Color, /*bAntialias=*/ true, LineThickness);
+
+	LinePoints.Empty();
+	LinePoints.Add(ArrowOrigin);
+	LinePoints.Add(ArrowOrigin + ArrowDirection.GetRotated(ArrowRotationAngle));
+	DrawContext.DrawLines(OutlineLayerId, 0.0f, 0.0f, LinePoints, ESlateDrawEffect::None, OutlineColor, /*bAntialias=*/ true, OutlineThickness);
+	DrawContext.DrawLines(LayerId, 0.0f, 0.0f, LinePoints, ESlateDrawEffect::None, Color, /*bAntialias=*/ true, LineThickness);
+}
+
+////////////////////////////////////////////////////////////////////////////////////////////////////
+
+} // namespace Insights
diff --git a/Engine/Source/Developer/TraceInsights/Private/Insights/TaskGraphProfiler/ViewModels/TaskGraphRelation.h b/Engine/Source/Developer/TraceInsights/Private/Insights/TaskGraphProfiler/ViewModels/TaskGraphRelation.h
new file mode 100644
--- /dev/null
+++ b/Engine/Source/Developer/TraceInsights/Private/Insights/TaskGraphProfiler/ViewModels/TaskGraphRelation.h
@@ -0,0 +1,54 @@
+// Copyright Epic Games, Inc. All Rights Reserved.
+
+#pragma once
+
+#include "CoreMinimal.h"
+
+#include "Insights/ViewModels/ITimingEvent.h"
+
+namespace Insights
+{
+
+enum class ETaskEventType : uint32;
+
+class FTaskGraphRelation : public ITimingEventRelation
+{
+	INSIGHTS_DECLARE_RTTI(FTaskGraphRelation, ITimingEventRelation)
+
+public:
+	FTaskGraphRelation(double InSourceTime, int32 InSourceThreadId, double InTargetTime, int32 InTargetThreadId, ETaskEventType InType);
+
+	virtual void Draw(const FDrawContext& DrawContext, const FTimingTrackViewport& Viewport, const ITimingViewDrawHelper& Helper, const ITimingEventRelation::EDrawFilter Filter) override;
+
+	void SetSourceTrack(TSharedPtr<const FBaseTimingTrack> InSourceTrack) { SourceTrack = InSourceTrack; }
+	TSharedPtr<const FBaseTimingTrack> GetSourceTrack() { return SourceTrack.Pin(); }
+
+	void SetTargetTrack(TSharedPtr<const FBaseTimingTrack> InTargetTrack) { TargetTrack = InTargetTrack; }
+	TSharedPtr<const FBaseTimingTrack> GetTargetTrack() { return TargetTrack.Pin(); }
+
+	double GetSourceTime() { return SourceTime; }
+	int32 GetSourceThreadId() { return SourceThreadId; }
+	void SetSourceDepth(int32 InDepth) { SourceDepth = InDepth; }
+	int32 GetSourceDepth() { return SourceDepth; }
+
+	double GetTargetTime() { return TargetTime; }
+	int32 GetTargetThreadId() { return TargetThreadId; }
+	void SetTargetDepth(int32 InDepth) { TargetDepth = InDepth; }
+	int32 GetTargetDepth() { return TargetDepth; }
+
+	ETaskEventType GetType() { return Type; }
+
+private:
+	double SourceTime;
+	int32 SourceThreadId;
+	int32 SourceDepth = -1;
+	double TargetTime;
+	int32 TargetThreadId;
+	int32 TargetDepth = -1;
+	ETaskEventType Type;
+
+	TWeakPtr<const FBaseTimingTrack> SourceTrack;
+	TWeakPtr<const FBaseTimingTrack> TargetTrack;
+};
+
+} // namespace Insights
diff --git a/Engine/Source/Developer/TraceInsights/Private/Insights/TaskGraphProfiler/ViewModels/TaskTimingTrack.cpp b/Engine/Source/Developer/TraceInsights/Private/Insights/TaskGraphProfiler/ViewModels/TaskTimingTrack.cpp
--- a/Engine/Source/Developer/TraceInsights/Private/Insights/TaskGraphProfiler/ViewModels/TaskTimingTrack.cpp
+++ b/Engine/Source/Developer/TraceInsights/Private/Insights/TaskGraphProfiler/ViewModels/TaskTimingTrack.cpp
@@ -12,8 +12,8 @@
 #include "Insights/InsightsStyle.h"
 #include "Insights/ITimingViewSession.h"
 #include "Insights/TaskGraphProfiler/TaskGraphProfilerManager.h"
+#include "Insights/TaskGraphProfiler/ViewModels/TaskGraphRelation.h"
 #include "Insights/TaskGraphProfiler/ViewModels/TaskTrackEvent.h"
-#include "Insights/ViewModels/TaskGraphRelation.h"
 #include "Insights/ViewModels/ThreadTimingTrack.h"
 #include "Insights/ViewModels/ThreadTrackEvent.h"
 #include "Insights/ViewModels/TimingTrackViewport.h"
diff --git a/Engine/Source/Developer/TraceInsights/Private/Insights/TaskGraphProfiler/Widgets/STaskTableTreeView.cpp b/Engine/Source/Developer/TraceInsights/Private/Insights/TaskGraphProfiler/Widgets/STaskTableTreeView.cpp
--- a/Engine/Source/Developer/TraceInsights/Private/Insights/TaskGraphProfiler/Widgets/STaskTableTreeView.cpp
+++ b/Engine/Source/Developer/TraceInsights/Private/Insights/TaskGraphProfiler/Widgets/STaskTableTreeView.cpp
@@ -12,12 +12,12 @@
 #include "Insights/InsightsStyle.h"
 #include "Insights/TaskGraphProfiler/TaskGraphProfilerManager.h"
 #include "Insights/TaskGraphProfiler/ViewModels/TaskEntry.h"
+#include "Insights/TaskGraphProfiler/ViewModels/TaskGraphRelation.h"
 #include "Insights/TaskGraphProfiler/ViewModels/TaskNode.h"
 #include "Insights/TaskGraphProfiler/ViewModels/TaskTimingTrack.h"
 #include "Insights/Table/ViewModels/TableColumn.h"
 #include "Insights/TimingProfilerManager.h"
 #include "Insights/ViewModels/FilterConfigurator.h"
-#include "Insights/ViewModels/TaskGraphRelation.h"
 #include "Insights/Widgets/STimingProfilerWindow.h"
 #include "Insights/Widgets/STimingView.h"
 
diff --git a/Engine/Source/Developer/TraceInsights/Private/Insights/ViewModels/TaskGraphRelation.cpp b/Engine/Source/Developer/TraceInsights/Private/Insights/ViewModels/TaskGraphRelation.cpp
deleted file mode 100644
diff --git a/Engine/Source/Developer/TraceInsights/Private/Insights/ViewModels/TaskGraphRelation.h b/Engine/Source/Developer/TraceInsights/Private/Insights/ViewModels/TaskGraphRelation.h
deleted file mode 100644
