[Insights] Renamed/moved Insights/TimingProfiler/QuickFindFilterConverters.h/.cpp to Insights/ViewModels/EventNameFilterValueConverter.h/cpp.

diff --git a/Engine/Source/Developer/TraceInsights/Private/Insights/TimingProfiler/QuickFindFilterConverters.cpp b/Engine/Source/Developer/TraceInsights/Private/Insights/TimingProfiler/QuickFindFilterConverters.cpp
deleted file mode 100644
diff --git a/Engine/Source/Developer/TraceInsights/Private/Insights/TimingProfiler/QuickFindFilterConverters.h b/Engine/Source/Developer/TraceInsights/Private/Insights/TimingProfiler/QuickFindFilterConverters.h
deleted file mode 100644
diff --git a/Engine/Source/Developer/TraceInsights/Private/Insights/ViewModels/EventNameFilterValueConverter.cpp b/Engine/Source/Developer/TraceInsights/Private/Insights/ViewModels/EventNameFilterValueConverter.cpp
new file mode 100644
--- /dev/null
+++ b/Engine/Source/Developer/TraceInsights/Private/Insights/ViewModels/EventNameFilterValueConverter.cpp
@@ -0,0 +1,64 @@
+// Copyright Epic Games, Inc. All Rights Reserved.
+
+#include "EventNameFilterValueConverter.h"
+
+#include "Insights/InsightsManager.h"
+
+#define LOCTEXT_NAMESPACE "Insights::FEventNameFilterValueConverter"
+
+namespace Insights
+{
+
+////////////////////////////////////////////////////////////////////////////////////////////////////
+
+bool FEventNameFilterValueConverter::Convert(const FString& Input, int64& Output, FText& OutError) const
+{
+	TSharedPtr<const TraceServices::IAnalysisSession> Session = FInsightsManager::Get()->GetSession();
+	if (Session.IsValid() && TraceServices::ReadTimingProfilerProvider(*Session.Get()))
+	{
+		TraceServices::FAnalysisSessionReadScope SessionReadScope(*Session.Get());
+
+		const TraceServices::ITimingProfilerProvider& TimingProfilerProvider = *TraceServices::ReadTimingProfilerProvider(*Session.Get());
+
+		const TraceServices::ITimingProfilerTimerReader* TimerReader;
+		TimingProfilerProvider.ReadTimers([&TimerReader](const TraceServices::ITimingProfilerTimerReader& Out) { TimerReader = &Out; });
+
+		uint32 TimerCount = TimerReader->GetTimerCount();
+		for (uint32 TimerIndex = 0; TimerIndex < TimerCount; ++TimerIndex)
+		{
+			const TraceServices::FTimingProfilerTimer* Timer = TimerReader->GetTimer(TimerIndex);
+			if (Timer && Timer->Name)
+			{
+				if (FCString::Strcmp(Timer->Name, *Input) == 0)
+				{
+					Output = Timer->Id;
+					return true;
+				}
+			}
+		}
+	}
+
+	OutError = LOCTEXT("NoTimerFound", "No timer with this name was found!");
+	return false;
+}
+
+////////////////////////////////////////////////////////////////////////////////////////////////////
+
+FText FEventNameFilterValueConverter::GetTooltipText() const
+{
+	return LOCTEXT("FEventNameFilterValueConverterTooltip", "Enter the exact name of the timer.");
+}
+
+////////////////////////////////////////////////////////////////////////////////////////////////////
+
+FText FEventNameFilterValueConverter::GetHintText() const
+{
+	// Use the name of a well known event as the hint to show the user what kind of name we are expecting.
+	return FText::FromString(TEXT("FEngineLoop"));
+}
+
+////////////////////////////////////////////////////////////////////////////////////////////////////
+
+} // namespace Insights
+
+#undef LOCTEXT_NAMESPACE
\ No newline at end of file
diff --git a/Engine/Source/Developer/TraceInsights/Private/Insights/ViewModels/EventNameFilterValueConverter.h b/Engine/Source/Developer/TraceInsights/Private/Insights/ViewModels/EventNameFilterValueConverter.h
new file mode 100644
--- /dev/null
+++ b/Engine/Source/Developer/TraceInsights/Private/Insights/ViewModels/EventNameFilterValueConverter.h
@@ -0,0 +1,18 @@
+// Copyright Epic Games, Inc. All Rights Reserved.
+
+#pragma once
+
+#include "Insights/ViewModels/Filters.h"
+
+namespace Insights
+{
+
+class FEventNameFilterValueConverter : public IFilterValueConverter
+{
+public:
+	virtual bool Convert(const FString& Input, int64& Output, FText& OutError) const override;
+	virtual FText GetTooltipText() const override;
+	virtual FText GetHintText() const override;
+};
+
+} // namespace Insights
diff --git a/Engine/Source/Developer/TraceInsights/Private/Insights/Widgets/STimingView.cpp b/Engine/Source/Developer/TraceInsights/Private/Insights/Widgets/STimingView.cpp
--- a/Engine/Source/Developer/TraceInsights/Private/Insights/Widgets/STimingView.cpp
+++ b/Engine/Source/Developer/TraceInsights/Private/Insights/Widgets/STimingView.cpp
@@ -49,9 +49,9 @@
 #include "Insights/Tests/TimingProfilerTests.h"
 #include "Insights/TimingProfilerCommon.h"
 #include "Insights/TimingProfilerManager.h"
-#include "Insights/TimingProfiler/QuickFindFilterConverters.h"
 #include "Insights/ViewModels/BaseTimingTrack.h"
 #include "Insights/ViewModels/DrawHelpers.h"
+#include "Insights/ViewModels/EventNameFilterValueConverter.h"
 #include "Insights/ViewModels/FileActivityTimingTrack.h"
 #include "Insights/ViewModels/FilterConfigurator.h"
 #include "Insights/ViewModels/FrameTimingTrack.h"
