[Insights] Draw info for selected event for Core Tracks

diff --git a/Engine/Source/Developer/TraceInsights/Private/Insights/ContextSwitches/ViewModels/CpuCoreTimingTrack.cpp b/Engine/Source/Developer/TraceInsights/Private/Insights/ContextSwitches/ViewModels/CpuCoreTimingTrack.cpp
--- a/Engine/Source/Developer/TraceInsights/Private/Insights/ContextSwitches/ViewModels/CpuCoreTimingTrack.cpp
+++ b/Engine/Source/Developer/TraceInsights/Private/Insights/ContextSwitches/ViewModels/CpuCoreTimingTrack.cpp
@@ -161,36 +161,9 @@ void FCpuCoreTimingTrack::AddCoreTimingEvent(ITimingEventsTrackDrawStateBuilder&
 	const uint32 SystemThreadId = CpuCoreEvent.SystemThreadId;
 
 	Builder.AddEvent(CpuCoreEvent.Start, CpuCoreEvent.End, EventDepth, EventColor,
-		[SystemThreadId, &CpuCoreEvent](float Width) -> const FString
+		[SystemThreadId, &CpuCoreEvent, this](float Width) -> const FString
 		{
-			const TCHAR* ThreadName = nullptr;
-
-			TSharedPtr<const TraceServices::IAnalysisSession> Session = FInsightsManager::Get()->GetSession();
-			if (Session.IsValid())
-			{
-				TraceServices::FAnalysisSessionReadScope SessionReadScope(*Session.Get());
-				const TraceServices::IContextSwitchesProvider* ContextSwitchesProvider = TraceServices::ReadContextSwitchesProvider(*Session.Get());
-				if (ContextSwitchesProvider)
-				{
-					uint32 ThreadId;
-					if (ContextSwitchesProvider->GetThreadId(SystemThreadId, ThreadId))
-					{
-						const TraceServices::IThreadProvider& ThreadProvider = TraceServices::ReadThreadProvider(*Session.Get());
-						ThreadName = ThreadProvider.GetThreadName(ThreadId);
-					}
-				}
-			}
-
-			FString EventName;
-
-			if (ThreadName)
-			{
-				EventName = ThreadName;
-			}
-			else
-			{
-				EventName = FString::Printf(TEXT("Unknown %d"), SystemThreadId);
-			}
+			FString EventName = GetThreadName(SystemThreadId);
 
 			if (Width > EventName.Len() * 4.0f + 32.0f)
 			{
@@ -211,6 +184,31 @@ void FCpuCoreTimingTrack::Draw(const ITimingTrackDrawContext& Context) const
 
 ////////////////////////////////////////////////////////////////////////////////////////////////////
 
+void FCpuCoreTimingTrack::PostDraw(const ITimingTrackDrawContext& Context) const
+{
+	if (ChildTrack.IsValid())
+	{
+		ChildTrack->PostDraw(Context);
+	}
+
+	const TSharedPtr<const ITimingEvent> SelectedEventPtr = Context.GetSelectedEvent();
+	if (SelectedEventPtr.IsValid() &&
+		SelectedEventPtr->CheckTrack(this) &&
+		SelectedEventPtr->Is<FTimingEvent>())
+	{
+		const FTimingEvent& SelectedEvent = SelectedEventPtr->As<FTimingEvent>();
+		const ITimingViewDrawHelper& Helper = Context.GetHelper();
+
+		FString Str = FString::Printf(TEXT("%s (Duration.: %s)"),
+			*GetThreadName(SelectedEvent.GetType()),
+			*TimeUtils::FormatTimeAuto(SelectedEvent.GetDuration()));
+
+		DrawSelectedEventInfo(Str, Context.GetViewport(), Context.GetDrawContext(), Helper.GetWhiteBrush(), Helper.GetEventFont());
+	}
+}
+
+////////////////////////////////////////////////////////////////////////////////////////////////////
+
 const TSharedPtr<const ITimingEvent> FCpuCoreTimingTrack::GetEvent(double InTime, double SecondsPerPixel, int32 Depth) const
 {
 	TSharedPtr<const TraceServices::IAnalysisSession> Session = FInsightsManager::Get()->GetSession();
@@ -343,7 +341,7 @@ const TSharedPtr<const ITimingEvent> FCpuCoreTimingTrack::SearchEvent(const FTim
 
 		[&FoundEvent, this](double InFoundStartTime, double InFoundEndTime, uint32 InFoundDepth, const TraceServices::FCpuCoreEvent& InEvent)
 		{
-			FoundEvent = MakeShared<FTimingEvent>(SharedThis(this), InFoundStartTime, InFoundEndTime, InFoundDepth);
+			FoundEvent = MakeShared<FTimingEvent>(SharedThis(this), InFoundStartTime, InFoundEndTime, InFoundDepth, InEvent.SystemThreadId);
 		},
 
 		TTimingEventSearch<TraceServices::FCpuCoreEvent>::NoMatch
@@ -434,6 +432,38 @@ bool FCpuCoreTimingTrack::HasCustomFilter() const
 
 ////////////////////////////////////////////////////////////////////////////////////////////////////
 
+FString FCpuCoreTimingTrack::GetThreadName(uint32 InSystemThreadId) const
+{
+	const TCHAR* ThreadName = nullptr;
+
+	TSharedPtr<const TraceServices::IAnalysisSession> Session = FInsightsManager::Get()->GetSession();
+	if (Session.IsValid())
+	{
+		TraceServices::FAnalysisSessionReadScope SessionReadScope(*Session.Get());
+		const TraceServices::IContextSwitchesProvider* ContextSwitchesProvider = TraceServices::ReadContextSwitchesProvider(*Session.Get());
+		if (ContextSwitchesProvider)
+		{
+			uint32 ThreadId;
+			if (ContextSwitchesProvider->GetThreadId(InSystemThreadId, ThreadId))
+			{
+				const TraceServices::IThreadProvider& ThreadProvider = TraceServices::ReadThreadProvider(*Session.Get());
+				ThreadName = ThreadProvider.GetThreadName(ThreadId);
+			}
+		}
+	}
+
+	if (ThreadName)
+	{
+		return FString(ThreadName);
+	}
+	else
+	{
+		return FString::Printf(TEXT("Unknown %d"), InSystemThreadId);
+	}
+}
+
+////////////////////////////////////////////////////////////////////////////////////////////////////
+
 } // namespace Insights
 
 #undef LOCTEXT_NAMESPACE
diff --git a/Engine/Source/Developer/TraceInsights/Private/Insights/ContextSwitches/ViewModels/CpuCoreTimingTrack.h b/Engine/Source/Developer/TraceInsights/Private/Insights/ContextSwitches/ViewModels/CpuCoreTimingTrack.h
--- a/Engine/Source/Developer/TraceInsights/Private/Insights/ContextSwitches/ViewModels/CpuCoreTimingTrack.h
+++ b/Engine/Source/Developer/TraceInsights/Private/Insights/ContextSwitches/ViewModels/CpuCoreTimingTrack.h
@@ -44,6 +44,8 @@ public:
 
 	virtual void Draw(const ITimingTrackDrawContext& Context) const override;
 
+	virtual void PostDraw(const ITimingTrackDrawContext& Context) const override;
+
 	virtual void InitTooltip(FTooltipDrawState& InOutTooltip, const ITimingEvent& InTooltipEvent) const override;
 
 	virtual void BuildContextMenu(FMenuBuilder& MenuBuilder) override;
@@ -57,6 +59,8 @@ protected:
 
 	void AddCoreTimingEvent(ITimingEventsTrackDrawStateBuilder& Builder, const TraceServices::FCpuCoreEvent& CpuCoreEvent);
 
+	FString GetThreadName(uint32 InSystemThreadId) const;
+
 private:
 	FContextSwitchesSharedState& SharedState;
 
diff --git a/Engine/Source/Developer/TraceInsights/Private/Insights/ViewModels/ThreadTimingTrack.cpp b/Engine/Source/Developer/TraceInsights/Private/Insights/ViewModels/ThreadTimingTrack.cpp
--- a/Engine/Source/Developer/TraceInsights/Private/Insights/ViewModels/ThreadTimingTrack.cpp
+++ b/Engine/Source/Developer/TraceInsights/Private/Insights/ViewModels/ThreadTimingTrack.cpp
@@ -3,7 +3,6 @@
 #include "ThreadTimingTrack.h"
 
 #include "CborReader.h"
-#include "Fonts/FontMeasure.h"
 #include "Framework/MultiBox/MultiBoxBuilder.h"
 #include "HAL/PlatformApplicationMisc.h"
 #include "Serialization/MemoryReader.h"
@@ -1009,46 +1008,27 @@ void FThreadTimingTrack::PostDraw(const ITimingTrackDrawContext& Context) const
 	{
 		const FThreadTrackEvent& SelectedEvent = SelectedEventPtr->As<FThreadTrackEvent>();
 		const ITimingViewDrawHelper& Helper = Context.GetHelper();
-		DrawSelectedEventInfo(SelectedEvent, Context.GetViewport(), Context.GetDrawContext(), Helper.GetWhiteBrush(), Helper.GetEventFont());
-	}
-}
-
-////////////////////////////////////////////////////////////////////////////////////////////////////
-
-void FThreadTimingTrack::DrawSelectedEventInfo(const FThreadTrackEvent& SelectedEvent, const FTimingTrackViewport& Viewport, const FDrawContext& DrawContext, const FSlateBrush* WhiteBrush, const FSlateFontInfo& Font) const
-{
-	TSharedPtr<const TraceServices::IAnalysisSession> Session = FInsightsManager::Get()->GetSession();
-	check(Session.IsValid());
-
-	TraceServices::FAnalysisSessionReadScope SessionReadScope(*Session.Get());
-
-	const TraceServices::ITimingProfilerProvider& TimingProfilerProvider = *TraceServices::ReadTimingProfilerProvider(*Session.Get());
 
-	const TraceServices::ITimingProfilerTimerReader* TimerReader;
-	TimingProfilerProvider.ReadTimers([&TimerReader](const TraceServices::ITimingProfilerTimerReader& Out) { TimerReader = &Out; });
+		TSharedPtr<const TraceServices::IAnalysisSession> Session = FInsightsManager::Get()->GetSession();
+		check(Session.IsValid());
 
-	const TraceServices::FTimingProfilerTimer* Timer = TimerReader->GetTimer(SelectedEvent.GetTimerIndex());
-	if (Timer != nullptr)
-	{
-		FString Str = FString::Printf(TEXT("%s (Incl.: %s, Excl.: %s)"),
-			Timer->Name,
-			*TimeUtils::FormatTimeAuto(SelectedEvent.GetDuration()),
-			*TimeUtils::FormatTimeAuto(SelectedEvent.GetExclusiveTime()));
+		TraceServices::FAnalysisSessionReadScope SessionReadScope(*Session.Get());
 
-		const TSharedRef<FSlateFontMeasure> FontMeasureService = FSlateApplication::Get().GetRenderer()->GetFontMeasureService();
-		const float FontScale = DrawContext.Geometry.Scale;
-		const FVector2D Size = FontMeasureService->Measure(Str, Font, FontScale) / FontScale;
-		const float X = Viewport.GetWidth() - Size.X - 23.0f;
-		const float Y = Viewport.GetPosY() + Viewport.GetHeight() - Size.Y - 18.0f;
+		const TraceServices::ITimingProfilerProvider& TimingProfilerProvider = *TraceServices::ReadTimingProfilerProvider(*Session.Get());
 
-		const FLinearColor BackgroundColor(0.05f, 0.05f, 0.05f, 1.0f);
-		const FLinearColor TextColor(0.7f, 0.7f, 0.7f, 1.0f);
+		const TraceServices::ITimingProfilerTimerReader* TimerReader;
+		TimingProfilerProvider.ReadTimers([&TimerReader](const TraceServices::ITimingProfilerTimerReader& Out) { TimerReader = &Out; });
 
-		DrawContext.DrawBox(X - 8.0f, Y - 2.0f, Size.X + 16.0f, Size.Y + 4.0f, WhiteBrush, BackgroundColor);
-		DrawContext.LayerId++;
+		const TraceServices::FTimingProfilerTimer* Timer = TimerReader->GetTimer(SelectedEvent.GetTimerIndex());
+		if (Timer != nullptr)
+		{
+			FString Str = FString::Printf(TEXT("%s (Incl.: %s, Excl.: %s)"),
+				Timer->Name,
+				*TimeUtils::FormatTimeAuto(SelectedEvent.GetDuration()),
+				*TimeUtils::FormatTimeAuto(SelectedEvent.GetExclusiveTime()));
 
-		DrawContext.DrawText(X, Y, Str, Font, TextColor);
-		DrawContext.LayerId++;
+			DrawSelectedEventInfo(Str, Context.GetViewport(), Context.GetDrawContext(), Helper.GetWhiteBrush(), Helper.GetEventFont());
+		}
 	}
 }
 
diff --git a/Engine/Source/Developer/TraceInsights/Private/Insights/ViewModels/ThreadTimingTrack.h b/Engine/Source/Developer/TraceInsights/Private/Insights/ViewModels/ThreadTimingTrack.h
--- a/Engine/Source/Developer/TraceInsights/Private/Insights/ViewModels/ThreadTimingTrack.h
+++ b/Engine/Source/Developer/TraceInsights/Private/Insights/ViewModels/ThreadTimingTrack.h
@@ -3,7 +3,6 @@
 #pragma once
 
 #include "CoreMinimal.h"
-#include "Fonts/SlateFontInfo.h"
 #include "TraceServices/Model/TimingProfiler.h"
 
 // Insights
@@ -16,7 +15,6 @@ class FTimingEventSearchParameters;
 class FGpuTimingTrack;
 class FCpuTimingTrack;
 class STimingView;
-struct FSlateBrush;
 
 namespace Insights
 {
@@ -172,8 +170,6 @@ protected:
 	virtual bool HasCustomFilter() const override;
 
 private:
-	void DrawSelectedEventInfo(const FThreadTrackEvent& SelectedEvent, const FTimingTrackViewport& Viewport, const FDrawContext& DrawContext, const FSlateBrush* WhiteBrush, const FSlateFontInfo& Font) const;
-
 	bool FindTimingProfilerEvent(const FThreadTrackEvent& InTimingEvent, TFunctionRef<void(double, double, uint32, const TraceServices::FTimingProfilerEvent&)> InFoundPredicate) const;
 	bool FindTimingProfilerEvent(const FTimingEventSearchParameters& InParameters, TFunctionRef<void(double, double, uint32, const TraceServices::FTimingProfilerEvent&)> InFoundPredicate) const;
 
diff --git a/Engine/Source/Developer/TraceInsights/Private/Insights/ViewModels/TimingEventsTrack.cpp b/Engine/Source/Developer/TraceInsights/Private/Insights/ViewModels/TimingEventsTrack.cpp
--- a/Engine/Source/Developer/TraceInsights/Private/Insights/ViewModels/TimingEventsTrack.cpp
+++ b/Engine/Source/Developer/TraceInsights/Private/Insights/ViewModels/TimingEventsTrack.cpp
@@ -2,6 +2,10 @@
 
 #include "Insights/ViewModels/TimingEventsTrack.h"
 
+#include "Fonts/FontMeasure.h"
+
+// Insights
+#include "Insights/Common/PaintUtils.h"
 #include "Insights/Common/Stopwatch.h"
 #include "Insights/TimingProfilerCommon.h"
 #include "Insights/ViewModels/TimingEvent.h"
@@ -456,4 +460,24 @@ TSharedPtr<ITimingEventFilter> FTimingEventsTrack::GetFilterByEvent(const TShare
 
 ////////////////////////////////////////////////////////////////////////////////////////////////////
 
+void FTimingEventsTrack::DrawSelectedEventInfo(const FString& InText, const FTimingTrackViewport& Viewport, const FDrawContext& DrawContext, const FSlateBrush* WhiteBrush, const FSlateFontInfo& Font) const
+{
+	const TSharedRef<FSlateFontMeasure> FontMeasureService = FSlateApplication::Get().GetRenderer()->GetFontMeasureService();
+	const float FontScale = DrawContext.Geometry.Scale;
+	const FVector2D Size = FontMeasureService->Measure(InText, Font, FontScale) / FontScale;
+	const float X = Viewport.GetWidth() - Size.X - 23.0f;
+	const float Y = Viewport.GetPosY() + Viewport.GetHeight() - Size.Y - 18.0f;
+
+	const FLinearColor BackgroundColor(0.05f, 0.05f, 0.05f, 1.0f);
+	const FLinearColor TextColor(0.7f, 0.7f, 0.7f, 1.0f);
+
+	DrawContext.DrawBox(X - 8.0f, Y - 2.0f, Size.X + 16.0f, Size.Y + 4.0f, WhiteBrush, BackgroundColor);
+	DrawContext.LayerId++;
+
+	DrawContext.DrawText(X, Y, InText, Font, TextColor);
+	DrawContext.LayerId++;
+}
+
+////////////////////////////////////////////////////////////////////////////////////////////////////
+
 #undef LOCTEXT_NAMESPACE
diff --git a/Engine/Source/Developer/TraceInsights/Public/Insights/ViewModels/TimingEventsTrack.h b/Engine/Source/Developer/TraceInsights/Public/Insights/ViewModels/TimingEventsTrack.h
--- a/Engine/Source/Developer/TraceInsights/Public/Insights/ViewModels/TimingEventsTrack.h
+++ b/Engine/Source/Developer/TraceInsights/Public/Insights/ViewModels/TimingEventsTrack.h
@@ -3,12 +3,15 @@
 #pragma once
 
 #include "CoreMinimal.h"
+#include "Fonts/SlateFontInfo.h"
 
 // Insights
 #include "Insights/ViewModels/BaseTimingTrack.h"
 
 ////////////////////////////////////////////////////////////////////////////////////////////////////
 
+struct FSlateBrush;
+
 class TRACEINSIGHTS_API ITimingEventsTrackDrawStateBuilder
 {
 public:
@@ -81,6 +84,8 @@ protected:
 
 	void DrawMarkers(const ITimingTrackDrawContext& Context, float LineY, float LineH) const;
 
+	void DrawSelectedEventInfo(const FString& InText, const FTimingTrackViewport& Viewport, const FDrawContext& DrawContext, const FSlateBrush* WhiteBrush, const FSlateFontInfo& Font) const;
+
 	int32 GetHeaderBackgroundLayerId(const ITimingTrackDrawContext& Context) const;
 	int32 GetHeaderTextLayerId(const ITimingTrackDrawContext& Context) const;
 
