StringBuilder: Deprecated Append(Char) in favor of AppendChar(Char)

Having only one Append overload with a single argument makes for clearer error messages. Otherwise, a failure to match Append(CharRangeType) leads to an error that ArgType cannot be converted to CharType.

diff --git a/Engine/Plugins/FX/Niagara/Source/Niagara/Private/NiagaraDebugHud.cpp b/Engine/Plugins/FX/Niagara/Source/Niagara/Private/NiagaraDebugHud.cpp
--- a/Engine/Plugins/FX/Niagara/Source/Niagara/Private/NiagaraDebugHud.cpp
+++ b/Engine/Plugins/FX/Niagara/Source/Niagara/Private/NiagaraDebugHud.cpp
@@ -1728,7 +1728,7 @@ void FNiagaraDebugHud::DrawGpuComputeOverriew(class FNiagaraWorldManager* WorldM
 		{
 			for ( int32 i=0; i < Columns.Num(); ++i )
 			{
-				Columns[i].AllRowsText.Append(TEXT('\n'));
+				Columns[i].AllRowsText.AppendChar(TEXT('\n'));
 			}
 		}
 
@@ -2308,11 +2308,11 @@ void FNiagaraDebugHud::DrawComponents(FNiagaraWorldManager* WorldManager, UCanva
 					StringBuilder.Appendf(TEXT("Particle(%u)"), iInstance);
 					for (const auto& ParticleVariable : CachedVariables.ParticleVariables[iEmitter])
 					{
-						StringBuilder.Append(Settings.bShowParticleVariablesVertical ? '\n' : ' ');
+						StringBuilder.AppendChar(Settings.bShowParticleVariablesVertical ? TEXT('\n') : TEXT(' '));
 						StringBuilder.Append(ParticleVariable.GetName().ToString());
-						StringBuilder.Append('(');
+						StringBuilder.AppendChar(TEXT('('));
 						ParticleVariable.StringAppend(StringBuilder, DataBuffer, iInstance);
-						StringBuilder.Append(')');
+						StringBuilder.AppendChar(TEXT(')'));
 					}
 
 					const TCHAR* FinalString = StringBuilder.ToString();
diff --git a/Engine/Plugins/FX/Niagara/Source/Niagara/Private/NiagaraTypes.cpp b/Engine/Plugins/FX/Niagara/Source/Niagara/Private/NiagaraTypes.cpp
--- a/Engine/Plugins/FX/Niagara/Source/Niagara/Private/NiagaraTypes.cpp
+++ b/Engine/Plugins/FX/Niagara/Source/Niagara/Private/NiagaraTypes.cpp
@@ -6,7 +6,7 @@ void FNiagaraVariableBase::SetNamespacedName(const FString& InNamespace, FName I
 {
 	TStringBuilder<128> NameBuilder;
 	NameBuilder.Append(InNamespace);
-	NameBuilder.Append('.');
+	NameBuilder.AppendChar(TEXT('.'));
 	InVariableName.AppendString(NameBuilder);
 	Name = FName(NameBuilder.ToString());
 }
diff --git a/Engine/Plugins/FX/Niagara/Source/NiagaraEditor/Private/NiagaraHlslTranslator.cpp b/Engine/Plugins/FX/Niagara/Source/NiagaraEditor/Private/NiagaraHlslTranslator.cpp
--- a/Engine/Plugins/FX/Niagara/Source/NiagaraEditor/Private/NiagaraHlslTranslator.cpp
+++ b/Engine/Plugins/FX/Niagara/Source/NiagaraEditor/Private/NiagaraHlslTranslator.cpp
@@ -4445,12 +4445,12 @@ int32 FHlslNiagaraTranslator::GetConstantDirect(int InConstantValue)
 
 bool FHlslNiagaraTranslator::GenerateStructInitializer(TStringBuilder<128>& InitializerString, UStruct* UserDefinedStruct, const void* StructData)
 {
-	InitializerString.Append('{');
+	InitializerString.AppendChar(TEXT('{'));
 	for (FField* ChildProperty = UserDefinedStruct->ChildProperties; ChildProperty; ChildProperty = ChildProperty->Next)
 	{
 		if (ChildProperty != UserDefinedStruct->ChildProperties)
 		{
-			InitializerString.Append(',');
+			InitializerString.AppendChar(TEXT(','));
 		}
 
 		if (FFloatProperty* FloatProperty = CastField<FFloatProperty>(ChildProperty))
@@ -4481,7 +4481,7 @@ bool FHlslNiagaraTranslator::GenerateStructInitializer(TStringBuilder<128>& Init
 			return false;
 		}
 	}
-	InitializerString.Append('}');
+	InitializerString.AppendChar(TEXT('}'));
 	return true;
 }
 
@@ -8070,7 +8070,7 @@ FString FHlslNiagaraTranslator::GenerateFunctionHlslPrototype(FStringView InVari
 		if (FunctionSignature.FunctionSpecifiers.Num())
 		{
 			bool bNeedsComma = false;
-			StringBuilder.Append('<');
+			StringBuilder.AppendChar(TEXT('<'));
 			for (auto SpecifierIt = FunctionSignature.FunctionSpecifiers.CreateConstIterator(); SpecifierIt; ++SpecifierIt)
 			{
 				if (bNeedsComma)
@@ -8080,12 +8080,12 @@ FString FHlslNiagaraTranslator::GenerateFunctionHlslPrototype(FStringView InVari
 				bNeedsComma = true;
 
 				StringBuilder.Append(SpecifierIt.Key().ToString());
-				StringBuilder.Append('=');
-				StringBuilder.Append('"');
+				StringBuilder.AppendChar(TEXT('='));
+				StringBuilder.AppendChar(TEXT('"'));
 				StringBuilder.Append(SpecifierIt.Value().IsNone() ? TEXT("Value") : *SpecifierIt.Value().ToString());
-				StringBuilder.Append('"');
+				StringBuilder.AppendChar(TEXT('"'));
 			}
-			StringBuilder.Append('>');
+			StringBuilder.AppendChar(TEXT('>'));
 		}
 
 		// Build function parameters
diff --git a/Engine/Source/Developer/TraceInsights/Private/Insights/Common/TimeUtils.cpp b/Engine/Source/Developer/TraceInsights/Private/Insights/Common/TimeUtils.cpp
--- a/Engine/Source/Developer/TraceInsights/Private/Insights/Common/TimeUtils.cpp
+++ b/Engine/Source/Developer/TraceInsights/Private/Insights/Common/TimeUtils.cpp
@@ -69,12 +69,12 @@ FString FormatTimeAuto(const double InDuration, const int32 NumDigits)
 	if (FMath::IsNegative(Duration))
 	{
 		Duration = -Duration;
-		StrBuilder.Append(TEXT('-'));
+		StrBuilder.AppendChar(TEXT('-'));
 	}
 
 	if (Duration == DBL_MAX || Duration == std::numeric_limits<double>::infinity())
 	{
-		StrBuilder.Append(TEXT('∞'));
+		StrBuilder.AppendChar(TEXT('∞'));
 	}
 	else if (Duration < TimeUtils::Picosecond)
 	{
@@ -163,7 +163,7 @@ FString FormatTimeAuto(const double InDuration, const int32 NumDigits)
 		else
 		{
 			StrBuilder.Append(FormatTimeValue(Duration, NumDigits));
-			StrBuilder.Append(TEXT('s'));
+			StrBuilder.AppendChar(TEXT('s'));
 		}
 	}
 	else if (Duration < TimeUtils::Hour)
@@ -182,9 +182,9 @@ FString FormatTimeAuto(const double InDuration, const int32 NumDigits)
 		}
 		else
 		{
-			StrBuilder.Append(TEXT(' '));
+			StrBuilder.AppendChar(TEXT(' '));
 			StrBuilder.Append(FormatTimeValue(Duration, NumDigits - 1));
-			StrBuilder.Append(TEXT('s'));
+			StrBuilder.AppendChar(TEXT('s'));
 		}
 	}
 	else if (Duration < TimeUtils::Day)
@@ -235,12 +235,12 @@ FString FormatTimeMs(const double InDuration, const int32 NumDigits, bool bAddTi
 	if (FMath::IsNegative(Duration))
 	{
 		Duration = -Duration;
-		StrBuilder.Append(TEXT('-'));
+		StrBuilder.AppendChar(TEXT('-'));
 	}
 
 	if (Duration == DBL_MAX || Duration == std::numeric_limits<double>::infinity())
 	{
-		StrBuilder.Append(TEXT('∞'));
+		StrBuilder.AppendChar(TEXT('∞'));
 	}
 	else
 	{
@@ -289,12 +289,12 @@ FString FormatTime(const double InTime, const double Precision)
 	if (FMath::IsNegative(Time))
 	{
 		Time = -Time;
-		StrBuilder.Append(TEXT('-'));
+		StrBuilder.AppendChar(TEXT('-'));
 	}
 
 	if (Time == DBL_MAX || Time == std::numeric_limits<double>::infinity())
 	{
-		StrBuilder.Append(TEXT('∞'));
+		StrBuilder.AppendChar(TEXT('∞'));
 		return FString(StrBuilder);
 	}
 
@@ -322,7 +322,7 @@ FString FormatTime(const double InTime, const double Precision)
 	{
 		if (bIsSpaceNeeded)
 		{
-			StrBuilder.Append(TEXT(' '));
+			StrBuilder.AppendChar(TEXT(' '));
 		}
 		StrBuilder.Appendf(TEXT("%dh"), Hours);
 		bIsSpaceNeeded = true;
@@ -343,7 +343,7 @@ FString FormatTime(const double InTime, const double Precision)
 	{
 		if (bIsSpaceNeeded)
 		{
-			StrBuilder.Append(TEXT(' '));
+			StrBuilder.AppendChar(TEXT(' '));
 		}
 		StrBuilder.Appendf(TEXT("%dm"), Minutes);
 		bIsSpaceNeeded = true;
@@ -395,7 +395,7 @@ FString FormatTime(const double InTime, const double Precision)
 		{
 			if (bIsSpaceNeeded)
 			{
-				StrBuilder.Append(TEXT(' '));
+				StrBuilder.AppendChar(TEXT(' '));
 			}
 			StrBuilder.Appendf(TEXT("%ds"), Seconds);
 		}
@@ -408,7 +408,7 @@ FString FormatTime(const double InTime, const double Precision)
 	//{
 	//	if (bIsSpaceNeeded)
 	//	{
-	//		StrBuilder.Append(TEXT(' '));
+	//		StrBuilder.AppendChar(TEXT(' '));
 	//	}
 	//	int32 Seconds = static_cast<int32>(Time / TimeUtils::Second);
 	//	Time -= static_cast<double>(Seconds) * TimeUtils::Second;
@@ -419,7 +419,7 @@ FString FormatTime(const double InTime, const double Precision)
 	{
 		if (bIsSpaceNeeded)
 		{
-			StrBuilder.Append(TEXT(' '));
+			StrBuilder.AppendChar(TEXT(' '));
 		}
 #if !PLATFORM_USE_GENERIC_STRING_IMPLEMENTATION
 		StrBuilder.Appendf(TEXT("%.*fs"), Digits, Time);
@@ -576,7 +576,7 @@ FString FormatTimeSplit(const FTimeSplit& InTimeSplit, const double Precision)
 
 	if (InTimeSplit.bIsNegative)
 	{
-		StrBuilder.Append(TEXT('-'));
+		StrBuilder.AppendChar(TEXT('-'));
 	}
 
 	bool bIsSpaceNeeded = false;
@@ -600,7 +600,7 @@ FString FormatTimeSplit(const FTimeSplit& InTimeSplit, const double Precision)
 	{
 		if (bIsSpaceNeeded)
 		{
-			StrBuilder.Append(TEXT(' '));
+			StrBuilder.AppendChar(TEXT(' '));
 		}
 		StrBuilder.Appendf(TEXT("%dh"), InTimeSplit.Hours);
 		bIsSpaceNeeded = true;
@@ -619,7 +619,7 @@ FString FormatTimeSplit(const FTimeSplit& InTimeSplit, const double Precision)
 	{
 		if (bIsSpaceNeeded)
 		{
-			StrBuilder.Append(TEXT(' '));
+			StrBuilder.AppendChar(TEXT(' '));
 		}
 		StrBuilder.Appendf(TEXT("%dm"), InTimeSplit.Minutes);
 		bIsSpaceNeeded = true;
@@ -638,7 +638,7 @@ FString FormatTimeSplit(const FTimeSplit& InTimeSplit, const double Precision)
 	{
 		if (bIsSpaceNeeded)
 		{
-			StrBuilder.Append(TEXT(' '));
+			StrBuilder.AppendChar(TEXT(' '));
 		}
 		StrBuilder.Appendf(TEXT("%ds"), InTimeSplit.Seconds);
 		bIsSpaceNeeded = true;
@@ -657,7 +657,7 @@ FString FormatTimeSplit(const FTimeSplit& InTimeSplit, const double Precision)
 	{
 		if (bIsSpaceNeeded)
 		{
-			StrBuilder.Append(TEXT(' '));
+			StrBuilder.AppendChar(TEXT(' '));
 		}
 		StrBuilder.Appendf(TEXT("%dms"), InTimeSplit.Miliseconds);
 		bIsSpaceNeeded = true;
@@ -676,7 +676,7 @@ FString FormatTimeSplit(const FTimeSplit& InTimeSplit, const double Precision)
 	{
 		if (bIsSpaceNeeded)
 		{
-			StrBuilder.Append(TEXT(' '));
+			StrBuilder.AppendChar(TEXT(' '));
 		}
 		StrBuilder.Appendf(TEXT("%dµs"), InTimeSplit.Microseconds);
 		bIsSpaceNeeded = true;
@@ -695,7 +695,7 @@ FString FormatTimeSplit(const FTimeSplit& InTimeSplit, const double Precision)
 	{
 		if (bIsSpaceNeeded)
 		{
-			StrBuilder.Append(TEXT(' '));
+			StrBuilder.AppendChar(TEXT(' '));
 		}
 		StrBuilder.Appendf(TEXT("%dns"), InTimeSplit.Nanoseconds);
 		bIsSpaceNeeded = true;
@@ -714,7 +714,7 @@ FString FormatTimeSplit(const FTimeSplit& InTimeSplit, const double Precision)
 	{
 		if (bIsSpaceNeeded)
 		{
-			StrBuilder.Append(TEXT(' '));
+			StrBuilder.AppendChar(TEXT(' '));
 		}
 		StrBuilder.Appendf(TEXT("%dps"), InTimeSplit.Picoseconds);
 		bIsSpaceNeeded = true;
diff --git a/Engine/Source/Developer/TraceInsights/Private/Insights/MemoryProfiler/Widgets/SMemInvestigationView.cpp b/Engine/Source/Developer/TraceInsights/Private/Insights/MemoryProfiler/Widgets/SMemInvestigationView.cpp
--- a/Engine/Source/Developer/TraceInsights/Private/Insights/MemoryProfiler/Widgets/SMemInvestigationView.cpp
+++ b/Engine/Source/Developer/TraceInsights/Private/Insights/MemoryProfiler/Widgets/SMemInvestigationView.cpp
@@ -495,7 +495,7 @@ FReply SMemInvestigationView::RunQuery()
 			{
 				Builder.Append(TEXT(", "));
 			}
-			Builder.Append((TCHAR)(TEXT('A') + TimeMarkerIndex));
+			Builder.AppendChar((TCHAR)(TEXT('A') + TimeMarkerIndex));
 			const TSharedRef<Insights::FTimeMarker>& TimeMarker = ProfilerWindow->GetCustomTimeMarker(TimeMarkerIndex);
 			Builder.Appendf(TEXT("=%.9f"), TimeMarker->GetTime());
 		}
diff --git a/Engine/Source/Developer/TraceInsights/Private/Insights/Widgets/SLogView.cpp b/Engine/Source/Developer/TraceInsights/Private/Insights/Widgets/SLogView.cpp
--- a/Engine/Source/Developer/TraceInsights/Private/Insights/Widgets/SLogView.cpp
+++ b/Engine/Source/Developer/TraceInsights/Private/Insights/Widgets/SLogView.cpp
@@ -1530,17 +1530,17 @@ void SLogView::AppendFormatMessageDetailed(const FLogMessageRecord& Log, TString
 void SLogView::AppendFormatMessageDelimitedHeader(TStringBuilderBase<TCHAR>& InOutStringBuilder, TCHAR Separator) const
 {
 	InOutStringBuilder.Append(TEXT("Index"));
-	InOutStringBuilder.Append(Separator);
+	InOutStringBuilder.AppendChar(Separator);
 	InOutStringBuilder.Append(TEXT("Time"));
-	InOutStringBuilder.Append(Separator);
+	InOutStringBuilder.AppendChar(Separator);
 	InOutStringBuilder.Append(TEXT("Verbosity"));
-	InOutStringBuilder.Append(Separator);
+	InOutStringBuilder.AppendChar(Separator);
 	InOutStringBuilder.Append(TEXT("Category"));
-	InOutStringBuilder.Append(Separator);
+	InOutStringBuilder.AppendChar(Separator);
 	InOutStringBuilder.Append(TEXT("Message"));
-	InOutStringBuilder.Append(Separator);
+	InOutStringBuilder.AppendChar(Separator);
 	InOutStringBuilder.Append(TEXT("File"));
-	InOutStringBuilder.Append(Separator);
+	InOutStringBuilder.AppendChar(Separator);
 	InOutStringBuilder.Append(TEXT("Line"));
 }
 
@@ -1549,13 +1549,13 @@ void SLogView::AppendFormatMessageDelimitedHeader(TStringBuilderBase<TCHAR>& InO
 void SLogView::AppendFormatMessageDelimited(const FLogMessageRecord& Log, TStringBuilderBase<TCHAR>& InOutStringBuilder, TCHAR Separator) const
 {
 	InOutStringBuilder.Appendf(TEXT("%d"), Log.GetIndex());
-	InOutStringBuilder.Append(Separator);
+	InOutStringBuilder.AppendChar(Separator);
 	InOutStringBuilder.Append(TimeUtils::FormatTimeHMS(Log.GetTime(), TimeUtils::Microsecond));
-	InOutStringBuilder.Append(Separator);
+	InOutStringBuilder.AppendChar(Separator);
 	InOutStringBuilder.Append(::ToString(Log.GetVerbosity()));
-	InOutStringBuilder.Append(Separator);
+	InOutStringBuilder.AppendChar(Separator);
 	InOutStringBuilder.Append(Log.GetCategoryAsString());
-	InOutStringBuilder.Append(Separator);
+	InOutStringBuilder.AppendChar(Separator);
 	FString Message = Log.GetMessageAsString();
 	if (Separator == TEXT('\t'))
 	{
@@ -1564,18 +1564,18 @@ void SLogView::AppendFormatMessageDelimited(const FLogMessageRecord& Log, TStrin
 	}
 	else if (Separator == TEXT(','))
 	{
-		InOutStringBuilder.Append(TEXT('\"'));
+		InOutStringBuilder.AppendChar(TEXT('\"'));
 		FString EscapedMessage = Message.Replace(TEXT("\""), TEXT("\"\""), ESearchCase::CaseSensitive);
 		InOutStringBuilder.Append(EscapedMessage);
-		InOutStringBuilder.Append(TEXT('\"'));
+		InOutStringBuilder.AppendChar(TEXT('\"'));
 	}
 	else
 	{
 		InOutStringBuilder.Append(Message);
 	}
-	InOutStringBuilder.Append(Separator);
+	InOutStringBuilder.AppendChar(Separator);
 	InOutStringBuilder.Append(Log.GetFile());
-	InOutStringBuilder.Append(Separator);
+	InOutStringBuilder.AppendChar(Separator);
 	InOutStringBuilder.Appendf(TEXT("%d"), Log.GetLine());
 }
 
@@ -1665,7 +1665,7 @@ void SLogView::CopyRange() const
 
 	TStringBuilder<1024> StringBuilder;
 	AppendFormatMessageDelimitedHeader(StringBuilder, TEXT('\t'));
-	StringBuilder.Append(TEXT('\n'));
+	StringBuilder.AppendChar(TEXT('\n'));
 
 	int32 NumMessagesInSelectedRange = 0;
 
@@ -1683,7 +1683,7 @@ void SLogView::CopyRange() const
 				{
 					FLogMessageRecord Log(MessageInfo);
 					AppendFormatMessageDelimited(Log, StringBuilder, TEXT('\t'));
-					StringBuilder.Append(TEXT('\n'));
+					StringBuilder.AppendChar(TEXT('\n'));
 					++NumMessagesInSelectedRange;
 				}
 			});
@@ -1715,7 +1715,7 @@ void SLogView::CopyAll() const
 
 	TStringBuilder<1024> StringBuilder;
 	AppendFormatMessageDelimitedHeader(StringBuilder, TEXT('\t'));
-	StringBuilder.Append(TEXT('\n'));
+	StringBuilder.AppendChar(TEXT('\n'));
 
 	TSharedPtr<const TraceServices::IAnalysisSession> Session = FInsightsManager::Get()->GetSession();
 	if (Session.IsValid())
@@ -1729,7 +1729,7 @@ void SLogView::CopyAll() const
 			{
 				FLogMessageRecord Log(MessageInfo);
 				AppendFormatMessageDelimited(Log, StringBuilder, TEXT('\t'));
-				StringBuilder.Append(TEXT('\n'));
+				StringBuilder.AppendChar(TEXT('\n'));
 			});
 		}
 	}
@@ -1844,7 +1844,7 @@ void SLogView::SaveLogsToFile(bool bSaveLogsInSelectedRangeOnly) const
 	{
 		TStringBuilder<1024> StringBuilder;
 		AppendFormatMessageDelimitedHeader(StringBuilder, Separator);
-		StringBuilder.Append(TEXT('\n'));
+		StringBuilder.AppendChar(TEXT('\n'));
 		FTCHARToUTF16 UTF16String(StringBuilder.ToString(), StringBuilder.Len());
 		ExportFileHandle->Write((const uint8*)UTF16String.Get(), UTF16String.Length() * sizeof(UTF16CHAR));
 	}
@@ -1866,7 +1866,7 @@ void SLogView::SaveLogsToFile(bool bSaveLogsInSelectedRangeOnly) const
 					TStringBuilder<1024> StringBuilder;
 					FLogMessageRecord Log(MessageInfo);
 					AppendFormatMessageDelimited(Log, StringBuilder, Separator);
-					StringBuilder.Append(TEXT('\n'));
+					StringBuilder.AppendChar(TEXT('\n'));
 					FTCHARToUTF16 UTF16String(StringBuilder.ToString(), StringBuilder.Len());
 					ExportFileHandle->Write((const uint8*)UTF16String.Get(), UTF16String.Length() * sizeof(UTF16CHAR));
 					++NumMessagesInSelectedRange;
diff --git a/Engine/Source/Developer/TraceInsights/Private/Insights/Widgets/STimersView.cpp b/Engine/Source/Developer/TraceInsights/Private/Insights/Widgets/STimersView.cpp
--- a/Engine/Source/Developer/TraceInsights/Private/Insights/Widgets/STimersView.cpp
+++ b/Engine/Source/Developer/TraceInsights/Private/Insights/Widgets/STimersView.cpp
@@ -2311,22 +2311,22 @@ void STimersView::ContextMenu_Export_Execute()
 			}
 			else
 			{
-				StringBuilder.Append(Separator);
+				StringBuilder.AppendChar(Separator);
 			}
 			FString Value = ColumnRef->GetShortName().ToString().ReplaceCharWithEscapedChar();
 			int32 CharIndex;
 			if (Value.FindChar(Separator, CharIndex))
 			{
-				StringBuilder.Append(QuotationMarkBegin);
+				StringBuilder.AppendChar(QuotationMarkBegin);
 				StringBuilder.Append(Value);
-				StringBuilder.Append(QuotationMarkEnd);
+				StringBuilder.AppendChar(QuotationMarkEnd);
 			}
 			else
 			{
 				StringBuilder.Append(Value);
 			}
 		}
-		StringBuilder.Append(LineEnd);
+		StringBuilder.AppendChar(LineEnd);
 		ExportFileHandle->Write((const uint8*)StringBuilder.ToString(), StringBuilder.Len() * sizeof(TCHAR));
 	}
 
@@ -2351,23 +2351,23 @@ void STimersView::ContextMenu_Export_Execute()
 			}
 			else
 			{
-				StringBuilder.Append(Separator);
+				StringBuilder.AppendChar(Separator);
 			}
 
 			FString Value = ColumnRef->GetValueAsSerializableString(*Node).ReplaceCharWithEscapedChar();
 			int32 CharIndex;
 			if (Value.FindChar(Separator, CharIndex))
 			{
-				StringBuilder.Append(QuotationMarkBegin);
+				StringBuilder.AppendChar(QuotationMarkBegin);
 				StringBuilder.Append(Value);
-				StringBuilder.Append(QuotationMarkEnd);
+				StringBuilder.AppendChar(QuotationMarkEnd);
 			}
 			else
 			{
 				StringBuilder.Append(Value);
 			}
 		}
-		StringBuilder.Append(LineEnd);
+		StringBuilder.AppendChar(LineEnd);
 		ExportFileHandle->Write((const uint8*)StringBuilder.ToString(), StringBuilder.Len() * sizeof(TCHAR));
 	}
 
@@ -2395,15 +2395,15 @@ struct FExportTimingEventsHeaderParams
 void ExportTimingEventsHeader(FExportTimingEventsHeaderParams& Params)
 {
 	Params.StringBuilder.Append(UTF8TEXT("ThreadId"));
-	Params.StringBuilder.Append(Params.Separator);
+	Params.StringBuilder.AppendChar(Params.Separator);
 	Params.StringBuilder.Append(UTF8TEXT("TimerId"));
-	Params.StringBuilder.Append(Params.Separator);
+	Params.StringBuilder.AppendChar(Params.Separator);
 	Params.StringBuilder.Append(UTF8TEXT("StartTime"));
-	Params.StringBuilder.Append(Params.Separator);
+	Params.StringBuilder.AppendChar(Params.Separator);
 	Params.StringBuilder.Append(UTF8TEXT("EndTime"));
-	Params.StringBuilder.Append(Params.Separator);
+	Params.StringBuilder.AppendChar(Params.Separator);
 	Params.StringBuilder.Append(UTF8TEXT("Depth"));
-	Params.StringBuilder.Append(Params.LineEnd);
+	Params.StringBuilder.AppendChar(Params.LineEnd);
 
 	Params.ExportFileHandle->Write((const uint8*)Params.StringBuilder.ToString(), Params.StringBuilder.Len() * sizeof(UTF8CHAR));
 }
@@ -2440,15 +2440,15 @@ void ExportTimingEventsEnumerate(FExportTimingEventsEnumerateParams& Params)
 					{
 						Params.StringBuilder.Reset();
 						Params.StringBuilder.Appendf(UTF8TEXT("%u"), Params.ThreadId);
-						Params.StringBuilder.Append(Params.Separator);
+						Params.StringBuilder.AppendChar(Params.Separator);
 						Params.StringBuilder.Appendf(UTF8TEXT("%u"), Event.TimerIndex);
-						Params.StringBuilder.Append(Params.Separator);
+						Params.StringBuilder.AppendChar(Params.Separator);
 						Params.StringBuilder.Appendf(UTF8TEXT("%.9g"), EventStartTime);
-						Params.StringBuilder.Append(Params.Separator);
+						Params.StringBuilder.AppendChar(Params.Separator);
 						Params.StringBuilder.Appendf(UTF8TEXT("%.9g"), EventEndTime);
-						Params.StringBuilder.Append(Params.Separator);
+						Params.StringBuilder.AppendChar(Params.Separator);
 						Params.StringBuilder.Appendf(UTF8TEXT("%u"), EventDepth);
-						Params.StringBuilder.Append(Params.LineEnd);
+						Params.StringBuilder.AppendChar(Params.LineEnd);
 
 						Params.ExportFileHandle->Write((const uint8*)Params.StringBuilder.ToString(), Params.StringBuilder.Len() * sizeof(UTF8CHAR));
 					}
@@ -2796,11 +2796,11 @@ void STimersView::ContextMenu_ExportThreads_Execute() const
 	// Write header.
 	{
 		StringBuilder.Append(UTF8TEXT("ThreadId"));
-		StringBuilder.Append(Separator);
+		StringBuilder.AppendChar(Separator);
 		StringBuilder.Append(UTF8TEXT("Name"));
-		StringBuilder.Append(Separator);
+		StringBuilder.AppendChar(Separator);
 		StringBuilder.Append(UTF8TEXT("GroupName"));
-		StringBuilder.Append(LineEnd);
+		StringBuilder.AppendChar(LineEnd);
 
 		ExportFileHandle->Write((const uint8*)StringBuilder.ToString(), StringBuilder.Len() * sizeof(UTF8CHAR));
 	}
@@ -2810,20 +2810,20 @@ void STimersView::ContextMenu_ExportThreads_Execute() const
 	{
 		StringBuilder.Reset();
 		StringBuilder.Appendf(UTF8TEXT("%u"), FGpuTimingTrack::Gpu1ThreadId);
-		StringBuilder.Append(Separator);
+		StringBuilder.AppendChar(Separator);
 		StringBuilder.Append(UTF8TEXT("GPU1"));
-		StringBuilder.Append(Separator);
+		StringBuilder.AppendChar(Separator);
 		StringBuilder.Append(UTF8TEXT("GPU"));
-		StringBuilder.Append(LineEnd);
+		StringBuilder.AppendChar(LineEnd);
 		ExportFileHandle->Write((const uint8*)StringBuilder.ToString(), StringBuilder.Len() * sizeof(UTF8CHAR));
 
 		StringBuilder.Reset();
 		StringBuilder.Appendf(UTF8TEXT("%u"), FGpuTimingTrack::Gpu2ThreadId);
-		StringBuilder.Append(Separator);
+		StringBuilder.AppendChar(Separator);
 		StringBuilder.Append(UTF8TEXT("GPU2"));
-		StringBuilder.Append(Separator);
+		StringBuilder.AppendChar(Separator);
 		StringBuilder.Append(UTF8TEXT("GPU"));
-		StringBuilder.Append(LineEnd);
+		StringBuilder.AppendChar(LineEnd);
 		ExportFileHandle->Write((const uint8*)StringBuilder.ToString(), StringBuilder.Len() * sizeof(UTF8CHAR));
 
 		// Iterate the CPU threads.
@@ -2835,11 +2835,11 @@ void STimersView::ContextMenu_ExportThreads_Execute() const
 				{
 					StringBuilder.Reset();
 					StringBuilder.Appendf(UTF8TEXT("%u"), ThreadInfo.Id);
-					StringBuilder.Append(Separator);
+					StringBuilder.AppendChar(Separator);
 					StringBuilder.Append(TCHAR_TO_UTF8(ThreadInfo.Name));
-					StringBuilder.Append(Separator);
+					StringBuilder.AppendChar(Separator);
 					StringBuilder.Append(TCHAR_TO_UTF8(ThreadInfo.GroupName));
-					StringBuilder.Append(LineEnd);
+					StringBuilder.AppendChar(LineEnd);
 					ExportFileHandle->Write((const uint8*)StringBuilder.ToString(), StringBuilder.Len() * sizeof(UTF8CHAR));
 				});
 		}
@@ -2890,15 +2890,15 @@ void STimersView::ContextMenu_ExportTimers_Execute() const
 	// Write header.
 	{
 		StringBuilder.Append(UTF8TEXT("TimerId"));
-		StringBuilder.Append(Separator);
+		StringBuilder.AppendChar(Separator);
 		StringBuilder.Append(UTF8TEXT("Type"));
-		StringBuilder.Append(Separator);
+		StringBuilder.AppendChar(Separator);
 		StringBuilder.Append(UTF8TEXT("Name"));
-		StringBuilder.Append(Separator);
+		StringBuilder.AppendChar(Separator);
 		StringBuilder.Append(UTF8TEXT("File"));
-		StringBuilder.Append(Separator);
+		StringBuilder.AppendChar(Separator);
 		StringBuilder.Append(UTF8TEXT("Line"));
-		StringBuilder.Append(LineEnd);
+		StringBuilder.AppendChar(LineEnd);
 
 		ExportFileHandle->Write((const uint8*)StringBuilder.ToString(), StringBuilder.Len() * sizeof(UTF8CHAR));
 	}
@@ -2919,15 +2919,15 @@ void STimersView::ContextMenu_ExportTimers_Execute() const
 			const TraceServices::FTimingProfilerTimer& Timer = *(TimerReader->GetTimer(TimerIndex));
 			StringBuilder.Reset();
 			StringBuilder.Appendf(UTF8TEXT("%u"), Timer.Id);
-			StringBuilder.Append(Separator);
+			StringBuilder.AppendChar(Separator);
 			StringBuilder.Append(Timer.IsGpuTimer ? UTF8TEXT("GPU") : UTF8TEXT("CPU"));
-			StringBuilder.Append(Separator);
+			StringBuilder.AppendChar(Separator);
 			StringBuilder.Append(TCHAR_TO_UTF8(Timer.Name));
-			StringBuilder.Append(Separator);
+			StringBuilder.AppendChar(Separator);
 			StringBuilder.Append(Timer.File ? (const UTF8CHAR*)TCHAR_TO_UTF8(Timer.File) : UTF8TEXT(""));
-			StringBuilder.Append(Separator);
+			StringBuilder.AppendChar(Separator);
 			StringBuilder.Appendf(UTF8TEXT("%u"), Timer.Line);
-			StringBuilder.Append(LineEnd);
+			StringBuilder.AppendChar(LineEnd);
 			ExportFileHandle->Write((const uint8*)StringBuilder.ToString(), StringBuilder.Len() * sizeof(UTF8CHAR));
 		}
 	}
diff --git a/Engine/Source/Programs/UnrealObjectPtrTool/Private/UnrealObjectPtrTool.cpp b/Engine/Source/Programs/UnrealObjectPtrTool/Private/UnrealObjectPtrTool.cpp
--- a/Engine/Source/Programs/UnrealObjectPtrTool/Private/UnrealObjectPtrTool.cpp
+++ b/Engine/Source/Programs/UnrealObjectPtrTool/Private/UnrealObjectPtrTool.cpp
@@ -263,7 +263,7 @@ bool FPointerUpgrader::TryFlushPendingWrites(FPendingWrites& PendingWrites)
 		TStringBuilder<2048> CompositeFilenamesBuilder;
 		for (const FPendingWrites::ElementType& PendingWrite : PendingWrites)
 		{
-			CompositeFilenamesBuilder.Append(TEXT('"'));
+			CompositeFilenamesBuilder.AppendChar(TEXT('"'));
 			CompositeFilenamesBuilder.Append(PendingWrite.Key);
 			CompositeFilenamesBuilder.Append(TEXT("\" "));
 		}
diff --git a/Engine/Source/Runtime/Core/Private/Internationalization/ICUCulture.cpp b/Engine/Source/Runtime/Core/Private/Internationalization/ICUCulture.cpp
--- a/Engine/Source/Runtime/Core/Private/Internationalization/ICUCulture.cpp
+++ b/Engine/Source/Runtime/Core/Private/Internationalization/ICUCulture.cpp
@@ -925,12 +925,12 @@ TSharedRef<const icu::DateFormat, ESPMode::ThreadSafe> FICUCultureImplementation
 					case TCHAR('p'): Result.Append(TEXT("a")); break;
 					case TCHAR('P'): Result.Append(TEXT("a")); break;
 					case TCHAR('j'): Result.Append(TEXT("D")); break;
-					default:		 Result.Append(*Format);
+					default:		 Result.AppendChar(*Format);
 					}
 				}
 				else
 				{
-					Result.Append(*Format);
+					Result.AppendChar(*Format);
 				}
 
 				// move to the next one
diff --git a/Engine/Source/Runtime/Core/Private/Misc/DateTime.cpp b/Engine/Source/Runtime/Core/Private/Misc/DateTime.cpp
--- a/Engine/Source/Runtime/Core/Private/Misc/DateTime.cpp
+++ b/Engine/Source/Runtime/Core/Private/Misc/DateTime.cpp
@@ -260,12 +260,12 @@ void FDateTime::ToString(const TCHAR* Format, FStringBuilderBase& Result) const
 				case TCHAR('M'): Result.Appendf(TEXT("%02i"), GetMinute()); break;
 				case TCHAR('S'): Result.Appendf(TEXT("%02i"), GetSecond()); break;
 				case TCHAR('s'): Result.Appendf(TEXT("%03i"), GetMillisecond()); break;
-				default:		 Result.Append(*Format);
+				default:		 Result.AppendChar(*Format);
 				}
 			}
 			else
 			{
-				Result.Append(*Format);
+				Result.AppendChar(*Format);
 			}
 
 			// move to the next one
@@ -304,12 +304,12 @@ FString FDateTime::ToFormattedString(const TCHAR* Format) const
 				case TCHAR('p'): Result.Append(IsMorning() ? TEXT("AM") : TEXT("PM")); break;
 				case TCHAR('P'): Result.Append(IsMorning() ? TEXT("am") : TEXT("pm")); break;
 				case TCHAR('j'): Result.Appendf(TEXT("%03i"), GetDayOfYear()); break;
-				default:		 Result.Append(*Format);
+				default:		 Result.AppendChar(*Format);
 				}
 			}
 			else
 			{
-				Result.Append(*Format);
+				Result.AppendChar(*Format);
 			}
 
 			// move to the next one
diff --git a/Engine/Source/Runtime/Core/Public/Misc/PathViews.h b/Engine/Source/Runtime/Core/Public/Misc/PathViews.h
--- a/Engine/Source/Runtime/Core/Public/Misc/PathViews.h
+++ b/Engine/Source/Runtime/Core/Public/Misc/PathViews.h
@@ -169,7 +169,7 @@ public:
 		{
 			if (!(Builder.Len() == 0 || Builder.LastChar() == '/' || Builder.LastChar() == '\\'))
 			{
-				Builder.Append('/');
+				Builder.AppendChar('/');
 			}
 			return Builder;
 		};
diff --git a/Engine/Source/Runtime/Core/Public/Misc/StringBuilder.h b/Engine/Source/Runtime/Core/Public/Misc/StringBuilder.h
--- a/Engine/Source/Runtime/Core/Public/Misc/StringBuilder.h
+++ b/Engine/Source/Runtime/Core/Public/Misc/StringBuilder.h
@@ -166,6 +166,7 @@ public:
 		return *this;
 	}
 
+	UE_DEPRECATED(5.0, "Use AppendChar instead of Append.")
 	inline BuilderType& Append(CharType Char) { return AppendChar(Char); }
 
 	UE_DEPRECATED(5.0, "Use Append instead of AppendAnsi.")
diff --git a/Engine/Source/Runtime/CoreUObject/Public/AssetRegistry/AssetData.h b/Engine/Source/Runtime/CoreUObject/Public/AssetRegistry/AssetData.h
--- a/Engine/Source/Runtime/CoreUObject/Public/AssetRegistry/AssetData.h
+++ b/Engine/Source/Runtime/CoreUObject/Public/AssetRegistry/AssetData.h
@@ -227,7 +227,7 @@ public:
 	{
 		OutFullNameBuilder.Reset();
 		AssetClass.AppendString(OutFullNameBuilder);
-		OutFullNameBuilder.Append(' ');
+		OutFullNameBuilder.AppendChar(' ');
 		ObjectPath.AppendString(OutFullNameBuilder);
 	}
 
@@ -254,9 +254,9 @@ public:
 	{
 		OutExportTextNameBuilder.Reset();
 		AssetClass.AppendString(OutExportTextNameBuilder);
-		OutExportTextNameBuilder.Append('\'');
+		OutExportTextNameBuilder.AppendChar('\'');
 		ObjectPath.AppendString(OutExportTextNameBuilder);
-		OutExportTextNameBuilder.Append('\'');
+		OutExportTextNameBuilder.AppendChar('\'');
 	}
 
 	/** Returns true if the this asset is a redirector. */
diff --git a/Engine/Source/Runtime/Engine/Private/Analytics/EngineNetAnalytics.cpp b/Engine/Source/Runtime/Engine/Private/Analytics/EngineNetAnalytics.cpp
--- a/Engine/Source/Runtime/Engine/Private/Analytics/EngineNetAnalytics.cpp
+++ b/Engine/Source/Runtime/Engine/Private/Analytics/EngineNetAnalytics.cpp
@@ -128,7 +128,7 @@ void FNetConnAnalyticsData::SendAnalytics()
 
 					if (!CurFullReason.IsEmpty())
 					{
-						CurFormattedReason.Append(TEXT(','));
+						CurFormattedReason.AppendChar(TEXT(','));
 					}
 
 					CurFormattedReason.Append(ToCStr(CurReason));
@@ -161,7 +161,7 @@ void FNetConnAnalyticsData::SendAnalytics()
 
 					if (!CurFullReason.IsEmpty())
 					{
-						CurFormattedReason.Append(TEXT(','));
+						CurFormattedReason.AppendChar(TEXT(','));
 					}
 
 					CurFormattedReason.Append(ToCStr(CurClientReason));
diff --git a/Engine/Source/Runtime/Engine/Private/BlueprintGeneratedClass.cpp b/Engine/Source/Runtime/Engine/Private/BlueprintGeneratedClass.cpp
--- a/Engine/Source/Runtime/Engine/Private/BlueprintGeneratedClass.cpp
+++ b/Engine/Source/Runtime/Engine/Private/BlueprintGeneratedClass.cpp
@@ -82,9 +82,9 @@ void WriteBPGCBreadcrumbs(FCrashContextExtendedWriter& Writer, const BPGCBreadcr
 		TStringBuilder<MAX_DATA_SIZE> Builder;
 
 		Params.Object.GetPathName(nullptr, Builder);
-		Builder.Append(TCHAR('\n'));
+		Builder.AppendChar(TEXT('\n'));
 		Params.BPGC.GetPathName(nullptr, Builder);
-		Builder.Append(TCHAR('\n'));
+		Builder.AppendChar(TEXT('\n'));
 
 		if (Params.Ar.GetSerializedProperty() != nullptr)
 		{
diff --git a/Engine/Source/Runtime/Engine/Private/HLSLTree/HLSLTree.cpp b/Engine/Source/Runtime/Engine/Private/HLSLTree/HLSLTree.cpp
--- a/Engine/Source/Runtime/Engine/Private/HLSLTree/HLSLTree.cpp
+++ b/Engine/Source/Runtime/Engine/Private/HLSLTree/HLSLTree.cpp
@@ -825,7 +825,7 @@ static void WriteIndent(int32 IndentLevel, FStringBuilderBase& InOutString)
 	TCHAR* Result = InOutString.GetData() + Offset;
 	for (int32 i = 0; i < IndentLevel; ++i)
 	{
-		*Result++ = TCHAR('\t');
+		*Result++ = TEXT('\t');
 	}
 }
 
@@ -853,7 +853,7 @@ void UE::HLSLTree::FScope::WriteHLSL(int32 Indent, FStringBuilderBase& OutString
 			check(!CodeDeclaration->NestedScope);
 			WriteIndent(Indent, OutString);
 			OutString.Append(CodeDeclaration->String, CodeDeclaration->Length);
-			OutString.Append(TEXT('\n'));
+			OutString.AppendChar(TEXT('\n'));
 			CodeDeclaration = CodeDeclaration->Next;
 		}
 	}
@@ -866,7 +866,7 @@ void UE::HLSLTree::FScope::WriteHLSL(int32 Indent, FStringBuilderBase& OutString
 			{
 				WriteIndent(Indent, OutString);
 				OutString.Append(CodeStatement->String, CodeStatement->Length);
-				OutString.Append(TEXT('\n'));
+				OutString.AppendChar(TEXT('\n'));
 			}
 			if (CodeStatement->NestedScope)
 			{
diff --git a/Engine/Source/Runtime/Engine/Private/Materials/MaterialUniformExpressions.cpp b/Engine/Source/Runtime/Engine/Private/Materials/MaterialUniformExpressions.cpp
--- a/Engine/Source/Runtime/Engine/Private/Materials/MaterialUniformExpressions.cpp
+++ b/Engine/Source/Runtime/Engine/Private/Materials/MaterialUniformExpressions.cpp
@@ -42,10 +42,10 @@ void WriteMaterialUniformAccess(UE::Shader::EValueComponentType ComponentType, u
 		// Can skip writing mask if we're taking all 4 components from the register
 		if (NumComponentsInRegister < 4u)
 		{
-			OutResult.Append(TCHAR('.'));
+			OutResult.AppendChar(TCHAR('.'));
 			for (uint32 i = 0u; i < NumComponentsInRegister; ++i)
 			{
-				OutResult.Append(IndexToMask[RegisterOffset + i]);
+				OutResult.AppendChar(IndexToMask[RegisterOffset + i]);
 			}
 		}
 		NumComponentsToWrite -= NumComponentsInRegister;
diff --git a/Engine/Source/Runtime/Engine/Private/NetConnection.cpp b/Engine/Source/Runtime/Engine/Private/NetConnection.cpp
--- a/Engine/Source/Runtime/Engine/Private/NetConnection.cpp
+++ b/Engine/Source/Runtime/Engine/Private/NetConnection.cpp
@@ -1384,9 +1384,9 @@ void UNetConnection::UpdateLevelVisibilityInternal(const FUpdateLevelVisibilityL
 			{
 				TStringBuilder<1024> VisibilityParms;
 
-				VisibilityParms.Append(ToCStr(PackageNameStr));
-				VisibilityParms.Append(TEXT(','));
-				VisibilityParms.Append(ToCStr(FileNameStr));
+				VisibilityParms.Append(PackageNameStr);
+				VisibilityParms.AppendChar(TEXT(','));
+				VisibilityParms.Append(FileNameStr);
 
 				Close({ENetCloseResult::MissingLevelPackage, VisibilityParms.ToString()});
 			}
diff --git a/Engine/Source/Runtime/MovieScene/Private/MovieSceneSpawnable.cpp b/Engine/Source/Runtime/MovieScene/Private/MovieSceneSpawnable.cpp
--- a/Engine/Source/Runtime/MovieScene/Private/MovieSceneSpawnable.cpp
+++ b/Engine/Source/Runtime/MovieScene/Private/MovieSceneSpawnable.cpp
@@ -72,12 +72,12 @@ FName FMovieSceneSpawnable::GetNetAddressableName(IMovieScenePlayer& Player, FMo
 	// Actor / player Name
 	if (AActor* OuterActor = PlayerObject->GetTypedOuter<AActor>())
 	{
-		AddressableName.Append('_');
+		AddressableName.AppendChar(TEXT('_'));
 		OuterActor->GetFName().AppendString(AddressableName);
 	}
 	else
 	{
-		AddressableName.Append('_');
+		AddressableName.AppendChar(TEXT('_'));
 		PlayerObject->GetFName().AppendString(AddressableName);
 	}
 
diff --git a/Engine/Source/Runtime/Net/Core/Private/Net/Core/Connection/NetResult.cpp b/Engine/Source/Runtime/Net/Core/Private/Net/Core/Connection/NetResult.cpp
--- a/Engine/Source/Runtime/Net/Core/Private/Net/Core/Connection/NetResult.cpp
+++ b/Engine/Source/Runtime/Net/Core/Private/Net/Core/Connection/NetResult.cpp
@@ -66,18 +66,18 @@ void NetResultToString(FString& OutResultStr, const TCHAR* LexResult, const FStr
 			CurResultStr.Append(TEXT(", "));
 		}
 
-		CurResultStr.Append(TEXT('('));
+		CurResultStr.AppendChar(TEXT('('));
 	}
 
 	CurResultStr.Append(TEXT("Result="));
 	CurResultStr.Append(LexResult);
 	CurResultStr.Append(TEXT(", ErrorContext=\""));
 	CurResultStr.Append(ToCStr(ErrorContext));
-	CurResultStr.Append(TEXT('\"'));
+	CurResultStr.AppendChar(TEXT('\"'));
 
 	if (bIncludeChain)
 	{
-		CurResultStr.Append(TEXT(')'));
+		CurResultStr.AppendChar(TEXT(')'));
 	}
 
 	OutResultStr += CurResultStr.ToString();
diff --git a/Engine/Source/Runtime/Net/Core/Private/Net/Core/Connection/StateStruct.cpp b/Engine/Source/Runtime/Net/Core/Private/Net/Core/Connection/StateStruct.cpp
--- a/Engine/Source/Runtime/Net/Core/Private/Net/Core/Connection/StateStruct.cpp
+++ b/Engine/Source/Runtime/Net/Core/Private/Net/Core/Connection/StateStruct.cpp
@@ -32,9 +32,9 @@ const UStatePerObjectConfig* UStatePerObjectConfig::Get(FStateConfigParms Config
 	TStringBuilder<1024> FormattedName;
 
 	FormattedName.Append(ToCStr(ConfigContext));
-	FormattedName.Append(TEXT('_'));
+	FormattedName.AppendChar(TEXT('_'));
 	FormattedName.Append(ToCStr(ConfigSection));
-	FormattedName.Append(TEXT('_'));
+	FormattedName.AppendChar(TEXT('_'));
 	FormattedName.Append(ToCStr(ConfigClass->GetFName().ToString()));
 
 	const TCHAR* FormattedNameStr = FormattedName.ToString();
@@ -47,7 +47,7 @@ const UStatePerObjectConfig* UStatePerObjectConfig::Get(FStateConfigParms Config
 		if (!ConfigContext.IsEmpty())
 		{
 			FullSection.Append(ToCStr(ConfigContext));
-			FullSection.Append(TEXT(' '));
+			FullSection.AppendChar(TEXT(' '));
 		}
 
 		FullSection.Append(ToCStr(ConfigSection));
@@ -113,7 +113,7 @@ void UStatePerObjectConfig::RegisterStateConfig(const TArray<FString>& StateName
 			TStringBuilder<1024> CurSection;
 
 			CurSection.Append(ToCStr(ConfigParms.ConfigSection));
-			CurSection.Append(TEXT('.'));
+			CurSection.AppendChar(TEXT('.'));
 			CurSection.Append(ToCStr(CurStateName));
 
 			const TCHAR* CurSectionStr = CurSection.ToString();
diff --git a/Engine/Source/Runtime/RenderCore/Private/ShaderCore.cpp b/Engine/Source/Runtime/RenderCore/Private/ShaderCore.cpp
--- a/Engine/Source/Runtime/RenderCore/Private/ShaderCore.cpp
+++ b/Engine/Source/Runtime/RenderCore/Private/ShaderCore.cpp
@@ -894,7 +894,7 @@ bool ReplaceVirtualFilePathForShaderAutogen(FString& InOutVirtualFilePath, EShad
 		const uint32 PlatformNameLen = PlatformName.GetPlainNameString(PlatformNameString);
 		OutputShaderName.Append(PlatformNameString, PlatformNameLen);
 
-		OutputShaderName.Append(TEXT('/'));
+		OutputShaderName.AppendChar(TEXT('/'));
 
 		// Append the relative name (the substring after "/Engine/Generated/ShaderAutogen/").
 		const TCHAR* RelativeShaderName = *InOutVirtualFilePath + ShaderAutogenStem.Len();
diff --git a/Engine/Source/Runtime/SlateCore/Private/Debugging/ConsoleSlateDebugger.cpp b/Engine/Source/Runtime/SlateCore/Private/Debugging/ConsoleSlateDebugger.cpp
--- a/Engine/Source/Runtime/SlateCore/Private/Debugging/ConsoleSlateDebugger.cpp
+++ b/Engine/Source/Runtime/SlateCore/Private/Debugging/ConsoleSlateDebugger.cpp
@@ -365,7 +365,7 @@ void FConsoleSlateDebugger::OnRouteInput(ESlateDebuggingInputEvent InputEventTyp
 		Message.Append(StaticEnum<ESlateDebuggingInputEvent>()->GetDisplayNameTextByValue((int64)InputEventType).ToString());
 		Message.Append(TEXT(" ("));
 		Message.Append(RoutedType.ToString());
-		Message.Append(TEXT(')'));
+		Message.AppendChar(TEXT(')'));
 		UE_LOG(LogSlateDebugger, Log, TEXT("%s"), *Message);
 	}
 	
diff --git a/Engine/Source/Runtime/SlateCore/Private/Debugging/SlateDebugging.cpp b/Engine/Source/Runtime/SlateCore/Private/Debugging/SlateDebugging.cpp
--- a/Engine/Source/Runtime/SlateCore/Private/Debugging/SlateDebugging.cpp
+++ b/Engine/Source/Runtime/SlateCore/Private/Debugging/SlateDebugging.cpp
@@ -258,7 +258,7 @@ FString LexToString(ESlateDebuggingInvalidateRootReason InValue)
 	}
 
 	TStringBuilder<512> Result;
-#define ENUM_CASE_TO_STRING(Enum) if (EnumHasAnyFlags(InValue, ESlateDebuggingInvalidateRootReason::Enum)) { if (Result.Len() != 0) { Result.Append(TEXT('|')); } Result.Append(TEXT(#Enum)); }
+#define ENUM_CASE_TO_STRING(Enum) if (EnumHasAnyFlags(InValue, ESlateDebuggingInvalidateRootReason::Enum)) { if (Result.Len() != 0) { Result.AppendChar(TEXT('|')); } Result.Append(TEXT(#Enum)); }
 	ENUM_CASE_TO_STRING(ChildOrder);
 	ENUM_CASE_TO_STRING(Root);
 	ENUM_CASE_TO_STRING(ScreenPosition);
diff --git a/Engine/Source/Runtime/SlateCore/Private/Widgets/InvalidateWidgetReason.cpp b/Engine/Source/Runtime/SlateCore/Private/Widgets/InvalidateWidgetReason.cpp
--- a/Engine/Source/Runtime/SlateCore/Private/Widgets/InvalidateWidgetReason.cpp
+++ b/Engine/Source/Runtime/SlateCore/Private/Widgets/InvalidateWidgetReason.cpp
@@ -69,7 +69,7 @@ FString LexToString(EInvalidateWidgetReason InValue)
 	}
 
 	TStringBuilder<512> Result;
-	#define ENUM_CASE_TO_STRING(Enum) if (EnumHasAnyFlags(InValue, EInvalidateWidgetReason::Enum)) { if (Result.Len() != 0) { Result.Append(TEXT('|')); } Result.Append(TEXT(#Enum)); }
+	#define ENUM_CASE_TO_STRING(Enum) if (EnumHasAnyFlags(InValue, EInvalidateWidgetReason::Enum)) { if (Result.Len() != 0) { Result.AppendChar(TEXT('|')); } Result.Append(TEXT(#Enum)); }
 	ENUM_CASE_TO_STRING(Layout);
 	ENUM_CASE_TO_STRING(Paint);
 	ENUM_CASE_TO_STRING(Volatility);
