[Insights] Printed address range in modules view was truncated.

diff --git a/Engine/Source/Developer/TraceInsights/Private/Insights/Widgets/SModulesView.cpp b/Engine/Source/Developer/TraceInsights/Private/Insights/Widgets/SModulesView.cpp
--- a/Engine/Source/Developer/TraceInsights/Private/Insights/Widgets/SModulesView.cpp
+++ b/Engine/Source/Developer/TraceInsights/Private/Insights/Widgets/SModulesView.cpp
@@ -41,7 +41,7 @@ public:
 	FText GetAddressRangeText() const
 	{
 		TStringBuilder<64> Sb;
-		Sb.Appendf(TEXT("0x%016x - 0x%016x"), Module->Base, Module->Base + Module->Size);
+		Sb.Appendf(TEXT("0x%016llX - 0x%016llX"), Module->Base, Module->Base + Module->Size);
 		return FText::FromStringView(Sb.ToView());
 	}
 	
