[Insights] Memory Insights: Added "Modules" toggle button to the toolbar.

diff --git a/Engine/Source/Developer/TraceInsights/Private/Insights/InsightsStyle.cpp b/Engine/Source/Developer/TraceInsights/Private/Insights/InsightsStyle.cpp
--- a/Engine/Source/Developer/TraceInsights/Private/Insights/InsightsStyle.cpp
+++ b/Engine/Source/Developer/TraceInsights/Private/Insights/InsightsStyle.cpp
@@ -211,6 +211,9 @@ void FInsightsStyle::FStyle::Initialize()
 
 	Set("Icons.MemAllocTableTreeView", new IMAGE_BRUSH_SVG("MemAllocTable", Icon16x16));
 
+	Set("Icons.ModulesView", new EDITOR_IMAGE_BRUSH_SVG("Starship/Common/Spreadsheet", Icon16x16));
+	Set("Icons.ModulesView.ToolBar", new EDITOR_IMAGE_BRUSH_SVG("Starship/Common/Spreadsheet", Icon20x20));
+
 	Set("Icons.AddAllMemTagGraphs", new CORE_IMAGE_BRUSH_SVG("Starship/Common/plus-circle", Icon16x16));
 	Set("Icons.RemoveAllMemTagGraphs", new CORE_IMAGE_BRUSH_SVG("Starship/Common/Delete", Icon16x16));
 
diff --git a/Engine/Source/Developer/TraceInsights/Private/Insights/MemoryProfiler/MemoryProfilerCommands.cpp b/Engine/Source/Developer/TraceInsights/Private/Insights/MemoryProfiler/MemoryProfilerCommands.cpp
--- a/Engine/Source/Developer/TraceInsights/Private/Insights/MemoryProfiler/MemoryProfilerCommands.cpp
+++ b/Engine/Source/Developer/TraceInsights/Private/Insights/MemoryProfiler/MemoryProfilerCommands.cpp
@@ -47,13 +47,39 @@ FMemoryProfilerCommands::FMemoryProfilerCommands()
 
 ////////////////////////////////////////////////////////////////////////////////////////////////////
 
+FMemoryProfilerCommands::~FMemoryProfilerCommands()
+{
+}
+
+////////////////////////////////////////////////////////////////////////////////////////////////////
+
 // UI_COMMAND takes long for the compiler to optimize
 PRAGMA_DISABLE_OPTIMIZATION
 void FMemoryProfilerCommands::RegisterCommands()
 {
-	UI_COMMAND(ToggleTimingViewVisibility, "Timing", "Toggles the visibility of the main Timing view", EUserInterfaceActionType::ToggleButton, FInputChord(EModifierKey::Control, EKeys::T));
-	UI_COMMAND(ToggleMemInvestigationViewVisibility, "Investigation", "Toggles the visibility of the Memory Investigation (Alloc Queries) view", EUserInterfaceActionType::ToggleButton, FInputChord(EModifierKey::Control, EKeys::I));
-	UI_COMMAND(ToggleMemTagTreeViewVisibility, "LLM Tags", "Toggles the visibility of the LLM Tags tree view", EUserInterfaceActionType::ToggleButton, FInputChord(EModifierKey::Control, EKeys::M));
+	UI_COMMAND(ToggleTimingViewVisibility,
+		"Timing",
+		"Toggles the visibility of the main Timing view.",
+		EUserInterfaceActionType::ToggleButton,
+		FInputChord());
+
+	UI_COMMAND(ToggleMemInvestigationViewVisibility,
+		"Investigation",
+		"Toggles the visibility of the Memory Investigation (Alloc Queries) view.",
+		EUserInterfaceActionType::ToggleButton,
+		FInputChord());
+
+	UI_COMMAND(ToggleMemTagTreeViewVisibility,
+		"LLM Tags",
+		"Toggles the visibility of the LLM Tags tree view.",
+		EUserInterfaceActionType::ToggleButton,
+		FInputChord());
+
+	UI_COMMAND(ToggleModulesViewVisibility,
+		"Modules",
+		"Toggles the visibility of the Modules view.",
+		EUserInterfaceActionType::ToggleButton,
+		FInputChord());
 }
 PRAGMA_ENABLE_OPTIMIZATION
 
@@ -97,6 +123,7 @@ PRAGMA_ENABLE_OPTIMIZATION
 IMPLEMENT_TOGGLE_COMMAND(ToggleTimingViewVisibility, IsTimingViewVisible, ShowHideTimingView)
 IMPLEMENT_TOGGLE_COMMAND(ToggleMemInvestigationViewVisibility, IsMemInvestigationViewVisible, ShowHideMemInvestigationView)
 IMPLEMENT_TOGGLE_COMMAND(ToggleMemTagTreeViewVisibility, IsMemTagTreeViewVisible, ShowHideMemTagTreeView)
+IMPLEMENT_TOGGLE_COMMAND(ToggleModulesViewVisibility, IsModulesViewVisible, ShowHideModulesView)
 
 #undef IMPLEMENT_TOGGLE_COMMAND
 
diff --git a/Engine/Source/Developer/TraceInsights/Private/Insights/MemoryProfiler/MemoryProfilerCommands.h b/Engine/Source/Developer/TraceInsights/Private/Insights/MemoryProfiler/MemoryProfilerCommands.h
--- a/Engine/Source/Developer/TraceInsights/Private/Insights/MemoryProfiler/MemoryProfilerCommands.h
+++ b/Engine/Source/Developer/TraceInsights/Private/Insights/MemoryProfiler/MemoryProfilerCommands.h
@@ -15,10 +15,8 @@ class FMenuBuilder;
 class FMemoryProfilerCommands : public TCommands<FMemoryProfilerCommands>
 {
 public:
-	/** Default constructor. */
 	FMemoryProfilerCommands();
-
-	/** Initialize commands. */
+	virtual ~FMemoryProfilerCommands();
 	virtual void RegisterCommands() override;
 
 public:
@@ -37,6 +35,9 @@ public:
 
 	/** Toggles visibility for the LLM Tags tree view. Global and custom command. */
 	TSharedPtr<FUICommandInfo> ToggleMemTagTreeViewVisibility;
+
+	/** Toggles visibility for the Modules view. Global and custom command. */
+	TSharedPtr<FUICommandInfo> ToggleModulesViewVisibility;
 };
 
 ////////////////////////////////////////////////////////////////////////////////////////////////////
@@ -91,6 +92,7 @@ protected:\
 	DECLARE_TOGGLE_COMMAND(ToggleTimingViewVisibility)
 	DECLARE_TOGGLE_COMMAND(ToggleMemInvestigationViewVisibility)
 	DECLARE_TOGGLE_COMMAND(ToggleMemTagTreeViewVisibility)
+	DECLARE_TOGGLE_COMMAND(ToggleModulesViewVisibility)
 #undef DECLARE_TOGGLE_COMMAND
 
 	//////////////////////////////////////////////////
diff --git a/Engine/Source/Developer/TraceInsights/Private/Insights/MemoryProfiler/MemoryProfilerManager.cpp b/Engine/Source/Developer/TraceInsights/Private/Insights/MemoryProfiler/MemoryProfilerManager.cpp
--- a/Engine/Source/Developer/TraceInsights/Private/Insights/MemoryProfiler/MemoryProfilerManager.cpp
+++ b/Engine/Source/Developer/TraceInsights/Private/Insights/MemoryProfiler/MemoryProfilerManager.cpp
@@ -60,6 +60,7 @@ FMemoryProfilerManager::FMemoryProfilerManager(TSharedRef<FUICommandList> InComm
 	, bIsTimingViewVisible(false)
 	, bIsMemInvestigationViewVisible(false)
 	, bIsMemTagTreeViewVisible(false)
+	, bIsModulesViewVisible(false)
 	, LogListingName(TEXT("MemoryInsights"))
 {
 }
@@ -134,6 +135,7 @@ void FMemoryProfilerManager::BindCommands()
 	ActionManager.Map_ToggleTimingViewVisibility_Global();
 	ActionManager.Map_ToggleMemInvestigationViewVisibility_Global();
 	ActionManager.Map_ToggleMemTagTreeViewVisibility_Global();
+	ActionManager.Map_ToggleModulesViewVisibility_Global();
 }
 
 ////////////////////////////////////////////////////////////////////////////////////////////////////
@@ -370,6 +372,19 @@ void FMemoryProfilerManager::ShowHideMemTagTreeView(const bool bIsVisible)
 
 ////////////////////////////////////////////////////////////////////////////////////////////////////
 
+void FMemoryProfilerManager::ShowHideModulesView(const bool bIsVisible)
+{
+	bIsModulesViewVisible = bIsVisible;
+
+	TSharedPtr<SMemoryProfilerWindow> Wnd = GetProfilerWindow();
+	if (Wnd.IsValid())
+	{
+		Wnd->ShowHideTab(FMemoryProfilerTabs::ModulesViewID, bIsVisible);
+	}
+}
+
+////////////////////////////////////////////////////////////////////////////////////////////////////
+
 void FMemoryProfilerManager::OnWindowClosedEvent()
 {
 	// Need to close MemAlloc window to prevent it being saved in the layout and spawning as an "Unregister Tab" on the next application start. 
diff --git a/Engine/Source/Developer/TraceInsights/Private/Insights/MemoryProfiler/MemoryProfilerManager.h b/Engine/Source/Developer/TraceInsights/Private/Insights/MemoryProfiler/MemoryProfilerManager.h
--- a/Engine/Source/Developer/TraceInsights/Private/Insights/MemoryProfiler/MemoryProfilerManager.h
+++ b/Engine/Source/Developer/TraceInsights/Private/Insights/MemoryProfiler/MemoryProfilerManager.h
@@ -100,10 +100,10 @@ public:
 	void SetMemTagTreeViewVisible(const bool bIsVisible) { bIsMemTagTreeViewVisible = bIsVisible; }
 	void ShowHideMemTagTreeView(const bool bIsVisible);
 
-	const bool IsMemModuleViewVisible() const { return bIsMemModuleViewVisible; }
-	void SetMemModuleViewVisible(const bool bIsVisible) { bIsMemModuleViewVisible = bIsVisible; }
-	void ShowHideMemModuleView(const bool bIsVisible);
-	
+	const bool IsModulesViewVisible() const { return bIsModulesViewVisible; }
+	void SetModulesViewVisible(const bool bIsVisible) { bIsModulesViewVisible = bIsVisible; }
+	void ShowHideModulesView(const bool bIsVisible);
+
 	////////////////////////////////////////////////////////////////////////////////////////////////////
 
 	const FName& GetLogListingName() const { return LogListingName; }
@@ -154,8 +154,8 @@ private:
 	/** If the LLM Tags tree view is visible or hidden. */
 	bool bIsMemTagTreeViewVisible;
 
-	/** If the modules view is visible or hidden. */
-	bool bIsMemModuleViewVisible;
+	/** If the Modules view is visible or hidden. */
+	bool bIsModulesViewVisible;
 
 	/** The name of the Memory Insights log listing. */
 	FName LogListingName;
diff --git a/Engine/Source/Developer/TraceInsights/Private/Insights/MemoryProfiler/Widgets/SMemoryProfilerToolbar.cpp b/Engine/Source/Developer/TraceInsights/Private/Insights/MemoryProfiler/Widgets/SMemoryProfilerToolbar.cpp
--- a/Engine/Source/Developer/TraceInsights/Private/Insights/MemoryProfiler/Widgets/SMemoryProfilerToolbar.cpp
+++ b/Engine/Source/Developer/TraceInsights/Private/Insights/MemoryProfiler/Widgets/SMemoryProfilerToolbar.cpp
@@ -46,6 +46,9 @@ void SMemoryProfilerToolbar::Construct(const FArguments& InArgs)
 				ToolbarBuilder.AddToolBarButton(FMemoryProfilerCommands::Get().ToggleMemTagTreeViewVisibility,
 					NAME_None, TAttribute<FText>(), TAttribute<FText>(),
 					FSlateIcon(FInsightsStyle::GetStyleSetName(), "Icons.MemTagTreeView.ToolBar"));
+				ToolbarBuilder.AddToolBarButton(FMemoryProfilerCommands::Get().ToggleModulesViewVisibility,
+					NAME_None, TAttribute<FText>(), TAttribute<FText>(),
+					FSlateIcon(FInsightsStyle::GetStyleSetName(), "Icons.ModulesView.ToolBar"));
 			}
 			ToolbarBuilder.EndSection();
 		}
diff --git a/Engine/Source/Developer/TraceInsights/Private/Insights/MemoryProfiler/Widgets/SMemoryProfilerWindow.cpp b/Engine/Source/Developer/TraceInsights/Private/Insights/MemoryProfiler/Widgets/SMemoryProfilerWindow.cpp
--- a/Engine/Source/Developer/TraceInsights/Private/Insights/MemoryProfiler/Widgets/SMemoryProfilerWindow.cpp
+++ b/Engine/Source/Developer/TraceInsights/Private/Insights/MemoryProfiler/Widgets/SMemoryProfilerWindow.cpp
@@ -49,7 +49,7 @@ const FName FMemoryProfilerTabs::TimingViewID(TEXT("TimingView"));
 const FName FMemoryProfilerTabs::MemInvestigationViewID(TEXT("MemInvestigation"));
 const FName FMemoryProfilerTabs::MemTagTreeViewID(TEXT("LowLevelMemTags"));
 const FName FMemoryProfilerTabs::MemAllocTableTreeViewID(TEXT("MemAllocTableTreeView"));
-const FName FMemoryProfilerTabs::MemModulesViewID(TEXT("MemModules"));
+const FName FMemoryProfilerTabs::ModulesViewID(TEXT("Modules"));
 
 ////////////////////////////////////////////////////////////////////////////////////////////////////
 
@@ -467,18 +467,18 @@ END_SLATE_FUNCTION_BUILD_OPTIMIZATION
 ////////////////////////////////////////////////////////////////////////////////////////////////////
 
 BEGIN_SLATE_FUNCTION_BUILD_OPTIMIZATION
-TSharedRef<SDockTab> SMemoryProfilerWindow::SpawnTab_MemModulesView(const FSpawnTabArgs& Args)
+TSharedRef<SDockTab> SMemoryProfilerWindow::SpawnTab_ModulesView(const FSpawnTabArgs& Args)
 {
-	FMemoryProfilerManager::Get()->SetMemModuleViewVisible(true);
+	FMemoryProfilerManager::Get()->SetModulesViewVisible(true);
 
 	const TSharedRef<SDockTab> DockTab = SNew(SDockTab)
 		.ShouldAutosize(false)
 		.TabRole(ETabRole::PanelTab)
 		[
-			SAssignNew(MemModulesView, Insights::SModulesView)
+			SAssignNew(ModulesView, Insights::SModulesView)
 		];
 
-	DockTab->SetOnTabClosed(SDockTab::FOnTabClosedCallback::CreateRaw(this, &SMemoryProfilerWindow::OnMemModulesViewClosed));
+	DockTab->SetOnTabClosed(SDockTab::FOnTabClosedCallback::CreateRaw(this, &SMemoryProfilerWindow::OnModulesViewClosed));
 
 	return DockTab;
 }
@@ -486,10 +486,10 @@ END_SLATE_FUNCTION_BUILD_OPTIMIZATION
 
 ////////////////////////////////////////////////////////////////////////////////////////////////////
 
-void SMemoryProfilerWindow::OnMemModulesViewClosed(TSharedRef<SDockTab> TabBeingClosed)
+void SMemoryProfilerWindow::OnModulesViewClosed(TSharedRef<SDockTab> TabBeingClosed)
 {
-	FMemoryProfilerManager::Get()->SetMemModuleViewVisible(false);
-	MemModulesView = nullptr;
+	FMemoryProfilerManager::Get()->SetModulesViewVisible(false);
+	ModulesView = nullptr;
 }
 
 ////////////////////////////////////////////////////////////////////////////////////////////////////
@@ -522,9 +522,9 @@ void SMemoryProfilerWindow::Construct(const FArguments& InArgs, const TSharedRef
 		.SetIcon(FSlateIcon(FInsightsStyle::GetStyleSetName(), "Icons.MemTagTreeView"))
 		.SetGroup(AppMenuGroup.ToSharedRef());
 
-	TabManager->RegisterTabSpawner(FMemoryProfilerTabs::MemModulesViewID, FOnSpawnTab::CreateRaw(this, &SMemoryProfilerWindow::SpawnTab_MemModulesView))
-		.SetDisplayName(LOCTEXT("MemModulesViewTabTitle", "Modules"))
-		.SetIcon(FSlateIcon(FInsightsStyle::GetStyleSetName(), "Icons.MemTagTreeView")) //todo: Update icon
+	TabManager->RegisterTabSpawner(FMemoryProfilerTabs::ModulesViewID, FOnSpawnTab::CreateRaw(this, &SMemoryProfilerWindow::SpawnTab_ModulesView))
+		.SetDisplayName(LOCTEXT("ModulesViewTabTitle", "Modules"))
+		.SetIcon(FSlateIcon(FInsightsStyle::GetStyleSetName(), "Icons.ModulesView"))
 		.SetGroup(AppMenuGroup.ToSharedRef());
 
 	TSharedPtr<FMemoryProfilerManager> MemoryProfilerManager = FMemoryProfilerManager::Get();
@@ -549,7 +549,7 @@ void SMemoryProfilerWindow::Construct(const FArguments& InArgs, const TSharedRef
 				->SetSizeCoefficient(0.23f)
 				->AddTab(FMemoryProfilerTabs::MemInvestigationViewID, ETabState::OpenedTab)
 				->AddTab(FMemoryProfilerTabs::MemTagTreeViewID, ETabState::OpenedTab)
-				->AddTab(FMemoryProfilerTabs::MemModulesViewID, ETabState::OpenedTab)
+				->AddTab(FMemoryProfilerTabs::ModulesViewID, ETabState::OpenedTab)
 				->SetForegroundTab(FMemoryProfilerTabs::MemInvestigationViewID)
 			)
 		);
diff --git a/Engine/Source/Developer/TraceInsights/Private/Insights/MemoryProfiler/Widgets/SMemoryProfilerWindow.h b/Engine/Source/Developer/TraceInsights/Private/Insights/MemoryProfiler/Widgets/SMemoryProfilerWindow.h
--- a/Engine/Source/Developer/TraceInsights/Private/Insights/MemoryProfiler/Widgets/SMemoryProfilerWindow.h
+++ b/Engine/Source/Developer/TraceInsights/Private/Insights/MemoryProfiler/Widgets/SMemoryProfilerWindow.h
@@ -43,7 +43,7 @@ struct FMemoryProfilerTabs
 	static const FName MemInvestigationViewID;
 	static const FName MemTagTreeViewID;
 	static const FName MemAllocTableTreeViewID; // base id
-	static const FName MemModulesViewID;
+	static const FName ModulesViewID;
 };
 
 ////////////////////////////////////////////////////////////////////////////////////////////////////
@@ -106,8 +106,8 @@ private:
 	TSharedRef<SDockTab> SpawnTab_MemAllocTableTreeView(const FSpawnTabArgs& Args, int32 TabIndex);
 	void OnMemAllocTableTreeViewTabClosed(TSharedRef<SDockTab> TabBeingClosed);
 
-	TSharedRef<SDockTab> SpawnTab_MemModulesView(const FSpawnTabArgs& Args);
-	void OnMemModulesViewClosed(TSharedRef<SDockTab> TabBeingClosed);
+	TSharedRef<SDockTab> SpawnTab_ModulesView(const FSpawnTabArgs& Args);
+	void OnModulesViewClosed(TSharedRef<SDockTab> TabBeingClosed);
 	
 	/**
 	 * Fill the main menu with menu items.
@@ -197,8 +197,8 @@ private:
 	/** The list of Allocations table tree view widgets */
 	TArray<TSharedPtr<Insights::SMemAllocTableTreeView>> MemAllocTableTreeViews;
 
-	/** The modules view widget. */
-	TSharedPtr<Insights::SModulesView> MemModulesView;
+	/** The Modules view widget. */
+	TSharedPtr<Insights::SModulesView> ModulesView;
 
 	const int32 MaxMemAllocTableTreeViews = 4;
 	int32 LastMemAllocTableTreeViewIndex = -1;
