[Insights]
  - Moved TimingView actions with key shortcuts into commands.
  - Removed shortcut keys from all toggle buttons from the main toolbars.

diff --git a/Engine/Source/Developer/TraceInsights/Private/Insights/Common/InsightsMenuBuilder.cpp b/Engine/Source/Developer/TraceInsights/Private/Insights/Common/InsightsMenuBuilder.cpp
--- a/Engine/Source/Developer/TraceInsights/Private/Insights/Common/InsightsMenuBuilder.cpp
+++ b/Engine/Source/Developer/TraceInsights/Private/Insights/Common/InsightsMenuBuilder.cpp
@@ -12,6 +12,8 @@
 
 #define LOCTEXT_NAMESPACE "InsightsMenuBuilder"
 
+////////////////////////////////////////////////////////////////////////////////////////////////////
+// FInsightsMenuBuilder
 ////////////////////////////////////////////////////////////////////////////////////////////////////
 
 FInsightsMenuBuilder::FInsightsMenuBuilder()
@@ -88,4 +90,46 @@ void FInsightsMenuBuilder::PopulateMenu(FMenuBuilder& MenuBuilder)
 
 ////////////////////////////////////////////////////////////////////////////////////////////////////
 
+void FInsightsMenuBuilder::AddMenuEntry(
+	FMenuBuilder& InOutMenuBuilder,
+	const FUIAction& InAction,
+	const TAttribute<FText>& InLabel,
+	const TAttribute<FText>& InToolTipText,
+	const TAttribute<FText>& InKeybinding,
+	const EUserInterfaceActionType InUserInterfaceActionType)
+{
+	InOutMenuBuilder.AddMenuEntry(
+		InAction,
+		SNew(SBox)
+		.Padding(FMargin(0.0f))
+		[
+			SNew(SHorizontalBox)
+			+ SHorizontalBox::Slot()
+			.Padding(FMargin(2.0f, 0.0f))
+			.VAlign(VAlign_Center)
+			.FillWidth(1.0f)
+			[
+				SNew(STextBlock)
+				.TextStyle(FAppStyle::Get(), "Menu.Label")
+				.Text(InLabel)
+			]
+			+ SHorizontalBox::Slot()
+			.Padding(FMargin(4.0f, 0.0f))
+			.VAlign(VAlign_Center)
+			.AutoWidth()
+			[
+				SNew(STextBlock)
+				.TextStyle(FAppStyle::Get(), "Menu.Keybinding")
+				.ColorAndOpacity(FSlateColor::UseSubduedForeground())
+				.Text(InKeybinding)
+			]
+		],
+		NAME_None,
+		InToolTipText,
+		InUserInterfaceActionType
+	);
+}
+
+////////////////////////////////////////////////////////////////////////////////////////////////////
+
 #undef LOCTEXT_NAMESPACE
diff --git a/Engine/Source/Developer/TraceInsights/Private/Insights/Common/InsightsMenuBuilder.h b/Engine/Source/Developer/TraceInsights/Private/Insights/Common/InsightsMenuBuilder.h
--- a/Engine/Source/Developer/TraceInsights/Private/Insights/Common/InsightsMenuBuilder.h
+++ b/Engine/Source/Developer/TraceInsights/Private/Insights/Common/InsightsMenuBuilder.h
@@ -3,10 +3,13 @@
 #pragma once
 
 #include "CoreMinimal.h"
+#include "Framework/Commands/UICommandInfo.h"
 
 class FWorkspaceItem;
 class FMenuBuilder;
 
+struct FUIAction;
+
 class FInsightsMenuBuilder
 {
 public:
@@ -17,6 +20,15 @@ public:
 
 	void PopulateMenu(FMenuBuilder& MenuBuilder);
 
+	// Adds a menu entry with a custom key binding text.
+	static void AddMenuEntry(
+		FMenuBuilder& InOutMenuBuilder,
+		const FUIAction& InAction,
+		const TAttribute<FText>& InLabel,
+		const TAttribute<FText>& InToolTipText,
+		const TAttribute<FText>& InKeybinding,
+		const EUserInterfaceActionType InUserInterfaceActionType);
+
 private:
 #if !WITH_EDITOR
 	TSharedRef<FWorkspaceItem> InsightsToolsGroup;
diff --git a/Engine/Source/Developer/TraceInsights/Private/Insights/LoadingProfiler/LoadingProfilerCommands.cpp b/Engine/Source/Developer/TraceInsights/Private/Insights/LoadingProfiler/LoadingProfilerCommands.cpp
--- a/Engine/Source/Developer/TraceInsights/Private/Insights/LoadingProfiler/LoadingProfilerCommands.cpp
+++ b/Engine/Source/Developer/TraceInsights/Private/Insights/LoadingProfiler/LoadingProfilerCommands.cpp
@@ -51,12 +51,41 @@ FLoadingProfilerCommands::FLoadingProfilerCommands()
 PRAGMA_DISABLE_OPTIMIZATION
 void FLoadingProfilerCommands::RegisterCommands()
 {
-	UI_COMMAND(ToggleTimingViewVisibility, "Timing", "Toggles the visibility of the main Timing view", EUserInterfaceActionType::ToggleButton, FInputChord(EModifierKey::Control, EKeys::T));
-	UI_COMMAND(ToggleEventAggregationTreeViewVisibility, "Event Aggregation", "Toggles the visibility of the Event Aggregation table/tree view", EUserInterfaceActionType::ToggleButton, FInputChord());
-	UI_COMMAND(ToggleObjectTypeAggregationTreeViewVisibility, "Object Type Aggregation", "Toggles the visibility of the Object Type Aggregation table/tree view", EUserInterfaceActionType::ToggleButton, FInputChord());
-	UI_COMMAND(TogglePackageDetailsTreeViewVisibility, "Package Details", "Toggles the visibility of the Package Details table/tree view", EUserInterfaceActionType::ToggleButton, FInputChord());
-	UI_COMMAND(ToggleExportDetailsTreeViewVisibility, "Export Details", "Toggles the visibility of the Export Details table/tree view", EUserInterfaceActionType::ToggleButton, FInputChord());
-	UI_COMMAND(ToggleRequestsTreeViewVisibility, "Requests", "Toggles the visibility of the Requests table/tree view", EUserInterfaceActionType::ToggleButton, FInputChord());
+	UI_COMMAND(ToggleTimingViewVisibility,
+		"Timing",
+		"Toggles the visibility of the main Timing view.",
+		EUserInterfaceActionType::ToggleButton,
+		FInputChord());
+
+	UI_COMMAND(ToggleEventAggregationTreeViewVisibility,
+		"Event Aggregation",
+		"Toggles the visibility of the Event Aggregation table/tree view.",
+		EUserInterfaceActionType::ToggleButton,
+		FInputChord());
+
+	UI_COMMAND(ToggleObjectTypeAggregationTreeViewVisibility,
+		"Object Type Aggregation",
+		"Toggles the visibility of the Object Type Aggregation table/tree view.",
+		EUserInterfaceActionType::ToggleButton,
+		FInputChord());
+
+	UI_COMMAND(TogglePackageDetailsTreeViewVisibility,
+		"Package Details",
+		"Toggles the visibility of the Package Details table/tree view.",
+		EUserInterfaceActionType::ToggleButton,
+		FInputChord());
+
+	UI_COMMAND(ToggleExportDetailsTreeViewVisibility,
+		"Export Details",
+		"Toggles the visibility of the Export Details table/tree view.",
+		EUserInterfaceActionType::ToggleButton,
+		FInputChord());
+
+	UI_COMMAND(ToggleRequestsTreeViewVisibility,
+		"Requests",
+		"Toggles the visibility of the Requests table/tree view.",
+		EUserInterfaceActionType::ToggleButton,
+		FInputChord());
 }
 PRAGMA_ENABLE_OPTIMIZATION
 
diff --git a/Engine/Source/Developer/TraceInsights/Private/Insights/MemoryProfiler/ViewModels/MemorySharedState.cpp b/Engine/Source/Developer/TraceInsights/Private/Insights/MemoryProfiler/ViewModels/MemorySharedState.cpp
--- a/Engine/Source/Developer/TraceInsights/Private/Insights/MemoryProfiler/ViewModels/MemorySharedState.cpp
+++ b/Engine/Source/Developer/TraceInsights/Private/Insights/MemoryProfiler/ViewModels/MemorySharedState.cpp
@@ -1,6 +1,8 @@
 // Copyright Epic Games, Inc. All Rights Reserved.
 
 #include "MemorySharedState.h"
+#include "Framework/Commands/Commands.h"
+#include "Framework/Commands/UICommandList.h"
 #include "Framework/MultiBox/MultiBoxBuilder.h"
 #include "Logging/MessageLog.h"
 #include "TraceServices/AnalysisService.h"
@@ -11,6 +13,7 @@
 #include "Insights/Common/Stopwatch.h"
 #include "Insights/Common/TimeUtils.h"
 #include "Insights/InsightsManager.h"
+#include "Insights/InsightsStyle.h"
 #include "Insights/MemoryProfiler/MemoryProfilerManager.h"
 #include "Insights/MemoryProfiler/ViewModels/MemoryTag.h"
 #include "Insights/MemoryProfiler/ViewModels/MemoryTracker.h"
@@ -29,6 +32,39 @@
 
 const FName Insights::FQueryTargetWindowSpec::NewWindow = TEXT("New Window");
 
+////////////////////////////////////////////////////////////////////////////////////////////////////
+// FMemoryTimingViewCommands
+////////////////////////////////////////////////////////////////////////////////////////////////////
+
+FMemoryTimingViewCommands::FMemoryTimingViewCommands()
+	: TCommands<FMemoryTimingViewCommands>(
+		TEXT("MemoryTimingViewCommands"),
+		NSLOCTEXT("Contexts", "MemoryTimingViewCommands", "Insights - Timing View - Memory"),
+		NAME_None,
+		FInsightsStyle::GetStyleSetName())
+{
+}
+
+////////////////////////////////////////////////////////////////////////////////////////////////////
+
+FMemoryTimingViewCommands::~FMemoryTimingViewCommands()
+{
+}
+
+////////////////////////////////////////////////////////////////////////////////////////////////////
+
+// UI_COMMAND takes long for the compiler to optimize
+PRAGMA_DISABLE_OPTIMIZATION
+void FMemoryTimingViewCommands::RegisterCommands()
+{
+	UI_COMMAND(ShowHideAllMemoryTracks,
+		"Memory Tracks",
+		"Shows/hides the Memory tracks.",
+		EUserInterfaceActionType::ToggleButton,
+		FInputChord(EModifierKey::Control, EKeys::M));
+}
+PRAGMA_ENABLE_OPTIMIZATION
+
 ////////////////////////////////////////////////////////////////////////////////////////////////////
 // FMemorySharedState
 ////////////////////////////////////////////////////////////////////////////////////////////////////
@@ -346,7 +382,7 @@ void FMemorySharedState::SetTrackHeightMode(EMemoryTrackHeightMode InTrackHeight
 
 ////////////////////////////////////////////////////////////////////////////////////////////////////
 
-void FMemorySharedState::ExtendFilterMenu(Insights::ITimingViewSession& InSession, FMenuBuilder& InOutMenuBuilder)
+void FMemorySharedState::ExtendOtherTracksFilterMenu(Insights::ITimingViewSession& InSession, FMenuBuilder& InOutMenuBuilder)
 {
 	if (&InSession != TimingView.Get())
 	{
@@ -355,22 +391,29 @@ void FMemorySharedState::ExtendFilterMenu(Insights::ITimingViewSession& InSessio
 
 	InOutMenuBuilder.BeginSection("Memory", LOCTEXT("ContextMenu_Section_Memory", "Memory"));
 	{
-		InOutMenuBuilder.AddMenuEntry(
-			LOCTEXT("AllMemoryTracks", "Memory Tracks - M"),
-			LOCTEXT("AllMemoryTracks_Tooltip", "Shows or hides the Memory tracks."),
-			FSlateIcon(),
-			FUIAction(FExecuteAction::CreateSP(this, &FMemorySharedState::ShowHideAllMemoryTracks),
-					  FCanExecuteAction(),
-					  FIsActionChecked::CreateSP(this, &FMemorySharedState::IsAllMemoryTracksToggleOn)),
-			NAME_None, //"QuickFilterSeparator",
-			EUserInterfaceActionType::ToggleButton
-		);
+		InOutMenuBuilder.AddMenuEntry(FMemoryTimingViewCommands::Get().ShowHideAllMemoryTracks);
 	}
 	InOutMenuBuilder.EndSection();
 }
 
 ////////////////////////////////////////////////////////////////////////////////////////////////////
 
+void FMemorySharedState::BindCommands()
+{
+	FMemoryTimingViewCommands::Register();
+
+	TSharedPtr<FUICommandList> CommandList = TimingView->GetCommandList();
+	ensure(CommandList.IsValid());
+
+	CommandList->MapAction(
+		FMemoryTimingViewCommands::Get().ShowHideAllMemoryTracks,
+		FExecuteAction::CreateSP(this, &FMemorySharedState::ShowHideAllMemoryTracks),
+		FCanExecuteAction(),
+		FIsActionChecked::CreateSP(this, &FMemorySharedState::IsAllMemoryTracksToggleOn));
+}
+
+////////////////////////////////////////////////////////////////////////////////////////////////////
+
 void FMemorySharedState::SetAllMemoryTracksToggle(bool bOnOff)
 {
 	bShowHideAllMemoryTracks = bOnOff;
diff --git a/Engine/Source/Developer/TraceInsights/Private/Insights/MemoryProfiler/ViewModels/MemorySharedState.h b/Engine/Source/Developer/TraceInsights/Private/Insights/MemoryProfiler/ViewModels/MemorySharedState.h
--- a/Engine/Source/Developer/TraceInsights/Private/Insights/MemoryProfiler/ViewModels/MemorySharedState.h
+++ b/Engine/Source/Developer/TraceInsights/Private/Insights/MemoryProfiler/ViewModels/MemorySharedState.h
@@ -3,6 +3,7 @@
 #pragma once
 
 #include "CoreMinimal.h"
+#include "Framework/Commands/Commands.h"
 #include "TraceServices/Model/AllocationsProvider.h"
 
 // Insights
@@ -23,6 +24,8 @@ struct FReportTypeConfig;
 struct FReportTypeGraphConfig;
 class FMemoryTracker;
 
+////////////////////////////////////////////////////////////////////////////////////////////////////
+
 class FMemoryRuleSpec
 {
 public:
@@ -51,6 +54,8 @@ private:
 	FText Description;     // ex.: "Allocations allocated and freed between time A and time B (A <= a <= f <= B)."
 };
 
+////////////////////////////////////////////////////////////////////////////////////////////////////
+
 class FQueryTargetWindowSpec
 {
 public:
@@ -70,10 +75,25 @@ private:
 	FName Name;
 };
 
+////////////////////////////////////////////////////////////////////////////////////////////////////
+
 } // namespace Insights
 
 ////////////////////////////////////////////////////////////////////////////////////////////////////
 
+class FMemoryTimingViewCommands : public TCommands<FMemoryTimingViewCommands>
+{
+public:
+	FMemoryTimingViewCommands();
+	virtual ~FMemoryTimingViewCommands();
+	virtual void RegisterCommands() override;
+
+public:
+	TSharedPtr<FUICommandInfo> ShowHideAllMemoryTracks;
+};
+
+////////////////////////////////////////////////////////////////////////////////////////////////////
+
 class FMemorySharedState : public Insights::ITimingViewExtender, public TSharedFromThis<FMemorySharedState>
 {
 public:
@@ -94,11 +114,17 @@ public:
 	EMemoryTrackHeightMode GetTrackHeightMode() const { return TrackHeightMode; }
 	void SetTrackHeightMode(EMemoryTrackHeightMode InTrackHeightMode);
 
-	// ITimingViewExtender
+	//////////////////////////////////////////////////
+	// ITimingViewExtender interface
+
 	virtual void OnBeginSession(Insights::ITimingViewSession& InSession) override;
 	virtual void OnEndSession(Insights::ITimingViewSession& InSession) override;
 	virtual void Tick(Insights::ITimingViewSession& InSession, const TraceServices::IAnalysisSession& InAnalysisSession) override;
-	virtual void ExtendFilterMenu(Insights::ITimingViewSession& InSession, FMenuBuilder& InMenuBuilder) override;
+	virtual void ExtendOtherTracksFilterMenu(Insights::ITimingViewSession& InSession, FMenuBuilder& InOutMenuBuilder) override;
+
+	//////////////////////////////////////////////////
+
+	void BindCommands();
 
 	bool IsAllMemoryTracksToggleOn() const { return bShowHideAllMemoryTracks; }
 	void SetAllMemoryTracksToggle(bool bOnOff);
diff --git a/Engine/Source/Developer/TraceInsights/Private/Insights/MemoryProfiler/Widgets/SMemoryProfilerWindow.cpp b/Engine/Source/Developer/TraceInsights/Private/Insights/MemoryProfiler/Widgets/SMemoryProfilerWindow.cpp
--- a/Engine/Source/Developer/TraceInsights/Private/Insights/MemoryProfiler/Widgets/SMemoryProfilerWindow.cpp
+++ b/Engine/Source/Developer/TraceInsights/Private/Insights/MemoryProfiler/Widgets/SMemoryProfilerWindow.cpp
@@ -249,6 +249,7 @@ TSharedRef<SDockTab> SMemoryProfilerWindow::SpawnTab_TimingView(const FSpawnTabA
 		];
 
 	SharedState->SetTimingView(TimingView);
+	SharedState->BindCommands();
 	IModularFeatures::Get().RegisterModularFeature(Insights::TimingViewExtenderFeatureName, &SharedState.Get());
 
 	TimingView->Reset(true);
diff --git a/Engine/Source/Developer/TraceInsights/Private/Insights/NetworkingProfiler/NetworkingProfilerCommands.cpp b/Engine/Source/Developer/TraceInsights/Private/Insights/NetworkingProfiler/NetworkingProfilerCommands.cpp
--- a/Engine/Source/Developer/TraceInsights/Private/Insights/NetworkingProfiler/NetworkingProfilerCommands.cpp
+++ b/Engine/Source/Developer/TraceInsights/Private/Insights/NetworkingProfiler/NetworkingProfilerCommands.cpp
@@ -51,9 +51,23 @@ FNetworkingProfilerCommands::FNetworkingProfilerCommands()
 PRAGMA_DISABLE_OPTIMIZATION
 void FNetworkingProfilerCommands::RegisterCommands()
 {
-	UI_COMMAND(TogglePacketViewVisibility, "Packets", "Toggles the visibility of the Packets view", EUserInterfaceActionType::ToggleButton, FInputChord());
-	UI_COMMAND(TogglePacketContentViewVisibility, "Packet Content", "Toggles the visibility of the Packet Content view", EUserInterfaceActionType::ToggleButton, FInputChord());
-	UI_COMMAND(ToggleNetStatsViewVisibility, "Net Stats", "Toggles the visibility of the Net Stats view", EUserInterfaceActionType::ToggleButton, FInputChord());
+	UI_COMMAND(TogglePacketViewVisibility,
+		"Packets",
+		"Toggles the visibility of the Packets view.",
+		EUserInterfaceActionType::ToggleButton,
+		FInputChord());
+
+	UI_COMMAND(TogglePacketContentViewVisibility,
+		"Packet Content",
+		"Toggles the visibility of the Packet Content view.",
+		EUserInterfaceActionType::ToggleButton,
+		FInputChord());
+
+	UI_COMMAND(ToggleNetStatsViewVisibility,
+		"Net Stats",
+		"Toggles the visibility of the Net Stats view.",
+		EUserInterfaceActionType::ToggleButton,
+		FInputChord());
 }
 PRAGMA_ENABLE_OPTIMIZATION
 
diff --git a/Engine/Source/Developer/TraceInsights/Private/Insights/TimingProfilerCommands.cpp b/Engine/Source/Developer/TraceInsights/Private/Insights/TimingProfilerCommands.cpp
--- a/Engine/Source/Developer/TraceInsights/Private/Insights/TimingProfilerCommands.cpp
+++ b/Engine/Source/Developer/TraceInsights/Private/Insights/TimingProfilerCommands.cpp
@@ -47,17 +47,57 @@ FTimingProfilerCommands::FTimingProfilerCommands()
 
 ////////////////////////////////////////////////////////////////////////////////////////////////////
 
+FTimingProfilerCommands::~FTimingProfilerCommands()
+{
+}
+
+////////////////////////////////////////////////////////////////////////////////////////////////////
+
 // UI_COMMAND takes long for the compiler to optimize
 PRAGMA_DISABLE_OPTIMIZATION
 void FTimingProfilerCommands::RegisterCommands()
 {
-	UI_COMMAND(ToggleFramesTrackVisibility, "Frames", "Toggles the visibility of the Frames track", EUserInterfaceActionType::ToggleButton, FInputChord());
-	UI_COMMAND(ToggleTimingViewVisibility, "Timing", "Toggles the visibility of the main Timing view", EUserInterfaceActionType::ToggleButton, FInputChord());
-	UI_COMMAND(ToggleTimersViewVisibility, "Timers", "Toggles the visibility of the Timers view", EUserInterfaceActionType::ToggleButton, FInputChord());
-	UI_COMMAND(ToggleCallersTreeViewVisibility, "Callers", "Toggles the visibility of the Callers tree view", EUserInterfaceActionType::ToggleButton, FInputChord());
-	UI_COMMAND(ToggleCalleesTreeViewVisibility, "Callees", "Toggles the visibility of the Callees tree view", EUserInterfaceActionType::ToggleButton, FInputChord());
-	UI_COMMAND(ToggleStatsCountersViewVisibility, "Counters", "Toggles the visibility of the Counters view", EUserInterfaceActionType::ToggleButton, FInputChord());
-	UI_COMMAND(ToggleLogViewVisibility, "Log", "Toggles the visibility of the Log view", EUserInterfaceActionType::ToggleButton, FInputChord());
+	UI_COMMAND(ToggleFramesTrackVisibility,
+		"Frames",
+		"Toggles the visibility of the Frames track.",
+		EUserInterfaceActionType::ToggleButton,
+		FInputChord());
+
+	UI_COMMAND(ToggleTimingViewVisibility,
+		"Timing",
+		"Toggles the visibility of the main Timing view.",
+		EUserInterfaceActionType::ToggleButton,
+		FInputChord());
+
+	UI_COMMAND(ToggleTimersViewVisibility,
+		"Timers",
+		"Toggles the visibility of the Timers view.",
+		EUserInterfaceActionType::ToggleButton,
+		FInputChord());
+
+	UI_COMMAND(ToggleCallersTreeViewVisibility,
+		"Callers",
+		"Toggles the visibility of the Callers tree view.",
+		EUserInterfaceActionType::ToggleButton,
+		FInputChord());
+
+	UI_COMMAND(ToggleCalleesTreeViewVisibility,
+		"Callees",
+		"Toggles the visibility of the Callees tree view.",
+		EUserInterfaceActionType::ToggleButton,
+		FInputChord());
+
+	UI_COMMAND(ToggleStatsCountersViewVisibility,
+		"Counters",
+		"Toggles the visibility of the Counters view.",
+		EUserInterfaceActionType::ToggleButton,
+		FInputChord());
+
+	UI_COMMAND(ToggleLogViewVisibility,
+		"Log",
+		"Toggles the visibility of the Log view.",
+		EUserInterfaceActionType::ToggleButton,
+		FInputChord());
 }
 PRAGMA_ENABLE_OPTIMIZATION
 
@@ -76,6 +116,12 @@ FTimingViewCommands::FTimingViewCommands()
 
 ////////////////////////////////////////////////////////////////////////////////////////////////////
 
+FTimingViewCommands::~FTimingViewCommands()
+{
+}
+
+////////////////////////////////////////////////////////////////////////////////////////////////////
+
 // UI_COMMAND takes long for the compiler to optimize
 PRAGMA_DISABLE_OPTIMIZATION
 void FTimingViewCommands::RegisterCommands()
@@ -100,7 +146,7 @@ void FTimingViewCommands::RegisterCommands()
 
 	UI_COMMAND(ShowMainGraphTrack,
 		"Graph Track",
-		"Show/hide the main Graph track.",
+		"Shows/hides the main Graph track.",
 		EUserInterfaceActionType::ToggleButton,
 		FInputChord(EKeys::G));
 
@@ -109,9 +155,6 @@ void FTimingViewCommands::RegisterCommands()
 		"Quick find or filter events in the timing view.",
 		EUserInterfaceActionType::Button,
 		FInputChord(EModifierKey::Control, EKeys::F));
-
-	//UI_COMMAND(ShowAllGpuTracks, "GPU Track(s)", "Show/hide the GPU track(s).", EUserInterfaceActionType::ToggleButton, FInputChord(EKeys::Y));
-	//UI_COMMAND(ShowAllCpuTracks, "CPU Thread Tracks", "Show/hide all CPU tracks (and all CPU thread groups).", EUserInterfaceActionType::ToggleButton, FInputChord(EKeys::U));
 }
 PRAGMA_ENABLE_OPTIMIZATION
 
diff --git a/Engine/Source/Developer/TraceInsights/Private/Insights/TimingProfilerCommands.h b/Engine/Source/Developer/TraceInsights/Private/Insights/TimingProfilerCommands.h
--- a/Engine/Source/Developer/TraceInsights/Private/Insights/TimingProfilerCommands.h
+++ b/Engine/Source/Developer/TraceInsights/Private/Insights/TimingProfilerCommands.h
@@ -15,10 +15,8 @@ class FMenuBuilder;
 class FTimingProfilerCommands : public TCommands<FTimingProfilerCommands>
 {
 public:
-	/** Default constructor. */
 	FTimingProfilerCommands();
-
-	/** Initialize commands. */
+	virtual ~FTimingProfilerCommands();
 	virtual void RegisterCommands() override;
 
 public:
@@ -56,10 +54,8 @@ public:
 class FTimingViewCommands : public TCommands<FTimingViewCommands>
 {
 public:
-	/** Default constructor. */
 	FTimingViewCommands();
-
-	/** Initialize commands. */
+	virtual ~FTimingViewCommands();
 	virtual void RegisterCommands() override;
 
 public:
@@ -77,12 +73,6 @@ public:
 
 	/** Opens the Quick Find widget. */
 	TSharedPtr<FUICommandInfo> QuickFind;
-
-	/** Toggles visibility for GPU thread track. */
-	//TSharedPtr<FUICommandInfo> ShowAllGpuTracks;
-
-	/** Toggles visibility for all CPU thread tracks at once. */
-	//TSharedPtr<FUICommandInfo> ShowAllCpuTracks;
 };
 
 ////////////////////////////////////////////////////////////////////////////////////////////////////
diff --git a/Engine/Source/Developer/TraceInsights/Private/Insights/ViewModels/FileActivityTimingTrack.cpp b/Engine/Source/Developer/TraceInsights/Private/Insights/ViewModels/FileActivityTimingTrack.cpp
--- a/Engine/Source/Developer/TraceInsights/Private/Insights/ViewModels/FileActivityTimingTrack.cpp
+++ b/Engine/Source/Developer/TraceInsights/Private/Insights/ViewModels/FileActivityTimingTrack.cpp
@@ -4,15 +4,19 @@
 
 #include "Algo/BinarySearch.h"
 #include "Fonts/FontMeasure.h"
+#include "Framework/Commands/Commands.h"
+#include "Framework/Commands/UICommandList.h"
 #include "Framework/MultiBox/MultiBoxBuilder.h"
 #include "Styling/SlateBrush.h"
 #include "TraceServices/AnalysisService.h"
 
 // Insights
+#include "Insights/Common/InsightsMenuBuilder.h"
 #include "Insights/Common/PaintUtils.h"
 #include "Insights/Common/Stopwatch.h"
 #include "Insights/Common/TimeUtils.h"
 #include "Insights/InsightsManager.h"
+#include "Insights/InsightsStyle.h"
 #include "Insights/ITimingViewSession.h"
 #include "Insights/TimingProfilerCommon.h"
 #include "Insights/ViewModels/TimingEvent.h"
@@ -25,9 +29,65 @@
 
 #define LOCTEXT_NAMESPACE "FileActivityTimingTrack"
 
-// The FileActivity (I/O) timelines are just prototypes for now.
-// Below code will be removed once the functionality is moved in analyzer.
+////////////////////////////////////////////////////////////////////////////////////////////////////
+// FFileActivityTimingViewCommands
+////////////////////////////////////////////////////////////////////////////////////////////////////
+
+FFileActivityTimingViewCommands::FFileActivityTimingViewCommands()
+: TCommands<FFileActivityTimingViewCommands>(
+	TEXT("FileActivityTimingViewCommands"),
+	NSLOCTEXT("Contexts", "FileActivityTimingViewCommands", "Insights - Timing View - File Activity"),
+	NAME_None,
+	FInsightsStyle::GetStyleSetName())
+{
+}
+
+////////////////////////////////////////////////////////////////////////////////////////////////////
+
+FFileActivityTimingViewCommands::~FFileActivityTimingViewCommands()
+{
+}
+
+////////////////////////////////////////////////////////////////////////////////////////////////////
+
+// UI_COMMAND takes long for the compiler to optimize
+PRAGMA_DISABLE_OPTIMIZATION
+void FFileActivityTimingViewCommands::RegisterCommands()
+{
+	// This command is used only for its key binding (to toggle both ShowHideIoOverviewTrack and ShowHideIoActivityTrack in the same time).
+	UI_COMMAND(ShowHideAllIoTracks,
+		"File Activity Tracks",
+		"Shows/hides the File Activity tracks.",
+		EUserInterfaceActionType::ToggleButton,
+		FInputChord(EKeys::I));
+
+	UI_COMMAND(ShowHideIoOverviewTrack,
+		"I/O Overview Track",
+		"Shows/hides the I/O Overview track.",
+		EUserInterfaceActionType::ToggleButton,
+		FInputChord());
+
+	UI_COMMAND(ToggleOnlyErrors,
+		"Only Errors (I/O Overview Track)",
+		"Shows only the events with errors, in the I/O Overview track.",
+		EUserInterfaceActionType::ToggleButton,
+		FInputChord());
+
+	UI_COMMAND(ShowHideIoActivityTrack,
+		"I/O Activity Track",
+		"Shows/hides the I/O Activity track.",
+		EUserInterfaceActionType::ToggleButton,
+		FInputChord());
+
+	UI_COMMAND(ToggleBackgroundEvents,
+		"Background Events (I/O Activity Track)",
+		"Shows/hides background events for file activities, in the I/O Activity track.",
+		EUserInterfaceActionType::ToggleButton,
+		FInputChord(EKeys::O));
+}
+PRAGMA_ENABLE_OPTIMIZATION
 
+////////////////////////////////////////////////////////////////////////////////////////////////////
 ////////////////////////////////////////////////////////////////////////////////////////////////////
 
 const TCHAR* GetFileActivityTypeName(TraceServices::EFileActivityType Type)
@@ -366,53 +426,78 @@ void FFileActivitySharedState::ExtendOtherTracksFilterMenu(Insights::ITimingView
 
 ////////////////////////////////////////////////////////////////////////////////////////////////////
 
+void FFileActivitySharedState::BindCommands()
+{
+	FFileActivityTimingViewCommands::Register();
+
+	TSharedPtr<FUICommandList> CommandList = TimingView->GetCommandList();
+	ensure(CommandList.IsValid());
+
+	// This command is used only for its key binding (to toggle both ShowHideIoOverviewTrack and ShowHideIoActivityTrack in the same time).
+	CommandList->MapAction(
+		FFileActivityTimingViewCommands::Get().ShowHideAllIoTracks,
+		FExecuteAction::CreateSP(this, &FFileActivitySharedState::ShowHideAllIoTracks),
+		FCanExecuteAction(),
+		FIsActionChecked::CreateSP(this, &FFileActivitySharedState::IsAllIoTracksToggleOn));
+
+	CommandList->MapAction(
+		FFileActivityTimingViewCommands::Get().ShowHideIoOverviewTrack,
+		FExecuteAction::CreateSP(this, &FFileActivitySharedState::ShowHideIoOverviewTrack),
+		FCanExecuteAction(),
+		FIsActionChecked::CreateSP(this, &FFileActivitySharedState::IsIoOverviewTrackVisible));
+
+	CommandList->MapAction(
+		FFileActivityTimingViewCommands::Get().ToggleOnlyErrors,
+		FExecuteAction::CreateSP(this, &FFileActivitySharedState::ToggleOnlyErrors),
+		FCanExecuteAction(),
+		FIsActionChecked::CreateSP(this, &FFileActivitySharedState::IsOnlyErrorsToggleOn));
+
+	CommandList->MapAction(
+		FFileActivityTimingViewCommands::Get().ShowHideIoActivityTrack,
+		FExecuteAction::CreateSP(this, &FFileActivitySharedState::ShowHideIoActivityTrack),
+		FCanExecuteAction(),
+		FIsActionChecked::CreateSP(this, &FFileActivitySharedState::IsIoActivityTrackVisible));
+
+	CommandList->MapAction(
+		FFileActivityTimingViewCommands::Get().ToggleBackgroundEvents,
+		FExecuteAction::CreateSP(this, &FFileActivitySharedState::ToggleBackgroundEvents),
+		FCanExecuteAction(),
+		FIsActionChecked::CreateSP(this, &FFileActivitySharedState::AreBackgroundEventsVisible));
+}
+
+////////////////////////////////////////////////////////////////////////////////////////////////////
+
 void FFileActivitySharedState::BuildSubMenu(FMenuBuilder& InOutMenuBuilder)
 {
 	InOutMenuBuilder.BeginSection("File Activity", LOCTEXT("ContextMenu_Section_FileActivity", "File Activity"));
 	{
-		InOutMenuBuilder.AddMenuEntry(
-			LOCTEXT("ShowIoOverviewTrack", "I/O Overview Track - I"),
-			LOCTEXT("ShowIoOverviewTrack_Tooltip", "Shows or hides the I/O Overview track."),
-			FSlateIcon(),
-			FUIAction(FExecuteAction::CreateSP(this, &FFileActivitySharedState::ShowHideIoOverviewTrack),
+		// Note: We use the custom AddMenuEntry in order to set the same key binding text for multiple menu items.
+
+		//InOutMenuBuilder.AddMenuEntry(FFileActivityTimingViewCommands::Get().ShowHideIoOverviewTrack);
+		FInsightsMenuBuilder::AddMenuEntry(InOutMenuBuilder,
+			FUIAction(
+				FExecuteAction::CreateSP(this, &FFileActivitySharedState::ShowHideIoOverviewTrack),
 				FCanExecuteAction(),
 				FIsActionChecked::CreateSP(this, &FFileActivitySharedState::IsIoOverviewTrackVisible)),
-			NAME_None,
-			EUserInterfaceActionType::ToggleButton
-		);
+			FFileActivityTimingViewCommands::Get().ShowHideIoOverviewTrack->GetLabel(),
+			FFileActivityTimingViewCommands::Get().ShowHideIoOverviewTrack->GetDescription(),
+			LOCTEXT("FileActivityTracksKeybinding", "I"),
+			EUserInterfaceActionType::ToggleButton);
 
-		InOutMenuBuilder.AddMenuEntry(
-			LOCTEXT("ShowOnlyErrors", "Only Errors (I/O Overview Track)"),
-			LOCTEXT("ShowOnlyErrors_Tooltip", "Shows only the events with errors, in the I/O Overview track."),
-			FSlateIcon(),
-			FUIAction(FExecuteAction::CreateSP(this, &FFileActivitySharedState::ToggleOnlyErrors),
-				FCanExecuteAction(),
-				FIsActionChecked::CreateSP(this, &FFileActivitySharedState::IsOnlyErrorsToggleOn)),
-			NAME_None,
-			EUserInterfaceActionType::ToggleButton
-		);
+		InOutMenuBuilder.AddMenuEntry(FFileActivityTimingViewCommands::Get().ToggleOnlyErrors);
 
-		InOutMenuBuilder.AddMenuEntry(
-			LOCTEXT("ShowIoActivityTrack", "I/O Activity Track - I"),
-			LOCTEXT("ShowIoActivityTrack_Tooltip", "Shows or hides the I/O Activity track."),
-			FSlateIcon(),
-			FUIAction(FExecuteAction::CreateSP(this, &FFileActivitySharedState::ShowHideIoActivityTrack),
+		//InOutMenuBuilder.AddMenuEntry(FFileActivityTimingViewCommands::Get().ShowHideIoActivityTrack);
+		FInsightsMenuBuilder::AddMenuEntry(InOutMenuBuilder,
+			FUIAction(
+				FExecuteAction::CreateSP(this, &FFileActivitySharedState::ShowHideIoActivityTrack),
 				FCanExecuteAction(),
 				FIsActionChecked::CreateSP(this, &FFileActivitySharedState::IsIoActivityTrackVisible)),
-			NAME_None,
-			EUserInterfaceActionType::ToggleButton
-		);
+			FFileActivityTimingViewCommands::Get().ShowHideIoActivityTrack->GetLabel(),
+			FFileActivityTimingViewCommands::Get().ShowHideIoActivityTrack->GetDescription(),
+			LOCTEXT("FileActivityTracksKeybinding", "I"),
+			EUserInterfaceActionType::ToggleButton);
 
-		InOutMenuBuilder.AddMenuEntry(
-			LOCTEXT("ShowBackgroundEvents", "Background Events (I/O Activity Track) - O"),
-			LOCTEXT("ShowBackgroundEvents_Tooltip", "Shows background events for file activities, in the I/O Activity track."),
-			FSlateIcon(),
-			FUIAction(FExecuteAction::CreateSP(this, &FFileActivitySharedState::ToggleBackgroundEvents),
-				FCanExecuteAction(),
-				FIsActionChecked::CreateSP(this, &FFileActivitySharedState::AreBackgroundEventsVisible)),
-			NAME_None,
-			EUserInterfaceActionType::ToggleButton
-		);
+		InOutMenuBuilder.AddMenuEntry(FFileActivityTimingViewCommands::Get().ToggleBackgroundEvents);
 	}
 	InOutMenuBuilder.EndSection();
 }
diff --git a/Engine/Source/Developer/TraceInsights/Private/Insights/ViewModels/FileActivityTimingTrack.h b/Engine/Source/Developer/TraceInsights/Private/Insights/ViewModels/FileActivityTimingTrack.h
--- a/Engine/Source/Developer/TraceInsights/Private/Insights/ViewModels/FileActivityTimingTrack.h
+++ b/Engine/Source/Developer/TraceInsights/Private/Insights/ViewModels/FileActivityTimingTrack.h
@@ -3,6 +3,7 @@
 #pragma once
 
 #include "CoreMinimal.h"
+#include "Framework/Commands/Commands.h"
 
 // Insights
 #include "Insights/ITimingViewExtender.h"
@@ -16,6 +17,23 @@ class FDetailedFileActivityTimingTrack;
 
 ////////////////////////////////////////////////////////////////////////////////////////////////////
 
+class FFileActivityTimingViewCommands : public TCommands<FFileActivityTimingViewCommands>
+{
+public:
+	FFileActivityTimingViewCommands();
+	virtual ~FFileActivityTimingViewCommands();
+	virtual void RegisterCommands() override;
+
+public:
+	TSharedPtr<FUICommandInfo> ShowHideAllIoTracks;
+	TSharedPtr<FUICommandInfo> ShowHideIoOverviewTrack;
+	TSharedPtr<FUICommandInfo> ToggleOnlyErrors;
+	TSharedPtr<FUICommandInfo> ShowHideIoActivityTrack;
+	TSharedPtr<FUICommandInfo> ToggleBackgroundEvents;
+};
+
+////////////////////////////////////////////////////////////////////////////////////////////////////
+
 class FFileActivitySharedState : public Insights::ITimingViewExtender, public TSharedFromThis<FFileActivitySharedState>
 {
 	friend class FOverviewFileActivityTimingTrack;
@@ -55,12 +73,18 @@ public:
 	explicit FFileActivitySharedState(STimingView* InTimingView) : TimingView(InTimingView) {}
 	virtual ~FFileActivitySharedState() = default;
 
-	// ITimingViewExtender
+	//////////////////////////////////////////////////
+	// ITimingViewExtender interface
+
 	virtual void OnBeginSession(Insights::ITimingViewSession& InSession) override;
 	virtual void OnEndSession(Insights::ITimingViewSession& InSession) override;
 	virtual void Tick(Insights::ITimingViewSession& InSession, const TraceServices::IAnalysisSession& InAnalysisSession) override;
 	virtual void ExtendOtherTracksFilterMenu(Insights::ITimingViewSession& InSession, FMenuBuilder& InOutMenuBuilder) override;
 
+	//////////////////////////////////////////////////
+
+	void BindCommands();
+
 	const TArray<FIoTimingEvent>& GetAllEvents() const { return AllIoEvents; }
 
 	void RequestUpdate() { bForceIoEventsUpdate = true; }
diff --git a/Engine/Source/Developer/TraceInsights/Private/Insights/ViewModels/FrameTimingTrack.cpp b/Engine/Source/Developer/TraceInsights/Private/Insights/ViewModels/FrameTimingTrack.cpp
--- a/Engine/Source/Developer/TraceInsights/Private/Insights/ViewModels/FrameTimingTrack.cpp
+++ b/Engine/Source/Developer/TraceInsights/Private/Insights/ViewModels/FrameTimingTrack.cpp
@@ -3,6 +3,8 @@
 #include "FrameTimingTrack.h"
 
 #include "Fonts/FontMeasure.h"
+#include "Framework/Commands/Commands.h"
+#include "Framework/Commands/UICommandList.h"
 #include "Framework/MultiBox/MultiBoxBuilder.h"
 #include "HAL/PlatformApplicationMisc.h"
 #include "Styling/SlateBrush.h"
@@ -11,6 +13,7 @@
 #include "Insights/Common/PaintUtils.h"
 #include "Insights/Common/TimeUtils.h"
 #include "Insights/InsightsManager.h"
+#include "Insights/InsightsStyle.h"
 #include "Insights/ITimingViewSession.h"
 #include "Insights/ViewModels/FrameTrackHelper.h"
 #include "Insights/ViewModels/TimerNode.h"
@@ -22,6 +25,39 @@
 
 #define LOCTEXT_NAMESPACE "FrameTimingTrack"
 
+////////////////////////////////////////////////////////////////////////////////////////////////////
+// FFrameTimingViewCommands
+////////////////////////////////////////////////////////////////////////////////////////////////////
+
+FFrameTimingViewCommands::FFrameTimingViewCommands()
+: TCommands<FFrameTimingViewCommands>(
+	TEXT("FrameTimingViewCommands"),
+	NSLOCTEXT("Contexts", "FrameTimingViewCommands", "Insights - Timing View - Frames"),
+	NAME_None,
+	FInsightsStyle::GetStyleSetName())
+{
+}
+
+////////////////////////////////////////////////////////////////////////////////////////////////////
+
+FFrameTimingViewCommands::~FFrameTimingViewCommands()
+{
+}
+
+////////////////////////////////////////////////////////////////////////////////////////////////////
+
+// UI_COMMAND takes long for the compiler to optimize
+PRAGMA_DISABLE_OPTIMIZATION
+void FFrameTimingViewCommands::RegisterCommands()
+{
+	UI_COMMAND(ShowHideAllFrameTracks,
+		"Frame Tracks",
+		"Shows/hides all Frame tracks.",
+		EUserInterfaceActionType::ToggleButton,
+		FInputChord(EKeys::R));
+}
+PRAGMA_ENABLE_OPTIMIZATION
+
 ////////////////////////////////////////////////////////////////////////////////////////////////////
 // FFrameSharedState
 ////////////////////////////////////////////////////////////////////////////////////////////////////
@@ -125,23 +161,29 @@ void FFrameSharedState::ExtendOtherTracksFilterMenu(Insights::ITimingViewSession
 
 	InOutMenuBuilder.BeginSection("FrameTracks", LOCTEXT("ContextMenu_Section_FrameTracks", "Frames"));
 	{
-		//TODO: MenuBuilder.AddMenuEntry(Commands.ShowAllFrameTracks);
-		InOutMenuBuilder.AddMenuEntry(
-			LOCTEXT("ShowAllFrameTracks", "Frame Tracks - R"),
-			LOCTEXT("ShowAllFrameTracks_Tooltip", "Show/hide all Frame tracks."),
-			FSlateIcon(),
-			FUIAction(FExecuteAction::CreateSP(this, &FFrameSharedState::ShowHideAllFrameTracks),
-					  FCanExecuteAction(),
-					  FIsActionChecked::CreateSP(this, &FFrameSharedState::IsAllFrameTracksToggleOn)),
-			NAME_None, //"QuickFilterSeparator",
-			EUserInterfaceActionType::ToggleButton
-		);
+		InOutMenuBuilder.AddMenuEntry(FFrameTimingViewCommands::Get().ShowHideAllFrameTracks);
 	}
 	InOutMenuBuilder.EndSection();
 }
 
 ////////////////////////////////////////////////////////////////////////////////////////////////////
 
+void FFrameSharedState::BindCommands()
+{
+	FFrameTimingViewCommands::Register();
+
+	TSharedPtr<FUICommandList> CommandList = TimingView->GetCommandList();
+	ensure(CommandList.IsValid());
+
+	CommandList->MapAction(
+		FFrameTimingViewCommands::Get().ShowHideAllFrameTracks,
+		FExecuteAction::CreateSP(this, &FFrameSharedState::ShowHideAllFrameTracks),
+		FCanExecuteAction(),
+		FIsActionChecked::CreateSP(this, &FFrameSharedState::IsAllFrameTracksToggleOn));
+}
+
+////////////////////////////////////////////////////////////////////////////////////////////////////
+
 void FFrameSharedState::SetAllFrameTracksToggle(bool bOnOff)
 {
 	bShowHideAllFrameTracks = bOnOff;
diff --git a/Engine/Source/Developer/TraceInsights/Private/Insights/ViewModels/FrameTimingTrack.h b/Engine/Source/Developer/TraceInsights/Private/Insights/ViewModels/FrameTimingTrack.h
--- a/Engine/Source/Developer/TraceInsights/Private/Insights/ViewModels/FrameTimingTrack.h
+++ b/Engine/Source/Developer/TraceInsights/Private/Insights/ViewModels/FrameTimingTrack.h
@@ -4,6 +4,7 @@
 
 #include "CoreMinimal.h"
 #include "Fonts/SlateFontInfo.h"
+#include "Framework/Commands/Commands.h"
 #include "TraceServices/Model/Frames.h"
 
 // Insights
@@ -20,6 +21,19 @@ struct FSlateBrush;
 
 ////////////////////////////////////////////////////////////////////////////////////////////////////
 
+class FFrameTimingViewCommands : public TCommands<FFrameTimingViewCommands>
+{
+public:
+	FFrameTimingViewCommands();
+	virtual ~FFrameTimingViewCommands();
+	virtual void RegisterCommands() override;
+
+public:
+	TSharedPtr<FUICommandInfo> ShowHideAllFrameTracks;
+};
+
+////////////////////////////////////////////////////////////////////////////////////////////////////
+
 class FFrameSharedState : public Insights::ITimingViewExtender, public TSharedFromThis<FFrameSharedState>
 {
 public:
@@ -40,6 +54,8 @@ public:
 
 	//////////////////////////////////////////////////
 
+	void BindCommands();
+
 	bool IsAllFrameTracksToggleOn() const { return bShowHideAllFrameTracks; }
 	void SetAllFrameTracksToggle(bool bOnOff);
 	void ShowAllFrameTracks() { SetAllFrameTracksToggle(true); }
diff --git a/Engine/Source/Developer/TraceInsights/Private/Insights/ViewModels/LoadingTimingTrack.cpp b/Engine/Source/Developer/TraceInsights/Private/Insights/ViewModels/LoadingTimingTrack.cpp
--- a/Engine/Source/Developer/TraceInsights/Private/Insights/ViewModels/LoadingTimingTrack.cpp
+++ b/Engine/Source/Developer/TraceInsights/Private/Insights/ViewModels/LoadingTimingTrack.cpp
@@ -3,6 +3,8 @@
 #include "LoadingTimingTrack.h"
 
 #include "Fonts/FontMeasure.h"
+#include "Framework/Commands/Commands.h"
+#include "Framework/Commands/UICommandList.h"
 #include "Framework/MultiBox/MultiBoxBuilder.h"
 #include "Styling/SlateBrush.h"
 #include "TraceServices/AnalysisService.h"
@@ -11,6 +13,7 @@
 #include "Insights/Common/PaintUtils.h"
 #include "Insights/Common/TimeUtils.h"
 #include "Insights/InsightsManager.h"
+#include "Insights/InsightsStyle.h"
 #include "Insights/ITimingViewSession.h"
 #include "Insights/ViewModels/TimingEvent.h"
 #include "Insights/ViewModels/TimingTrackViewport.h"
@@ -20,6 +23,39 @@
 
 #define LOCTEXT_NAMESPACE "LoadingTimingTrack"
 
+////////////////////////////////////////////////////////////////////////////////////////////////////
+// FLoadingTimingViewCommands
+////////////////////////////////////////////////////////////////////////////////////////////////////
+
+FLoadingTimingViewCommands::FLoadingTimingViewCommands()
+: TCommands<FLoadingTimingViewCommands>(
+	TEXT("LoadingTimingViewCommands"),
+	NSLOCTEXT("Contexts", "LoadingTimingViewCommands", "Insights - Timing View - Asset Loading"),
+	NAME_None,
+	FInsightsStyle::GetStyleSetName())
+{
+}
+
+////////////////////////////////////////////////////////////////////////////////////////////////////
+
+FLoadingTimingViewCommands::~FLoadingTimingViewCommands()
+{
+}
+
+////////////////////////////////////////////////////////////////////////////////////////////////////
+
+// UI_COMMAND takes long for the compiler to optimize
+PRAGMA_DISABLE_OPTIMIZATION
+void FLoadingTimingViewCommands::RegisterCommands()
+{
+	UI_COMMAND(ShowHideAllLoadingTracks,
+		"Asset Loading Tracks",
+		"Shows/hides the Asset Loading tracks.",
+		EUserInterfaceActionType::ToggleButton,
+		FInputChord(EKeys::L));
+}
+PRAGMA_ENABLE_OPTIMIZATION
+
 ////////////////////////////////////////////////////////////////////////////////////////////////////
 // FLoadingSharedState
 ////////////////////////////////////////////////////////////////////////////////////////////////////
@@ -120,22 +156,29 @@ void FLoadingSharedState::ExtendOtherTracksFilterMenu(Insights::ITimingViewSessi
 
 	InOutMenuBuilder.BeginSection("Asset Loading", LOCTEXT("ContextMenu_Section_AssetLoading", "Asset Loading"));
 	{
-		InOutMenuBuilder.AddMenuEntry(
-			LOCTEXT("ShowAllLoadingTracks", "Asset Loading Tracks - L"),
-			LOCTEXT("ShowAllLoadingTracks_Tooltip", "Shows or hides the Asset Loading tracks."),
-			FSlateIcon(),
-			FUIAction(FExecuteAction::CreateSP(this, &FLoadingSharedState::ShowHideAllLoadingTracks),
-					  FCanExecuteAction(),
-					  FIsActionChecked::CreateSP(this, &FLoadingSharedState::IsAllLoadingTracksToggleOn)),
-			NAME_None, //"QuickFilterSeparator",
-			EUserInterfaceActionType::ToggleButton
-		);
+		InOutMenuBuilder.AddMenuEntry(FLoadingTimingViewCommands::Get().ShowHideAllLoadingTracks);
 	}
 	InOutMenuBuilder.EndSection();
 }
 
 ////////////////////////////////////////////////////////////////////////////////////////////////////
 
+void FLoadingSharedState::BindCommands()
+{
+	FLoadingTimingViewCommands::Register();
+
+	TSharedPtr<FUICommandList> CommandList = TimingView->GetCommandList();
+	ensure(CommandList.IsValid());
+
+	CommandList->MapAction(
+		FLoadingTimingViewCommands::Get().ShowHideAllLoadingTracks,
+		FExecuteAction::CreateSP(this, &FLoadingSharedState::ShowHideAllLoadingTracks),
+		FCanExecuteAction(),
+		FIsActionChecked::CreateSP(this, &FLoadingSharedState::IsAllLoadingTracksToggleOn));
+}
+
+////////////////////////////////////////////////////////////////////////////////////////////////////
+
 void FLoadingSharedState::SetAllLoadingTracksToggle(bool bOnOff)
 {
 	bShowHideAllLoadingTracks = bOnOff;
diff --git a/Engine/Source/Developer/TraceInsights/Private/Insights/ViewModels/LoadingTimingTrack.h b/Engine/Source/Developer/TraceInsights/Private/Insights/ViewModels/LoadingTimingTrack.h
--- a/Engine/Source/Developer/TraceInsights/Private/Insights/ViewModels/LoadingTimingTrack.h
+++ b/Engine/Source/Developer/TraceInsights/Private/Insights/ViewModels/LoadingTimingTrack.h
@@ -3,6 +3,7 @@
 #pragma once
 
 #include "CoreMinimal.h"
+#include "Framework/Commands/Commands.h"
 #include "TraceServices/AnalysisService.h"
 
 // Insights
@@ -11,8 +12,24 @@
 
 class FTimingEventSearchParameters;
 class STimingView;
+
 class FLoadingTimingTrack;
 
+////////////////////////////////////////////////////////////////////////////////////////////////////
+
+class FLoadingTimingViewCommands : public TCommands<FLoadingTimingViewCommands>
+{
+public:
+	FLoadingTimingViewCommands();
+	virtual ~FLoadingTimingViewCommands();
+	virtual void RegisterCommands() override;
+
+public:
+	TSharedPtr<FUICommandInfo> ShowHideAllLoadingTracks;
+};
+
+////////////////////////////////////////////////////////////////////////////////////////////////////
+
 /** Defines FLoadingTrackGetEventNameDelegate delegate interface. Returns the name for a timing event in a Loading track. */
 DECLARE_DELEGATE_RetVal_TwoParams(const TCHAR*, FLoadingTrackGetEventNameDelegate, uint32 /*Depth*/, const TraceServices::FLoadTimeProfilerCpuEvent& /*Event*/);
 
@@ -24,12 +41,18 @@ public:
 	explicit FLoadingSharedState(STimingView* InTimingView) : TimingView(InTimingView) {}
 	virtual ~FLoadingSharedState() = default;
 
-	// ITimingViewExtender
+	//////////////////////////////////////////////////
+	// ITimingViewExtender interface
+
 	virtual void OnBeginSession(Insights::ITimingViewSession& InSession) override;
 	virtual void OnEndSession(Insights::ITimingViewSession& InSession) override;
 	virtual void Tick(Insights::ITimingViewSession& InSession, const TraceServices::IAnalysisSession& InAnalysisSession) override;
 	virtual void ExtendOtherTracksFilterMenu(Insights::ITimingViewSession& InSession, FMenuBuilder& InMenuBuilder) override;
 
+	//////////////////////////////////////////////////
+
+	void BindCommands();
+
 	const TCHAR* GetEventName(uint32 Depth, const TraceServices::FLoadTimeProfilerCpuEvent& Event) const;
 	void SetColorSchema(int32 Schema);
 
diff --git a/Engine/Source/Developer/TraceInsights/Private/Insights/ViewModels/ThreadTimingTrack.cpp b/Engine/Source/Developer/TraceInsights/Private/Insights/ViewModels/ThreadTimingTrack.cpp
--- a/Engine/Source/Developer/TraceInsights/Private/Insights/ViewModels/ThreadTimingTrack.cpp
+++ b/Engine/Source/Developer/TraceInsights/Private/Insights/ViewModels/ThreadTimingTrack.cpp
@@ -3,6 +3,8 @@
 #include "ThreadTimingTrack.h"
 
 #include "CborReader.h"
+#include "Framework/Commands/Commands.h"
+#include "Framework/Commands/UICommandList.h"
 #include "Framework/MultiBox/MultiBoxBuilder.h"
 #include "HAL/PlatformApplicationMisc.h"
 #include "Serialization/MemoryReader.h"
@@ -31,6 +33,46 @@
 
 using namespace Insights;
 
+////////////////////////////////////////////////////////////////////////////////////////////////////
+// FThreadTimingViewCommands
+////////////////////////////////////////////////////////////////////////////////////////////////////
+
+FThreadTimingViewCommands::FThreadTimingViewCommands()
+: TCommands<FThreadTimingViewCommands>(
+	TEXT("ThreadTimingViewCommands"),
+	NSLOCTEXT("Contexts", "ThreadTimingViewCommands", "Insights - Timing View - Threads"),
+	NAME_None,
+	FInsightsStyle::GetStyleSetName())
+{
+}
+
+////////////////////////////////////////////////////////////////////////////////////////////////////
+
+FThreadTimingViewCommands::~FThreadTimingViewCommands()
+{
+}
+
+////////////////////////////////////////////////////////////////////////////////////////////////////
+
+// UI_COMMAND takes long for the compiler to optimize
+PRAGMA_DISABLE_OPTIMIZATION
+void FThreadTimingViewCommands::RegisterCommands()
+{
+	UI_COMMAND(ShowHideAllGpuTracks,
+		"GPU Track(s)",
+		"Shows/hides the GPU track(s).",
+		EUserInterfaceActionType::ToggleButton,
+		FInputChord(EKeys::Y));
+
+	UI_COMMAND(ShowHideAllCpuTracks,
+		"CPU Thread Tracks",
+		"Shows/hides all CPU tracks (and all CPU thread groups).",
+		EUserInterfaceActionType::ToggleButton,
+		FInputChord(EKeys::U));
+}
+PRAGMA_ENABLE_OPTIMIZATION
+
+////////////////////////////////////////////////////////////////////////////////////////////////////
 ////////////////////////////////////////////////////////////////////////////////////////////////////
 
 static void AppendMetadataToTooltip(FTooltipDrawState& Tooltip, TArrayView<const uint8>& Metadata)
@@ -546,17 +588,7 @@ void FThreadTimingSharedState::ExtendGpuTracksFilterMenu(Insights::ITimingViewSe
 
 	InOutMenuBuilder.BeginSection("GpuTracks", LOCTEXT("ContextMenu_Section_GpuTracks", "GPU Tracks"));
 	{
-		//TODO: MenuBuilder.AddMenuEntry(Commands.ShowAllGpuTracks);
-		InOutMenuBuilder.AddMenuEntry(
-			LOCTEXT("ShowAllGpuTracks", "GPU Track(s) - Y"),
-			LOCTEXT("ShowAllGpuTracks_Tooltip", "Show/hide the GPU track(s)."),
-			FSlateIcon(),
-			FUIAction(FExecuteAction::CreateSP(this, &FThreadTimingSharedState::ShowHideAllGpuTracks),
-					  FCanExecuteAction(),
-					  FIsActionChecked::CreateSP(this, &FThreadTimingSharedState::IsAllGpuTracksToggleOn)),
-			NAME_None,
-			EUserInterfaceActionType::ToggleButton
-		);
+		InOutMenuBuilder.AddMenuEntry(FThreadTimingViewCommands::Get().ShowHideAllGpuTracks);
 	}
 	InOutMenuBuilder.EndSection();
 }
@@ -572,17 +604,7 @@ void FThreadTimingSharedState::ExtendCpuTracksFilterMenu(Insights::ITimingViewSe
 
 	InOutMenuBuilder.BeginSection("CpuTracks", LOCTEXT("ContextMenu_Section_CpuTracks", "CPU Tracks"));
 	{
-		//TODO: MenuBuilder.AddMenuEntry(Commands.ShowAllCpuTracks);
-		InOutMenuBuilder.AddMenuEntry(
-			LOCTEXT("ShowAllCpuTracks", "CPU Thread Tracks - U"),
-			LOCTEXT("ShowAllCpuTracks_Tooltip", "Show/hide all CPU tracks (and all CPU thread groups)."),
-			FSlateIcon(),
-			FUIAction(FExecuteAction::CreateSP(this, &FThreadTimingSharedState::ShowHideAllCpuTracks),
-					  FCanExecuteAction(),
-					  FIsActionChecked::CreateSP(this, &FThreadTimingSharedState::IsAllCpuTracksToggleOn)),
-			NAME_None,
-			EUserInterfaceActionType::ToggleButton
-		);
+		InOutMenuBuilder.AddMenuEntry(FThreadTimingViewCommands::Get().ShowHideAllCpuTracks);
 	}
 	InOutMenuBuilder.EndSection();
 
@@ -593,6 +615,28 @@ void FThreadTimingSharedState::ExtendCpuTracksFilterMenu(Insights::ITimingViewSe
 
 ////////////////////////////////////////////////////////////////////////////////////////////////////
 
+void FThreadTimingSharedState::BindCommands()
+{
+	FThreadTimingViewCommands::Register();
+
+	TSharedPtr<FUICommandList> CommandList = TimingView->GetCommandList();
+	ensure(CommandList.IsValid());
+
+	CommandList->MapAction(
+		FThreadTimingViewCommands::Get().ShowHideAllGpuTracks,
+		FExecuteAction::CreateSP(this, &FThreadTimingSharedState::ShowHideAllGpuTracks),
+		FCanExecuteAction(),
+		FIsActionChecked::CreateSP(this, &FThreadTimingSharedState::IsAllGpuTracksToggleOn));
+
+	CommandList->MapAction(
+		FThreadTimingViewCommands::Get().ShowHideAllCpuTracks,
+		FExecuteAction::CreateSP(this, &FThreadTimingSharedState::ShowHideAllCpuTracks),
+		FCanExecuteAction(),
+		FIsActionChecked::CreateSP(this, &FThreadTimingSharedState::IsAllCpuTracksToggleOn));
+}
+
+////////////////////////////////////////////////////////////////////////////////////////////////////
+
 void FThreadTimingSharedState::CreateThreadGroupsMenu(FMenuBuilder& InOutMenuBuilder)
 {
 	// Sort the list of thread groups.
diff --git a/Engine/Source/Developer/TraceInsights/Private/Insights/ViewModels/ThreadTimingTrack.h b/Engine/Source/Developer/TraceInsights/Private/Insights/ViewModels/ThreadTimingTrack.h
--- a/Engine/Source/Developer/TraceInsights/Private/Insights/ViewModels/ThreadTimingTrack.h
+++ b/Engine/Source/Developer/TraceInsights/Private/Insights/ViewModels/ThreadTimingTrack.h
@@ -3,6 +3,7 @@
 #pragma once
 
 #include "CoreMinimal.h"
+#include "Framework/Commands/Commands.h"
 #include "TraceServices/Model/TimingProfiler.h"
 
 // Insights
@@ -23,6 +24,23 @@ namespace Insights
 
 ////////////////////////////////////////////////////////////////////////////////////////////////////
 
+class FThreadTimingViewCommands : public TCommands<FThreadTimingViewCommands>
+{
+public:
+	FThreadTimingViewCommands();
+	virtual ~FThreadTimingViewCommands();
+	virtual void RegisterCommands() override;
+
+public:
+	/** Toggles visibility for GPU thread track. */
+	TSharedPtr<FUICommandInfo> ShowHideAllGpuTracks;
+
+	/** Toggles visibility for all CPU thread tracks at once. */
+	TSharedPtr<FUICommandInfo> ShowHideAllCpuTracks;
+};
+
+////////////////////////////////////////////////////////////////////////////////////////////////////
+
 class FThreadTimingSharedState : public Insights::ITimingViewExtender, public TSharedFromThis<FThreadTimingSharedState>
 {
 private:
@@ -61,6 +79,8 @@ public:
 
 	//////////////////////////////////////////////////
 
+	void BindCommands();
+
 	bool IsAllGpuTracksToggleOn() const { return bShowHideAllGpuTracks; }
 	void SetAllGpuTracksToggle(bool bOnOff);
 	void ShowAllGpuTracks() { SetAllGpuTracksToggle(true); }
diff --git a/Engine/Source/Developer/TraceInsights/Private/Insights/Widgets/SFrameTrack.cpp b/Engine/Source/Developer/TraceInsights/Private/Insights/Widgets/SFrameTrack.cpp
--- a/Engine/Source/Developer/TraceInsights/Private/Insights/Widgets/SFrameTrack.cpp
+++ b/Engine/Source/Developer/TraceInsights/Private/Insights/Widgets/SFrameTrack.cpp
@@ -1150,7 +1150,7 @@ void SFrameTrack::ShowContextMenu(const FPointerEvent& MouseEvent)
 		MenuBuilder.AddMenuEntry
 		(
 			LOCTEXT("ContextMenu_ShowGameFrames", "Game Frames"),
-			LOCTEXT("ContextMenu_ShowGameFrames_Desc", "Shows or hides the Game frames."),
+			LOCTEXT("ContextMenu_ShowGameFrames_Desc", "Shows/hides the Game frames."),
 			FSlateIcon(),
 			Action_ShowGameFrames,
 			NAME_None,
@@ -1166,7 +1166,7 @@ void SFrameTrack::ShowContextMenu(const FPointerEvent& MouseEvent)
 		MenuBuilder.AddMenuEntry
 		(
 			LOCTEXT("ContextMenu_ShowRenderingFrames", "Rendering Frames"),
-			LOCTEXT("ContextMenu_ShowRenderingFrames_Desc", "Shows or hides the Rendering frames."),
+			LOCTEXT("ContextMenu_ShowRenderingFrames_Desc", "Shows/hides the Rendering frames."),
 			FSlateIcon(),
 			Action_ShowRenderingFrames,
 			NAME_None,
diff --git a/Engine/Source/Developer/TraceInsights/Private/Insights/Widgets/STimingView.cpp b/Engine/Source/Developer/TraceInsights/Private/Insights/Widgets/STimingView.cpp
--- a/Engine/Source/Developer/TraceInsights/Private/Insights/Widgets/STimingView.cpp
+++ b/Engine/Source/Developer/TraceInsights/Private/Insights/Widgets/STimingView.cpp
@@ -36,6 +36,7 @@
 #include "Widgets/Text/STextBlock.h"
 
 // Insights
+#include "Insights/Common/InsightsMenuBuilder.h"
 #include "Insights/Common/PaintUtils.h"
 #include "Insights/Common/Stopwatch.h"
 #include "Insights/Common/TimeUtils.h"
@@ -2728,6 +2729,7 @@ FReply STimingView::OnKeyDown(const FGeometry& MyGeometry, const FKeyEvent& InKe
 	}
 	else if (InKeyEvent.GetKey() == EKeys::Enter)
 	{
+		// Enter: Selects the time range of the currently selected timing event.
 		if (SelectedEvent.IsValid())
 		{
 			const double Duration = Viewport.RestrictDuration(SelectedEvent->GetStartTime(), SelectedEvent->GetEndTime());
@@ -2741,67 +2743,21 @@ FReply STimingView::OnKeyDown(const FGeometry& MyGeometry, const FKeyEvent& InKe
 		FSlateApplication::Get().QueryCursor();
 		return FReply::Handled();
 	}
-	else if (InKeyEvent.GetKey() == EKeys::D) // debug: toggles down-sampling on/off
+	else if (InKeyEvent.GetKey() == EKeys::D)
 	{
-		if (!InKeyEvent.GetModifierKeys().IsControlDown() && !InKeyEvent.GetModifierKeys().IsShiftDown())
+		if (InKeyEvent.GetModifierKeys().IsControlDown() && InKeyEvent.GetModifierKeys().IsShiftDown())
 		{
+			// Ctrl+Shift+D: Toggles down-sampling on/off (for debugging purposes only).
 			FTimingEventsTrack::bUseDownSampling = !FTimingEventsTrack::bUseDownSampling;
 			Viewport.AddDirtyFlags(ETimingTrackViewportDirtyFlags::HInvalidated);
 			return FReply::Handled();
 		}
 	}
-	else if (InKeyEvent.GetKey() == EKeys::R)
-	{
-		if (!InKeyEvent.GetModifierKeys().IsControlDown() && !InKeyEvent.GetModifierKeys().IsShiftDown())
-		{
-			FrameSharedState->ShowHideAllFrameTracks();
-			return FReply::Handled();
-		}
-	}
-	else if (InKeyEvent.GetKey() == EKeys::Y)
-	{
-		if (!InKeyEvent.GetModifierKeys().IsControlDown() && !InKeyEvent.GetModifierKeys().IsShiftDown())
-		{
-			ThreadTimingSharedState->ShowHideAllGpuTracks();
-			return FReply::Handled();
-		}
-	}
-	else if (InKeyEvent.GetKey() == EKeys::U)
-	{
-		if (!InKeyEvent.GetModifierKeys().IsControlDown() && !InKeyEvent.GetModifierKeys().IsShiftDown())
-		{
-			ThreadTimingSharedState->ShowHideAllCpuTracks();
-			return FReply::Handled();
-		}
-	}
-	else if (InKeyEvent.GetKey() == EKeys::L)
-	{
-		if (!InKeyEvent.GetModifierKeys().IsControlDown() && !InKeyEvent.GetModifierKeys().IsShiftDown())
-		{
-			LoadingSharedState->ShowHideAllLoadingTracks();
-			return FReply::Handled();
-		}
-	}
-	else if (InKeyEvent.GetKey() == EKeys::I)
-	{
-		if (!InKeyEvent.GetModifierKeys().IsControlDown() && !InKeyEvent.GetModifierKeys().IsShiftDown())
-		{
-			FileActivitySharedState->ShowHideAllIoTracks();
-			return FReply::Handled();
-		}
-	}
-	else if (InKeyEvent.GetKey() == EKeys::O)
-	{
-		if (!InKeyEvent.GetModifierKeys().IsControlDown() && !InKeyEvent.GetModifierKeys().IsShiftDown())
-		{
-			FileActivitySharedState->ToggleBackgroundEvents();
-			return FReply::Handled();
-		}
-	}
 	else if (InKeyEvent.GetKey() == EKeys::A)
 	{
 		if (InKeyEvent.GetModifierKeys().IsControlDown())
 		{
+			// Ctrl+A: Selects the entire time range of session.
 			SelectTimeInterval(0, Viewport.GetMaxValidTime());
 			return FReply::Handled();
 		}
@@ -3129,17 +3085,10 @@ void STimingView::BindCommands()
 		FExecuteAction::CreateSP(this, &STimingView::QuickFind_Execute),
 		FCanExecuteAction::CreateSP(this, &STimingView::QuickFind_CanExecute));
 
-	//CommandList->MapAction(
-	//	Commands.ShowAllGpuTracks,
-	//	FExecuteAction::CreateSP(this, &STimingView::ShowHideAllGpuTracks_Execute),
-	//	FCanExecuteAction(), //FCanExecuteAction::CreateLambda([] { return true; }),
-	//	FIsActionChecked::CreateSP(this, &STimingView::ShowHideAllGpuTracks_IsChecked));
-
-	//CommandList->MapAction(
-	//	Commands.ShowAllCpuTracks,
-	//	FExecuteAction::CreateSP(this, &STimingView::ShowHideAllCpuTracks_Execute),
-	//	FCanExecuteAction(), //FCanExecuteAction::CreateLambda([] { return true; }),
-	//	FIsActionChecked::CreateSP(this, &STimingView::ShowHideAllCpuTracks_IsChecked));
+	FrameSharedState->BindCommands();
+	ThreadTimingSharedState->BindCommands();
+	LoadingSharedState->BindCommands();
+	FileActivitySharedState->BindCommands();
 }
 
 ////////////////////////////////////////////////////////////////////////////////////////////////////
@@ -4601,77 +4550,41 @@ void STimingView::TogglePanningOnScreenEdges()
 
 ////////////////////////////////////////////////////////////////////////////////////////////////////
 
-void AddMenuEntryRadioButton(
-	FMenuBuilder& InOutMenuBuilder,
-	const FUIAction& InAction,
-	const TAttribute<FText>& InLabel,
-	const TAttribute<FText>& InToolTipText,
-	const TAttribute<FText>& InKeybinding)
-{
-	InOutMenuBuilder.AddMenuEntry(
-		InAction,
-		SNew(SBox)
-		.Padding(FMargin(0.0f))
-		[
-			SNew(SHorizontalBox)
-			+ SHorizontalBox::Slot()
-			.Padding(FMargin(2.0f, 0.0f))
-			.VAlign(VAlign_Center)
-			.FillWidth(1.0f)
-			[
-				SNew(STextBlock)
-				.TextStyle(FAppStyle::Get(), "Menu.Label")
-				.Text(InLabel)
-			]
-			+ SHorizontalBox::Slot()
-			.Padding(FMargin(4.0f, 0.0f))
-			.VAlign(VAlign_Center)
-			.AutoWidth()
-			[
-				SNew(STextBlock)
-				.TextStyle(FAppStyle::Get(), "Menu.Keybinding")
-				.ColorAndOpacity(FSlateColor::UseSubduedForeground())
-				.Text(InKeybinding)
-			]
-		],
-		NAME_None,
-		InToolTipText,
-		EUserInterfaceActionType::RadioButton
-	);
-}
-
-////////////////////////////////////////////////////////////////////////////////////////////////////
-
 void STimingView::CreateDepthLimitMenu(FMenuBuilder& MenuBuilder)
 {
 	MenuBuilder.BeginSection("DepthLimit", LOCTEXT("ContextMenu_Section_DepthLimit", "Depth Limit"));
 	{
-		AddMenuEntryRadioButton(MenuBuilder,
+		// Note: We use the custom AddMenuEntry in order to set the same key binding text for multiple menu items.
+
+		FInsightsMenuBuilder::AddMenuEntry(MenuBuilder,
 			FUIAction(
 				FExecuteAction::CreateSP(this, &STimingView::SetEventDepthLimit, FTimingProfilerManager::UnlimitedEventDepth),
 				FCanExecuteAction(),
 				FIsActionChecked::CreateSP(this, &STimingView::CheckEventDepthLimit, FTimingProfilerManager::UnlimitedEventDepth)),
 			LOCTEXT("UnlimitedDepth", "Unlimited"),
 			LOCTEXT("UnlimitedDepth_Desc", "Timing Events tracks (like the Cpu Thread tracks) can have unlimited depth (lanes)."),
-			TAttribute<FText>::Create(TAttribute<FText>::FGetter::CreateSP(this, &STimingView::GetEventDepthLimitKeybindingText, FTimingProfilerManager::UnlimitedEventDepth)));
+			TAttribute<FText>::Create(TAttribute<FText>::FGetter::CreateSP(this, &STimingView::GetEventDepthLimitKeybindingText, FTimingProfilerManager::UnlimitedEventDepth)),
+			EUserInterfaceActionType::RadioButton);
 
-		AddMenuEntryRadioButton(MenuBuilder,
+		FInsightsMenuBuilder::AddMenuEntry(MenuBuilder,
 			FUIAction(
 				FExecuteAction::CreateSP(this, &STimingView::SetEventDepthLimit, (uint32)4),
 				FCanExecuteAction(),
 				FIsActionChecked::CreateSP(this, &STimingView::CheckEventDepthLimit, (uint32)4)),
 			LOCTEXT("DepthLimit4", "4 Lanes"),
 			LOCTEXT("DepthLimit4_Desc", "Timing Events tracks (like the Cpu Thread tracks) can have maximum 4 lanes."),
-			TAttribute<FText>::Create(TAttribute<FText>::FGetter::CreateSP(this, &STimingView::GetEventDepthLimitKeybindingText, (uint32)4)));
+			TAttribute<FText>::Create(TAttribute<FText>::FGetter::CreateSP(this, &STimingView::GetEventDepthLimitKeybindingText, (uint32)4)),
+			EUserInterfaceActionType::RadioButton);
 
-		AddMenuEntryRadioButton(MenuBuilder,
+		FInsightsMenuBuilder::AddMenuEntry(MenuBuilder,
 			FUIAction(
 				FExecuteAction::CreateSP(this, &STimingView::SetEventDepthLimit, (uint32)1),
 				FCanExecuteAction(),
 				FIsActionChecked::CreateSP(this, &STimingView::CheckEventDepthLimit, (uint32)1)),
 			LOCTEXT("DepthLimit1", "Single Lane"),
 			LOCTEXT("DepthLimit1_Desc", "Timing Events tracks (like the Cpu Thread tracks) can have a single lane."),
-			TAttribute<FText>::Create(TAttribute<FText>::FGetter::CreateSP(this, &STimingView::GetEventDepthLimitKeybindingText, (uint32)1)));
+			TAttribute<FText>::Create(TAttribute<FText>::FGetter::CreateSP(this, &STimingView::GetEventDepthLimitKeybindingText, (uint32)1)),
+			EUserInterfaceActionType::RadioButton);
 	}
 	MenuBuilder.EndSection();
 }
