Shift CPU acronym to upper case in menus and tooltips.

diff --git a/Engine/Source/Developer/TraceInsights/Private/Insights/ContextSwitches/ViewModels/ContextSwitchesSharedState.cpp b/Engine/Source/Developer/TraceInsights/Private/Insights/ContextSwitches/ViewModels/ContextSwitchesSharedState.cpp
--- a/Engine/Source/Developer/TraceInsights/Private/Insights/ContextSwitches/ViewModels/ContextSwitchesSharedState.cpp
+++ b/Engine/Source/Developer/TraceInsights/Private/Insights/ContextSwitches/ViewModels/ContextSwitchesSharedState.cpp
@@ -54,19 +54,19 @@ void FContextSwitchesStateCommands::RegisterCommands()
 {
 	UI_COMMAND(Command_ShowCoreTracks,
 		"Core Tracks",
-		"Shows/hides the Cpu Core tracks.",
+		"Shows/hides the CPU Core tracks.",
 		EUserInterfaceActionType::ToggleButton,
 		FInputChord(EModifierKey::Alt, EKeys::C));
 
 	UI_COMMAND(Command_ShowContextSwitches,
 		"Context Switches",
-		"Shows/hides the context switches on top of cpu timing tracks.",
+		"Shows/hides the context switches on top of CPU timing tracks.",
 		EUserInterfaceActionType::ToggleButton,
 		FInputChord(EModifierKey::Shift, EKeys::C));
 
 	UI_COMMAND(Command_ShowOverlays,
 		"Overlays",
-		"Extends the visualisation of context switches over the cpu timing tracks.",
+		"Extends the visualization of context switches over the CPU timing tracks.",
 		EUserInterfaceActionType::ToggleButton,
 		FInputChord(EModifierKey::Shift, EKeys::O));
 
@@ -77,26 +77,26 @@ void FContextSwitchesStateCommands::RegisterCommands()
 		FInputChord(EModifierKey::Shift, EKeys::L));
 
 	UI_COMMAND(Command_NavigateToCpuThreadEvent,
-		"Go To Cpu Thread Track",
-		"Selects the context switch event in the corresponding Cpu Thread track.",
+		"Go To CPU Thread Track",
+		"Selects the context switch event in the corresponding CPU Thread track.",
 		EUserInterfaceActionType::Button,
 		FInputChord(/*EModifierKey::Control, EKeys::Tab*/));
 
 	UI_COMMAND(Command_DockCpuThreadTrackToBottom,
-		"Dock Cpu Thread Track To Bottom",
-		"Docks the corresponding Cpu Thread track to the bottom of the Timing view.",
+		"Dock CPU Thread Track To Bottom",
+		"Docks the corresponding CPU Thread track to the bottom of the Timing view.",
 		EUserInterfaceActionType::Button,
 		FInputChord());
 
 	UI_COMMAND(Command_NavigateToCpuCoreEvent,
-		"Go To Cpu Core Track",
-		"Selects the timing event in the corresponding Cpu Core track.",
+		"Go To CPU Core Track",
+		"Selects the timing event in the corresponding CPU Core track.",
 		EUserInterfaceActionType::Button,
 		FInputChord(EModifierKey::Control, EKeys::Tab));
 
 	UI_COMMAND(Command_DockCpuCoreTrackToTop,
-		"Dock Cpu Core Track To Top",
-		"Docks the corresponding Cpu Core track to the top of the Timing view.",
+		"Dock CPU Core Track To Top",
+		"Docks the corresponding CPU Core track to the top of the Timing view.",
 		EUserInterfaceActionType::Button,
 		FInputChord());
 }
diff --git a/Engine/Source/Developer/TraceInsights/Private/Insights/Widgets/STimingView.cpp b/Engine/Source/Developer/TraceInsights/Private/Insights/Widgets/STimingView.cpp
--- a/Engine/Source/Developer/TraceInsights/Private/Insights/Widgets/STimingView.cpp
+++ b/Engine/Source/Developer/TraceInsights/Private/Insights/Widgets/STimingView.cpp
@@ -4562,7 +4562,7 @@ void STimingView::CreateDepthLimitMenu(FMenuBuilder& MenuBuilder)
 				FCanExecuteAction(),
 				FIsActionChecked::CreateSP(this, &STimingView::CheckEventDepthLimit, FTimingProfilerManager::UnlimitedEventDepth)),
 			LOCTEXT("UnlimitedDepth", "Unlimited"),
-			LOCTEXT("UnlimitedDepth_Desc", "Timing Events tracks (like the Cpu Thread tracks) can have unlimited depth (lanes)."),
+			LOCTEXT("UnlimitedDepth_Desc", "Timing Events tracks (like the CPU Thread tracks) can have unlimited depth (lanes)."),
 			TAttribute<FText>::Create(TAttribute<FText>::FGetter::CreateSP(this, &STimingView::GetEventDepthLimitKeybindingText, FTimingProfilerManager::UnlimitedEventDepth)),
 			EUserInterfaceActionType::RadioButton);
 
@@ -4572,7 +4572,7 @@ void STimingView::CreateDepthLimitMenu(FMenuBuilder& MenuBuilder)
 				FCanExecuteAction(),
 				FIsActionChecked::CreateSP(this, &STimingView::CheckEventDepthLimit, (uint32)4)),
 			LOCTEXT("DepthLimit4", "4 Lanes"),
-			LOCTEXT("DepthLimit4_Desc", "Timing Events tracks (like the Cpu Thread tracks) can have maximum 4 lanes."),
+			LOCTEXT("DepthLimit4_Desc", "Timing Events tracks (like the CPU Thread tracks) can have maximum 4 lanes."),
 			TAttribute<FText>::Create(TAttribute<FText>::FGetter::CreateSP(this, &STimingView::GetEventDepthLimitKeybindingText, (uint32)4)),
 			EUserInterfaceActionType::RadioButton);
 
@@ -4582,7 +4582,7 @@ void STimingView::CreateDepthLimitMenu(FMenuBuilder& MenuBuilder)
 				FCanExecuteAction(),
 				FIsActionChecked::CreateSP(this, &STimingView::CheckEventDepthLimit, (uint32)1)),
 			LOCTEXT("DepthLimit1", "Single Lane"),
-			LOCTEXT("DepthLimit1_Desc", "Timing Events tracks (like the Cpu Thread tracks) can have a single lane."),
+			LOCTEXT("DepthLimit1_Desc", "Timing Events tracks (like the CPU Thread tracks) can have a single lane."),
 			TAttribute<FText>::Create(TAttribute<FText>::FGetter::CreateSP(this, &STimingView::GetEventDepthLimitKeybindingText, (uint32)1)),
 			EUserInterfaceActionType::RadioButton);
 	}
