[Insights]
  - Timing View: Changed time selection to be unique for each instance of STimingView. Also removed a hardcoded time selection changed notification for Asset Loading Insights.
  - Memory Insights: Changed time selection made in the Timing view from Memory Insights to not update anymore the time selection from Timing Insights. (UE-135582)
  - Memory Insights: Fixed bug where the time markers from Mem Investigation panel become inconsistent with ones from Timing view, when the Timing view is toggled on/off.
  - Memory Insights: Changed the time markers to be always sorted by time, i.e. A <= B <= C <= D. (UE-137557)
  - Memory Insights: Changed the time markers A and B to be automatically set by a time range selection.

diff --git a/Engine/Source/Developer/TraceInsights/Private/Insights/LoadingProfiler/Widgets/SLoadingProfilerWindow.cpp b/Engine/Source/Developer/TraceInsights/Private/Insights/LoadingProfiler/Widgets/SLoadingProfilerWindow.cpp
--- a/Engine/Source/Developer/TraceInsights/Private/Insights/LoadingProfiler/Widgets/SLoadingProfilerWindow.cpp
+++ b/Engine/Source/Developer/TraceInsights/Private/Insights/LoadingProfiler/Widgets/SLoadingProfilerWindow.cpp
@@ -51,6 +51,8 @@ const FName FLoadingProfilerTabs::RequestsTreeViewID(TEXT("Requests"));
 
 SLoadingProfilerWindow::SLoadingProfilerWindow()
 	: DurationActive(0.0f)
+	, SelectionStartTime(0.0f)
+	, SelectionEndTime(0.0f)
 {
 }
 
@@ -162,9 +164,6 @@ void SLoadingProfilerWindow::UpdateEventAggregationTreeView()
 			TraceServices::FAnalysisSessionReadScope SessionReadScope(*Session.Get());
 			const TraceServices::ILoadTimeProfilerProvider& LoadTimeProfilerProvider = *TraceServices::ReadLoadTimeProfilerProvider(*Session.Get());
 
-			const double SelectionStartTime = TimingView ? TimingView->GetSelectionStartTime() : 0.0;
-			const double SelectionEndTime = TimingView ? TimingView->GetSelectionEndTime() : 0.0;
-
 			TraceServices::ITable<TraceServices::FLoadTimeProfilerAggregatedStats>* EventAggregationTable = LoadTimeProfilerProvider.CreateEventAggregation(SelectionStartTime, SelectionEndTime);
 			EventAggregationTreeView->UpdateSourceTable(MakeShareable(EventAggregationTable));
 		}
@@ -187,9 +186,6 @@ void SLoadingProfilerWindow::UpdateObjectTypeAggregationTreeView()
 			TraceServices::FAnalysisSessionReadScope SessionReadScope(*Session.Get());
 			const TraceServices::ILoadTimeProfilerProvider& LoadTimeProfilerProvider = *TraceServices::ReadLoadTimeProfilerProvider(*Session.Get());
 
-			const double SelectionStartTime = TimingView ? TimingView->GetSelectionStartTime() : 0.0;
-			const double SelectionEndTime = TimingView ? TimingView->GetSelectionEndTime() : 0.0;
-
 			TraceServices::ITable<TraceServices::FLoadTimeProfilerAggregatedStats>* ObjectTypeAggregationTable = LoadTimeProfilerProvider.CreateObjectTypeAggregation(SelectionStartTime, SelectionEndTime);
 			ObjectTypeAggregationTreeView->UpdateSourceTable(MakeShareable(ObjectTypeAggregationTable));
 		}
@@ -212,9 +208,6 @@ void SLoadingProfilerWindow::UpdatePackageDetailsTreeView()
 			TraceServices::FAnalysisSessionReadScope SessionReadScope(*Session.Get());
 			const TraceServices::ILoadTimeProfilerProvider& LoadTimeProfilerProvider = *TraceServices::ReadLoadTimeProfilerProvider(*Session.Get());
 
-			const double SelectionStartTime = TimingView ? TimingView->GetSelectionStartTime() : 0.0;
-			const double SelectionEndTime = TimingView ? TimingView->GetSelectionEndTime() : 0.0;
-
 			TraceServices::ITable<TraceServices::FPackagesTableRow>* PackageDetailsTable = LoadTimeProfilerProvider.CreatePackageDetailsTable(SelectionStartTime, SelectionEndTime);
 			PackageDetailsTreeView->UpdateSourceTable(MakeShareable(PackageDetailsTable));
 		}
@@ -237,9 +230,6 @@ void SLoadingProfilerWindow::UpdateExportDetailsTreeView()
 			TraceServices::FAnalysisSessionReadScope SessionReadScope(*Session.Get());
 			const TraceServices::ILoadTimeProfilerProvider& LoadTimeProfilerProvider = *TraceServices::ReadLoadTimeProfilerProvider(*Session.Get());
 
-			const double SelectionStartTime = TimingView ? TimingView->GetSelectionStartTime() : 0.0;
-			const double SelectionEndTime = TimingView ? TimingView->GetSelectionEndTime() : 0.0;
-
 			TraceServices::ITable<TraceServices::FExportsTableRow>* ExportDetailsTable = LoadTimeProfilerProvider.CreateExportDetailsTable(SelectionStartTime, SelectionEndTime);
 			ExportDetailsTreeView->UpdateSourceTable(MakeShareable(ExportDetailsTable));
 		}
@@ -302,6 +292,8 @@ TSharedRef<SDockTab> SLoadingProfilerWindow::SpawnTab_TimingView(const FSpawnTab
 
 	TimingView->EnableAssetLoadingMode();
 	TimingView->Reset(true);
+	TimingView->OnSelectionChanged().AddSP(this, &SLoadingProfilerWindow::OnTimeSelectionChanged);
+	TimingView->SelectTimeInterval(SelectionStartTime, SelectionEndTime - SelectionStartTime);
 
 	DockTab->SetOnTabClosed(SDockTab::FOnTabClosedCallback::CreateRaw(this, &SLoadingProfilerWindow::OnTimingViewTabClosed));
 
@@ -314,7 +306,11 @@ END_SLATE_FUNCTION_BUILD_OPTIMIZATION
 void SLoadingProfilerWindow::OnTimingViewTabClosed(TSharedRef<SDockTab> TabBeingClosed)
 {
 	FLoadingProfilerManager::Get()->SetTimingViewVisible(false);
-	TimingView = nullptr;
+	if (TimingView)
+	{
+		TimingView->OnSelectionChanged().RemoveAll(this);
+		TimingView = nullptr;
+	}
 }
 
 ////////////////////////////////////////////////////////////////////////////////////////////////////
@@ -772,4 +768,19 @@ FReply SLoadingProfilerWindow::OnDrop(const FGeometry& MyGeometry, const FDragDr
 
 ////////////////////////////////////////////////////////////////////////////////////////////////////
 
+void SLoadingProfilerWindow::OnTimeSelectionChanged(Insights::ETimeChangedFlags InFlags, double InStartTime, double InEndTime)
+{
+	if (InFlags != Insights::ETimeChangedFlags::Interactive)
+	{
+		if (InStartTime < InEndTime)
+		{
+			SelectionStartTime = InStartTime;
+			SelectionEndTime = InEndTime;
+			UpdateTableTreeViews();
+		}
+	}
+}
+
+////////////////////////////////////////////////////////////////////////////////////////////////////
+
 #undef LOCTEXT_NAMESPACE
diff --git a/Engine/Source/Developer/TraceInsights/Private/Insights/LoadingProfiler/Widgets/SLoadingProfilerWindow.h b/Engine/Source/Developer/TraceInsights/Private/Insights/LoadingProfiler/Widgets/SLoadingProfilerWindow.h
--- a/Engine/Source/Developer/TraceInsights/Private/Insights/LoadingProfiler/Widgets/SLoadingProfilerWindow.h
+++ b/Engine/Source/Developer/TraceInsights/Private/Insights/LoadingProfiler/Widgets/SLoadingProfilerWindow.h
@@ -14,6 +14,7 @@
 
 // Insights
 #include "Insights/InsightsManager.h"
+#include "Insights/ITimingViewSession.h"
 
 ////////////////////////////////////////////////////////////////////////////////////////////////////
 
@@ -169,6 +170,8 @@ private:
 	 */
 	virtual FReply OnDragOver(const FGeometry& MyGeometry, const FDragDropEvent& DragDropEvent)  override;
 
+	void OnTimeSelectionChanged(Insights::ETimeChangedFlags InFlags, double InStartTime, double InEndTime);
+
 private:
 	/** The Timing view (multi-track) widget */
 	TSharedPtr<STimingView> TimingView;
@@ -196,4 +199,7 @@ private:
 
 	/** The number of seconds the profiler has been active */
 	float DurationActive;
+
+	double SelectionStartTime;
+	double SelectionEndTime;
 };
diff --git a/Engine/Source/Developer/TraceInsights/Private/Insights/MemoryProfiler/ViewModels/MemorySharedState.cpp b/Engine/Source/Developer/TraceInsights/Private/Insights/MemoryProfiler/ViewModels/MemorySharedState.cpp
--- a/Engine/Source/Developer/TraceInsights/Private/Insights/MemoryProfiler/ViewModels/MemorySharedState.cpp
+++ b/Engine/Source/Developer/TraceInsights/Private/Insights/MemoryProfiler/ViewModels/MemorySharedState.cpp
@@ -116,8 +116,6 @@ void FMemorySharedState::OnBeginSession(Insights::ITimingViewSession& InSession)
 	bShowHideAllMemoryTracks = true;
 
 	CreatedDefaultTracks.Reset();
-
-	CurrentMemoryRule = MemoryRules[0];
 }
 
 ////////////////////////////////////////////////////////////////////////////////////////////////////
@@ -143,8 +141,6 @@ void FMemorySharedState::OnEndSession(Insights::ITimingViewSession& InSession)
 	bShowHideAllMemoryTracks = false;
 
 	CreatedDefaultTracks.Reset();
-
-	CurrentMemoryRule = nullptr;
 }
 
 ////////////////////////////////////////////////////////////////////////////////////////////////////
diff --git a/Engine/Source/Developer/TraceInsights/Private/Insights/MemoryProfiler/Widgets/SMemInvestigationView.cpp b/Engine/Source/Developer/TraceInsights/Private/Insights/MemoryProfiler/Widgets/SMemInvestigationView.cpp
--- a/Engine/Source/Developer/TraceInsights/Private/Insights/MemoryProfiler/Widgets/SMemInvestigationView.cpp
+++ b/Engine/Source/Developer/TraceInsights/Private/Insights/MemoryProfiler/Widgets/SMemInvestigationView.cpp
@@ -247,6 +247,11 @@ TSharedRef<SWidget> SMemInvestigationView::ConstructTimeMarkerWidget(uint32 Time
 	}
 
 	const TSharedRef<Insights::FTimeMarker>& TimeMarker = ProfilerWindow->GetCustomTimeMarker(TimeMarkerIndex);
+	TSharedPtr<Insights::FTimeMarker> PreviousTimeMarker;
+	if (TimeMarkerIndex > 0)
+	{
+		PreviousTimeMarker = ProfilerWindow->GetCustomTimeMarker(TimeMarkerIndex - 1);
+	}
 
 	TSharedRef<SWidget> Widget = SNew(SHorizontalBox)
 		.Visibility_Lambda([TimeMarker]()
@@ -276,10 +281,11 @@ TSharedRef<SWidget> SMemInvestigationView::ConstructTimeMarkerWidget(uint32 Time
 				{
 					return FText::FromString(FString::Printf(TEXT("%.9f"), TimeMarker->GetTime()));
 				})
-			.OnTextCommitted_Lambda([TimeMarker](const FText& InText, ETextCommit::Type InCommitType)
+			.OnTextCommitted_Lambda([TimeMarker, ProfilerWindow](const FText& InText, ETextCommit::Type InCommitType)
 				{
 					const double Time = FCString::Atod(*InText.ToString());
 					TimeMarker->SetTime(Time);
+					ProfilerWindow->OnTimeMarkerChanged(Insights::ETimeChangedFlags::None, TimeMarker);
 				})
 		]
 
@@ -289,9 +295,18 @@ TSharedRef<SWidget> SMemInvestigationView::ConstructTimeMarkerWidget(uint32 Time
 		.Padding(4.0f, 1.0f, 0.0f, 1.0f)
 		[
 			SNew(STextBlock)
-			.Text_Lambda([TimeMarker]()
+			.Text_Lambda([TimeMarker, PreviousTimeMarker]()
 				{
-					return FText::FromString(TimeUtils::FormatTime(TimeMarker->GetTime(), 0.1));
+					if (PreviousTimeMarker.IsValid())
+					{
+						return FText::FromString(FString::Printf(TEXT("%s (+%s)"),
+							*TimeUtils::FormatTime(TimeMarker->GetTime(), 0.1),
+							*TimeUtils::FormatTime(TimeMarker->GetTime() - PreviousTimeMarker->GetTime(), 0.1)));
+					}
+					else
+					{
+						return FText::FromString(TimeUtils::FormatTime(TimeMarker->GetTime(), 0.1));
+					}
 				})
 		];
 
@@ -353,7 +368,7 @@ void SMemInvestigationView::QueryRule_OnSelectionChanged(TSharedPtr<Insights::FM
 		{
 			FMemorySharedState& SharedState = ProfilerWindow->GetSharedState();
 			SharedState.SetCurrentMemoryRule(InRule);
-			ProfilerWindow->UpdateTimingViewMarkers();
+			ProfilerWindow->OnMemoryRuleChanged();
 		}
 	}
 }
diff --git a/Engine/Source/Developer/TraceInsights/Private/Insights/MemoryProfiler/Widgets/SMemoryProfilerWindow.cpp b/Engine/Source/Developer/TraceInsights/Private/Insights/MemoryProfiler/Widgets/SMemoryProfilerWindow.cpp
--- a/Engine/Source/Developer/TraceInsights/Private/Insights/MemoryProfiler/Widgets/SMemoryProfilerWindow.cpp
+++ b/Engine/Source/Developer/TraceInsights/Private/Insights/MemoryProfiler/Widgets/SMemoryProfilerWindow.cpp
@@ -57,6 +57,7 @@
 	: SharedState(MakeShared<FMemorySharedState>())
 	, DurationActive(0.0f)
 {
+	CreateTimingViewMarkers();
 }
 
 ////////////////////////////////////////////////////////////////////////////////////////////////////
@@ -115,68 +116,6 @@
 
 ////////////////////////////////////////////////////////////////////////////////////////////////////
 
-void SMemoryProfilerWindow::ResetTimingViewMarkers()
-{
-	TSharedRef<FTimeRulerTrack> TimeRulerTrack = TimingView->GetTimeRulerTrack();
-
-	TimeRulerTrack->RemoveAllTimeMarkers();
-	CustomTimeMarkers.Reset();
-
-	constexpr uint32 MaxNumTimeMarkers = 5;
-
-	TCHAR TimeMarkerName[2];
-	TimeMarkerName[1] = 0;
-
-	for (uint32 Index = 0; Index < MaxNumTimeMarkers; ++Index)
-	{
-		// Keep (re-add) the "Default Time Marker" as the first time marker and create new ones for the rest.
-		TSharedRef<Insights::FTimeMarker> TimeMarker = (Index == 0) ? TimingView->GetDefaultTimeMarker() : MakeShared<Insights::FTimeMarker>();
-
-		TimeMarkerName[0] = TEXT('A') + Index; // "A", "B", "C", etc.
-		TimeMarker->SetName(TimeMarkerName);
-
-		const uint32 HueStep = 256 / MaxNumTimeMarkers;
-		const uint8 H = uint8(HueStep * Index);
-		const uint8 S = 192;
-		const uint8 V = 255;
-		const FLinearColor Color = FLinearColor::MakeFromHSV8(H, S, V);
-		TimeMarker->SetColor(Color);
-
-		TimeMarker->SetTime((Index + 1) * 10.0); // 10s, 20s, 30s, etc.
-
-		TimeRulerTrack->AddTimeMarker(TimeMarker);
-		CustomTimeMarkers.Add(TimeMarker);
-	}
-
-	UpdateTimingViewMarkers();
-}
-
-////////////////////////////////////////////////////////////////////////////////////////////////////
-
-void SMemoryProfilerWindow::UpdateTimingViewMarkers()
-{
-	TSharedPtr<Insights::FMemoryRuleSpec> Rule = SharedState->GetCurrentMemoryRule();
-	const uint32 NumVisibleTimeMarkers = Rule ? Rule->GetNumTimeMarkers() : 0;
-
-	const uint32 NumTimeMarkers = CustomTimeMarkers.Num();
-	ensure(NumVisibleTimeMarkers <= NumTimeMarkers);
-
-	for (uint32 Index = 0; Index < NumTimeMarkers; ++Index)
-	{
-		TSharedRef<Insights::FTimeMarker>& TimeMarker = CustomTimeMarkers[Index];
-		if (Index < NumVisibleTimeMarkers)
-		{
-			TimeMarker->SetVisibility(true);
-		}
-		else
-		{
-			TimeMarker->SetVisibility(false);
-		}
-	}
-}
-
-////////////////////////////////////////////////////////////////////////////////////////////////////
-
 void SMemoryProfilerWindow::UpdateMemInvestigationView()
 {
 	if (MemInvestigationView)
@@ -253,6 +192,8 @@
 	IModularFeatures::Get().RegisterModularFeature(Insights::TimingViewExtenderFeatureName, &SharedState.Get());
 
 	TimingView->Reset(true);
+	TimingView->OnSelectionChanged().AddSP(this, &SMemoryProfilerWindow::OnTimeSelectionChanged);
+	TimingView->OnCustomTimeMarkerChanged().AddSP(this, &SMemoryProfilerWindow::OnTimeMarkerChanged);
 	ResetTimingViewMarkers();
 	TimingView->HideAllDefaultTracks();
 
@@ -266,12 +207,17 @@
 
 void SMemoryProfilerWindow::OnTimingViewTabClosed(TSharedRef<SDockTab> TabBeingClosed)
 {
+	FMemoryProfilerManager::Get()->SetTimingViewVisible(false);
+
+	if (TimingView)
+	{
+		TimingView->OnSelectionChanged().RemoveAll(this);
+		TimingView->OnCustomTimeMarkerChanged().RemoveAll(this);
+		TimingView = nullptr;
+	}
+
 	IModularFeatures::Get().UnregisterModularFeature(Insights::TimingViewExtenderFeatureName, &SharedState.Get());
 	SharedState->SetTimingView(nullptr);
-
-	TimingView = nullptr;
-
-	FMemoryProfilerManager::Get()->SetTimingViewVisible(false);
 }
 
 ////////////////////////////////////////////////////////////////////////////////////////////////////
@@ -766,4 +712,167 @@
 
 ////////////////////////////////////////////////////////////////////////////////////////////////////
 
+void SMemoryProfilerWindow::OnTimeSelectionChanged(Insights::ETimeChangedFlags InFlags, double InStartTime, double InEndTime)
+{
+	if (InFlags != Insights::ETimeChangedFlags::Interactive)
+	{
+		if (InStartTime < InEndTime)
+		{
+			const uint32 NumTimeMarkers = CustomTimeMarkers.Num();
+			if (NumTimeMarkers >= 1)
+			{
+				CustomTimeMarkers[0]->SetTime(InStartTime);
+			}
+			if (NumTimeMarkers >= 2)
+			{
+				CustomTimeMarkers[1]->SetTime(InEndTime);
+			}
+
+			double Time = InEndTime;
+			for (uint32 Index = 2; Index < NumTimeMarkers; ++Index)
+			{
+				const double TimeMarkerTime = CustomTimeMarkers[Index]->GetTime();
+				if (TimeMarkerTime < Time)
+				{
+					CustomTimeMarkers[Index]->SetTime(Time);
+				}
+				else
+				{
+					Time = TimeMarkerTime;
+				}
+			}
+
+			UpdateMemInvestigationView();
+		}
+	}
+}
+
+////////////////////////////////////////////////////////////////////////////////////////////////////
+
+void SMemoryProfilerWindow::CreateTimingViewMarkers()
+{
+	check(CustomTimeMarkers.Num() == 0);
+
+	constexpr uint32 MaxNumTimeMarkers = 5;
+
+	TCHAR TimeMarkerName[2];
+	TimeMarkerName[1] = 0;
+
+	for (uint32 Index = 0; Index < MaxNumTimeMarkers; ++Index)
+	{
+		TSharedRef<Insights::FTimeMarker> TimeMarker = MakeShared<Insights::FTimeMarker>();
+
+		TimeMarkerName[0] = TEXT('A') + Index; // "A", "B", "C", etc.
+		TimeMarker->SetName(TimeMarkerName);
+
+		const uint32 HueStep = 256 / MaxNumTimeMarkers;
+		const uint8 H = uint8(HueStep * Index);
+		const uint8 S = 192;
+		const uint8 V = 255;
+		const FLinearColor Color = FLinearColor::MakeFromHSV8(H, S, V);
+		TimeMarker->SetColor(Color);
+
+		TimeMarker->SetTime((Index + 1) * 10.0); // 10s, 20s, 30s, etc.
+		TimeMarker->SetVisibility(false);
+
+		CustomTimeMarkers.Add(TimeMarker);
+	}
+}
+
+////////////////////////////////////////////////////////////////////////////////////////////////////
+
+void SMemoryProfilerWindow::ResetTimingViewMarkers()
+{
+	TSharedRef<FTimeRulerTrack> TimeRulerTrack = TimingView->GetTimeRulerTrack();
+
+	TimeRulerTrack->RemoveAllTimeMarkers();
+
+	// Hide the "Default Time Marker".
+	TSharedRef<Insights::FTimeMarker> DefaultTimeMarker = TimingView->GetDefaultTimeMarker();
+	DefaultTimeMarker->SetVisibility(false);
+
+	const uint32 NumTimeMarkers = CustomTimeMarkers.Num();
+	for (uint32 Index = 0; Index < NumTimeMarkers; ++Index)
+	{
+		TSharedRef<Insights::FTimeMarker>& TimeMarker = CustomTimeMarkers[Index];
+		TimeMarker->SetTime((Index + 1) * 10.0); // 10s, 20s, 30s, etc.
+		TimeRulerTrack->AddTimeMarker(TimeMarker);
+	}
+
+	UpdateTimingViewMarkers();
+}
+
+////////////////////////////////////////////////////////////////////////////////////////////////////
+
+void SMemoryProfilerWindow::OnMemoryRuleChanged()
+{
+	UpdateTimingViewMarkers();
+}
+
+////////////////////////////////////////////////////////////////////////////////////////////////////
+
+void SMemoryProfilerWindow::UpdateTimingViewMarkers()
+{
+	TSharedPtr<Insights::FMemoryRuleSpec> Rule = SharedState->GetCurrentMemoryRule();
+	const uint32 NumVisibleTimeMarkers = Rule ? Rule->GetNumTimeMarkers() : 0;
+
+	const uint32 NumTimeMarkers = CustomTimeMarkers.Num();
+	ensure(NumVisibleTimeMarkers <= NumTimeMarkers);
+
+	for (uint32 Index = 0; Index < NumTimeMarkers; ++Index)
+	{
+		TSharedRef<Insights::FTimeMarker>& TimeMarker = CustomTimeMarkers[Index];
+		if (Index < NumVisibleTimeMarkers)
+		{
+			TimeMarker->SetVisibility(true);
+		}
+		else
+		{
+			TimeMarker->SetVisibility(false);
+		}
+	}
+}
+
+////////////////////////////////////////////////////////////////////////////////////////////////////
+
+void SMemoryProfilerWindow::OnTimeMarkerChanged(Insights::ETimeChangedFlags InFlags, TSharedRef<Insights::ITimeMarker> InTimeMarker)
+{
+	const uint32 NumTimeMarkers = CustomTimeMarkers.Num();
+
+	// Find index of the changing time marker.
+	uint32 ChangedTimeMarkerIndex = -1;
+	for (uint32 Index = 0; Index < NumTimeMarkers; ++Index)
+	{
+		TSharedRef<Insights::FTimeMarker>& TimeMarker = CustomTimeMarkers[Index];
+		if (TimeMarker == InTimeMarker)
+		{
+			ChangedTimeMarkerIndex = Index;
+			break;
+		}
+	}
+
+	// Ensure the rest of time markers are orderd by time.
+	if (ChangedTimeMarkerIndex >= 0)
+	{
+		for (uint32 Index = 0; Index < ChangedTimeMarkerIndex; ++Index)
+		{
+			TSharedRef<Insights::FTimeMarker>& TimeMarker = CustomTimeMarkers[Index];
+			if (TimeMarker->GetTime() > InTimeMarker->GetTime())
+			{
+				TimeMarker->SetTime(InTimeMarker->GetTime());
+			}
+		}
+		for (uint32 Index = ChangedTimeMarkerIndex + 1; Index < NumTimeMarkers; ++Index)
+		{
+			TSharedRef<Insights::FTimeMarker>& TimeMarker = CustomTimeMarkers[Index];
+			if (TimeMarker->GetTime() < InTimeMarker->GetTime())
+			{
+				TimeMarker->SetTime(InTimeMarker->GetTime());
+			}
+		}
+	}
+}
+
+////////////////////////////////////////////////////////////////////////////////////////////////////
+
 #undef LOCTEXT_NAMESPACE
diff --git a/Engine/Source/Developer/TraceInsights/Private/Insights/MemoryProfiler/Widgets/SMemoryProfilerWindow.h b/Engine/Source/Developer/TraceInsights/Private/Insights/MemoryProfiler/Widgets/SMemoryProfilerWindow.h
--- a/Engine/Source/Developer/TraceInsights/Private/Insights/MemoryProfiler/Widgets/SMemoryProfilerWindow.h
+++ b/Engine/Source/Developer/TraceInsights/Private/Insights/MemoryProfiler/Widgets/SMemoryProfilerWindow.h
@@ -14,6 +14,7 @@
 
 // Insights
 #include "Insights/InsightsManager.h"
+#include "Insights/ITimingViewSession.h"
 #include "Insights/Widgets/SModulesView.h"
 
 ////////////////////////////////////////////////////////////////////////////////////////////////////
@@ -62,8 +63,6 @@ public:
 	SLATE_END_ARGS()
 
 	void Reset();
-	void ResetTimingViewMarkers();
-	void UpdateTimingViewMarkers();
 	void UpdateTableTreeViews();
 	void UpdateMemInvestigationView();
 	void UpdateMemTagTreeView();
@@ -93,6 +92,9 @@ public:
 
 	void CloseMemAllocTableTreeTabs();
 
+	void OnMemoryRuleChanged();
+	void OnTimeMarkerChanged(Insights::ETimeChangedFlags InFlags, TSharedRef<Insights::ITimeMarker> InTimeMarker);
+
 private:
 	TSharedRef<SDockTab> SpawnTab_TimingView(const FSpawnTabArgs& Args);
 	void OnTimingViewTabClosed(TSharedRef<SDockTab> TabBeingClosed);
@@ -108,7 +110,7 @@ private:
 
 	TSharedRef<SDockTab> SpawnTab_ModulesView(const FSpawnTabArgs& Args);
 	void OnModulesViewClosed(TSharedRef<SDockTab> TabBeingClosed);
-	
+
 	/**
 	 * Fill the main menu with menu items.
 	 *
@@ -180,6 +182,12 @@ private:
 	 */
 	virtual FReply OnDragOver(const FGeometry& MyGeometry, const FDragDropEvent& DragDropEvent)  override;
 
+	void OnTimeSelectionChanged(Insights::ETimeChangedFlags InFlags, double InStartTime, double InEndTime);
+
+	void CreateTimingViewMarkers();
+	void ResetTimingViewMarkers();
+	void UpdateTimingViewMarkers();
+
 private:
 	TSharedRef<FMemorySharedState> SharedState;
 
diff --git a/Engine/Source/Developer/TraceInsights/Private/Insights/ViewModels/TimeRulerTrack.h b/Engine/Source/Developer/TraceInsights/Private/Insights/ViewModels/TimeRulerTrack.h
--- a/Engine/Source/Developer/TraceInsights/Private/Insights/ViewModels/TimeRulerTrack.h
+++ b/Engine/Source/Developer/TraceInsights/Private/Insights/ViewModels/TimeRulerTrack.h
@@ -7,6 +7,7 @@
 
 // Insights
 #include "Insights/ViewModels/BaseTimingTrack.h"
+#include "Insights/ITimingViewSession.h"
 
 struct FDrawContext;
 class FMenuBuilder;
@@ -18,7 +19,7 @@ namespace Insights
 
 class ITimingViewExtender;
 
-class FTimeMarker
+class FTimeMarker : public  ITimeMarker
 {
 public:
 	FTimeMarker()
@@ -30,8 +31,10 @@ public:
 		, CrtTextWidth(0.0f)
 	{}
 
-	double GetTime() const { return Time; }
-	void SetTime(const double  InTime) { Time = InTime; }
+	virtual ~FTimeMarker() {}
+
+	double GetTime() const override { return Time; }
+	void SetTime(const double  InTime) override { Time = InTime; }
 
 	const FString& GetName() const { return Name; }
 	void SetName(const FString& InName) { Name = InName; }
diff --git a/Engine/Source/Developer/TraceInsights/Private/Insights/Widgets/STimingProfilerWindow.cpp b/Engine/Source/Developer/TraceInsights/Private/Insights/Widgets/STimingProfilerWindow.cpp
--- a/Engine/Source/Developer/TraceInsights/Private/Insights/Widgets/STimingProfilerWindow.cpp
+++ b/Engine/Source/Developer/TraceInsights/Private/Insights/Widgets/STimingProfilerWindow.cpp
@@ -199,6 +199,10 @@ TSharedRef<SDockTab> STimingProfilerWindow::SpawnTab_TimingView(const FSpawnTabA
 		];
 
 	TimingView->Reset(true);
+	TimingView->OnSelectionChanged().AddSP(this, &STimingProfilerWindow::OnTimeSelectionChanged);
+	const double SelectionStartTime = FTimingProfilerManager::Get()->GetSelectionStartTime();
+	const double SelectionEndTime = FTimingProfilerManager::Get()->GetSelectionEndTime();
+	TimingView->SelectTimeInterval(SelectionStartTime, SelectionEndTime - SelectionStartTime);
 
 	DockTab->SetOnTabClosed(SDockTab::FOnTabClosedCallback::CreateRaw(this, &STimingProfilerWindow::OnTimingViewTabClosed));
 
@@ -211,7 +215,11 @@ END_SLATE_FUNCTION_BUILD_OPTIMIZATION
 void STimingProfilerWindow::OnTimingViewTabClosed(TSharedRef<SDockTab> TabBeingClosed)
 {
 	FTimingProfilerManager::Get()->SetTimingViewVisible(false);
-	TimingView = nullptr;
+	if (TimingView)
+	{
+		TimingView->OnSelectionChanged().RemoveAll(this);
+		TimingView = nullptr;
+	}
 }
 
 ////////////////////////////////////////////////////////////////////////////////////////////////////
@@ -732,4 +740,14 @@ FReply STimingProfilerWindow::OnDrop(const FGeometry& MyGeometry, const FDragDro
 
 ////////////////////////////////////////////////////////////////////////////////////////////////////
 
+void STimingProfilerWindow::OnTimeSelectionChanged(Insights::ETimeChangedFlags InFlags, double InStartTime, double InEndTime)
+{
+	if (InFlags != Insights::ETimeChangedFlags::Interactive)
+	{
+		FTimingProfilerManager::Get()->SetSelectedTimeRange(InStartTime, InEndTime);
+	}
+}
+
+////////////////////////////////////////////////////////////////////////////////////////////////////
+
 #undef LOCTEXT_NAMESPACE
diff --git a/Engine/Source/Developer/TraceInsights/Private/Insights/Widgets/STimingProfilerWindow.h b/Engine/Source/Developer/TraceInsights/Private/Insights/Widgets/STimingProfilerWindow.h
--- a/Engine/Source/Developer/TraceInsights/Private/Insights/Widgets/STimingProfilerWindow.h
+++ b/Engine/Source/Developer/TraceInsights/Private/Insights/Widgets/STimingProfilerWindow.h
@@ -14,7 +14,7 @@
 
 // Insights
 #include "Insights/InsightsManager.h"
-#include "Insights/IUnrealInsightsModule.h"
+#include "Insights/ITimingViewSession.h"
 
 ////////////////////////////////////////////////////////////////////////////////////////////////////
 
@@ -155,6 +155,8 @@ private:
 	 */
 	virtual FReply OnDragOver(const FGeometry& MyGeometry, const FDragDropEvent& DragDropEvent)  override;
 
+	void OnTimeSelectionChanged(Insights::ETimeChangedFlags InFlags, double InStartTime, double InEndTime);
+
 private:
 	/** The Frame track widget */
 	TSharedPtr<SFrameTrack> FrameTrack;
diff --git a/Engine/Source/Developer/TraceInsights/Private/Insights/Widgets/STimingView.cpp b/Engine/Source/Developer/TraceInsights/Private/Insights/Widgets/STimingView.cpp
--- a/Engine/Source/Developer/TraceInsights/Private/Insights/Widgets/STimingView.cpp
+++ b/Engine/Source/Developer/TraceInsights/Private/Insights/Widgets/STimingView.cpp
@@ -427,6 +427,7 @@ void STimingView::Reset(bool bIsFirstReset)
 	bIsSelecting = false;
 	SelectionStartTime = 0.0;
 	SelectionEndTime = 0.0;
+	RaiseSelectionChanged();
 
 	if (HoveredTrack.IsValid())
 	{
@@ -685,7 +686,7 @@ void STimingView::Tick(const FGeometry& AllottedGeometry, const double InCurrent
 			: TimingView(InTimingView)
 			, Geometry(InGeometry)
 			, CurrentTime(InCurrentTime)
-			, DeltaTime(InDeltaTime)		
+			, DeltaTime(InDeltaTime)
 		{}
 
 		virtual const FGeometry& GetGeometry() const override { return Geometry; }
@@ -2860,7 +2861,7 @@ void STimingView::ShowContextMenu(const FPointerEvent& MouseEvent)
 		bHasAnyActions |= Extender->ExtendGlobalContextMenu(*this, MenuBuilder);
 	}
 
-	if(!bHasAnyActions)
+	if (!bHasAnyActions)
 	{
 		MenuBuilder.BeginSection(TEXT("Empty"));
 		{
@@ -3358,13 +3359,6 @@ void STimingView::RaiseSelectionChanging()
 void STimingView::RaiseSelectionChanged()
 {
 	OnSelectionChangedDelegate.Broadcast(Insights::ETimeChangedFlags::None, SelectionStartTime, SelectionEndTime);
-
-	FTimingProfilerManager::Get()->SetSelectedTimeRange(SelectionStartTime, SelectionEndTime);
-
-	if (SelectionStartTime < SelectionEndTime)
-	{
-		UpdateAggregatedStats();
-	}
 }
 
 ////////////////////////////////////////////////////////////////////////////////////////////////////
@@ -3376,6 +3370,7 @@ void STimingView::RaiseTimeMarkerChanging(TSharedRef<Insights::FTimeMarker> InTi
 		const double Time = DefaultTimeMarker->GetTime();
 		OnTimeMarkerChangedDelegate.Broadcast(Insights::ETimeChangedFlags::Interactive, Time);
 	}
+	OnCustomTimeMarkerChangedDelegate.Broadcast(Insights::ETimeChangedFlags::Interactive, InTimeMarker);
 }
 
 ////////////////////////////////////////////////////////////////////////////////////////////////////
@@ -3387,6 +3382,7 @@ void STimingView::RaiseTimeMarkerChanged(TSharedRef<Insights::FTimeMarker> InTim
 		const double Time = DefaultTimeMarker->GetTime();
 		OnTimeMarkerChangedDelegate.Broadcast(Insights::ETimeChangedFlags::None, Time);
 	}
+	OnCustomTimeMarkerChangedDelegate.Broadcast(Insights::ETimeChangedFlags::None, InTimeMarker);
 }
 
 ////////////////////////////////////////////////////////////////////////////////////////////////////
@@ -3419,20 +3415,6 @@ void STimingView::AddOverlayWidget(const TSharedRef<SWidget>& InWidget)
 
 ////////////////////////////////////////////////////////////////////////////////////////////////////
 
-void STimingView::UpdateAggregatedStats()
-{
-	if (bAssetLoadingMode)
-	{
-		TSharedPtr<SLoadingProfilerWindow> LoadingProfilerWnd = FLoadingProfilerManager::Get()->GetProfilerWindow();
-		if (LoadingProfilerWnd.IsValid())
-		{
-			LoadingProfilerWnd->UpdateTableTreeViews();
-		}
-	}
-}
-
-////////////////////////////////////////////////////////////////////////////////////////////////////
-
 void STimingView::SetAndCenterOnTimeMarker(double Time)
 {
 	SetTimeMarker(Time);
@@ -4869,13 +4851,8 @@ void STimingView::OnTrackVisibilityChanged()
 	}
 	Tooltip.SetDesiredOpacity(0.0f);
 
-	//TODO: ThreadFilterChangedEvent.Broadcast();
+	//TODO: TrackVisibilityChangedEvent.Broadcast();
 	FTimingProfilerManager::Get()->OnThreadFilterChanged();
-
-	if (SelectionStartTime < SelectionEndTime)
-	{
-		UpdateAggregatedStats();
-	}
 }
 
 ////////////////////////////////////////////////////////////////////////////////////////////////////
diff --git a/Engine/Source/Developer/TraceInsights/Private/Insights/Widgets/STimingView.h b/Engine/Source/Developer/TraceInsights/Private/Insights/Widgets/STimingView.h
--- a/Engine/Source/Developer/TraceInsights/Private/Insights/Widgets/STimingView.h
+++ b/Engine/Source/Developer/TraceInsights/Private/Insights/Widgets/STimingView.h
@@ -267,6 +267,7 @@ public:
 
 	virtual Insights::FSelectionChangedDelegate& OnSelectionChanged() override { return OnSelectionChangedDelegate; }
 	virtual Insights::FTimeMarkerChangedDelegate& OnTimeMarkerChanged() override { return OnTimeMarkerChangedDelegate; }
+	virtual Insights::FCustomTimeMarkerChangedDelegate& OnCustomTimeMarkerChanged() override { return OnCustomTimeMarkerChangedDelegate; }
 	virtual Insights::FHoveredTrackChangedDelegate& OnHoveredTrackChanged() override { return OnHoveredTrackChangedDelegate; }
 	virtual Insights::FHoveredEventChangedDelegate& OnHoveredEventChanged() override { return OnHoveredEventChangedDelegate; }
 	virtual Insights::FSelectedTrackChangedDelegate& OnSelectedTrackChanged() override { return OnSelectedTrackChangedDelegate; }
@@ -701,6 +702,7 @@ protected:
 
 	Insights::FSelectionChangedDelegate OnSelectionChangedDelegate;
 	Insights::FTimeMarkerChangedDelegate OnTimeMarkerChangedDelegate;
+	Insights::FCustomTimeMarkerChangedDelegate OnCustomTimeMarkerChangedDelegate;
 	Insights::FHoveredTrackChangedDelegate OnHoveredTrackChangedDelegate;
 	Insights::FHoveredEventChangedDelegate OnHoveredEventChangedDelegate;
 	Insights::FSelectedTrackChangedDelegate OnSelectedTrackChangedDelegate;
diff --git a/Engine/Source/Developer/TraceInsights/Public/Insights/ITimingViewSession.h b/Engine/Source/Developer/TraceInsights/Public/Insights/ITimingViewSession.h
--- a/Engine/Source/Developer/TraceInsights/Public/Insights/ITimingViewSession.h
+++ b/Engine/Source/Developer/TraceInsights/Public/Insights/ITimingViewSession.h
@@ -26,12 +26,26 @@ ENUM_CLASS_FLAGS(ETimeChangedFlags);
 
 ////////////////////////////////////////////////////////////////////////////////////////////////////
 
+class TRACEINSIGHTS_API ITimeMarker
+{
+public:
+	virtual ~ITimeMarker() = default;
+
+	virtual double GetTime() const = 0;
+	virtual void SetTime(const double  InTime) = 0;
+};
+
+////////////////////////////////////////////////////////////////////////////////////////////////////
+
 /** The delegate to be invoked when the selection have been changed */
 DECLARE_MULTICAST_DELEGATE_ThreeParams(FSelectionChangedDelegate, ETimeChangedFlags /*InFlags*/, double /*StartTime*/, double /*EndTime*/);
 
 /** The delegate to be invoked when the time marker has changed */
 DECLARE_MULTICAST_DELEGATE_TwoParams(FTimeMarkerChangedDelegate, ETimeChangedFlags /*InFlags*/, double /*TimeMarker*/);
 
+/** The delegate to be invoked when a custom time marker has changed */
+DECLARE_MULTICAST_DELEGATE_TwoParams(FCustomTimeMarkerChangedDelegate, ETimeChangedFlags /*InFlags*/, TSharedRef<ITimeMarker> /*TimeMarker*/);
+
 /** The delegate to be invoked when the timing track being hovered by the mouse has changed */
 DECLARE_MULTICAST_DELEGATE_OneParam(FHoveredTrackChangedDelegate, const TSharedPtr<FBaseTimingTrack> /*InTrack*/);
 
@@ -97,6 +111,8 @@ public:
 	virtual FSelectionChangedDelegate& OnSelectionChanged() = 0;
 	/** Gets the delegate to be invoked when the time marker has changed. */
 	virtual FTimeMarkerChangedDelegate& OnTimeMarkerChanged() = 0;
+	/** Gets the delegate to be invoked when a custom time marker has changed. */
+	virtual FCustomTimeMarkerChangedDelegate& OnCustomTimeMarkerChanged() = 0;
 
 	/** Gets the delegate to be invoked when the timing track being hovered by the mouse has changed. */
 	virtual FHoveredTrackChangedDelegate& OnHoveredTrackChanged() = 0;
