[Insights] Fixed static analysis warning.

diff --git a/Engine/Source/Developer/TraceInsights/Private/Insights/MemoryProfiler/Widgets/SMemoryProfilerWindow.cpp b/Engine/Source/Developer/TraceInsights/Private/Insights/MemoryProfiler/Widgets/SMemoryProfilerWindow.cpp
--- a/Engine/Source/Developer/TraceInsights/Private/Insights/MemoryProfiler/Widgets/SMemoryProfilerWindow.cpp
+++ b/Engine/Source/Developer/TraceInsights/Private/Insights/MemoryProfiler/Widgets/SMemoryProfilerWindow.cpp
@@ -837,11 +837,11 @@ void SMemoryProfilerWindow::UpdateTimingViewMarkers()
 
 void SMemoryProfilerWindow::OnTimeMarkerChanged(Insights::ETimeChangedFlags InFlags, TSharedRef<Insights::ITimeMarker> InTimeMarker)
 {
-	const uint32 NumTimeMarkers = CustomTimeMarkers.Num();
+	const int32 NumTimeMarkers = CustomTimeMarkers.Num();
 
 	// Find index of the changing time marker.
-	uint32 ChangedTimeMarkerIndex = -1;
-	for (uint32 Index = 0; Index < NumTimeMarkers; ++Index)
+	int32 ChangedTimeMarkerIndex = -1;
+	for (int32 Index = 0; Index < NumTimeMarkers; ++Index)
 	{
 		TSharedRef<Insights::FTimeMarker>& TimeMarker = CustomTimeMarkers[Index];
 		if (TimeMarker == InTimeMarker)
@@ -854,7 +854,7 @@ void SMemoryProfilerWindow::OnTimeMarkerChanged(Insights::ETimeChangedFlags InFl
 	// Ensure the rest of time markers are orderd by time.
 	if (ChangedTimeMarkerIndex >= 0)
 	{
-		for (uint32 Index = 0; Index < ChangedTimeMarkerIndex; ++Index)
+		for (int32 Index = 0; Index < ChangedTimeMarkerIndex; ++Index)
 		{
 			TSharedRef<Insights::FTimeMarker>& TimeMarker = CustomTimeMarkers[Index];
 			if (TimeMarker->GetTime() > InTimeMarker->GetTime())
@@ -862,7 +862,7 @@ void SMemoryProfilerWindow::OnTimeMarkerChanged(Insights::ETimeChangedFlags InFl
 				TimeMarker->SetTime(InTimeMarker->GetTime());
 			}
 		}
-		for (uint32 Index = ChangedTimeMarkerIndex + 1; Index < NumTimeMarkers; ++Index)
+		for (int32 Index = ChangedTimeMarkerIndex + 1; Index < NumTimeMarkers; ++Index)
 		{
 			TSharedRef<Insights::FTimeMarker>& TimeMarker = CustomTimeMarkers[Index];
 			if (TimeMarker->GetTime() < InTimeMarker->GetTime())
