[Insights] Memory Insights: Change Time Marker A when changing the default time marker (i.e. when using Ctrl + mouse click/drag).

diff --git a/Engine/Source/Developer/TraceInsights/Private/Insights/MemoryProfiler/Widgets/SMemoryProfilerWindow.cpp b/Engine/Source/Developer/TraceInsights/Private/Insights/MemoryProfiler/Widgets/SMemoryProfilerWindow.cpp
--- a/Engine/Source/Developer/TraceInsights/Private/Insights/MemoryProfiler/Widgets/SMemoryProfilerWindow.cpp
+++ b/Engine/Source/Developer/TraceInsights/Private/Insights/MemoryProfiler/Widgets/SMemoryProfilerWindow.cpp
@@ -851,6 +851,17 @@ void SMemoryProfilerWindow::OnTimeMarkerChanged(Insights::ETimeChangedFlags InFl
 		}
 	}
 
+	// Change Time Marker A when changing the default time marker (i.e. when using Ctrl + click/drag).
+	if (ChangedTimeMarkerIndex < 0 &&
+		NumTimeMarkers > 0 &&
+		TimingView.IsValid() &&
+		InTimeMarker == TimingView->GetDefaultTimeMarker())
+	{
+		TSharedRef<Insights::FTimeMarker>& TimeMarkerA = CustomTimeMarkers[0];
+		TimeMarkerA->SetTime(InTimeMarker->GetTime());
+		ChangedTimeMarkerIndex = 0;
+	}
+
 	// Ensure the rest of time markers are orderd by time.
 	if (ChangedTimeMarkerIndex >= 0)
 	{
