[Insights] Add UI to set initial channel when late connecting.

diff --git a/Engine/Source/Developer/TraceInsights/Private/Insights/Widgets/SStartPageWindow.cpp b/Engine/Source/Developer/TraceInsights/Private/Insights/Widgets/SStartPageWindow.cpp
--- a/Engine/Source/Developer/TraceInsights/Private/Insights/Widgets/SStartPageWindow.cpp
+++ b/Engine/Source/Developer/TraceInsights/Private/Insights/Widgets/SStartPageWindow.cpp
@@ -1863,6 +1863,54 @@ TSharedRef<SWidget> SConnectionWindow::ConstructConnectPanel()
 			]
 		]
 
+		+ SVerticalBox::Slot()
+		.AutoHeight()
+		.HAlign(HAlign_Fill)
+		.Padding(12.0f, 12.0f, 12.0f, 0.0f)
+		[
+			SNew(SHorizontalBox)
+
+			+ SHorizontalBox::Slot()
+			.AutoWidth()
+			.VAlign(VAlign_Top)
+			[
+				SNew(SBox)
+				.MinDesiredWidth(180.0f)
+				.HAlign(HAlign_Right)
+				[
+					SNew(STextBlock)
+					.Text(LOCTEXT("TraceRecorderChannelsText", "Initial channels"))
+				]
+			]
+
+			+ SHorizontalBox::Slot()
+			.FillWidth(1.0f)
+			.VAlign(VAlign_Center)
+			.Padding(6.0f, 0.0f, 0.0f, 0.0f)
+			[
+				SNew(SVerticalBox)
+
+				+SVerticalBox::Slot()
+				.AutoHeight()
+				.HAlign(HAlign_Fill)
+				[
+					 SNew(SBox)
+					 .MinDesiredWidth(120.0f)
+					 [
+						 SAssignNew(ChannelsTextBox, SEditableTextBox)
+					 ]
+				]
+
+				+SVerticalBox::Slot()
+				.AutoHeight()
+				.HAlign(HAlign_Fill)
+				[
+					SNew(STextBlock)
+					.Text(LOCTEXT("TraceRecorderChannelNoteText", "Comma-separated list of channel names (or \"default\"= cpu,gpu,frame,log,bookmark) to enable when connected."))
+				]
+			]
+		]
+	
 		+ SVerticalBox::Slot()
 		.AutoHeight()
 		.HAlign(HAlign_Fill)
@@ -1909,6 +1957,7 @@ TSharedRef<SWidget> SConnectionWindow::ConstructConnectPanel()
 	}
 
 	RunningInstanceAddressTextBox->SetText(LocalHost);
+	ChannelsTextBox->SetText(FText::FromStringView(TEXT("default")));
 
 	return Widget;
 }
@@ -1942,9 +1991,11 @@ FReply SConnectionWindow::Connect_OnClicked()
 		Prerequisites.Add(ConnectTask);
 		PrerequisitesPtr = &Prerequisites;
 	}
+	
+	const FString ChannelsExpandedStr = ChannelsTextBox->GetText().ToString().Replace(TEXT("default"), TEXT("cpu,gpu,frame,log,bookmark"));
 
 	FGraphEventRef PreConnectTask = FFunctionGraphTask::CreateAndDispatchWhenReady(
-		[this, TraceRecorderAddressStr, RunningInstanceAddressStr]
+		[this, TraceRecorderAddressStr, RunningInstanceAddressStr, ChannelsExpandedStr]
 		{
 			bIsConnecting = true;
 
@@ -1955,6 +2006,8 @@ FReply SConnectionWindow::Connect_OnClicked()
 			{
 				UE_LOG(TraceInsights, Log, TEXT("[Connection] SendSendTo(\"%s\")..."), *TraceRecorderAddressStr);
 				ControlClient.SendSendTo(*TraceRecorderAddressStr);
+				UE_LOG(TraceInsights, Log, TEXT("[Connection] ToggleChannel(\"%s\")..."), *ChannelsExpandedStr);
+				ControlClient.SendToggleChannel(*ChannelsExpandedStr, true);
 				bIsConnectedSuccessfully = true;
 			}
 			else
diff --git a/Engine/Source/Developer/TraceInsights/Private/Insights/Widgets/SStartPageWindow.h b/Engine/Source/Developer/TraceInsights/Private/Insights/Widgets/SStartPageWindow.h
--- a/Engine/Source/Developer/TraceInsights/Private/Insights/Widgets/SStartPageWindow.h
+++ b/Engine/Source/Developer/TraceInsights/Private/Insights/Widgets/SStartPageWindow.h
@@ -298,6 +298,7 @@ private:
 	TSharedPtr<SVerticalBox> MainContentPanel;
 	TSharedPtr<SEditableTextBox> TraceRecorderAddressTextBox;
 	TSharedPtr<SEditableTextBox> RunningInstanceAddressTextBox;
+	TSharedPtr<SEditableTextBox> ChannelsTextBox;
 
 	/** Widget for the non-intrusive notifications. */
 	TSharedPtr<SNotificationList> NotificationList;
