[Insights] Time converter for Filters

diff --git a/Engine/Source/Developer/TraceInsights/Private/Insights/MemoryProfiler/ViewModels/MemAllocTable.cpp b/Engine/Source/Developer/TraceInsights/Private/Insights/MemoryProfiler/ViewModels/MemAllocTable.cpp
--- a/Engine/Source/Developer/TraceInsights/Private/Insights/MemoryProfiler/ViewModels/MemAllocTable.cpp
+++ b/Engine/Source/Developer/TraceInsights/Private/Insights/MemoryProfiler/ViewModels/MemAllocTable.cpp
@@ -16,6 +16,7 @@
 #include "Insights/Table/ViewModels/TableCellValueGetter.h"
 #include "Insights/Table/ViewModels/TableCellValueSorter.h"
 #include "Insights/Table/ViewModels/TableColumn.h"
+#include "Insights/ViewModels/TimeFilterValueConverter.h"
 
 #define LOCTEXT_NAMESPACE "Insights::FMemAllocTable"
 
@@ -311,6 +312,9 @@ void FMemAllocTable::AddDefaultColumns()
 		TSharedRef<ITableCellValueSorter> Sorter = MakeShared<FSorterByDoubleValue>(ColumnRef);
 		Column.SetValueSorter(Sorter);
 
+		TSharedRef<IFilterValueConverter> Converter = MakeShared<FTimeFilterValueConverter>();
+		Column.SetValueConverter(Converter);
+
 		Column.SetAggregation(ETableColumnAggregation::Min);
 
 		AddColumn(ColumnRef);
@@ -369,6 +373,9 @@ void FMemAllocTable::AddDefaultColumns()
 		TSharedRef<ITableCellValueSorter> Sorter = MakeShared<FSorterByDoubleValue>(ColumnRef);
 		Column.SetValueSorter(Sorter);
 
+		TSharedRef<IFilterValueConverter> Converter = MakeShared<FTimeFilterValueConverter>();
+		Column.SetValueConverter(Converter);
+
 		Column.SetAggregation(ETableColumnAggregation::Max);
 
 		AddColumn(ColumnRef);
@@ -427,6 +434,9 @@ void FMemAllocTable::AddDefaultColumns()
 		TSharedRef<ITableCellValueSorter> Sorter = MakeShared<FSorterByDoubleValue>(ColumnRef);
 		Column.SetValueSorter(Sorter);
 
+		TSharedRef<IFilterValueConverter> Converter = MakeShared<FTimeFilterValueConverter>();
+		Column.SetValueConverter(Converter);
+
 		AddColumn(ColumnRef);
 	}
 	//////////////////////////////////////////////////
diff --git a/Engine/Source/Developer/TraceInsights/Private/Insights/Tests/UITests.cpp b/Engine/Source/Developer/TraceInsights/Private/Insights/Tests/UITests.cpp
--- a/Engine/Source/Developer/TraceInsights/Private/Insights/Tests/UITests.cpp
+++ b/Engine/Source/Developer/TraceInsights/Private/Insights/Tests/UITests.cpp
@@ -8,6 +8,7 @@
 #include "Insights/InsightsManager.h"
 #include "Insights/MemoryProfiler/ViewModels/MemAllocFilterValueConverter.h"
 #include "Insights/TimingProfilerManager.h"
+#include "Insights/ViewModels/TimeFilterValueConverter.h"
 #include "Insights/Widgets/STimingProfilerWindow.h"
 
 ////////////////////////////////////////////////////////////////////////////////////////////////////
@@ -81,3 +82,52 @@ bool FMemoryFilterValueConverterTest::RunTest(const FString& Parameters)
 }
 
 ////////////////////////////////////////////////////////////////////////////////////////////////////
+
+bool FTimeFilterValueConverterTest::RunTest(const FString& Parameters)
+{
+	Insights::FTimeFilterValueConverter Converter;
+
+	FText Error;
+	double Value;
+
+	Converter.Convert(TEXT("15.3"), Value, Error);
+	TestEqual(TEXT("ValueInSeconds"), 15.3, Value);
+
+	Converter.Convert(TEXT("0.2"), Value, Error);
+	TestEqual(TEXT("ValueInSeconds2"), 0.2, Value);
+
+	Converter.Convert(TEXT("0.4s"), Value, Error);
+	TestEqual(TEXT("ValueInSeconds3"), 0.4, Value);
+
+	Converter.Convert(TEXT("125.56ms"), Value, Error);
+	TestEqual(TEXT("ValueInMiliseconds"), 0.12556, Value);
+
+	Converter.Convert(TEXT("14.2 µs"), Value, Error);
+	TestEqual(TEXT("ValueInMicroseconds"), 14.2 * 1.e-6, Value);
+
+	Converter.Convert(TEXT("0.72us"), Value, Error);
+	TestEqual(TEXT("ValueInMicroseconds2"), 0.72 * 1.e-6, Value); 
+
+	Converter.Convert(TEXT("3ns"), Value, Error);
+	TestEqual(TEXT("ValueInNanoseconds"), 3 * 1.e-9, Value);
+
+	Converter.Convert(TEXT("17ns"), Value, Error);
+	TestEqual(TEXT("ValueInPicoseconds"), 17 * 1.e-12, Value);
+
+	Converter.Convert(TEXT("0.5m"), Value, Error);
+	TestEqual(TEXT("ValueInMinutes"), 30.0, Value);
+
+	Converter.Convert(TEXT("1.1h"), Value, Error);
+	TestEqual(TEXT("ValueInHours"), 3960.0, Value);
+
+	Converter.Convert(TEXT("2d"), Value, Error);
+	TestEqual(TEXT("ValueInHours"), 2 * 60 * 60 * 24.0, Value);
+
+	TestFalse(TEXT("Fail1"), Converter.Convert(TEXT("1.23ss"), Value, Error));
+	TestFalse(TEXT("Fail2"), Converter.Convert(TEXT("abc"), Value, Error));
+	TestFalse(TEXT("FailInvalidChar"), Converter.Convert(TEXT("0,2"), Value, Error));
+
+	return !HasAnyErrors();
+}
+
+////////////////////////////////////////////////////////////////////////////////////////////////////
diff --git a/Engine/Source/Developer/TraceInsights/Private/Insights/Tests/UITests.h b/Engine/Source/Developer/TraceInsights/Private/Insights/Tests/UITests.h
--- a/Engine/Source/Developer/TraceInsights/Private/Insights/Tests/UITests.h
+++ b/Engine/Source/Developer/TraceInsights/Private/Insights/Tests/UITests.h
@@ -15,4 +15,5 @@ IMPLEMENT_SIMPLE_AUTOMATION_TEST(FHideAndShowAllTimingViewTabs, "Insights.HideAn
 #endif
 
 IMPLEMENT_SIMPLE_AUTOMATION_TEST(FMemoryFilterValueConverterTest, "Insights.FMemoryFilterValueConverterTest", EAutomationTestFlags::ApplicationContextMask | EAutomationTestFlags::ProductFilter)
+IMPLEMENT_SIMPLE_AUTOMATION_TEST(FTimeFilterValueConverterTest, "Insights.FTimeFilterValueConverterTest", EAutomationTestFlags::ApplicationContextMask | EAutomationTestFlags::ProductFilter)
 
diff --git a/Engine/Source/Developer/TraceInsights/Private/Insights/ViewModels/TimeFilterValueConverter.cpp b/Engine/Source/Developer/TraceInsights/Private/Insights/ViewModels/TimeFilterValueConverter.cpp
new file mode 100644
--- /dev/null
+++ b/Engine/Source/Developer/TraceInsights/Private/Insights/ViewModels/TimeFilterValueConverter.cpp
@@ -0,0 +1,105 @@
+// Copyright Epic Games, Inc. All Rights Reserved.
+
+#include "TimeFilterValueConverter.h"
+
+#include "Insights/InsightsManager.h"
+
+#define LOCTEXT_NAMESPACE "Insights::FTimeFilterValueConverter"
+
+namespace Insights
+{
+
+////////////////////////////////////////////////////////////////////////////////////////////////////
+
+bool FTimeFilterValueConverter::Convert(const FString& Input, double& Output, FText& OutError) const
+{
+	FString TrimmedInput = Input.TrimStartAndEnd();
+	int Index = 0;
+	for (; Index < TrimmedInput.Len(); ++Index)
+	{
+		TCHAR Character = TrimmedInput[Index];
+		if (!TChar<TCHAR>::IsDigit(Character) && !(Character == '.'))
+		{
+			break;
+		}
+	}
+
+	FString Value = TrimmedInput.Left(Index);
+	if (!Value.IsNumeric())
+	{
+		OutError = LOCTEXT("InvalidValue", "Invalid value!");
+	}
+	Output = FCString::Atod(*Value);
+
+	FString Unit = TrimmedInput.RightChop(Index);
+	Unit.TrimStartAndEndInline();
+	if (Unit.Len() == 0)
+	{
+		// The default unit if one is not specified by the user.
+		Unit = TEXT("s");
+	}
+
+	if (Unit.Compare(TEXT("ps"), ESearchCase::IgnoreCase) == 0)
+	{
+		Output *= 1.e-12f;
+		return true;
+	}
+	else if (Unit.Compare(TEXT("ns"), ESearchCase::IgnoreCase) == 0)
+	{
+		Output *= 1.e-9f;
+		return true;
+	}
+	else if (Unit.Compare(TEXT("µs"), ESearchCase::IgnoreCase) == 0 || Unit.Compare(TEXT("us"), ESearchCase::IgnoreCase) == 0)
+	{
+		Output *= 1.e-6f;
+		return true;
+	}
+	else if (Unit.Compare(TEXT("ms"), ESearchCase::IgnoreCase) == 0)
+	{
+		Output *= 1.e-3f;
+		return true;
+	}
+	else if (Unit.Compare(TEXT("s"), ESearchCase::IgnoreCase) == 0)
+	{
+		return true;
+	}
+	else if (Unit.Compare(TEXT("m"), ESearchCase::IgnoreCase) == 0)
+	{
+		Output *= 60;
+		return true;
+	}
+	else if (Unit.Compare(TEXT("h"), ESearchCase::IgnoreCase) == 0)
+	{
+		Output *= 60 * 60;
+		return true;
+	}
+	else if (Unit.Compare(TEXT("d"), ESearchCase::IgnoreCase) == 0)
+	{
+		Output *= 60 * 60 * 24;
+		return true;
+	}
+
+	OutError = LOCTEXT("InvalidUnit", "Invalid unit! Can be one of: ps,ns,µs/us,ms,s,m,d.");
+	return false;
+}
+
+////////////////////////////////////////////////////////////////////////////////////////////////////
+
+FText FTimeFilterValueConverter::GetTooltipText() const
+{
+	return LOCTEXT("FTimeFilterValueConverterTooltip", "Enter a time value as a floating point number followed by an optional unit (ps,ns,µs/us,ms,s,m,d).\nThe default unit is \"s\" (seconds).");
+}
+
+////////////////////////////////////////////////////////////////////////////////////////////////////
+
+FText FTimeFilterValueConverter::GetHintText() const
+{
+	// An example value.
+	return FText::FromString(TEXT("1.5s"));
+}
+
+////////////////////////////////////////////////////////////////////////////////////////////////////
+
+} // namespace Insights
+
+#undef LOCTEXT_NAMESPACE
\ No newline at end of file
diff --git a/Engine/Source/Developer/TraceInsights/Private/Insights/ViewModels/TimeFilterValueConverter.h b/Engine/Source/Developer/TraceInsights/Private/Insights/ViewModels/TimeFilterValueConverter.h
new file mode 100644
--- /dev/null
+++ b/Engine/Source/Developer/TraceInsights/Private/Insights/ViewModels/TimeFilterValueConverter.h
@@ -0,0 +1,18 @@
+// Copyright Epic Games, Inc. All Rights Reserved.
+
+#pragma once
+
+#include "Insights/ViewModels/Filters.h"
+
+namespace Insights
+{
+
+class FTimeFilterValueConverter : public IFilterValueConverter
+{
+public:
+	virtual bool Convert(const FString& Input, double& Output, FText& OutError) const override;
+	virtual FText GetTooltipText() const override;
+	virtual FText GetHintText() const override;
+};
+
+} // namespace Insights
diff --git a/Engine/Source/Developer/TraceInsights/Private/Insights/Widgets/SFilterConfiguratorRow.cpp b/Engine/Source/Developer/TraceInsights/Private/Insights/Widgets/SFilterConfiguratorRow.cpp
--- a/Engine/Source/Developer/TraceInsights/Private/Insights/Widgets/SFilterConfiguratorRow.cpp
+++ b/Engine/Source/Developer/TraceInsights/Private/Insights/Widgets/SFilterConfiguratorRow.cpp
@@ -506,8 +506,16 @@ bool SFilterConfiguratorRow::TextBox_OnVerifyTextChanged(const FText& InText, FT
 	TSharedPtr<IFilterValueConverter> Converter = FilterConfiguratorNodePtr->GetSelectedFilter()->Converter;
 	if (Converter.IsValid())
 	{
-		int64 Value;
-		return Converter->Convert(InText.ToString(), Value, OutErrorMessage);
+		if (FilterConfiguratorNodePtr->GetSelectedFilter()->DataType == EFilterDataType::Int64)
+		{
+			int64 Value;
+			return Converter->Convert(InText.ToString(), Value, OutErrorMessage);
+		}		
+		else if (FilterConfiguratorNodePtr->GetSelectedFilter()->DataType == EFilterDataType::Double)
+		{
+			double Value;
+			return Converter->Convert(InText.ToString(), Value, OutErrorMessage);
+		}
 	}
 
 	return true;
diff --git a/Engine/Source/Developer/TraceInsights/Private/Insights/Widgets/STimingView.cpp b/Engine/Source/Developer/TraceInsights/Private/Insights/Widgets/STimingView.cpp
--- a/Engine/Source/Developer/TraceInsights/Private/Insights/Widgets/STimingView.cpp
+++ b/Engine/Source/Developer/TraceInsights/Private/Insights/Widgets/STimingView.cpp
@@ -61,6 +61,7 @@
 #include "Insights/ViewModels/LoadingTimingTrack.h"
 #include "Insights/ViewModels/MarkersTimingTrack.h"
 #include "Insights/ViewModels/ThreadTimingTrack.h"
+#include "Insights/ViewModels/TimeFilterValueConverter.h"
 #include "Insights/ViewModels/TimeRulerTrack.h"
 #include "Insights/ViewModels/TimingEventSearch.h"
 #include "Insights/ViewModels/TimingGraphTrack.h"
@@ -4732,10 +4733,16 @@ void STimingView::QuickFind_Execute()
 	{
 		TSharedPtr<FFilterConfigurator> NewFilterConfigurator = MakeShared<FFilterConfigurator>();
 		TSharedPtr<TArray<TSharedPtr<struct FFilter>>>& AvailableFilters = NewFilterConfigurator->GetAvailableFilters();
+		TSharedPtr<FFilter> CurrentFilter;
 
-		AvailableFilters->Add(MakeShared<FFilter>(static_cast<int32>(EFilterField::StartTime), LOCTEXT("StartTime", "Start Time"), LOCTEXT("StartTime", "Start Time"), EFilterDataType::Double, FFilterService::Get()->GetDoubleOperators()));
-		AvailableFilters->Add(MakeShared<FFilter>(static_cast<int32>(EFilterField::EndTime), LOCTEXT("EndTime", "End Time"), LOCTEXT("EndTime", "End Time"), EFilterDataType::Double, FFilterService::Get()->GetDoubleOperators()));
-		AvailableFilters->Add(MakeShared<FFilter>(static_cast<int32>(EFilterField::Duration), LOCTEXT("Duration", "Duration"), LOCTEXT("Duration", "Duration"), EFilterDataType::Double, FFilterService::Get()->GetDoubleOperators()));
+		CurrentFilter = AvailableFilters->Add_GetRef(MakeShared<FFilter>(static_cast<int32>(EFilterField::StartTime), LOCTEXT("StartTime", "Start Time"), LOCTEXT("StartTime", "Start Time"), EFilterDataType::Double, FFilterService::Get()->GetDoubleOperators()));
+		CurrentFilter->Converter = MakeShared<FTimeFilterValueConverter>();
+
+		CurrentFilter = AvailableFilters->Add_GetRef(MakeShared<FFilter>(static_cast<int32>(EFilterField::EndTime), LOCTEXT("EndTime", "End Time"), LOCTEXT("EndTime", "End Time"), EFilterDataType::Double, FFilterService::Get()->GetDoubleOperators()));
+		CurrentFilter->Converter = MakeShared<FTimeFilterValueConverter>();
+
+		CurrentFilter = AvailableFilters->Add_GetRef(MakeShared<FFilter>(static_cast<int32>(EFilterField::Duration), LOCTEXT("Duration", "Duration"), LOCTEXT("Duration", "Duration"), EFilterDataType::Double, FFilterService::Get()->GetDoubleOperators()));
+		CurrentFilter->Converter = MakeShared<FTimeFilterValueConverter>();
 
 		TSharedPtr<FFilterWithSuggestions> TrackFilter = MakeShared<FFilterWithSuggestions>(static_cast<int32>(EFilterField::TrackName), LOCTEXT("Track", "Track"), LOCTEXT("Track", "Track"), EFilterDataType::String, FFilterService::Get()->GetStringOperators());
 		TrackFilter->Callback = [this](const FString& Text, TArray<FString>& OutSuggestions)
