[Insights] Fixed style for BreadcrumbTrail widget used for Hierarcy in tree table views.

diff --git a/Engine/Source/Developer/TraceInsights/Private/Insights/Table/Widgets/STableTreeView.cpp b/Engine/Source/Developer/TraceInsights/Private/Insights/Table/Widgets/STableTreeView.cpp
--- a/Engine/Source/Developer/TraceInsights/Private/Insights/Table/Widgets/STableTreeView.cpp
+++ b/Engine/Source/Developer/TraceInsights/Private/Insights/Table/Widgets/STableTreeView.cpp
@@ -350,11 +350,10 @@ TSharedRef<SWidget> STableTreeView::ConstructHierarchyBreadcrumbTrail()
 	.VAlign(VAlign_Center)
 	[
 		SAssignNew(GroupingBreadcrumbTrail, SBreadcrumbTrail<TSharedPtr<FTreeNodeGrouping>>)
-		.ButtonContentPadding(FMargin(1.0f, 1.0f))
-		//.DelimiterImage(FInsightsStyle::GetBrush("HierarchyPathDelimiter"))
-		//.TextStyle(FInsightsStyle::Get(), "HierarchyPathText")
-		//.ShowLeadingDelimiter(true)
-		//.PersistentBreadcrumbs(true)
+		.ButtonStyle(FAppStyle::Get(), "SimpleButton")
+		.ButtonContentPadding(FMargin(3.0f, 2.0f))
+		.TextStyle(FAppStyle::Get(), "NormalText")
+		.DelimiterImage(FAppStyle::Get().GetBrush("Icons.ChevronRight"))
 		.OnCrumbClicked(this, &STableTreeView::OnGroupingCrumbClicked)
 		.GetCrumbMenuContent(this, &STableTreeView::GetGroupingCrumbMenuContent)
 	];
