[Insights] Memory Insights: Added number of allocs in status of Mem Allocs table.

diff --git a/Engine/Source/Developer/TraceInsights/Private/Insights/MemoryProfiler/Widgets/SMemAllocTableTreeView.cpp b/Engine/Source/Developer/TraceInsights/Private/Insights/MemoryProfiler/Widgets/SMemAllocTableTreeView.cpp
--- a/Engine/Source/Developer/TraceInsights/Private/Insights/MemoryProfiler/Widgets/SMemAllocTableTreeView.cpp
+++ b/Engine/Source/Developer/TraceInsights/Private/Insights/MemoryProfiler/Widgets/SMemAllocTableTreeView.cpp
@@ -165,6 +165,7 @@ void SMemAllocTableTreeView::RebuildTree(bool bResync)
 					TableTreeNodes.Add(NodePtr);
 				}
 				ensure(TableTreeNodes.Num() == TotalAllocCount - HeapAllocCount);
+				UpdateQueryInfo();
 			}
 		}
 	}
@@ -1107,7 +1108,7 @@ void SMemAllocTableTreeView::UpdateQueryInfo()
 			check(false);
 		}
 
-		QueryInfo = FText::Format(LOCTEXT("QueryInfoFmt", "{0} ({1})"), Rule->GetVerboseName(), TimeMarkersText);
+		QueryInfo = FText::Format(LOCTEXT("QueryInfoFmt", "{0} ({1}) : {2} allocs"), Rule->GetVerboseName(), TimeMarkersText, FText::AsNumber(TableTreeNodes.Num()));
 		QueryInfoTooltip = Rule->GetDescription();
 	}
 }
