[Insights]
  - Memory Insights / Table Tree Views: Improved performance of tree grouping, filtering and aggregation with more than 2x (ex. filtering step for a large 145M allocs table from 52s --> 21s).
  - Memory Insights: Added display name (currently used in logging) for Mem Allocs table.
  - Table Tree Views: Added logging for duration of aggregation step (when executed from the filtering step).

diff --git a/Engine/Source/Developer/TraceInsights/Private/Insights/MemoryProfiler/ViewModels/MemAllocNode.h b/Engine/Source/Developer/TraceInsights/Private/Insights/MemoryProfiler/ViewModels/MemAllocNode.h
--- a/Engine/Source/Developer/TraceInsights/Private/Insights/MemoryProfiler/ViewModels/MemAllocNode.h
+++ b/Engine/Source/Developer/TraceInsights/Private/Insights/MemoryProfiler/ViewModels/MemAllocNode.h
@@ -16,18 +16,6 @@
 namespace Insights
 {
 
-enum class EMemAllocNodeType
-{
-	/** The MemAllocNode is an allocation node. */
-	MemAlloc,
-
-	/** The MemAllocNode is a group node. */
-	Group,
-
-	/** Invalid enum type, may be used as a number of enumerations. */
-	InvalidOrMax,
-};
-
 ////////////////////////////////////////////////////////////////////////////////////////////////////
 
 class FMemAllocNode;
@@ -57,30 +45,20 @@ public:
 	/** Initialization constructor for the MemAlloc node. */
 	explicit FMemAllocNode(const FName InName, TWeakPtr<FMemAllocTable> InParentTable, int32 InRowIndex)
 		: FTableTreeNode(InName, InParentTable, InRowIndex)
-		, Type(EMemAllocNodeType::MemAlloc)
+		, MemAllocTable(InParentTable.Pin().Get())
 	{
 	}
 
 	/** Initialization constructor for the group node. */
 	explicit FMemAllocNode(const FName InGroupName, TWeakPtr<FMemAllocTable> InParentTable)
 		: FTableTreeNode(InGroupName, InParentTable)
-		, Type(EMemAllocNodeType::Group)
+		, MemAllocTable(InParentTable.Pin().Get())
 	{
 	}
 
 	virtual const FName& GetTypeName() const override { return TypeName; }
 
-	/**
-	 * @return a type of this MemAlloc node or EMemAllocNodeType::Group for group nodes.
-	 */
-	EMemAllocNodeType GetType() const { return Type; }
-
-	FMemAllocTable& GetMemTableChecked() const
-	{
-		const TSharedPtr<FTable>& TablePin = GetParentTable().Pin();
-		check(TablePin.IsValid());
-		return *StaticCastSharedPtr<FMemAllocTable>(TablePin);
-	}
+	FMemAllocTable& GetMemTableChecked() const { return *MemAllocTable; }
 
 	bool IsValidMemAlloc() const { return GetMemTableChecked().IsValidRowIndex(GetRowIndex()); }
 	const FMemoryAlloc* GetMemAlloc() const { return GetMemTableChecked().GetMemAlloc(GetRowIndex()); }
@@ -88,7 +66,7 @@ public:
 	FText GetFullCallstack() const;
 
 private:
-	const EMemAllocNodeType Type;
+	FMemAllocTable* MemAllocTable;
 };
 
 ////////////////////////////////////////////////////////////////////////////////////////////////////
diff --git a/Engine/Source/Developer/TraceInsights/Private/Insights/MemoryProfiler/Widgets/SMemoryProfilerWindow.cpp b/Engine/Source/Developer/TraceInsights/Private/Insights/MemoryProfiler/Widgets/SMemoryProfilerWindow.cpp
--- a/Engine/Source/Developer/TraceInsights/Private/Insights/MemoryProfiler/Widgets/SMemoryProfilerWindow.cpp
+++ b/Engine/Source/Developer/TraceInsights/Private/Insights/MemoryProfiler/Widgets/SMemoryProfilerWindow.cpp
@@ -289,6 +289,7 @@ TSharedRef<SDockTab> SMemoryProfilerWindow::SpawnTab_MemAllocTableTreeView(const
 
 	TSharedRef<Insights::FMemAllocTable> MemAllocTable = MakeShared<Insights::FMemAllocTable>();
 	MemAllocTable->Reset();
+	MemAllocTable->SetDisplayName(FText::FromString(TEXT("MemAllocs")));
 
 	TSharedPtr<Insights::SMemAllocTableTreeView> MemAllocTableTreeView;
 
diff --git a/Engine/Source/Developer/TraceInsights/Private/Insights/Table/Widgets/STableTreeView.cpp b/Engine/Source/Developer/TraceInsights/Private/Insights/Table/Widgets/STableTreeView.cpp
--- a/Engine/Source/Developer/TraceInsights/Private/Insights/Table/Widgets/STableTreeView.cpp
+++ b/Engine/Source/Developer/TraceInsights/Private/Insights/Table/Widgets/STableTreeView.cpp
@@ -1001,8 +1001,14 @@ void STableTreeView::ApplyFiltering()
 		}
 	}
 
+	Stopwatch.Update();
+	const double Time1 = Stopwatch.GetAccumulatedTime();
+
 	UpdateAggregatedValues(*Root);
 
+	Stopwatch.Update();
+	const double Time2 = Stopwatch.GetAccumulatedTime();
+
 	// Cannot call TreeView functions from other threads than MainThread and SlateThread.
 	if (!bRunInAsyncMode)
 	{
@@ -1045,7 +1051,7 @@ void STableTreeView::ApplyFiltering()
 	const double TotalTime = Stopwatch.GetAccumulatedTime();
 	if (TotalTime > 0.1)
 	{
-		UE_LOG(TraceInsights, Log, TEXT("[Tree - %s] Filtering completed in %.3fs."), *Table->GetDisplayName().ToString(), TotalTime);
+		UE_LOG(TraceInsights, Log, TEXT("[Tree - %s] Filtering completed in %.3fs (aggregation done in %.3fs)."), *Table->GetDisplayName().ToString(), TotalTime, Time2 - Time1);
 	}
 }
 
@@ -1450,16 +1456,16 @@ void STableTreeView::UpdateAggregatedValues(FTableTreeNode& GroupNode)
 			case ETableColumnAggregation::Sum:
 				if (Column.GetDataType() == ETableCellDataType::Float || Column.GetDataType() == ETableCellDataType::Double)
 				{
-					STableTreeView::UpdateAggregationRec<double>(Column, GroupNode, 0, true, [](double InValue, TOptional<FTableCellValue> InTableCellValue)
+					STableTreeView::UpdateAggregationRec<double>(Column, GroupNode, 0, true, [](double InValue, const FTableCellValue& InTableCellValue)
 					{
-						return InValue + InTableCellValue->AsDouble();
+						return InValue + InTableCellValue.AsDouble();
 					});
 				}
 				else
 				{
-					STableTreeView::UpdateAggregationRec<int64>(Column, GroupNode, 0, true, [](int64 InValue, TOptional<FTableCellValue> InTableCellValue)
+					STableTreeView::UpdateAggregationRec<int64>(Column, GroupNode, 0, true, [](int64 InValue, const FTableCellValue& InTableCellValue)
 					{
-						return InValue + InTableCellValue->AsInt64();
+						return InValue + InTableCellValue.AsInt64();
 					});
 				}
 				break;
@@ -1467,16 +1473,16 @@ void STableTreeView::UpdateAggregatedValues(FTableTreeNode& GroupNode)
 			case ETableColumnAggregation::Min:
 				if (Column.GetDataType() == ETableCellDataType::Float || Column.GetDataType() == ETableCellDataType::Double)
 				{
-					STableTreeView::UpdateAggregationRec<double>(Column, GroupNode, std::numeric_limits<double>::max(), false, [](double InValue, TOptional<FTableCellValue> InTableCellValue)
+					STableTreeView::UpdateAggregationRec<double>(Column, GroupNode, std::numeric_limits<double>::max(), false, [](double InValue, const FTableCellValue& InTableCellValue)
 					{
-						return FMath::Min(InValue, InTableCellValue->AsDouble());
+						return FMath::Min(InValue, InTableCellValue.AsDouble());
 					});
 				}
 				else
 				{
-					STableTreeView::UpdateAggregationRec<int64>(Column, GroupNode, std::numeric_limits<int64>::max(), false, [](int64 InValue, TOptional<FTableCellValue> InTableCellValue)
+					STableTreeView::UpdateAggregationRec<int64>(Column, GroupNode, std::numeric_limits<int64>::max(), false, [](int64 InValue, const FTableCellValue& InTableCellValue)
 					{
-						return FMath::Min(InValue, InTableCellValue->AsInt64());
+						return FMath::Min(InValue, InTableCellValue.AsInt64());
 					});
 				}
 				break;
@@ -1484,16 +1490,16 @@ void STableTreeView::UpdateAggregatedValues(FTableTreeNode& GroupNode)
 			case ETableColumnAggregation::Max:
 				if (Column.GetDataType() == ETableCellDataType::Float || Column.GetDataType() == ETableCellDataType::Double)
 				{
-					STableTreeView::UpdateAggregationRec<double>(Column, GroupNode, std::numeric_limits<double>::lowest(), false, [](double InValue, TOptional<FTableCellValue> InTableCellValue)
+					STableTreeView::UpdateAggregationRec<double>(Column, GroupNode, std::numeric_limits<double>::lowest(), false, [](double InValue, const FTableCellValue& InTableCellValue)
 					{
-						return FMath::Max(InValue, InTableCellValue->AsDouble());
+						return FMath::Max(InValue, InTableCellValue.AsDouble());
 					});
 				}
 				else
 				{
-					STableTreeView::UpdateAggregationRec<int64>(Column, GroupNode, std::numeric_limits<int64>::min(), false, [](int64 InValue, TOptional<FTableCellValue> InTableCellValue)
+					STableTreeView::UpdateAggregationRec<int64>(Column, GroupNode, std::numeric_limits<int64>::min(), false, [](int64 InValue, const FTableCellValue& InTableCellValue)
 					{
-						return FMath::Max(InValue, InTableCellValue->AsInt64());
+						return FMath::Max(InValue, InTableCellValue.AsInt64());
 					});
 				}
 				break;
diff --git a/Engine/Source/Developer/TraceInsights/Private/Insights/Table/Widgets/STableTreeView.h b/Engine/Source/Developer/TraceInsights/Private/Insights/Table/Widgets/STableTreeView.h
--- a/Engine/Source/Developer/TraceInsights/Private/Insights/Table/Widgets/STableTreeView.h
+++ b/Engine/Source/Developer/TraceInsights/Private/Insights/Table/Widgets/STableTreeView.h
@@ -251,7 +251,7 @@ protected:
 	void UpdateAggregatedValues(FTableTreeNode& GroupNode);
 	
 	template<typename T>
-	static void UpdateAggregationRec(FTableColumn& Column, FTableTreeNode& GroupNode, T InitialAggregatedValue, bool bSetInitialValue, TFunctionRef<T(T, TOptional<FTableCellValue>)> ValueGetterFunc)
+	static void UpdateAggregationRec(FTableColumn& Column, FTableTreeNode& GroupNode, T InitialAggregatedValue, bool bSetInitialValue, TFunctionRef<T(T, const FTableCellValue&)> ValueGetterFunc)
 	{
 		T AggregatedValue = InitialAggregatedValue;
 
@@ -261,22 +261,23 @@ protected:
 			{
 				continue;
 			}
-			FTableTreeNode& TableNode = *StaticCastSharedPtr<FTableTreeNode>(NodePtr);
-			if (TableNode.IsGroup())
+			
+			if (!NodePtr->IsGroup())
 			{
-				TableNode.ResetAggregatedValues(Column.GetId());
-				UpdateAggregationRec(Column, TableNode, InitialAggregatedValue, bSetInitialValue, ValueGetterFunc);
-				if (TableNode.HasAggregatedValue(Column.GetId()))
+				const TOptional<FTableCellValue> NodeValue = Column.GetValue(*NodePtr);
+				if (NodeValue.IsSet())
 				{
-					AggregatedValue = ValueGetterFunc(AggregatedValue, TableNode.GetAggregatedValue(Column.GetId()));
+					AggregatedValue = ValueGetterFunc(AggregatedValue, NodeValue.GetValue());
 				}
 			}
 			else
 			{
-				const TOptional<FTableCellValue> OptionalValue = Column.GetValue(TableNode);
-				if (OptionalValue.IsSet())
+				FTableTreeNode& TableNode = *(FTableTreeNode*)NodePtr.Get();
+				TableNode.ResetAggregatedValues(Column.GetId());
+				UpdateAggregationRec(Column, TableNode, InitialAggregatedValue, bSetInitialValue, ValueGetterFunc);
+				if (TableNode.HasAggregatedValue(Column.GetId()))
 				{
-					AggregatedValue = ValueGetterFunc(AggregatedValue, OptionalValue);
+					AggregatedValue = ValueGetterFunc(AggregatedValue, TableNode.GetAggregatedValue(Column.GetId()));
 				}
 			}
 		}
