[Insights]
  - Memory Insights: Fixed grouping in allocs tree view when node name exeeds NAME_SIZE.
  - Memory Insights: Improved FName creation in few places to avoid a strlen.

diff --git a/Engine/Source/Developer/TraceInsights/Private/Insights/MemoryProfiler/ViewModels/MemAllocGroupingByCallstack.cpp b/Engine/Source/Developer/TraceInsights/Private/Insights/MemoryProfiler/ViewModels/MemAllocGroupingByCallstack.cpp
--- a/Engine/Source/Developer/TraceInsights/Private/Insights/MemoryProfiler/ViewModels/MemAllocGroupingByCallstack.cpp
+++ b/Engine/Source/Developer/TraceInsights/Private/Insights/MemoryProfiler/ViewModels/MemAllocGroupingByCallstack.cpp
@@ -161,15 +161,20 @@ FName FMemAllocGroupingByCallstack::GetGroupName(const TraceServices::FStackFram
 	const TraceServices::ESymbolQueryResult Result = Frame->Symbol->GetResult();
 	if (Result == TraceServices::ESymbolQueryResult::OK)
 	{
-		return FName(Frame->Symbol->Name);
+		FStringView GroupName(Frame->Symbol->Name);
+		if (GroupName.Len() >= NAME_SIZE)
+		{
+			GroupName = FStringView(Frame->Symbol->Name, NAME_SIZE - 1);
+		}
+		return FName(GroupName, 0);
 	}
 	else if (Result == TraceServices::ESymbolQueryResult::Pending)
 	{
-		return FName(*FString::Printf(TEXT("0x%X [...]"), Frame->Addr));
+		return FName(FString::Printf(TEXT("0x%X [...]"), Frame->Addr), 0);
 	}
 	else
 	{
-		return FName(*FString::Printf(TEXT("0x%X"), Frame->Addr));
+		return FName(FString::Printf(TEXT("0x%X"), Frame->Addr), 0);
 	}
 }
 
diff --git a/Engine/Source/Developer/TraceInsights/Private/Insights/MemoryProfiler/ViewModels/MemAllocGroupingBySize.cpp b/Engine/Source/Developer/TraceInsights/Private/Insights/MemoryProfiler/ViewModels/MemAllocGroupingBySize.cpp
--- a/Engine/Source/Developer/TraceInsights/Private/Insights/MemoryProfiler/ViewModels/MemAllocGroupingBySize.cpp
+++ b/Engine/Source/Developer/TraceInsights/Private/Insights/MemoryProfiler/ViewModels/MemAllocGroupingBySize.cpp
@@ -117,7 +117,7 @@ void FMemAllocGroupingBySize::GroupNodes(const TArray<FTableTreeNodePtr>& Nodes,
 			if (!GroupPtr)
 			{
 				const FThreshold& Threshold = Thresholds[ThresholdIndex];
-				GroupPtr = MakeShared<FTableTreeNode>(FName(*Threshold.Name.ToString()), InParentTable);
+				GroupPtr = MakeShared<FTableTreeNode>(FName(Threshold.Name.ToString(), 0), InParentTable);
 				GroupPtr->SetExpansion(false);
 				ParentGroup.AddChildAndSetGroupPtr(GroupPtr);
 				GroupMap[ThresholdIndex] = GroupPtr;
diff --git a/Engine/Source/Developer/TraceInsights/Private/Insights/MemoryProfiler/ViewModels/MemTagNode.h b/Engine/Source/Developer/TraceInsights/Private/Insights/MemoryProfiler/ViewModels/MemTagNode.h
--- a/Engine/Source/Developer/TraceInsights/Private/Insights/MemoryProfiler/ViewModels/MemTagNode.h
+++ b/Engine/Source/Developer/TraceInsights/Private/Insights/MemoryProfiler/ViewModels/MemTagNode.h
@@ -66,7 +66,7 @@ public:
 public:
 	/** Initialization constructor for the MemTag node. */
 	explicit FMemTagNode(Insights::FMemoryTag* InMemTag)
-		: FBaseTreeNode(FName(*InMemTag->GetStatFullName()), false)
+		: FBaseTreeNode(FName(InMemTag->GetStatFullName(), 0), false)
 		, Type(EMemTagNodeType::MemTag)
 		, MemTag(InMemTag)
 	{
diff --git a/Engine/Source/Developer/TraceInsights/Private/Insights/Table/ViewModels/TreeNodeGrouping.cpp b/Engine/Source/Developer/TraceInsights/Private/Insights/Table/ViewModels/TreeNodeGrouping.cpp
--- a/Engine/Source/Developer/TraceInsights/Private/Insights/Table/ViewModels/TreeNodeGrouping.cpp
+++ b/Engine/Source/Developer/TraceInsights/Private/Insights/Table/ViewModels/TreeNodeGrouping.cpp
@@ -126,7 +126,13 @@ FTreeNodeGroupingByUniqueValue::FTreeNodeGroupingByUniqueValue(TSharedRef<FTable
 FTreeNodeGroupInfo FTreeNodeGroupingByUniqueValue::GetGroupForNode(const FBaseTreeNodePtr InNode) const
 {
 	FTableTreeNodePtr TableTreeNodePtr = StaticCastSharedPtr<FTableTreeNode>(InNode);
-	return { FName(*ColumnRef->GetValueAsText(*TableTreeNodePtr).ToString()), false };
+	FText ValueAsText = ColumnRef->GetValueAsText(*TableTreeNodePtr);
+	FStringView GroupName(ValueAsText.ToString());
+	if (GroupName.Len() >= NAME_SIZE)
+	{
+		GroupName = FStringView(GroupName.GetData(), NAME_SIZE - 1);
+	}
+	return { FName(GroupName, 0), false };
 }
 
 ////////////////////////////////////////////////////////////////////////////////////////////////////
