[Insights] Fixed group name for "unnamed" threads.

diff --git a/Engine/Source/Developer/TraceInsights/Private/Insights/ViewModels/ThreadTimingTrack.cpp b/Engine/Source/Developer/TraceInsights/Private/Insights/ViewModels/ThreadTimingTrack.cpp
--- a/Engine/Source/Developer/TraceInsights/Private/Insights/ViewModels/ThreadTimingTrack.cpp
+++ b/Engine/Source/Developer/TraceInsights/Private/Insights/ViewModels/ThreadTimingTrack.cpp
@@ -500,20 +500,25 @@ void FThreadTimingSharedState::Tick(Insights::ITimingViewSession& InSession, con
 			{
 				// Check if this thread is part of a group?
 				bool bIsGroupVisible = bShowHideAllCpuTracks;
-				const TCHAR* const GroupName = ThreadInfo.GroupName ? ThreadInfo.GroupName : ThreadInfo.Name;
-				if (GroupName != nullptr)
+				const TCHAR* GroupName = ThreadInfo.GroupName;
+				if (!GroupName || *GroupName == 0)
 				{
-					if (!ThreadGroups.Contains(GroupName))
-					{
-						//UE_LOG(TimingProfiler, Log, TEXT("New CPU Thread Group (%d) : \"%s\""), ThreadGroups.Num() + 1, GroupName);
-						ThreadGroups.Add(GroupName, { GroupName, bIsGroupVisible, 0, Order });
-					}
-					else
-					{
-						FThreadGroup& ThreadGroup = ThreadGroups[GroupName];
-						bIsGroupVisible = ThreadGroup.bIsVisible;
-						ThreadGroup.Order = Order;
-					}
+					GroupName = ThreadInfo.Name;
+				}
+				if (!GroupName || *GroupName == 0)
+				{
+					GroupName = TEXT("Other Threads");
+				}
+				if (!ThreadGroups.Contains(GroupName))
+				{
+					// Note: The GroupName pointer should be valid for the duration of the session.
+					ThreadGroups.Add(GroupName, { GroupName, bIsGroupVisible, 0, Order });
+				}
+				else
+				{
+					FThreadGroup& ThreadGroup = ThreadGroups[GroupName];
+					bIsGroupVisible = ThreadGroup.bIsVisible;
+					ThreadGroup.Order = Order;
 				}
 
 				// Check if there is an available Asset Loading track for this thread.
diff --git a/Engine/Source/Developer/TraceServices/Private/Model/Threads.cpp b/Engine/Source/Developer/TraceServices/Private/Model/Threads.cpp
--- a/Engine/Source/Developer/TraceServices/Private/Model/Threads.cpp
+++ b/Engine/Source/Developer/TraceServices/Private/Model/Threads.cpp
@@ -27,7 +27,7 @@ FThreadProvider::~FThreadProvider()
 
 void FThreadProvider::AddGameThread(uint32 Id)
 {
-	const FString Name = *FName(NAME_GameThread).GetPlainNameString();
+	const FString Name = FName(NAME_GameThread).GetPlainNameString();
 	AddThread(Id, *Name, EThreadPriority(-2));
 }
 
@@ -50,7 +50,11 @@ void FThreadProvider::AddThread(uint32 Id, const TCHAR* Name, EThreadPriority Pr
 	}
 	if (Name != nullptr)
 	{
-		ThreadInfo->Name = Session.StoreString(Name);
+		if (Name[0] != 0)
+		{
+			ThreadInfo->Name = Session.StoreString(Name);
+		}
+
 		if (!FCString::Strcmp(Name, TEXT("RHIThread"))) //-V1051
 		{
 			const TCHAR* GroupName = Session.StoreString(TEXT("Render"));
