[Insights]
  - Improved UI for Modules view:
    * added monospace font for the Address Range column;
    * added tooltip to module name (with the full file name of a module);
    * added tooltip to Statistics column header;
    * added color for Statistics column value, based on the number of symbols cached / resolved / failed;
    * fixed row selection.
  - Minor fixes re UI for Connection tab.
  - Fixed icon for the Test Automation tab.

diff --git a/Engine/Source/Developer/TraceInsights/Private/Insights/InsightsStyle.cpp b/Engine/Source/Developer/TraceInsights/Private/Insights/InsightsStyle.cpp
--- a/Engine/Source/Developer/TraceInsights/Private/Insights/InsightsStyle.cpp
+++ b/Engine/Source/Developer/TraceInsights/Private/Insights/InsightsStyle.cpp
@@ -96,6 +96,9 @@ void FInsightsStyle::FStyle::Initialize()
 	// Sync styles from the parent style that will be used as templates for styles defined here
 	SyncParentStyles();
 
+	Set("Mono.9", DEFAULT_FONT("Mono", 9));
+	Set("Mono.10", DEFAULT_FONT("Mono", 10));
+
 	SetContentRoot(FPaths::EngineContentDir() / TEXT("Slate/Starship/Insights"));
 	SetCoreContentRoot(FPaths::EngineContentDir() / TEXT("Slate"));
 
diff --git a/Engine/Source/Developer/TraceInsights/Private/Insights/Tests/InsightsTestRunner.cpp b/Engine/Source/Developer/TraceInsights/Private/Insights/Tests/InsightsTestRunner.cpp
--- a/Engine/Source/Developer/TraceInsights/Private/Insights/Tests/InsightsTestRunner.cpp
+++ b/Engine/Source/Developer/TraceInsights/Private/Insights/Tests/InsightsTestRunner.cpp
@@ -58,13 +58,13 @@ void FInsightsTestRunner::Initialize(IUnrealInsightsModule& InsightsModule)
 		auto SessionService = SessionServicesModule.GetSessionService();
 		SessionService->Start();
 
-		//Create Session Manager
+		// Create Session Manager.
 		SessionServicesModule.GetSessionManager();
 
 		IAutomationControllerModule& AutomationControllerModule = FModuleManager::LoadModuleChecked<IAutomationControllerModule>(TEXT("AutomationController"));
 		AutomationControllerModule.Init();
 
-		// Initialize the target platform manager as it is needed by Automation Window
+		// Initialize the target platform manager as it is needed by Automation Window.
 		GetTargetPlatformManager();
 		FModuleManager::Get().LoadModule("AutomationWindow");
 		FModuleManager::Get().LoadModule("AutomationWorker");
@@ -90,17 +90,17 @@ void FInsightsTestRunner::RegisterMajorTabs(IUnrealInsightsModule& InsightsModul
 {
 	if (bInitAutomationModules)
 	{
-		const FInsightsMajorTabConfig& AutomationInfoConfig = InsightsModule.FindMajorTabConfig(FInsightsManagerTabs::AutomationWindowTabId);
-		if (AutomationInfoConfig.bIsAvailable)
+		const FInsightsMajorTabConfig& AutomationConfig = InsightsModule.FindMajorTabConfig(FInsightsManagerTabs::AutomationWindowTabId);
+		if (AutomationConfig.bIsAvailable)
 		{
-			// Register tab spawner for the Session Info.
+			// Register tab spawner for the Automation window.
 			FTabSpawnerEntry& TabSpawnerEntry = FGlobalTabmanager::Get()->RegisterNomadTabSpawner(FInsightsManagerTabs::AutomationWindowTabId,
 				FOnSpawnTab::CreateRaw(this, &FInsightsTestRunner::SpawnAutomationWindowTab))
-				.SetDisplayName(AutomationInfoConfig.TabLabel.IsSet() ? AutomationInfoConfig.TabLabel.GetValue() : LOCTEXT("AutomationTab", "Automation"))
-				.SetTooltipText(AutomationInfoConfig.TabTooltip.IsSet() ? AutomationInfoConfig.TabTooltip.GetValue() : LOCTEXT("AutomationTooltipText", "Opens the automation tab."))
-				.SetIcon(AutomationInfoConfig.TabIcon.IsSet() ? AutomationInfoConfig.TabIcon.GetValue() : FSlateIcon(FInsightsStyle::GetStyleSetName(), "Icons.SessionInfo"));
+				.SetDisplayName(AutomationConfig.TabLabel.IsSet() ? AutomationConfig.TabLabel.GetValue() : LOCTEXT("AutomationTab", "Automation"))
+				.SetTooltipText(AutomationConfig.TabTooltip.IsSet() ? AutomationConfig.TabTooltip.GetValue() : LOCTEXT("AutomationTooltipText", "Opens the automation tab."))
+				.SetIcon(AutomationConfig.TabIcon.IsSet() ? AutomationConfig.TabIcon.GetValue() : FSlateIcon(FInsightsStyle::GetStyleSetName(), "Icons.TestAutomation"));
 
-			TSharedRef<FWorkspaceItem> Group = AutomationInfoConfig.WorkspaceGroup.IsValid() ? AutomationInfoConfig.WorkspaceGroup.ToSharedRef() : FInsightsManager::Get()->GetInsightsMenuBuilder()->GetWindowsGroup();
+			TSharedRef<FWorkspaceItem> Group = AutomationConfig.WorkspaceGroup.IsValid() ? AutomationConfig.WorkspaceGroup.ToSharedRef() : FInsightsManager::Get()->GetInsightsMenuBuilder()->GetWindowsGroup();
 			TabSpawnerEntry.SetGroup(Group);
 		}
 	}
diff --git a/Engine/Source/Developer/TraceInsights/Private/Insights/Widgets/SModulesView.cpp b/Engine/Source/Developer/TraceInsights/Private/Insights/Widgets/SModulesView.cpp
--- a/Engine/Source/Developer/TraceInsights/Private/Insights/Widgets/SModulesView.cpp
+++ b/Engine/Source/Developer/TraceInsights/Private/Insights/Widgets/SModulesView.cpp
@@ -25,7 +25,7 @@ namespace Insights
 
 ////////////////////////////////////////////////////////////////////////////////////////////////////
 // FModule
-// 
+//
 // ViewModel for modules
 ////////////////////////////////////////////////////////////////////////////////////////////////////
 
@@ -38,23 +38,33 @@ public:
 		check(InModule);
 	}
 
+	const TraceServices::FModule* GetModule() const
+	{
+		return Module;
+	}
+
 	FText GetAddressRangeText() const
 	{
 		TStringBuilder<64> Sb;
-		Sb.Appendf(TEXT("0x%016llX - 0x%016llX"), Module->Base, Module->Base + Module->Size);
+		Sb.Appendf(TEXT("0x%016llX-0x%016llX"), Module->Base, Module->Base + Module->Size);
 		return FText::FromStringView(Sb.ToView());
 	}
-	
+
 	FText GetModuleNameText() const
 	{
 		return FText::FromStringView(Module->Name);
 	}
-	
+
+	FText GetModuleFullNameText() const
+	{
+		return FText::FromStringView(Module->FullName);
+	}
+
 	FText GetSymbolsFileText() const
 	{
 		return Module->StatusMessage ? FText::FromStringView(Module->StatusMessage) : FText();
 	}
-	
+
 	FText GetStatusText() const
 	{
 		return FText::FromStringView(ModuleStatusToString(Module->Status.load()));
@@ -83,7 +93,7 @@ public:
 private:
 	const TraceServices::FModule* Module;
 };
-	
+
 ////////////////////////////////////////////////////////////////////////////////////////////////////
 // SModuleRow
 ////////////////////////////////////////////////////////////////////////////////////////////////////
@@ -113,8 +123,8 @@ public:
 		ChildSlot
 		[
 			SNew(SBorder)
-			//.BorderImage(FInsightsStyle::Get().GetBrush("WhiteBrush"))
-			//.BorderBackgroundColor(this, &SModuleRow::GetBackgroundColor)
+			.BorderImage(FInsightsStyle::Get().GetBrush("WhiteBrush"))
+			.BorderBackgroundColor(this, &SModuleRow::GetBackgroundColor)
 			[
 				Row
 			]
@@ -130,6 +140,8 @@ public:
 				.Padding(FMargin(4.0, 0.0))
 				[
 					SNew(STextBlock)
+					.Font(FInsightsStyle::Get().GetFontStyle("Mono.10"))
+					.ColorAndOpacity(this, &SModuleRow::GetAddressRangeColor)
 					.Text(this, &SModuleRow::GetAddressRange)
 				];
 		}
@@ -140,6 +152,7 @@ public:
 				[
 					SNew(STextBlock)
 					.Text(this, &SModuleRow::GetModuleName)
+					.ToolTipText(this, &SModuleRow::GetModuleFullName)
 				];
 		}
 		else if (ColumnName == ModulesViewColumns::SymbolsFileColumnName)
@@ -166,6 +179,7 @@ public:
 				.Padding(FMargin(4.0, 0.0))
 				[
 					SNew(STextBlock)
+					.ColorAndOpacity(this, &SModuleRow::GetStatsColor)
 					.Text(this, &SModuleRow::GetStats)
 				];
 		}
@@ -193,6 +207,18 @@ public:
 		}
 	}
 
+	FSlateColor GetAddressRangeColor() const
+	{
+		if (IsSelected())
+		{
+			return FLinearColor(0.5f, 0.5f, 0.5f, 1.0f);
+		}
+		else
+		{
+			return FLinearColor(0.2f, 0.2f, 0.2f, 1.0f);
+		}
+	}
+
 	FText GetModuleName() const
 	{
 		TSharedPtr<FModule> ModulePin = WeakModule.Pin();
@@ -206,6 +232,19 @@ public:
 		}
 	}
 
+	FText GetModuleFullName() const
+	{
+		TSharedPtr<FModule> ModulePin = WeakModule.Pin();
+		if (ModulePin.IsValid())
+		{
+			return ModulePin->GetModuleFullNameText();
+		}
+		else
+		{
+			return FText();
+		}
+	}
+
 	FText GetSymbolsFile() const
 	{
 		TSharedPtr<FModule> ModulePin = WeakModule.Pin();
@@ -245,6 +284,51 @@ public:
 		}
 	}
 
+	FSlateColor GetStatsColor() const
+	{
+		TSharedPtr<FModule> ModulePin = WeakModule.Pin();
+		if (ModulePin.IsValid())
+		{
+			const TraceServices::FModule* Module = ModulePin->GetModule();
+			check(Module != nullptr);
+
+			if (Module->Stats.Failed > 0)
+			{
+				if (Module->Stats.Cached + Module->Stats.Resolved > Module->Stats.Failed)
+				{
+					return FLinearColor(1.0f, 0.9f, 0.6f, 1.0f);
+				}
+				else
+				{
+					return FLinearColor(1.0f, 0.7f, 0.7f, 1.0f);
+				}
+			}
+			else if (Module->Stats.Cached > 0 || Module->Stats.Resolved > 0)
+			{
+				return FLinearColor(0.7f, 1.0f, 0.7f, 1.0f);
+			}
+			else if (Module->Stats.Discovered > 0)
+			{
+				return FLinearColor(1.0f, 1.0f, 1.0f, 1.0f);
+			}
+			else
+			{
+				if (IsSelected())
+				{
+					return FLinearColor(0.5f, 0.5f, 0.5f, 1.0f);
+				}
+				else
+				{
+					return FLinearColor(0.2f, 0.2f, 0.2f, 1.0f);
+				}
+			}
+		}
+		else
+		{
+			return FLinearColor(0.0f, 0.0f, 0.0f, 1.0f);
+		}
+	}
+
 	FText GetRowToolTip() const
 	{
 		TSharedPtr<FModule> ModulePin = WeakModule.Pin();
@@ -303,21 +387,21 @@ void SModulesView::Construct(const FArguments& InArgs)
 		SNew(SVerticalBox)
 
 		// Toolbar
-		+ SVerticalBox::Slot()
-		.AutoHeight()
-		[
-			SNew(SHorizontalBox)
-
-			// Stats Text
-			+ SHorizontalBox::Slot()
-			.VAlign(VAlign_Center)
-			.Padding(2.0f, 0.0f)
-			.AutoWidth()
-			[
-				SNew(STextBlock)
-				.Text(this, &SModulesView::GetStatsText)
-			]
-		]
+		//+ SVerticalBox::Slot()
+		//.AutoHeight()
+		//[
+		//	SNew(SHorizontalBox)
+		//
+		//	// Stats Text
+		//	+ SHorizontalBox::Slot()
+		//	.VAlign(VAlign_Center)
+		//	.Padding(2.0f, 0.0f)
+		//	.AutoWidth()
+		//	[
+		//		SNew(STextBlock)
+		//		.Text(this, &SModulesView::GetStatsText)
+		//	]
+		//]
 
 		+ SVerticalBox::Slot()
 		.FillHeight(1.0f)
@@ -354,7 +438,7 @@ void SModulesView::Construct(const FArguments& InArgs)
 
 							+ SHeaderRow::Column(ModulesViewColumns::AddressRangeColumnName)
 							.ManualWidth(320.0f)
-							.DefaultLabel(LOCTEXT("AddressRangeColumn", "Address range"))
+							.DefaultLabel(LOCTEXT("AddressRangeColumn", "Address Range"))
 
 							+ SHeaderRow::Column(ModulesViewColumns::ModuleNameColumnName)
 							.ManualWidth(300.0f)
@@ -363,11 +447,12 @@ void SModulesView::Construct(const FArguments& InArgs)
 							+ SHeaderRow::Column(ModulesViewColumns::StatsColumnName)
 							.ManualWidth(350.0f)
 							.DefaultLabel(LOCTEXT("StatsColumn", "Statistics"))
+							.DefaultTooltip(LOCTEXT("StatsColumnTooltip", "Discovered = The total number of symbols discovered.\nCached = The number of symbols loaded from cache (*.ucache file).\nResolved = The number of symbols successfully resolved.\nFailed = The number of symbols failed to resolve."))
 
 							+ SHeaderRow::Column(ModulesViewColumns::StatusColumnName)
 							.ManualWidth(100.0f)
 							.DefaultLabel(LOCTEXT("StatusColumn", "Status"))
-							
+
 							+ SHeaderRow::Column(ModulesViewColumns::SymbolsFileColumnName)
 							.ManualWidth(800.0f)
 							.DefaultLabel(LOCTEXT("SymbolsFileColumn", "Symbols"))
@@ -404,13 +489,13 @@ TSharedRef<ITableRow> SModulesView::OnGenerateRow(TSharedPtr<FModule> InModule,
 void SModulesView::Tick(const FGeometry& AllottedGeometry, const double InCurrentTime, const float InDeltaTime)
 {
 	using namespace TraceServices;
-	
+
 	if ((InCurrentTime - LastUpdateTime) < 1.0)
 	{
 		return;
 	}
 	LastUpdateTime = InCurrentTime;
-	
+
 	const TSharedPtr<const IAnalysisSession> Session = FInsightsManager::Get()->GetSession();
 	if (const IModuleProvider* ModuleProvider = Session->ReadProvider<IModuleProvider>("ModuleProvider"))
 	{
@@ -537,7 +622,7 @@ bool SModulesView::LoadSymbols_CanExecute() const
 {
 	return !LoadSymbolsTask.IsValid() || LoadSymbolsTask->IsComplete();
 }
-	
+
 ////////////////////////////////////////////////////////////////////////////////////////////////////
 
 void SModulesView::LoadSymbols_Execute(TSharedPtr<FModule> Module, bool bOpenFile)
diff --git a/Engine/Source/Developer/TraceInsights/Private/Insights/Widgets/SModulesView.h b/Engine/Source/Developer/TraceInsights/Private/Insights/Widgets/SModulesView.h
--- a/Engine/Source/Developer/TraceInsights/Private/Insights/Widgets/SModulesView.h
+++ b/Engine/Source/Developer/TraceInsights/Private/Insights/Widgets/SModulesView.h
@@ -30,7 +30,7 @@ namespace ModulesViewColumns
 ////////////////////////////////////////////////////////////////////////////////////////////////////
 
 /**
- * Trace log window.
+ * Modules window.
  */
 class SModulesView : public SCompoundWidget
 {
@@ -88,7 +88,7 @@ protected:
 	/** External scrollbar used to synchronize tree view position. */
 	TSharedPtr<SScrollBar> ExternalScrollbar;
 
-	/** List of callstack frames (view model) to show in list view. */
+	/** List of modules (view model) to show in list view. */
 	TArray<TSharedPtr<FModule>> Modules;
 
 	/** Reference to ongoing load symbols task. */
@@ -100,4 +100,4 @@ protected:
 
 ////////////////////////////////////////////////////////////////////////////////////////////////////
 
-} //namespace Insights
+} // namespace Insights
diff --git a/Engine/Source/Developer/TraceInsights/Private/Insights/Widgets/SStartPageWindow.cpp b/Engine/Source/Developer/TraceInsights/Private/Insights/Widgets/SStartPageWindow.cpp
--- a/Engine/Source/Developer/TraceInsights/Private/Insights/Widgets/SStartPageWindow.cpp
+++ b/Engine/Source/Developer/TraceInsights/Private/Insights/Widgets/SStartPageWindow.cpp
@@ -1866,20 +1866,20 @@ TSharedRef<SWidget> SConnectionWindow::ConstructConnectPanel()
 		+ SVerticalBox::Slot()
 		.AutoHeight()
 		.HAlign(HAlign_Fill)
-		.Padding(12.0f, 12.0f, 12.0f, 0.0f)
+		.Padding(12.0f, 8.0f, 12.0f, 0.0f)
 		[
 			SNew(SHorizontalBox)
 
 			+ SHorizontalBox::Slot()
 			.AutoWidth()
-			.VAlign(VAlign_Top)
+			.VAlign(VAlign_Center)
 			[
 				SNew(SBox)
 				.MinDesiredWidth(180.0f)
 				.HAlign(HAlign_Right)
 				[
 					SNew(STextBlock)
-					.Text(LOCTEXT("TraceRecorderChannelsText", "Initial channels"))
+					.Text(LOCTEXT("InitialChannelsText", "Initial channels"))
 				]
 			]
 
@@ -1888,29 +1888,23 @@ TSharedRef<SWidget> SConnectionWindow::ConstructConnectPanel()
 			.VAlign(VAlign_Center)
 			.Padding(6.0f, 0.0f, 0.0f, 0.0f)
 			[
-				SNew(SVerticalBox)
-
-				+SVerticalBox::Slot()
-				.AutoHeight()
-				.HAlign(HAlign_Fill)
-				[
-					 SNew(SBox)
-					 .MinDesiredWidth(120.0f)
-					 [
-						 SAssignNew(ChannelsTextBox, SEditableTextBox)
-					 ]
-				]
-
-				+SVerticalBox::Slot()
-				.AutoHeight()
-				.HAlign(HAlign_Fill)
+				SNew(SBox)
+				.MinDesiredWidth(120.0f)
 				[
-					SNew(STextBlock)
-					.Text(LOCTEXT("TraceRecorderChannelNoteText", "Comma-separated list of channel names (or \"default\"= cpu,gpu,frame,log,bookmark) to enable when connected."))
+					SAssignNew(ChannelsTextBox, SEditableTextBox)
 				]
 			]
 		]
-	
+
+		+ SVerticalBox::Slot()
+		.AutoHeight()
+		.HAlign(HAlign_Fill)
+		.Padding(198.0f, 4.0f, 12.0f, 0.0f)
+		[
+			SNew(STextBlock)
+			.Text(LOCTEXT("InitialChannelsNoteText", "Comma-separated list of channel names (or \"default\"=cpu,gpu,frame,log,bookmark) to enable when connected."))
+		]
+
 		+ SVerticalBox::Slot()
 		.AutoHeight()
 		.HAlign(HAlign_Fill)
@@ -1991,7 +1985,7 @@ FReply SConnectionWindow::Connect_OnClicked()
 		Prerequisites.Add(ConnectTask);
 		PrerequisitesPtr = &Prerequisites;
 	}
-	
+
 	const FString ChannelsExpandedStr = ChannelsTextBox->GetText().ToString().Replace(TEXT("default"), TEXT("cpu,gpu,frame,log,bookmark"));
 
 	FGraphEventRef PreConnectTask = FFunctionGraphTask::CreateAndDispatchWhenReady(
