[Insights] Task Graph Insights usability improvements

diff --git a/Engine/Source/Developer/TraceInsights/Private/Insights/InsightsStyle.cpp b/Engine/Source/Developer/TraceInsights/Private/Insights/InsightsStyle.cpp
--- a/Engine/Source/Developer/TraceInsights/Private/Insights/InsightsStyle.cpp
+++ b/Engine/Source/Developer/TraceInsights/Private/Insights/InsightsStyle.cpp
@@ -225,7 +225,8 @@ void FInsightsStyle::FStyle::Initialize()
 
 	Set("Icons.GoToTask", new EDITOR_IMAGE_BRUSH("Icons/Profiler/profiler_ViewColumn_32x", Icon16x16));
 	Set("Icons.ShowTaskCriticalPath", new EDITOR_IMAGE_BRUSH("Icons/Profiler/profiler_HotPath_32x", Icon16x16));
-	Set("Icons.ShowTaskDependencies", new EDITOR_IMAGE_BRUSH("Icons/Profiler/profiler_Calls_32x", Icon16x16));
+	Set("Icons.ShowTaskTransitions", new EDITOR_IMAGE_BRUSH("Icons/Profiler/profiler_Calls_32x", Icon16x16));
+	Set("Icons.ShowTaskConnections", new EDITOR_IMAGE_BRUSH("Icons/Profiler/profiler_Calls_32x", Icon16x16));
 	Set("Icons.ShowTaskPrerequisites", new EDITOR_IMAGE_BRUSH("Icons/Profiler/profiler_Calls_32x", Icon16x16));
 	Set("Icons.ShowTaskSubsequents", new EDITOR_IMAGE_BRUSH("Icons/Profiler/profiler_Calls_32x", Icon16x16));
 	Set("Icons.ShowNestedTasks", new EDITOR_IMAGE_BRUSH("Icons/Profiler/profiler_Calls_32x", Icon16x16));
diff --git a/Engine/Source/Developer/TraceInsights/Private/Insights/TaskGraphProfiler/TaskGraphProfilerManager.cpp b/Engine/Source/Developer/TraceInsights/Private/Insights/TaskGraphProfiler/TaskGraphProfilerManager.cpp
--- a/Engine/Source/Developer/TraceInsights/Private/Insights/TaskGraphProfiler/TaskGraphProfilerManager.cpp
+++ b/Engine/Source/Developer/TraceInsights/Private/Insights/TaskGraphProfiler/TaskGraphProfilerManager.cpp
@@ -275,7 +275,7 @@ void FTaskGraphProfilerManager::ShowTaskRelations(const TraceServices::FTaskInfo
 			const TraceServices::FTaskInfo* Task = TasksProvider->TryGetTask(Relation.RelativeId);
 			if (Task != nullptr)
 			{
-				GetSingleTaskRelations(Task, TasksProvider, InSelectedEvent);
+				GetSingleTaskTransitions(Task, TasksProvider, InSelectedEvent);
 			}
 		}
 	};
@@ -284,9 +284,13 @@ void FTaskGraphProfilerManager::ShowTaskRelations(const TraceServices::FTaskInfo
 	{
 		GetSingleTaskRelationsForAll(Task->Prerequisites);
 	}
-	if (bShowRelations)
+	if (bShowTransitions)
 	{
-		GetSingleTaskRelations(Task, TasksProvider, InSelectedEvent);
+		GetSingleTaskTransitions(Task, TasksProvider, InSelectedEvent);
+	}
+	if (bShowConnections)
+	{
+		GetSingleTaskConnections(Task, TasksProvider, InSelectedEvent);
 	}
 	if (bShowNestedTasks)
 	{
@@ -302,10 +306,8 @@ void FTaskGraphProfilerManager::ShowTaskRelations(const TraceServices::FTaskInfo
 	}
 }
 
-void FTaskGraphProfilerManager::GetSingleTaskRelations(const TraceServices::FTaskInfo* Task, const TraceServices::ITasksProvider* TasksProvider, const FThreadTrackEvent* InSelectedEvent)
+void FTaskGraphProfilerManager::GetSingleTaskTransitions(const TraceServices::FTaskInfo* Task, const TraceServices::ITasksProvider* TasksProvider, const FThreadTrackEvent* InSelectedEvent)
 {
-	const int32 MaxTasksToShow = 30;
-
 	if (Task->CreatedTimestamp != Task->LaunchedTimestamp || Task->CreatedThreadId != Task->LaunchedThreadId)
 	{
 		AddRelation(InSelectedEvent, Task->CreatedTimestamp, Task->CreatedThreadId, Task->LaunchedTimestamp, Task->LaunchedThreadId, ETaskEventType::Created);
@@ -315,17 +317,37 @@ void FTaskGraphProfilerManager::GetSingleTaskRelations(const TraceServices::FTas
 	{
 		AddRelation(InSelectedEvent, Task->LaunchedTimestamp, Task->LaunchedThreadId, Task->ScheduledTimestamp, Task->ScheduledThreadId, ETaskEventType::Launched);
 	}
+
+	int32 ExecutionStartedDepth = GetDepthOfTaskExecution(Task->StartedTimestamp, Task->FinishedTimestamp, Task->StartedThreadId);
+	AddRelation(InSelectedEvent, Task->ScheduledTimestamp, Task->ScheduledThreadId, -1, Task->StartedTimestamp, Task->StartedThreadId, ExecutionStartedDepth, ETaskEventType::Scheduled);
+
+	if (Task->FinishedTimestamp != Task->CompletedTimestamp || Task->CompletedThreadId != Task->StartedThreadId)
+	{
+		AddRelation(InSelectedEvent, Task->FinishedTimestamp, Task->StartedThreadId, ExecutionStartedDepth, Task->CompletedTimestamp, Task->CompletedThreadId, -1,  ETaskEventType::Completed);
+	}
+}
+
+////////////////////////////////////////////////////////////////////////////////////////////////////
+
+void FTaskGraphProfilerManager::GetSingleTaskConnections(const TraceServices::FTaskInfo* Task, const TraceServices::ITasksProvider* TasksProvider, const FThreadTrackEvent* InSelectedEvent)
+{
+	const int32 MaxTasksToShow = 30;
+
 	int32 NumPrerequisitesToShow = FMath::Min(Task->Prerequisites.Num(), MaxTasksToShow);
 	for (int32 i = 0; i != NumPrerequisitesToShow; ++i)
 	{
 		const TraceServices::FTaskInfo* Prerequisite = TasksProvider->TryGetTask(Task->Prerequisites[i].RelativeId);
 		check(Prerequisite != nullptr);
-		AddRelation(InSelectedEvent, Prerequisite->CompletedTimestamp, Prerequisite->CompletedThreadId, Task->ScheduledTimestamp, Task->ScheduledThreadId, ETaskEventType::Prerequisite);
+		double StartTime = Prerequisite->CompletedTimestamp;
+		uint32 ThreadId = Prerequisite->CompletedThreadId;
+		if (StartTime <= 0.0)
+		{
+			StartTime = Prerequisite->FinishedTimestamp;
+			ThreadId = Prerequisite->StartedThreadId;
+		}
+		AddRelation(InSelectedEvent, StartTime, ThreadId, Task->StartedTimestamp, Task->StartedThreadId, ETaskEventType::Prerequisite);
 	}
 
-	int32 ExecutionStartedDepth = GetDepthOfTaskExecution(Task->StartedTimestamp, Task->FinishedTimestamp, Task->StartedThreadId);
-	AddRelation(InSelectedEvent, Task->ScheduledTimestamp, Task->ScheduledThreadId, -1, Task->StartedTimestamp, Task->StartedThreadId, ExecutionStartedDepth, ETaskEventType::Scheduled);
-
 	int32 NumNestedToShow = FMath::Min(Task->NestedTasks.Num(), MaxTasksToShow);
 	for (int32 i = 0; i != NumNestedToShow; ++i)
 	{
@@ -335,9 +357,7 @@ void FTaskGraphProfilerManager::GetSingleTaskRelations(const TraceServices::FTas
 
 		int32 NestedExecutionStartedDepth = GetDepthOfTaskExecution(Task->StartedTimestamp, Task->FinishedTimestamp, Task->StartedThreadId);
 
-		AddRelation(InSelectedEvent, RelationInfo.Timestamp, Task->StartedThreadId, -1, NestedTask->StartedTimestamp,  NestedTask->StartedThreadId, NestedExecutionStartedDepth,  ETaskEventType::AddedNested);
-
-		AddRelation(InSelectedEvent, NestedTask->CompletedTimestamp, NestedTask->CompletedThreadId, NestedTask->CompletedTimestamp, Task->StartedThreadId, ETaskEventType::NestedCompleted);
+		AddRelation(InSelectedEvent, RelationInfo.Timestamp, Task->StartedThreadId, -1, NestedTask->StartedTimestamp, NestedTask->StartedThreadId, NestedExecutionStartedDepth, ETaskEventType::AddedNested);
 	}
 
 	int32 NumSubsequentsToShow = FMath::Min(Task->Subsequents.Num(), MaxTasksToShow);
@@ -345,15 +365,14 @@ void FTaskGraphProfilerManager::GetSingleTaskRelations(const TraceServices::FTas
 	{
 		const TraceServices::FTaskInfo* Subsequent = TasksProvider->TryGetTask(Task->Subsequents[i].RelativeId);
 		check(Subsequent != nullptr);
-		if (Task->CompletedTimestamp < Subsequent->ScheduledTimestamp)
+		double StartTime = Task->CompletedTimestamp;
+		uint32 ThreadId = Task->CompletedThreadId;
+		if (StartTime <= 0.0)
 		{
-			AddRelation(InSelectedEvent, Task->CompletedTimestamp, Task->CompletedThreadId, Subsequent->ScheduledTimestamp, Subsequent->ScheduledThreadId, ETaskEventType::Subsequent);
+			StartTime = Task->FinishedTimestamp;
+			ThreadId = Task->StartedThreadId;
 		}
-	}
-
-	if (Task->FinishedTimestamp != Task->CompletedTimestamp || Task->CompletedThreadId != Task->StartedThreadId)
-	{
-		AddRelation(InSelectedEvent, Task->FinishedTimestamp, Task->StartedThreadId, ExecutionStartedDepth, Task->CompletedTimestamp, Task->StartedThreadId, -1,  ETaskEventType::Completed);
+		AddRelation(InSelectedEvent, StartTime, ThreadId, Subsequent->StartedTimestamp, Subsequent->StartedThreadId, ETaskEventType::Subsequent);
 	}
 }
 
@@ -476,16 +495,25 @@ void FTaskGraphProfilerManager::OnWindowClosedEvent()
 
 void FTaskGraphProfilerManager::InitializeColorCode()
 {
-	ColorCode[static_cast<uint32>(ETaskEventType::Created)] = FLinearColor::Yellow;
-	ColorCode[static_cast<uint32>(ETaskEventType::Launched)] = FLinearColor::Green;
-	ColorCode[static_cast<uint32>(ETaskEventType::Prerequisite)] = FLinearColor::Red;
-	ColorCode[static_cast<uint32>(ETaskEventType::Scheduled)] = FLinearColor::Blue;
-	ColorCode[static_cast<uint32>(ETaskEventType::Started)] = FLinearColor::Red;
-	ColorCode[static_cast<uint32>(ETaskEventType::AddedNested)] = FLinearColor::Blue;
-	ColorCode[static_cast<uint32>(ETaskEventType::NestedCompleted)] = FLinearColor::Red;
-	ColorCode[static_cast<uint32>(ETaskEventType::Subsequent)] = FLinearColor::Red;
-	ColorCode[static_cast<uint32>(ETaskEventType::Completed)] = FLinearColor::Yellow;
-	ColorCode[static_cast<uint32>(ETaskEventType::CriticalPath)] = FLinearColor(FColor::Purple);
+	auto ToLiniarColorNoAlpha = [](uint32 Value)
+	{
+		float R = (Value & 0xFF000000) >> 24;
+		float G = (Value & 0x00FF0000) >> 16;
+		float B = (Value & 0x0000FF00) >> 8;
+
+		return FLinearColor(R / 255, G / 255, B / 255);
+	};
+
+	ColorCode[static_cast<uint32>(ETaskEventType::Created)] = ToLiniarColorNoAlpha(0xFFDC1AFF); // Yellow;
+	ColorCode[static_cast<uint32>(ETaskEventType::Launched)] = ToLiniarColorNoAlpha(0x8BC24AFF); // Green;
+	ColorCode[static_cast<uint32>(ETaskEventType::Prerequisite)] = ToLiniarColorNoAlpha(0xFF729CFF); //Pink;
+	ColorCode[static_cast<uint32>(ETaskEventType::Scheduled)] = ToLiniarColorNoAlpha(0x26BBFFFF); // Blue;
+	ColorCode[static_cast<uint32>(ETaskEventType::Started)] = ToLiniarColorNoAlpha(0xFF0000FF); // Red;
+	ColorCode[static_cast<uint32>(ETaskEventType::AddedNested)] = ToLiniarColorNoAlpha(0xB68F55FF); // Folder
+	ColorCode[static_cast<uint32>(ETaskEventType::NestedCompleted)] = ToLiniarColorNoAlpha(0x804D39FF); // Brown
+	ColorCode[static_cast<uint32>(ETaskEventType::Subsequent)] = ToLiniarColorNoAlpha(0xFE9B07FF); // Orange;;
+	ColorCode[static_cast<uint32>(ETaskEventType::Completed)] = ToLiniarColorNoAlpha(0xFFDC1AFF); // Yellow
+	ColorCode[static_cast<uint32>(ETaskEventType::CriticalPath)] = ToLiniarColorNoAlpha(0xA139BFFF); // Purple
 }
 
 ////////////////////////////////////////////////////////////////////////////////////////////////////
@@ -498,6 +526,13 @@ FLinearColor FTaskGraphProfilerManager::GetColorForTaskEvent(ETaskEventType InEv
 
 ////////////////////////////////////////////////////////////////////////////////////////////////////
 
+uint32 FTaskGraphProfilerManager::GetColorForTaskEventAsPackedARGB(ETaskEventType InEvent)
+{
+	return GetColorForTaskEvent(InEvent).ToFColor(false).ToPackedARGB();
+}
+
+////////////////////////////////////////////////////////////////////////////////////////////////////
+
 void FTaskGraphProfilerManager::AddRelation(const FThreadTrackEvent* InSelectedEvent, double SourceTimestamp, uint32 SourceThreadId, double TargetTimestamp, uint32 TargetThreadId, ETaskEventType Type)
 {
 	AddRelation(InSelectedEvent, SourceTimestamp, SourceThreadId, -1, TargetTimestamp, TargetThreadId, -1, Type);
@@ -596,10 +631,10 @@ void FTaskGraphProfilerManager::ClearTaskRelations()
 
 ////////////////////////////////////////////////////////////////////////////////////////////////////
 
-void FTaskGraphProfilerManager::SetShowRelations(bool bInValue)
+void FTaskGraphProfilerManager::SetShowTransitions(bool bInValue)
 {
-	bShowRelations = bInValue;
-	if (!bShowRelations)
+	bShowTransitions = bInValue;
+	if (!bShowTransitions)
 	{
 		TaskTimingSharedState->SetTaskId(TaskTrace::InvalidId);
 	}
diff --git a/Engine/Source/Developer/TraceInsights/Private/Insights/TaskGraphProfiler/TaskGraphProfilerManager.h b/Engine/Source/Developer/TraceInsights/Private/Insights/TaskGraphProfiler/TaskGraphProfilerManager.h
--- a/Engine/Source/Developer/TraceInsights/Private/Insights/TaskGraphProfiler/TaskGraphProfilerManager.h
+++ b/Engine/Source/Developer/TraceInsights/Private/Insights/TaskGraphProfiler/TaskGraphProfilerManager.h
@@ -102,11 +102,15 @@ public:
 	void AddRelation(const FThreadTrackEvent* InSelectedEvent, double SourceTimestamp, uint32 SourceThreadId, int32 SourceDepth, double TargetTimestamp, uint32 TargetThreadId, int32 TargetDepth, ETaskEventType Type);
 	void ClearTaskRelations();
 	FLinearColor GetColorForTaskEvent(ETaskEventType InEvent);
+	uint32 GetColorForTaskEventAsPackedARGB(ETaskEventType InEvent);
 
 	TSharedPtr<Insights::FTaskTimingSharedState> GetTaskTimingSharedState() { return TaskTimingSharedState;	}
 
-	bool GetShowRelations() const { return bShowRelations; }
-	void SetShowRelations(bool bInValue);
+	bool GetShowTransitions() const { return bShowTransitions; }
+	void SetShowTransitions(bool bInValue);
+
+	bool GetShowConnections() const { return bShowConnections; }
+	void SetShowConnections(bool bInValue) { bShowConnections = bInValue; }
 
 	bool GetShowPrerequisites() const { return bShowPrerequisites; }
 	void SetShowPrerequisites(bool bInValue) { bShowPrerequisites = bInValue; }
@@ -120,7 +124,7 @@ public:
 	bool GetShowCriticalPath() const { return bShowCriticalPath; }
 	void SetShowCriticalPath(bool bInValue) { bShowCriticalPath = bInValue; }
 
-	bool GetShowAnyRelations() const { return GetShowRelations() || GetShowPrerequisites() || GetShowSubsequents() || GetShowNestedTasks() || GetShowCriticalPath(); }
+	bool GetShowAnyRelations() const { return GetShowTransitions() || GetShowConnections() || GetShowPrerequisites() || GetShowSubsequents() || GetShowNestedTasks() || GetShowCriticalPath(); }
 
 	int32 GetDepthOfTaskExecution(double TaskStartedTime, double TaskFinishedTime, uint32 ThreadId);
 
@@ -131,7 +135,8 @@ private:
 	void RegisterTimingProfilerLayoutExtensions(FInsightsMajorTabExtender& InOutExtender);
 
 	void ShowTaskRelations(const TraceServices::FTaskInfo* Task, const TraceServices::ITasksProvider* TasksProvider, const FThreadTrackEvent* InSelectedEvent);
-	void GetSingleTaskRelations(const TraceServices::FTaskInfo* Task, const TraceServices::ITasksProvider* TasksProvider, const FThreadTrackEvent* InSelectedEvent);
+	void GetSingleTaskTransitions(const TraceServices::FTaskInfo* Task, const TraceServices::ITasksProvider* TasksProvider, const FThreadTrackEvent* InSelectedEvent);
+	void GetSingleTaskConnections(const TraceServices::FTaskInfo* Task, const TraceServices::ITasksProvider* TasksProvider, const FThreadTrackEvent* InSelectedEvent);
 	void GetRelationsOnCriticalPath(const TraceServices::FTaskInfo* Task);
 
 	double GetRelationsOnCriticalPathAscendingRec(const TraceServices::FTaskInfo* Task, const TraceServices::ITasksProvider* TasksProvider, TArray<FTaskGraphRelation>& Relations);
@@ -161,7 +166,8 @@ private:
 
 	TSharedPtr<Insights::STaskTableTreeView> TaskTableTreeView;
 	FLinearColor ColorCode[static_cast<uint32>(ETaskEventType::NumTaskEventTypes)];
-	bool bShowRelations = true;
+	bool bShowTransitions = true;
+	bool bShowConnections = false;
 	bool bShowPrerequisites = false;
 	bool bShowSubsequents = false;
 	bool bShowNestedTasks = false;
diff --git a/Engine/Source/Developer/TraceInsights/Private/Insights/TaskGraphProfiler/ViewModels/TaskTimingTrack.cpp b/Engine/Source/Developer/TraceInsights/Private/Insights/TaskGraphProfiler/ViewModels/TaskTimingTrack.cpp
--- a/Engine/Source/Developer/TraceInsights/Private/Insights/TaskGraphProfiler/ViewModels/TaskTimingTrack.cpp
+++ b/Engine/Source/Developer/TraceInsights/Private/Insights/TaskGraphProfiler/ViewModels/TaskTimingTrack.cpp
@@ -51,15 +51,45 @@ public:
 	PRAGMA_DISABLE_OPTIMIZATION
 	virtual void RegisterCommands() override
 	{
-		UI_COMMAND(Command_ShowTaskDependencies, "Show Task Dependencies", "Show/hide dependencies of the current task (for a selected cpu timing event.)", EUserInterfaceActionType::ToggleButton, FInputChord(EKeys::T));
-		UI_COMMAND(Command_ShowTaskPrerequisites, "Show Dependencies of Prerequisites", "Show/hide dependencies of the current task's prerequisites.", EUserInterfaceActionType::ToggleButton, FInputChord(EKeys::P));
-		UI_COMMAND(Command_ShowTaskSubsequents, "Show Dependencies of Subsequents", "Show/hide dependencies of the current task's subsequents.", EUserInterfaceActionType::ToggleButton, FInputChord(EKeys::S));
-		UI_COMMAND(Command_ShowNestedTasks, "Show Dependencies of Nested Tasks", "Show/hide dependencies of the current task's nested tasks.", EUserInterfaceActionType::ToggleButton, FInputChord(EKeys::N));
-		UI_COMMAND(Command_ShowCriticalPath, "Show Task Critical Path", "Show/hide relations representing the critical path containing the current task.", EUserInterfaceActionType::ToggleButton, FInputChord());
+		UI_COMMAND(Command_ShowTaskTransitions,
+				   "Show Task Transitions",
+				   "Show/hide transitions between the stages of the current task (for a selected cpu timing event.)",
+				    EUserInterfaceActionType::ToggleButton,
+				    FInputChord(EKeys::T));
+
+		UI_COMMAND(Command_ShowTaskConnections,
+				   "Show Task Connections",
+				   "Show/hide conections between:\nThe current task's prerequisites completed time and the current task's started time.\nThe current task's completed time and the current task's subsequents started time.\nThe current task's nested tasks added time and their started time.",
+				   EUserInterfaceActionType::ToggleButton, FInputChord());
+
+		UI_COMMAND(Command_ShowTaskPrerequisites,
+				   "Show Transitions of Prerequisites",
+				   "Show/hide stage transitions for the current task's prerequisites.",
+				   EUserInterfaceActionType::ToggleButton,
+				   FInputChord(EKeys::P));
+
+		UI_COMMAND(Command_ShowTaskSubsequents,
+				   "Show Transitions of Subsequents",
+				   "Show/hide stage transitions for the current task's subsequents.",
+				   EUserInterfaceActionType::ToggleButton,
+				   FInputChord(EKeys::S));
+
+		UI_COMMAND(Command_ShowNestedTasks,
+				   "Show Transitions of Nested Tasks",
+				   "Show/hide stage transitions for the current task's nested tasks.",
+				   EUserInterfaceActionType::ToggleButton,
+				   FInputChord(EKeys::N));
+
+		UI_COMMAND(Command_ShowCriticalPath,
+				   "Show Task Critical Path",
+				   "Show/hide relations representing the critical path containing the current task.",
+				   EUserInterfaceActionType::ToggleButton,
+				   FInputChord());
 	}
 	PRAGMA_ENABLE_OPTIMIZATION
 
-	TSharedPtr<FUICommandInfo> Command_ShowTaskDependencies;
+	TSharedPtr<FUICommandInfo> Command_ShowTaskTransitions;
+	TSharedPtr<FUICommandInfo> Command_ShowTaskConnections;
 	TSharedPtr<FUICommandInfo> Command_ShowTaskPrerequisites;
 	TSharedPtr<FUICommandInfo> Command_ShowTaskSubsequents;
 	TSharedPtr<FUICommandInfo> Command_ShowNestedTasks;
@@ -205,11 +235,20 @@ void FTaskTimingSharedState::BuildTasksSubMenu(FMenuBuilder& MenuBuilder)
 
 	MenuBuilder.AddMenuEntry
 	(
-		FTaskTimingStateCommands::Get().Command_ShowTaskDependencies,
+		FTaskTimingStateCommands::Get().Command_ShowTaskTransitions,
 		NAME_None,
 		TAttribute<FText>(),
 		TAttribute<FText>(),
-		FSlateIcon(FInsightsStyle::GetStyleSetName(), "Icons.ShowTaskDependencies")
+		FSlateIcon(FInsightsStyle::GetStyleSetName(), "Icons.ShowTaskTransitions")
+	);
+
+	MenuBuilder.AddMenuEntry
+	(
+		FTaskTimingStateCommands::Get().Command_ShowTaskConnections,
+		NAME_None,
+		TAttribute<FText>(),
+		TAttribute<FText>(),
+		FSlateIcon(FInsightsStyle::GetStyleSetName(), "Icons.ShowTaskConnections")
 	);
 
 	MenuBuilder.AddSeparator();
@@ -252,10 +291,16 @@ void FTaskTimingSharedState::InitCommandList()
 	FTaskTimingStateCommands::Register();
 
 	CommandList->MapAction(
-		FTaskTimingStateCommands::Get().Command_ShowTaskDependencies,
-		FExecuteAction::CreateSP(this, &FTaskTimingSharedState::ContextMenu_ShowTaskDependencies_Execute),
-		FCanExecuteAction::CreateSP(this, &FTaskTimingSharedState::ContextMenu_ShowTaskDependencies_CanExecute),
-		FIsActionChecked::CreateSP(this, &FTaskTimingSharedState::ContextMenu_ShowTaskDependencies_IsChecked));
+		FTaskTimingStateCommands::Get().Command_ShowTaskTransitions,
+		FExecuteAction::CreateSP(this, &FTaskTimingSharedState::ContextMenu_ShowTaskTransitions_Execute),
+		FCanExecuteAction::CreateSP(this, &FTaskTimingSharedState::ContextMenu_ShowTaskTransitions_CanExecute),
+		FIsActionChecked::CreateSP(this, &FTaskTimingSharedState::ContextMenu_ShowTaskTransitions_IsChecked));
+
+	CommandList->MapAction(
+		FTaskTimingStateCommands::Get().Command_ShowTaskConnections,
+		FExecuteAction::CreateSP(this, &FTaskTimingSharedState::ContextMenu_ShowTaskConnections_Execute),
+		FCanExecuteAction::CreateSP(this, &FTaskTimingSharedState::ContextMenu_ShowTaskConnections_CanExecute),
+		FIsActionChecked::CreateSP(this, &FTaskTimingSharedState::ContextMenu_ShowTaskConnections_IsChecked));
 
 	CommandList->MapAction(
 		FTaskTimingStateCommands::Get().Command_ShowTaskPrerequisites,
@@ -283,31 +328,59 @@ void FTaskTimingSharedState::InitCommandList()
 }
 
 ////////////////////////////////////////////////////////////////////////////////////////////////////
-// Dependencies
+// Transitions
+
+void FTaskTimingSharedState::ContextMenu_ShowTaskTransitions_Execute()
+{
+	TSharedPtr<Insights::FTaskGraphProfilerManager> TaskGraphManager = Insights::FTaskGraphProfilerManager::Get();
+	if (TaskGraphManager.IsValid() && TaskGraphManager->GetIsAvailable())
+	{
+		TaskGraphManager->SetShowTransitions(!TaskGraphManager->GetShowTransitions());
+		OnTaskSettingsChanged();
+	}
+}
+
+////////////////////////////////////////////////////////////////////////////////////////////////////
+
+bool FTaskTimingSharedState::ContextMenu_ShowTaskTransitions_CanExecute()
+{
+	return Insights::FTaskGraphProfilerManager::Get().IsValid() && Insights::FTaskGraphProfilerManager::Get()->GetIsAvailable();
+}
+
+////////////////////////////////////////////////////////////////////////////////////////////////////
+
+bool FTaskTimingSharedState::ContextMenu_ShowTaskTransitions_IsChecked()
+{
+	TSharedPtr<Insights::FTaskGraphProfilerManager> TaskGraphManager = Insights::FTaskGraphProfilerManager::Get();
+	return TaskGraphManager.IsValid() && TaskGraphManager->GetIsAvailable() && TaskGraphManager->GetShowTransitions();
+}
+
+///////////////////////////////////////////////////////////////////////////////////////////////////
+// Connections
 
-void FTaskTimingSharedState::ContextMenu_ShowTaskDependencies_Execute()
+void FTaskTimingSharedState::ContextMenu_ShowTaskConnections_Execute()
 {
 	TSharedPtr<Insights::FTaskGraphProfilerManager> TaskGraphManager = Insights::FTaskGraphProfilerManager::Get();
 	if (TaskGraphManager.IsValid() && TaskGraphManager->GetIsAvailable())
 	{
-		TaskGraphManager->SetShowRelations(!TaskGraphManager->GetShowRelations());
+		TaskGraphManager->SetShowConnections(!TaskGraphManager->GetShowConnections());
 		OnTaskSettingsChanged();
 	}
 }
 
 ////////////////////////////////////////////////////////////////////////////////////////////////////
 
-bool FTaskTimingSharedState::ContextMenu_ShowTaskDependencies_CanExecute()
+bool FTaskTimingSharedState::ContextMenu_ShowTaskConnections_CanExecute()
 {
 	return Insights::FTaskGraphProfilerManager::Get().IsValid() && Insights::FTaskGraphProfilerManager::Get()->GetIsAvailable();
 }
 
 ////////////////////////////////////////////////////////////////////////////////////////////////////
 
-bool FTaskTimingSharedState::ContextMenu_ShowTaskDependencies_IsChecked()
+bool FTaskTimingSharedState::ContextMenu_ShowTaskConnections_IsChecked()
 {
 	TSharedPtr<Insights::FTaskGraphProfilerManager> TaskGraphManager = Insights::FTaskGraphProfilerManager::Get();
-	return TaskGraphManager.IsValid() && TaskGraphManager->GetIsAvailable() && TaskGraphManager->GetShowRelations();
+	return TaskGraphManager.IsValid() && TaskGraphManager->GetIsAvailable() && TaskGraphManager->GetShowConnections();
 }
 
 ////////////////////////////////////////////////////////////////////////////////////////////////////
@@ -488,13 +561,13 @@ void FTaskTimingTrack::BuildDrawState(ITimingEventsTrackDrawStateBuilder& Builde
 		return;
 	}
 
-	Builder.AddEvent(Task->CreatedTimestamp, Task->LaunchedTimestamp, 0, TEXT("Launched"), 0, FTaskGraphProfilerManager::Get()->GetColorForTaskEvent(ETaskEventType::Created).ToFColor(true).ToPackedARGB());
-	Builder.AddEvent(Task->LaunchedTimestamp, Task->ScheduledTimestamp, 0, TEXT("Dispatched"), 0, FTaskGraphProfilerManager::Get()->GetColorForTaskEvent(ETaskEventType::Launched).ToFColor(true).ToPackedARGB());
-	Builder.AddEvent(Task->ScheduledTimestamp, Task->StartedTimestamp, 0, TEXT("Scheduled"), 0,  FTaskGraphProfilerManager::Get()->GetColorForTaskEvent(ETaskEventType::Scheduled).ToFColor(true).ToPackedARGB());
-	Builder.AddEvent(Task->StartedTimestamp, Task->FinishedTimestamp, 0, TEXT("Executed"), 0, FTaskGraphProfilerManager::Get()->GetColorForTaskEvent(ETaskEventType::NestedCompleted).ToFColor(true).ToPackedARGB());
+	Builder.AddEvent(Task->CreatedTimestamp, Task->LaunchedTimestamp, 0, TEXT("Launched"), 0, FTaskGraphProfilerManager::Get()->GetColorForTaskEventAsPackedARGB(ETaskEventType::Created));
+	Builder.AddEvent(Task->LaunchedTimestamp, Task->ScheduledTimestamp, 0, TEXT("Dispatched"), 0, FTaskGraphProfilerManager::Get()->GetColorForTaskEventAsPackedARGB(ETaskEventType::Launched));
+	Builder.AddEvent(Task->ScheduledTimestamp, Task->StartedTimestamp, 0, TEXT("Scheduled"), 0,  FTaskGraphProfilerManager::Get()->GetColorForTaskEventAsPackedARGB(ETaskEventType::Scheduled));
+	Builder.AddEvent(Task->StartedTimestamp, Task->FinishedTimestamp, 0, TEXT("Executed"), 0, FTaskGraphProfilerManager::Get()->GetColorForTaskEventAsPackedARGB(ETaskEventType::Started));
 	if (Task->CompletedTimestamp > Task->FinishedTimestamp)
 	{
-		Builder.AddEvent(Task->FinishedTimestamp, Task->CompletedTimestamp, 0, TEXT("Completed"), 0, FTaskGraphProfilerManager::Get()->GetColorForTaskEvent(ETaskEventType::Completed).ToFColor(true).ToPackedARGB());
+		Builder.AddEvent(Task->FinishedTimestamp, Task->CompletedTimestamp, 0, TEXT("Completed"), 0, FTaskGraphProfilerManager::Get()->GetColorForTaskEventAsPackedARGB(ETaskEventType::Completed));
 	}
 }
 
diff --git a/Engine/Source/Developer/TraceInsights/Private/Insights/TaskGraphProfiler/ViewModels/TaskTimingTrack.h b/Engine/Source/Developer/TraceInsights/Private/Insights/TaskGraphProfiler/ViewModels/TaskTimingTrack.h
--- a/Engine/Source/Developer/TraceInsights/Private/Insights/TaskGraphProfiler/ViewModels/TaskTimingTrack.h
+++ b/Engine/Source/Developer/TraceInsights/Private/Insights/TaskGraphProfiler/ViewModels/TaskTimingTrack.h
@@ -58,9 +58,13 @@ private:
 
 	void BuildTasksSubMenu(FMenuBuilder& MenuBuilder);
 
-	void ContextMenu_ShowTaskDependencies_Execute();
-	bool ContextMenu_ShowTaskDependencies_CanExecute();
-	bool ContextMenu_ShowTaskDependencies_IsChecked();
+	void ContextMenu_ShowTaskTransitions_Execute();
+	bool ContextMenu_ShowTaskTransitions_CanExecute();
+	bool ContextMenu_ShowTaskTransitions_IsChecked();
+
+	void ContextMenu_ShowTaskConnections_Execute();
+	bool ContextMenu_ShowTaskConnections_CanExecute();
+	bool ContextMenu_ShowTaskConnections_IsChecked();
 
 	void ContextMenu_ShowTaskPrerequisites_Execute();
 	bool ContextMenu_ShowTaskPrerequisites_CanExecute();
diff --git a/Engine/Source/Developer/TraceInsights/Private/Insights/TaskGraphProfiler/Widgets/STaskTableTreeView.cpp b/Engine/Source/Developer/TraceInsights/Private/Insights/TaskGraphProfiler/Widgets/STaskTableTreeView.cpp
--- a/Engine/Source/Developer/TraceInsights/Private/Insights/TaskGraphProfiler/Widgets/STaskTableTreeView.cpp
+++ b/Engine/Source/Developer/TraceInsights/Private/Insights/TaskGraphProfiler/Widgets/STaskTableTreeView.cpp
@@ -445,7 +445,7 @@ void STaskTableTreeView::ContextMenu_GoToTask_Execute()
 
 	TSharedPtr<FTaskTimingSharedState> TaskSharedState = FTaskGraphProfilerManager::Get()->GetTaskTimingSharedState();
 
-	if (TaskSharedState.IsValid() && FTaskGraphProfilerManager::Get()->GetShowRelations())
+	if (TaskSharedState.IsValid() && FTaskGraphProfilerManager::Get()->GetShowAnyRelations())
 	{
 		TaskSharedState->SetTaskId(SelectedTask->GetTask()->GetId());
 	}
