[Insights] Fixed bug when deselecting "Show Task Transitions" hid all relations.
Moved task menu options to the "Other" timing view menu.
Task connections are now visible by default.

diff --git a/Engine/Source/Developer/TraceInsights/Private/Insights/TaskGraphProfiler/TaskGraphProfilerManager.cpp b/Engine/Source/Developer/TraceInsights/Private/Insights/TaskGraphProfiler/TaskGraphProfilerManager.cpp
--- a/Engine/Source/Developer/TraceInsights/Private/Insights/TaskGraphProfiler/TaskGraphProfilerManager.cpp
+++ b/Engine/Source/Developer/TraceInsights/Private/Insights/TaskGraphProfiler/TaskGraphProfilerManager.cpp
@@ -631,17 +631,6 @@ void FTaskGraphProfilerManager::ClearTaskRelations()
 
 ////////////////////////////////////////////////////////////////////////////////////////////////////
 
-void FTaskGraphProfilerManager::SetShowTransitions(bool bInValue)
-{
-	bShowTransitions = bInValue;
-	if (!bShowTransitions)
-	{
-		TaskTimingSharedState->SetTaskId(TaskTrace::InvalidId);
-	}
-}
-
-////////////////////////////////////////////////////////////////////////////////////////////////////
-
 int32 FTaskGraphProfilerManager::GetDepthOfTaskExecution(double TaskStartedTime, double TaskFinishedTime, uint32 ThreadId)
 {
 	int32 Depth = -1;
diff --git a/Engine/Source/Developer/TraceInsights/Private/Insights/TaskGraphProfiler/TaskGraphProfilerManager.h b/Engine/Source/Developer/TraceInsights/Private/Insights/TaskGraphProfiler/TaskGraphProfilerManager.h
--- a/Engine/Source/Developer/TraceInsights/Private/Insights/TaskGraphProfiler/TaskGraphProfilerManager.h
+++ b/Engine/Source/Developer/TraceInsights/Private/Insights/TaskGraphProfiler/TaskGraphProfilerManager.h
@@ -107,7 +107,7 @@ public:
 	TSharedPtr<Insights::FTaskTimingSharedState> GetTaskTimingSharedState() { return TaskTimingSharedState;	}
 
 	bool GetShowTransitions() const { return bShowTransitions; }
-	void SetShowTransitions(bool bInValue);
+	void SetShowTransitions(bool bInValue) { bShowTransitions = bInValue; }
 
 	bool GetShowConnections() const { return bShowConnections; }
 	void SetShowConnections(bool bInValue) { bShowConnections = bInValue; }
@@ -167,7 +167,7 @@ private:
 	TSharedPtr<Insights::STaskTableTreeView> TaskTableTreeView;
 	FLinearColor ColorCode[static_cast<uint32>(ETaskEventType::NumTaskEventTypes)];
 	bool bShowTransitions = true;
-	bool bShowConnections = false;
+	bool bShowConnections = true;
 	bool bShowPrerequisites = false;
 	bool bShowSubsequents = false;
 	bool bShowNestedTasks = false;
diff --git a/Engine/Source/Developer/TraceInsights/Private/Insights/TaskGraphProfiler/ViewModels/TaskTimingTrack.cpp b/Engine/Source/Developer/TraceInsights/Private/Insights/TaskGraphProfiler/ViewModels/TaskTimingTrack.cpp
--- a/Engine/Source/Developer/TraceInsights/Private/Insights/TaskGraphProfiler/ViewModels/TaskTimingTrack.cpp
+++ b/Engine/Source/Developer/TraceInsights/Private/Insights/TaskGraphProfiler/ViewModels/TaskTimingTrack.cpp
@@ -203,13 +203,15 @@ void FTaskTimingSharedState::SetTaskId(TaskTrace::FId InTaskId)
 
 ////////////////////////////////////////////////////////////////////////////////////////////////////
 
-bool FTaskTimingSharedState::ExtendGlobalContextMenu(ITimingViewSession& InSession, FMenuBuilder& InMenuBuilder)
+void FTaskTimingSharedState::ExtendOtherTracksFilterMenu(ITimingViewSession& InSession, FMenuBuilder& InMenuBuilder)
 {
 	if (&InSession != TimingView)
 	{
-		return false;
+		return;
 	}
 
+	InMenuBuilder.BeginSection("TaskGraphInsightsOptions", LOCTEXT("OthersMenu_Section_Tasks", "Tasks"));
+
 	InMenuBuilder.AddSubMenu
 	(
 		LOCTEXT("ContextMenu_Tasks_SubMenu", "Tasks"),
@@ -217,7 +219,7 @@ bool FTaskTimingSharedState::ExtendGlobalContextMenu(ITimingViewSession& InSessi
 		FNewMenuDelegate::CreateSP(this, &FTaskTimingSharedState::BuildTasksSubMenu)
 	);
 
-	return true;
+	InMenuBuilder.EndSection();
 }
 
 ////////////////////////////////////////////////////////////////////////////////////////////////////
diff --git a/Engine/Source/Developer/TraceInsights/Private/Insights/TaskGraphProfiler/ViewModels/TaskTimingTrack.h b/Engine/Source/Developer/TraceInsights/Private/Insights/TaskGraphProfiler/ViewModels/TaskTimingTrack.h
--- a/Engine/Source/Developer/TraceInsights/Private/Insights/TaskGraphProfiler/ViewModels/TaskTimingTrack.h
+++ b/Engine/Source/Developer/TraceInsights/Private/Insights/TaskGraphProfiler/ViewModels/TaskTimingTrack.h
@@ -38,7 +38,7 @@ public:
 	virtual void OnEndSession(Insights::ITimingViewSession& InSession) override;
 	virtual void Tick(Insights::ITimingViewSession& InSession, const TraceServices::IAnalysisSession& InAnalysisSession) override;
 	virtual void ExtendFilterMenu(Insights::ITimingViewSession& InSession, FMenuBuilder& InMenuBuilder) override;
-	virtual bool ExtendGlobalContextMenu(ITimingViewSession& InSession, FMenuBuilder& InMenuBuilder) override;
+	virtual void ExtendOtherTracksFilterMenu(ITimingViewSession& InSession, FMenuBuilder& InMenuBuilder) override;
 
 	//////////////////////////////////////////////////
 
