[Insights] Memory Insights / Tree Table views: Fixed grouping "by unique values" to use the tooltip text for a cell value in order to ensure unique string names for created groups.

diff --git a/Engine/Source/Developer/TraceInsights/Private/Insights/Table/ViewModels/TreeNodeGrouping.cpp b/Engine/Source/Developer/TraceInsights/Private/Insights/Table/ViewModels/TreeNodeGrouping.cpp
--- a/Engine/Source/Developer/TraceInsights/Private/Insights/Table/ViewModels/TreeNodeGrouping.cpp
+++ b/Engine/Source/Developer/TraceInsights/Private/Insights/Table/ViewModels/TreeNodeGrouping.cpp
@@ -145,6 +145,12 @@ template<> float TTreeNodeGroupingByUniqueValue<float>::GetValue(const FTableCel
 template<> double TTreeNodeGroupingByUniqueValue<double>::GetValue(const FTableCellValue& CellValue) { return CellValue.Double; }
 template<> const TCHAR* TTreeNodeGroupingByUniqueValue<const TCHAR*>::GetValue(const FTableCellValue& CellValue) { return CellValue.CString; }
 
+template<>
+FText TTreeNodeGroupingByUniqueValue<const TCHAR*>::GetValueAsText(const FTableColumn& Column, const FTableTreeNode& Node)
+{
+	return Column.GetValueAsText(Node);
+}
+
 ////////////////////////////////////////////////////////////////////////////////////////////////////
 // FTreeNodeGroupingByNameFirstLetter
 ////////////////////////////////////////////////////////////////////////////////////////////////////
diff --git a/Engine/Source/Developer/TraceInsights/Private/Insights/Table/ViewModels/TreeNodeGrouping.h b/Engine/Source/Developer/TraceInsights/Private/Insights/Table/ViewModels/TreeNodeGrouping.h
--- a/Engine/Source/Developer/TraceInsights/Private/Insights/Table/ViewModels/TreeNodeGrouping.h
+++ b/Engine/Source/Developer/TraceInsights/Private/Insights/Table/ViewModels/TreeNodeGrouping.h
@@ -122,6 +122,10 @@ public:
 
 private:
 	static Type GetValue(const FTableCellValue& CellValue);
+	static FText GetValueAsText(const FTableColumn& Column, const FTableTreeNode& Node)
+	{
+		return Column.GetValueAsTooltipText(Node);
+	}
 };
 
 template<typename Type>
@@ -147,7 +151,8 @@ void TTreeNodeGroupingByUniqueValue<Type>::GroupNodes(const TArray<FTableTreeNod
 
 		FTableTreeNodePtr GroupPtr = nullptr;
 
-		const TOptional<FTableCellValue> CellValue = GetColumn()->GetValue(*NodePtr);
+		const FTableColumn& Column = *GetColumn();
+		const TOptional<FTableCellValue> CellValue = Column.GetValue(*NodePtr);
 		if (CellValue.IsSet())
 		{
 			const Type Value = GetValue(CellValue.GetValue());
@@ -155,8 +160,13 @@ void TTreeNodeGroupingByUniqueValue<Type>::GroupNodes(const TArray<FTableTreeNod
 			FTableTreeNodePtr* GroupPtrPtr = GroupMap.Find(Value);
 			if (!GroupPtrPtr)
 			{
-				FName GroupName(GetColumn()->GetValueAsText(*NodePtr).ToString());
-				GroupPtr = MakeShared<FTableTreeNode>(GroupName, InParentTable);
+				FText ValueAsText = GetValueAsText(Column, *NodePtr);
+				FStringView GroupName(ValueAsText.ToString());
+				if (GroupName.Len() >= NAME_SIZE)
+				{
+					GroupName = FStringView(GroupName.GetData(), NAME_SIZE - 1);
+				}
+				GroupPtr = MakeShared<FTableTreeNode>(FName(GroupName, 0), InParentTable);
 				GroupPtr->SetExpansion(false);
 				ParentGroup.AddChildAndSetGroupPtr(GroupPtr);
 				GroupMap.Add(Value, GroupPtr);
