[Insights] Added more information about symbol search paths.

Refactored so code could share symbol search path text generation. Also added the information to the investigation tab.

diff --git a/Engine/Programs/UnrealInsights/Config/DefaultUnrealInsightsSettings.ini b/Engine/Programs/UnrealInsights/Config/DefaultUnrealInsightsSettings.ini
--- a/Engine/Programs/UnrealInsights/Config/DefaultUnrealInsightsSettings.ini
+++ b/Engine/Programs/UnrealInsights/Config/DefaultUnrealInsightsSettings.ini
@@ -19,3 +19,9 @@ AutoScrollFrameAlignment=game
 AutoScrollViewportOffsetPercent=0.1
 ; Minimum time between two auto-scroll updates, in [seconds].
 AutoScrollMinDelay=0.3
+
+[Insights.MemoryProfiler]
+; Add symbol search paths by adding to this array. Order is important.
+;+SymbolSearchPaths=c:\symbols
+; Symbol server format is also accepted, such as the official Microsoft symbol servers (see https://docs.microsoft.com/en-us/windows-hardware/drivers/debugger/advanced-symsrv-use for details)
+;+SymbolSearchPaths=srv*C:\Symbols;https://msdl.microsoft.com/download/symbols
\ No newline at end of file
diff --git a/Engine/Source/Developer/TraceInsights/Private/Insights/MemoryProfiler/Common/SymbolSearchPathsHelper.h b/Engine/Source/Developer/TraceInsights/Private/Insights/MemoryProfiler/Common/SymbolSearchPathsHelper.h
new file mode 100644
--- /dev/null
+++ b/Engine/Source/Developer/TraceInsights/Private/Insights/MemoryProfiler/Common/SymbolSearchPathsHelper.h
@@ -0,0 +1,29 @@
+﻿#pragma once
+#include "Internationalization/Text.h"
+#include "TraceServices/Private/Model/ModuleProvider.h"
+
+#define LOCTEXT_NAMESPACE "SymbolSearchPathsHelper"
+
+class FSymbolSearchPathsHelper
+{
+public:
+	/**
+	 * Get localized text describing symbol search paths, or if none is set, instruction on how to setup those.
+	 * Safe to call from any thread.
+	 * @param ModuleProvider Provider to ask.
+	 * @return Text listing symbol search paths.
+	 */
+	static FText GetLocalizedSymbolSearchPathsText(const TraceServices::IModuleProvider* ModuleProvider)
+	{
+		TStringBuilder<2048> PathsListStr;
+		ModuleProvider->EnumerateSymbolSearchPaths([&](FStringView Path) { PathsListStr << TEXT("\n \u2022 ") << Path;});
+		if (PathsListStr.Len())
+		{
+			return FText::Format(LOCTEXT("SymbolResolutionPaths", "Symbol search paths (in order of precedence): {0}"),
+					 FText::FromStringView(PathsListStr));
+		}
+		return LOCTEXT("SymbolResolutionNoPaths", "No symbol paths has been setup. Use configuration file, UE_INSIGHTS_SYMBOL_PATH environment variable or manually select a file or directory from the Modules view.");
+	}
+};
+
+#undef LOCTEXT_NAMESPACE
\ No newline at end of file
diff --git a/Engine/Source/Developer/TraceInsights/Private/Insights/MemoryProfiler/Widgets/SMemAllocTableTreeView.cpp b/Engine/Source/Developer/TraceInsights/Private/Insights/MemoryProfiler/Widgets/SMemAllocTableTreeView.cpp
--- a/Engine/Source/Developer/TraceInsights/Private/Insights/MemoryProfiler/Widgets/SMemAllocTableTreeView.cpp
+++ b/Engine/Source/Developer/TraceInsights/Private/Insights/MemoryProfiler/Widgets/SMemAllocTableTreeView.cpp
@@ -19,8 +19,10 @@
 
 // Insights
 #include "Insights/InsightsStyle.h"
+#include "Insights/MemoryProfiler/Common/SymbolSearchPathsHelper.h"
 #include "Insights/MemoryProfiler/MemoryProfilerManager.h"
 #include "Insights/MemoryProfiler/ViewModels/MemAllocGroupingByCallstack.h"
+#include "Insights/MemoryProfiler/ViewModels/MemAllocGroupingByHeap.h"
 #include "Insights/MemoryProfiler/ViewModels/MemAllocGroupingBySize.h"
 #include "Insights/MemoryProfiler/ViewModels/MemAllocNode.h"
 #include "Insights/MemoryProfiler/ViewModels/MemAllocTable.h"
@@ -33,8 +35,6 @@
 
 #include <limits>
 
-#include "Insights/MemoryProfiler/ViewModels/MemAllocGroupingByHeap.h"
-
 #define LOCTEXT_NAMESPACE "SMemAllocTableTreeView"
 
 using namespace TraceServices;
@@ -1026,12 +1026,12 @@ FText SMemAllocTableTreeView::GetSymbolResolutionStatus() const
 
 FText SMemAllocTableTreeView::GetSymbolResolutionTooltip() const
 {
-	auto ModuleProvider = Session->ReadProvider<IModuleProvider>(FName("ModuleProvider"));
-	TStringBuilder<2048> Sb;
-	uint32 Index = 0;
-	ModuleProvider->EnumerateSymbolSearchPaths([&Sb, &Index](FStringView Path) { Sb << ++Index << TEXT(": ") << Path << TEXT("\n");});
-	return FText::Format(LOCTEXT("SymbolResolutionEnvVarHelp", "Symbol paths:\n{0}"),
-		FText::FromStringView(Sb));
+	const TraceServices::IModuleProvider* ModuleProvider = ReadModuleProvider(*Session.Get());
+	if (ModuleProvider)
+	{
+		return FSymbolSearchPathsHelper::GetLocalizedSymbolSearchPathsText(ModuleProvider);
+	}
+	return FText();
 }
 
 ////////////////////////////////////////////////////////////////////////////////////////////////////
@@ -1052,7 +1052,7 @@ FText SMemAllocTableTreeView::GetQueryInfoTooltip() const
 
 void SMemAllocTableTreeView::InternalCreateGroupings()
 {
-	STableTreeView::InternalCreateGroupings();
+	Insights::STableTreeView::InternalCreateGroupings();
 
 	int32 Index = 1; // after the Flat ("All") grouping
 
diff --git a/Engine/Source/Developer/TraceInsights/Private/Insights/MemoryProfiler/Widgets/SMemInvestigationView.cpp b/Engine/Source/Developer/TraceInsights/Private/Insights/MemoryProfiler/Widgets/SMemInvestigationView.cpp
--- a/Engine/Source/Developer/TraceInsights/Private/Insights/MemoryProfiler/Widgets/SMemInvestigationView.cpp
+++ b/Engine/Source/Developer/TraceInsights/Private/Insights/MemoryProfiler/Widgets/SMemInvestigationView.cpp
@@ -13,6 +13,7 @@
 #include "Insights/Common/TimeUtils.h"
 #include "Insights/InsightsStyle.h"
 #include "Insights/MemoryProfiler/MemoryProfilerManager.h"
+#include "Insights/MemoryProfiler/Common/SymbolSearchPathsHelper.h"
 #include "Insights/MemoryProfiler/ViewModels/MemorySharedState.h"
 #include "Insights/MemoryProfiler/Widgets/SMemAllocTableTreeView.h"
 #include "Insights/MemoryProfiler/Widgets/SMemoryProfilerWindow.h"
@@ -180,6 +181,16 @@ TSharedRef<SWidget> SMemInvestigationView::ConstructInvestigationWidgetArea()
 			.ToolTipText(LOCTEXT("RunQueryBtnToolTipText", "Run Memory Query.\nThe resulting list of allocations will be available in a tree view."))
 			.OnClicked(this, &SMemInvestigationView::RunQuery)
 		]
+	
+		+ SVerticalBox::Slot()
+		.AutoHeight()
+		.Padding(20.0f, 0.0f)
+		.HAlign(HAlign_Left)
+		[
+			SAssignNew(SymbolPathsTextBlock, STextBlock)
+			.Text(GetSymbolPathsText())
+			.AutoWrapText(true)
+		]
 	;
 
 	TSharedPtr<SMemoryProfilerWindow> ProfilerWindow = ProfilerWindowWeakPtr.Pin();
@@ -331,7 +342,7 @@ void SMemInvestigationView::InsightsManager_OnSessionChanged()
 
 void SMemInvestigationView::Reset()
 {
-	//...
+	SymbolPathsTextBlock->SetText(GetSymbolPathsText());
 }
 
 ////////////////////////////////////////////////////////////////////////////////////////////////////
@@ -443,6 +454,19 @@ void SMemInvestigationView::QueryTarget_OnSelectionChanged(TSharedPtr<Insights::
 	}
 }
 
+////////////////////////////////////////////////////////////////////////////////////////////////////
+FText SMemInvestigationView::GetSymbolPathsText() const
+{
+	if (Session)
+	{
+		if (const TraceServices::IModuleProvider* ModuleProvider = ReadModuleProvider(*Session.Get()))
+		{
+			return FSymbolSearchPathsHelper::GetLocalizedSymbolSearchPathsText(ModuleProvider);
+		}
+	}
+	return FText();
+}
+
 ////////////////////////////////////////////////////////////////////////////////////////////////////
 
 TSharedRef<SWidget> SMemInvestigationView::QueryTarget_OnGenerateWidget(TSharedPtr<Insights::FQueryTargetWindowSpec> InTarget)
diff --git a/Engine/Source/Developer/TraceInsights/Private/Insights/MemoryProfiler/Widgets/SMemInvestigationView.h b/Engine/Source/Developer/TraceInsights/Private/Insights/MemoryProfiler/Widgets/SMemInvestigationView.h
--- a/Engine/Source/Developer/TraceInsights/Private/Insights/MemoryProfiler/Widgets/SMemInvestigationView.h
+++ b/Engine/Source/Developer/TraceInsights/Private/Insights/MemoryProfiler/Widgets/SMemInvestigationView.h
@@ -52,6 +52,7 @@ public:
 	void QueryTarget_OnSelectionChanged(TSharedPtr<Insights::FQueryTargetWindowSpec> InRule, ESelectInfo::Type SelectInfo);
 
 private:
+	FText GetSymbolPathsText() const;
 	TSharedRef<SWidget> ConstructInvestigationWidgetArea();
 	TSharedRef<SWidget> ConstructTimeMarkerWidget(uint32 TimeMarkerIndex);
 
@@ -88,6 +89,8 @@ private:
 	TSharedPtr<SComboBox<TSharedPtr<Insights::FMemoryRuleSpec>>> QueryRuleComboBox;
 
 	TSharedPtr<SComboBox<TSharedPtr<Insights::FQueryTargetWindowSpec>>> QueryTargetComboBox;
+	
+	TSharedPtr<STextBlock> SymbolPathsTextBlock;
 };
 
 ////////////////////////////////////////////////////////////////////////////////////////////////////
