Added copyright boiler plate

diff --git a/Engine/Source/Developer/TraceInsights/Private/Insights/MemoryProfiler/Common/SymbolSearchPathsHelper.h b/Engine/Source/Developer/TraceInsights/Private/Insights/MemoryProfiler/Common/SymbolSearchPathsHelper.h
--- a/Engine/Source/Developer/TraceInsights/Private/Insights/MemoryProfiler/Common/SymbolSearchPathsHelper.h
+++ b/Engine/Source/Developer/TraceInsights/Private/Insights/MemoryProfiler/Common/SymbolSearchPathsHelper.h
@@ -1,4 +1,5 @@
-﻿#pragma once
+﻿// Copyright Epic Games, Inc. All Rights Reserved.
+#pragma once
 #include "Internationalization/Text.h"
 #include "TraceServices/Private/Model/ModuleProvider.h"
 
