[Insights] Show Prerequisites/Nestes/Subsequent's executed time in the Task Overview Graph.

diff --git a/Engine/Source/Developer/TraceInsights/Private/Insights/InsightsStyle.cpp b/Engine/Source/Developer/TraceInsights/Private/Insights/InsightsStyle.cpp
--- a/Engine/Source/Developer/TraceInsights/Private/Insights/InsightsStyle.cpp
+++ b/Engine/Source/Developer/TraceInsights/Private/Insights/InsightsStyle.cpp
@@ -230,6 +230,8 @@ void FInsightsStyle::FStyle::Initialize()
 	Set("Icons.ShowTaskPrerequisites", new EDITOR_IMAGE_BRUSH("Icons/Profiler/profiler_Calls_32x", Icon16x16));
 	Set("Icons.ShowTaskSubsequents", new EDITOR_IMAGE_BRUSH("Icons/Profiler/profiler_Calls_32x", Icon16x16));
 	Set("Icons.ShowNestedTasks", new EDITOR_IMAGE_BRUSH("Icons/Profiler/profiler_Calls_32x", Icon16x16));
+	Set("Icons.ShowTaskTrack", new EDITOR_IMAGE_BRUSH("Icons/Profiler/profiler_Calls_32x", Icon16x16));
+	Set("Icons.ShowDetailedTaskTrackInfo", new EDITOR_IMAGE_BRUSH("Icons/Profiler/profiler_Calls_32x", Icon16x16));
 
 	//////////////////////////////////////////////////
 
diff --git a/Engine/Source/Developer/TraceInsights/Private/Insights/TaskGraphProfiler/TaskGraphProfilerManager.cpp b/Engine/Source/Developer/TraceInsights/Private/Insights/TaskGraphProfiler/TaskGraphProfilerManager.cpp
--- a/Engine/Source/Developer/TraceInsights/Private/Insights/TaskGraphProfiler/TaskGraphProfilerManager.cpp
+++ b/Engine/Source/Developer/TraceInsights/Private/Insights/TaskGraphProfiler/TaskGraphProfilerManager.cpp
@@ -338,14 +338,8 @@ void FTaskGraphProfilerManager::GetSingleTaskConnections(const TraceServices::FT
 	{
 		const TraceServices::FTaskInfo* Prerequisite = TasksProvider->TryGetTask(Task->Prerequisites[i].RelativeId);
 		check(Prerequisite != nullptr);
-		double StartTime = Prerequisite->CompletedTimestamp;
-		uint32 ThreadId = Prerequisite->CompletedThreadId;
-		if (StartTime <= 0.0)
-		{
-			StartTime = Prerequisite->FinishedTimestamp;
-			ThreadId = Prerequisite->StartedThreadId;
-		}
-		AddRelation(InSelectedEvent, StartTime, ThreadId, Task->StartedTimestamp, Task->StartedThreadId, ETaskEventType::Prerequisite);
+
+		AddRelation(InSelectedEvent, Prerequisite->FinishedTimestamp, Prerequisite->StartedThreadId, Task->StartedTimestamp, Task->StartedThreadId, ETaskEventType::Prerequisite);
 	}
 
 	int32 NumNestedToShow = FMath::Min(Task->NestedTasks.Num(), MaxTasksToShow);
@@ -365,14 +359,8 @@ void FTaskGraphProfilerManager::GetSingleTaskConnections(const TraceServices::FT
 	{
 		const TraceServices::FTaskInfo* Subsequent = TasksProvider->TryGetTask(Task->Subsequents[i].RelativeId);
 		check(Subsequent != nullptr);
-		double StartTime = Task->CompletedTimestamp;
-		uint32 ThreadId = Task->CompletedThreadId;
-		if (StartTime <= 0.0)
-		{
-			StartTime = Task->FinishedTimestamp;
-			ThreadId = Task->StartedThreadId;
-		}
-		AddRelation(InSelectedEvent, StartTime, ThreadId, Subsequent->StartedTimestamp, Subsequent->StartedThreadId, ETaskEventType::Subsequent);
+
+		AddRelation(InSelectedEvent, Task->FinishedTimestamp, Task->StartedThreadId, Subsequent->StartedTimestamp, Subsequent->StartedThreadId, ETaskEventType::Subsequent);
 	}
 }
 
diff --git a/Engine/Source/Developer/TraceInsights/Private/Insights/TaskGraphProfiler/ViewModels/TaskTimingTrack.cpp b/Engine/Source/Developer/TraceInsights/Private/Insights/TaskGraphProfiler/ViewModels/TaskTimingTrack.cpp
--- a/Engine/Source/Developer/TraceInsights/Private/Insights/TaskGraphProfiler/ViewModels/TaskTimingTrack.cpp
+++ b/Engine/Source/Developer/TraceInsights/Private/Insights/TaskGraphProfiler/ViewModels/TaskTimingTrack.cpp
@@ -85,6 +85,18 @@ public:
 				   "Show/hide relations representing the critical path containing the current task.",
 				   EUserInterfaceActionType::ToggleButton,
 				   FInputChord());
+
+		UI_COMMAND(Command_ShowTaskTrack,
+				   "Show Task Overview Track",
+				   "Show/hide the Task Overview Track when a task is selected.",
+				   EUserInterfaceActionType::ToggleButton,
+				   FInputChord());
+
+		UI_COMMAND(Command_ShowDetailedTaskTrackInfo,
+				   "Show Detailed Info on the Task Overview Track",
+				   "Show the current task's prerequisites/nested tasks/subsequents in the Task Overview Track.",
+				   EUserInterfaceActionType::ToggleButton,
+				   FInputChord());
 	}
 	PRAGMA_ENABLE_OPTIMIZATION
 
@@ -94,6 +106,8 @@ public:
 	TSharedPtr<FUICommandInfo> Command_ShowTaskSubsequents;
 	TSharedPtr<FUICommandInfo> Command_ShowNestedTasks;
 	TSharedPtr<FUICommandInfo> Command_ShowCriticalPath;
+	TSharedPtr<FUICommandInfo> Command_ShowTaskTrack;
+	TSharedPtr<FUICommandInfo> Command_ShowDetailedTaskTrackInfo;
 };
 
 ////////////////////////////////////////////////////////////////////////////////////////////////////
@@ -281,6 +295,26 @@ void FTaskTimingSharedState::BuildTasksSubMenu(FMenuBuilder& MenuBuilder)
 		TAttribute<FText>(),
 		FSlateIcon(FInsightsStyle::GetStyleSetName(), "Icons.ShowNestedTasks")
 	);
+
+	MenuBuilder.AddSeparator();
+
+	MenuBuilder.AddMenuEntry
+	(
+		FTaskTimingStateCommands::Get().Command_ShowTaskTrack,
+		NAME_None,
+		TAttribute<FText>(),
+		TAttribute<FText>(),
+		FSlateIcon(FInsightsStyle::GetStyleSetName(), "Icons.ShowTaskTrack")
+	);
+
+	MenuBuilder.AddMenuEntry
+	(
+		FTaskTimingStateCommands::Get().Command_ShowDetailedTaskTrackInfo,
+		NAME_None,
+		TAttribute<FText>(),
+		TAttribute<FText>(),
+		FSlateIcon(FInsightsStyle::GetStyleSetName(), "Icons.ShowDetailedTaskTrackInfo")
+	);
 }
 
 ////////////////////////////////////////////////////////////////////////////////////////////////////
@@ -327,6 +361,18 @@ void FTaskTimingSharedState::InitCommandList()
 		FExecuteAction::CreateSP(this, &FTaskTimingSharedState::ContextMenu_ShowCriticalPath_Execute),
 		FCanExecuteAction::CreateSP(this, &FTaskTimingSharedState::ContextMenu_ShowCriticalPath_CanExecute),
 		FIsActionChecked::CreateSP(this, &FTaskTimingSharedState::ContextMenu_ShowCriticalPath_IsChecked));
+
+	CommandList->MapAction(
+		FTaskTimingStateCommands::Get().Command_ShowTaskTrack,
+		FExecuteAction::CreateSP(this, &FTaskTimingSharedState::ContextMenu_ShowTaskTrack_Execute),
+		FCanExecuteAction::CreateSP(this, &FTaskTimingSharedState::ContextMenu_ShowTaskTrack_CanExecute),
+		FIsActionChecked::CreateSP(this, &FTaskTimingSharedState::ContextMenu_ShowTaskTrack_IsChecked));
+
+	CommandList->MapAction(
+		FTaskTimingStateCommands::Get().Command_ShowDetailedTaskTrackInfo,
+		FExecuteAction::CreateSP(this, &FTaskTimingSharedState::ContextMenu_ShowDetailedTaskTrackInfo_Execute),
+		FCanExecuteAction::CreateSP(this, &FTaskTimingSharedState::ContextMenu_ShowDetailedTaskTrackInfo_CanExecute),
+		FIsActionChecked::CreateSP(this, &FTaskTimingSharedState::ContextMenu_ShowDetailedTaskTrackInfo_IsChecked));
 }
 
 ////////////////////////////////////////////////////////////////////////////////////////////////////
@@ -499,6 +545,65 @@ void FTaskTimingSharedState::ContextMenu_ShowCriticalPath_Execute()
 
 ////////////////////////////////////////////////////////////////////////////////////////////////////
 
+void FTaskTimingSharedState::ContextMenu_ShowTaskTrack_Execute()
+{
+	if (TaskTrack.IsValid())
+	{
+		TaskTrack->ToggleVisibility();
+	}
+}
+
+////////////////////////////////////////////////////////////////////////////////////////////////////
+
+bool FTaskTimingSharedState::ContextMenu_ShowTaskTrack_CanExecute()
+{
+	return true;
+}
+
+////////////////////////////////////////////////////////////////////////////////////////////////////
+
+bool FTaskTimingSharedState::ContextMenu_ShowTaskTrack_IsChecked()
+{
+	if (TaskTrack.IsValid())
+	{
+		return TaskTrack->IsVisible();
+	}
+
+	return false;
+}
+
+////////////////////////////////////////////////////////////////////////////////////////////////////
+
+void FTaskTimingSharedState::ContextMenu_ShowDetailedTaskTrackInfo_Execute()
+{
+	if (TaskTrack.IsValid())
+	{
+		TaskTrack->SetShowDetailedInfoOnTaskTrack(!TaskTrack->GetShowDetailedInfoOnTaskTrack());
+		TaskTrack->SetDirtyFlag();
+	}
+}
+////////////////////////////////////////////////////////////////////////////////////////////////////
+
+
+bool FTaskTimingSharedState::ContextMenu_ShowDetailedTaskTrackInfo_CanExecute()
+{
+	return true;
+}
+
+////////////////////////////////////////////////////////////////////////////////////////////////////
+
+bool FTaskTimingSharedState::ContextMenu_ShowDetailedTaskTrackInfo_IsChecked()
+{
+	if (TaskTrack.IsValid())
+	{
+		return TaskTrack->GetShowDetailedInfoOnTaskTrack();
+	}
+
+	return false;
+}
+
+////////////////////////////////////////////////////////////////////////////////////////////////////
+
 void FTaskTimingSharedState::OnTaskSettingsChanged()
 {
 	if (!TaskTrack.IsValid())
@@ -571,6 +676,39 @@ void FTaskTimingTrack::BuildDrawState(ITimingEventsTrackDrawStateBuilder& Builde
 	{
 		Builder.AddEvent(Task->FinishedTimestamp, Task->CompletedTimestamp, 0, TEXT("Completed"), 0, FTaskGraphProfilerManager::Get()->GetColorForTaskEventAsPackedARGB(ETaskEventType::Completed));
 	}
+
+	if (bShowDetailInfoOnTaskTrack)
+	{
+		uint32 Depth = 0;
+		for (TraceServices::FTaskInfo::FRelationInfo Relation : Task->Prerequisites)
+		{
+			const TraceServices::FTaskInfo* TaskInfo = TasksProvider->TryGetTask(Relation.RelativeId);
+			if (TaskInfo)
+			{
+				Builder.AddEvent(TaskInfo->StartedTimestamp, TaskInfo->FinishedTimestamp, ++Depth, *FString::Printf(TEXT("Prerequisite Task %d"), TaskInfo->Id), 0, FTaskGraphProfilerManager::Get()->GetColorForTaskEventAsPackedARGB(ETaskEventType::Prerequisite));
+			}
+		}
+
+		// The nested/subsequents tasks should not overlap the prerequisites tasks, so we can draw them on same depth."
+		Depth = 0;
+		for (TraceServices::FTaskInfo::FRelationInfo Relation : Task->NestedTasks)
+		{
+			const TraceServices::FTaskInfo* TaskInfo = TasksProvider->TryGetTask(Relation.RelativeId);
+			if (TaskInfo)
+			{
+				Builder.AddEvent(TaskInfo->StartedTimestamp, TaskInfo->FinishedTimestamp, ++Depth, *FString::Printf(TEXT("Nested Task %d"), TaskInfo->Id), 0, FTaskGraphProfilerManager::Get()->GetColorForTaskEventAsPackedARGB(ETaskEventType::AddedNested));
+			}
+		}
+
+		for (TraceServices::FTaskInfo::FRelationInfo Relation : Task->Subsequents)
+		{
+			const TraceServices::FTaskInfo* TaskInfo = TasksProvider->TryGetTask(Relation.RelativeId);
+			if (TaskInfo)
+			{
+				Builder.AddEvent(TaskInfo->StartedTimestamp, TaskInfo->FinishedTimestamp, ++Depth, *FString::Printf(TEXT("Subsequent Task %d"), TaskInfo->Id), 0, FTaskGraphProfilerManager::Get()->GetColorForTaskEventAsPackedARGB(ETaskEventType::Subsequent));
+			}
+		}
+	}
 }
 
 ////////////////////////////////////////////////////////////////////////////////////////////////////
@@ -665,6 +803,7 @@ const TSharedPtr<const ITimingEvent> FTaskTimingTrack::GetEvent(float InPosX, fl
 	if (DY >= 0 && DY < GetHeight() - 1.0f - 2 * Layout.TimelineDY)
 	{
 		const double EventTime = Viewport.SlateUnitsToTime(InPosX);
+		const int32 Depth = DY / (Layout.EventH + Layout.EventDY);
 
 		TSharedPtr<const TraceServices::IAnalysisSession> Session = FInsightsManager::Get()->GetSession();
 		if (!Session.IsValid())
@@ -688,31 +827,67 @@ const TSharedPtr<const ITimingEvent> FTaskTimingTrack::GetEvent(float InPosX, fl
 			return TimingEvent;
 		}
 
-		// This logic will be replaced
-		if (EventTime >= Task->CreatedTimestamp && EventTime < Task->LaunchedTimestamp)
-		{
-			TimingEvent = MakeShared<FTaskTrackEvent>(SharedThis(this), Task->CreatedTimestamp, Task->LaunchedTimestamp, 0, ETaskTrackEventType::Launched);
-		}
-		else if (EventTime >= Task->LaunchedTimestamp && EventTime < Task->ScheduledTimestamp)
-		{
-			TimingEvent = MakeShared<FTaskTrackEvent>(SharedThis(this), Task->LaunchedTimestamp, Task->ScheduledTimestamp, 0, ETaskTrackEventType::Dispatched);
-		}
-		else if (EventTime >= Task->ScheduledTimestamp && EventTime < Task->StartedTimestamp)
-		{
-			TimingEvent = MakeShared<FTaskTrackEvent>(SharedThis(this), Task->ScheduledTimestamp, Task->StartedTimestamp, 0, ETaskTrackEventType::Scheduled);
-		}
-		else if (EventTime >= Task->StartedTimestamp && EventTime < Task->FinishedTimestamp)
+		if (Depth == 0)
 		{
-			TimingEvent = MakeShared<FTaskTrackEvent>(SharedThis(this), Task->StartedTimestamp, Task->FinishedTimestamp, 0, ETaskTrackEventType::Executed);
+			if (EventTime >= Task->CreatedTimestamp && EventTime < Task->LaunchedTimestamp)
+			{
+				TimingEvent = MakeShared<FTaskTrackEvent>(SharedThis(this), Task->CreatedTimestamp, Task->LaunchedTimestamp, 0, ETaskTrackEventType::Launched);
+			}
+			else if (EventTime >= Task->LaunchedTimestamp && EventTime < Task->ScheduledTimestamp)
+			{
+				TimingEvent = MakeShared<FTaskTrackEvent>(SharedThis(this), Task->LaunchedTimestamp, Task->ScheduledTimestamp, 0, ETaskTrackEventType::Dispatched);
+			}
+			else if (EventTime >= Task->ScheduledTimestamp && EventTime < Task->StartedTimestamp)
+			{
+				TimingEvent = MakeShared<FTaskTrackEvent>(SharedThis(this), Task->ScheduledTimestamp, Task->StartedTimestamp, 0, ETaskTrackEventType::Scheduled);
+			}
+			else if (EventTime >= Task->StartedTimestamp && EventTime < Task->FinishedTimestamp)
+			{
+				TimingEvent = MakeShared<FTaskTrackEvent>(SharedThis(this), Task->StartedTimestamp, Task->FinishedTimestamp, 0, ETaskTrackEventType::Executed);
+			}
+			else if (EventTime >= Task->FinishedTimestamp && EventTime < Task->CompletedTimestamp)
+			{
+				TimingEvent = MakeShared<FTaskTrackEvent>(SharedThis(this), Task->FinishedTimestamp, Task->CompletedTimestamp, 0, ETaskTrackEventType::Completed);
+			}
+
+			if (TimingEvent.IsValid())
+			{
+				TimingEvent->SetTaskId(Task->Id);
+			}
 		}
-		else if (EventTime >= Task->FinishedTimestamp && EventTime < Task->CompletedTimestamp)
+		else if(bShowDetailInfoOnTaskTrack)
 		{
-			TimingEvent = MakeShared<FTaskTrackEvent>(SharedThis(this), Task->FinishedTimestamp, Task->CompletedTimestamp, 0, ETaskTrackEventType::Completed);
-		}
+			auto GetEventFromRelations = [&TasksProvider, this, EventTime](const TArray<TraceServices::FTaskInfo::FRelationInfo>& Relations, int32 Depth, ETaskTrackEventType EventType, int32 PosOffset) -> TSharedPtr<FTaskTrackEvent>
+			{
+				if (Depth - PosOffset <= 0 || Relations.Num() < Depth - PosOffset)
+				{
+					return nullptr;
+				}
+				const TraceServices::FTaskInfo::FRelationInfo& Relation = Relations[Depth - PosOffset - 1];
+				{
+					const TraceServices::FTaskInfo* TaskInfo = TasksProvider->TryGetTask(Relation.RelativeId);
+					if (TaskInfo && EventTime >= TaskInfo->StartedTimestamp && EventTime <= TaskInfo->FinishedTimestamp)
+					{
+						TSharedPtr<FTaskTrackEvent> NewTimingEvent = MakeShared<FTaskTrackEvent>(SharedThis(this), TaskInfo->StartedTimestamp, TaskInfo->FinishedTimestamp, Depth, EventType);
+						NewTimingEvent->SetTaskId(TaskInfo->Id);
+						return NewTimingEvent;
+					}
+				}
 
-		if (TimingEvent.IsValid())
-		{
-			TimingEvent->SetTaskId(Task->Id);
+				return nullptr;
+			};
+
+			TimingEvent = GetEventFromRelations(Task->Prerequisites, Depth, ETaskTrackEventType::Prerequisite, 0);
+
+			if (!TimingEvent.IsValid())
+			{
+				TimingEvent = GetEventFromRelations(Task->NestedTasks, Depth, ETaskTrackEventType::Nested, 0);
+			}
+
+			if (!TimingEvent.IsValid())
+			{
+				TimingEvent = GetEventFromRelations(Task->Subsequents, Depth, ETaskTrackEventType::Subsequent, Task->NestedTasks.Num());
+			}
 		}
 	}
 
@@ -765,6 +940,13 @@ void FTaskTimingTrack::InitTooltip(FTooltipDrawState& InOutTooltip, const ITimin
 	{
 	case ETaskTrackEventType::Launched:
 		break;
+	case ETaskTrackEventType::Prerequisite:
+	case ETaskTrackEventType::Nested:
+	case ETaskTrackEventType::Subsequent:
+	{
+		InOutTooltip.AddNameValueTextLine(TEXT("Executed Thread Id:"), FString::Printf(TEXT("%d"), Task->StartedThreadId));
+		break;
+	}
 	case ETaskTrackEventType::Dispatched:
 	{
 		InOutTooltip.AddNameValueTextLine(TEXT("Prerequisite tasks:"), FString::Printf(TEXT("%d"), Task->Prerequisites.Num()));
diff --git a/Engine/Source/Developer/TraceInsights/Private/Insights/TaskGraphProfiler/ViewModels/TaskTimingTrack.h b/Engine/Source/Developer/TraceInsights/Private/Insights/TaskGraphProfiler/ViewModels/TaskTimingTrack.h
--- a/Engine/Source/Developer/TraceInsights/Private/Insights/TaskGraphProfiler/ViewModels/TaskTimingTrack.h
+++ b/Engine/Source/Developer/TraceInsights/Private/Insights/TaskGraphProfiler/ViewModels/TaskTimingTrack.h
@@ -82,6 +82,14 @@ private:
 	bool ContextMenu_ShowCriticalPath_CanExecute();
 	bool ContextMenu_ShowCriticalPath_IsChecked();
 
+	void ContextMenu_ShowTaskTrack_Execute();
+	bool ContextMenu_ShowTaskTrack_CanExecute();
+	bool ContextMenu_ShowTaskTrack_IsChecked();
+
+	void ContextMenu_ShowDetailedTaskTrackInfo_Execute();
+	bool ContextMenu_ShowDetailedTaskTrackInfo_CanExecute();
+	bool ContextMenu_ShowDetailedTaskTrackInfo_IsChecked();
+
 	void OnTaskSettingsChanged();
 
 private:
@@ -128,6 +136,9 @@ public:
 	void OnTimingEventSelected(TSharedPtr<const ITimingEvent> InSelectedEvent);
 	void GetEventRelations(const FThreadTrackEvent& InSelectedEvent);
 
+	void SetShowDetailedInfoOnTaskTrack(bool InValue) { bShowDetailInfoOnTaskTrack = InValue; }
+	bool GetShowDetailedInfoOnTaskTrack() { return bShowDetailInfoOnTaskTrack; }
+
 private:
 	uint32 TimelineIndex;
 
@@ -136,6 +147,8 @@ private:
 	TaskTrace::FId TaskId = TaskTrace::InvalidId;
 
 	FVector2D MousePositionOnButtonDown;
+
+	bool bShowDetailInfoOnTaskTrack = true;
 };
 
 ////////////////////////////////////////////////////////////////////////////////////////////////////
diff --git a/Engine/Source/Developer/TraceInsights/Private/Insights/TaskGraphProfiler/ViewModels/TaskTrackEvent.cpp b/Engine/Source/Developer/TraceInsights/Private/Insights/TaskGraphProfiler/ViewModels/TaskTrackEvent.cpp
--- a/Engine/Source/Developer/TraceInsights/Private/Insights/TaskGraphProfiler/ViewModels/TaskTrackEvent.cpp
+++ b/Engine/Source/Developer/TraceInsights/Private/Insights/TaskGraphProfiler/ViewModels/TaskTrackEvent.cpp
@@ -26,6 +26,9 @@ FString FTaskTrackEvent::GetStartLabel() const
 	case ETaskTrackEventType::Scheduled:
 		return TEXT("Scheduled Time:");
 	case ETaskTrackEventType::Executed:
+	case ETaskTrackEventType::Prerequisite:
+	case ETaskTrackEventType::Nested:
+	case ETaskTrackEventType::Subsequent:
 		return TEXT("Started Time:");
 	case ETaskTrackEventType::Completed:
 		return TEXT("Finished Time:");
@@ -50,6 +53,9 @@ FString FTaskTrackEvent::GetEndLabel() const
 	case ETaskTrackEventType::Scheduled:
 		return TEXT("Started Time:");
 	case ETaskTrackEventType::Executed:
+	case ETaskTrackEventType::Prerequisite:
+	case ETaskTrackEventType::Nested:
+	case ETaskTrackEventType::Subsequent:
 		return TEXT("Finished Time:");
 	case ETaskTrackEventType::Completed:
 		return TEXT("Completed Time:");
@@ -77,6 +83,13 @@ FString FTaskTrackEvent::GetEventName() const
 		return TEXT("Executed");
 	case ETaskTrackEventType::Completed:
 		return TEXT("Completed");
+	case ETaskTrackEventType::Prerequisite:
+		return FString::Printf(TEXT("Prerequisite Task %d"), GetTaskId());	
+	case ETaskTrackEventType::Nested:
+		return FString::Printf(TEXT("Nested Task %d"), GetTaskId());
+	case ETaskTrackEventType::Subsequent:
+		return FString::Printf(TEXT("Subsequent Task %d"), GetTaskId());
+
 	default:
 		checkf(false, TEXT("Unknown task event type"));
 		break;
diff --git a/Engine/Source/Developer/TraceInsights/Private/Insights/TaskGraphProfiler/ViewModels/TaskTrackEvent.h b/Engine/Source/Developer/TraceInsights/Private/Insights/TaskGraphProfiler/ViewModels/TaskTrackEvent.h
--- a/Engine/Source/Developer/TraceInsights/Private/Insights/TaskGraphProfiler/ViewModels/TaskTrackEvent.h
+++ b/Engine/Source/Developer/TraceInsights/Private/Insights/TaskGraphProfiler/ViewModels/TaskTrackEvent.h
@@ -16,6 +16,9 @@ enum class ETaskTrackEventType : uint32
 	Scheduled,
 	Executed,
 	Completed,
+	Prerequisite,
+	Nested,
+	Subsequent,
 };
 
 ////////////////////////////////////////////////////////////////////////////////////////////////////
