[Insights] Fix Task Relations remaining displayed after menu option is unset.

diff --git a/Engine/Source/Developer/TraceInsights/Private/Insights/TaskGraphProfiler/ViewModels/TaskTimingTrack.cpp b/Engine/Source/Developer/TraceInsights/Private/Insights/TaskGraphProfiler/ViewModels/TaskTimingTrack.cpp
--- a/Engine/Source/Developer/TraceInsights/Private/Insights/TaskGraphProfiler/ViewModels/TaskTimingTrack.cpp
+++ b/Engine/Source/Developer/TraceInsights/Private/Insights/TaskGraphProfiler/ViewModels/TaskTimingTrack.cpp
@@ -611,10 +611,8 @@ void FTaskTimingSharedState::OnTaskSettingsChanged()
 		return;
 	}
 
-	if (TaskTrack->GetTaskId() == TaskTrace::InvalidId)
-	{
-		FTaskGraphProfilerManager::Get()->ClearTaskRelations();
-	}
+	FTaskGraphProfilerManager::Get()->ClearTaskRelations();
+	TaskTrack->SetTaskId(TaskTrace::InvalidId);
 
 	TSharedPtr<Insights::FTaskGraphProfilerManager> TaskGraphManager = Insights::FTaskGraphProfilerManager::Get();
 
