fix for animation insights menu not appearing on mac

diff --git a/Engine/Source/Developer/TraceInsights/Private/Insights/Widgets/STimingProfilerWindow.cpp b/Engine/Source/Developer/TraceInsights/Private/Insights/Widgets/STimingProfilerWindow.cpp
--- a/Engine/Source/Developer/TraceInsights/Private/Insights/Widgets/STimingProfilerWindow.cpp
+++ b/Engine/Source/Developer/TraceInsights/Private/Insights/Widgets/STimingProfilerWindow.cpp
@@ -369,7 +369,10 @@ void STimingProfilerWindow::Construct(const FArguments& InArgs, const TSharedRef
 {
 	// Create & initialize tab manager.
 	TabManager = FGlobalTabmanager::Get()->NewTabManager(ConstructUnderMajorTab);
-	//TabManager->SetAllowWindowMenuBar(true);
+
+#if WITH_EDITOR
+	TabManager->SetAllowWindowMenuBar(true);
+#endif
 
 	const auto& PersistLayout = [](const TSharedRef<FTabManager::FLayout>& LayoutToSave)
 	{
