[Insights] Memory Insights / Tree Table views: Updated tooltip of memory values to show byte value first.

diff --git a/Engine/Source/Developer/TraceInsights/Private/Insights/Table/ViewModels/TableCellValueFormatter.cpp b/Engine/Source/Developer/TraceInsights/Private/Insights/Table/ViewModels/TableCellValueFormatter.cpp
--- a/Engine/Source/Developer/TraceInsights/Private/Insights/Table/ViewModels/TableCellValueFormatter.cpp
+++ b/Engine/Source/Developer/TraceInsights/Private/Insights/Table/ViewModels/TableCellValueFormatter.cpp
@@ -115,13 +115,20 @@ FText FInt64ValueFormatterAsMemory::FormatValueForTooltip(const TOptional<FTable
 		{
 			if (Value < 1024)
 			{
-				return FText::Format(LOCTEXT("AsMemory_PositiveValue_TooltipFmt1", "{0} bytes"), FText::AsNumber(Value));
+				if (Value == 1)
+				{
+					return LOCTEXT("AsMemory_Tooltip_1byte", "1 byte");
+				}
+				else
+				{
+					return FText::Format(LOCTEXT("AsMemory_PositiveValue_TooltipFmt1", "{0} bytes"), FText::AsNumber(Value));
+				}
 			}
 			else
 			{
 				FNumberFormattingOptions FormattingOptions;
 				FormattingOptions.MaximumFractionalDigits = 2;
-				return FText::Format(LOCTEXT("AsMemory_PositiveValue_TooltipFmt2", "{0} ({1} bytes)"), FText::AsMemory(Value, &FormattingOptions), FText::AsNumber(Value));
+				return FText::Format(LOCTEXT("AsMemory_PositiveValue_TooltipFmt2", "{0} bytes ({1})"), FText::AsNumber(Value), FText::AsMemory(Value, &FormattingOptions));
 			}
 		}
 		else if (Value == 0)
@@ -132,13 +139,20 @@ FText FInt64ValueFormatterAsMemory::FormatValueForTooltip(const TOptional<FTable
 		{
 			if (-Value < 1024)
 			{
-				return FText::Format(LOCTEXT("AsMemory_NegativeValue_TooltipFmt1", "-{0} bytes"), FText::AsNumber(-Value));
+				if (Value == -1)
+				{
+					return LOCTEXT("AsMemory_Tooltip_minus1byte", "-1 byte");
+				}
+				else
+				{
+					return FText::Format(LOCTEXT("AsMemory_NegativeValue_TooltipFmt1", "-{0} bytes"), FText::AsNumber(-Value));
+				}
 			}
 			else
 			{
 				FNumberFormattingOptions FormattingOptions;
 				FormattingOptions.MaximumFractionalDigits = 2;
-				return FText::Format(LOCTEXT("AsMemory_NegativeValue_TooltipFmt2", "-{0} (-{1} bytes)"), FText::AsMemory(-Value, &FormattingOptions), FText::AsNumber(-Value));
+				return FText::Format(LOCTEXT("AsMemory_NegativeValue_TooltipFmt2", "-{0} bytes (-{1})"), FText::AsNumber(-Value), FText::AsMemory(-Value, &FormattingOptions));
 			}
 		}
 	}
