[Insights]
  - Removed log messages re added/removed tracks in STimingView (too noisy).
  - Update log messages re duplicated categories in SLogView (also to reduce noise).
  - Added better logging for canceling an async operation in FStatsAggregator.

diff --git a/Engine/Source/Developer/TraceInsights/Private/Insights/ViewModels/StatsAggregator.cpp b/Engine/Source/Developer/TraceInsights/Private/Insights/ViewModels/StatsAggregator.cpp
--- a/Engine/Source/Developer/TraceInsights/Private/Insights/ViewModels/StatsAggregator.cpp
+++ b/Engine/Source/Developer/TraceInsights/Private/Insights/ViewModels/StatsAggregator.cpp
@@ -92,14 +92,29 @@ void FStatsAggregator::Start()
 			OperationCount = 0;
 		}
 
-		UE_LOG(TimingProfiler, Log, TEXT("[%s] Request async aggregation (op %d)..."), *LogName, OperationCount + 1);
+		UE_LOG(TimingProfiler, Log, TEXT("[%s] Request async aggregation (op %d) [%fs to %fs] (%fs)..."),
+			*LogName, OperationCount + 1, IntervalStartTime, IntervalEndTime, IntervalEndTime - IntervalStartTime);
 		bIsStartRequested = true;
 	}
 }
 
 ////////////////////////////////////////////////////////////////////////////////////////////////////
 
-void FStatsAggregator::Tick(TSharedPtr<const TraceServices::IAnalysisSession> InSession, const double InCurrentTime, const float InDeltaTime, TFunctionRef<void()> OnFinishedCallback)
+void FStatsAggregator::Cancel()
+{
+	if (AsyncTask)
+	{
+		UE_LOG(TimingProfiler, Log, TEXT("[%s] Cancel requested for async aggregation (op %d)..."), *LogName, OperationCount);
+		bIsCancelRequested = true;
+	}
+}
+
+////////////////////////////////////////////////////////////////////////////////////////////////////
+
+void FStatsAggregator::Tick(TSharedPtr<const TraceServices::IAnalysisSession> InSession,
+							const double InCurrentTime,
+							const float InDeltaTime,
+							TFunctionRef<void()> OnFinishedCallback)
 {
 	if (AsyncTask && AsyncTask->IsDone())
 	{
@@ -141,8 +156,12 @@ void FStatsAggregator::Tick(TSharedPtr<const TraceServices::IAnalysisSession> In
 	{
 		if (AsyncTask)
 		{
-			// Cancel and wait for the previous async task to finish.
-			bIsCancelRequested = true;
+			if (!bIsCancelRequested)
+			{
+				// Cancel and wait for the previous async task to finish.
+				UE_LOG(TimingProfiler, Log, TEXT("[%s] Cancel previous async aggregation (op %d)..."), *LogName, OperationCount);
+				bIsCancelRequested = true;
+			}
 		}
 		else
 		{
diff --git a/Engine/Source/Developer/TraceInsights/Private/Insights/ViewModels/StatsAggregator.h b/Engine/Source/Developer/TraceInsights/Private/Insights/ViewModels/StatsAggregator.h
--- a/Engine/Source/Developer/TraceInsights/Private/Insights/ViewModels/StatsAggregator.h
+++ b/Engine/Source/Developer/TraceInsights/Private/Insights/ViewModels/StatsAggregator.h
@@ -66,7 +66,7 @@ public:
 	// IStatsAggregator
 
 	virtual void Start() override;
-	virtual void Cancel() override { bIsCancelRequested = true; }
+	virtual void Cancel() override;
 
 	virtual bool IsCancelRequested() const override { return bIsCancelRequested; }
 	virtual bool IsRunning() const override { return AsyncTask != nullptr; }
diff --git a/Engine/Source/Developer/TraceInsights/Private/Insights/Widgets/SLogView.cpp b/Engine/Source/Developer/TraceInsights/Private/Insights/Widgets/SLogView.cpp
--- a/Engine/Source/Developer/TraceInsights/Private/Insights/Widgets/SLogView.cpp
+++ b/Engine/Source/Developer/TraceInsights/Private/Insights/Widgets/SLogView.cpp
@@ -764,21 +764,43 @@ void SLogView::Tick(const FGeometry& AllottedGeometry, const double InCurrentTim
 			UE_LOG(TraceInsights, Log, TEXT("[LogView] Total Log Categories: %d"), TotalNumCategories);
 
 			TSet<FName> Categories;
-			LogProvider.EnumerateCategories([&Categories](const TraceServices::FLogCategoryInfo& Category)
+			TMap<FName, int32> DuplicatedCategories;
+			LogProvider.EnumerateCategories([&Categories, &DuplicatedCategories](const TraceServices::FLogCategoryInfo& Category)
 			{
 				FString CategoryStr(Category.Name);
 				if (CategoryStr.StartsWith(TEXT("Log")))
 				{
 					CategoryStr.RightChopInline(3, false);
 				}
-				if (Categories.Contains(FName(*CategoryStr)))
+				FName CategoryName(CategoryStr);
+				if (Categories.Contains(CategoryName))
 				{
-					UE_LOG(TraceInsights, Log, TEXT("[LogView] Duplicated Log Category: \"%s\""), Category.Name);
+					int32* CountPtr = DuplicatedCategories.Find(CategoryName);
+					if (CountPtr)
+					{
+						++(*CountPtr);
+					}
+					else
+					{
+						DuplicatedCategories.Add(CategoryName, 1);
+					}
+				}
+				else
+				{
+					Categories.Add(CategoryName);
 				}
-				Categories.Add(FName(*CategoryStr));
 			});
 			Filter.SyncAvailableCategories(Categories);
-			UE_LOG(TraceInsights, Log, TEXT("[LogView] Unique Log Categories: %d"), Filter.GetAvailableLogCategories().Num());
+			const int32 NumAvailableLogCategories = Filter.GetAvailableLogCategories().Num();
+			if (DuplicatedCategories.Num() > 0)
+			{
+				UE_LOG(TraceInsights, Warning, TEXT("[LogView] Duplicated Log Categories: %d (+%dx)"), DuplicatedCategories.Num(), TotalNumCategories - NumAvailableLogCategories);
+				for (const auto& KV : DuplicatedCategories)
+				{
+					UE_LOG(TraceInsights, Log, TEXT("[LogView]    \"%s\" (+%dx)"), *KV.Key.GetPlainNameString(), KV.Value);
+				}
+			}
+			UE_LOG(TraceInsights, Log, TEXT("[LogView] Unique Log Categories: %d"), NumAvailableLogCategories);
 
 			//Cache.Reset();
 			Messages.Reset();
diff --git a/Engine/Source/Developer/TraceInsights/Private/Insights/Widgets/STimingView.cpp b/Engine/Source/Developer/TraceInsights/Private/Insights/Widgets/STimingView.cpp
--- a/Engine/Source/Developer/TraceInsights/Private/Insights/Widgets/STimingView.cpp
+++ b/Engine/Source/Developer/TraceInsights/Private/Insights/Widgets/STimingView.cpp
@@ -1738,11 +1738,13 @@ void STimingView::AddTrack(TSharedPtr<FBaseTimingTrack> Track, ETimingTrackLocat
 		return;
 	}
 
+#if 0
 	UE_LOG(TimingProfiler, Log, TEXT("New %s Track (%d) : %s (\"%s\")"),
 		LocationName,
 		TrackList.Num() + 1,
 		*Track->GetTypeName().ToString(),
 		*Track->GetName());
+#endif
 
 	ensure(Track->GetLocation() == ETimingTrackLocation::None);
 	Track->SetLocation(Location);
@@ -1785,11 +1787,13 @@ bool STimingView::RemoveTrack(TSharedPtr<FBaseTimingTrack> Track)
 			InvalidateScrollableTracksOrder();
 		}
 
+#if 0
 		UE_LOG(TimingProfiler, Log, TEXT("Removed %s Track (%d) : %s (\"%s\")"),
 			LocationName,
 			TrackList.Num(),
 			*Track->GetTypeName().ToString(),
 			*Track->GetName());
+#endif
 
 		return true;
 	}
