[Insights]
  - Added debug code to detect leaked shared references to analysis session when the TraceInsights module shutsdown.

diff --git a/Engine/Source/Developer/TraceInsights/Private/Insights/InsightsManager.h b/Engine/Source/Developer/TraceInsights/Private/Insights/InsightsManager.h
--- a/Engine/Source/Developer/TraceInsights/Private/Insights/InsightsManager.h
+++ b/Engine/Source/Developer/TraceInsights/Private/Insights/InsightsManager.h
@@ -115,7 +115,7 @@ public:
 	/** @return the filename of the trace being analyzed. */
 	const FString& GetTraceFilename() const { return CurrentTraceFilename; }
 
-	/** @returns UI command list for the main manager. */
+	/** @return UI command list for the main manager. */
 	const TSharedRef<FUICommandList> GetCommandList() const;
 
 	/** @return an instance of the main commands. */
diff --git a/Engine/Source/Developer/TraceInsights/Private/Insights/LoadingProfiler/LoadingProfilerManager.cpp b/Engine/Source/Developer/TraceInsights/Private/Insights/LoadingProfiler/LoadingProfilerManager.cpp
--- a/Engine/Source/Developer/TraceInsights/Private/Insights/LoadingProfiler/LoadingProfilerManager.cpp
+++ b/Engine/Source/Developer/TraceInsights/Private/Insights/LoadingProfiler/LoadingProfilerManager.cpp
@@ -54,7 +54,7 @@ FLoadingProfilerManager::FLoadingProfilerManager(TSharedRef<FUICommandList> InCo
 	, bIsAvailable(false)
 	, CommandList(InCommandList)
 	, ActionManager(this)
-	, ProfilerWindow(nullptr)
+	, ProfilerWindowWeakPtr()
 	, bIsTimingViewVisible(false)
 	, bIsEventAggregationTreeViewVisible(false)
 	, bIsObjectTypeAggregationTreeViewVisible(false)
diff --git a/Engine/Source/Developer/TraceInsights/Private/Insights/LoadingProfiler/LoadingProfilerManager.h b/Engine/Source/Developer/TraceInsights/Private/Insights/LoadingProfiler/LoadingProfilerManager.h
--- a/Engine/Source/Developer/TraceInsights/Private/Insights/LoadingProfiler/LoadingProfilerManager.h
+++ b/Engine/Source/Developer/TraceInsights/Private/Insights/LoadingProfiler/LoadingProfilerManager.h
@@ -54,7 +54,7 @@ public:
 
 	//////////////////////////////////////////////////
 
-	/** @returns UI command list for the Loading Profiler manager. */
+	/** @return UI command list for the Loading Profiler manager. */
 	const TSharedRef<FUICommandList> GetCommandList() const;
 
 	/** @return an instance of the Loading Profiler commands. */
@@ -63,23 +63,13 @@ public:
 	/** @return an instance of the Loading Profiler action manager. */
 	static FLoadingProfilerActionManager& GetActionManager();
 
-	void AssignProfilerWindow(const TSharedRef<SLoadingProfilerWindow>& InProfilerWindow)
-	{
-		ProfilerWindow = InProfilerWindow;
-	}
-
-	void RemoveProfilerWindow()
-	{
-		ProfilerWindow.Reset();
-	}
-
 	/**
 	 * Converts profiler window weak pointer to a shared pointer and returns it.
 	 * Make sure the returned pointer is valid before trying to dereference it.
 	 */
-	TSharedPtr<class SLoadingProfilerWindow> GetProfilerWindow() const
+	TSharedPtr<SLoadingProfilerWindow> GetProfilerWindow() const
 	{
-		return ProfilerWindow.Pin();
+		return ProfilerWindowWeakPtr.Pin();
 	}
 
 	////////////////////////////////////////////////////////////////////////////////////////////////////
@@ -136,6 +126,16 @@ private:
 	/** Updates this manager, done through FCoreTicker. */
 	bool Tick(float DeltaTime);
 
+	void AssignProfilerWindow(const TSharedRef<SLoadingProfilerWindow>& InProfilerWindow)
+	{
+		ProfilerWindowWeakPtr = InProfilerWindow;
+	}
+
+	void RemoveProfilerWindow()
+	{
+		ProfilerWindowWeakPtr.Reset();
+	}
+
 private:
 	bool bIsInitialized;
 	bool bIsAvailable;
@@ -153,8 +153,8 @@ private:
 	/** An instance of the Loading Profiler action manager. */
 	FLoadingProfilerActionManager ActionManager;
 
-	/** A weak pointer to the Loading Profiler window. */
-	TWeakPtr<class SLoadingProfilerWindow> ProfilerWindow;
+	/** A weak pointer to the Asset Loading Insights window. */
+	TWeakPtr<SLoadingProfilerWindow> ProfilerWindowWeakPtr;
 
 	/** If the Timing view is visible or hidden. */
 	bool bIsTimingViewVisible;
diff --git a/Engine/Source/Developer/TraceInsights/Private/Insights/MemoryProfiler/MemoryProfilerManager.cpp b/Engine/Source/Developer/TraceInsights/Private/Insights/MemoryProfiler/MemoryProfilerManager.cpp
--- a/Engine/Source/Developer/TraceInsights/Private/Insights/MemoryProfiler/MemoryProfilerManager.cpp
+++ b/Engine/Source/Developer/TraceInsights/Private/Insights/MemoryProfiler/MemoryProfilerManager.cpp
@@ -56,7 +56,7 @@ FMemoryProfilerManager::FMemoryProfilerManager(TSharedRef<FUICommandList> InComm
 	, bIsAvailable(false)
 	, CommandList(InCommandList)
 	, ActionManager(this)
-	, ProfilerWindow(nullptr)
+	, ProfilerWindowWeakPtr()
 	, bIsTimingViewVisible(false)
 	, bIsMemInvestigationViewVisible(false)
 	, bIsMemTagTreeViewVisible(false)
@@ -198,10 +198,10 @@ bool FMemoryProfilerManager::CanSpawnTab(const FSpawnTabArgs& Args) const
 void FMemoryProfilerManager::OnTabClosed(TSharedRef<SDockTab> TabBeingClosed)
 {
 	{
-		TSharedPtr<SMemoryProfilerWindow> ProfilerWindowSharedPtr = ProfilerWindow.Pin();
-		if (ProfilerWindowSharedPtr.IsValid())
+		TSharedPtr<SMemoryProfilerWindow> ProfilerWindow = GetProfilerWindow();
+		if (ProfilerWindow.IsValid())
 		{
-			ProfilerWindowSharedPtr->CloseMemAllocTableTreeTabs();
+			ProfilerWindow->CloseMemAllocTableTreeTabs();
 		}
 	}
 
diff --git a/Engine/Source/Developer/TraceInsights/Private/Insights/MemoryProfiler/MemoryProfilerManager.h b/Engine/Source/Developer/TraceInsights/Private/Insights/MemoryProfiler/MemoryProfilerManager.h
--- a/Engine/Source/Developer/TraceInsights/Private/Insights/MemoryProfiler/MemoryProfilerManager.h
+++ b/Engine/Source/Developer/TraceInsights/Private/Insights/MemoryProfiler/MemoryProfilerManager.h
@@ -19,7 +19,7 @@ DECLARE_LOG_CATEGORY_EXTERN(MemoryProfiler, Log, All);
 
 ////////////////////////////////////////////////////////////////////////////////////////////////////
 /**
- * This class manages the Memory Profiler (Asset Memory Insights) state and settings.
+ * This class manages the Memory Profiler (Memory Insights) state and settings.
  */
 class FMemoryProfilerManager : public TSharedFromThis<FMemoryProfilerManager>, public IInsightsComponent
 {
@@ -55,7 +55,7 @@ public:
 
 	//////////////////////////////////////////////////
 
-	/** @returns UI command list for the Memory Profiler manager. */
+	/** @return UI command list for the Memory Profiler manager. */
 	const TSharedRef<FUICommandList> GetCommandList() const;
 
 	/** @return an instance of the Memory Profiler commands. */
@@ -64,23 +64,13 @@ public:
 	/** @return an instance of the Memory Profiler action manager. */
 	static FMemoryProfilerActionManager& GetActionManager();
 
-	void AssignProfilerWindow(const TSharedRef<SMemoryProfilerWindow>& InProfilerWindow)
-	{
-		ProfilerWindow = InProfilerWindow;
-	}
-
-	void RemoveProfilerWindow()
-	{
-		ProfilerWindow.Reset();
-	}
-
 	/**
 	 * Converts profiler window weak pointer to a shared pointer and returns it.
 	 * Make sure the returned pointer is valid before trying to dereference it.
 	 */
-	TSharedPtr<class SMemoryProfilerWindow> GetProfilerWindow() const
+	TSharedPtr<SMemoryProfilerWindow> GetProfilerWindow() const
 	{
-		return ProfilerWindow.Pin();
+		return ProfilerWindowWeakPtr.Pin();
 	}
 
 	FMemorySharedState* GetSharedState();
@@ -125,6 +115,16 @@ private:
 	/** Updates this manager, done through FCoreTicker. */
 	bool Tick(float DeltaTime);
 
+	void AssignProfilerWindow(const TSharedRef<SMemoryProfilerWindow>& InProfilerWindow)
+	{
+		ProfilerWindowWeakPtr = InProfilerWindow;
+	}
+
+	void RemoveProfilerWindow()
+	{
+		ProfilerWindowWeakPtr.Reset();
+	}
+
 private:
 	bool bIsInitialized;
 	bool bIsAvailable;
@@ -142,8 +142,8 @@ private:
 	/** An instance of the Memory Profiler action manager. */
 	FMemoryProfilerActionManager ActionManager;
 
-	/** A weak pointer to the Memory Profiler window. */
-	TWeakPtr<class SMemoryProfilerWindow> ProfilerWindow;
+	/** A weak pointer to the Memory Insights window. */
+	TWeakPtr<SMemoryProfilerWindow> ProfilerWindowWeakPtr;
 
 	/** If the Timing view is visible or hidden. */
 	bool bIsTimingViewVisible;
diff --git a/Engine/Source/Developer/TraceInsights/Private/Insights/MemoryProfiler/Widgets/SMemAllocTableTreeView.cpp b/Engine/Source/Developer/TraceInsights/Private/Insights/MemoryProfiler/Widgets/SMemAllocTableTreeView.cpp
--- a/Engine/Source/Developer/TraceInsights/Private/Insights/MemoryProfiler/Widgets/SMemAllocTableTreeView.cpp
+++ b/Engine/Source/Developer/TraceInsights/Private/Insights/MemoryProfiler/Widgets/SMemAllocTableTreeView.cpp
@@ -37,8 +37,6 @@
 
 #define LOCTEXT_NAMESPACE "SMemAllocTableTreeView"
 
-using namespace TraceServices;
-
 namespace Insights
 {
 
@@ -73,21 +71,6 @@ void SMemAllocTableTreeView::Reset()
 	STableTreeView::Reset();
 }
 
-////////////////////////////////////////////////////////////////////////////////////////////////////
-/*
-void SMemAllocTableTreeView::UpdateSourceTable(TSharedPtr<TraceServices::IMemAllocTable> SourceTable)
-{
-	//check(Table->Is<Insights::FMemAllocTable>());
-	TSharedPtr<Insights::FMemAllocTable> MemAllocTable = StaticCastSharedPtr<Insights::FMemAllocTable>(Table);
-
-	if (MemAllocTable->UpdateSourceTable(SourceTable))
-	{
-		RebuildColumns();
-	}
-
-	RebuildTree(true);
-}
-*/
 ////////////////////////////////////////////////////////////////////////////////////////////////////
 
 void SMemAllocTableTreeView::Tick(const FGeometry& AllottedGeometry, const double InCurrentTime, const float InDeltaTime)
@@ -110,10 +93,12 @@ void SMemAllocTableTreeView::Tick(const FGeometry& AllottedGeometry, const doubl
 
 void SMemAllocTableTreeView::RebuildTree(bool bResync)
 {
-	FStopwatch SyncStopwatch;
 	FStopwatch Stopwatch;
 	Stopwatch.Start();
 
+	FStopwatch SyncStopwatch;
+	SyncStopwatch.Start();
+
 	if (bResync)
 	{
 		TableTreeNodes.Empty();
@@ -123,7 +108,6 @@ void SMemAllocTableTreeView::RebuildTree(bool bResync)
 
 	TSharedPtr<Insights::FMemAllocTable> MemAllocTable = GetMemAllocTable();
 
-	SyncStopwatch.Start();
 	if (Session.IsValid() && MemAllocTable.IsValid())
 	{
 		TraceServices::IAllocationsProvider::EQueryStatus QueryStatus;
@@ -169,6 +153,7 @@ void SMemAllocTableTreeView::RebuildTree(bool bResync)
 			}
 		}
 	}
+
 	SyncStopwatch.Stop();
 
 	if (bResync || TableTreeNodes.Num() != PreviousNodeCount)
@@ -203,7 +188,7 @@ void SMemAllocTableTreeView::RebuildTree(bool bResync)
 	if (TotalTime > 0.01)
 	{
 		const double SyncTime = SyncStopwatch.GetAccumulatedTime();
-		UE_LOG(MemoryProfiler, Log, TEXT("[MemAlloc] Tree view rebuilt in %.3fs (%.3fs + %.3fs) --> %d nodes (%d added)"),
+		UE_LOG(MemoryProfiler, Log, TEXT("[MemAlloc] Tree view rebuilt in %.4fs (sync: %.4fs + update: %.4fs) --> %d nodes (%d added)"),
 			TotalTime, SyncTime, TotalTime - SyncTime, TableTreeNodes.Num(), TableTreeNodes.Num() - PreviousNodeCount);
 	}
 }
@@ -256,6 +241,12 @@ void SMemAllocTableTreeView::StartQuery()
 		return;
 	}
 
+	if (!Session.IsValid())
+	{
+		UE_LOG(MemoryProfiler, Warning, TEXT("[MemAlloc] Invalid analysis session!"));
+		return;
+	}
+
 	const TraceServices::IAllocationsProvider* AllocationsProvider = TraceServices::ReadAllocationsProvider(*Session.Get());
 	if (!AllocationsProvider)
 	{
@@ -293,6 +284,12 @@ void SMemAllocTableTreeView::UpdateQuery(TraceServices::IAllocationsProvider::EQ
 		return;
 	}
 
+	if (!Session.IsValid())
+	{
+		UE_LOG(MemoryProfiler, Warning, TEXT("[MemAlloc] Invalid analysis session!"));
+		return;
+	}
+
 	const TraceServices::IAllocationsProvider* AllocationsProvider = TraceServices::ReadAllocationsProvider(*Session.Get());
 	if (!AllocationsProvider)
 	{
@@ -318,12 +315,12 @@ void SMemAllocTableTreeView::UpdateQuery(TraceServices::IAllocationsProvider::EQ
 			return;
 		}
 
-		if (Status.Status == IAllocationsProvider::EQueryStatus::Working)
+		if (Status.Status == TraceServices::IAllocationsProvider::EQueryStatus::Working)
 		{
 			break;
 		}
 
-		check(Status.Status == IAllocationsProvider::EQueryStatus::Available);
+		check(Status.Status == TraceServices::IAllocationsProvider::EQueryStatus::Available);
 
 		TSharedPtr<Insights::FMemAllocTable> MemAllocTable = GetMemAllocTable();
 		if (MemAllocTable)
@@ -412,11 +409,14 @@ void SMemAllocTableTreeView::CancelQuery()
 {
 	if (Query != 0)
 	{
-		const TraceServices::IAllocationsProvider* AllocationsProvider = TraceServices::ReadAllocationsProvider(*Session.Get());
-		if (AllocationsProvider)
+		if (Session.IsValid())
 		{
-			AllocationsProvider->CancelQuery(Query);
-			UE_LOG(MemoryProfiler, Log, TEXT("[MemAlloc] Query canceled."));
+			const TraceServices::IAllocationsProvider* AllocationsProvider = TraceServices::ReadAllocationsProvider(*Session.Get());
+			if (AllocationsProvider)
+			{
+				AllocationsProvider->CancelQuery(Query);
+				UE_LOG(MemoryProfiler, Log, TEXT("[MemAlloc] Query canceled."));
+			}
 		}
 
 		Query = 0;
@@ -999,37 +999,40 @@ END_SLATE_FUNCTION_BUILD_OPTIMIZATION
 
 FText SMemAllocTableTreeView::GetSymbolResolutionStatus() const
 {
-	auto ModuleProvider = Session->ReadProvider<IModuleProvider>(FName("ModuleProvider"));
-
-	IModuleProvider::FStats Stats;
-	if (ModuleProvider)
+	if (Session.IsValid())
 	{
-		ModuleProvider->GetStats(&Stats);
-		//check(Stats.SymbolsDiscovered >= Stats.SymbolsResolved + Stats.SymbolsFailed);
-		const int32 SymbolsPending = Stats.SymbolsDiscovered - Stats.SymbolsResolved - Stats.SymbolsFailed;
-		if (SymbolsPending > 0)
-		{
-			return FText::Format(LOCTEXT("SymbolsResolved1", "Resolving {0} / {1} symbols ({2} resolved, {3} failed)"), SymbolsPending, Stats.SymbolsDiscovered, Stats.SymbolsResolved, Stats.SymbolsFailed);
-		}
-		else
+		const TraceServices::IModuleProvider* ModuleProvider = ReadModuleProvider(*Session.Get());
+		if (ModuleProvider)
 		{
-			return FText::Format(LOCTEXT("SymbolsResolved2", "{0} symbols ({1} resolved, {2} failed)"), Stats.SymbolsDiscovered, Stats.SymbolsResolved, Stats.SymbolsFailed);
+			TraceServices::IModuleProvider::FStats Stats;
+			ModuleProvider->GetStats(&Stats);
+			//check(Stats.SymbolsDiscovered >= Stats.SymbolsResolved + Stats.SymbolsFailed);
+			const int32 SymbolsPending = Stats.SymbolsDiscovered - Stats.SymbolsResolved - Stats.SymbolsFailed;
+			if (SymbolsPending > 0)
+			{
+				return FText::Format(LOCTEXT("SymbolsResolved1", "Resolving {0} / {1} symbols ({2} resolved, {3} failed)"), SymbolsPending, Stats.SymbolsDiscovered, Stats.SymbolsResolved, Stats.SymbolsFailed);
+			}
+			else
+			{
+				return FText::Format(LOCTEXT("SymbolsResolved2", "{0} symbols ({1} resolved, {2} failed)"), Stats.SymbolsDiscovered, Stats.SymbolsResolved, Stats.SymbolsFailed);
+			}
 		}
 	}
-	else
-	{
-		return LOCTEXT("SymbolsResolutionNotPossible", "Symbol resolution was not possible.");
-	}
+
+	return LOCTEXT("SymbolsResolutionNotPossible", "Symbol resolution was not possible.");
 }
 
 ////////////////////////////////////////////////////////////////////////////////////////////////////
 
 FText SMemAllocTableTreeView::GetSymbolResolutionTooltip() const
 {
-	const TraceServices::IModuleProvider* ModuleProvider = ReadModuleProvider(*Session.Get());
-	if (ModuleProvider)
+	if (Session.IsValid())
 	{
-		return FSymbolSearchPathsHelper::GetLocalizedSymbolSearchPathsText(ModuleProvider);
+		const TraceServices::IModuleProvider* ModuleProvider = ReadModuleProvider(*Session.Get());
+		if (ModuleProvider)
+		{
+			return FSymbolSearchPathsHelper::GetLocalizedSymbolSearchPathsText(ModuleProvider);
+		}
 	}
 	return FText();
 }
@@ -1076,8 +1079,14 @@ void SMemAllocTableTreeView::InternalCreateGroupings()
 	AvailableGroupings.Insert(MakeShared<FMemAllocGroupingByCallstack>(false, bIsCallstackGroupingByFunction), Index++);
 	AvailableGroupings.Insert(MakeShared<FMemAllocGroupingByCallstack>(true, bIsCallstackGroupingByFunction), Index++);
 
-	const TraceServices::IAllocationsProvider* AllocationsProvider = TraceServices::ReadAllocationsProvider(*Session.Get());
-	AvailableGroupings.Insert(MakeShared<FMemAllocGroupingByHeap>(*AllocationsProvider), Index++);
+	if (Session.IsValid())
+	{
+		const TraceServices::IAllocationsProvider* AllocationsProvider = TraceServices::ReadAllocationsProvider(*Session.Get());
+		if (AllocationsProvider)
+		{
+			AvailableGroupings.Insert(MakeShared<FMemAllocGroupingByHeap>(*AllocationsProvider), Index++);
+		}
+	}
 }
 
 ////////////////////////////////////////////////////////////////////////////////////////////////////
diff --git a/Engine/Source/Developer/TraceInsights/Private/Insights/MemoryProfiler/Widgets/SMemInvestigationView.cpp b/Engine/Source/Developer/TraceInsights/Private/Insights/MemoryProfiler/Widgets/SMemInvestigationView.cpp
--- a/Engine/Source/Developer/TraceInsights/Private/Insights/MemoryProfiler/Widgets/SMemInvestigationView.cpp
+++ b/Engine/Source/Developer/TraceInsights/Private/Insights/MemoryProfiler/Widgets/SMemInvestigationView.cpp
@@ -38,6 +38,8 @@ SMemInvestigationView::~SMemInvestigationView()
 	{
 		FInsightsManager::Get()->GetSessionChangedEvent().RemoveAll(this);
 	}
+
+	Session.Reset();
 }
 
 ////////////////////////////////////////////////////////////////////////////////////////////////////
@@ -58,8 +60,6 @@ void SMemInvestigationView::Construct(const FArguments& InArgs, TSharedPtr<SMemo
 		]
 	];
 
-	//BindCommands();
-
 	// Register ourselves with the Insights manager.
 	FInsightsManager::Get()->GetSessionChangedEvent().AddSP(this, &SMemInvestigationView::InsightsManager_OnSessionChanged);
 
@@ -184,17 +184,17 @@ TSharedRef<SWidget> SMemInvestigationView::ConstructInvestigationWidgetArea()
 	
 		+ SVerticalBox::Slot()
 		.AutoHeight()
-		.Padding(20.0f, 0.0f)
+		.Padding(FMargin(0.0f, 10.0f, 0.0f, 0.0f))
 		.HAlign(HAlign_Left)
 		[
 			SAssignNew(SymbolPathsTextBlock, STextBlock)
 			.Text(GetSymbolPathsText())
+			.ColorAndOpacity(FLinearColor(0.3f, 0.3f, 0.3f, 1.0f))
 			.AutoWrapText(true)
 		]
 	;
 
-	TSharedPtr<SMemoryProfilerWindow> ProfilerWindow = ProfilerWindowWeakPtr.Pin();
-
+	TSharedPtr<SMemoryProfilerWindow> ProfilerWindow = GetProfilerWindow();
 	if (ProfilerWindow.IsValid())
 	{
 		FMemorySharedState& SharedState = ProfilerWindow->GetSharedState();
@@ -243,8 +243,7 @@ END_SLATE_FUNCTION_BUILD_OPTIMIZATION
 BEGIN_SLATE_FUNCTION_BUILD_OPTIMIZATION
 TSharedRef<SWidget> SMemInvestigationView::ConstructTimeMarkerWidget(uint32 TimeMarkerIndex)
 {
-	TSharedPtr<SMemoryProfilerWindow> ProfilerWindow = ProfilerWindowWeakPtr.Pin();
-
+	TSharedPtr<SMemoryProfilerWindow> ProfilerWindow = GetProfilerWindow();
 	if (!ProfilerWindow.IsValid())
 	{
 		return SNew(SBox);
@@ -292,11 +291,15 @@ TSharedRef<SWidget> SMemInvestigationView::ConstructTimeMarkerWidget(uint32 Time
 				{
 					return FText::FromString(FString::Printf(TEXT("%.9f"), TimeMarker->GetTime()));
 				})
-			.OnTextCommitted_Lambda([TimeMarker, ProfilerWindow](const FText& InText, ETextCommit::Type InCommitType)
+			.OnTextCommitted_Lambda([this, TimeMarker](const FText& InText, ETextCommit::Type InCommitType)
 				{
 					const double Time = FCString::Atod(*InText.ToString());
 					TimeMarker->SetTime(Time);
-					ProfilerWindow->OnTimeMarkerChanged(Insights::ETimeChangedFlags::None, TimeMarker);
+					TSharedPtr<SMemoryProfilerWindow> ProfilerWindow = GetProfilerWindow();
+					if (ProfilerWindow.IsValid())
+					{
+						ProfilerWindow->OnTimeMarkerChanged(Insights::ETimeChangedFlags::None, TimeMarker);
+					}
 				})
 		]
 
@@ -330,7 +333,6 @@ END_SLATE_FUNCTION_BUILD_OPTIMIZATION
 void SMemInvestigationView::InsightsManager_OnSessionChanged()
 {
 	TSharedPtr<const TraceServices::IAnalysisSession> NewSession = FInsightsManager::Get()->GetSession();
-
 	if (NewSession != Session)
 	{
 		Session = NewSession;
@@ -357,8 +359,7 @@ void SMemInvestigationView::Tick(const FGeometry& AllottedGeometry, const double
 
 const TArray<TSharedPtr<Insights::FMemoryRuleSpec>>* SMemInvestigationView::GetAvailableQueryRules()
 {
-	TSharedPtr<SMemoryProfilerWindow> ProfilerWindow = ProfilerWindowWeakPtr.Pin();
-
+	TSharedPtr<SMemoryProfilerWindow> ProfilerWindow = GetProfilerWindow();
 	if (ProfilerWindow.IsValid())
 	{
 		FMemorySharedState& SharedState = ProfilerWindow->GetSharedState();
@@ -373,8 +374,7 @@ void SMemInvestigationView::QueryRule_OnSelectionChanged(TSharedPtr<Insights::FM
 {
 	if (SelectInfo != ESelectInfo::Direct)
 	{
-		TSharedPtr<SMemoryProfilerWindow> ProfilerWindow = ProfilerWindowWeakPtr.Pin();
-
+		TSharedPtr<SMemoryProfilerWindow> ProfilerWindow = GetProfilerWindow();
 		if (ProfilerWindow.IsValid() && InRule)
 		{
 			FMemorySharedState& SharedState = ProfilerWindow->GetSharedState();
@@ -388,8 +388,7 @@ void SMemInvestigationView::QueryRule_OnSelectionChanged(TSharedPtr<Insights::FM
 
 FText SMemInvestigationView::QueryRule_GetSelectedText() const
 {
-	TSharedPtr<SMemoryProfilerWindow> ProfilerWindow = ProfilerWindowWeakPtr.Pin();
-
+	TSharedPtr<SMemoryProfilerWindow> ProfilerWindow = GetProfilerWindow();
 	if (ProfilerWindow.IsValid())
 	{
 		FMemorySharedState& SharedState = ProfilerWindow->GetSharedState();
@@ -406,8 +405,7 @@ FText SMemInvestigationView::QueryRule_GetSelectedText() const
 
 FText SMemInvestigationView::QueryRule_GetTooltipText() const
 {
-	TSharedPtr<SMemoryProfilerWindow> ProfilerWindow = ProfilerWindowWeakPtr.Pin();
-
+	TSharedPtr<SMemoryProfilerWindow> ProfilerWindow = GetProfilerWindow();
 	if (ProfilerWindow.IsValid())
 	{
 		FMemorySharedState& SharedState = ProfilerWindow->GetSharedState();
@@ -424,8 +422,7 @@ FText SMemInvestigationView::QueryRule_GetTooltipText() const
 
 const TArray<TSharedPtr<Insights::FQueryTargetWindowSpec>>* SMemInvestigationView::GetAvailableQueryTargets()
 {
-	TSharedPtr<SMemoryProfilerWindow> ProfilerWindow = ProfilerWindowWeakPtr.Pin();
-
+	TSharedPtr<SMemoryProfilerWindow> ProfilerWindow = GetProfilerWindow();
 	if (ProfilerWindow.IsValid())
 	{
 		FMemorySharedState& SharedState = ProfilerWindow->GetSharedState();
@@ -440,8 +437,7 @@ void SMemInvestigationView::QueryTarget_OnSelectionChanged(TSharedPtr<Insights::
 {
 	if (SelectInfo != ESelectInfo::Type::Direct)
 	{
-		TSharedPtr<SMemoryProfilerWindow> ProfilerWindow = ProfilerWindowWeakPtr.Pin();
-
+		TSharedPtr<SMemoryProfilerWindow> ProfilerWindow = GetProfilerWindow();
 		if (ProfilerWindow.IsValid() && InTarget)
 		{
 			FMemorySharedState& SharedState = ProfilerWindow->GetSharedState();
@@ -480,8 +476,7 @@ TSharedRef<SWidget> SMemInvestigationView::QueryTarget_OnGenerateWidget(TSharedP
 
 FText SMemInvestigationView::QueryTarget_GetSelectedText() const
 {
-	TSharedPtr<SMemoryProfilerWindow> ProfilerWindow = ProfilerWindowWeakPtr.Pin();
-
+	TSharedPtr<SMemoryProfilerWindow> ProfilerWindow = GetProfilerWindow();
 	if (ProfilerWindow.IsValid())
 	{
 		FMemorySharedState& SharedState = ProfilerWindow->GetSharedState();
@@ -498,8 +493,7 @@ FText SMemInvestigationView::QueryTarget_GetSelectedText() const
 
 FReply SMemInvestigationView::RunQuery()
 {
-	TSharedPtr<SMemoryProfilerWindow> ProfilerWindow = ProfilerWindowWeakPtr.Pin();
-
+	TSharedPtr<SMemoryProfilerWindow> ProfilerWindow = GetProfilerWindow();
 	if (!ProfilerWindow.IsValid())
 	{
 		UE_LOG(MemoryProfiler, Error, TEXT("[MemQuery] Invalid Profiler Window!"));
@@ -561,7 +555,7 @@ FReply SMemInvestigationView::RunQuery()
 
 FReply SMemInvestigationView::OnTimeMarkerLabelDoubleClicked(const FGeometry& MyGeometry, const FPointerEvent& PointerEvent, uint32 TimeMarkerIndex)
 {
-	TSharedPtr<SMemoryProfilerWindow> ProfilerWindow = ProfilerWindowWeakPtr.Pin();
+	TSharedPtr<SMemoryProfilerWindow> ProfilerWindow = GetProfilerWindow();
 	if (ProfilerWindow.IsValid())
 	{
 		const uint32 NumTimeMarkers = ProfilerWindow->GetNumCustomTimeMarkers();
diff --git a/Engine/Source/Developer/TraceInsights/Private/Insights/MemoryProfiler/Widgets/SMemInvestigationView.h b/Engine/Source/Developer/TraceInsights/Private/Insights/MemoryProfiler/Widgets/SMemInvestigationView.h
--- a/Engine/Source/Developer/TraceInsights/Private/Insights/MemoryProfiler/Widgets/SMemInvestigationView.h
+++ b/Engine/Source/Developer/TraceInsights/Private/Insights/MemoryProfiler/Widgets/SMemInvestigationView.h
@@ -41,6 +41,15 @@ public:
 	SLATE_BEGIN_ARGS(SMemInvestigationView) {}
 	SLATE_END_ARGS()
 
+	/**
+	 * Converts profiler window weak pointer to a shared pointer and returns it.
+	 * Make sure the returned pointer is valid before trying to dereference it.
+	 */
+	TSharedPtr<SMemoryProfilerWindow> GetProfilerWindow() const
+	{
+		return ProfilerWindowWeakPtr.Pin();
+	}
+
 	/**
 	 * Construct this widget
 	 * @param InArgs - The declaration data for this widget
@@ -81,10 +90,11 @@ private:
 	FReply OnTimeMarkerLabelDoubleClicked(const FGeometry& MyGeometry, const FPointerEvent& PointerEvent, uint32 TimeMarkerIndex);
 
 private:
+	/** A weak pointer to the Memory Insights window. */
 	TWeakPtr<SMemoryProfilerWindow> ProfilerWindowWeakPtr;
 
-	/** A weak pointer to the profiler session used to populate this widget. */
-	TSharedPtr<const TraceServices::IAnalysisSession>/*Weak*/ Session;
+	/** The analysis session used to populate this widget. */
+	TSharedPtr<const TraceServices::IAnalysisSession> Session;
 
 	TSharedPtr<SComboBox<TSharedPtr<Insights::FMemoryRuleSpec>>> QueryRuleComboBox;
 
diff --git a/Engine/Source/Developer/TraceInsights/Private/Insights/MemoryProfiler/Widgets/SMemTagTreeView.cpp b/Engine/Source/Developer/TraceInsights/Private/Insights/MemoryProfiler/Widgets/SMemTagTreeView.cpp
--- a/Engine/Source/Developer/TraceInsights/Private/Insights/MemoryProfiler/Widgets/SMemTagTreeView.cpp
+++ b/Engine/Source/Developer/TraceInsights/Private/Insights/MemoryProfiler/Widgets/SMemTagTreeView.cpp
@@ -60,6 +60,8 @@ SMemTagTreeView::~SMemTagTreeView()
 	{
 		FInsightsManager::Get()->GetSessionChangedEvent().RemoveAll(this);
 	}
+
+	Session.Reset();
 }
 
 ////////////////////////////////////////////////////////////////////////////////////////////////////
@@ -154,7 +156,6 @@ void SMemTagTreeView::Construct(const FArguments& InArgs, TSharedPtr<SMemoryProf
 	];
 
 	InitializeAndShowHeaderColumns();
-	//BindCommands();
 
 	// Create the search filters: text based, type based etc.
 	TextFilter = MakeShared<FMemTagNodeTextFilter>(FMemTagNodeTextFilter::FItemToStringArray::CreateSP(this, &SMemTagTreeView::HandleItemToStringArray));
@@ -188,8 +189,7 @@ TSharedRef<SWidget> SMemTagTreeView::MakeTrackersMenu()
 
 void SMemTagTreeView::CreateTrackersMenuSection(FMenuBuilder& MenuBuilder)
 {
-	TSharedPtr<SMemoryProfilerWindow> ProfilerWindow = ProfilerWindowWeakPtr.Pin();
-
+	TSharedPtr<SMemoryProfilerWindow> ProfilerWindow = GetProfilerWindow();
 	if (ProfilerWindow.IsValid())
 	{
 		FMemorySharedState& SharedState = ProfilerWindow->GetSharedState();
@@ -435,7 +435,7 @@ TSharedRef<SWidget> SMemTagTreeView::ConstructTracksMiniToolbar()
 		SNew(SSegmentedControl<uint32>)
 		.OnValueChanged_Lambda([this](uint32 InValue)
 			{
-				TSharedPtr<SMemoryProfilerWindow> ProfilerWindow = ProfilerWindowWeakPtr.Pin();
+				TSharedPtr<SMemoryProfilerWindow> ProfilerWindow = GetProfilerWindow();
 				if (ProfilerWindow.IsValid())
 				{
 					FMemorySharedState& SharedState = ProfilerWindow->GetSharedState();
@@ -444,7 +444,7 @@ TSharedRef<SWidget> SMemTagTreeView::ConstructTracksMiniToolbar()
 			})
 		.Value_Lambda([this]
 			{
-				TSharedPtr<SMemoryProfilerWindow> ProfilerWindow = ProfilerWindowWeakPtr.Pin();
+				TSharedPtr<SMemoryProfilerWindow> ProfilerWindow = GetProfilerWindow();
 				if (ProfilerWindow.IsValid())
 				{
 					FMemorySharedState& SharedState = ProfilerWindow->GetSharedState();
@@ -942,7 +942,6 @@ TSharedRef<SWidget> SMemTagTreeView::TreeViewHeaderRow_GenerateColumnMenu(const
 void SMemTagTreeView::InsightsManager_OnSessionChanged()
 {
 	TSharedPtr<const TraceServices::IAnalysisSession> NewSession = FInsightsManager::Get()->GetSession();
-
 	if (NewSession != Session)
 	{
 		Session = NewSession;
@@ -988,8 +987,6 @@ void SMemTagTreeView::ApplyFiltering()
 {
 	FilteredGroupNodes.Reset();
 
-	TSharedPtr<SMemoryProfilerWindow> ProfilerWindow = ProfilerWindowWeakPtr.Pin();
-
 	// Apply filter to all groups and its children.
 	const int32 NumGroups = GroupNodes.Num();
 	for (int32 GroupIndex = 0; GroupIndex < NumGroups; ++GroupIndex)
@@ -1133,8 +1130,7 @@ void SMemTagTreeView::TreeView_OnMouseButtonDoubleClick(FMemTagNodePtr MemTagNod
 	}
 	else
 	{
-		TSharedPtr<SMemoryProfilerWindow> ProfilerWindow = ProfilerWindowWeakPtr.Pin();
-
+		TSharedPtr<SMemoryProfilerWindow> ProfilerWindow = GetProfilerWindow();
 		if (ProfilerWindow.IsValid())
 		{
 			FMemorySharedState& SharedState = ProfilerWindow->GetSharedState();
@@ -1868,10 +1864,12 @@ void SMemTagTreeView::Tick(const FGeometry& AllottedGeometry, const double InCur
 
 void SMemTagTreeView::RebuildTree(bool bResync)
 {
-	FStopwatch SyncStopwatch;
 	FStopwatch Stopwatch;
 	Stopwatch.Start();
 
+	FStopwatch SyncStopwatch;
+	SyncStopwatch.Start();
+
 	bool bListHasChanged = false;
 
 	if (bResync)
@@ -1884,9 +1882,7 @@ void SMemTagTreeView::RebuildTree(bool bResync)
 
 	const uint32 PreviousNodeCount = MemTagNodes.Num();
 
-	SyncStopwatch.Start();
-	TSharedPtr<SMemoryProfilerWindow> ProfilerWindow = ProfilerWindowWeakPtr.Pin();
-
+	TSharedPtr<SMemoryProfilerWindow> ProfilerWindow = GetProfilerWindow();
 	if (ProfilerWindow.IsValid())
 	{
 		FMemorySharedState& SharedState = ProfilerWindow->GetSharedState();
@@ -1925,6 +1921,7 @@ void SMemTagTreeView::RebuildTree(bool bResync)
 			}
 		}
 	}
+
 	SyncStopwatch.Stop();
 
 	if (bListHasChanged)
@@ -1960,7 +1957,7 @@ void SMemTagTreeView::RebuildTree(bool bResync)
 	if (TotalTime > 0.01)
 	{
 		const double SyncTime = SyncStopwatch.GetAccumulatedTime();
-		UE_LOG(MemoryProfiler, Log, TEXT("[LLM Tags] Tree view rebuilt in %.3fs (%.3fs + %.3fs) --> %d LLM tags (%d added)"),
+		UE_LOG(MemoryProfiler, Log, TEXT("[LLM Tags] Tree view rebuilt in %.4fs (sync: %.4fs + update: %.4fs) --> %d LLM tags (%d added)"),
 			TotalTime, SyncTime, TotalTime - SyncTime, MemTagNodes.Num(), MemTagNodes.Num() - PreviousNodeCount);
 	}
 }
@@ -2075,8 +2072,7 @@ void SMemTagTreeView::SelectMemTagNode(Insights::FMemoryTagId MemTagId)
 
 FReply SMemTagTreeView::LoadReportXML_OnClicked()
 {
-	TSharedPtr<SMemoryProfilerWindow> ProfilerWindow = ProfilerWindowWeakPtr.Pin();
-
+	TSharedPtr<SMemoryProfilerWindow> ProfilerWindow = GetProfilerWindow();
 	if (ProfilerWindow.IsValid())
 	{
 		TArray<FString> OutFiles;
@@ -2137,8 +2133,7 @@ FReply SMemTagTreeView::HideAllTracks_OnClicked()
 
 bool SMemTagTreeView::CanCreateGraphTracksForSelectedMemTags() const
 {
-	TSharedPtr<SMemoryProfilerWindow> ProfilerWindow = ProfilerWindowWeakPtr.Pin();
-
+	TSharedPtr<SMemoryProfilerWindow> ProfilerWindow = GetProfilerWindow();
 	return ProfilerWindow.IsValid() && TreeView->GetNumItemsSelected() > 0;
 }
 
@@ -2146,8 +2141,7 @@ bool SMemTagTreeView::CanCreateGraphTracksForSelectedMemTags() const
 
 void SMemTagTreeView::CreateGraphTracksForSelectedMemTags()
 {
-	TSharedPtr<SMemoryProfilerWindow> ProfilerWindow = ProfilerWindowWeakPtr.Pin();
-
+	TSharedPtr<SMemoryProfilerWindow> ProfilerWindow = GetProfilerWindow();
 	if (ProfilerWindow.IsValid())
 	{
 		FMemorySharedState& SharedState = ProfilerWindow->GetSharedState();
@@ -2181,8 +2175,7 @@ void SMemTagTreeView::CreateGraphTracksForSelectedMemTags()
 
 bool SMemTagTreeView::CanCreateGraphTracksForFilteredMemTags() const
 {
-	TSharedPtr<SMemoryProfilerWindow> ProfilerWindow = ProfilerWindowWeakPtr.Pin();
-
+	TSharedPtr<SMemoryProfilerWindow> ProfilerWindow = GetProfilerWindow();
 	if (!ProfilerWindow.IsValid())
 	{
 		return false;
@@ -2199,8 +2192,7 @@ bool SMemTagTreeView::CanCreateGraphTracksForFilteredMemTags() const
 
 void SMemTagTreeView::CreateGraphTracksForFilteredMemTags()
 {
-	TSharedPtr<SMemoryProfilerWindow> ProfilerWindow = ProfilerWindowWeakPtr.Pin();
-
+	TSharedPtr<SMemoryProfilerWindow> ProfilerWindow = GetProfilerWindow();
 	if (ProfilerWindow.IsValid())
 	{
 		FMemorySharedState& SharedState = ProfilerWindow->GetSharedState();
@@ -2224,8 +2216,7 @@ void SMemTagTreeView::CreateGraphTracksForFilteredMemTags()
 
 bool SMemTagTreeView::CanCreateAllGraphTracks() const
 {
-	TSharedPtr<SMemoryProfilerWindow> ProfilerWindow = ProfilerWindowWeakPtr.Pin();
-
+	TSharedPtr<SMemoryProfilerWindow> ProfilerWindow = GetProfilerWindow();
 	return ProfilerWindow.IsValid();
 }
 
@@ -2233,8 +2224,7 @@ bool SMemTagTreeView::CanCreateAllGraphTracks() const
 
 void SMemTagTreeView::CreateAllGraphTracks()
 {
-	TSharedPtr<SMemoryProfilerWindow> ProfilerWindow = ProfilerWindowWeakPtr.Pin();
-
+	TSharedPtr<SMemoryProfilerWindow> ProfilerWindow = GetProfilerWindow();
 	if (ProfilerWindow.IsValid())
 	{
 		FMemorySharedState& SharedState = ProfilerWindow->GetSharedState();
@@ -2253,8 +2243,7 @@ void SMemTagTreeView::CreateAllGraphTracks()
 
 bool SMemTagTreeView::CanRemoveGraphTracksForSelectedMemTags() const
 {
-	TSharedPtr<SMemoryProfilerWindow> ProfilerWindow = ProfilerWindowWeakPtr.Pin();
-
+	TSharedPtr<SMemoryProfilerWindow> ProfilerWindow = GetProfilerWindow();
 	return ProfilerWindow.IsValid() && TreeView->GetNumItemsSelected() > 0;
 }
 
@@ -2262,8 +2251,7 @@ bool SMemTagTreeView::CanRemoveGraphTracksForSelectedMemTags() const
 
 void SMemTagTreeView::RemoveGraphTracksForSelectedMemTags()
 {
-	TSharedPtr<SMemoryProfilerWindow> ProfilerWindow = ProfilerWindowWeakPtr.Pin();
-
+	TSharedPtr<SMemoryProfilerWindow> ProfilerWindow = GetProfilerWindow();
 	if (ProfilerWindow.IsValid())
 	{
 		FMemorySharedState& SharedState = ProfilerWindow->GetSharedState();
@@ -2297,8 +2285,7 @@ void SMemTagTreeView::RemoveGraphTracksForSelectedMemTags()
 
 bool SMemTagTreeView::CanRemoveAllGraphTracks() const
 {
-	TSharedPtr<SMemoryProfilerWindow> ProfilerWindow = ProfilerWindowWeakPtr.Pin();
-
+	TSharedPtr<SMemoryProfilerWindow> ProfilerWindow = GetProfilerWindow();
 	return ProfilerWindow.IsValid();
 }
 
@@ -2306,8 +2293,7 @@ bool SMemTagTreeView::CanRemoveAllGraphTracks() const
 
 void SMemTagTreeView::RemoveAllGraphTracks()
 {
-	TSharedPtr<SMemoryProfilerWindow> ProfilerWindow = ProfilerWindowWeakPtr.Pin();
-
+	TSharedPtr<SMemoryProfilerWindow> ProfilerWindow = GetProfilerWindow();
 	if (ProfilerWindow.IsValid())
 	{
 		FMemorySharedState& SharedState = ProfilerWindow->GetSharedState();
@@ -2321,8 +2307,7 @@ void SMemTagTreeView::RemoveAllGraphTracks()
 
 bool SMemTagTreeView::CanGenerateColorForSelectedMemTags() const
 {
-	TSharedPtr<SMemoryProfilerWindow> ProfilerWindow = ProfilerWindowWeakPtr.Pin();
-
+	TSharedPtr<SMemoryProfilerWindow> ProfilerWindow = GetProfilerWindow();
 	return ProfilerWindow.IsValid() && TreeView->GetNumItemsSelected() > 0;
 }
 
@@ -2330,8 +2315,7 @@ bool SMemTagTreeView::CanGenerateColorForSelectedMemTags() const
 
 void SMemTagTreeView::GenerateColorForSelectedMemTags()
 {
-	TSharedPtr<SMemoryProfilerWindow> ProfilerWindow = ProfilerWindowWeakPtr.Pin();
-
+	TSharedPtr<SMemoryProfilerWindow> ProfilerWindow = GetProfilerWindow();
 	if (ProfilerWindow.IsValid())
 	{
 		const TArray<FMemTagNodePtr> SelectedNodes = TreeView->GetSelectedItems();
@@ -2378,11 +2362,14 @@ void SMemTagTreeView::SetColorToNode(const FMemTagNodePtr& MemTagNode, FLinearCo
 
 	const Insights::FMemoryTagId MemTagId = MemTagNode->GetMemTagId();
 
-	TSharedPtr<SMemoryProfilerWindow> ProfilerWindow = ProfilerWindowWeakPtr.Pin();
-
-	FMemorySharedState& SharedState = ProfilerWindow->GetSharedState();
+	TSharedPtr<FMemoryGraphTrack> MainGraphTrack;
+	TSharedPtr<SMemoryProfilerWindow> ProfilerWindow = GetProfilerWindow();
+	if (ProfilerWindow.IsValid())
+	{
+		FMemorySharedState& SharedState = ProfilerWindow->GetSharedState();
+		MainGraphTrack = SharedState.GetMainGraphTrack();
+	}
 
-	TSharedPtr<FMemoryGraphTrack> MainGraphTrack = SharedState.GetMainGraphTrack();
 	for (const TSharedPtr<FMemoryGraphTrack>& GraphTrack : MemTag->GetGraphTracks())
 	{
 		for (TSharedPtr<FGraphSeries>& Series : GraphTrack->GetSeries())
@@ -2410,8 +2397,7 @@ void SMemTagTreeView::SetColorToNode(const FMemTagNodePtr& MemTagNode, FLinearCo
 
 bool SMemTagTreeView::CanEditColorForSelectedMemTags() const
 {
-	TSharedPtr<SMemoryProfilerWindow> ProfilerWindow = ProfilerWindowWeakPtr.Pin();
-
+	TSharedPtr<SMemoryProfilerWindow> ProfilerWindow = GetProfilerWindow();
 	return ProfilerWindow.IsValid() && TreeView->GetNumItemsSelected() > 0;
 }
 
@@ -2419,8 +2405,7 @@ bool SMemTagTreeView::CanEditColorForSelectedMemTags() const
 
 void SMemTagTreeView::EditColorForSelectedMemTags()
 {
-	TSharedPtr<SMemoryProfilerWindow> ProfilerWindow = ProfilerWindowWeakPtr.Pin();
-
+	TSharedPtr<SMemoryProfilerWindow> ProfilerWindow = GetProfilerWindow();
 	if (ProfilerWindow.IsValid())
 	{
 		EditableColorValue = FLinearColor(0.5f, 0.5f, 0.5f, 1.0f);
diff --git a/Engine/Source/Developer/TraceInsights/Private/Insights/MemoryProfiler/Widgets/SMemTagTreeView.h b/Engine/Source/Developer/TraceInsights/Private/Insights/MemoryProfiler/Widgets/SMemTagTreeView.h
--- a/Engine/Source/Developer/TraceInsights/Private/Insights/MemoryProfiler/Widgets/SMemTagTreeView.h
+++ b/Engine/Source/Developer/TraceInsights/Private/Insights/MemoryProfiler/Widgets/SMemTagTreeView.h
@@ -61,6 +61,15 @@ public:
 	SLATE_BEGIN_ARGS(SMemTagTreeView) {}
 	SLATE_END_ARGS()
 
+	/**
+	 * Converts profiler window weak pointer to a shared pointer and returns it.
+	 * Make sure the returned pointer is valid before trying to dereference it.
+	 */
+	TSharedPtr<SMemoryProfilerWindow> GetProfilerWindow() const
+	{
+		return ProfilerWindowWeakPtr.Pin();
+	}
+
 	/**
 	 * Construct this widget
 	 * @param InArgs - The declaration data for this widget
@@ -300,13 +309,14 @@ private:
 	virtual void Tick(const FGeometry& AllottedGeometry, const double InCurrentTime, const float InDeltaTime) override;
 
 private:
+	/** A weak pointer to the Memory Insights window. */
 	TWeakPtr<SMemoryProfilerWindow> ProfilerWindowWeakPtr;
 
 	/** Table view model. */
 	TSharedPtr<Insights::FTable> Table;
 
-	/** A weak pointer to the profiler session used to populate this widget. */
-	TSharedPtr<const TraceServices::IAnalysisSession>/*Weak*/ Session;
+	/** The analysis session used to populate this widget. */
+	TSharedPtr<const TraceServices::IAnalysisSession> Session;
 
 	//////////////////////////////////////////////////
 	// Tree View, Columns
diff --git a/Engine/Source/Developer/TraceInsights/Private/Insights/MemoryProfiler/Widgets/SMemoryProfilerWindow.cpp b/Engine/Source/Developer/TraceInsights/Private/Insights/MemoryProfiler/Widgets/SMemoryProfilerWindow.cpp
--- a/Engine/Source/Developer/TraceInsights/Private/Insights/MemoryProfiler/Widgets/SMemoryProfilerWindow.cpp
+++ b/Engine/Source/Developer/TraceInsights/Private/Insights/MemoryProfiler/Widgets/SMemoryProfilerWindow.cpp
@@ -64,6 +64,12 @@ SMemoryProfilerWindow::SMemoryProfilerWindow()
 
 SMemoryProfilerWindow::~SMemoryProfilerWindow()
 {
+	if (ModulesView)
+	{
+		HideTab(FMemoryProfilerTabs::ModulesViewID);
+		check(ModulesView == nullptr);
+	}
+
 	if (MemTagTreeView)
 	{
 		HideTab(FMemoryProfilerTabs::MemTagTreeViewID);
@@ -110,66 +116,9 @@ void SMemoryProfilerWindow::Reset()
 		MemTagTreeView->Reset();
 	}
 
-	UpdateMemInvestigationView();
-	UpdateTableTreeViews();
-}
-
-////////////////////////////////////////////////////////////////////////////////////////////////////
-
-void SMemoryProfilerWindow::UpdateMemInvestigationView()
-{
-	if (MemInvestigationView)
-	{
-		//MemInvestigationView->Update();
-	}
-}
-
-////////////////////////////////////////////////////////////////////////////////////////////////////
-
-void SMemoryProfilerWindow::UpdateTableTreeViews()
-{
-	UpdateMemTagTreeView();
-
-	for (TSharedPtr<Insights::SMemAllocTableTreeView>& MemAllocTableTreeView : MemAllocTableTreeViews)
-	{
-		UpdateMemAllocTableTreeView(MemAllocTableTreeView);
-	}
-}
-
-////////////////////////////////////////////////////////////////////////////////////////////////////
-
-void SMemoryProfilerWindow::UpdateMemTagTreeView()
-{
-	if (MemTagTreeView)
-	{
-		/*
-		TSharedPtr<const TraceServices::IAnalysisSession> Session = FInsightsManager::Get()->GetSession();
-		if (Session.IsValid() && TraceServices::ReadMemoryProfilerProvider(*Session.Get()))
-		{
-			TraceServices::FAnalysisSessionReadScope SessionReadScope(*Session.Get());
-			const TraceServices::IMemoryProfilerProvider& MemoryProfilerProvider = *TraceServices::ReadMemoryProfilerProvider(*Session.Get());
-
-			const double SelectionStartTime = TimingView ? TimingView->GetSelectionStartTime() : 0.0;
-			const double SelectionEndTime = TimingView ? TimingView->GetSelectionEndTime() : 0.0;
-
-			TraceServices::ITable<TraceServices::FMemoryProfilerAggregatedStats>* EventAggregationTable = MemoryProfilerProvider.CreateEventAggregation(SelectionStartTime, SelectionEndTime);
-			MemTagTreeView->UpdateSourceTable(MakeShareable(EventAggregationTable));
-		}
-		else
-		{
-			MemTagTreeView->UpdateSourceTable(nullptr);
-		}
-		*/
-	}
-}
-
-////////////////////////////////////////////////////////////////////////////////////////////////////
-
-void SMemoryProfilerWindow::UpdateMemAllocTableTreeView(TSharedPtr<Insights::SMemAllocTableTreeView> MemAllocTableTreeView)
-{
-	if (MemAllocTableTreeView)
+	if (ModulesView)
 	{
-		//TODO
+		ModulesView->Reset();
 	}
 }
 
@@ -234,8 +183,6 @@ TSharedRef<SDockTab> SMemoryProfilerWindow::SpawnTab_MemInvestigationView(const
 			SAssignNew(MemInvestigationView, SMemInvestigationView, SharedThis(this))
 		];
 
-	UpdateMemInvestigationView();
-
 	DockTab->SetOnTabClosed(SDockTab::FOnTabClosedCallback::CreateRaw(this, &SMemoryProfilerWindow::OnMemInvestigationViewTabClosed));
 
 	return DockTab;
@@ -264,8 +211,6 @@ TSharedRef<SDockTab> SMemoryProfilerWindow::SpawnTab_MemTagTreeView(const FSpawn
 			SAssignNew(MemTagTreeView, SMemTagTreeView, SharedThis(this))
 		];
 
-	UpdateMemTagTreeView();
-
 	DockTab->SetOnTabClosed(SDockTab::FOnTabClosedCallback::CreateRaw(this, &SMemoryProfilerWindow::OnMemTagTreeViewTabClosed));
 
 	return DockTab;
@@ -303,7 +248,6 @@ TSharedRef<SDockTab> SMemoryProfilerWindow::SpawnTab_MemAllocTableTreeView(const
 	MemAllocTableTreeView->SetLogListingName(FMemoryProfilerManager::Get()->GetLogListingName());
 	MemAllocTableTreeView->SetTabIndex(TabIndex);
 	MemAllocTableTreeViews.Add(MemAllocTableTreeView);
-	UpdateMemAllocTableTreeView(MemAllocTableTreeView);
 
 	DockTab->SetOnTabClosed(SDockTab::FOnTabClosedCallback::CreateRaw(this, &SMemoryProfilerWindow::OnMemAllocTableTreeViewTabClosed));
 
@@ -742,8 +686,6 @@ void SMemoryProfilerWindow::OnTimeSelectionChanged(Insights::ETimeChangedFlags I
 					Time = TimeMarkerTime;
 				}
 			}
-
-			UpdateMemInvestigationView();
 		}
 	}
 }
diff --git a/Engine/Source/Developer/TraceInsights/Private/Insights/MemoryProfiler/Widgets/SMemoryProfilerWindow.h b/Engine/Source/Developer/TraceInsights/Private/Insights/MemoryProfiler/Widgets/SMemoryProfilerWindow.h
--- a/Engine/Source/Developer/TraceInsights/Private/Insights/MemoryProfiler/Widgets/SMemoryProfilerWindow.h
+++ b/Engine/Source/Developer/TraceInsights/Private/Insights/MemoryProfiler/Widgets/SMemoryProfilerWindow.h
@@ -63,10 +63,6 @@ public:
 	SLATE_END_ARGS()
 
 	void Reset();
-	void UpdateTableTreeViews();
-	void UpdateMemInvestigationView();
-	void UpdateMemTagTreeView();
-	void UpdateMemAllocTableTreeView(TSharedPtr<Insights::SMemAllocTableTreeView> MemAllocTableTreeView);
 
 	/** Constructs this widget. */
 	void Construct(const FArguments& InArgs, const TSharedRef<SDockTab>& ConstructUnderMajorTab, const TSharedPtr<SWindow>& ConstructUnderWindow);
diff --git a/Engine/Source/Developer/TraceInsights/Private/Insights/NetworkingProfiler/NetworkingProfilerManager.h b/Engine/Source/Developer/TraceInsights/Private/Insights/NetworkingProfiler/NetworkingProfilerManager.h
--- a/Engine/Source/Developer/TraceInsights/Private/Insights/NetworkingProfiler/NetworkingProfilerManager.h
+++ b/Engine/Source/Developer/TraceInsights/Private/Insights/NetworkingProfiler/NetworkingProfilerManager.h
@@ -53,7 +53,7 @@ public:
 
 	//////////////////////////////////////////////////
 
-	/** @returns UI command list for the Networking Profiler manager. */
+	/** @return UI command list for the Networking Profiler manager. */
 	const TSharedRef<FUICommandList> GetCommandList() const;
 
 	/** @return an instance of the Networking Profiler commands. */
@@ -62,21 +62,21 @@ public:
 	/** @return an instance of the Networking Profiler action manager. */
 	static FNetworkingProfilerActionManager& GetActionManager();
 
-	void AddProfilerWindow(const TSharedRef<SNetworkingProfilerWindow>& InProfilerWindow)
-	{
-		ProfilerWindows.Add(InProfilerWindow);
-	}
-
-	void RemoveProfilerWindow(const TSharedRef<SNetworkingProfilerWindow>& InProfilerWindow)
+	/** @return the number of the "Networking Insights" windows currently available. */
+	int32 GetNumProfilerWindows() const
 	{
-		ProfilerWindows.Remove(InProfilerWindow);
+		return ProfilerWindows.Num();
 	}
 
 	/**
 	 * Converts profiler window weak pointer to a shared pointer and returns it.
 	 * Make sure the returned pointer is valid before trying to dereference it.
 	 */
-	TSharedPtr<class SNetworkingProfilerWindow> GetProfilerWindow(int32 Index) const
+	TSharedPtr<SNetworkingProfilerWindow> GetProfilerWindow(int32 Index) const
+	{
+		return (Index >= 0 && Index < ProfilerWindows.Num()) ? ProfilerWindows[Index].Pin() : nullptr;
+	}
+	TSharedPtr<SNetworkingProfilerWindow> GetProfilerWindowChecked(int32 Index) const
 	{
 		return ProfilerWindows[Index].Pin();
 	}
@@ -98,6 +98,16 @@ private:
 	/** Updates this manager, done through FCoreTicker. */
 	bool Tick(float DeltaTime);
 
+	void AddProfilerWindow(const TSharedRef<SNetworkingProfilerWindow>& InProfilerWindow)
+	{
+		ProfilerWindows.Add(InProfilerWindow);
+	}
+
+	void RemoveProfilerWindow(const TSharedRef<SNetworkingProfilerWindow>& InProfilerWindow)
+	{
+		ProfilerWindows.Remove(InProfilerWindow);
+	}
+
 private:
 	bool bIsInitialized;
 	bool bIsAvailable;
diff --git a/Engine/Source/Developer/TraceInsights/Private/Insights/NetworkingProfiler/Widgets/SNetStatsView.cpp b/Engine/Source/Developer/TraceInsights/Private/Insights/NetworkingProfiler/Widgets/SNetStatsView.cpp
--- a/Engine/Source/Developer/TraceInsights/Private/Insights/NetworkingProfiler/Widgets/SNetStatsView.cpp
+++ b/Engine/Source/Developer/TraceInsights/Private/Insights/NetworkingProfiler/Widgets/SNetStatsView.cpp
@@ -32,7 +32,7 @@
 ////////////////////////////////////////////////////////////////////////////////////////////////////
 
 SNetStatsView::SNetStatsView()
-	: ProfilerWindow()
+	: ProfilerWindowWeakPtr()
 	, Table(MakeShared<Insights::FTable>())
 	, bExpansionSaved(false)
 	, bFilterOutZeroCountEvents(false)
@@ -63,6 +63,8 @@ SNetStatsView::~SNetStatsView()
 	{
 		FInsightsManager::Get()->GetSessionChangedEvent().RemoveAll(this);
 	}
+
+	Session.Reset();
 }
 
 ////////////////////////////////////////////////////////////////////////////////////////////////////
@@ -70,7 +72,7 @@ SNetStatsView::~SNetStatsView()
 BEGIN_SLATE_FUNCTION_BUILD_OPTIMIZATION
 void SNetStatsView::Construct(const FArguments& InArgs, TSharedPtr<SNetworkingProfilerWindow> InProfilerWindow)
 {
-	ProfilerWindow = InProfilerWindow;
+	ProfilerWindowWeakPtr = InProfilerWindow;
 
 	SAssignNew(ExternalScrollbar, SScrollBar)
 	.AlwaysShowScrollbar(true);
@@ -206,7 +208,6 @@ void SNetStatsView::Construct(const FArguments& InArgs, TSharedPtr<SNetworkingPr
 	];
 
 	InitializeAndShowHeaderColumns();
-	//BindCommands();
 
 	// Create the search filters: text based, type based etc.
 	TextFilter = MakeShared<FNetEventNodeTextFilter>(FNetEventNodeTextFilter::FItemToStringArray::CreateSP(this, &SNetStatsView::HandleItemToStringArray));
@@ -549,7 +550,6 @@ TSharedRef<SWidget> SNetStatsView::TreeViewHeaderRow_GenerateColumnMenu(const In
 void SNetStatsView::InsightsManager_OnSessionChanged()
 {
 	TSharedPtr<const TraceServices::IAnalysisSession> NewSession = FInsightsManager::Get()->GetSession();
-
 	if (NewSession != Session)
 	{
 		Session = NewSession;
@@ -784,6 +784,7 @@ void SNetStatsView::TreeView_OnMouseButtonDoubleClick(FNetEventNodePtr NetEventN
 	}
 	else
 	{
+		TSharedPtr<SNetworkingProfilerWindow> ProfilerWindow = GetProfilerWindow();
 		TSharedPtr<SPacketContentView> PacketContentView = ProfilerWindow.IsValid() ? ProfilerWindow->GetPacketContentView() : nullptr;
 		if (PacketContentView.IsValid())
 		{
@@ -1480,10 +1481,12 @@ void SNetStatsView::Tick(const FGeometry& AllottedGeometry, const double InCurre
 
 void SNetStatsView::RebuildTree(bool bResync)
 {
-	FStopwatch SyncStopwatch;
 	FStopwatch Stopwatch;
 	Stopwatch.Start();
 
+	FStopwatch SyncStopwatch;
+	SyncStopwatch.Start();
+
 	if (bResync)
 	{
 		NetEventNodes.Empty();
@@ -1491,7 +1494,6 @@ void SNetStatsView::RebuildTree(bool bResync)
 
 	const uint32 PreviousNodeCount = NetEventNodes.Num();
 
-	SyncStopwatch.Start();
 	if (Session.IsValid())
 	{
 		TraceServices::FAnalysisSessionReadScope SessionReadScope(*Session.Get());
@@ -1523,6 +1525,7 @@ void SNetStatsView::RebuildTree(bool bResync)
 			});
 		}
 	}
+
 	SyncStopwatch.Stop();
 
 	if (bResync || NetEventNodes.Num() != PreviousNodeCount)
@@ -1558,7 +1561,7 @@ void SNetStatsView::RebuildTree(bool bResync)
 	if (TotalTime > 0.01)
 	{
 		const double SyncTime = SyncStopwatch.GetAccumulatedTime();
-		UE_LOG(NetworkingProfiler, Log, TEXT("[NetStats] Tree view rebuilt in %.4fs (%.4fs + %.4fs) --> %d net events (%d added)"),
+		UE_LOG(NetworkingProfiler, Log, TEXT("[NetStats] Tree view rebuilt in %.4fs (sync: %.4fs + update: %.4fs) --> %d net events (%d added)"),
 			TotalTime, SyncTime, TotalTime - SyncTime, NetEventNodes.Num(), NetEventNodes.Num() - PreviousNodeCount);
 	}
 }
diff --git a/Engine/Source/Developer/TraceInsights/Private/Insights/NetworkingProfiler/Widgets/SNetStatsView.h b/Engine/Source/Developer/TraceInsights/Private/Insights/NetworkingProfiler/Widgets/SNetStatsView.h
--- a/Engine/Source/Developer/TraceInsights/Private/Insights/NetworkingProfiler/Widgets/SNetStatsView.h
+++ b/Engine/Source/Developer/TraceInsights/Private/Insights/NetworkingProfiler/Widgets/SNetStatsView.h
@@ -60,14 +60,23 @@ public:
 	SLATE_BEGIN_ARGS(SNetStatsView) {}
 	SLATE_END_ARGS()
 
+	/**
+	 * Converts profiler window weak pointer to a shared pointer and returns it.
+	 * Make sure the returned pointer is valid before trying to dereference it.
+	 */
+	TSharedPtr<SNetworkingProfilerWindow> GetProfilerWindow() const
+	{
+		return ProfilerWindowWeakPtr.Pin();
+	}
+
+	TSharedPtr<Insights::FTable> GetTable() const { return Table; }
+
 	/**
 	 * Construct this widget
 	 * @param InArgs - The declaration data for this widget
 	 */
 	void Construct(const FArguments& InArgs, TSharedPtr<SNetworkingProfilerWindow> InProfilerWindow);
 
-	TSharedPtr<Insights::FTable> GetTable() const { return Table; }
-
 	/**
 	 * Ticks this widget. Override in derived classes, but always call the parent implementation.
 	 *
@@ -251,13 +260,14 @@ protected:
 	////////////////////////////////////////////////////////////////////////////////////////////////////
 
 protected:
-	TSharedPtr<SNetworkingProfilerWindow> ProfilerWindow;
+	/** A weak pointer to the Networking Insights window. */
+	TWeakPtr<SNetworkingProfilerWindow> ProfilerWindowWeakPtr;
 
 	/** Table view model. */
 	TSharedPtr<Insights::FTable> Table;
 
-	/** A weak pointer to the profiler session used to populate this widget. */
-	TSharedPtr<const TraceServices::IAnalysisSession>/*Weak*/ Session;
+	/** The analysis session used to populate this widget. */
+	TSharedPtr<const TraceServices::IAnalysisSession> Session;
 
 	//////////////////////////////////////////////////
 	// Tree View, Columns
diff --git a/Engine/Source/Developer/TraceInsights/Private/Insights/Table/Widgets/STableTreeView.cpp b/Engine/Source/Developer/TraceInsights/Private/Insights/Table/Widgets/STableTreeView.cpp
--- a/Engine/Source/Developer/TraceInsights/Private/Insights/Table/Widgets/STableTreeView.cpp
+++ b/Engine/Source/Developer/TraceInsights/Private/Insights/Table/Widgets/STableTreeView.cpp
@@ -100,7 +100,6 @@ const FName STableTreeView::RootNodeName(TEXT("Root"));
 
 STableTreeView::STableTreeView()
 	: Table()
-	, Session(FInsightsManager::Get()->GetSession())
 	, TreeView(nullptr)
 	, TreeViewHeaderRow(nullptr)
 	, ExternalScrollbar(nullptr)
@@ -142,6 +141,27 @@ STableTreeView::~STableTreeView()
 		delete CurrentAsyncOpFilterConfigurator;
 		CurrentAsyncOpFilterConfigurator = nullptr;
 	}
+
+	//FTableTreeViewCommands::Unregister();
+
+	Session.Reset();
+}
+
+////////////////////////////////////////////////////////////////////////////////////////////////////
+
+void STableTreeView::InitCommandList()
+{
+	FTableTreeViewCommands::Register();
+	CommandList = MakeShared<FUICommandList>();
+	CommandList->MapAction(FTableTreeViewCommands::Get().Command_CopyToClipboard, FExecuteAction::CreateSP(this, &STableTreeView::ContextMenu_CopySelectedToClipboard_Execute), FCanExecuteAction::CreateSP(this, &STableTreeView::ContextMenu_CopySelectedToClipboard_CanExecute));
+	CommandList->MapAction(FTableTreeViewCommands::Get().Command_CopyColumnToClipboard, FExecuteAction::CreateSP(this, &STableTreeView::ContextMenu_CopyColumnToClipboard_Execute), FCanExecuteAction::CreateSP(this, &STableTreeView::ContextMenu_CopyColumnToClipboard_CanExecute));
+	CommandList->MapAction(FTableTreeViewCommands::Get().Command_CopyColumnTooltipToClipboard, FExecuteAction::CreateSP(this, &STableTreeView::ContextMenu_CopyColumnTooltipToClipboard_Execute), FCanExecuteAction::CreateSP(this, &STableTreeView::ContextMenu_CopyColumnTooltipToClipboard_CanExecute));
+	CommandList->MapAction(FTableTreeViewCommands::Get().Command_ExpandSubtree, FExecuteAction::CreateSP(this, &STableTreeView::ContextMenu_ExpandSubtree_Execute), FCanExecuteAction::CreateSP(this, &STableTreeView::ContextMenu_ExpandSubtree_CanExecute));
+	CommandList->MapAction(FTableTreeViewCommands::Get().Command_ExpandCriticalPath, FExecuteAction::CreateSP(this, &STableTreeView::ContextMenu_ExpandCriticalPath_Execute), FCanExecuteAction::CreateSP(this, &STableTreeView::ContextMenu_ExpandCriticalPath_CanExecute));
+	CommandList->MapAction(FTableTreeViewCommands::Get().Command_CollapseSubtree, FExecuteAction::CreateSP(this, &STableTreeView::ContextMenu_CollapseSubtree_Execute), FCanExecuteAction::CreateSP(this, &STableTreeView::ContextMenu_CollapseSubtree_CanExecute));
+	CommandList->MapAction(FTableTreeViewCommands::Get().Command_ExportToFile, FExecuteAction::CreateSP(this, &STableTreeView::ContextMenu_ExportToFile_Execute, false, true), FCanExecuteAction::CreateSP(this, &STableTreeView::ContextMenu_ExportToFile_CanExecute));
+	CommandList->MapAction(FTableTreeViewCommands::Get().Command_ExportEntireTreeToFile, FExecuteAction::CreateSP(this, &STableTreeView::ContextMenu_ExportToFile_Execute, true, true), FCanExecuteAction::CreateSP(this, &STableTreeView::ContextMenu_ExportToFile_CanExecute));
+	CommandList->MapAction(FTableTreeViewCommands::Get().Command_ExportEntireTreeToFileNoLeafs, FExecuteAction::CreateSP(this, &STableTreeView::ContextMenu_ExportToFile_Execute, true, false), FCanExecuteAction::CreateSP(this, &STableTreeView::ContextMenu_ExportToFile_CanExecute));
 }
 
 ////////////////////////////////////////////////////////////////////////////////////////////////////
@@ -249,6 +269,9 @@ void STableTreeView::ConstructWidget(TSharedPtr<FTable> InTablePtr)
 
 	// Register ourselves with the Insights manager.
 	FInsightsManager::Get()->GetSessionChangedEvent().AddSP(this, &STableTreeView::InsightsManager_OnSessionChanged);
+
+	// Update the Session (i.e. when analysis session was already started).
+	InsightsManager_OnSessionChanged();
 }
 END_SLATE_FUNCTION_BUILD_OPTIMIZATION
 
@@ -403,23 +426,6 @@ END_SLATE_FUNCTION_BUILD_OPTIMIZATION
 
 ////////////////////////////////////////////////////////////////////////////////////////////////////
 
-void STableTreeView::InitCommandList()
-{
-	FTableTreeViewCommands::Register();
-	CommandList = MakeShared<FUICommandList>();
-	CommandList->MapAction(FTableTreeViewCommands::Get().Command_CopyToClipboard, FExecuteAction::CreateSP(this, &STableTreeView::ContextMenu_CopySelectedToClipboard_Execute), FCanExecuteAction::CreateSP(this, &STableTreeView::ContextMenu_CopySelectedToClipboard_CanExecute));
-	CommandList->MapAction(FTableTreeViewCommands::Get().Command_CopyColumnToClipboard, FExecuteAction::CreateSP(this, &STableTreeView::ContextMenu_CopyColumnToClipboard_Execute), FCanExecuteAction::CreateSP(this, &STableTreeView::ContextMenu_CopyColumnToClipboard_CanExecute));
-	CommandList->MapAction(FTableTreeViewCommands::Get().Command_CopyColumnTooltipToClipboard, FExecuteAction::CreateSP(this, &STableTreeView::ContextMenu_CopyColumnTooltipToClipboard_Execute), FCanExecuteAction::CreateSP(this, &STableTreeView::ContextMenu_CopyColumnTooltipToClipboard_CanExecute));
-	CommandList->MapAction(FTableTreeViewCommands::Get().Command_ExpandSubtree, FExecuteAction::CreateSP(this, &STableTreeView::ContextMenu_ExpandSubtree_Execute), FCanExecuteAction::CreateSP(this, &STableTreeView::ContextMenu_ExpandSubtree_CanExecute));
-	CommandList->MapAction(FTableTreeViewCommands::Get().Command_ExpandCriticalPath, FExecuteAction::CreateSP(this, &STableTreeView::ContextMenu_ExpandCriticalPath_Execute), FCanExecuteAction::CreateSP(this, &STableTreeView::ContextMenu_ExpandCriticalPath_CanExecute));
-	CommandList->MapAction(FTableTreeViewCommands::Get().Command_CollapseSubtree, FExecuteAction::CreateSP(this, &STableTreeView::ContextMenu_CollapseSubtree_Execute), FCanExecuteAction::CreateSP(this, &STableTreeView::ContextMenu_CollapseSubtree_CanExecute));
-	CommandList->MapAction(FTableTreeViewCommands::Get().Command_ExportToFile, FExecuteAction::CreateSP(this, &STableTreeView::ContextMenu_ExportToFile_Execute, false, true), FCanExecuteAction::CreateSP(this, &STableTreeView::ContextMenu_ExportToFile_CanExecute));
-	CommandList->MapAction(FTableTreeViewCommands::Get().Command_ExportEntireTreeToFile, FExecuteAction::CreateSP(this, &STableTreeView::ContextMenu_ExportToFile_Execute, true, true), FCanExecuteAction::CreateSP(this, &STableTreeView::ContextMenu_ExportToFile_CanExecute));
-	CommandList->MapAction(FTableTreeViewCommands::Get().Command_ExportEntireTreeToFileNoLeafs, FExecuteAction::CreateSP(this, &STableTreeView::ContextMenu_ExportToFile_Execute, true, false), FCanExecuteAction::CreateSP(this, &STableTreeView::ContextMenu_ExportToFile_CanExecute));
-}
-
-////////////////////////////////////////////////////////////////////////////////////////////////////
-
 TSharedPtr<SWidget> STableTreeView::TreeView_GetMenuContent()
 {
 	const TArray<FTableTreeNodePtr> SelectedNodes = TreeView->GetSelectedItems();
@@ -846,7 +852,6 @@ TSharedRef<SWidget> STableTreeView::TreeViewHeaderRow_GenerateColumnMenu(const F
 void STableTreeView::InsightsManager_OnSessionChanged()
 {
 	TSharedPtr<const TraceServices::IAnalysisSession> NewSession = FInsightsManager::Get()->GetSession();
-
 	if (NewSession != Session)
 	{
 		Session = NewSession;
diff --git a/Engine/Source/Developer/TraceInsights/Private/Insights/Table/Widgets/STableTreeView.h b/Engine/Source/Developer/TraceInsights/Private/Insights/Table/Widgets/STableTreeView.h
--- a/Engine/Source/Developer/TraceInsights/Private/Insights/Table/Widgets/STableTreeView.h
+++ b/Engine/Source/Developer/TraceInsights/Private/Insights/Table/Widgets/STableTreeView.h
@@ -133,8 +133,9 @@ public:
 	const FName& GetLogListingName() { return LogListingName; }
 
 protected:
-	void ConstructWidget(TSharedPtr<FTable> InTablePtr);
 	void InitCommandList();
+
+	void ConstructWidget(TSharedPtr<FTable> InTablePtr);
 	virtual TSharedRef<SWidget> ConstructSearchBox();
 	virtual TSharedRef<SWidget> ConstructAdvancedFiltersButton();
 	virtual TSharedRef<SWidget> ConstructHierarchyBreadcrumbTrail();
@@ -407,8 +408,8 @@ protected:
 	/** Table view model. */
 	TSharedPtr<Insights::FTable> Table;
 
-	/** A weak pointer to the profiler session used to populate this widget. */
-	TSharedPtr<const TraceServices::IAnalysisSession>/*Weak*/ Session;
+	/** The analysis session used to populate this widget. */
+	TSharedPtr<const TraceServices::IAnalysisSession> Session;
 
 	//////////////////////////////////////////////////
 	// Tree View, Columns
diff --git a/Engine/Source/Developer/TraceInsights/Private/Insights/Table/Widgets/SUntypedTableTreeView.cpp b/Engine/Source/Developer/TraceInsights/Private/Insights/Table/Widgets/SUntypedTableTreeView.cpp
--- a/Engine/Source/Developer/TraceInsights/Private/Insights/Table/Widgets/SUntypedTableTreeView.cpp
+++ b/Engine/Source/Developer/TraceInsights/Private/Insights/Table/Widgets/SUntypedTableTreeView.cpp
@@ -60,10 +60,12 @@ void SUntypedTableTreeView::UpdateSourceTable(TSharedPtr<TraceServices::IUntyped
 
 void SUntypedTableTreeView::RebuildTree(bool bResync)
 {
-	FStopwatch SyncStopwatch;
 	FStopwatch Stopwatch;
 	Stopwatch.Start();
 
+	FStopwatch SyncStopwatch;
+	SyncStopwatch.Start();
+
 	if (bResync)
 	{
 		TableTreeNodes.Empty();
@@ -77,7 +79,6 @@ void SUntypedTableTreeView::RebuildTree(bool bResync)
 	TSharedPtr<TraceServices::IUntypedTable> SourceTable = UntypedTable->GetSourceTable();
 	TSharedPtr<TraceServices::IUntypedTableReader> TableReader = UntypedTable->GetTableReader();
 
-	SyncStopwatch.Start();
 	if (Session.IsValid() && SourceTable.IsValid() && TableReader.IsValid())
 	{
 		const int32 TotalRowCount = SourceTable->GetRowCount();
@@ -96,6 +97,7 @@ void SUntypedTableTreeView::RebuildTree(bool bResync)
 			ensure(TableTreeNodes.Num() == TotalRowCount);
 		}
 	}
+
 	SyncStopwatch.Stop();
 
 	if (bResync || TableTreeNodes.Num() != PreviousNodeCount)
@@ -130,7 +132,7 @@ void SUntypedTableTreeView::RebuildTree(bool bResync)
 	if (TotalTime > 0.01)
 	{
 		const double SyncTime = SyncStopwatch.GetAccumulatedTime();
-		UE_LOG(TraceInsights, Log, TEXT("[Table] Tree view rebuilt in %.3fs (%.3fs + %.3fs) --> %d rows (%d added)"),
+		UE_LOG(TraceInsights, Log, TEXT("[Table] Tree view rebuilt in %.4fs (sync: %.4fs + update: %.4fs) --> %d rows (%d added)"),
 			TotalTime, SyncTime, TotalTime - SyncTime, TableTreeNodes.Num(), TableTreeNodes.Num() - PreviousNodeCount);
 	}
 }
diff --git a/Engine/Source/Developer/TraceInsights/Private/Insights/TimingProfilerManager.cpp b/Engine/Source/Developer/TraceInsights/Private/Insights/TimingProfilerManager.cpp
--- a/Engine/Source/Developer/TraceInsights/Private/Insights/TimingProfilerManager.cpp
+++ b/Engine/Source/Developer/TraceInsights/Private/Insights/TimingProfilerManager.cpp
@@ -59,7 +59,7 @@ FTimingProfilerManager::FTimingProfilerManager(TSharedRef<FUICommandList> InComm
 	, bIsAvailable(false)
 	, CommandList(InCommandList)
 	, ActionManager(this)
-	, ProfilerWindow(nullptr)
+	, ProfilerWindowWeakPtr()
 	, bIsFramesTrackVisible(false)
 	, bIsTimingViewVisible(false)
 	, bIsTimersViewVisible(false)
diff --git a/Engine/Source/Developer/TraceInsights/Private/Insights/TimingProfilerManager.h b/Engine/Source/Developer/TraceInsights/Private/Insights/TimingProfilerManager.h
--- a/Engine/Source/Developer/TraceInsights/Private/Insights/TimingProfilerManager.h
+++ b/Engine/Source/Developer/TraceInsights/Private/Insights/TimingProfilerManager.h
@@ -69,7 +69,7 @@ public:
 
 	//////////////////////////////////////////////////
 
-	/** @returns UI command list for the Timing Profiler manager. */
+	/** @return UI command list for the Timing Profiler manager. */
 	const TSharedRef<FUICommandList> GetCommandList() const;
 
 	/** @return an instance of the Timing Profiler commands. */
@@ -78,23 +78,13 @@ public:
 	/** @return an instance of the Timing Profiler action manager. */
 	static FTimingProfilerActionManager& GetActionManager();
 
-	void AssignProfilerWindow(const TSharedRef<STimingProfilerWindow>& InProfilerWindow)
-	{
-		ProfilerWindow = InProfilerWindow;
-	}
-
-	void RemoveProfilerWindow()
-	{
-		ProfilerWindow.Reset();
-	}
-
 	/**
 	 * Converts profiler window weak pointer to a shared pointer and returns it.
 	 * Make sure the returned pointer is valid before trying to dereference it.
 	 */
-	TSharedPtr<class STimingProfilerWindow> GetProfilerWindow() const
+	TSharedPtr<STimingProfilerWindow> GetProfilerWindow() const
 	{
-		return ProfilerWindow.Pin();
+		return ProfilerWindowWeakPtr.Pin();
 	}
 
 	////////////////////////////////////////////////////////////////////////////////////////////////////
@@ -181,6 +171,16 @@ private:
 
 	void FinishTimerButterflyAggregation();
 
+	void AssignProfilerWindow(const TSharedRef<STimingProfilerWindow>& InProfilerWindow)
+	{
+		ProfilerWindowWeakPtr = InProfilerWindow;
+	}
+
+	void RemoveProfilerWindow()
+	{
+		ProfilerWindowWeakPtr.Reset();
+	}
+
 private:
 	bool bIsInitialized;
 	bool bIsAvailable;
@@ -198,8 +198,8 @@ private:
 	/** An instance of the Timing Profiler action manager. */
 	FTimingProfilerActionManager ActionManager;
 
-	/** A weak pointer to the Timing Profiler window. */
-	TWeakPtr<class STimingProfilerWindow> ProfilerWindow;
+	/** A weak pointer to the Timing Insights window. */
+	TWeakPtr<STimingProfilerWindow> ProfilerWindowWeakPtr;
 
 	/** If the Frames Track is visible or hidden. */
 	bool bIsFramesTrackVisible;
diff --git a/Engine/Source/Developer/TraceInsights/Private/Insights/TraceInsightsModule.cpp b/Engine/Source/Developer/TraceInsights/Private/Insights/TraceInsightsModule.cpp
--- a/Engine/Source/Developer/TraceInsights/Private/Insights/TraceInsightsModule.cpp
+++ b/Engine/Source/Developer/TraceInsights/Private/Insights/TraceInsightsModule.cpp
@@ -73,6 +73,7 @@ void FTraceInsightsModule::StartupModule()
 
 ////////////////////////////////////////////////////////////////////////////////////////////////////
 
+PRAGMA_DISABLE_OPTIMIZATION
 void FTraceInsightsModule::ShutdownModule()
 {
 #if !WITH_EDITOR
@@ -92,6 +93,15 @@ void FTraceInsightsModule::ShutdownModule()
 
 	UnregisterTabSpawners();
 
+#define INSIGHTS_CHECK_SHARED_REFERENCES 1
+#if INSIGHTS_CHECK_SHARED_REFERENCES
+	TSharedPtr<const TraceServices::IAnalysisSession> Session = FInsightsManager::Get()->GetSession();
+	auto TimingInsightsWindow = FTimingProfilerManager::Get()->GetProfilerWindow();
+	auto AssetLoadingInsightsWindow = FLoadingProfilerManager::Get()->GetProfilerWindow();
+	auto NetworkingInsightsWindow0 = FNetworkingProfilerManager::Get()->GetProfilerWindow(0);
+	auto MemoryInsightsWindow = FMemoryProfilerManager::Get()->GetProfilerWindow();
+#endif
+
 	// Unregister components. Shutdown in the reverse order they were registered.
 	for (int32 ComponentIndex = Components.Num() - 1; ComponentIndex >= 0; --ComponentIndex)
 	{
@@ -99,8 +109,36 @@ void FTraceInsightsModule::ShutdownModule()
 	}
 	Components.Reset();
 
+#if INSIGHTS_CHECK_SHARED_REFERENCES
+	ensure(MemoryInsightsWindow.GetSharedReferenceCount() <= 1);
+	MemoryInsightsWindow.Reset();
+
+	ensure(NetworkingInsightsWindow0.GetSharedReferenceCount() <= 1);
+	NetworkingInsightsWindow0.Reset();
+
+	ensure(AssetLoadingInsightsWindow.GetSharedReferenceCount() <= 1);
+	AssetLoadingInsightsWindow.Reset();
+
+	ensure(TimingInsightsWindow.GetSharedReferenceCount() <= 1);
+	TimingInsightsWindow.Reset();
+
+	if (ensure(Session.GetSharedReferenceCount() <= 1))
+	{
+		Session.Reset();
+	}
+	else // Some component(s) failed to release the references to Session shared ptr!
+	{
+		UE_LOG(TraceInsights, Warning, TEXT("The analysis Session is still referenced! Force delete!"));
+		const TraceServices::IAnalysisSession* SessionPtr = Session.Get();
+		Session.Reset();
+		delete SessionPtr;
+	}
+#endif
+#undef INSIGHTS_CHECK_SHARED_REFERENCES
+
 	FInsightsStyle::Shutdown();
 }
+PRAGMA_ENABLE_OPTIMIZATION
 
 ////////////////////////////////////////////////////////////////////////////////////////////////////
 
diff --git a/Engine/Source/Developer/TraceInsights/Private/Insights/Widgets/SStatsView.cpp b/Engine/Source/Developer/TraceInsights/Private/Insights/Widgets/SStatsView.cpp
--- a/Engine/Source/Developer/TraceInsights/Private/Insights/Widgets/SStatsView.cpp
+++ b/Engine/Source/Developer/TraceInsights/Private/Insights/Widgets/SStatsView.cpp
@@ -106,6 +106,8 @@ SStatsView::~SStatsView()
 	}
 
 	FStatsViewCommands::Unregister();
+
+	Session.Reset();
 }
 
 ////////////////////////////////////////////////////////////////////////////////////////////////////
@@ -308,7 +310,6 @@ void SStatsView::Construct(const FArguments& InArgs)
 	];
 
 	InitializeAndShowHeaderColumns();
-	//BindCommands();
 
 	// Create the search filters: text based, type based etc.
 	TextFilter = MakeShared<FStatsNodeTextFilter>(FStatsNodeTextFilter::FItemToStringArray::CreateSP(this, &SStatsView::HandleItemToStringArray));
@@ -712,7 +713,6 @@ TSharedRef<SWidget> SStatsView::TreeViewHeaderRow_GenerateColumnMenu(const Insig
 void SStatsView::InsightsManager_OnSessionChanged()
 {
 	TSharedPtr<const TraceServices::IAnalysisSession> NewSession = FInsightsManager::Get()->GetSession();
-
 	if (NewSession != Session)
 	{
 		Session = NewSession;
@@ -1825,10 +1825,12 @@ void SStatsView::Tick(const FGeometry& AllottedGeometry, const double InCurrentT
 
 void SStatsView::RebuildTree(bool bResync)
 {
-	FStopwatch SyncStopwatch;
 	FStopwatch Stopwatch;
 	Stopwatch.Start();
 
+	FStopwatch SyncStopwatch;
+	SyncStopwatch.Start();
+
 	if (bResync)
 	{
 		StatsNodes.Empty();
@@ -1837,7 +1839,6 @@ void SStatsView::RebuildTree(bool bResync)
 
 	const uint32 PreviousNodeCount = StatsNodes.Num();
 
-	SyncStopwatch.Start();
 	if (Session.IsValid())
 	{
 		TraceServices::FAnalysisSessionReadScope SessionReadScope(*Session.Get());
@@ -1876,6 +1877,7 @@ void SStatsView::RebuildTree(bool bResync)
 			ensure(StatsNodes.Num() == CounterCount);
 		}
 	}
+
 	SyncStopwatch.Stop();
 
 	if (bResync || StatsNodes.Num() != PreviousNodeCount)
@@ -1920,7 +1922,7 @@ void SStatsView::RebuildTree(bool bResync)
 	if (TotalTime > 0.01)
 	{
 		const double SyncTime = SyncStopwatch.GetAccumulatedTime();
-		UE_LOG(TimingProfiler, Log, TEXT("[Counters] Tree view rebuilt in %.4fs (%.4fs + %.4fs) --> %d counters (%d added)"),
+		UE_LOG(TimingProfiler, Log, TEXT("[Counters] Tree view rebuilt in %.4fs (sync: %.4fs + update: %.4fs) --> %d counters (%d added)"),
 			TotalTime, SyncTime, TotalTime - SyncTime, StatsNodes.Num(), StatsNodes.Num() - PreviousNodeCount);
 	}
 }
diff --git a/Engine/Source/Developer/TraceInsights/Private/Insights/Widgets/SStatsView.h b/Engine/Source/Developer/TraceInsights/Private/Insights/Widgets/SStatsView.h
--- a/Engine/Source/Developer/TraceInsights/Private/Insights/Widgets/SStatsView.h
+++ b/Engine/Source/Developer/TraceInsights/Private/Insights/Widgets/SStatsView.h
@@ -281,8 +281,8 @@ private:
 	/** Table view model. */
 	TSharedPtr<Insights::FTable> Table;
 
-	/** A weak pointer to the profiler session used to populate this widget. */
-	TSharedPtr<const TraceServices::IAnalysisSession>/*Weak*/ Session;
+	/** The analysis session used to populate this widget. */
+	TSharedPtr<const TraceServices::IAnalysisSession> Session;
 
 	TSharedPtr<FUICommandList> CommandList;
 
diff --git a/Engine/Source/Developer/TraceInsights/Private/Insights/Widgets/STimersView.cpp b/Engine/Source/Developer/TraceInsights/Private/Insights/Widgets/STimersView.cpp
--- a/Engine/Source/Developer/TraceInsights/Private/Insights/Widgets/STimersView.cpp
+++ b/Engine/Source/Developer/TraceInsights/Private/Insights/Widgets/STimersView.cpp
@@ -122,6 +122,8 @@ STimersView::~STimersView()
 	}
 
 	FTimersViewCommands::Unregister();
+
+	Session.Reset();
 }
 
 ////////////////////////////////////////////////////////////////////////////////////////////////////
@@ -339,7 +341,6 @@ void STimersView::Construct(const FArguments& InArgs)
 	];
 
 	InitializeAndShowHeaderColumns();
-	//BindCommands();
 
 	// Create the search filters: text based, type based etc.
 	TextFilter = MakeShared<FTimerNodeTextFilter>(FTimerNodeTextFilter::FItemToStringArray::CreateSP(this, &STimersView::HandleItemToStringArray));
@@ -880,7 +881,6 @@ TSharedRef<SWidget> STimersView::TreeViewHeaderRow_GenerateColumnMenu(const Insi
 void STimersView::InsightsManager_OnSessionChanged()
 {
 	TSharedPtr<const TraceServices::IAnalysisSession> NewSession = FInsightsManager::Get()->GetSession();
-
 	if (NewSession != Session)
 	{
 		Session = NewSession;
@@ -1914,10 +1914,12 @@ void STimersView::Tick(const FGeometry& AllottedGeometry, const double InCurrent
 
 void STimersView::RebuildTree(bool bResync)
 {
-	FStopwatch SyncStopwatch;
 	FStopwatch Stopwatch;
 	Stopwatch.Start();
 
+	FStopwatch SyncStopwatch;
+	SyncStopwatch.Start();
+
 	if (bResync)
 	{
 		TimerNodes.Empty();
@@ -1925,7 +1927,6 @@ void STimersView::RebuildTree(bool bResync)
 
 	const uint32 PreviousNodeCount = TimerNodes.Num();
 
-	SyncStopwatch.Start();
 	if (Session.IsValid() && TraceServices::ReadTimingProfilerProvider(*Session.Get()))
 	{
 		TraceServices::FAnalysisSessionReadScope SessionReadScope(*Session.Get());
@@ -1953,6 +1954,7 @@ void STimersView::RebuildTree(bool bResync)
 			ensure(TimerNodes.Num() == TimerCount);
 		}
 	}
+
 	SyncStopwatch.Stop();
 
 	if (bResync || TimerNodes.Num() != PreviousNodeCount)
@@ -1997,7 +1999,7 @@ void STimersView::RebuildTree(bool bResync)
 	if (TotalTime > 0.01)
 	{
 		const double SyncTime = SyncStopwatch.GetAccumulatedTime();
-		UE_LOG(TimingProfiler, Log, TEXT("[Timers] Tree view rebuilt in %.4fs (%.4fs + %.4fs) --> %d timers (%d added)"),
+		UE_LOG(TimingProfiler, Log, TEXT("[Timers] Tree view rebuilt in %.4fs (sync: %.4fs + update: %.4fs) --> %d timers (%d added)"),
 			TotalTime, SyncTime, TotalTime - SyncTime, TimerNodes.Num(), TimerNodes.Num() - PreviousNodeCount);
 	}
 }
diff --git a/Engine/Source/Developer/TraceInsights/Private/Insights/Widgets/STimersView.h b/Engine/Source/Developer/TraceInsights/Private/Insights/Widgets/STimersView.h
--- a/Engine/Source/Developer/TraceInsights/Private/Insights/Widgets/STimersView.h
+++ b/Engine/Source/Developer/TraceInsights/Private/Insights/Widgets/STimersView.h
@@ -303,8 +303,8 @@ private:
 	/** Table view model. */
 	TSharedPtr<Insights::FTable> Table;
 
-	/** A weak pointer to the profiler session used to populate this widget. */
-	TSharedPtr<const TraceServices::IAnalysisSession>/*Weak*/ Session;
+	/** The analysis session used to populate this widget. */
+	TSharedPtr<const TraceServices::IAnalysisSession> Session;
 
 	TSharedPtr<FUICommandList> CommandList;
 
