[Insights] Select a timing event when the Go To Task functionality from the Task's Table is used.

diff --git a/Engine/Source/Developer/TraceInsights/Private/Insights/TaskGraphProfiler/Widgets/STaskTableTreeView.cpp b/Engine/Source/Developer/TraceInsights/Private/Insights/TaskGraphProfiler/Widgets/STaskTableTreeView.cpp
--- a/Engine/Source/Developer/TraceInsights/Private/Insights/TaskGraphProfiler/Widgets/STaskTableTreeView.cpp
+++ b/Engine/Source/Developer/TraceInsights/Private/Insights/TaskGraphProfiler/Widgets/STaskTableTreeView.cpp
@@ -18,6 +18,7 @@
 #include "Insights/Table/ViewModels/TableColumn.h"
 #include "Insights/TimingProfilerManager.h"
 #include "Insights/ViewModels/FilterConfigurator.h"
+#include "Insights/ViewModels/ThreadTimingTrack.h"
 #include "Insights/Widgets/STimingProfilerWindow.h"
 #include "Insights/Widgets/STimingView.h"
 
@@ -425,6 +426,7 @@ void STaskTableTreeView::ContextMenu_GoToTask_Execute()
 	}
 
 	FTaskNodePtr SelectedTask = StaticCastSharedPtr<FTaskNode>(SelectedItems[0]);
+	const FTaskEntry* TaskEntry = SelectedTask->GetTask();
 
 	TSharedPtr<STimingProfilerWindow> TimingWindow = FTimingProfilerManager::Get()->GetProfilerWindow();
 	if (!TimingWindow.IsValid())
@@ -438,17 +440,45 @@ void STaskTableTreeView::ContextMenu_GoToTask_Execute()
 		return;
 	}
 
-	FTaskGraphProfilerManager::Get()->ShowTaskRelations(SelectedTask->GetTask()->GetId());
+	FTaskGraphProfilerManager::Get()->ShowTaskRelations(TaskEntry->GetId());
 
-	double Duration = (SelectedTask->GetTask()->GetFinishedTimestamp() - SelectedTask->GetTask()->GetCreatedTimestamp()) * 1.5;
-	TimingView->ZoomOnTimeInterval(SelectedTask->GetTask()->GetCreatedTimestamp() - Duration * 0.15, Duration);
+	double Duration = (TaskEntry->GetFinishedTimestamp() - TaskEntry->GetCreatedTimestamp()) * 1.5;
+	TimingView->ZoomOnTimeInterval(TaskEntry->GetCreatedTimestamp() - Duration * 0.15, Duration);
 
 	TSharedPtr<FTaskTimingSharedState> TaskSharedState = FTaskGraphProfilerManager::Get()->GetTaskTimingSharedState();
 
 	if (TaskSharedState.IsValid() && FTaskGraphProfilerManager::Get()->GetShowAnyRelations())
 	{
-		TaskSharedState->SetTaskId(SelectedTask->GetTask()->GetId());
+		TaskSharedState->SetTaskId(TaskEntry->GetId());
 	}
+
+	TSharedPtr<FThreadTimingSharedState> ThreadTimingState = TimingView->GetThreadTimingSharedState();
+	if (!ThreadTimingState.IsValid())
+	{
+		return;
+	}
+
+	TSharedPtr<FCpuTimingTrack> Track = ThreadTimingState->GetCpuTrack(TaskEntry->GetStartedThreadId());
+	if (!Track.IsValid())
+	{
+		return;
+	}
+
+	TimingView->SelectTimingTrack(Track, true);
+
+	FTimingEventSearchParameters SearchParams(TaskEntry->StartedTimestamp, TaskEntry->FinishedTimestamp, ETimingEventSearchFlags::StopAtFirstMatch, 
+		[TaskEntry](double StartTime, double EndTime, uint32 Depth)
+		{
+			if (StartTime >= TaskEntry->GetStartedTimestamp() && EndTime <= TaskEntry->GetFinishedTimestamp())
+			{
+				return true;
+			}
+
+			return false;
+		});
+
+	const TSharedPtr<const ITimingEvent> FoundEvent = Track->SearchEvent(SearchParams);
+	TimingView->SelectTimingEvent(FoundEvent, true);
 }
 
 ////////////////////////////////////////////////////////////////////////////////////////////////////
