[Insights] Show Waiting Task Id's as a multi line value instead of just their number

diff --git a/Engine/Source/Developer/TraceInsights/Private/Insights/ViewModels/ThreadTimingTrack.cpp b/Engine/Source/Developer/TraceInsights/Private/Insights/ViewModels/ThreadTimingTrack.cpp
--- a/Engine/Source/Developer/TraceInsights/Private/Insights/ViewModels/ThreadTimingTrack.cpp
+++ b/Engine/Source/Developer/TraceInsights/Private/Insights/ViewModels/ThreadTimingTrack.cpp
@@ -1236,23 +1236,54 @@ void FThreadTimingTrack::InitTooltip(FTooltipDrawState& InOutTooltip, const ITim
 			do // info about blocking
 			{
 				const TraceServices::FWaitingForTasks* Waiting = TasksProvider->TryGetWaiting(*TimerName, ThreadId, TooltipEvent.GetStartTime());
-				if (Waiting == nullptr)
+				if (Waiting == nullptr || Waiting->Tasks.Num() == 0)
 				{
 					break;
 				}
 
 				InOutTooltip.AddTextLine(TEXT("-------- Waiting for tasks --------"), FLinearColor::Red);
-				const int32 MaxWaitedTasksToList = 5;
-				FString TaskIdsStr = Waiting->Tasks.Num() <= MaxWaitedTasksToList ?
-					FString::JoinBy(Waiting->Tasks, TEXT(", "), [](TaskTrace::FId TaskId) { return FString::FromInt(TaskId); }) :
-					FString::Printf(TEXT("[%d]"), Waiting->Tasks.Num());
-				InOutTooltip.AddNameValueTextLine(TEXT("Tasks:"), FString::Printf(TEXT("%s"), *TaskIdsStr));
+				constexpr int32 NumIdsOnRow = 4;
+				TStringBuilder<1024> StringBuilder;
+
+				// Add the first line of Task Id values.
+				for (int32 Index = 0; Index < Waiting->Tasks.Num() && Index < NumIdsOnRow; ++Index)
+				{
+					StringBuilder.Appendf(TEXT("%d, "), Waiting->Tasks[Index]);
+				}
+
+				// Remove separators from last entry.
+				if (Waiting->Tasks.Num() <= NumIdsOnRow)
+				{
+					StringBuilder.RemoveSuffix(2);
+				}
+				InOutTooltip.AddNameValueTextLine(FString::Printf(TEXT("Tasks[%d]:"), Waiting->Tasks.Num()), StringBuilder.ToString());
+				StringBuilder.Reset();
+
+				// Add the rest of the lines with an empty name so they appear as a multi line value.
+				for (int32 Index = NumIdsOnRow; Index < Waiting->Tasks.Num(); ++Index)
+				{
+					StringBuilder.Appendf(TEXT("%d, "), Waiting->Tasks[Index]);
+
+					if ((Index + 1) % NumIdsOnRow == 0)
+					{
+						InOutTooltip.AddNameValueTextLine(TEXT(""), StringBuilder.ToString());
+						StringBuilder.Reset();
+					}
+				}
+
+				if (StringBuilder.Len() > 1)
+				{
+					StringBuilder.RemoveSuffix(2);
+					InOutTooltip.AddNameValueTextLine(TEXT(""), StringBuilder.ToString());
+				}
+
 				InOutTooltip.AddNameValueTextLine(TEXT("Started waiting:"), FString::Printf(TEXT("%f"), Waiting->StartedTimestamp));
 				InOutTooltip.AddNameValueTextLine(TEXT("Finished waiting:"),
 					FString::Printf(TEXT("%s (+%s)"),
 						Waiting->FinishedTimestamp == TraceServices::FTaskInfo::InvalidTimestamp ? TEXT("[not set]") : *FString::SanitizeFloat(Waiting->FinishedTimestamp),
 						*TimeUtils::FormatTimeAuto(Waiting->FinishedTimestamp - Waiting->StartedTimestamp)));
 
+				const int32 MaxWaitedTasksToList = 5;
 				int32 NumTasksToList = FMath::Min(Waiting->Tasks.Num(), MaxWaitedTasksToList);
 				for (int32 i = 0; i != NumTasksToList; ++i)
 				{
