[Insights]
  - Mem Allocs Tree view: Fixed crash when choosing grouping by Heap (groupings were initialized before Session ptr).
  - Mem Investigation view: Fixed alignment for the "symbol paths info" text block.

diff --git a/Engine/Source/Developer/TraceInsights/Private/Insights/MemoryProfiler/Widgets/SMemInvestigationView.cpp b/Engine/Source/Developer/TraceInsights/Private/Insights/MemoryProfiler/Widgets/SMemInvestigationView.cpp
--- a/Engine/Source/Developer/TraceInsights/Private/Insights/MemoryProfiler/Widgets/SMemInvestigationView.cpp
+++ b/Engine/Source/Developer/TraceInsights/Private/Insights/MemoryProfiler/Widgets/SMemInvestigationView.cpp
@@ -185,7 +185,7 @@ TSharedRef<SWidget> SMemInvestigationView::ConstructInvestigationWidgetArea()
 		+ SVerticalBox::Slot()
 		.AutoHeight()
 		.Padding(FMargin(0.0f, 10.0f, 0.0f, 0.0f))
-		.HAlign(HAlign_Left)
+		.HAlign(HAlign_Fill)
 		[
 			SAssignNew(SymbolPathsTextBlock, STextBlock)
 			.Text(GetSymbolPathsText())
diff --git a/Engine/Source/Developer/TraceInsights/Private/Insights/Table/Widgets/STableTreeView.cpp b/Engine/Source/Developer/TraceInsights/Private/Insights/Table/Widgets/STableTreeView.cpp
--- a/Engine/Source/Developer/TraceInsights/Private/Insights/Table/Widgets/STableTreeView.cpp
+++ b/Engine/Source/Developer/TraceInsights/Private/Insights/Table/Widgets/STableTreeView.cpp
@@ -262,8 +262,6 @@ void STableTreeView::ConstructWidget(TSharedPtr<FTable> InTablePtr)
 	}
 
 	InitializeAndShowHeaderColumns();
-	CreateGroupings();
-	CreateSortings();
 
 	InitCommandList();
 
@@ -272,6 +270,9 @@ void STableTreeView::ConstructWidget(TSharedPtr<FTable> InTablePtr)
 
 	// Update the Session (i.e. when analysis session was already started).
 	InsightsManager_OnSessionChanged();
+
+	CreateGroupings();
+	CreateSortings();
 }
 END_SLATE_FUNCTION_BUILD_OPTIMIZATION
 
@@ -1368,7 +1369,10 @@ void STableTreeView::GroupNodesRec(const TArray<FTableTreeNodePtr>& Nodes, FTabl
 		return;
 	}
 
-	ensure(Groupings.Num() > 0);
+	if (Groupings.Num() == 0)
+	{
+		return;
+	}
 
 	FTreeNodeGrouping& Grouping = *Groupings[GroupingDepth];
 
