[Insights] Updated version to 1.00 in preparation for UE 5.0 release. Also fixed offset to window edge for displayed version.

diff --git a/Engine/Source/Developer/TraceInsights/Private/Insights/LoadingProfiler/Widgets/SLoadingProfilerWindow.cpp b/Engine/Source/Developer/TraceInsights/Private/Insights/LoadingProfiler/Widgets/SLoadingProfilerWindow.cpp
--- a/Engine/Source/Developer/TraceInsights/Private/Insights/LoadingProfiler/Widgets/SLoadingProfilerWindow.cpp
+++ b/Engine/Source/Developer/TraceInsights/Private/Insights/LoadingProfiler/Widgets/SLoadingProfilerWindow.cpp
@@ -592,7 +592,7 @@ void SLoadingProfilerWindow::Construct(const FArguments& InArgs, const TSharedRe
 		+ SOverlay::Slot()
 		.HAlign(HAlign_Right)
 		.VAlign(VAlign_Top)
-		.Padding(0.0f, -16.0f, 0.0f, 0.0f)
+		.Padding(0.0f, -16.0f, 4.0f, 0.0f)
 		[
 			SNew(STextBlock)
 			.Clipping(EWidgetClipping::ClipToBoundsWithoutIntersecting)
diff --git a/Engine/Source/Developer/TraceInsights/Private/Insights/MemoryProfiler/Widgets/SMemoryProfilerWindow.cpp b/Engine/Source/Developer/TraceInsights/Private/Insights/MemoryProfiler/Widgets/SMemoryProfilerWindow.cpp
--- a/Engine/Source/Developer/TraceInsights/Private/Insights/MemoryProfiler/Widgets/SMemoryProfilerWindow.cpp
+++ b/Engine/Source/Developer/TraceInsights/Private/Insights/MemoryProfiler/Widgets/SMemoryProfilerWindow.cpp
@@ -481,7 +481,7 @@ void SMemoryProfilerWindow::Construct(const FArguments& InArgs, const TSharedRef
 		+ SOverlay::Slot()
 		.HAlign(HAlign_Right)
 		.VAlign(VAlign_Top)
-		.Padding(0.0f, -16.0f, 0.0f, 0.0f)
+		.Padding(0.0f, -16.0f, 4.0f, 0.0f)
 		[
 			SNew(STextBlock)
 			.Clipping(EWidgetClipping::ClipToBoundsWithoutIntersecting)
diff --git a/Engine/Source/Developer/TraceInsights/Private/Insights/NetworkingProfiler/Widgets/SNetworkingProfilerWindow.cpp b/Engine/Source/Developer/TraceInsights/Private/Insights/NetworkingProfiler/Widgets/SNetworkingProfilerWindow.cpp
--- a/Engine/Source/Developer/TraceInsights/Private/Insights/NetworkingProfiler/Widgets/SNetworkingProfilerWindow.cpp
+++ b/Engine/Source/Developer/TraceInsights/Private/Insights/NetworkingProfiler/Widgets/SNetworkingProfilerWindow.cpp
@@ -301,7 +301,7 @@ void SNetworkingProfilerWindow::Construct(const FArguments& InArgs, const TShare
 		+ SOverlay::Slot()
 		.HAlign(HAlign_Right)
 		.VAlign(VAlign_Top)
-		.Padding(0.0f, -16.0f, 0.0f, 0.0f)
+		.Padding(0.0f, -16.0f, 4.0f, 0.0f)
 		[
 			SNew(STextBlock)
 			.Clipping(EWidgetClipping::ClipToBoundsWithoutIntersecting)
diff --git a/Engine/Source/Developer/TraceInsights/Private/Insights/Widgets/SSessionInfoWindow.cpp b/Engine/Source/Developer/TraceInsights/Private/Insights/Widgets/SSessionInfoWindow.cpp
--- a/Engine/Source/Developer/TraceInsights/Private/Insights/Widgets/SSessionInfoWindow.cpp
+++ b/Engine/Source/Developer/TraceInsights/Private/Insights/Widgets/SSessionInfoWindow.cpp
@@ -123,7 +123,7 @@ void SSessionInfoWindow::Construct(const FArguments& InArgs, const TSharedRef<SD
 		+ SOverlay::Slot()
 		.HAlign(HAlign_Right)
 		.VAlign(VAlign_Top)
-		.Padding(0.0f, -16.0f, 0.0f, 0.0f)
+		.Padding(0.0f, -16.0f, 4.0f, 0.0f)
 		[
 			SNew(STextBlock)
 			.Clipping(EWidgetClipping::ClipToBoundsWithoutIntersecting)
diff --git a/Engine/Source/Developer/TraceInsights/Private/Insights/Widgets/STimingProfilerWindow.cpp b/Engine/Source/Developer/TraceInsights/Private/Insights/Widgets/STimingProfilerWindow.cpp
--- a/Engine/Source/Developer/TraceInsights/Private/Insights/Widgets/STimingProfilerWindow.cpp
+++ b/Engine/Source/Developer/TraceInsights/Private/Insights/Widgets/STimingProfilerWindow.cpp
@@ -556,7 +556,7 @@ void STimingProfilerWindow::Construct(const FArguments& InArgs, const TSharedRef
 		+ SOverlay::Slot()
 		.HAlign(HAlign_Right)
 		.VAlign(VAlign_Top)
-		.Padding(0.0f, -16.0f, 0.0f, 0.0f)
+		.Padding(0.0f, -16.0f, 4.0f, 0.0f)
 		[
 			SNew(STextBlock)
 			.Clipping(EWidgetClipping::ClipToBoundsWithoutIntersecting)
diff --git a/Engine/Source/Developer/TraceInsights/Public/Insights/Version.h b/Engine/Source/Developer/TraceInsights/Public/Insights/Version.h
--- a/Engine/Source/Developer/TraceInsights/Public/Insights/Version.h
+++ b/Engine/Source/Developer/TraceInsights/Public/Insights/Version.h
@@ -4,11 +4,11 @@
 
 // See https://semver.org/
 
-#define UNREAL_INSIGHTS_VERSION_MAJOR 0
-#define UNREAL_INSIGHTS_VERSION_MINOR 16
+#define UNREAL_INSIGHTS_VERSION_MAJOR 1
+#define UNREAL_INSIGHTS_VERSION_MINOR 0
 #define UNREAL_INSIGHTS_VERSION_PATCH 0
 
-#define UNREAL_INSIGHTS_VERSION_STRING "0.16"
+#define UNREAL_INSIGHTS_VERSION_STRING "1.0"
 
 /** Extra identifier string (like: "alpha", "beta.1", "beta.2", "test", "featureX", etc.). */
 //#define UNREAL_INSIGHTS_VERSION_ID_STRING "beta.1"
