[Insights] Change the name of Task Track events.
+ Cleanup of task enums
+ Nested Tasks and Subsequents are drawn on the same depth in the Task Overview Track.

diff --git a/Engine/Source/Developer/TraceInsights/Private/Insights/TaskGraphProfiler/TaskGraphProfilerManager.cpp b/Engine/Source/Developer/TraceInsights/Private/Insights/TaskGraphProfiler/TaskGraphProfilerManager.cpp
--- a/Engine/Source/Developer/TraceInsights/Private/Insights/TaskGraphProfiler/TaskGraphProfilerManager.cpp
+++ b/Engine/Source/Developer/TraceInsights/Private/Insights/TaskGraphProfiler/TaskGraphProfilerManager.cpp
@@ -323,7 +323,7 @@ void FTaskGraphProfilerManager::GetSingleTaskTransitions(const TraceServices::FT
 
 	if (Task->FinishedTimestamp != Task->CompletedTimestamp || Task->CompletedThreadId != Task->StartedThreadId)
 	{
-		AddRelation(InSelectedEvent, Task->FinishedTimestamp, Task->StartedThreadId, ExecutionStartedDepth, Task->CompletedTimestamp, Task->CompletedThreadId, -1,  ETaskEventType::Completed);
+		AddRelation(InSelectedEvent, Task->FinishedTimestamp, Task->StartedThreadId, ExecutionStartedDepth, Task->CompletedTimestamp, Task->CompletedThreadId, -1,  ETaskEventType::Finished);
 	}
 }
 
@@ -339,7 +339,7 @@ void FTaskGraphProfilerManager::GetSingleTaskConnections(const TraceServices::FT
 		const TraceServices::FTaskInfo* Prerequisite = TasksProvider->TryGetTask(Task->Prerequisites[i].RelativeId);
 		check(Prerequisite != nullptr);
 
-		AddRelation(InSelectedEvent, Prerequisite->FinishedTimestamp, Prerequisite->StartedThreadId, Task->StartedTimestamp, Task->StartedThreadId, ETaskEventType::Prerequisite);
+		AddRelation(InSelectedEvent, Prerequisite->FinishedTimestamp, Prerequisite->StartedThreadId, Task->StartedTimestamp, Task->StartedThreadId, ETaskEventType::PrerequisiteStarted);
 	}
 
 	int32 NumNestedToShow = FMath::Min(Task->NestedTasks.Num(), MaxTasksToShow);
@@ -351,7 +351,7 @@ void FTaskGraphProfilerManager::GetSingleTaskConnections(const TraceServices::FT
 
 		int32 NestedExecutionStartedDepth = GetDepthOfTaskExecution(Task->StartedTimestamp, Task->FinishedTimestamp, Task->StartedThreadId);
 
-		AddRelation(InSelectedEvent, RelationInfo.Timestamp, Task->StartedThreadId, -1, NestedTask->StartedTimestamp, NestedTask->StartedThreadId, NestedExecutionStartedDepth, ETaskEventType::AddedNested);
+		AddRelation(InSelectedEvent, RelationInfo.Timestamp, Task->StartedThreadId, -1, NestedTask->StartedTimestamp, NestedTask->StartedThreadId, NestedExecutionStartedDepth, ETaskEventType::NestedStarted);
 	}
 
 	int32 NumSubsequentsToShow = FMath::Min(Task->Subsequents.Num(), MaxTasksToShow);
@@ -360,7 +360,7 @@ void FTaskGraphProfilerManager::GetSingleTaskConnections(const TraceServices::FT
 		const TraceServices::FTaskInfo* Subsequent = TasksProvider->TryGetTask(Task->Subsequents[i].RelativeId);
 		check(Subsequent != nullptr);
 
-		AddRelation(InSelectedEvent, Task->FinishedTimestamp, Task->StartedThreadId, Subsequent->StartedTimestamp, Subsequent->StartedThreadId, ETaskEventType::Subsequent);
+		AddRelation(InSelectedEvent, Task->FinishedTimestamp, Task->StartedThreadId, Subsequent->StartedTimestamp, Subsequent->StartedThreadId, ETaskEventType::SubsequentStarted);
 	}
 }
 
@@ -494,13 +494,13 @@ void FTaskGraphProfilerManager::InitializeColorCode()
 
 	ColorCode[static_cast<uint32>(ETaskEventType::Created)] = ToLiniarColorNoAlpha(0xFFDC1AFF); // Yellow;
 	ColorCode[static_cast<uint32>(ETaskEventType::Launched)] = ToLiniarColorNoAlpha(0x8BC24AFF); // Green;
-	ColorCode[static_cast<uint32>(ETaskEventType::Prerequisite)] = ToLiniarColorNoAlpha(0xFF729CFF); //Pink;
 	ColorCode[static_cast<uint32>(ETaskEventType::Scheduled)] = ToLiniarColorNoAlpha(0x26BBFFFF); // Blue;
 	ColorCode[static_cast<uint32>(ETaskEventType::Started)] = ToLiniarColorNoAlpha(0xFF0000FF); // Red;
-	ColorCode[static_cast<uint32>(ETaskEventType::AddedNested)] = ToLiniarColorNoAlpha(0xB68F55FF); // Folder
+	ColorCode[static_cast<uint32>(ETaskEventType::Finished)] = ToLiniarColorNoAlpha(0xFFDC1AFF); // Yellow
+	ColorCode[static_cast<uint32>(ETaskEventType::PrerequisiteStarted)] = ToLiniarColorNoAlpha(0xFF729CFF); //Pink;
+	ColorCode[static_cast<uint32>(ETaskEventType::NestedStarted)] = ToLiniarColorNoAlpha(0xB68F55FF); // Folder
+	ColorCode[static_cast<uint32>(ETaskEventType::SubsequentStarted)] = ToLiniarColorNoAlpha(0xFE9B07FF); // Orange;;
 	ColorCode[static_cast<uint32>(ETaskEventType::NestedCompleted)] = ToLiniarColorNoAlpha(0x804D39FF); // Brown
-	ColorCode[static_cast<uint32>(ETaskEventType::Subsequent)] = ToLiniarColorNoAlpha(0xFE9B07FF); // Orange;;
-	ColorCode[static_cast<uint32>(ETaskEventType::Completed)] = ToLiniarColorNoAlpha(0xFFDC1AFF); // Yellow
 	ColorCode[static_cast<uint32>(ETaskEventType::CriticalPath)] = ToLiniarColorNoAlpha(0xA139BFFF); // Purple
 }
 
diff --git a/Engine/Source/Developer/TraceInsights/Private/Insights/TaskGraphProfiler/TaskGraphProfilerManager.h b/Engine/Source/Developer/TraceInsights/Private/Insights/TaskGraphProfiler/TaskGraphProfilerManager.h
--- a/Engine/Source/Developer/TraceInsights/Private/Insights/TaskGraphProfiler/TaskGraphProfilerManager.h
+++ b/Engine/Source/Developer/TraceInsights/Private/Insights/TaskGraphProfiler/TaskGraphProfilerManager.h
@@ -39,13 +39,13 @@ enum class ETaskEventType : uint32
 {
 	Created,
 	Launched,
-	Prerequisite,
 	Scheduled,
 	Started,
-	AddedNested,
+	Finished,
+	PrerequisiteStarted,
+	NestedStarted,
 	NestedCompleted,
-	Subsequent,
-	Completed,
+	SubsequentStarted,
 	CriticalPath,
 	
 	NumTaskEventTypes,
diff --git a/Engine/Source/Developer/TraceInsights/Private/Insights/TaskGraphProfiler/ViewModels/TaskTimingTrack.cpp b/Engine/Source/Developer/TraceInsights/Private/Insights/TaskGraphProfiler/ViewModels/TaskTimingTrack.cpp
--- a/Engine/Source/Developer/TraceInsights/Private/Insights/TaskGraphProfiler/ViewModels/TaskTimingTrack.cpp
+++ b/Engine/Source/Developer/TraceInsights/Private/Insights/TaskGraphProfiler/ViewModels/TaskTimingTrack.cpp
@@ -666,13 +666,13 @@ void FTaskTimingTrack::BuildDrawState(ITimingEventsTrackDrawStateBuilder& Builde
 		return;
 	}
 
-	Builder.AddEvent(Task->CreatedTimestamp, Task->LaunchedTimestamp, 0, TEXT("Launched"), 0, FTaskGraphProfilerManager::Get()->GetColorForTaskEventAsPackedARGB(ETaskEventType::Created));
-	Builder.AddEvent(Task->LaunchedTimestamp, Task->ScheduledTimestamp, 0, TEXT("Dispatched"), 0, FTaskGraphProfilerManager::Get()->GetColorForTaskEventAsPackedARGB(ETaskEventType::Launched));
+	Builder.AddEvent(Task->CreatedTimestamp, Task->LaunchedTimestamp, 0, TEXT("Created"), 0, FTaskGraphProfilerManager::Get()->GetColorForTaskEventAsPackedARGB(ETaskEventType::Created));
+	Builder.AddEvent(Task->LaunchedTimestamp, Task->ScheduledTimestamp, 0, TEXT("Launched"), 0, FTaskGraphProfilerManager::Get()->GetColorForTaskEventAsPackedARGB(ETaskEventType::Launched));
 	Builder.AddEvent(Task->ScheduledTimestamp, Task->StartedTimestamp, 0, TEXT("Scheduled"), 0,  FTaskGraphProfilerManager::Get()->GetColorForTaskEventAsPackedARGB(ETaskEventType::Scheduled));
-	Builder.AddEvent(Task->StartedTimestamp, Task->FinishedTimestamp, 0, TEXT("Executed"), 0, FTaskGraphProfilerManager::Get()->GetColorForTaskEventAsPackedARGB(ETaskEventType::Started));
+	Builder.AddEvent(Task->StartedTimestamp, Task->FinishedTimestamp, 0, TEXT("Executing"), 0, FTaskGraphProfilerManager::Get()->GetColorForTaskEventAsPackedARGB(ETaskEventType::Started));
 	if (Task->CompletedTimestamp > Task->FinishedTimestamp)
 	{
-		Builder.AddEvent(Task->FinishedTimestamp, Task->CompletedTimestamp, 0, TEXT("Completed"), 0, FTaskGraphProfilerManager::Get()->GetColorForTaskEventAsPackedARGB(ETaskEventType::Completed));
+		Builder.AddEvent(Task->FinishedTimestamp, Task->CompletedTimestamp, 0, TEXT("Finished"), 0, FTaskGraphProfilerManager::Get()->GetColorForTaskEventAsPackedARGB(ETaskEventType::Finished));
 	}
 
 	if (bShowDetailInfoOnTaskTrack)
@@ -683,27 +683,27 @@ void FTaskTimingTrack::BuildDrawState(ITimingEventsTrackDrawStateBuilder& Builde
 			const TraceServices::FTaskInfo* TaskInfo = TasksProvider->TryGetTask(Relation.RelativeId);
 			if (TaskInfo)
 			{
-				Builder.AddEvent(TaskInfo->StartedTimestamp, TaskInfo->FinishedTimestamp, ++Depth, *FString::Printf(TEXT("Prerequisite Task %d"), TaskInfo->Id), 0, FTaskGraphProfilerManager::Get()->GetColorForTaskEventAsPackedARGB(ETaskEventType::Prerequisite));
+				Builder.AddEvent(TaskInfo->StartedTimestamp, TaskInfo->FinishedTimestamp, ++Depth, *FString::Printf(TEXT("Prerequisite Task %d Executing"), TaskInfo->Id), 0, FTaskGraphProfilerManager::Get()->GetColorForTaskEventAsPackedARGB(ETaskEventType::PrerequisiteStarted));
 			}
 		}
 
-		// The nested/subsequents tasks should not overlap the prerequisites tasks, so we can draw them on same depth."
 		Depth = 0;
 		for (TraceServices::FTaskInfo::FRelationInfo Relation : Task->NestedTasks)
 		{
 			const TraceServices::FTaskInfo* TaskInfo = TasksProvider->TryGetTask(Relation.RelativeId);
 			if (TaskInfo)
 			{
-				Builder.AddEvent(TaskInfo->StartedTimestamp, TaskInfo->FinishedTimestamp, ++Depth, *FString::Printf(TEXT("Nested Task %d"), TaskInfo->Id), 0, FTaskGraphProfilerManager::Get()->GetColorForTaskEventAsPackedARGB(ETaskEventType::AddedNested));
+				Builder.AddEvent(TaskInfo->StartedTimestamp, TaskInfo->FinishedTimestamp, ++Depth, *FString::Printf(TEXT("Nested Task %d Executing"), TaskInfo->Id), 0, FTaskGraphProfilerManager::Get()->GetColorForTaskEventAsPackedARGB(ETaskEventType::NestedStarted));
 			}
 		}
 
+		Depth = 0;
 		for (TraceServices::FTaskInfo::FRelationInfo Relation : Task->Subsequents)
 		{
 			const TraceServices::FTaskInfo* TaskInfo = TasksProvider->TryGetTask(Relation.RelativeId);
 			if (TaskInfo)
 			{
-				Builder.AddEvent(TaskInfo->StartedTimestamp, TaskInfo->FinishedTimestamp, ++Depth, *FString::Printf(TEXT("Subsequent Task %d"), TaskInfo->Id), 0, FTaskGraphProfilerManager::Get()->GetColorForTaskEventAsPackedARGB(ETaskEventType::Subsequent));
+				Builder.AddEvent(TaskInfo->StartedTimestamp, TaskInfo->FinishedTimestamp, ++Depth, *FString::Printf(TEXT("Subsequent Task %d Executing"), TaskInfo->Id), 0, FTaskGraphProfilerManager::Get()->GetColorForTaskEventAsPackedARGB(ETaskEventType::SubsequentStarted));
 			}
 		}
 	}
@@ -829,23 +829,23 @@ const TSharedPtr<const ITimingEvent> FTaskTimingTrack::GetEvent(float InPosX, fl
 		{
 			if (EventTime >= Task->CreatedTimestamp && EventTime < Task->LaunchedTimestamp)
 			{
-				TimingEvent = MakeShared<FTaskTrackEvent>(SharedThis(this), Task->CreatedTimestamp, Task->LaunchedTimestamp, 0, ETaskTrackEventType::Launched);
+				TimingEvent = MakeShared<FTaskTrackEvent>(SharedThis(this), Task->CreatedTimestamp, Task->LaunchedTimestamp, 0, ETaskEventType::Created);
 			}
 			else if (EventTime >= Task->LaunchedTimestamp && EventTime < Task->ScheduledTimestamp)
 			{
-				TimingEvent = MakeShared<FTaskTrackEvent>(SharedThis(this), Task->LaunchedTimestamp, Task->ScheduledTimestamp, 0, ETaskTrackEventType::Dispatched);
+				TimingEvent = MakeShared<FTaskTrackEvent>(SharedThis(this), Task->LaunchedTimestamp, Task->ScheduledTimestamp, 0, ETaskEventType::Launched);
 			}
 			else if (EventTime >= Task->ScheduledTimestamp && EventTime < Task->StartedTimestamp)
 			{
-				TimingEvent = MakeShared<FTaskTrackEvent>(SharedThis(this), Task->ScheduledTimestamp, Task->StartedTimestamp, 0, ETaskTrackEventType::Scheduled);
+				TimingEvent = MakeShared<FTaskTrackEvent>(SharedThis(this), Task->ScheduledTimestamp, Task->StartedTimestamp, 0, ETaskEventType::Scheduled);
 			}
 			else if (EventTime >= Task->StartedTimestamp && EventTime < Task->FinishedTimestamp)
 			{
-				TimingEvent = MakeShared<FTaskTrackEvent>(SharedThis(this), Task->StartedTimestamp, Task->FinishedTimestamp, 0, ETaskTrackEventType::Executed);
+				TimingEvent = MakeShared<FTaskTrackEvent>(SharedThis(this), Task->StartedTimestamp, Task->FinishedTimestamp, 0, ETaskEventType::Started);
 			}
 			else if (EventTime >= Task->FinishedTimestamp && EventTime < Task->CompletedTimestamp)
 			{
-				TimingEvent = MakeShared<FTaskTrackEvent>(SharedThis(this), Task->FinishedTimestamp, Task->CompletedTimestamp, 0, ETaskTrackEventType::Completed);
+				TimingEvent = MakeShared<FTaskTrackEvent>(SharedThis(this), Task->FinishedTimestamp, Task->CompletedTimestamp, 0, ETaskEventType::Finished);
 			}
 
 			if (TimingEvent.IsValid())
@@ -855,13 +855,13 @@ const TSharedPtr<const ITimingEvent> FTaskTimingTrack::GetEvent(float InPosX, fl
 		}
 		else if(bShowDetailInfoOnTaskTrack)
 		{
-			auto GetEventFromRelations = [&TasksProvider, this, EventTime](const TArray<TraceServices::FTaskInfo::FRelationInfo>& Relations, int32 Depth, ETaskTrackEventType EventType, int32 PosOffset) -> TSharedPtr<FTaskTrackEvent>
+			auto GetEventFromRelations = [&TasksProvider, this, EventTime](const TArray<TraceServices::FTaskInfo::FRelationInfo>& Relations, int32 Depth, ETaskEventType EventType) -> TSharedPtr<FTaskTrackEvent>
 			{
-				if (Depth - PosOffset <= 0 || Relations.Num() < Depth - PosOffset)
+				if (Relations.Num() < Depth)
 				{
 					return nullptr;
 				}
-				const TraceServices::FTaskInfo::FRelationInfo& Relation = Relations[Depth - PosOffset - 1];
+				const TraceServices::FTaskInfo::FRelationInfo& Relation = Relations[Depth - 1];
 				{
 					const TraceServices::FTaskInfo* TaskInfo = TasksProvider->TryGetTask(Relation.RelativeId);
 					if (TaskInfo && EventTime >= TaskInfo->StartedTimestamp && EventTime <= TaskInfo->FinishedTimestamp)
@@ -875,16 +875,16 @@ const TSharedPtr<const ITimingEvent> FTaskTimingTrack::GetEvent(float InPosX, fl
 				return nullptr;
 			};
 
-			TimingEvent = GetEventFromRelations(Task->Prerequisites, Depth, ETaskTrackEventType::Prerequisite, 0);
+			TimingEvent = GetEventFromRelations(Task->Prerequisites, Depth, ETaskEventType::PrerequisiteStarted);
 
 			if (!TimingEvent.IsValid())
 			{
-				TimingEvent = GetEventFromRelations(Task->NestedTasks, Depth, ETaskTrackEventType::Nested, 0);
+				TimingEvent = GetEventFromRelations(Task->NestedTasks, Depth, ETaskEventType::NestedStarted);
 			}
 
 			if (!TimingEvent.IsValid())
 			{
-				TimingEvent = GetEventFromRelations(Task->Subsequents, Depth, ETaskTrackEventType::Subsequent, Task->NestedTasks.Num());
+				TimingEvent = GetEventFromRelations(Task->Subsequents, Depth, ETaskEventType::SubsequentStarted);
 			}
 		}
 	}
@@ -936,26 +936,26 @@ void FTaskTimingTrack::InitTooltip(FTooltipDrawState& InOutTooltip, const ITimin
 
 	switch (TaskTrackEvent.GetTaskEventType())
 	{
-	case ETaskTrackEventType::Launched:
+	case ETaskEventType::Created:
 		break;
-	case ETaskTrackEventType::Prerequisite:
-	case ETaskTrackEventType::Nested:
-	case ETaskTrackEventType::Subsequent:
+	case ETaskEventType::PrerequisiteStarted:
+	case ETaskEventType::NestedStarted:
+	case ETaskEventType::SubsequentStarted:
 	{
 		InOutTooltip.AddNameValueTextLine(TEXT("Executed Thread Id:"), FString::Printf(TEXT("%d"), Task->StartedThreadId));
 		break;
 	}
-	case ETaskTrackEventType::Dispatched:
+	case ETaskEventType::Launched:
 	{
 		InOutTooltip.AddNameValueTextLine(TEXT("Prerequisite tasks:"), FString::Printf(TEXT("%d"), Task->Prerequisites.Num()));
 		break;
 	}
-	case ETaskTrackEventType::Scheduled:
+	case ETaskEventType::Scheduled:
 		break;
-	case ETaskTrackEventType::Executed:
+	case ETaskEventType::Started:
 		InOutTooltip.AddNameValueTextLine(TEXT("Nested tasks:"), FString::Printf(TEXT("%d"), Task->NestedTasks.Num()));
 		break;
-	case ETaskTrackEventType::Completed:
+	case ETaskEventType::Finished:
 		InOutTooltip.AddNameValueTextLine(TEXT("Subsequent tasks:"), FString::Printf(TEXT("%d"), Task->Subsequents.Num()));
 		break;
 	default:
@@ -1004,7 +1004,7 @@ void FTaskTimingTrack::GetEventRelations(const FThreadTrackEvent& InSelectedEven
 				{
 					int32 WaitingTaskExecutionDepth = FTaskGraphProfilerManager::Get()->GetDepthOfTaskExecution(WaitedTask->StartedTimestamp, WaitedTask->FinishedTimestamp, WaitedTask->StartedThreadId);
 
-					FTaskGraphProfilerManager::Get()->AddRelation(&InSelectedEvent, StartTime, ThreadId, -1, WaitedTask->StartedTimestamp, WaitedTask->StartedThreadId, WaitingTaskExecutionDepth, ETaskEventType::AddedNested);
+					FTaskGraphProfilerManager::Get()->AddRelation(&InSelectedEvent, StartTime, ThreadId, -1, WaitedTask->StartedTimestamp, WaitedTask->StartedThreadId, WaitingTaskExecutionDepth, ETaskEventType::NestedStarted);
 					FTaskGraphProfilerManager::Get()->AddRelation(&InSelectedEvent, WaitedTask->CompletedTimestamp, WaitedTask->CompletedThreadId, WaitedTask->CompletedTimestamp, ThreadId, ETaskEventType::NestedCompleted);
 				}
 			}
diff --git a/Engine/Source/Developer/TraceInsights/Private/Insights/TaskGraphProfiler/ViewModels/TaskTrackEvent.cpp b/Engine/Source/Developer/TraceInsights/Private/Insights/TaskGraphProfiler/ViewModels/TaskTrackEvent.cpp
--- a/Engine/Source/Developer/TraceInsights/Private/Insights/TaskGraphProfiler/ViewModels/TaskTrackEvent.cpp
+++ b/Engine/Source/Developer/TraceInsights/Private/Insights/TaskGraphProfiler/ViewModels/TaskTrackEvent.cpp
@@ -2,13 +2,16 @@
 
 #include "Insights/TaskGraphProfiler/ViewModels/TaskTrackEvent.h"
 
+namespace Insights
+{
+
 ////////////////////////////////////////////////////////////////////////////////////////////////////
 
 INSIGHTS_IMPLEMENT_RTTI(FTaskTrackEvent)
 
 ////////////////////////////////////////////////////////////////////////////////////////////////////
 
-FTaskTrackEvent::FTaskTrackEvent(const TSharedRef<const FBaseTimingTrack> InTrack, double InStartTime, double InEndTime, uint32 InDepth, ETaskTrackEventType InType)
+FTaskTrackEvent::FTaskTrackEvent(const TSharedRef<const FBaseTimingTrack> InTrack, double InStartTime, double InEndTime, uint32 InDepth, ETaskEventType InType)
 	: FTimingEvent(InTrack, InStartTime, InEndTime, InDepth)
 	, TaskEventType(InType)
 {}
@@ -19,18 +22,18 @@ FString FTaskTrackEvent::GetStartLabel() const
 {
 	switch (TaskEventType)
 	{
-	case ETaskTrackEventType::Launched:
+	case ETaskEventType::Created:
 		return TEXT("Created Time:");
-	case ETaskTrackEventType::Dispatched:
+	case ETaskEventType::Launched:
 		return TEXT("Launched Time");
-	case ETaskTrackEventType::Scheduled:
+	case ETaskEventType::Scheduled:
 		return TEXT("Scheduled Time:");
-	case ETaskTrackEventType::Executed:
-	case ETaskTrackEventType::Prerequisite:
-	case ETaskTrackEventType::Nested:
-	case ETaskTrackEventType::Subsequent:
+	case ETaskEventType::Started:
+	case ETaskEventType::PrerequisiteStarted:
+	case ETaskEventType::NestedStarted:
+	case ETaskEventType::SubsequentStarted:
 		return TEXT("Started Time:");
-	case ETaskTrackEventType::Completed:
+	case ETaskEventType::Finished:
 		return TEXT("Finished Time:");
 	default:
 		checkf(false, TEXT("Unknown task event type"));
@@ -46,18 +49,18 @@ FString FTaskTrackEvent::GetEndLabel() const
 {
 	switch (TaskEventType)
 	{
-	case ETaskTrackEventType::Launched:
+	case ETaskEventType::Created:
 		return TEXT("Launched Time:");
-	case ETaskTrackEventType::Dispatched:
+	case ETaskEventType::Launched:
 		return TEXT("Scheduled Time");
-	case ETaskTrackEventType::Scheduled:
+	case ETaskEventType::Scheduled:
 		return TEXT("Started Time:");
-	case ETaskTrackEventType::Executed:
-	case ETaskTrackEventType::Prerequisite:
-	case ETaskTrackEventType::Nested:
-	case ETaskTrackEventType::Subsequent:
+	case ETaskEventType::Started:
+	case ETaskEventType::PrerequisiteStarted:
+	case ETaskEventType::NestedStarted:
+	case ETaskEventType::SubsequentStarted:
 		return TEXT("Finished Time:");
-	case ETaskTrackEventType::Completed:
+	case ETaskEventType::Finished:
 		return TEXT("Completed Time:");
 	default:
 		checkf(false, TEXT("Unknown task event type"));
@@ -73,22 +76,22 @@ FString FTaskTrackEvent::GetEventName() const
 {
 	switch (TaskEventType)
 	{
-	case ETaskTrackEventType::Launched:
+	case ETaskEventType::Created:
+		return TEXT("Created");
+	case ETaskEventType::Launched:
 		return TEXT("Launched");
-	case ETaskTrackEventType::Dispatched:
-		return TEXT("Dispatched");
-	case ETaskTrackEventType::Scheduled:
+	case ETaskEventType::Scheduled:
 		return TEXT("Scheduled");
-	case ETaskTrackEventType::Executed:
-		return TEXT("Executed");
-	case ETaskTrackEventType::Completed:
-		return TEXT("Completed");
-	case ETaskTrackEventType::Prerequisite:
-		return FString::Printf(TEXT("Prerequisite Task %d"), GetTaskId());	
-	case ETaskTrackEventType::Nested:
-		return FString::Printf(TEXT("Nested Task %d"), GetTaskId());
-	case ETaskTrackEventType::Subsequent:
-		return FString::Printf(TEXT("Subsequent Task %d"), GetTaskId());
+	case ETaskEventType::Started:
+		return TEXT("Executing");
+	case ETaskEventType::Finished:
+		return TEXT("Finished");
+	case ETaskEventType::PrerequisiteStarted:
+		return FString::Printf(TEXT("Prerequisite Task %d Executing"), GetTaskId());	
+	case ETaskEventType::NestedStarted:
+		return FString::Printf(TEXT("Nested Task %d Executing"), GetTaskId());
+	case ETaskEventType::SubsequentStarted:
+		return FString::Printf(TEXT("Subsequent Task %d Executing"), GetTaskId());
 
 	default:
 		checkf(false, TEXT("Unknown task event type"));
@@ -100,3 +103,5 @@ FString FTaskTrackEvent::GetEventName() const
 
 ////////////////////////////////////////////////////////////////////////////////////////////////////
 ////////////////////////////////////////////////////////////////////////////////////////////////////
+
+} // namespace Insights
diff --git a/Engine/Source/Developer/TraceInsights/Private/Insights/TaskGraphProfiler/ViewModels/TaskTrackEvent.h b/Engine/Source/Developer/TraceInsights/Private/Insights/TaskGraphProfiler/ViewModels/TaskTrackEvent.h
--- a/Engine/Source/Developer/TraceInsights/Private/Insights/TaskGraphProfiler/ViewModels/TaskTrackEvent.h
+++ b/Engine/Source/Developer/TraceInsights/Private/Insights/TaskGraphProfiler/ViewModels/TaskTrackEvent.h
@@ -5,21 +5,11 @@
 #include "CoreMinimal.h"
 
 // Insights
+#include "Insights/TaskGraphProfiler/TaskGraphProfilerManager.h"
 #include "Insights/ViewModels/TimingEvent.h"
 
-////////////////////////////////////////////////////////////////////////////////////////////////////
-
-enum class ETaskTrackEventType : uint32
+namespace Insights
 {
-	Launched,
-	Dispatched,
-	Scheduled,
-	Executed,
-	Completed,
-	Prerequisite,
-	Nested,
-	Subsequent,
-};
 
 ////////////////////////////////////////////////////////////////////////////////////////////////////
 
@@ -28,7 +18,7 @@ class FTaskTrackEvent : public FTimingEvent
 	INSIGHTS_DECLARE_RTTI(FTaskTrackEvent, FTimingEvent)
 
 public:
-	FTaskTrackEvent(const TSharedRef<const FBaseTimingTrack> InTrack, double InStartTime, double InEndTime, uint32 InDepth, ETaskTrackEventType InTaskEventType);
+	FTaskTrackEvent(const TSharedRef<const FBaseTimingTrack> InTrack, double InStartTime, double InEndTime, uint32 InDepth, ETaskEventType InTaskEventType);
 	virtual ~FTaskTrackEvent() {}
 
 	FString GetStartLabel() const;
@@ -36,14 +26,16 @@ public:
 
 	FString GetEventName() const;
 
-	ETaskTrackEventType GetTaskEventType() const { return TaskEventType; }
+	ETaskEventType GetTaskEventType() const { return TaskEventType; }
 
 	uint32 GetTaskId() const { return TaskId; }
 	void SetTaskId(uint32 InTaskId) { TaskId = InTaskId; }
 
 private:
-	ETaskTrackEventType TaskEventType;
+	ETaskEventType TaskEventType;
 	uint32 TaskId = ~0;
 };
 
 ////////////////////////////////////////////////////////////////////////////////////////////////////
+
+} // namespace Insights
\ No newline at end of file
