[Insights] Task Relations and the Task Track now adapt to the Depth Limit setting.

diff --git a/Engine/Source/Developer/TraceInsights/Private/Insights/TaskGraphProfiler/ViewModels/TaskGraphRelation.cpp b/Engine/Source/Developer/TraceInsights/Private/Insights/TaskGraphProfiler/ViewModels/TaskGraphRelation.cpp
--- a/Engine/Source/Developer/TraceInsights/Private/Insights/TaskGraphProfiler/ViewModels/TaskGraphRelation.cpp
+++ b/Engine/Source/Developer/TraceInsights/Private/Insights/TaskGraphProfiler/ViewModels/TaskGraphRelation.cpp
@@ -4,6 +4,7 @@
 
 // Insights
 #include "Insights/Common/PaintUtils.h"
+#include "Insights/TimingProfilerManager.h"
 #include "Insights/TaskGraphProfiler/TaskGraphProfilerManager.h"
 #include "Insights/ViewModels/TimingTrackViewport.h"
 #include "Insights/ViewModels/TimingViewDrawHelper.h"
@@ -69,15 +70,17 @@ void FTaskGraphRelation::Draw(const FDrawContext& DrawContext, const FTimingTrac
 		return;
 	}
 
+	int32 ActualSourceDepth = FMath::Min(SourceDepth, (int32)FTimingProfilerManager::Get()->GetEventDepthLimit() - 1);
 	float Y1 = SourceTrackShared->GetPosY();
-	Y1 += Viewport.GetLayout().GetLaneY(SourceDepth) + Viewport.GetLayout().EventH / 2.0f;
+	Y1 += Viewport.GetLayout().GetLaneY(ActualSourceDepth) + Viewport.GetLayout().EventH / 2.0f;
 	if (SourceTrackShared->GetChildTrack() && SourceTrackShared->GetChildTrack()->GetHeight() > 0.0f)
 	{
 		Y1 += SourceTrackShared->GetChildTrack()->GetHeight() + Viewport.GetLayout().ChildTimelineDY;
 	}
 
+	int32 ActualTargetDepth = FMath::Min(TargetDepth, (int32)FTimingProfilerManager::Get()->GetEventDepthLimit() - 1);
 	float Y2 = TargetTrackShared->GetPosY();
-	Y2 += Viewport.GetLayout().GetLaneY(TargetDepth) + Viewport.GetLayout().EventH / 2.0f;
+	Y2 += Viewport.GetLayout().GetLaneY(ActualTargetDepth) + Viewport.GetLayout().EventH / 2.0f;
 	if (TargetTrackShared->GetChildTrack() && TargetTrackShared->GetChildTrack()->GetHeight() > 0.0f)
 	{
 		Y2 += TargetTrackShared->GetChildTrack()->GetHeight() + Viewport.GetLayout().ChildTimelineDY;
diff --git a/Engine/Source/Developer/TraceInsights/Private/Insights/TaskGraphProfiler/ViewModels/TaskTimingTrack.cpp b/Engine/Source/Developer/TraceInsights/Private/Insights/TaskGraphProfiler/ViewModels/TaskTimingTrack.cpp
--- a/Engine/Source/Developer/TraceInsights/Private/Insights/TaskGraphProfiler/ViewModels/TaskTimingTrack.cpp
+++ b/Engine/Source/Developer/TraceInsights/Private/Insights/TaskGraphProfiler/ViewModels/TaskTimingTrack.cpp
@@ -680,6 +680,10 @@ void FTaskTimingTrack::BuildDrawState(ITimingEventsTrackDrawStateBuilder& Builde
 		uint32 Depth = 0;
 		for (TraceServices::FTaskInfo::FRelationInfo Relation : Task->Prerequisites)
 		{
+			if (Depth >= FTimingProfilerManager::Get()->GetEventDepthLimit() - 1)
+			{
+				break;
+			}
 			const TraceServices::FTaskInfo* TaskInfo = TasksProvider->TryGetTask(Relation.RelativeId);
 			if (TaskInfo)
 			{
@@ -690,6 +694,10 @@ void FTaskTimingTrack::BuildDrawState(ITimingEventsTrackDrawStateBuilder& Builde
 		Depth = 0;
 		for (TraceServices::FTaskInfo::FRelationInfo Relation : Task->NestedTasks)
 		{
+			if (Depth >= FTimingProfilerManager::Get()->GetEventDepthLimit() - 1)
+			{
+				break;
+			}
 			const TraceServices::FTaskInfo* TaskInfo = TasksProvider->TryGetTask(Relation.RelativeId);
 			if (TaskInfo)
 			{
@@ -700,6 +708,10 @@ void FTaskTimingTrack::BuildDrawState(ITimingEventsTrackDrawStateBuilder& Builde
 		Depth = 0;
 		for (TraceServices::FTaskInfo::FRelationInfo Relation : Task->Subsequents)
 		{
+			if (Depth >= FTimingProfilerManager::Get()->GetEventDepthLimit() - 1)
+			{
+				break;
+			}
 			const TraceServices::FTaskInfo* TaskInfo = TasksProvider->TryGetTask(Relation.RelativeId);
 			if (TaskInfo)
 			{
