[Insights] Fix crash when using the "Go To Task" option on a group node.

diff --git a/Engine/Source/Developer/TraceInsights/Private/Insights/TaskGraphProfiler/Widgets/STaskTableTreeView.cpp b/Engine/Source/Developer/TraceInsights/Private/Insights/TaskGraphProfiler/Widgets/STaskTableTreeView.cpp
--- a/Engine/Source/Developer/TraceInsights/Private/Insights/TaskGraphProfiler/Widgets/STaskTableTreeView.cpp
+++ b/Engine/Source/Developer/TraceInsights/Private/Insights/TaskGraphProfiler/Widgets/STaskTableTreeView.cpp
@@ -410,7 +410,22 @@ bool STaskTableTreeView::TimestampOptions_IsEnabled() const
 
 bool STaskTableTreeView::ContextMenu_GoToTask_CanExecute() const
 {
-	return TreeView->GetNumItemsSelected() == 1;
+	TArray<FTableTreeNodePtr> SelectedItems;
+	TreeView->GetSelectedItems(SelectedItems);
+
+	if (SelectedItems.Num() != 1)
+	{
+		return false;
+	}
+
+	FTaskNodePtr SelectedTask = StaticCastSharedPtr<FTaskNode>(SelectedItems[0]);
+
+	if (!SelectedTask.IsValid() || SelectedTask->IsGroup())
+	{
+		return false;
+	}
+
+	return true;
 }
 
 ////////////////////////////////////////////////////////////////////////////////////////////////////
@@ -426,7 +441,12 @@ void STaskTableTreeView::ContextMenu_GoToTask_Execute()
 	}
 
 	FTaskNodePtr SelectedTask = StaticCastSharedPtr<FTaskNode>(SelectedItems[0]);
-	const FTaskEntry* TaskEntry = SelectedTask->GetTask();
+	const FTaskEntry* TaskEntry = SelectedTask.IsValid() ? SelectedTask->GetTask() : nullptr;
+
+	if (TaskEntry == nullptr)
+	{
+		return;
+	}
 
 	TSharedPtr<STimingProfilerWindow> TimingWindow = FTimingProfilerManager::Get()->GetProfilerWindow();
 	if (!TimingWindow.IsValid())
