[Insights] Improve GetEvent for the Task Track to help with tooltip vizualization for small events.

diff --git a/Engine/Source/Developer/TraceInsights/Private/Insights/TaskGraphProfiler/ViewModels/TaskTimingTrack.cpp b/Engine/Source/Developer/TraceInsights/Private/Insights/TaskGraphProfiler/ViewModels/TaskTimingTrack.cpp
--- a/Engine/Source/Developer/TraceInsights/Private/Insights/TaskGraphProfiler/ViewModels/TaskTimingTrack.cpp
+++ b/Engine/Source/Developer/TraceInsights/Private/Insights/TaskGraphProfiler/ViewModels/TaskTimingTrack.cpp
@@ -837,27 +837,34 @@ const TSharedPtr<const ITimingEvent> FTaskTimingTrack::GetEvent(float InPosX, fl
 			return TimingEvent;
 		}
 
+		const double SecondsPerPixel = 1.0 / Viewport.GetScaleX();
+
+		auto IsEventTimeBetween = [SecondsPerPixel, EventTime](double TimeA, double TimeB)
+		{
+			return EventTime >= TimeA - 2 * SecondsPerPixel && EventTime <= TimeB + 2 * SecondsPerPixel;
+		};
+
 		if (Depth == 0)
 		{
-			if (EventTime >= Task->CreatedTimestamp && EventTime < Task->LaunchedTimestamp)
+			if (IsEventTimeBetween(Task->CreatedTimestamp, Task->LaunchedTimestamp))
 			{
 				TimingEvent = MakeShared<FTaskTrackEvent>(SharedThis(this), Task->CreatedTimestamp, Task->LaunchedTimestamp, 0, ETaskEventType::Created);
 			}
-			else if (EventTime >= Task->LaunchedTimestamp && EventTime < Task->ScheduledTimestamp)
+			else if (IsEventTimeBetween(Task->LaunchedTimestamp, Task->ScheduledTimestamp))
 			{
 				TimingEvent = MakeShared<FTaskTrackEvent>(SharedThis(this), Task->LaunchedTimestamp, Task->ScheduledTimestamp, 0, ETaskEventType::Launched);
 			}
-			else if (EventTime >= Task->ScheduledTimestamp && EventTime < Task->StartedTimestamp)
+			else if (IsEventTimeBetween(Task->ScheduledTimestamp, Task->StartedTimestamp))
 			{
 				TimingEvent = MakeShared<FTaskTrackEvent>(SharedThis(this), Task->ScheduledTimestamp, Task->StartedTimestamp, 0, ETaskEventType::Scheduled);
 			}
-			else if (EventTime >= Task->StartedTimestamp && EventTime < Task->FinishedTimestamp)
+			else if (IsEventTimeBetween(Task->FinishedTimestamp, Task->CompletedTimestamp))
 			{
-				TimingEvent = MakeShared<FTaskTrackEvent>(SharedThis(this), Task->StartedTimestamp, Task->FinishedTimestamp, 0, ETaskEventType::Started);
+				TimingEvent = MakeShared<FTaskTrackEvent>(SharedThis(this), Task->FinishedTimestamp, Task->CompletedTimestamp, 0, ETaskEventType::Finished);
 			}
-			else if (EventTime >= Task->FinishedTimestamp && EventTime < Task->CompletedTimestamp)
+			else if (IsEventTimeBetween(Task->StartedTimestamp, Task->FinishedTimestamp))
 			{
-				TimingEvent = MakeShared<FTaskTrackEvent>(SharedThis(this), Task->FinishedTimestamp, Task->CompletedTimestamp, 0, ETaskEventType::Finished);
+				TimingEvent = MakeShared<FTaskTrackEvent>(SharedThis(this), Task->StartedTimestamp, Task->FinishedTimestamp, 0, ETaskEventType::Started);
 			}
 
 			if (TimingEvent.IsValid())
@@ -867,7 +874,7 @@ const TSharedPtr<const ITimingEvent> FTaskTimingTrack::GetEvent(float InPosX, fl
 		}
 		else if(bShowDetailInfoOnTaskTrack)
 		{
-			auto GetEventFromRelations = [&TasksProvider, this, EventTime](const TArray<TraceServices::FTaskInfo::FRelationInfo>& Relations, int32 Depth, ETaskEventType EventType) -> TSharedPtr<FTaskTrackEvent>
+			auto GetEventFromRelations = [&TasksProvider, this, EventTime, SecondsPerPixel](const TArray<TraceServices::FTaskInfo::FRelationInfo>& Relations, int32 Depth, ETaskEventType EventType) -> TSharedPtr<FTaskTrackEvent>
 			{
 				if (Relations.Num() < Depth)
 				{
@@ -876,7 +883,7 @@ const TSharedPtr<const ITimingEvent> FTaskTimingTrack::GetEvent(float InPosX, fl
 				const TraceServices::FTaskInfo::FRelationInfo& Relation = Relations[Depth - 1];
 				{
 					const TraceServices::FTaskInfo* TaskInfo = TasksProvider->TryGetTask(Relation.RelativeId);
-					if (TaskInfo && EventTime >= TaskInfo->StartedTimestamp && EventTime <= TaskInfo->FinishedTimestamp)
+					if (TaskInfo && EventTime >= TaskInfo->StartedTimestamp - 2 * SecondsPerPixel && EventTime <= TaskInfo->FinishedTimestamp + 2 * SecondsPerPixel)
 					{
 						TSharedPtr<FTaskTrackEvent> NewTimingEvent = MakeShared<FTaskTrackEvent>(SharedThis(this), TaskInfo->StartedTimestamp, TaskInfo->FinishedTimestamp, Depth, EventType);
 						NewTimingEvent->SetTaskId(TaskInfo->Id);
