Unreal Insights: Stats group information is now used for meta group column

diff --git a/Engine/Source/Developer/TraceInsights/Private/Insights/Widgets/SStatsView.cpp b/Engine/Source/Developer/TraceInsights/Private/Insights/Widgets/SStatsView.cpp
--- a/Engine/Source/Developer/TraceInsights/Private/Insights/Widgets/SStatsView.cpp
+++ b/Engine/Source/Developer/TraceInsights/Private/Insights/Widgets/SStatsView.cpp
@@ -1864,7 +1864,8 @@ void SStatsView::RebuildTree(bool bResync)
 				if (!NodePtr)
 				{
 					FName Name(Counter.GetName());
-					const FName Group = ((Counter.GetDisplayHint() == TraceServices::CounterDisplayHint_Memory) ? MemoryGroup :
+					const FName Group = Counter.GetGroup() ? Counter.GetGroup() :
+										((Counter.GetDisplayHint() == TraceServices::CounterDisplayHint_Memory) ? MemoryGroup :
 										  Counter.IsFloatingPoint() ? MiscFloatGroup : MiscInt64Group);
 					const EStatsNodeType Type = EStatsNodeType::Counter;
 					const EStatsNodeDataType DataType = Counter.IsFloatingPoint() ? EStatsNodeDataType::Double : EStatsNodeDataType::Int64;
