[Insights] Fixed crash (check) when closing Unreal Insights while having a minor tab undocked in Timing Insights, Memory Insights, Asset Loading Insights or Networking Insights.

diff --git a/Engine/Source/Developer/TraceInsights/Private/Insights/LoadingProfiler/Widgets/SLoadingProfilerWindow.cpp b/Engine/Source/Developer/TraceInsights/Private/Insights/LoadingProfiler/Widgets/SLoadingProfilerWindow.cpp
--- a/Engine/Source/Developer/TraceInsights/Private/Insights/LoadingProfiler/Widgets/SLoadingProfilerWindow.cpp
+++ b/Engine/Source/Developer/TraceInsights/Private/Insights/LoadingProfiler/Widgets/SLoadingProfilerWindow.cpp
@@ -60,41 +60,25 @@ SLoadingProfilerWindow::SLoadingProfilerWindow()
 
 SLoadingProfilerWindow::~SLoadingProfilerWindow()
 {
-	if (RequestsTreeView)
+	// Close all tabs.
+	TArray<TSharedPtr<SDockTab>> LocalOpenTabs;
+	for (TSharedPtr<SDockTab>& Tab : OpenTabs)
 	{
-		HideTab(FLoadingProfilerTabs::RequestsTreeViewID);
-		check(RequestsTreeView == nullptr);
+		LocalOpenTabs.Add(Tab);
 	}
-
-	if (ExportDetailsTreeView)
-	{
-		HideTab(FLoadingProfilerTabs::ExportDetailsTreeViewID);
-		check(ExportDetailsTreeView == nullptr);
-	}
-
-	if (PackageDetailsTreeView)
+	for (TSharedPtr<SDockTab>& Tab : LocalOpenTabs)
 	{
-		HideTab(FLoadingProfilerTabs::PackageDetailsTreeViewID);
-		check(PackageDetailsTreeView == nullptr);
-	}
-
-	if (ObjectTypeAggregationTreeView)
-	{
-		HideTab(FLoadingProfilerTabs::ObjectTypeAggregationTreeViewID);
-		check(ObjectTypeAggregationTreeView == nullptr);
-	}
-
-	if (EventAggregationTreeView)
-	{
-		HideTab(FLoadingProfilerTabs::EventAggregationTreeViewID);
-		check(EventAggregationTreeView == nullptr);
+		Tab->RequestCloseTab();
 	}
+	LocalOpenTabs.Reset();
+	check(OpenTabs.IsEmpty());
 
-	if (TimingView)
-	{
-		HideTab(FLoadingProfilerTabs::TimingViewID);
-		check(TimingView == nullptr);
-	}
+	check(RequestsTreeView == nullptr);
+	check(ExportDetailsTreeView == nullptr);
+	check(PackageDetailsTreeView == nullptr);
+	check(ObjectTypeAggregationTreeView == nullptr);
+	check(EventAggregationTreeView == nullptr);
+	check(TimingView == nullptr);
 
 #if WITH_EDITOR
 	if (DurationActive > 0.0f && FEngineAnalytics::IsAvailable())
@@ -296,6 +280,7 @@ TSharedRef<SDockTab> SLoadingProfilerWindow::SpawnTab_TimingView(const FSpawnTab
 	TimingView->SelectTimeInterval(SelectionStartTime, SelectionEndTime - SelectionStartTime);
 
 	DockTab->SetOnTabClosed(SDockTab::FOnTabClosedCallback::CreateRaw(this, &SLoadingProfilerWindow::OnTimingViewTabClosed));
+	OpenTabs.Add(DockTab);
 
 	return DockTab;
 }
@@ -311,6 +296,8 @@ void SLoadingProfilerWindow::OnTimingViewTabClosed(TSharedRef<SDockTab> TabBeing
 		TimingView->OnSelectionChanged().RemoveAll(this);
 		TimingView = nullptr;
 	}
+
+	OpenTabs.Remove(TabBeingClosed);
 }
 
 ////////////////////////////////////////////////////////////////////////////////////////////////////
@@ -334,6 +321,7 @@ TSharedRef<SDockTab> SLoadingProfilerWindow::SpawnTab_EventAggregationTreeView(c
 	UpdateEventAggregationTreeView();
 
 	DockTab->SetOnTabClosed(SDockTab::FOnTabClosedCallback::CreateRaw(this, &SLoadingProfilerWindow::OnEventAggregationTreeViewTabClosed));
+	OpenTabs.Add(DockTab);
 
 	return DockTab;
 }
@@ -345,6 +333,8 @@ void SLoadingProfilerWindow::OnEventAggregationTreeViewTabClosed(TSharedRef<SDoc
 {
 	FLoadingProfilerManager::Get()->SetEventAggregationTreeViewVisible(false);
 	EventAggregationTreeView = nullptr;
+
+	OpenTabs.Remove(TabBeingClosed);
 }
 
 ////////////////////////////////////////////////////////////////////////////////////////////////////
@@ -368,6 +358,7 @@ TSharedRef<SDockTab> SLoadingProfilerWindow::SpawnTab_ObjectTypeAggregationTreeV
 	UpdateObjectTypeAggregationTreeView();
 
 	DockTab->SetOnTabClosed(SDockTab::FOnTabClosedCallback::CreateRaw(this, &SLoadingProfilerWindow::OnObjectTypeAggregationTreeViewTabClosed));
+	OpenTabs.Add(DockTab);
 
 	return DockTab;
 }
@@ -379,6 +370,8 @@ void SLoadingProfilerWindow::OnObjectTypeAggregationTreeViewTabClosed(TSharedRef
 {
 	FLoadingProfilerManager::Get()->SetObjectTypeAggregationTreeViewVisible(false);
 	ObjectTypeAggregationTreeView = nullptr;
+
+	OpenTabs.Remove(TabBeingClosed);
 }
 
 ////////////////////////////////////////////////////////////////////////////////////////////////////
@@ -402,6 +395,7 @@ TSharedRef<SDockTab> SLoadingProfilerWindow::SpawnTab_PackageDetailsTreeView(con
 	UpdatePackageDetailsTreeView();
 
 	DockTab->SetOnTabClosed(SDockTab::FOnTabClosedCallback::CreateRaw(this, &SLoadingProfilerWindow::OnPackageDetailsTreeViewTabClosed));
+	OpenTabs.Add(DockTab);
 
 	return DockTab;
 }
@@ -413,6 +407,8 @@ void SLoadingProfilerWindow::OnPackageDetailsTreeViewTabClosed(TSharedRef<SDockT
 {
 	FLoadingProfilerManager::Get()->SetPackageDetailsTreeViewVisible(false);
 	PackageDetailsTreeView = nullptr;
+
+	OpenTabs.Remove(TabBeingClosed);
 }
 
 ////////////////////////////////////////////////////////////////////////////////////////////////////
@@ -436,6 +432,7 @@ TSharedRef<SDockTab> SLoadingProfilerWindow::SpawnTab_ExportDetailsTreeView(cons
 	UpdateExportDetailsTreeView();
 
 	DockTab->SetOnTabClosed(SDockTab::FOnTabClosedCallback::CreateRaw(this, &SLoadingProfilerWindow::OnExportDetailsTreeViewTabClosed));
+	OpenTabs.Add(DockTab);
 
 	return DockTab;
 }
@@ -447,6 +444,8 @@ void SLoadingProfilerWindow::OnExportDetailsTreeViewTabClosed(TSharedRef<SDockTa
 {
 	FLoadingProfilerManager::Get()->SetExportDetailsTreeViewVisible(false);
 	ExportDetailsTreeView = nullptr;
+
+	OpenTabs.Remove(TabBeingClosed);
 }
 
 ////////////////////////////////////////////////////////////////////////////////////////////////////
@@ -470,6 +469,7 @@ TSharedRef<SDockTab> SLoadingProfilerWindow::SpawnTab_RequestsTreeView(const FSp
 	UpdateRequestsTreeView();
 
 	DockTab->SetOnTabClosed(SDockTab::FOnTabClosedCallback::CreateRaw(this, &SLoadingProfilerWindow::OnRequestsTreeViewTabClosed));
+	OpenTabs.Add(DockTab);
 
 	return DockTab;
 }
@@ -481,6 +481,8 @@ void SLoadingProfilerWindow::OnRequestsTreeViewTabClosed(TSharedRef<SDockTab> Ta
 {
 	FLoadingProfilerManager::Get()->SetRequestsTreeViewVisible(false);
 	RequestsTreeView = nullptr;
+
+	OpenTabs.Remove(TabBeingClosed);
 }
 
 ////////////////////////////////////////////////////////////////////////////////////////////////////
diff --git a/Engine/Source/Developer/TraceInsights/Private/Insights/LoadingProfiler/Widgets/SLoadingProfilerWindow.h b/Engine/Source/Developer/TraceInsights/Private/Insights/LoadingProfiler/Widgets/SLoadingProfilerWindow.h
--- a/Engine/Source/Developer/TraceInsights/Private/Insights/LoadingProfiler/Widgets/SLoadingProfilerWindow.h
+++ b/Engine/Source/Developer/TraceInsights/Private/Insights/LoadingProfiler/Widgets/SLoadingProfilerWindow.h
@@ -173,6 +173,12 @@ private:
 	void OnTimeSelectionChanged(Insights::ETimeChangedFlags InFlags, double InStartTime, double InEndTime);
 
 private:
+	/** Holds the tab manager that manages the front-end's tabs. */
+	TSharedPtr<FTabManager> TabManager;
+
+	/** All tabs owned by this window */
+	TSet<TSharedPtr<SDockTab>> OpenTabs;
+
 	/** The Timing view (multi-track) widget */
 	TSharedPtr<STimingView> TimingView;
 
@@ -191,9 +197,6 @@ private:
 	/** The Requests tree view widget */
 	TSharedPtr<Insights::SUntypedTableTreeView> RequestsTreeView;
 
-	/** Holds the tab manager that manages the front-end's tabs. */
-	TSharedPtr<FTabManager> TabManager;
-
 	/** The handle to the active update duration tick */
 	TWeakPtr<FActiveTimerHandle> ActiveTimerHandle;
 
diff --git a/Engine/Source/Developer/TraceInsights/Private/Insights/MemoryProfiler/Widgets/SMemoryProfilerWindow.cpp b/Engine/Source/Developer/TraceInsights/Private/Insights/MemoryProfiler/Widgets/SMemoryProfilerWindow.cpp
--- a/Engine/Source/Developer/TraceInsights/Private/Insights/MemoryProfiler/Widgets/SMemoryProfilerWindow.cpp
+++ b/Engine/Source/Developer/TraceInsights/Private/Insights/MemoryProfiler/Widgets/SMemoryProfilerWindow.cpp
@@ -64,29 +64,23 @@ SMemoryProfilerWindow::SMemoryProfilerWindow()
 
 SMemoryProfilerWindow::~SMemoryProfilerWindow()
 {
-	if (ModulesView)
-	{
-		HideTab(FMemoryProfilerTabs::ModulesViewID);
-		check(ModulesView == nullptr);
-	}
-
-	if (MemTagTreeView)
+	// Close all tabs.
+	TArray<TSharedPtr<SDockTab>> LocalOpenTabs;
+	for (TSharedPtr<SDockTab>& Tab : OpenTabs)
 	{
-		HideTab(FMemoryProfilerTabs::MemTagTreeViewID);
-		check(MemTagTreeView == nullptr);
+		LocalOpenTabs.Add(Tab);
 	}
-
-	if (MemInvestigationView)
+	for (TSharedPtr<SDockTab>& Tab : LocalOpenTabs)
 	{
-		HideTab(FMemoryProfilerTabs::MemInvestigationViewID);
-		check(MemInvestigationView == nullptr);
+		Tab->RequestCloseTab();
 	}
+	LocalOpenTabs.Reset();
+	check(OpenTabs.IsEmpty());
 
-	if (TimingView)
-	{
-		HideTab(FMemoryProfilerTabs::TimingViewID);
-		check(TimingView == nullptr);
-	}
+	check(ModulesView == nullptr);
+	check(MemTagTreeView == nullptr);
+	check(MemInvestigationView == nullptr);
+	check(TimingView == nullptr);
 
 #if WITH_EDITOR
 	if (DurationActive > 0.0f && FEngineAnalytics::IsAvailable())
@@ -147,6 +141,7 @@ TSharedRef<SDockTab> SMemoryProfilerWindow::SpawnTab_TimingView(const FSpawnTabA
 	TimingView->HideAllDefaultTracks();
 
 	DockTab->SetOnTabClosed(SDockTab::FOnTabClosedCallback::CreateRaw(this, &SMemoryProfilerWindow::OnTimingViewTabClosed));
+	OpenTabs.Add(DockTab);
 
 	return DockTab;
 }
@@ -167,6 +162,8 @@ void SMemoryProfilerWindow::OnTimingViewTabClosed(TSharedRef<SDockTab> TabBeingC
 
 	IModularFeatures::Get().UnregisterModularFeature(Insights::TimingViewExtenderFeatureName, &SharedState.Get());
 	SharedState->SetTimingView(nullptr);
+
+	OpenTabs.Remove(TabBeingClosed);
 }
 
 ////////////////////////////////////////////////////////////////////////////////////////////////////
@@ -184,6 +181,7 @@ TSharedRef<SDockTab> SMemoryProfilerWindow::SpawnTab_MemInvestigationView(const
 		];
 
 	DockTab->SetOnTabClosed(SDockTab::FOnTabClosedCallback::CreateRaw(this, &SMemoryProfilerWindow::OnMemInvestigationViewTabClosed));
+	OpenTabs.Add(DockTab);
 
 	return DockTab;
 }
@@ -195,6 +193,8 @@ void SMemoryProfilerWindow::OnMemInvestigationViewTabClosed(TSharedRef<SDockTab>
 {
 	FMemoryProfilerManager::Get()->SetMemInvestigationViewVisible(false);
 	MemInvestigationView = nullptr;
+
+	OpenTabs.Remove(TabBeingClosed);
 }
 
 ////////////////////////////////////////////////////////////////////////////////////////////////////
@@ -212,6 +212,7 @@ TSharedRef<SDockTab> SMemoryProfilerWindow::SpawnTab_MemTagTreeView(const FSpawn
 		];
 
 	DockTab->SetOnTabClosed(SDockTab::FOnTabClosedCallback::CreateRaw(this, &SMemoryProfilerWindow::OnMemTagTreeViewTabClosed));
+	OpenTabs.Add(DockTab);
 
 	return DockTab;
 }
@@ -223,6 +224,8 @@ void SMemoryProfilerWindow::OnMemTagTreeViewTabClosed(TSharedRef<SDockTab> TabBe
 {
 	FMemoryProfilerManager::Get()->SetMemTagTreeViewVisible(false);
 	MemTagTreeView = nullptr;
+
+	OpenTabs.Remove(TabBeingClosed);
 }
 
 ////////////////////////////////////////////////////////////////////////////////////////////////////
@@ -250,6 +253,7 @@ TSharedRef<SDockTab> SMemoryProfilerWindow::SpawnTab_MemAllocTableTreeView(const
 	MemAllocTableTreeViews.Add(MemAllocTableTreeView);
 
 	DockTab->SetOnTabClosed(SDockTab::FOnTabClosedCallback::CreateRaw(this, &SMemoryProfilerWindow::OnMemAllocTableTreeViewTabClosed));
+	OpenTabs.Add(DockTab);
 
 	return DockTab;
 }
@@ -294,6 +298,8 @@ void SMemoryProfilerWindow::OnMemAllocTableTreeViewTabClosed(TSharedRef<SDockTab
 
 	MemAllocTableTreeView->OnClose();
 	MemAllocTableTreeViews.Remove(MemAllocTableTreeView);
+
+	OpenTabs.Remove(TabBeingClosed);
 }
 
 ////////////////////////////////////////////////////////////////////////////////////////////////////
@@ -371,6 +377,7 @@ TSharedRef<SDockTab> SMemoryProfilerWindow::SpawnTab_ModulesView(const FSpawnTab
 		];
 
 	DockTab->SetOnTabClosed(SDockTab::FOnTabClosedCallback::CreateRaw(this, &SMemoryProfilerWindow::OnModulesViewClosed));
+	OpenTabs.Add(DockTab);
 
 	return DockTab;
 }
@@ -382,6 +389,8 @@ void SMemoryProfilerWindow::OnModulesViewClosed(TSharedRef<SDockTab> TabBeingClo
 {
 	FMemoryProfilerManager::Get()->SetModulesViewVisible(false);
 	ModulesView = nullptr;
+
+	OpenTabs.Remove(TabBeingClosed);
 }
 
 ////////////////////////////////////////////////////////////////////////////////////////////////////
diff --git a/Engine/Source/Developer/TraceInsights/Private/Insights/MemoryProfiler/Widgets/SMemoryProfilerWindow.h b/Engine/Source/Developer/TraceInsights/Private/Insights/MemoryProfiler/Widgets/SMemoryProfilerWindow.h
--- a/Engine/Source/Developer/TraceInsights/Private/Insights/MemoryProfiler/Widgets/SMemoryProfilerWindow.h
+++ b/Engine/Source/Developer/TraceInsights/Private/Insights/MemoryProfiler/Widgets/SMemoryProfilerWindow.h
@@ -187,6 +187,14 @@ private:
 private:
 	TSharedRef<FMemorySharedState> SharedState;
 
+	/** Holds the tab manager that manages the front-end's tabs. */
+	TSharedPtr<FTabManager> TabManager;
+
+	TSharedPtr<FWorkspaceItem> AppMenuGroup;
+
+	/** All tabs owned by this window */
+	TSet<TSharedPtr<SDockTab>> OpenTabs;
+
 	/** The Timing view (multi-track) widget */
 	TSharedPtr<STimingView> TimingView;
 
@@ -207,11 +215,6 @@ private:
 	const int32 MaxMemAllocTableTreeViews = 4;
 	int32 LastMemAllocTableTreeViewIndex = -1;
 
-	/** Holds the tab manager that manages the front-end's tabs. */
-	TSharedPtr<FTabManager> TabManager;
-
-	TSharedPtr<FWorkspaceItem> AppMenuGroup;
-
 	/** The handle to the active update duration tick */
 	TWeakPtr<FActiveTimerHandle> ActiveTimerHandle;
 
diff --git a/Engine/Source/Developer/TraceInsights/Private/Insights/NetworkingProfiler/Widgets/SNetworkingProfilerWindow.cpp b/Engine/Source/Developer/TraceInsights/Private/Insights/NetworkingProfiler/Widgets/SNetworkingProfilerWindow.cpp
--- a/Engine/Source/Developer/TraceInsights/Private/Insights/NetworkingProfiler/Widgets/SNetworkingProfilerWindow.cpp
+++ b/Engine/Source/Developer/TraceInsights/Private/Insights/NetworkingProfiler/Widgets/SNetworkingProfilerWindow.cpp
@@ -61,23 +61,22 @@ SNetworkingProfilerWindow::SNetworkingProfilerWindow()
 
 SNetworkingProfilerWindow::~SNetworkingProfilerWindow()
 {
-	if (NetStatsView)
+	// Close all tabs.
+	TArray<TSharedPtr<SDockTab>> LocalOpenTabs;
+	for (TSharedPtr<SDockTab>& Tab : OpenTabs)
 	{
-		HideTab(FNetworkingProfilerTabs::NetStatsViewID);
-		check(NetStatsView == nullptr);
+		LocalOpenTabs.Add(Tab);
 	}
-
-	if (PacketContentView)
+	for (TSharedPtr<SDockTab>& Tab : LocalOpenTabs)
 	{
-		HideTab(FNetworkingProfilerTabs::PacketContentViewID);
-		check(PacketContentView == nullptr);
+		Tab->RequestCloseTab();
 	}
+	LocalOpenTabs.Reset();
+	check(OpenTabs.IsEmpty());
 
-	if (PacketView)
-	{
-		HideTab(FNetworkingProfilerTabs::PacketViewID);
-		check(PacketView == nullptr);
-	}
+	check(NetStatsView == nullptr);
+	check(PacketContentView == nullptr);
+	check(PacketView == nullptr);
 
 #if WITH_EDITOR
 	if (DurationActive > 0.0f && FEngineAnalytics::IsAvailable())
@@ -134,6 +133,7 @@ TSharedRef<SDockTab> SNetworkingProfilerWindow::SpawnTab_PacketView(const FSpawn
 		];
 
 	DockTab->SetOnTabClosed(SDockTab::FOnTabClosedCallback::CreateRaw(this, &SNetworkingProfilerWindow::OnPacketViewTabClosed));
+	OpenTabs.Add(DockTab);
 
 	return DockTab;
 }
@@ -144,6 +144,8 @@ END_SLATE_FUNCTION_BUILD_OPTIMIZATION
 void SNetworkingProfilerWindow::OnPacketViewTabClosed(TSharedRef<SDockTab> TabBeingClosed)
 {
 	PacketView = nullptr;
+
+	OpenTabs.Remove(TabBeingClosed);
 }
 
 ////////////////////////////////////////////////////////////////////////////////////////////////////
@@ -159,6 +161,7 @@ TSharedRef<SDockTab> SNetworkingProfilerWindow::SpawnTab_PacketContentView(const
 		];
 
 	DockTab->SetOnTabClosed(SDockTab::FOnTabClosedCallback::CreateRaw(this, &SNetworkingProfilerWindow::OnPacketContentViewTabClosed));
+	OpenTabs.Add(DockTab);
 
 	return DockTab;
 }
@@ -169,6 +172,8 @@ END_SLATE_FUNCTION_BUILD_OPTIMIZATION
 void SNetworkingProfilerWindow::OnPacketContentViewTabClosed(TSharedRef<SDockTab> TabBeingClosed)
 {
 	PacketContentView = nullptr;
+
+	OpenTabs.Remove(TabBeingClosed);
 }
 
 ////////////////////////////////////////////////////////////////////////////////////////////////////
@@ -184,6 +189,7 @@ TSharedRef<SDockTab> SNetworkingProfilerWindow::SpawnTab_NetStatsView(const FSpa
 		];
 
 	DockTab->SetOnTabClosed(SDockTab::FOnTabClosedCallback::CreateRaw(this, &SNetworkingProfilerWindow::OnNetStatsViewTabClosed));
+	OpenTabs.Add(DockTab);
 
 	return DockTab;
 }
@@ -194,6 +200,8 @@ END_SLATE_FUNCTION_BUILD_OPTIMIZATION
 void SNetworkingProfilerWindow::OnNetStatsViewTabClosed(TSharedRef<SDockTab> TabBeingClosed)
 {
 	NetStatsView = nullptr;
+
+	OpenTabs.Remove(TabBeingClosed);
 }
 
 ////////////////////////////////////////////////////////////////////////////////////////////////////
diff --git a/Engine/Source/Developer/TraceInsights/Private/Insights/NetworkingProfiler/Widgets/SNetworkingProfilerWindow.h b/Engine/Source/Developer/TraceInsights/Private/Insights/NetworkingProfiler/Widgets/SNetworkingProfilerWindow.h
--- a/Engine/Source/Developer/TraceInsights/Private/Insights/NetworkingProfiler/Widgets/SNetworkingProfilerWindow.h
+++ b/Engine/Source/Developer/TraceInsights/Private/Insights/NetworkingProfiler/Widgets/SNetworkingProfilerWindow.h
@@ -274,12 +274,15 @@ private:\
 	//////////////////////////////////////////////////
 
 private:
-	/** Commandlist used in the window (Maps commands to window specific actions) */
-	TSharedPtr<FUICommandList> CommandList;
-
 	/** Holds the tab manager that manages the front-end's tabs. */
 	TSharedPtr<FTabManager> TabManager;
 
+	/** All tabs owned by this window */
+	TSet<TSharedPtr<SDockTab>> OpenTabs;
+
+	/** Commandlist used in the window (Maps commands to window specific actions) */
+	TSharedPtr<FUICommandList> CommandList;
+
 	/** The Packet View widget */
 	TSharedPtr<SPacketView> PacketView;
 
diff --git a/Engine/Source/Developer/TraceInsights/Private/Insights/Widgets/STimingProfilerWindow.cpp b/Engine/Source/Developer/TraceInsights/Private/Insights/Widgets/STimingProfilerWindow.cpp
--- a/Engine/Source/Developer/TraceInsights/Private/Insights/Widgets/STimingProfilerWindow.cpp
+++ b/Engine/Source/Developer/TraceInsights/Private/Insights/Widgets/STimingProfilerWindow.cpp
@@ -66,47 +66,26 @@ STimingProfilerWindow::STimingProfilerWindow()
 
 STimingProfilerWindow::~STimingProfilerWindow()
 {
-	if (LogView)
-	{
-		HideTab(FTimingProfilerTabs::LogViewID);
-		check(LogView == nullptr);
-	}
-
-	if (StatsView)
-	{
-		HideTab(FTimingProfilerTabs::StatsCountersID);
-		check(StatsView == nullptr);
-	}
-
-	if (CalleesTreeView)
-	{
-		HideTab(FTimingProfilerTabs::CalleesID);
-		check(CalleesTreeView == nullptr);
-	}
-
-	if (CallersTreeView)
-	{
-		HideTab(FTimingProfilerTabs::CallersID);
-		check(CallersTreeView == nullptr);
-	}
-
-	if (TimersView)
+	// Close all tabs.
+	TArray<TSharedPtr<SDockTab>> LocalOpenTabs;
+	for (TSharedPtr<SDockTab>& Tab : OpenTabs)
 	{
-		HideTab(FTimingProfilerTabs::TimersID);
-		check(TimersView == nullptr);
+		LocalOpenTabs.Add(Tab);
 	}
-
-	if (TimingView)
+	for (TSharedPtr<SDockTab>& Tab : LocalOpenTabs)
 	{
-		HideTab(FTimingProfilerTabs::TimingViewID);
-		check(TimingView == nullptr);
+		Tab->RequestCloseTab();
 	}
+	LocalOpenTabs.Reset();
+	check(OpenTabs.IsEmpty());
 
-	if (FrameTrack)
-	{
-		HideTab(FTimingProfilerTabs::FramesTrackID);
-		check(FrameTrack == nullptr);
-	}
+	check(LogView == nullptr);
+	check(StatsView == nullptr);
+	check(CalleesTreeView == nullptr);
+	check(CallersTreeView == nullptr);
+	check(TimersView == nullptr);
+	check(TimingView == nullptr);
+	check(FrameTrack == nullptr);
 
 #if WITH_EDITOR
 	if (DurationActive > 0.0f && FEngineAnalytics::IsAvailable())
@@ -171,6 +150,7 @@ TSharedRef<SDockTab> STimingProfilerWindow::SpawnTab_FramesTrack(const FSpawnTab
 		];
 
 	DockTab->SetOnTabClosed(SDockTab::FOnTabClosedCallback::CreateRaw(this, &STimingProfilerWindow::OnFramesTrackTabClosed));
+	OpenTabs.Add(DockTab);
 
 	return DockTab;
 }
@@ -182,6 +162,8 @@ void STimingProfilerWindow::OnFramesTrackTabClosed(TSharedRef<SDockTab> TabBeing
 {
 	FTimingProfilerManager::Get()->SetFramesTrackVisible(false);
 	FrameTrack = nullptr;
+
+	OpenTabs.Remove(TabBeingClosed);
 }
 
 ////////////////////////////////////////////////////////////////////////////////////////////////////
@@ -205,6 +187,7 @@ TSharedRef<SDockTab> STimingProfilerWindow::SpawnTab_TimingView(const FSpawnTabA
 	TimingView->SelectTimeInterval(SelectionStartTime, SelectionEndTime - SelectionStartTime);
 
 	DockTab->SetOnTabClosed(SDockTab::FOnTabClosedCallback::CreateRaw(this, &STimingProfilerWindow::OnTimingViewTabClosed));
+	OpenTabs.Add(DockTab);
 
 	return DockTab;
 }
@@ -220,6 +203,8 @@ void STimingProfilerWindow::OnTimingViewTabClosed(TSharedRef<SDockTab> TabBeingC
 		TimingView->OnSelectionChanged().RemoveAll(this);
 		TimingView = nullptr;
 	}
+
+	OpenTabs.Remove(TabBeingClosed);
 }
 
 ////////////////////////////////////////////////////////////////////////////////////////////////////
@@ -237,6 +222,7 @@ TSharedRef<SDockTab> STimingProfilerWindow::SpawnTab_Timers(const FSpawnTabArgs&
 		];
 
 	DockTab->SetOnTabClosed(SDockTab::FOnTabClosedCallback::CreateRaw(this, &STimingProfilerWindow::OnTimersTabClosed));
+	OpenTabs.Add(DockTab);
 
 	return DockTab;
 }
@@ -248,6 +234,8 @@ void STimingProfilerWindow::OnTimersTabClosed(TSharedRef<SDockTab> TabBeingClose
 {
 	FTimingProfilerManager::Get()->SetTimersViewVisible(false);
 	TimersView = nullptr;
+
+	OpenTabs.Remove(TabBeingClosed);
 }
 
 ////////////////////////////////////////////////////////////////////////////////////////////////////
@@ -265,6 +253,7 @@ TSharedRef<SDockTab> STimingProfilerWindow::SpawnTab_Callers(const FSpawnTabArgs
 		];
 
 	DockTab->SetOnTabClosed(SDockTab::FOnTabClosedCallback::CreateRaw(this, &STimingProfilerWindow::OnCallersTabClosed));
+	OpenTabs.Add(DockTab);
 
 	return DockTab;
 }
@@ -276,6 +265,8 @@ void STimingProfilerWindow::OnCallersTabClosed(TSharedRef<SDockTab> TabBeingClos
 {
 	FTimingProfilerManager::Get()->SetCallersTreeViewVisible(false);
 	CallersTreeView = nullptr;
+
+	OpenTabs.Remove(TabBeingClosed);
 }
 
 ////////////////////////////////////////////////////////////////////////////////////////////////////
@@ -293,6 +284,7 @@ TSharedRef<SDockTab> STimingProfilerWindow::SpawnTab_Callees(const FSpawnTabArgs
 		];
 
 	DockTab->SetOnTabClosed(SDockTab::FOnTabClosedCallback::CreateRaw(this, &STimingProfilerWindow::OnCalleesTabClosed));
+	OpenTabs.Add(DockTab);
 
 	return DockTab;
 }
@@ -304,6 +296,8 @@ void STimingProfilerWindow::OnCalleesTabClosed(TSharedRef<SDockTab> TabBeingClos
 {
 	FTimingProfilerManager::Get()->SetCalleesTreeViewVisible(false);
 	CalleesTreeView = nullptr;
+
+	OpenTabs.Remove(TabBeingClosed);
 }
 
 ////////////////////////////////////////////////////////////////////////////////////////////////////
@@ -321,6 +315,7 @@ TSharedRef<SDockTab> STimingProfilerWindow::SpawnTab_StatsCounters(const FSpawnT
 		];
 
 	DockTab->SetOnTabClosed(SDockTab::FOnTabClosedCallback::CreateRaw(this, &STimingProfilerWindow::OnStatsCountersTabClosed));
+	OpenTabs.Add(DockTab);
 
 	return DockTab;
 }
@@ -332,6 +327,8 @@ void STimingProfilerWindow::OnStatsCountersTabClosed(TSharedRef<SDockTab> TabBei
 {
 	FTimingProfilerManager::Get()->SetStatsCountersViewVisible(false);
 	StatsView = nullptr;
+
+	OpenTabs.Remove(TabBeingClosed);
 }
 
 ////////////////////////////////////////////////////////////////////////////////////////////////////
@@ -349,6 +346,7 @@ TSharedRef<SDockTab> STimingProfilerWindow::SpawnTab_LogView(const FSpawnTabArgs
 		];
 
 	DockTab->SetOnTabClosed(SDockTab::FOnTabClosedCallback::CreateRaw(this, &STimingProfilerWindow::OnLogViewTabClosed));
+	OpenTabs.Add(DockTab);
 
 	return DockTab;
 }
@@ -360,6 +358,8 @@ void STimingProfilerWindow::OnLogViewTabClosed(TSharedRef<SDockTab> TabBeingClos
 {
 	FTimingProfilerManager::Get()->SetLogViewVisible(false);
 	LogView = nullptr;
+
+	OpenTabs.Remove(TabBeingClosed);
 }
 
 ////////////////////////////////////////////////////////////////////////////////////////////////////
diff --git a/Engine/Source/Developer/TraceInsights/Private/Insights/Widgets/STimingProfilerWindow.h b/Engine/Source/Developer/TraceInsights/Private/Insights/Widgets/STimingProfilerWindow.h
--- a/Engine/Source/Developer/TraceInsights/Private/Insights/Widgets/STimingProfilerWindow.h
+++ b/Engine/Source/Developer/TraceInsights/Private/Insights/Widgets/STimingProfilerWindow.h
@@ -158,6 +158,12 @@ private:
 	void OnTimeSelectionChanged(Insights::ETimeChangedFlags InFlags, double InStartTime, double InEndTime);
 
 private:
+	/** Holds the tab manager that manages the front-end's tabs. */
+	TSharedPtr<FTabManager> TabManager;
+
+	/** All tabs owned by this window */
+	TSet<TSharedPtr<SDockTab>> OpenTabs;
+
 	/** The Frame track widget */
 	TSharedPtr<SFrameTrack> FrameTrack;
 
@@ -179,9 +185,6 @@ private:
 	/** The Log view widget */
 	TSharedPtr<SLogView> LogView;
 
-	/** Holds the tab manager that manages the front-end's tabs. */
-	TSharedPtr<FTabManager> TabManager;
-
 	/** The handle to the active update duration tick */
 	TWeakPtr<FActiveTimerHandle> ActiveTimerHandle;
 
