[Insights] Fixed leaked shared references to SNetworkingProfilerWindow.

diff --git a/Engine/Source/Developer/TraceInsights/Private/Insights/NetworkingProfiler/Widgets/SNetworkingProfilerToolbar.cpp b/Engine/Source/Developer/TraceInsights/Private/Insights/NetworkingProfiler/Widgets/SNetworkingProfilerToolbar.cpp
--- a/Engine/Source/Developer/TraceInsights/Private/Insights/NetworkingProfiler/Widgets/SNetworkingProfilerToolbar.cpp
+++ b/Engine/Source/Developer/TraceInsights/Private/Insights/NetworkingProfiler/Widgets/SNetworkingProfilerToolbar.cpp
@@ -30,13 +30,11 @@ SNetworkingProfilerToolbar::~SNetworkingProfilerToolbar()
 
 ////////////////////////////////////////////////////////////////////////////////////////////////////
 
-void SNetworkingProfilerToolbar::Construct(const FArguments& InArgs, TSharedPtr<SNetworkingProfilerWindow> InProfilerWindow)
+void SNetworkingProfilerToolbar::Construct(const FArguments& InArgs, TSharedRef<SNetworkingProfilerWindow> InProfilerWindow)
 {
-	ProfilerWindow = InProfilerWindow;
-
 	struct Local
 	{
-		static void FillViewToolbar(TSharedPtr<SNetworkingProfilerWindow> ProfilerWindow, FToolBarBuilder& ToolbarBuilder)
+		static void FillViewToolbar(FToolBarBuilder& ToolbarBuilder, TSharedRef<SNetworkingProfilerWindow> ProfilerWindow)
 		{
 			ToolbarBuilder.BeginSection("View");
 			{
@@ -84,11 +82,11 @@ void SNetworkingProfilerToolbar::Construct(const FArguments& InArgs, TSharedPtr<
 		}
 	};
 
-	TSharedPtr<FUICommandList> CommandList = ProfilerWindow->GetCommandList();
+	TSharedPtr<FUICommandList> CommandList = InProfilerWindow->GetCommandList();
 
 	FSlimHorizontalToolBarBuilder ToolbarBuilder(CommandList.ToSharedRef(), FMultiBoxCustomization::None);
 	ToolbarBuilder.SetStyle(&FInsightsStyle::Get(), "PrimaryToolbar");
-	Local::FillViewToolbar(ProfilerWindow, ToolbarBuilder);
+	Local::FillViewToolbar(ToolbarBuilder, InProfilerWindow);
 
 	FSlimHorizontalToolBarBuilder RightSideToolbarBuilder(CommandList.ToSharedRef(), FMultiBoxCustomization::None);
 	RightSideToolbarBuilder.SetStyle(&FInsightsStyle::Get(), "PrimaryToolbar");
diff --git a/Engine/Source/Developer/TraceInsights/Private/Insights/NetworkingProfiler/Widgets/SNetworkingProfilerToolbar.h b/Engine/Source/Developer/TraceInsights/Private/Insights/NetworkingProfiler/Widgets/SNetworkingProfilerToolbar.h
--- a/Engine/Source/Developer/TraceInsights/Private/Insights/NetworkingProfiler/Widgets/SNetworkingProfilerToolbar.h
+++ b/Engine/Source/Developer/TraceInsights/Private/Insights/NetworkingProfiler/Widgets/SNetworkingProfilerToolbar.h
@@ -25,8 +25,5 @@ public:
 	 *
 	 * @param InArgs The declaration data for this widget
 	 */
-	void Construct(const FArguments& InArgs, TSharedPtr<SNetworkingProfilerWindow> InProfilerWindow);
-
-private:
-	TSharedPtr<SNetworkingProfilerWindow> ProfilerWindow;
+	void Construct(const FArguments& InArgs, TSharedRef<SNetworkingProfilerWindow> InProfilerWindow);
 };
diff --git a/Engine/Source/Developer/TraceInsights/Private/Insights/NetworkingProfiler/Widgets/SPacketContentView.cpp b/Engine/Source/Developer/TraceInsights/Private/Insights/NetworkingProfiler/Widgets/SPacketContentView.cpp
--- a/Engine/Source/Developer/TraceInsights/Private/Insights/NetworkingProfiler/Widgets/SPacketContentView.cpp
+++ b/Engine/Source/Developer/TraceInsights/Private/Insights/NetworkingProfiler/Widgets/SPacketContentView.cpp
@@ -31,7 +31,7 @@
 ////////////////////////////////////////////////////////////////////////////////////////////////////
 
 SPacketContentView::SPacketContentView()
-	: ProfilerWindow()
+	: ProfilerWindowWeakPtr()
 	, DrawState(MakeShared<FPacketContentViewDrawState>())
 	, FilteredDrawState(MakeShared<FPacketContentViewDrawState>())
 {
@@ -48,7 +48,7 @@ SPacketContentView::~SPacketContentView()
 
 void SPacketContentView::Reset()
 {
-	//ProfilerWindow
+	//ProfilerWindowWeakPtr
 
 	Viewport.Reset();
 	//FAxisViewportDouble& ViewportX = Viewport.GetHorizontalAxisViewport();
@@ -104,9 +104,9 @@ void SPacketContentView::Reset()
 
 ////////////////////////////////////////////////////////////////////////////////////////////////////
 
-void SPacketContentView::Construct(const FArguments& InArgs, TSharedPtr<SNetworkingProfilerWindow> InProfilerWindow)
+void SPacketContentView::Construct(const FArguments& InArgs, TSharedRef<SNetworkingProfilerWindow> InProfilerWindow)
 {
-	ProfilerWindow = InProfilerWindow;
+	ProfilerWindowWeakPtr = InProfilerWindow;
 
 	FSlimHorizontalToolBarBuilder ToolbarBuilder(TSharedPtr<const FUICommandList>(), FMultiBoxCustomization::None);
 	ToolbarBuilder.SetStyle(&FInsightsStyle::Get(), "SecondaryToolbar2");
@@ -339,10 +339,14 @@ void SPacketContentView::Construct(const FArguments& InArgs, TSharedPtr<SNetwork
 
 void SPacketContentView::FindPreviousPacket()
 {
-	TSharedPtr<SPacketView> PacketView = ProfilerWindow->GetPacketView();
-	if (PacketView.IsValid())
+	TSharedPtr<SNetworkingProfilerWindow> ProfilerWindow = ProfilerWindowWeakPtr.Pin();
+	if (ProfilerWindow.IsValid())
 	{
-		PacketView->SelectPreviousPacket();
+		TSharedPtr<SPacketView> PacketView = ProfilerWindow->GetPacketView();
+		if (PacketView.IsValid())
+		{
+			PacketView->SelectPreviousPacket();
+		}
 	}
 }
 
@@ -350,10 +354,14 @@ void SPacketContentView::FindPreviousPacket()
 
 void SPacketContentView::FindNextPacket()
 {
-	TSharedPtr<SPacketView> PacketView = ProfilerWindow->GetPacketView();
-	if (PacketView.IsValid())
+	TSharedPtr<SNetworkingProfilerWindow> ProfilerWindow = ProfilerWindowWeakPtr.Pin();
+	if (ProfilerWindow.IsValid())
 	{
-		PacketView->SelectNextPacket();
+		TSharedPtr<SPacketView> PacketView = ProfilerWindow->GetPacketView();
+		if (PacketView.IsValid())
+		{
+			PacketView->SelectNextPacket();
+		}
 	}
 }
 
@@ -373,10 +381,14 @@ void SPacketContentView::Packet_OnTextCommitted(const FText& InNewText, ETextCom
 		uint32 NewPacketSequence = 0;
 		TTypeFromString<uint32>::FromString(NewPacketSequence, *InNewText.ToString());
 
-		TSharedPtr<SPacketView> PacketView = ProfilerWindow->GetPacketView();
-		if (PacketView.IsValid())
+		TSharedPtr<SNetworkingProfilerWindow> ProfilerWindow = ProfilerWindowWeakPtr.Pin();
+		if (ProfilerWindow.IsValid())
 		{
-			PacketView->SelectPacketBySequenceNumber(NewPacketSequence);
+			TSharedPtr<SPacketView> PacketView = ProfilerWindow->GetPacketView();
+			if (PacketView.IsValid())
+			{
+				PacketView->SelectPacketBySequenceNumber(NewPacketSequence);
+			}
 		}
 	}
 }
@@ -601,10 +613,14 @@ void SPacketContentView::FilterByNetId_OnCheckStateChanged(ECheckBoxState NewSta
 	bFilterByNetId = (NewState == ECheckBoxState::Checked);
 	bIsStateDirty = true;
 
-	TSharedPtr<SPacketView> PacketView = ProfilerWindow ? ProfilerWindow->GetPacketView() : nullptr;
-	if (PacketView.IsValid())
+	TSharedPtr<SNetworkingProfilerWindow> ProfilerWindow = ProfilerWindowWeakPtr.Pin();
+	if (ProfilerWindow.IsValid())
 	{
-		PacketView->InvalidateState();
+		TSharedPtr<SPacketView> PacketView = ProfilerWindow->GetPacketView();
+		if (PacketView.IsValid())
+		{
+			PacketView->InvalidateState();
+		}
 	}
 }
 
@@ -641,10 +657,14 @@ void SPacketContentView::FilterByEventType_OnCheckStateChanged(ECheckBoxState Ne
 	bFilterByEventType = (NewState == ECheckBoxState::Checked);
 	bIsStateDirty = true;
 
-	TSharedPtr<SPacketView> PacketView = ProfilerWindow ? ProfilerWindow->GetPacketView() : nullptr;
-	if (PacketView.IsValid())
+	TSharedPtr<SNetworkingProfilerWindow> ProfilerWindow = ProfilerWindowWeakPtr.Pin();
+	if (ProfilerWindow.IsValid())
 	{
-		PacketView->InvalidateState();
+		TSharedPtr<SPacketView> PacketView = ProfilerWindow->GetPacketView();
+		if (PacketView.IsValid())
+		{
+			PacketView->InvalidateState();
+		}
 	}
 }
 
@@ -764,10 +784,14 @@ void SPacketContentView::SetFilterNetId(const uint32 InNetId)
 	{
 		bIsStateDirty = true;
 
-		TSharedPtr<SPacketView> PacketView = ProfilerWindow ? ProfilerWindow->GetPacketView() : nullptr;
-		if (PacketView.IsValid())
+		TSharedPtr<SNetworkingProfilerWindow> ProfilerWindow = ProfilerWindowWeakPtr.Pin();
+		if (ProfilerWindow.IsValid())
 		{
-			PacketView->InvalidateState();
+			TSharedPtr<SPacketView> PacketView = ProfilerWindow->GetPacketView();
+			if (PacketView.IsValid())
+			{
+				PacketView->InvalidateState();
+			}
 		}
 	}
 }
@@ -783,10 +807,14 @@ void SPacketContentView::SetFilterEventType(const uint32 InEventTypeIndex, const
 	{
 		bIsStateDirty = true;
 
-		TSharedPtr<SPacketView> PacketView = ProfilerWindow ? ProfilerWindow->GetPacketView() : nullptr;
-		if (PacketView.IsValid())
+		TSharedPtr<SNetworkingProfilerWindow> ProfilerWindow = ProfilerWindowWeakPtr.Pin();
+		if (ProfilerWindow.IsValid())
 		{
-			PacketView->InvalidateState();
+			TSharedPtr<SPacketView> PacketView = ProfilerWindow->GetPacketView();
+			if (PacketView.IsValid())
+			{
+				PacketView->InvalidateState();
+			}
 		}
 	}
 }
@@ -997,10 +1025,14 @@ void SPacketContentView::UpdateHoveredEvent()
 
 void SPacketContentView::OnSelectedEventChanged()
 {
-	if (SelectedEvent.IsValid() && ProfilerWindow.IsValid())
+	TSharedPtr<SNetworkingProfilerWindow> ProfilerWindow = ProfilerWindowWeakPtr.Pin();
+	if (ProfilerWindow.IsValid())
 	{
-		// Select the node coresponding to net event type of selected net event instance.
-		ProfilerWindow->SetSelectedEventTypeIndex(SelectedEvent.Event.EventTypeIndex);
+		if (SelectedEvent.IsValid())
+		{
+			// Select the node coresponding to net event type of selected net event instance.
+			ProfilerWindow->SetSelectedEventTypeIndex(SelectedEvent.Event.EventTypeIndex);
+		}
 	}
 }
 
@@ -1425,6 +1457,7 @@ FReply SPacketContentView::OnKeyDown(const FGeometry& MyGeometry, const FKeyEven
 		else if (InKeyEvent.GetModifierKeys().IsControlDown() ||
 				 InKeyEvent.GetModifierKeys().IsCommandDown())
 		{
+			TSharedPtr<SNetworkingProfilerWindow> ProfilerWindow = ProfilerWindowWeakPtr.Pin();
 			if (ProfilerWindow.IsValid())
 			{
 				const TSharedPtr<SPacketView> PacketView = ProfilerWindow->GetPacketView();
@@ -1449,6 +1482,7 @@ FReply SPacketContentView::OnKeyDown(const FGeometry& MyGeometry, const FKeyEven
 		else if (InKeyEvent.GetModifierKeys().IsControlDown() ||
 				 InKeyEvent.GetModifierKeys().IsCommandDown())
 		{
+			TSharedPtr<SNetworkingProfilerWindow> ProfilerWindow = ProfilerWindowWeakPtr.Pin();
 			if (ProfilerWindow.IsValid())
 			{
 				const TSharedPtr<SPacketView> PacketView = ProfilerWindow->GetPacketView();
diff --git a/Engine/Source/Developer/TraceInsights/Private/Insights/NetworkingProfiler/Widgets/SPacketContentView.h b/Engine/Source/Developer/TraceInsights/Private/Insights/NetworkingProfiler/Widgets/SPacketContentView.h
--- a/Engine/Source/Developer/TraceInsights/Private/Insights/NetworkingProfiler/Widgets/SPacketContentView.h
+++ b/Engine/Source/Developer/TraceInsights/Private/Insights/NetworkingProfiler/Widgets/SPacketContentView.h
@@ -112,7 +112,7 @@ public:
 	SLATE_BEGIN_ARGS(SPacketContentView) {}
 	SLATE_END_ARGS()
 
-	void Construct(const FArguments& InArgs, TSharedPtr<SNetworkingProfilerWindow> InProfilerWindow);
+	void Construct(const FArguments& InArgs, TSharedRef<SNetworkingProfilerWindow> InProfilerWindow);
 
 	/**
 	 * Ticks this widget. Override in derived classes, but always call the parent implementation.
@@ -211,7 +211,7 @@ private:
 	void AdjustForSplitContent();
 
 private:
-	TSharedPtr<SNetworkingProfilerWindow> ProfilerWindow;
+	TWeakPtr<SNetworkingProfilerWindow> ProfilerWindowWeakPtr;
 
 	/** The track's viewport. Encapsulates info about position and scale. */
 	FPacketContentViewport Viewport;
diff --git a/Engine/Source/Developer/TraceInsights/Private/Insights/NetworkingProfiler/Widgets/SPacketView.cpp b/Engine/Source/Developer/TraceInsights/Private/Insights/NetworkingProfiler/Widgets/SPacketView.cpp
--- a/Engine/Source/Developer/TraceInsights/Private/Insights/NetworkingProfiler/Widgets/SPacketView.cpp
+++ b/Engine/Source/Developer/TraceInsights/Private/Insights/NetworkingProfiler/Widgets/SPacketView.cpp
@@ -31,7 +31,7 @@
 ////////////////////////////////////////////////////////////////////////////////////////////////////
 
 SPacketView::SPacketView()
-	: ProfilerWindow()
+	: ProfilerWindowWeakPtr()
 	, PacketSeries(MakeShared<FNetworkPacketSeries>())
 {
 	Reset();
@@ -47,7 +47,7 @@ SPacketView::~SPacketView()
 
 void SPacketView::Reset()
 {
-	//ProfilerWindow
+	//ProfilerWindowWeakPtr
 
 	GameInstanceIndex = 0;
 	ConnectionIndex = 0;
@@ -105,6 +105,7 @@ void SPacketView::Reset()
 	OnPaintDurationHistory.Reset();
 	LastOnPaintTime = FPlatformTime::Cycles64();
 
+	TSharedPtr<SNetworkingProfilerWindow> ProfilerWindow = ProfilerWindowWeakPtr.Pin();
 	if (ProfilerWindow.IsValid())
 	{
 		SetConnection(ProfilerWindow->GetSelectedGameInstanceIndex(), ProfilerWindow->GetSelectedConnectionIndex(), ProfilerWindow->GetSelectedConnectionMode());
@@ -132,9 +133,9 @@ void SPacketView::SetConnection(uint32 InGameInstanceIndex, uint32 InConnectionI
 
 ////////////////////////////////////////////////////////////////////////////////////////////////////
 
-void SPacketView::Construct(const FArguments& InArgs, TSharedPtr<SNetworkingProfilerWindow> InProfilerWindow)
+void SPacketView::Construct(const FArguments& InArgs, TSharedRef<SNetworkingProfilerWindow> InProfilerWindow)
 {
-	ProfilerWindow = InProfilerWindow;
+	ProfilerWindowWeakPtr = InProfilerWindow;
 
 	ChildSlot
 	[
@@ -284,6 +285,7 @@ void SPacketView::UpdateState()
 
 	FPacketFilter Filter;
 
+	TSharedPtr<SNetworkingProfilerWindow> ProfilerWindow = ProfilerWindowWeakPtr.Pin();
 	if (ProfilerWindow)
 	{
 		TSharedPtr<SPacketContentView> PacketContentView = ProfilerWindow->GetPacketContentView();
@@ -603,6 +605,7 @@ void SPacketView::SelectSampleAtMousePosition(float X, float Y, const FPointerEv
 
 void SPacketView::OnSelectionChanged()
 {
+	TSharedPtr<SNetworkingProfilerWindow> ProfilerWindow = ProfilerWindowWeakPtr.Pin();
 	if (ProfilerWindow.IsValid())
 	{
 		if (SelectedSample.IsValid())
@@ -628,7 +631,7 @@ const TCHAR* StatusToString(TraceServices::ENetProfilerDeliveryStatus Status)
 		case TraceServices::ENetProfilerDeliveryStatus::Delivered:  return TEXT("Delivered");
 		case TraceServices::ENetProfilerDeliveryStatus::Dropped:    return TEXT("Dropped");
 		case TraceServices::ENetProfilerDeliveryStatus::Unknown:
-		default:                                               return TEXT("Unknown");
+		default:                                                    return TEXT("Unknown");
 	};
 }
 
@@ -639,7 +642,7 @@ const TCHAR* AggregatedStatusToString(TraceServices::ENetProfilerDeliveryStatus
 		case TraceServices::ENetProfilerDeliveryStatus::Delivered:  return TEXT("all packets are Delivered");
 		case TraceServices::ENetProfilerDeliveryStatus::Dropped:    return TEXT("at least one Dropped packet");
 		case TraceServices::ENetProfilerDeliveryStatus::Unknown:
-		default:                                               return TEXT("Unknown");
+		default:                                                    return TEXT("Unknown");
 	}
 }
 
diff --git a/Engine/Source/Developer/TraceInsights/Private/Insights/NetworkingProfiler/Widgets/SPacketView.h b/Engine/Source/Developer/TraceInsights/Private/Insights/NetworkingProfiler/Widgets/SPacketView.h
--- a/Engine/Source/Developer/TraceInsights/Private/Insights/NetworkingProfiler/Widgets/SPacketView.h
+++ b/Engine/Source/Developer/TraceInsights/Private/Insights/NetworkingProfiler/Widgets/SPacketView.h
@@ -106,7 +106,7 @@ public:
 	}
 	SLATE_END_ARGS()
 
-	void Construct(const FArguments& InArgs, TSharedPtr<SNetworkingProfilerWindow> InProfilerWindow);
+	void Construct(const FArguments& InArgs, TSharedRef<SNetworkingProfilerWindow> InProfilerWindow);
 
 	/**
 	 * Ticks this widget. Override in derived classes, but always call the parent implementation.
@@ -178,7 +178,7 @@ private:
 	void ZoomHorizontally(const float Delta, const float X);
 
 private:
-	TSharedPtr<SNetworkingProfilerWindow> ProfilerWindow;
+	TWeakPtr<SNetworkingProfilerWindow> ProfilerWindowWeakPtr;
 
 	uint32 GameInstanceIndex;
 	uint32 ConnectionIndex;
