[Insights]
  - Added support for InitialSortMode for columns.
  - Set InitialSortMode to "Descending" for most numeric columns in Timers, Counters, Net Stats, Mem Allocs and generic table/tree views.

diff --git a/Engine/Source/Developer/TraceInsights/Private/Insights/MemoryProfiler/ViewModels/MemAllocTable.cpp b/Engine/Source/Developer/TraceInsights/Private/Insights/MemoryProfiler/ViewModels/MemAllocTable.cpp
--- a/Engine/Source/Developer/TraceInsights/Private/Insights/MemoryProfiler/ViewModels/MemAllocTable.cpp
+++ b/Engine/Source/Developer/TraceInsights/Private/Insights/MemoryProfiler/ViewModels/MemAllocTable.cpp
@@ -604,6 +604,7 @@ void FMemAllocTable::AddDefaultColumns()
 
 		TSharedRef<ITableCellValueSorter> Sorter = MakeShared<FSorterByInt64Value>(ColumnRef);
 		Column.SetValueSorter(Sorter);
+		Column.SetInitialSortMode(EColumnSortMode::Descending);
 
 		Column.SetAggregation(ETableColumnAggregation::Sum);
 
@@ -662,6 +663,7 @@ void FMemAllocTable::AddDefaultColumns()
 
 		TSharedRef<ITableCellValueSorter> Sorter = MakeShared<FSorterByInt64Value>(ColumnRef);
 		Column.SetValueSorter(Sorter);
+		Column.SetInitialSortMode(EColumnSortMode::Descending);
 
 		TSharedRef<IFilterValueConverter> Converter = MakeShared<FMemoryFilterValueConverter>();
 		Column.SetValueConverter(Converter);
diff --git a/Engine/Source/Developer/TraceInsights/Private/Insights/MemoryProfiler/Widgets/SMemTagTreeView.cpp b/Engine/Source/Developer/TraceInsights/Private/Insights/MemoryProfiler/Widgets/SMemTagTreeView.cpp
--- a/Engine/Source/Developer/TraceInsights/Private/Insights/MemoryProfiler/Widgets/SMemTagTreeView.cpp
+++ b/Engine/Source/Developer/TraceInsights/Private/Insights/MemoryProfiler/Widgets/SMemTagTreeView.cpp
@@ -1632,6 +1632,7 @@ void SMemTagTreeView::ShowColumn(const FName ColumnId)
 		.VAlignHeader(VAlign_Center)
 		.HAlignCell(HAlign_Fill)
 		.VAlignCell(VAlign_Fill)
+		.InitialSortMode(Column.GetInitialSortMode())
 		.SortMode(this, &SMemTagTreeView::GetSortModeForColumn, Column.GetId())
 		.OnSort(this, &SMemTagTreeView::OnSortModeChanged)
 		.FillWidth(Column.GetInitialWidth())
diff --git a/Engine/Source/Developer/TraceInsights/Private/Insights/NetworkingProfiler/ViewModels/NetStatsViewColumnFactory.cpp b/Engine/Source/Developer/TraceInsights/Private/Insights/NetworkingProfiler/ViewModels/NetStatsViewColumnFactory.cpp
--- a/Engine/Source/Developer/TraceInsights/Private/Insights/NetworkingProfiler/ViewModels/NetStatsViewColumnFactory.cpp
+++ b/Engine/Source/Developer/TraceInsights/Private/Insights/NetworkingProfiler/ViewModels/NetStatsViewColumnFactory.cpp
@@ -217,6 +217,7 @@ TSharedRef<Insights::FTableColumn> FNetStatsViewColumnFactory::CreateInstanceCou
 	//TSharedRef<ITableCellValueSorter> Sorter = MakeShared<FSorterByInt64Value>(ColumnRef);
 	TSharedRef<ITableCellValueSorter> Sorter = MakeShared<FNetEventNodeSortingByInstanceCount>(ColumnRef);
 	Column.SetValueSorter(Sorter);
+	Column.SetInitialSortMode(EColumnSortMode::Descending);
 
 	return ColumnRef;
 }
@@ -266,6 +267,7 @@ TSharedRef<Insights::FTableColumn> FNetStatsViewColumnFactory::CreateTotalInclus
 	//TSharedRef<ITableCellValueSorter> Sorter = MakeShared<FSorterByInt64Value>(ColumnRef);
 	TSharedRef<ITableCellValueSorter> Sorter = MakeShared<FNetEventNodeSortingByTotalInclusiveSize>(ColumnRef);
 	Column.SetValueSorter(Sorter);
+	Column.SetInitialSortMode(EColumnSortMode::Descending);
 
 	return ColumnRef;
 }
@@ -311,6 +313,7 @@ TSharedRef<Insights::FTableColumn> FNetStatsViewColumnFactory::CreateMaxInclusiv
 
 	TSharedRef<ITableCellValueSorter> Sorter = MakeShared<FSorterByInt64Value>(ColumnRef);
 	Column.SetValueSorter(Sorter);
+	Column.SetInitialSortMode(EColumnSortMode::Descending);
 
 	return ColumnRef;
 }
@@ -356,6 +359,7 @@ TSharedRef<Insights::FTableColumn> FNetStatsViewColumnFactory::CreateAverageIncl
 
 	TSharedRef<ITableCellValueSorter> Sorter = MakeShared<FSorterByInt64Value>(ColumnRef);
 	Column.SetValueSorter(Sorter);
+	Column.SetInitialSortMode(EColumnSortMode::Descending);
 
 	return ColumnRef;
 }
@@ -405,6 +409,7 @@ TSharedRef<Insights::FTableColumn> FNetStatsViewColumnFactory::CreateTotalExclus
 	//TSharedRef<ITableCellValueSorter> Sorter = MakeShared<FSorterByInt64Value(ColumnRef));
 	TSharedRef<ITableCellValueSorter> Sorter = MakeShared<FNetEventNodeSortingByTotalExclusiveSize>(ColumnRef);
 	Column.SetValueSorter(Sorter);
+	Column.SetInitialSortMode(EColumnSortMode::Descending);
 
 	return ColumnRef;
 }
@@ -450,6 +455,7 @@ TSharedRef<Insights::FTableColumn> FNetStatsViewColumnFactory::CreateMaxExclusiv
 
 	TSharedRef<ITableCellValueSorter> Sorter = MakeShared<FSorterByInt64Value>(ColumnRef);
 	Column.SetValueSorter(Sorter);
+	Column.SetInitialSortMode(EColumnSortMode::Descending);
 
 	return ColumnRef;
 }
@@ -496,6 +502,7 @@ TSharedRef<Insights::FTableColumn> FNetStatsViewColumnFactory::CreateAverageExcl
 
 	TSharedRef<ITableCellValueSorter> Sorter = MakeShared<FSorterByInt64Value>(ColumnRef);
 	Column.SetValueSorter(Sorter);
+	Column.SetInitialSortMode(EColumnSortMode::Descending);
 
 	return ColumnRef;
 }
diff --git a/Engine/Source/Developer/TraceInsights/Private/Insights/NetworkingProfiler/Widgets/SNetStatsView.cpp b/Engine/Source/Developer/TraceInsights/Private/Insights/NetworkingProfiler/Widgets/SNetStatsView.cpp
--- a/Engine/Source/Developer/TraceInsights/Private/Insights/NetworkingProfiler/Widgets/SNetStatsView.cpp
+++ b/Engine/Source/Developer/TraceInsights/Private/Insights/NetworkingProfiler/Widgets/SNetStatsView.cpp
@@ -70,7 +70,7 @@ SNetStatsView::~SNetStatsView()
 ////////////////////////////////////////////////////////////////////////////////////////////////////
 
 BEGIN_SLATE_FUNCTION_BUILD_OPTIMIZATION
-void SNetStatsView::Construct(const FArguments& InArgs, TSharedPtr<SNetworkingProfilerWindow> InProfilerWindow)
+void SNetStatsView::Construct(const FArguments& InArgs, TSharedRef<SNetworkingProfilerWindow> InProfilerWindow)
 {
 	ProfilerWindowWeakPtr = InProfilerWindow;
 
@@ -1247,6 +1247,7 @@ void SNetStatsView::ShowColumn(const FName ColumnId)
 		.VAlignHeader(VAlign_Center)
 		.HAlignCell(HAlign_Fill)
 		.VAlignCell(VAlign_Fill)
+		.InitialSortMode(Column.GetInitialSortMode())
 		.SortMode(this, &SNetStatsView::GetSortModeForColumn, Column.GetId())
 		.OnSort(this, &SNetStatsView::OnSortModeChanged)
 		.FillWidth(Column.GetInitialWidth())
diff --git a/Engine/Source/Developer/TraceInsights/Private/Insights/NetworkingProfiler/Widgets/SNetStatsView.h b/Engine/Source/Developer/TraceInsights/Private/Insights/NetworkingProfiler/Widgets/SNetStatsView.h
--- a/Engine/Source/Developer/TraceInsights/Private/Insights/NetworkingProfiler/Widgets/SNetStatsView.h
+++ b/Engine/Source/Developer/TraceInsights/Private/Insights/NetworkingProfiler/Widgets/SNetStatsView.h
@@ -75,7 +75,7 @@ public:
 	 * Construct this widget
 	 * @param InArgs - The declaration data for this widget
 	 */
-	void Construct(const FArguments& InArgs, TSharedPtr<SNetworkingProfilerWindow> InProfilerWindow);
+	void Construct(const FArguments& InArgs, TSharedRef<SNetworkingProfilerWindow> InProfilerWindow);
 
 	/**
 	 * Ticks this widget. Override in derived classes, but always call the parent implementation.
diff --git a/Engine/Source/Developer/TraceInsights/Private/Insights/Table/ViewModels/TableColumn.h b/Engine/Source/Developer/TraceInsights/Private/Insights/Table/ViewModels/TableColumn.h
--- a/Engine/Source/Developer/TraceInsights/Private/Insights/Table/ViewModels/TableColumn.h
+++ b/Engine/Source/Developer/TraceInsights/Private/Insights/Table/ViewModels/TableColumn.h
@@ -5,6 +5,7 @@
 #include "CoreMinimal.h"
 #include "Types/SlateEnums.h"
 #include "Misc/EnumClassFlags.h"
+#include "Widgets/Views/SHeaderRow.h" // for EColumnSortMode
 
 #include "Insights/Table/ViewModels/TableCellValue.h"
 
@@ -73,6 +74,7 @@ public:
 		, ValueGetter(GetDefaultValueGetter())
 		, ValueFormatter(GetDefaultValueFormatter())
 		, ValueSorter(nullptr)
+		, InitialSortMode(EColumnSortMode::Ascending)
 		, ParentTable(nullptr)
 	{
 	}
@@ -185,6 +187,12 @@ public:
 	/** Whether this column can be used for sorting. */
 	bool CanBeSorted() const { return ValueSorter.IsValid(); }
 
+	/** Gets the initial sorting mode. */
+	EColumnSortMode::Type GetInitialSortMode() const { return InitialSortMode; }
+
+	/** Sets the initial sorting mode. */
+	void SetInitialSortMode(EColumnSortMode::Type InMode) { InitialSortMode = InMode; }
+
 	//////////////////////////////////////////////////
 	// Value Converter (can be nullptr)
 
@@ -248,6 +256,9 @@ private:
 
 	/** Custom sorter for values displayed by this column. */
 	TSharedPtr<ITableCellValueSorter> ValueSorter;
+
+	/** Initial sorting mode. */
+	EColumnSortMode::Type InitialSortMode;
 	
 	/** Used to convert in a custom way from string to column data type in FilterConfigurator */
 	TSharedPtr<IFilterValueConverter> ValueConverter;
diff --git a/Engine/Source/Developer/TraceInsights/Private/Insights/Table/ViewModels/UntypedTable.cpp b/Engine/Source/Developer/TraceInsights/Private/Insights/Table/ViewModels/UntypedTable.cpp
--- a/Engine/Source/Developer/TraceInsights/Private/Insights/Table/ViewModels/UntypedTable.cpp
+++ b/Engine/Source/Developer/TraceInsights/Private/Insights/Table/ViewModels/UntypedTable.cpp
@@ -208,6 +208,7 @@ void FUntypedTable::CreateColumns(const TraceServices::ITableLayout& TableLayout
 
 		TSharedPtr<ITableCellValueFormatter> FormatterPtr;
 		TSharedPtr<ITableCellValueSorter> SorterPtr;
+		EColumnSortMode::Type InitialSortMode = EColumnSortMode::Ascending;
 
 		switch (ColumnType)
 		{
@@ -219,6 +220,7 @@ void FUntypedTable::CreateColumns(const TraceServices::ITableLayout& TableLayout
 			//else // if (Hint == AsTrueFalse)
 			FormatterPtr = MakeShared<FBoolValueFormatterAsTrueFalse>();
 			SorterPtr = MakeShared<FSorterByBoolValue>(ColumnRef);
+			InitialSortMode = EColumnSortMode::Ascending;
 			break;
 
 		case TraceServices::TableColumnType_Int:
@@ -238,6 +240,7 @@ void FUntypedTable::CreateColumns(const TraceServices::ITableLayout& TableLayout
 				FormatterPtr = MakeShared<FInt64ValueFormatterAsNumber>();
 			}
 			SorterPtr = MakeShared<FSorterByInt64Value>(ColumnRef);
+			InitialSortMode = EColumnSortMode::Descending;
 			break;
 
 		case TraceServices::TableColumnType_Float:
@@ -251,10 +254,12 @@ void FUntypedTable::CreateColumns(const TraceServices::ITableLayout& TableLayout
 			if (ColumnDisplayHintFlags & TraceServices::TableColumnDisplayHint_Time)
 			{
 				FormatterPtr = MakeShared<FFloatValueFormatterAsTimeAuto>();
+				InitialSortMode = EColumnSortMode::Ascending;
 			}
 			else
 			{
 				FormatterPtr = MakeShared<FFloatValueFormatterAsNumber>();
+				InitialSortMode = EColumnSortMode::Descending;
 			}
 			SorterPtr = MakeShared<FSorterByFloatValue>(ColumnRef);
 			break;
@@ -270,10 +275,12 @@ void FUntypedTable::CreateColumns(const TraceServices::ITableLayout& TableLayout
 			if (ColumnDisplayHintFlags & TraceServices::TableColumnDisplayHint_Time)
 			{
 				FormatterPtr = MakeShared<FDoubleValueFormatterAsTimeAuto>();
+				InitialSortMode = EColumnSortMode::Ascending;
 			}
 			else
 			{
 				FormatterPtr = MakeShared<FDoubleValueFormatterAsNumber>();
+				InitialSortMode = EColumnSortMode::Descending;
 			}
 			SorterPtr = MakeShared<FSorterByDoubleValue>(ColumnRef);
 			break;
@@ -284,6 +291,7 @@ void FUntypedTable::CreateColumns(const TraceServices::ITableLayout& TableLayout
 			InitialColumnWidth = FMath::Max(120.0f, 6.0f * ColumnNameStr.Len());
 			FormatterPtr = MakeShared<FCStringValueFormatterAsText>();
 			SorterPtr = MakeShared<FSorterByCStringValue>(ColumnRef);
+			InitialSortMode = EColumnSortMode::Ascending;
 			break;
 		}
 
@@ -309,6 +317,7 @@ void FUntypedTable::CreateColumns(const TraceServices::ITableLayout& TableLayout
 		}
 
 		Column.SetValueSorter(SorterPtr);
+		Column.SetInitialSortMode(InitialSortMode);
 
 		AddColumn(ColumnRef);
 	}
diff --git a/Engine/Source/Developer/TraceInsights/Private/Insights/Table/Widgets/STableTreeView.cpp b/Engine/Source/Developer/TraceInsights/Private/Insights/Table/Widgets/STableTreeView.cpp
--- a/Engine/Source/Developer/TraceInsights/Private/Insights/Table/Widgets/STableTreeView.cpp
+++ b/Engine/Source/Developer/TraceInsights/Private/Insights/Table/Widgets/STableTreeView.cpp
@@ -2293,6 +2293,7 @@ void STableTreeView::ShowColumn(FTableColumn& Column)
 			.VAlignHeader(VAlign_Center)
 			.HAlignCell(HAlign_Fill)
 			.VAlignCell(VAlign_Fill)
+			.InitialSortMode(Column.GetInitialSortMode())
 			.SortMode(this, &STableTreeView::GetSortModeForColumn, Column.GetId())
 			.OnSort(this, &STableTreeView::OnSortModeChanged)
 			.FillWidth(Column.GetInitialWidth())
diff --git a/Engine/Source/Developer/TraceInsights/Private/Insights/ViewModels/StatsViewColumnFactory.cpp b/Engine/Source/Developer/TraceInsights/Private/Insights/ViewModels/StatsViewColumnFactory.cpp
--- a/Engine/Source/Developer/TraceInsights/Private/Insights/ViewModels/StatsViewColumnFactory.cpp
+++ b/Engine/Source/Developer/TraceInsights/Private/Insights/ViewModels/StatsViewColumnFactory.cpp
@@ -261,6 +261,7 @@ TSharedRef<Insights::FTableColumn> FStatsViewColumnFactory::CreateCountColumn()
 	//TSharedRef<ITableCellValueSorter> Sorter = MakeShared<FSorterByInt64Value>(ColumnRef);
 	TSharedRef<ITableCellValueSorter> Sorter = MakeShared<FStatsNodeSortingByCount>(ColumnRef);
 	Column.SetValueSorter(Sorter);
+	Column.SetInitialSortMode(EColumnSortMode::Descending);
 
 	return ColumnRef;
 }
@@ -318,6 +319,7 @@ TSharedRef<Insights::FTableColumn> FStatsViewColumnFactory::CreateSumColumn()
 	TSharedRef<ITableCellValueSorter> Sorter = MakeShared<FSorterByDoubleValue>(ColumnRef);
 	//TSharedRef<ITableCellValueSorter> Sorter = MakeShared<FStatsNodeSortingBySum>(ColumnRef);
 	Column.SetValueSorter(Sorter);
+	Column.SetInitialSortMode(EColumnSortMode::Descending);
 
 	return ColumnRef;
 }
@@ -375,6 +377,7 @@ TSharedRef<Insights::FTableColumn> FStatsViewColumnFactory::CreateMaxColumn()
 	TSharedRef<ITableCellValueSorter> Sorter = MakeShared<FSorterByDoubleValue>(ColumnRef);
 	//TSharedRef<ITableCellValueSorter> Sorter = MakeShared<FStatsNodeSortingByMax>(ColumnRef);
 	Column.SetValueSorter(Sorter);
+	Column.SetInitialSortMode(EColumnSortMode::Descending);
 
 	return ColumnRef;
 }
@@ -432,6 +435,7 @@ TSharedRef<Insights::FTableColumn> FStatsViewColumnFactory::CreateUpperQuartileC
 	TSharedRef<ITableCellValueSorter> Sorter = MakeShared<FSorterByDoubleValue>(ColumnRef);
 	//TSharedRef<ITableCellValueSorter> Sorter = MakeShared<FStatsNodeSortingByUpperQuartile>(ColumnRef);
 	Column.SetValueSorter(Sorter);
+	Column.SetInitialSortMode(EColumnSortMode::Descending);
 
 	return ColumnRef;
 }
@@ -488,6 +492,7 @@ TSharedRef<Insights::FTableColumn> FStatsViewColumnFactory::CreateAverageColumn(
 	TSharedRef<ITableCellValueSorter> Sorter = MakeShared<FSorterByDoubleValue>(ColumnRef);
 	//TSharedRef<ITableCellValueSorter> Sorter = MakeShared<FStatsNodeSortingByAverage>(ColumnRef);
 	Column.SetValueSorter(Sorter);
+	Column.SetInitialSortMode(EColumnSortMode::Descending);
 
 	return ColumnRef;
 }
@@ -544,6 +549,7 @@ TSharedRef<Insights::FTableColumn> FStatsViewColumnFactory::CreateMedianColumn()
 	TSharedRef<ITableCellValueSorter> Sorter = MakeShared<FSorterByDoubleValue>(ColumnRef);
 	//TSharedRef<ITableCellValueSorter> Sorter = MakeShared<FStatsNodeSortingByMedian>(ColumnRef);
 	Column.SetValueSorter(Sorter);
+	Column.SetInitialSortMode(EColumnSortMode::Descending);
 
 	return ColumnRef;
 }
@@ -601,6 +607,7 @@ TSharedRef<Insights::FTableColumn> FStatsViewColumnFactory::CreateLowerQuartileC
 	TSharedRef<ITableCellValueSorter> Sorter = MakeShared<FSorterByDoubleValue>(ColumnRef);
 	//TSharedRef<ITableCellValueSorter> Sorter = MakeShared<FStatsNodeSortingByLowerQuartile>(ColumnRef);
 	Column.SetValueSorter(Sorter);
+	Column.SetInitialSortMode(EColumnSortMode::Descending);
 
 	return ColumnRef;
 }
@@ -658,6 +665,7 @@ TSharedRef<Insights::FTableColumn> FStatsViewColumnFactory::CreateMinColumn()
 	TSharedRef<ITableCellValueSorter> Sorter = MakeShared<FSorterByDoubleValue>(ColumnRef);
 	//TSharedRef<ITableCellValueSorter> Sorter = MakeShared<FStatsNodeSortingByMin>(ColumnRef);
 	Column.SetValueSorter(Sorter);
+	Column.SetInitialSortMode(EColumnSortMode::Descending);
 
 	return ColumnRef;
 }
@@ -711,6 +719,7 @@ TSharedRef<Insights::FTableColumn> FStatsViewColumnFactory::CreateDiffColumn()
 
 	TSharedRef<ITableCellValueSorter> Sorter = MakeShared<FSorterByDoubleValue>(ColumnRef);
 	Column.SetValueSorter(Sorter);
+	Column.SetInitialSortMode(EColumnSortMode::Descending);
 
 	return ColumnRef;
 }
diff --git a/Engine/Source/Developer/TraceInsights/Private/Insights/ViewModels/TimersViewColumnFactory.cpp b/Engine/Source/Developer/TraceInsights/Private/Insights/ViewModels/TimersViewColumnFactory.cpp
--- a/Engine/Source/Developer/TraceInsights/Private/Insights/ViewModels/TimersViewColumnFactory.cpp
+++ b/Engine/Source/Developer/TraceInsights/Private/Insights/ViewModels/TimersViewColumnFactory.cpp
@@ -243,6 +243,7 @@ TSharedRef<Insights::FTableColumn> FTimersViewColumnFactory::CreateInstanceCount
 	//TSharedRef<ITableCellValueSorter> Sorter = MakeShared<FSorterByInt64Value>(ColumnRef);
 	TSharedRef<ITableCellValueSorter> Sorter = MakeShared<FTimerNodeSortingByInstanceCount>(ColumnRef);
 	Column.SetValueSorter(Sorter);
+	Column.SetInitialSortMode(EColumnSortMode::Descending);
 
 	return ColumnRef;
 }
@@ -292,6 +293,7 @@ TSharedRef<Insights::FTableColumn> FTimersViewColumnFactory::CreateTotalInclusiv
 	//TSharedRef<ITableCellValueSorter> Sorter = MakeShared<FSorterByDoubleValue>(ColumnRef);
 	TSharedRef<ITableCellValueSorter> Sorter = MakeShared<FTimerNodeSortingByTotalInclusiveTime>(ColumnRef);
 	Column.SetValueSorter(Sorter);
+	Column.SetInitialSortMode(EColumnSortMode::Descending);
 
 	return ColumnRef;
 }
@@ -337,6 +339,7 @@ TSharedRef<Insights::FTableColumn> FTimersViewColumnFactory::CreateMaxInclusiveT
 
 	TSharedRef<ITableCellValueSorter> Sorter = MakeShared<FSorterByDoubleValue>(ColumnRef);
 	Column.SetValueSorter(Sorter);
+	Column.SetInitialSortMode(EColumnSortMode::Descending);
 
 	return ColumnRef;
 }
@@ -382,6 +385,7 @@ TSharedRef<Insights::FTableColumn> FTimersViewColumnFactory::CreateAverageInclus
 
 	TSharedRef<ITableCellValueSorter> Sorter = MakeShared<FSorterByDoubleValue>(ColumnRef);
 	Column.SetValueSorter(Sorter);
+	Column.SetInitialSortMode(EColumnSortMode::Descending);
 
 	return ColumnRef;
 }
@@ -427,6 +431,7 @@ TSharedRef<Insights::FTableColumn> FTimersViewColumnFactory::CreateMedianInclusi
 
 	TSharedRef<ITableCellValueSorter> Sorter = MakeShared<FSorterByDoubleValue>(ColumnRef);
 	Column.SetValueSorter(Sorter);
+	Column.SetInitialSortMode(EColumnSortMode::Descending);
 
 	return ColumnRef;
 }
@@ -472,6 +477,7 @@ TSharedRef<Insights::FTableColumn> FTimersViewColumnFactory::CreateMinInclusiveT
 
 	TSharedRef<ITableCellValueSorter> Sorter = MakeShared<FSorterByDoubleValue>(ColumnRef);
 	Column.SetValueSorter(Sorter);
+	Column.SetInitialSortMode(EColumnSortMode::Descending);
 
 	return ColumnRef;
 }
@@ -521,6 +527,7 @@ TSharedRef<Insights::FTableColumn> FTimersViewColumnFactory::CreateTotalExclusiv
 	//TSharedRef<ITableCellValueSorter> Sorter = MakeShared<FSorterByDoubleValue>(ColumnRef);
 	TSharedRef<ITableCellValueSorter> Sorter = MakeShared<FTimerNodeSortingByTotalExclusiveTime>(ColumnRef);
 	Column.SetValueSorter(Sorter);
+	Column.SetInitialSortMode(EColumnSortMode::Descending);
 
 	return ColumnRef;
 }
@@ -566,6 +573,7 @@ TSharedRef<Insights::FTableColumn> FTimersViewColumnFactory::CreateMaxExclusiveT
 
 	TSharedRef<ITableCellValueSorter> Sorter = MakeShared<FSorterByDoubleValue>(ColumnRef);
 	Column.SetValueSorter(Sorter);
+	Column.SetInitialSortMode(EColumnSortMode::Descending);
 
 	return ColumnRef;
 }
@@ -611,6 +619,7 @@ TSharedRef<Insights::FTableColumn> FTimersViewColumnFactory::CreateAverageExclus
 
 	TSharedRef<ITableCellValueSorter> Sorter = MakeShared<FSorterByDoubleValue>(ColumnRef);
 	Column.SetValueSorter(Sorter);
+	Column.SetInitialSortMode(EColumnSortMode::Descending);
 
 	return ColumnRef;
 }
@@ -656,6 +665,7 @@ TSharedRef<Insights::FTableColumn> FTimersViewColumnFactory::CreateMedianExclusi
 
 	TSharedRef<ITableCellValueSorter> Sorter = MakeShared<FSorterByDoubleValue>(ColumnRef);
 	Column.SetValueSorter(Sorter);
+	Column.SetInitialSortMode(EColumnSortMode::Descending);
 
 	return ColumnRef;
 }
@@ -701,6 +711,7 @@ TSharedRef<Insights::FTableColumn> FTimersViewColumnFactory::CreateMinExclusiveT
 
 	TSharedRef<ITableCellValueSorter> Sorter = MakeShared<FSorterByDoubleValue>(ColumnRef);
 	Column.SetValueSorter(Sorter);
+	Column.SetInitialSortMode(EColumnSortMode::Descending);
 
 	return ColumnRef;
 }
diff --git a/Engine/Source/Developer/TraceInsights/Private/Insights/Widgets/SStatsView.cpp b/Engine/Source/Developer/TraceInsights/Private/Insights/Widgets/SStatsView.cpp
--- a/Engine/Source/Developer/TraceInsights/Private/Insights/Widgets/SStatsView.cpp
+++ b/Engine/Source/Developer/TraceInsights/Private/Insights/Widgets/SStatsView.cpp
@@ -1590,6 +1590,7 @@ void SStatsView::ShowColumn(const FName ColumnId)
 		.VAlignHeader(VAlign_Center)
 		.HAlignCell(HAlign_Fill)
 		.VAlignCell(VAlign_Fill)
+		.InitialSortMode(Column.GetInitialSortMode())
 		.SortMode(this, &SStatsView::GetSortModeForColumn, Column.GetId())
 		.OnSort(this, &SStatsView::OnSortModeChanged)
 		.FillWidth(Column.GetInitialWidth())
diff --git a/Engine/Source/Developer/TraceInsights/Private/Insights/Widgets/STimerTreeView.cpp b/Engine/Source/Developer/TraceInsights/Private/Insights/Widgets/STimerTreeView.cpp
--- a/Engine/Source/Developer/TraceInsights/Private/Insights/Widgets/STimerTreeView.cpp
+++ b/Engine/Source/Developer/TraceInsights/Private/Insights/Widgets/STimerTreeView.cpp
@@ -835,6 +835,7 @@ void STimerTreeView::ShowColumn(const FName ColumnId)
 		.VAlignHeader(VAlign_Center)
 		.HAlignCell(HAlign_Fill)
 		.VAlignCell(VAlign_Fill)
+		.InitialSortMode(Column.GetInitialSortMode())
 		.SortMode(this, &STimerTreeView::GetSortModeForColumn, Column.GetId())
 		.OnSort(this, &STimerTreeView::OnSortModeChanged)
 		.FillWidth(Column.GetInitialWidth())
diff --git a/Engine/Source/Developer/TraceInsights/Private/Insights/Widgets/STimersView.cpp b/Engine/Source/Developer/TraceInsights/Private/Insights/Widgets/STimersView.cpp
--- a/Engine/Source/Developer/TraceInsights/Private/Insights/Widgets/STimersView.cpp
+++ b/Engine/Source/Developer/TraceInsights/Private/Insights/Widgets/STimersView.cpp
@@ -1677,6 +1677,7 @@ void STimersView::ShowColumn(const FName ColumnId)
 		.VAlignHeader(VAlign_Center)
 		.HAlignCell(HAlign_Fill)
 		.VAlignCell(VAlign_Fill)
+		.InitialSortMode(Column.GetInitialSortMode())
 		.SortMode(this, &STimersView::GetSortModeForColumn, Column.GetId())
 		.OnSort(this, &STimersView::OnSortModeChanged)
 		.FillWidth(Column.GetInitialWidth())
