[Insights] Add "No event found" error messages to Networking Profiler find next/last functionality.

diff --git a/Engine/Source/Developer/TraceInsights/Private/Insights/NetworkingProfiler/NetworkingProfilerManager.cpp b/Engine/Source/Developer/TraceInsights/Private/Insights/NetworkingProfiler/NetworkingProfilerManager.cpp
--- a/Engine/Source/Developer/TraceInsights/Private/Insights/NetworkingProfiler/NetworkingProfilerManager.cpp
+++ b/Engine/Source/Developer/TraceInsights/Private/Insights/NetworkingProfiler/NetworkingProfilerManager.cpp
@@ -2,6 +2,7 @@
 
 #include "NetworkingProfilerManager.h"
 
+#include "MessageLog/Public/MessageLogModule.h"
 #include "Modules/ModuleManager.h"
 #include "TraceServices/AnalysisService.h"
 #include "WorkspaceMenuStructure.h"
@@ -51,6 +52,7 @@ FNetworkingProfilerManager::FNetworkingProfilerManager(TSharedRef<FUICommandList
 	, CommandList(InCommandList)
 	, ActionManager(this)
 	, ProfilerWindows()
+	, LogListingName(TEXT("NetworkingInsights"))
 {
 }
 
@@ -88,6 +90,16 @@ void FNetworkingProfilerManager::Shutdown()
 	}
 	bIsInitialized = false;
 
+	// If the MessageLog module was already unloaded as part of the global Shutdown process, do not load it again.
+	if (FModuleManager::Get().IsModuleLoaded("MessageLog"))
+	{
+		FMessageLogModule& MessageLogModule = FModuleManager::LoadModuleChecked<FMessageLogModule>("MessageLog");
+		if (MessageLogModule.IsRegisteredLogListing(GetLogListingName()))
+		{
+			MessageLogModule.UnregisterLogListing(GetLogListingName());
+		}
+	}
+
 	FInsightsManager::Get()->GetSessionChangedEvent().RemoveAll(this);
 
 	FNetworkingProfilerCommands::Unregister();
@@ -260,6 +272,10 @@ bool FNetworkingProfilerManager::Tick(float DeltaTime)
 				FGlobalTabmanager::Get()->TryInvokeTab(TabId);
 			}
 
+			FMessageLogModule& MessageLogModule = FModuleManager::LoadModuleChecked<FMessageLogModule>("MessageLog");
+			MessageLogModule.RegisterLogListing(GetLogListingName(), LOCTEXT("NetworkingInsights", "Networking Profiler Insights"));
+			MessageLogModule.EnableMessageLogDisplay(true);
+
 			//int32 SpawnTabCount = 2; // we want to spawn 2 tabs
 			//for (int32 ReservedId = 0; SpawnTabCount > 0 && ReservedId < 10; ++ReservedId)
 			//{
diff --git a/Engine/Source/Developer/TraceInsights/Private/Insights/NetworkingProfiler/NetworkingProfilerManager.h b/Engine/Source/Developer/TraceInsights/Private/Insights/NetworkingProfiler/NetworkingProfilerManager.h
--- a/Engine/Source/Developer/TraceInsights/Private/Insights/NetworkingProfiler/NetworkingProfilerManager.h
+++ b/Engine/Source/Developer/TraceInsights/Private/Insights/NetworkingProfiler/NetworkingProfilerManager.h
@@ -83,6 +83,8 @@ public:
 
 	void OnSessionChanged();
 
+	const FName& GetLogListingName() const { return LogListingName; }
+
 private:
 	/** Binds our UI commands to delegates. */
 	void BindCommands();
@@ -130,4 +132,7 @@ private:
 
 	/** A shared pointer to the global instance of the NetworkingProfiler manager. */
 	static TSharedPtr<FNetworkingProfilerManager> Instance;
+
+	/** The name of the Networking Profiler log listing. */
+	FName LogListingName;
 };
diff --git a/Engine/Source/Developer/TraceInsights/Private/Insights/NetworkingProfiler/Widgets/SPacketContentView.cpp b/Engine/Source/Developer/TraceInsights/Private/Insights/NetworkingProfiler/Widgets/SPacketContentView.cpp
--- a/Engine/Source/Developer/TraceInsights/Private/Insights/NetworkingProfiler/Widgets/SPacketContentView.cpp
+++ b/Engine/Source/Developer/TraceInsights/Private/Insights/NetworkingProfiler/Widgets/SPacketContentView.cpp
@@ -7,6 +7,8 @@
 #include "Framework/Application/SlateApplication.h"
 #include "Framework/MultiBox/MultiBoxBuilder.h"
 #include "HAL/PlatformApplicationMisc.h"
+#include "Logging/MessageLog.h"
+#include "MessageLog/Public/MessageLogModule.h"
 #include "Rendering/DrawElements.h"
 #include "Styling/AppStyle.h"
 #include "Widgets/Images/SImage.h"
@@ -21,6 +23,7 @@
 #include "Insights/Common/Stopwatch.h"
 #include "Insights/InsightsManager.h"
 #include "Insights/InsightsStyle.h"
+#include "Insights/NetworkingProfiler/NetworkingProfilerManager.h"
 #include "Insights/NetworkingProfiler/Widgets/SNetworkingProfilerWindow.h"
 #include "Insights/NetworkingProfiler/Widgets/SPacketView.h"
 
@@ -403,6 +406,12 @@ void SPacketContentView::FindFirstEvent()
 		OnSelectedEventChanged();
 		BringEventIntoView(SelectedEvent);
 	}
+	else
+	{
+		FMessageLog ReportMessageLog(FNetworkingProfilerManager::Get()->GetLogListingName());
+		ReportMessageLog.Error(LOCTEXT("NoEventFound", "No event found!"));
+		ReportMessageLog.Notify();
+	}
 }
 
 ////////////////////////////////////////////////////////////////////////////////////////////////////
@@ -415,6 +424,8 @@ void SPacketContentView::FindPreviousEvent(EEventNavigationType NavigationType)
 		return;
 	}
 
+	FNetworkPacketEventRef PreviousSelectedEvent = SelectedEvent;
+
 	const int32 EventCount = FilteredDrawState->Events.Num();
 	for (int32 EventIndex = EventCount - 1; EventIndex >= 0; --EventIndex)
 	{
@@ -459,6 +470,13 @@ void SPacketContentView::FindPreviousEvent(EEventNavigationType NavigationType)
 	}
 
 	BringEventIntoView(SelectedEvent);
+
+	if (PreviousSelectedEvent.Equals(SelectedEvent))
+	{
+		FMessageLog ReportMessageLog(FNetworkingProfilerManager::Get()->GetLogListingName());
+		ReportMessageLog.Error(LOCTEXT("NoEventFound", "No event found!"));
+		ReportMessageLog.Notify();
+	}
 }
 
 ////////////////////////////////////////////////////////////////////////////////////////////////////
@@ -471,6 +489,8 @@ void SPacketContentView::FindNextEvent(EEventNavigationType NavigationType)
 		return;
 	}
 
+	FNetworkPacketEventRef PreviousSelectedEvent = SelectedEvent;
+
 	const int32 EventCount = FilteredDrawState->Events.Num();
 	for (int32 EventIndex = 0; EventIndex < EventCount; ++EventIndex)
 	{
@@ -515,6 +535,13 @@ void SPacketContentView::FindNextEvent(EEventNavigationType NavigationType)
 	}
 
 	BringEventIntoView(SelectedEvent);
+
+	if (PreviousSelectedEvent.Equals(SelectedEvent))
+	{
+		FMessageLog ReportMessageLog(FNetworkingProfilerManager::Get()->GetLogListingName());
+		ReportMessageLog.Error(LOCTEXT("NoEventFound", "No event found!"));
+		ReportMessageLog.Notify();
+	}
 }
 
 ////////////////////////////////////////////////////////////////////////////////////////////////////
@@ -527,6 +554,12 @@ void SPacketContentView::FindLastEvent()
 		OnSelectedEventChanged();
 		BringEventIntoView(SelectedEvent);
 	}
+	else
+	{
+		FMessageLog ReportMessageLog(FNetworkingProfilerManager::Get()->GetLogListingName());
+		ReportMessageLog.Error(LOCTEXT("NoEventFound", "No event found!"));
+		ReportMessageLog.Notify();
+	}
 }
 
 ////////////////////////////////////////////////////////////////////////////////////////////////////
