[Insights]
  - Session Browser: Added time limit for updating metadata for a trace. If the time limit is exceeded, the current metadata analysis stops, allowing other traces to update metadata, then it will retry with an increased time limit (up to 5 retries).

diff --git a/Engine/Source/Developer/TraceInsights/Private/Insights/StoreService/StoreBrowser.cpp b/Engine/Source/Developer/TraceInsights/Private/Insights/StoreService/StoreBrowser.cpp
--- a/Engine/Source/Developer/TraceInsights/Private/Insights/StoreService/StoreBrowser.cpp
+++ b/Engine/Source/Developer/TraceInsights/Private/Insights/StoreService/StoreBrowser.cpp
@@ -271,10 +271,9 @@ void FStoreBrowser::UpdateTraces()
 
 		int32 MetadataUpdateCount = 0; // for debugging
 
-		//if (false)
 		for (TSharedPtr<FStoreBrowserTraceInfo>& TracePtr : Traces)
 		{
-			if (!TracePtr->bIsMetadataUpdated)
+			if (TracePtr->MetadataUpdateCount > 0)
 			{
 				MetadataUpdateCount++;
 				UpdateMetadata(*TracePtr);
@@ -313,8 +312,6 @@ void FStoreBrowser::ResetTraces()
 
 void FStoreBrowser::UpdateMetadata(FStoreBrowserTraceInfo& Trace)
 {
-	//UE_LOG(TraceInsights, Log, TEXT("[StoreBrowser] Updating metadata for trace 0x%08X (%s)..."), TraceSession.TraceId, *TraceSession.Name.ToString());
-
 	UE::Trace::FStoreClient* StoreClient = GetStoreClient();
 	if (StoreClient == nullptr)
 	{
@@ -331,10 +328,25 @@ void FStoreBrowser::UpdateMetadata(FStoreBrowserTraceInfo& Trace)
 
 	struct FDataStream : public UE::Trace::IInDataStream
 	{
+		enum class EReadStatus
+		{
+			Ready = 0,
+			StoppedByReadSizeLimit,
+			StoppedByTimeLimit,
+		};
+
 		virtual int32 Read(void* Data, uint32 Size) override
 		{
 			if (BytesRead >= 1024 * 1024)
 			{
+				Status = EReadStatus::StoppedByReadSizeLimit;
+				return 0;
+			}
+
+			Stopwatch.Update();
+			if (Stopwatch.GetAccumulatedTime() > TimeLimit)
+			{
+				Status = EReadStatus::StoppedByTimeLimit;
 				return 0;
 			}
 
@@ -349,12 +361,17 @@ void FStoreBrowser::UpdateMetadata(FStoreBrowserTraceInfo& Trace)
 			Inner->Close();
 		}
 
-		int32 BytesRead = 0;
 		UE::Trace::IInDataStream* Inner;
+		double TimeLimit = 1.0;
+		FStopwatch Stopwatch;
+		int32 BytesRead = 0;
+		EReadStatus Status = EReadStatus::Ready;
 	};
 
 	FDataStream DataStream;
 	DataStream.Inner = TraceData.Get();
+	DataStream.TimeLimit = 1.0 + (double)(Trace.MetadataUpdateCount - 1) * 2.0; // 1s, 3s, 5s, ...
+	DataStream.Stopwatch.Start();
 
 	FDiagnosticsSessionAnalyzer Analyzer;
 	UE::Trace::FAnalysisContext Context;
@@ -364,11 +381,11 @@ void FStoreBrowser::UpdateMetadata(FStoreBrowserTraceInfo& Trace)
 	// Update the FStoreBrowserTraceInfo object.
 	{
 		FScopeLock Lock(&TracesCriticalSection);
-		TracesChangeSerial++;
-		Trace.ChangeSerial++;
-		Trace.bIsMetadataUpdated = true;
+
 		if (Analyzer.Platform.Len() != 0)
 		{
+			TracesChangeSerial++;
+			Trace.ChangeSerial++;
 			Trace.Platform = Analyzer.Platform;
 			Trace.AppName = Analyzer.AppName;
 			Trace.CommandLine = Analyzer.CommandLine;
@@ -378,6 +395,20 @@ void FStoreBrowser::UpdateMetadata(FStoreBrowserTraceInfo& Trace)
 			Trace.ConfigurationType = static_cast<EBuildConfiguration>(Analyzer.ConfigurationType);
 			Trace.TargetType = static_cast<EBuildTargetType>(Analyzer.TargetType);
 		}
+
+		if (DataStream.Status == FDataStream::EReadStatus::StoppedByTimeLimit)
+		{
+			// Try again later.
+			++Trace.MetadataUpdateCount;
+			if (Trace.MetadataUpdateCount > 5)
+			{
+				Trace.MetadataUpdateCount = 0; // no more updates
+			}
+		}
+		else
+		{
+			Trace.MetadataUpdateCount = 0; // no more updates
+		}
 	}
 }
 
diff --git a/Engine/Source/Developer/TraceInsights/Private/Insights/StoreService/StoreBrowser.h b/Engine/Source/Developer/TraceInsights/Private/Insights/StoreService/StoreBrowser.h
--- a/Engine/Source/Developer/TraceInsights/Private/Insights/StoreService/StoreBrowser.h
+++ b/Engine/Source/Developer/TraceInsights/Private/Insights/StoreService/StoreBrowser.h
@@ -45,8 +45,8 @@ struct FStoreBrowserTraceInfo
 	EBuildConfiguration ConfigurationType = EBuildConfiguration::Unknown;
 	EBuildTargetType TargetType = EBuildTargetType::Unknown;
 
-	bool bIsMetadataUpdated = false;
 	bool bIsLive = false;
+	uint8 MetadataUpdateCount = 1;
 
 	uint32 IpAddress = 0;
 
diff --git a/Engine/Source/Developer/TraceInsights/Private/Insights/Widgets/SStartPageWindow.cpp b/Engine/Source/Developer/TraceInsights/Private/Insights/Widgets/SStartPageWindow.cpp
--- a/Engine/Source/Developer/TraceInsights/Private/Insights/Widgets/SStartPageWindow.cpp
+++ b/Engine/Source/Developer/TraceInsights/Private/Insights/Widgets/SStartPageWindow.cpp
@@ -1171,7 +1171,7 @@ void STraceStoreWindow::UpdateTrace(FTraceViewModel& InOutTrace, const Insights:
 	InOutTrace.IpAddress = InSourceTrace.IpAddress;
 
 	// Is metadata updated?
-	if (!InOutTrace.bIsMetadataUpdated && InSourceTrace.bIsMetadataUpdated)
+	if (!InOutTrace.bIsMetadataUpdated && InSourceTrace.MetadataUpdateCount == 0)
 	{
 		InOutTrace.bIsMetadataUpdated = true;
 		InOutTrace.Platform = FText::FromString(InSourceTrace.Platform);
