[Insights]
  - Timing Insights: Implemented micro-optimizations to sorting code for Timers and Counters (15% improvement for a list of 100K timers; 30% improvement for a list of 16M timers).

diff --git a/Engine/Source/Developer/TraceInsights/Private/Insights/MemoryProfiler/ViewModels/MemTagNodeGroupingAndSorting.cpp b/Engine/Source/Developer/TraceInsights/Private/Insights/MemoryProfiler/ViewModels/MemTagNodeGroupingAndSorting.cpp
--- a/Engine/Source/Developer/TraceInsights/Private/Insights/MemoryProfiler/ViewModels/MemTagNodeGroupingAndSorting.cpp
+++ b/Engine/Source/Developer/TraceInsights/Private/Insights/MemoryProfiler/ViewModels/MemTagNodeGroupingAndSorting.cpp
@@ -7,6 +7,9 @@
 
 #define LOCTEXT_NAMESPACE "MemTagNode"
 
+#define INSIGHTS_ENSURE ensure
+//#define INSIGHTS_ENSURE(...)
+
 // Sort by name (ascending).
 #define INSIGHTS_DEFAULT_SORTING_NODES(A, B) return A->GetName().LexicalLess(B->GetName());
 //#define INSIGHTS_DEFAULT_SORTING_NODES(A, B) return A->GetDefaultSortOrder() < B->GetDefaultSortOrder();
@@ -33,10 +36,10 @@ void FMemTagNodeSortingByType::Sort(TArray<Insights::FBaseTreeNodePtr>& NodesToS
 	{
 		NodesToSort.Sort([](const Insights::FBaseTreeNodePtr& A, const Insights::FBaseTreeNodePtr& B) -> bool
 		{
-			ensure(A.IsValid() && A->GetTypeName() == FMemTagNode::TypeName);
+			INSIGHTS_ENSURE(A.IsValid() && A->GetTypeName() == FMemTagNode::TypeName);
 			const FMemTagNodePtr MemTagNodeA = StaticCastSharedPtr<FMemTagNode, Insights::FBaseTreeNode>(A);
 
-			ensure(B.IsValid() && B->GetTypeName() == FMemTagNode::TypeName);
+			INSIGHTS_ENSURE(B.IsValid() && B->GetTypeName() == FMemTagNode::TypeName);
 			const FMemTagNodePtr MemTagNodeB = StaticCastSharedPtr<FMemTagNode, Insights::FBaseTreeNode>(B);
 
 			if (MemTagNodeA->GetType() == MemTagNodeB->GetType())
@@ -54,10 +57,10 @@ void FMemTagNodeSortingByType::Sort(TArray<Insights::FBaseTreeNodePtr>& NodesToS
 	{
 		NodesToSort.Sort([](const Insights::FBaseTreeNodePtr& A, const Insights::FBaseTreeNodePtr& B) -> bool
 		{
-			ensure(A.IsValid() && A->GetTypeName() == FMemTagNode::TypeName);
+			INSIGHTS_ENSURE(A.IsValid() && A->GetTypeName() == FMemTagNode::TypeName);
 			const FMemTagNodePtr MemTagNodeA = StaticCastSharedPtr<FMemTagNode, Insights::FBaseTreeNode>(A);
 
-			ensure(B.IsValid() && B->GetTypeName() == FMemTagNode::TypeName);
+			INSIGHTS_ENSURE(B.IsValid() && B->GetTypeName() == FMemTagNode::TypeName);
 			const FMemTagNodePtr MemTagNodeB = StaticCastSharedPtr<FMemTagNode, Insights::FBaseTreeNode>(B);
 
 			if (MemTagNodeA->GetType() == MemTagNodeB->GetType())
@@ -95,10 +98,10 @@ void FMemTagNodeSortingByTracker::Sort(TArray<Insights::FBaseTreeNodePtr>& Nodes
 	{
 		NodesToSort.Sort([](const Insights::FBaseTreeNodePtr& A, const Insights::FBaseTreeNodePtr& B) -> bool
 		{
-			ensure(A.IsValid() && A->GetTypeName() == FMemTagNode::TypeName);
+			INSIGHTS_ENSURE(A.IsValid() && A->GetTypeName() == FMemTagNode::TypeName);
 			const FMemTagNodePtr MemTagNodeA = StaticCastSharedPtr<FMemTagNode, Insights::FBaseTreeNode>(A);
 
-			ensure(B.IsValid() && B->GetTypeName() == FMemTagNode::TypeName);
+			INSIGHTS_ENSURE(B.IsValid() && B->GetTypeName() == FMemTagNode::TypeName);
 			const FMemTagNodePtr MemTagNodeB = StaticCastSharedPtr<FMemTagNode, Insights::FBaseTreeNode>(B);
 
 			if (MemTagNodeA->GetMemTrackerId() == MemTagNodeB->GetMemTrackerId())
@@ -116,10 +119,10 @@ void FMemTagNodeSortingByTracker::Sort(TArray<Insights::FBaseTreeNodePtr>& Nodes
 	{
 		NodesToSort.Sort([](const Insights::FBaseTreeNodePtr& A, const Insights::FBaseTreeNodePtr& B) -> bool
 		{
-			ensure(A.IsValid() && A->GetTypeName() == FMemTagNode::TypeName);
+			INSIGHTS_ENSURE(A.IsValid() && A->GetTypeName() == FMemTagNode::TypeName);
 			const FMemTagNodePtr MemTagNodeA = StaticCastSharedPtr<FMemTagNode, Insights::FBaseTreeNode>(A);
 
-			ensure(B.IsValid() && B->GetTypeName() == FMemTagNode::TypeName);
+			INSIGHTS_ENSURE(B.IsValid() && B->GetTypeName() == FMemTagNode::TypeName);
 			const FMemTagNodePtr MemTagNodeB = StaticCastSharedPtr<FMemTagNode, Insights::FBaseTreeNode>(B);
 
 			if (MemTagNodeA->GetMemTrackerId() == MemTagNodeB->GetMemTrackerId())
@@ -157,11 +160,11 @@ void FMemTagNodeSortingByInstanceCount::Sort(TArray<Insights::FBaseTreeNodePtr>&
 	{
 		NodesToSort.Sort([](const Insights::FBaseTreeNodePtr& A, const Insights::FBaseTreeNodePtr& B) -> bool
 		{
-			ensure(A.IsValid() && A->GetTypeName() == FMemTagNode::TypeName);
+			INSIGHTS_ENSURE(A.IsValid() && A->GetTypeName() == FMemTagNode::TypeName);
 			const FMemTagNodePtr MemTagNodeA = StaticCastSharedPtr<FMemTagNode, Insights::FBaseTreeNode>(A);
 			const uint64 ValueA = MemTagNodeA->GetAggregatedStats().InstanceCount;
 
-			ensure(B.IsValid() && B->GetTypeName() == FMemTagNode::TypeName);
+			INSIGHTS_ENSURE(B.IsValid() && B->GetTypeName() == FMemTagNode::TypeName);
 			const FMemTagNodePtr MemTagNodeB = StaticCastSharedPtr<FMemTagNode, Insights::FBaseTreeNode>(B);
 			const uint64 ValueB = MemTagNodeB->GetAggregatedStats().InstanceCount;
 
@@ -180,11 +183,11 @@ void FMemTagNodeSortingByInstanceCount::Sort(TArray<Insights::FBaseTreeNodePtr>&
 	{
 		NodesToSort.Sort([](const Insights::FBaseTreeNodePtr& A, const Insights::FBaseTreeNodePtr& B) -> bool
 		{
-			ensure(A.IsValid() && A->GetTypeName() == FMemTagNode::TypeName);
+			INSIGHTS_ENSURE(A.IsValid() && A->GetTypeName() == FMemTagNode::TypeName);
 			const FMemTagNodePtr MemTagNodeA = StaticCastSharedPtr<FMemTagNode, Insights::FBaseTreeNode>(A);
 			const uint64 ValueA = MemTagNodeA->GetAggregatedStats().InstanceCount;
 
-			ensure(B.IsValid() && B->GetTypeName() == FMemTagNode::TypeName);
+			INSIGHTS_ENSURE(B.IsValid() && B->GetTypeName() == FMemTagNode::TypeName);
 			const FMemTagNodePtr MemTagNodeB = StaticCastSharedPtr<FMemTagNode, Insights::FBaseTreeNode>(B);
 			const uint64 ValueB = MemTagNodeB->GetAggregatedStats().InstanceCount;
 
@@ -204,4 +207,5 @@ void FMemTagNodeSortingByInstanceCount::Sort(TArray<Insights::FBaseTreeNodePtr>&
 ////////////////////////////////////////////////////////////////////////////////////////////////////
 
 #undef INSIGHTS_DEFAULT_SORTING_NODES
+#undef INSIGHTS_ENSURE
 #undef LOCTEXT_NAMESPACE
diff --git a/Engine/Source/Developer/TraceInsights/Private/Insights/NetworkingProfiler/ViewModels/NetEventGroupingAndSorting.cpp b/Engine/Source/Developer/TraceInsights/Private/Insights/NetworkingProfiler/ViewModels/NetEventGroupingAndSorting.cpp
--- a/Engine/Source/Developer/TraceInsights/Private/Insights/NetworkingProfiler/ViewModels/NetEventGroupingAndSorting.cpp
+++ b/Engine/Source/Developer/TraceInsights/Private/Insights/NetworkingProfiler/ViewModels/NetEventGroupingAndSorting.cpp
@@ -7,6 +7,9 @@
 
 #define LOCTEXT_NAMESPACE "NetEventNode"
 
+#define INSIGHTS_ENSURE ensure
+//#define INSIGHTS_ENSURE(...)
+ 
 // Sort by name (ascending).
 #define INSIGHTS_DEFAULT_SORTING_NODES(A, B) return A->GetName().LexicalLess(B->GetName());
 //#define INSIGHTS_DEFAULT_SORTING_NODES(A, B) return A->GetDefaultSortOrder() < B->GetDefaultSortOrder();
@@ -33,10 +36,10 @@ void FNetEventNodeSortingByEventType::Sort(TArray<Insights::FBaseTreeNodePtr>& N
 	{
 		NodesToSort.Sort([](const Insights::FBaseTreeNodePtr& A, const Insights::FBaseTreeNodePtr& B) -> bool
 		{
-			ensure(A.IsValid() && A->GetTypeName() == FNetEventNode::TypeName);
+			INSIGHTS_ENSURE(A.IsValid() && A->GetTypeName() == FNetEventNode::TypeName);
 			const FNetEventNodePtr NetEventNodeA = StaticCastSharedPtr<FNetEventNode, Insights::FBaseTreeNode>(A);
 
-			ensure(B.IsValid() && B->GetTypeName() == FNetEventNode::TypeName);
+			INSIGHTS_ENSURE(B.IsValid() && B->GetTypeName() == FNetEventNode::TypeName);
 			const FNetEventNodePtr NetEventNodeB = StaticCastSharedPtr<FNetEventNode, Insights::FBaseTreeNode>(B);
 
 			if (NetEventNodeA->GetType() == NetEventNodeB->GetType())
@@ -54,10 +57,10 @@ void FNetEventNodeSortingByEventType::Sort(TArray<Insights::FBaseTreeNodePtr>& N
 	{
 		NodesToSort.Sort([](const Insights::FBaseTreeNodePtr& A, const Insights::FBaseTreeNodePtr& B) -> bool
 		{
-			ensure(A.IsValid() && A->GetTypeName() == FNetEventNode::TypeName);
+			INSIGHTS_ENSURE(A.IsValid() && A->GetTypeName() == FNetEventNode::TypeName);
 			const FNetEventNodePtr NetEventNodeA = StaticCastSharedPtr<FNetEventNode, Insights::FBaseTreeNode>(A);
 
-			ensure(B.IsValid() && B->GetTypeName() == FNetEventNode::TypeName);
+			INSIGHTS_ENSURE(B.IsValid() && B->GetTypeName() == FNetEventNode::TypeName);
 			const FNetEventNodePtr NetEventNodeB = StaticCastSharedPtr<FNetEventNode, Insights::FBaseTreeNode>(B);
 
 			if (NetEventNodeA->GetType() == NetEventNodeB->GetType())
@@ -95,11 +98,11 @@ void FNetEventNodeSortingByInstanceCount::Sort(TArray<Insights::FBaseTreeNodePtr
 	{
 		NodesToSort.Sort([](const Insights::FBaseTreeNodePtr& A, const Insights::FBaseTreeNodePtr& B) -> bool
 		{
-			ensure(A.IsValid() && A->GetTypeName() == FNetEventNode::TypeName);
+			INSIGHTS_ENSURE(A.IsValid() && A->GetTypeName() == FNetEventNode::TypeName);
 			const FNetEventNodePtr NetEventNodeA = StaticCastSharedPtr<FNetEventNode, Insights::FBaseTreeNode>(A);
 			const uint64 ValueA = NetEventNodeA->GetAggregatedStats().InstanceCount;
 
-			ensure(B.IsValid() && B->GetTypeName() == FNetEventNode::TypeName);
+			INSIGHTS_ENSURE(B.IsValid() && B->GetTypeName() == FNetEventNode::TypeName);
 			const FNetEventNodePtr NetEventNodeB = StaticCastSharedPtr<FNetEventNode, Insights::FBaseTreeNode>(B);
 			const uint64 ValueB = NetEventNodeB->GetAggregatedStats().InstanceCount;
 
@@ -118,11 +121,11 @@ void FNetEventNodeSortingByInstanceCount::Sort(TArray<Insights::FBaseTreeNodePtr
 	{
 		NodesToSort.Sort([](const Insights::FBaseTreeNodePtr& A, const Insights::FBaseTreeNodePtr& B) -> bool
 		{
-			ensure(A.IsValid() && A->GetTypeName() == FNetEventNode::TypeName);
+			INSIGHTS_ENSURE(A.IsValid() && A->GetTypeName() == FNetEventNode::TypeName);
 			const FNetEventNodePtr NetEventNodeA = StaticCastSharedPtr<FNetEventNode, Insights::FBaseTreeNode>(A);
 			const uint64 ValueA = NetEventNodeA->GetAggregatedStats().InstanceCount;
 
-			ensure(B.IsValid() && B->GetTypeName() == FNetEventNode::TypeName);
+			INSIGHTS_ENSURE(B.IsValid() && B->GetTypeName() == FNetEventNode::TypeName);
 			const FNetEventNodePtr NetEventNodeB = StaticCastSharedPtr<FNetEventNode, Insights::FBaseTreeNode>(B);
 			const uint64 ValueB = NetEventNodeB->GetAggregatedStats().InstanceCount;
 
@@ -161,11 +164,11 @@ void FNetEventNodeSortingByTotalInclusiveSize::Sort(TArray<Insights::FBaseTreeNo
 	{
 		NodesToSort.Sort([](const Insights::FBaseTreeNodePtr& A, const Insights::FBaseTreeNodePtr& B) -> bool
 		{
-			ensure(A.IsValid() && A->GetTypeName() == FNetEventNode::TypeName);
+			INSIGHTS_ENSURE(A.IsValid() && A->GetTypeName() == FNetEventNode::TypeName);
 			const FNetEventNodePtr NetEventNodeA = StaticCastSharedPtr<FNetEventNode, Insights::FBaseTreeNode>(A);
 			const uint32 ValueA = NetEventNodeA->GetAggregatedStats().TotalInclusive;
 
-			ensure(B.IsValid() && B->GetTypeName() == FNetEventNode::TypeName);
+			INSIGHTS_ENSURE(B.IsValid() && B->GetTypeName() == FNetEventNode::TypeName);
 			const FNetEventNodePtr NetEventNodeB = StaticCastSharedPtr<FNetEventNode, Insights::FBaseTreeNode>(B);
 			const uint32 ValueB = NetEventNodeB->GetAggregatedStats().TotalInclusive;
 
@@ -184,11 +187,11 @@ void FNetEventNodeSortingByTotalInclusiveSize::Sort(TArray<Insights::FBaseTreeNo
 	{
 		NodesToSort.Sort([](const Insights::FBaseTreeNodePtr& A, const Insights::FBaseTreeNodePtr& B) -> bool
 		{
-			ensure(A.IsValid() && A->GetTypeName() == FNetEventNode::TypeName);
+			INSIGHTS_ENSURE(A.IsValid() && A->GetTypeName() == FNetEventNode::TypeName);
 			const FNetEventNodePtr NetEventNodeA = StaticCastSharedPtr<FNetEventNode, Insights::FBaseTreeNode>(A);
 			const uint32 ValueA = NetEventNodeA->GetAggregatedStats().TotalInclusive;
 
-			ensure(B.IsValid() && B->GetTypeName() == FNetEventNode::TypeName);
+			INSIGHTS_ENSURE(B.IsValid() && B->GetTypeName() == FNetEventNode::TypeName);
 			const FNetEventNodePtr NetEventNodeB = StaticCastSharedPtr<FNetEventNode, Insights::FBaseTreeNode>(B);
 			const uint32 ValueB = NetEventNodeB->GetAggregatedStats().TotalInclusive;
 
@@ -227,11 +230,11 @@ void FNetEventNodeSortingByTotalExclusiveSize::Sort(TArray<Insights::FBaseTreeNo
 	{
 		NodesToSort.Sort([](const Insights::FBaseTreeNodePtr& A, const Insights::FBaseTreeNodePtr& B) -> bool
 		{
-			ensure(A.IsValid() && A->GetTypeName() == FNetEventNode::TypeName);
+			INSIGHTS_ENSURE(A.IsValid() && A->GetTypeName() == FNetEventNode::TypeName);
 			const FNetEventNodePtr NetEventNodeA = StaticCastSharedPtr<FNetEventNode, Insights::FBaseTreeNode>(A);
 			const uint32 ValueA = NetEventNodeA->GetAggregatedStats().TotalExclusive;
 
-			ensure(B.IsValid() && B->GetTypeName() == FNetEventNode::TypeName);
+			INSIGHTS_ENSURE(B.IsValid() && B->GetTypeName() == FNetEventNode::TypeName);
 			const FNetEventNodePtr NetEventNodeB = StaticCastSharedPtr<FNetEventNode, Insights::FBaseTreeNode>(B);
 			const uint32 ValueB = NetEventNodeB->GetAggregatedStats().TotalExclusive;
 
@@ -250,11 +253,11 @@ void FNetEventNodeSortingByTotalExclusiveSize::Sort(TArray<Insights::FBaseTreeNo
 	{
 		NodesToSort.Sort([](const Insights::FBaseTreeNodePtr& A, const Insights::FBaseTreeNodePtr& B) -> bool
 		{
-			ensure(A.IsValid() && A->GetTypeName() == FNetEventNode::TypeName);
+			INSIGHTS_ENSURE(A.IsValid() && A->GetTypeName() == FNetEventNode::TypeName);
 			const FNetEventNodePtr NetEventNodeA = StaticCastSharedPtr<FNetEventNode, Insights::FBaseTreeNode>(A);
 			const uint32 ValueA = NetEventNodeA->GetAggregatedStats().TotalExclusive;
 
-			ensure(B.IsValid() && B->GetTypeName() == FNetEventNode::TypeName);
+			INSIGHTS_ENSURE(B.IsValid() && B->GetTypeName() == FNetEventNode::TypeName);
 			const FNetEventNodePtr NetEventNodeB = StaticCastSharedPtr<FNetEventNode, Insights::FBaseTreeNode>(B);
 			const uint32 ValueB = NetEventNodeB->GetAggregatedStats().TotalExclusive;
 
@@ -274,4 +277,5 @@ void FNetEventNodeSortingByTotalExclusiveSize::Sort(TArray<Insights::FBaseTreeNo
 ////////////////////////////////////////////////////////////////////////////////////////////////////
 
 #undef INSIGHTS_DEFAULT_SORTING_NODES
+#undef INSIGHTS_ENSURE
 #undef LOCTEXT_NAMESPACE
diff --git a/Engine/Source/Developer/TraceInsights/Private/Insights/ViewModels/StatsGroupingAndSorting.cpp b/Engine/Source/Developer/TraceInsights/Private/Insights/ViewModels/StatsGroupingAndSorting.cpp
--- a/Engine/Source/Developer/TraceInsights/Private/Insights/ViewModels/StatsGroupingAndSorting.cpp
+++ b/Engine/Source/Developer/TraceInsights/Private/Insights/ViewModels/StatsGroupingAndSorting.cpp
@@ -7,6 +7,9 @@
 
 #define LOCTEXT_NAMESPACE "StatsNode"
 
+//#define INSIGHTS_ENSURE ensure
+#define INSIGHTS_ENSURE(...)
+
 // Sort by name (ascending).
 #define INSIGHTS_DEFAULT_SORTING_NODES(A, B) return A->GetName().LexicalLess(B->GetName());
 //#define INSIGHTS_DEFAULT_SORTING_NODES(A, B) return A->GetDefaultSortOrder() < B->GetDefaultSortOrder();
@@ -20,7 +23,7 @@ FStatsNodeSortingByStatsType::FStatsNodeSortingByStatsType(TSharedRef<Insights::
 		FName(TEXT("ByStatsType")),
 		LOCTEXT("Sorting_ByStatsType_Name", "By Type"),
 		LOCTEXT("Sorting_ByStatsType_Title", "Sort By Type"),
-		LOCTEXT("Sorting_ByStatsType_Desc", "Sort by stats counter type."),
+		LOCTEXT("Sorting_ByStatsType_Desc", "Sort by counter type."),
 		InColumnRef)
 {
 }
@@ -33,20 +36,20 @@ void FStatsNodeSortingByStatsType::Sort(TArray<Insights::FBaseTreeNodePtr>& Node
 	{
 		NodesToSort.Sort([](const Insights::FBaseTreeNodePtr& A, const Insights::FBaseTreeNodePtr& B) -> bool
 		{
-			ensure(A.IsValid() && A->GetTypeName() == FStatsNode::TypeName);
-			const FStatsNodePtr StatsNodeA = StaticCastSharedPtr<FStatsNode, Insights::FBaseTreeNode>(A);
+			INSIGHTS_ENSURE(A.IsValid() && A->GetTypeName() == FStatsNode::TypeName);
+			const EStatsNodeType ValueA = reinterpret_cast<FStatsNode*>(A.Get())->GetType();
 
-			ensure(B.IsValid() && B->GetTypeName() == FStatsNode::TypeName);
-			const FStatsNodePtr StatsNodeB = StaticCastSharedPtr<FStatsNode, Insights::FBaseTreeNode>(B);
+			INSIGHTS_ENSURE(B.IsValid() && B->GetTypeName() == FStatsNode::TypeName);
+			const EStatsNodeType ValueB = reinterpret_cast<FStatsNode*>(B.Get())->GetType();
 
-			if (StatsNodeA->GetType() == StatsNodeB->GetType())
+			if (ValueA == ValueB)
 			{
 				INSIGHTS_DEFAULT_SORTING_NODES(A, B)
 			}
 			else
 			{
 				// Sort by node type (ascending).
-				return StatsNodeA->GetType() < StatsNodeB->GetType();
+				return ValueA < ValueB;
 			}
 		});
 	}
@@ -54,20 +57,20 @@ void FStatsNodeSortingByStatsType::Sort(TArray<Insights::FBaseTreeNodePtr>& Node
 	{
 		NodesToSort.Sort([](const Insights::FBaseTreeNodePtr& A, const Insights::FBaseTreeNodePtr& B) -> bool
 		{
-			ensure(A.IsValid() && A->GetTypeName() == FStatsNode::TypeName);
-			const FStatsNodePtr StatsNodeA = StaticCastSharedPtr<FStatsNode, Insights::FBaseTreeNode>(A);
+			INSIGHTS_ENSURE(A.IsValid() && A->GetTypeName() == FStatsNode::TypeName);
+			const EStatsNodeType ValueA = reinterpret_cast<FStatsNode*>(A.Get())->GetType();
 
-			ensure(B.IsValid() && B->GetTypeName() == FStatsNode::TypeName);
-			const FStatsNodePtr StatsNodeB = StaticCastSharedPtr<FStatsNode, Insights::FBaseTreeNode>(B);
+			INSIGHTS_ENSURE(B.IsValid() && B->GetTypeName() == FStatsNode::TypeName);
+			const EStatsNodeType ValueB = reinterpret_cast<FStatsNode*>(B.Get())->GetType();
 
-			if (StatsNodeA->GetType() == StatsNodeB->GetType())
+			if (ValueA == ValueB)
 			{
 				INSIGHTS_DEFAULT_SORTING_NODES(A, B)
 			}
 			else
 			{
 				// Sort by node type (descending).
-				return StatsNodeB->GetType() < StatsNodeA->GetType();
+				return ValueB < ValueA;
 			}
 		});
 	}
@@ -95,20 +98,20 @@ void FStatsNodeSortingByDataType::Sort(TArray<Insights::FBaseTreeNodePtr>& Nodes
 	{
 		NodesToSort.Sort([](const Insights::FBaseTreeNodePtr& A, const Insights::FBaseTreeNodePtr& B) -> bool
 		{
-			ensure(A.IsValid() && A->GetTypeName() == FStatsNode::TypeName);
-			const FStatsNodePtr StatsNodeA = StaticCastSharedPtr<FStatsNode, Insights::FBaseTreeNode>(A);
+			INSIGHTS_ENSURE(A.IsValid() && A->GetTypeName() == FStatsNode::TypeName);
+			const EStatsNodeDataType ValueA = reinterpret_cast<FStatsNode*>(A.Get())->GetDataType();
 
-			ensure(B.IsValid() && B->GetTypeName() == FStatsNode::TypeName);
-			const FStatsNodePtr StatsNodeB = StaticCastSharedPtr<FStatsNode, Insights::FBaseTreeNode>(B);
+			INSIGHTS_ENSURE(B.IsValid() && B->GetTypeName() == FStatsNode::TypeName);
+			const EStatsNodeDataType ValueB = reinterpret_cast<FStatsNode*>(B.Get())->GetDataType();
 
-			if (StatsNodeA->GetDataType() == StatsNodeB->GetDataType())
+			if (ValueA == ValueB)
 			{
 				INSIGHTS_DEFAULT_SORTING_NODES(A, B)
 			}
 			else
 			{
 				// Sort by data type (ascending).
-				return StatsNodeA->GetDataType() < StatsNodeB->GetDataType();
+				return ValueA < ValueB;
 			}
 		});
 	}
@@ -116,20 +119,20 @@ void FStatsNodeSortingByDataType::Sort(TArray<Insights::FBaseTreeNodePtr>& Nodes
 	{
 		NodesToSort.Sort([](const Insights::FBaseTreeNodePtr& A, const Insights::FBaseTreeNodePtr& B) -> bool
 		{
-			ensure(A.IsValid() && A->GetTypeName() == FStatsNode::TypeName);
-			const FStatsNodePtr StatsNodeA = StaticCastSharedPtr<FStatsNode, Insights::FBaseTreeNode>(A);
+			INSIGHTS_ENSURE(A.IsValid() && A->GetTypeName() == FStatsNode::TypeName);
+			const EStatsNodeDataType ValueA = reinterpret_cast<FStatsNode*>(A.Get())->GetDataType();
 
-			ensure(B.IsValid() && B->GetTypeName() == FStatsNode::TypeName);
-			const FStatsNodePtr StatsNodeB = StaticCastSharedPtr<FStatsNode, Insights::FBaseTreeNode>(B);
+			INSIGHTS_ENSURE(B.IsValid() && B->GetTypeName() == FStatsNode::TypeName);
+			const EStatsNodeDataType ValueB = reinterpret_cast<FStatsNode*>(B.Get())->GetDataType();
 
-			if (StatsNodeA->GetDataType() == StatsNodeB->GetDataType())
+			if (ValueA == ValueB)
 			{
 				INSIGHTS_DEFAULT_SORTING_NODES(A, B)
 			}
 			else
 			{
 				// Sort by data type (descending).
-				return StatsNodeB->GetDataType() < StatsNodeA->GetDataType();
+				return ValueB < ValueA;
 			}
 		});
 	}
@@ -157,13 +160,11 @@ void FStatsNodeSortingByCount::Sort(TArray<Insights::FBaseTreeNodePtr>& NodesToS
 	{
 		NodesToSort.Sort([](const Insights::FBaseTreeNodePtr& A, const Insights::FBaseTreeNodePtr& B) -> bool
 		{
-			ensure(A.IsValid() && A->GetTypeName() == FStatsNode::TypeName);
-			const FStatsNodePtr StatsNodeA = StaticCastSharedPtr<FStatsNode, Insights::FBaseTreeNode>(A);
-			const uint64 ValueA = StatsNodeA->GetAggregatedStats().Count;
+			INSIGHTS_ENSURE(A.IsValid() && A->GetTypeName() == FStatsNode::TypeName);
+			const uint64 ValueA = reinterpret_cast<FStatsNode*>(A.Get())->GetAggregatedStats().Count;
 
-			ensure(B.IsValid() && B->GetTypeName() == FStatsNode::TypeName);
-			const FStatsNodePtr StatsNodeB = StaticCastSharedPtr<FStatsNode, Insights::FBaseTreeNode>(B);
-			const uint64 ValueB = StatsNodeB->GetAggregatedStats().Count;
+			INSIGHTS_ENSURE(B.IsValid() && B->GetTypeName() == FStatsNode::TypeName);
+			const uint64 ValueB = reinterpret_cast<FStatsNode*>(B.Get())->GetAggregatedStats().Count;
 
 			if (ValueA == ValueB)
 			{
@@ -180,13 +181,11 @@ void FStatsNodeSortingByCount::Sort(TArray<Insights::FBaseTreeNodePtr>& NodesToS
 	{
 		NodesToSort.Sort([](const Insights::FBaseTreeNodePtr& A, const Insights::FBaseTreeNodePtr& B) -> bool
 		{
-			ensure(A.IsValid() && A->GetTypeName() == FStatsNode::TypeName);
-			const FStatsNodePtr StatsNodeA = StaticCastSharedPtr<FStatsNode, Insights::FBaseTreeNode>(A);
-			const uint64 ValueA = StatsNodeA->GetAggregatedStats().Count;
+			INSIGHTS_ENSURE(A.IsValid() && A->GetTypeName() == FStatsNode::TypeName);
+			const uint64 ValueA = reinterpret_cast<FStatsNode*>(A.Get())->GetAggregatedStats().Count;
 
-			ensure(B.IsValid() && B->GetTypeName() == FStatsNode::TypeName);
-			const FStatsNodePtr StatsNodeB = StaticCastSharedPtr<FStatsNode, Insights::FBaseTreeNode>(B);
-			const uint64 ValueB = StatsNodeB->GetAggregatedStats().Count;
+			INSIGHTS_ENSURE(B.IsValid() && B->GetTypeName() == FStatsNode::TypeName);
+			const uint64 ValueB = reinterpret_cast<FStatsNode*>(B.Get())->GetAggregatedStats().Count;
 
 			if (ValueA == ValueB)
 			{
@@ -204,4 +203,5 @@ void FStatsNodeSortingByCount::Sort(TArray<Insights::FBaseTreeNodePtr>& NodesToS
 ////////////////////////////////////////////////////////////////////////////////////////////////////
 
 #undef INSIGHTS_DEFAULT_SORTING_NODES
+#undef INSIGHTS_ENSURE
 #undef LOCTEXT_NAMESPACE
diff --git a/Engine/Source/Developer/TraceInsights/Private/Insights/ViewModels/TimerGroupingAndSorting.cpp b/Engine/Source/Developer/TraceInsights/Private/Insights/ViewModels/TimerGroupingAndSorting.cpp
--- a/Engine/Source/Developer/TraceInsights/Private/Insights/ViewModels/TimerGroupingAndSorting.cpp
+++ b/Engine/Source/Developer/TraceInsights/Private/Insights/ViewModels/TimerGroupingAndSorting.cpp
@@ -7,6 +7,9 @@
 
 #define LOCTEXT_NAMESPACE "TimerNode"
 
+//#define INSIGHTS_ENSURE ensure
+#define INSIGHTS_ENSURE(...)
+
 // Sort by name (ascending).
 #define INSIGHTS_DEFAULT_SORTING_NODES(A, B) return A->GetName().LexicalLess(B->GetName());
 //#define INSIGHTS_DEFAULT_SORTING_NODES(A, B) return A->GetDefaultSortOrder() < B->GetDefaultSortOrder();
@@ -33,20 +36,20 @@ void FTimerNodeSortingByTimerType::Sort(TArray<Insights::FBaseTreeNodePtr>& Node
 	{
 		NodesToSort.Sort([](const Insights::FBaseTreeNodePtr& A, const Insights::FBaseTreeNodePtr& B) -> bool
 		{
-			ensure(A.IsValid() && A->GetTypeName() == FTimerNode::TypeName);
-			const FTimerNodePtr TimerNodeA = StaticCastSharedPtr<FTimerNode, Insights::FBaseTreeNode>(A);
+			INSIGHTS_ENSURE(A.IsValid() && A->GetTypeName() == FTimerNode::TypeName);
+			const ETimerNodeType ValueA = reinterpret_cast<FTimerNode*>(A.Get())->GetType();
 
-			ensure(B.IsValid() && B->GetTypeName() == FTimerNode::TypeName);
-			const FTimerNodePtr TimerNodeB = StaticCastSharedPtr<FTimerNode, Insights::FBaseTreeNode>(B);
+			INSIGHTS_ENSURE(B.IsValid() && B->GetTypeName() == FTimerNode::TypeName);
+			const ETimerNodeType ValueB = reinterpret_cast<FTimerNode*>(B.Get())->GetType();
 
-			if (TimerNodeA->GetType() == TimerNodeB->GetType())
+			if (ValueA == ValueB)
 			{
 				INSIGHTS_DEFAULT_SORTING_NODES(A, B)
 			}
 			else
 			{
 				// Sort by timer type (ascending).
-				return TimerNodeA->GetType() < TimerNodeB->GetType();
+				return ValueA < ValueB;
 			}
 		});
 	}
@@ -54,20 +57,20 @@ void FTimerNodeSortingByTimerType::Sort(TArray<Insights::FBaseTreeNodePtr>& Node
 	{
 		NodesToSort.Sort([](const Insights::FBaseTreeNodePtr& A, const Insights::FBaseTreeNodePtr& B) -> bool
 		{
-			ensure(A.IsValid() && A->GetTypeName() == FTimerNode::TypeName);
-			const FTimerNodePtr TimerNodeA = StaticCastSharedPtr<FTimerNode, Insights::FBaseTreeNode>(A);
+			INSIGHTS_ENSURE(A.IsValid() && A->GetTypeName() == FTimerNode::TypeName);
+			const ETimerNodeType ValueA = reinterpret_cast<FTimerNode*>(A.Get())->GetType();
 
-			ensure(B.IsValid() && B->GetTypeName() == FTimerNode::TypeName);
-			const FTimerNodePtr TimerNodeB = StaticCastSharedPtr<FTimerNode, Insights::FBaseTreeNode>(B);
+			INSIGHTS_ENSURE(B.IsValid() && B->GetTypeName() == FTimerNode::TypeName);
+			const ETimerNodeType ValueB = reinterpret_cast<FTimerNode*>(B.Get())->GetType();
 
-			if (TimerNodeA->GetType() == TimerNodeB->GetType())
+			if (ValueA == ValueB)
 			{
 				INSIGHTS_DEFAULT_SORTING_NODES(A, B)
 			}
 			else
 			{
 				// Sort by timer type (descending).
-				return TimerNodeB->GetType() < TimerNodeA->GetType();
+				return ValueB < ValueA;
 			}
 		});
 	}
@@ -95,13 +98,11 @@ void FTimerNodeSortingByInstanceCount::Sort(TArray<Insights::FBaseTreeNodePtr>&
 	{
 		NodesToSort.Sort([](const Insights::FBaseTreeNodePtr& A, const Insights::FBaseTreeNodePtr& B) -> bool
 		{
-			ensure(A.IsValid() && A->GetTypeName() == FTimerNode::TypeName);
-			const FTimerNodePtr TimerNodeA = StaticCastSharedPtr<FTimerNode, Insights::FBaseTreeNode>(A);
-			const uint64 ValueA = TimerNodeA->GetAggregatedStats().InstanceCount;
+			INSIGHTS_ENSURE(A.IsValid() && A->GetTypeName() == FTimerNode::TypeName);
+			const uint64 ValueA = reinterpret_cast<FTimerNode*>(A.Get())->GetAggregatedStats().InstanceCount;
 
-			ensure(B.IsValid() && B->GetTypeName() == FTimerNode::TypeName);
-			const FTimerNodePtr TimerNodeB = StaticCastSharedPtr<FTimerNode, Insights::FBaseTreeNode>(B);
-			const uint64 ValueB = TimerNodeB->GetAggregatedStats().InstanceCount;
+			INSIGHTS_ENSURE(B.IsValid() && B->GetTypeName() == FTimerNode::TypeName);
+			const uint64 ValueB = reinterpret_cast<FTimerNode*>(B.Get())->GetAggregatedStats().InstanceCount;
 
 			if (ValueA == ValueB)
 			{
@@ -118,13 +119,11 @@ void FTimerNodeSortingByInstanceCount::Sort(TArray<Insights::FBaseTreeNodePtr>&
 	{
 		NodesToSort.Sort([](const Insights::FBaseTreeNodePtr& A, const Insights::FBaseTreeNodePtr& B) -> bool
 		{
-			ensure(A.IsValid() && A->GetTypeName() == FTimerNode::TypeName);
-			const FTimerNodePtr TimerNodeA = StaticCastSharedPtr<FTimerNode, Insights::FBaseTreeNode>(A);
-			const uint64 ValueA = TimerNodeA->GetAggregatedStats().InstanceCount;
+			INSIGHTS_ENSURE(A.IsValid() && A->GetTypeName() == FTimerNode::TypeName);
+			const uint64 ValueA = reinterpret_cast<FTimerNode*>(A.Get())->GetAggregatedStats().InstanceCount;
 
-			ensure(B.IsValid() && B->GetTypeName() == FTimerNode::TypeName);
-			const FTimerNodePtr TimerNodeB = StaticCastSharedPtr<FTimerNode, Insights::FBaseTreeNode>(B);
-			const uint64 ValueB = TimerNodeB->GetAggregatedStats().InstanceCount;
+			INSIGHTS_ENSURE(B.IsValid() && B->GetTypeName() == FTimerNode::TypeName);
+			const uint64 ValueB = reinterpret_cast<FTimerNode*>(B.Get())->GetAggregatedStats().InstanceCount;
 
 			if (ValueA == ValueB)
 			{
@@ -161,13 +160,11 @@ void FTimerNodeSortingByTotalInclusiveTime::Sort(TArray<Insights::FBaseTreeNodeP
 	{
 		NodesToSort.Sort([](const Insights::FBaseTreeNodePtr& A, const Insights::FBaseTreeNodePtr& B) -> bool
 		{
-			ensure(A.IsValid() && A->GetTypeName() == FTimerNode::TypeName);
-			const FTimerNodePtr TimerNodeA = StaticCastSharedPtr<FTimerNode, Insights::FBaseTreeNode>(A);
-			const double ValueA = TimerNodeA->GetAggregatedStats().TotalInclusiveTime;
+			INSIGHTS_ENSURE(A.IsValid() && A->GetTypeName() == FTimerNode::TypeName);
+			const double ValueA = reinterpret_cast<FTimerNode*>(A.Get())->GetAggregatedStats().TotalInclusiveTime;
 
-			ensure(B.IsValid() && B->GetTypeName() == FTimerNode::TypeName);
-			const FTimerNodePtr TimerNodeB = StaticCastSharedPtr<FTimerNode, Insights::FBaseTreeNode>(B);
-			const double ValueB = TimerNodeB->GetAggregatedStats().TotalInclusiveTime;
+			INSIGHTS_ENSURE(B.IsValid() && B->GetTypeName() == FTimerNode::TypeName);
+			const double ValueB = reinterpret_cast<FTimerNode*>(B.Get())->GetAggregatedStats().TotalInclusiveTime;
 
 			if (ValueA == ValueB)
 			{
@@ -184,13 +181,11 @@ void FTimerNodeSortingByTotalInclusiveTime::Sort(TArray<Insights::FBaseTreeNodeP
 	{
 		NodesToSort.Sort([](const Insights::FBaseTreeNodePtr& A, const Insights::FBaseTreeNodePtr& B) -> bool
 		{
-			ensure(A.IsValid() && A->GetTypeName() == FTimerNode::TypeName);
-			const FTimerNodePtr TimerNodeA = StaticCastSharedPtr<FTimerNode, Insights::FBaseTreeNode>(A);
-			const double ValueA = TimerNodeA->GetAggregatedStats().TotalInclusiveTime;
+			INSIGHTS_ENSURE(A.IsValid() && A->GetTypeName() == FTimerNode::TypeName);
+			const double ValueA = reinterpret_cast<FTimerNode*>(A.Get())->GetAggregatedStats().TotalInclusiveTime;
 
-			ensure(B.IsValid() && B->GetTypeName() == FTimerNode::TypeName);
-			const FTimerNodePtr TimerNodeB = StaticCastSharedPtr<FTimerNode, Insights::FBaseTreeNode>(B);
-			const double ValueB = TimerNodeB->GetAggregatedStats().TotalInclusiveTime;
+			INSIGHTS_ENSURE(B.IsValid() && B->GetTypeName() == FTimerNode::TypeName);
+			const double ValueB = reinterpret_cast<FTimerNode*>(B.Get())->GetAggregatedStats().TotalInclusiveTime;
 
 			if (ValueA == ValueB)
 			{
@@ -227,13 +222,11 @@ void FTimerNodeSortingByTotalExclusiveTime::Sort(TArray<Insights::FBaseTreeNodeP
 	{
 		NodesToSort.Sort([](const Insights::FBaseTreeNodePtr& A, const Insights::FBaseTreeNodePtr& B) -> bool
 		{
-			ensure(A.IsValid() && A->GetTypeName() == FTimerNode::TypeName);
-			const FTimerNodePtr TimerNodeA = StaticCastSharedPtr<FTimerNode, Insights::FBaseTreeNode>(A);
-			const double ValueA = TimerNodeA->GetAggregatedStats().TotalExclusiveTime;
+			INSIGHTS_ENSURE(A.IsValid() && A->GetTypeName() == FTimerNode::TypeName);
+			const double ValueA = reinterpret_cast<FTimerNode*>(A.Get())->GetAggregatedStats().TotalExclusiveTime;
 
-			ensure(B.IsValid() && B->GetTypeName() == FTimerNode::TypeName);
-			const FTimerNodePtr TimerNodeB = StaticCastSharedPtr<FTimerNode, Insights::FBaseTreeNode>(B);
-			const double ValueB = TimerNodeB->GetAggregatedStats().TotalExclusiveTime;
+			INSIGHTS_ENSURE(B.IsValid() && B->GetTypeName() == FTimerNode::TypeName);
+			const double ValueB = reinterpret_cast<FTimerNode*>(B.Get())->GetAggregatedStats().TotalExclusiveTime;
 
 			if (ValueA == ValueB)
 			{
@@ -250,13 +243,11 @@ void FTimerNodeSortingByTotalExclusiveTime::Sort(TArray<Insights::FBaseTreeNodeP
 	{
 		NodesToSort.Sort([](const Insights::FBaseTreeNodePtr& A, const Insights::FBaseTreeNodePtr& B) -> bool
 		{
-			ensure(A.IsValid() && A->GetTypeName() == FTimerNode::TypeName);
-			const FTimerNodePtr TimerNodeA = StaticCastSharedPtr<FTimerNode, Insights::FBaseTreeNode>(A);
-			const double ValueA = TimerNodeA->GetAggregatedStats().TotalExclusiveTime;
+			INSIGHTS_ENSURE(A.IsValid() && A->GetTypeName() == FTimerNode::TypeName);
+			const double ValueA = reinterpret_cast<FTimerNode*>(A.Get())->GetAggregatedStats().TotalExclusiveTime;
 
-			ensure(B.IsValid() && B->GetTypeName() == FTimerNode::TypeName);
-			const FTimerNodePtr TimerNodeB = StaticCastSharedPtr<FTimerNode, Insights::FBaseTreeNode>(B);
-			const double ValueB = TimerNodeB->GetAggregatedStats().TotalExclusiveTime;
+			INSIGHTS_ENSURE(B.IsValid() && B->GetTypeName() == FTimerNode::TypeName);
+			const double ValueB = reinterpret_cast<FTimerNode*>(B.Get())->GetAggregatedStats().TotalExclusiveTime;
 
 			if (ValueA == ValueB)
 			{
@@ -274,4 +265,5 @@ void FTimerNodeSortingByTotalExclusiveTime::Sort(TArray<Insights::FBaseTreeNodeP
 ////////////////////////////////////////////////////////////////////////////////////////////////////
 
 #undef INSIGHTS_DEFAULT_SORTING_NODES
+#undef INSIGHTS_ENSURE
 #undef LOCTEXT_NAMESPACE
diff --git a/Engine/Source/Developer/TraceInsights/Private/Insights/Widgets/SStatsView.cpp b/Engine/Source/Developer/TraceInsights/Private/Insights/Widgets/SStatsView.cpp
--- a/Engine/Source/Developer/TraceInsights/Private/Insights/Widgets/SStatsView.cpp
+++ b/Engine/Source/Developer/TraceInsights/Private/Insights/Widgets/SStatsView.cpp
@@ -757,34 +757,18 @@ void SStatsView::InsightsManager_OnSessionAnalysisCompleted()
 
 void SStatsView::UpdateTree()
 {
-	FStopwatch Stopwatch;
-	Stopwatch.Start();
-
 	CreateGroups();
-
-	Stopwatch.Update();
-	const double Time1 = Stopwatch.GetAccumulatedTime();
-
 	SortTreeNodes();
-
-	Stopwatch.Update();
-	const double Time2 = Stopwatch.GetAccumulatedTime();
-
 	ApplyFiltering();
-
-	Stopwatch.Stop();
-	const double TotalTime = Stopwatch.GetAccumulatedTime();
-	if (TotalTime > 0.1)
-	{
-		UE_LOG(TimingProfiler, Log, TEXT("[Counters] Tree view updated in %.3fs (%d counters) --> G:%.3fs + S:%.3fs + F:%.3fs"),
-			TotalTime, StatsNodes.Num(), Time1, Time2 - Time1, TotalTime - Time2);
-	}
 }
 
 ////////////////////////////////////////////////////////////////////////////////////////////////////
 
 void SStatsView::ApplyFiltering()
 {
+	FStopwatch Stopwatch;
+	Stopwatch.Start();
+
 	FilteredGroupNodes.Reset();
 
 	// Apply filter to all groups and its children.
@@ -928,6 +912,14 @@ void SStatsView::ApplyFiltering()
 
 	// Request tree refresh
 	TreeView->RequestTreeRefresh();
+
+	Stopwatch.Stop();
+	const double TotalTime = Stopwatch.GetAccumulatedTime();
+	if (TotalTime > 0.1)
+	{
+		UE_LOG(TimingProfiler, Log, TEXT("[Counters] Tree view filtered in %.3fs (%d counters)"),
+			TotalTime, StatsNodes.Num());
+	}
 }
 
 ////////////////////////////////////////////////////////////////////////////////////////////////////
@@ -1200,6 +1192,9 @@ bool SStatsView::SearchBox_IsEnabled() const
 
 void SStatsView::CreateGroups()
 {
+	FStopwatch Stopwatch;
+	Stopwatch.Start();
+
 	if (GroupingMode == EStatsGroupingMode::Flat)
 	{
 		GroupNodes.Reset();
@@ -1325,6 +1320,14 @@ void SStatsView::CreateGroups()
 		GroupNodeSet.KeySort([](const uint32& A, const uint32& B) { return A > B; }); // sort groups by order
 		GroupNodeSet.GenerateValueArray(GroupNodes);
 	}
+
+	Stopwatch.Stop();
+	const double TotalTime = Stopwatch.GetAccumulatedTime();
+	if (TotalTime > 0.1)
+	{
+		UE_LOG(TimingProfiler, Log, TEXT("[Counters] Tree view grouping updated in %.3fs (%d counters)"),
+			TotalTime, StatsNodes.Num());
+	}
 }
 
 ////////////////////////////////////////////////////////////////////////////////////////////////////
@@ -1437,6 +1440,9 @@ void SStatsView::UpdateCurrentSortingByColumn()
 
 void SStatsView::SortTreeNodes()
 {
+	FStopwatch Stopwatch;
+	Stopwatch.Start();
+
 	if (CurrentSorter.IsValid())
 	{
 		for (FStatsNodePtr& Root : GroupNodes)
@@ -1444,6 +1450,16 @@ void SStatsView::SortTreeNodes()
 			SortTreeNodesRec(*Root, *CurrentSorter);
 		}
 	}
+
+	Stopwatch.Stop();
+	const double TotalTime = Stopwatch.GetAccumulatedTime();
+	if (TotalTime > 0.1)
+	{
+		UE_LOG(TimingProfiler, Log, TEXT("[Counters] Tree view sorted (%s, %c) in %.3fs (%d counters)"),
+			CurrentSorter.IsValid() ? *CurrentSorter->GetShortName().ToString() : TEXT("N/A"),
+			(ColumnSortMode == EColumnSortMode::Type::Descending) ? TEXT('D') : TEXT('A'),
+			TotalTime, StatsNodes.Num());
+	}
 }
 
 ////////////////////////////////////////////////////////////////////////////////////////////////////
@@ -1458,6 +1474,16 @@ void SStatsView::SortTreeNodesRec(FStatsNode& Node, const Insights::ITableCellVa
 	{
 		Node.SortChildrenAscending(Sorter);
 	}
+
+#if 0 // Current groupings creates only one level.
+	for (Insights::FBaseTreeNodePtr ChildPtr : Node.GetChildren())
+	{
+		if (ChildPtr->GetChildren().Num() > 0)
+		{
+			SortTreeNodesRec(*StaticCastSharedPtr<FTimerNode>(ChildPtr), Sorter);
+		}
+	}
+#endif
 }
 
 ////////////////////////////////////////////////////////////////////////////////////////////////////
@@ -1859,7 +1885,8 @@ void SStatsView::RebuildTree(bool bResync)
 			const FName MiscInt64Group(TEXT("Misc_int64"));
 
 			// Add nodes only for new counters.
-			CountersProvider.EnumerateCounters([this, MemoryGroup, MiscFloatGroup, MiscInt64Group](uint32 CounterId, const TraceServices::ICounter& Counter)
+			uint32 CounterIndex = PreviousNodeCount;
+			CountersProvider.EnumerateCounters([this, &CounterIndex, MemoryGroup, MiscFloatGroup, MiscInt64Group](uint32 CounterId, const TraceServices::ICounter& Counter)
 			{
 				FStatsNodePtr NodePtr = StatsNodesIdMap.FindRef(CounterId);
 				if (!NodePtr)
@@ -1871,6 +1898,7 @@ void SStatsView::RebuildTree(bool bResync)
 					const EStatsNodeType Type = EStatsNodeType::Counter;
 					const EStatsNodeDataType DataType = Counter.IsFloatingPoint() ? EStatsNodeDataType::Double : EStatsNodeDataType::Int64;
 					NodePtr = MakeShared<FStatsNode>(CounterId, Name, Group, Type, DataType);
+					NodePtr->SetDefaultSortOrder(++CounterIndex);
 					UpdateNode(NodePtr);
 					StatsNodes.Add(NodePtr);
 					StatsNodesIdMap.Add(CounterId, NodePtr);
diff --git a/Engine/Source/Developer/TraceInsights/Private/Insights/Widgets/STimersView.cpp b/Engine/Source/Developer/TraceInsights/Private/Insights/Widgets/STimersView.cpp
--- a/Engine/Source/Developer/TraceInsights/Private/Insights/Widgets/STimersView.cpp
+++ b/Engine/Source/Developer/TraceInsights/Private/Insights/Widgets/STimersView.cpp
@@ -925,34 +925,18 @@ void STimersView::InsightsManager_OnSessionAnalysisCompleted()
 
 void STimersView::UpdateTree()
 {
-	FStopwatch Stopwatch;
-	Stopwatch.Start();
-
 	CreateGroups();
-
-	Stopwatch.Update();
-	const double Time1 = Stopwatch.GetAccumulatedTime();
-
 	SortTreeNodes();
-
-	Stopwatch.Update();
-	const double Time2 = Stopwatch.GetAccumulatedTime();
-
 	ApplyFiltering();
-
-	Stopwatch.Stop();
-	const double TotalTime = Stopwatch.GetAccumulatedTime();
-	if (TotalTime > 0.1)
-	{
-		UE_LOG(TimingProfiler, Log, TEXT("[Timers] Tree view updated in %.3fs (%d timers) --> G:%.3fs + S:%.3fs + F:%.3fs"),
-			TotalTime, TimerNodes.Num(), Time1, Time2 - Time1, TotalTime - Time2);
-	}
 }
 
 ////////////////////////////////////////////////////////////////////////////////////////////////////
 
 void STimersView::ApplyFiltering()
 {
+	FStopwatch Stopwatch;
+	Stopwatch.Start();
+
 	FilteredGroupNodes.Reset();
 
 	// Apply filter to all groups and its children.
@@ -1056,6 +1040,14 @@ void STimersView::ApplyFiltering()
 
 	// Request tree refresh
 	TreeView->RequestTreeRefresh();
+
+	Stopwatch.Stop();
+	const double TotalTime = Stopwatch.GetAccumulatedTime();
+	if (TotalTime > 0.1)
+	{
+		UE_LOG(TimingProfiler, Log, TEXT("[Timers] Tree view filtered in %.3fs (%d timers)"),
+			TotalTime, TimerNodes.Num());
+	}
 }
 
 ////////////////////////////////////////////////////////////////////////////////////////////////////
@@ -1288,6 +1280,9 @@ bool STimersView::SearchBox_IsEnabled() const
 
 void STimersView::CreateGroups()
 {
+	FStopwatch Stopwatch;
+	Stopwatch.Start();
+
 	if (GroupingMode == ETimerGroupingMode::Flat)
 	{
 		GroupNodes.Reset();
@@ -1404,6 +1399,14 @@ void STimersView::CreateGroups()
 	{
 		//im:TODO:
 	}
+
+	Stopwatch.Stop();
+	const double TotalTime = Stopwatch.GetAccumulatedTime();
+	if (TotalTime > 0.1)
+	{
+		UE_LOG(TimingProfiler, Log, TEXT("[Timers] Tree view grouping updated in %.3fs (%d timers)"),
+			TotalTime, TimerNodes.Num());
+	}
 }
 
 ////////////////////////////////////////////////////////////////////////////////////////////////////
@@ -1515,6 +1518,9 @@ void STimersView::UpdateCurrentSortingByColumn()
 
 void STimersView::SortTreeNodes()
 {
+	FStopwatch Stopwatch;
+	Stopwatch.Start();
+
 	if (CurrentSorter.IsValid())
 	{
 		for (FTimerNodePtr& Root : GroupNodes)
@@ -1522,6 +1528,16 @@ void STimersView::SortTreeNodes()
 			SortTreeNodesRec(*Root, *CurrentSorter);
 		}
 	}
+
+	Stopwatch.Stop();
+	const double TotalTime = Stopwatch.GetAccumulatedTime();
+	if (TotalTime > 0.1)
+	{
+		UE_LOG(TimingProfiler, Log, TEXT("[Timers] Tree view sorted (%s, %c) in %.3fs (%d timers)"),
+			CurrentSorter.IsValid() ? *CurrentSorter->GetShortName().ToString() : TEXT("N/A"),
+			(ColumnSortMode == EColumnSortMode::Type::Descending) ? TEXT('D') : TEXT('A'),
+			TotalTime, TimerNodes.Num());
+	}
 }
 
 ////////////////////////////////////////////////////////////////////////////////////////////////////
@@ -1536,7 +1552,8 @@ void STimersView::SortTreeNodesRec(FTimerNode& Node, const Insights::ITableCellV
 	{
 		Node.SortChildrenAscending(Sorter);
 	}
-/*
+
+#if 0 // Current groupings creates only one level.
 	for (Insights::FBaseTreeNodePtr ChildPtr : Node.GetChildren())
 	{
 		if (ChildPtr->GetChildren().Num() > 0)
@@ -1544,7 +1561,7 @@ void STimersView::SortTreeNodesRec(FTimerNode& Node, const Insights::ITableCellV
 			SortTreeNodesRec(*StaticCastSharedPtr<FTimerNode>(ChildPtr), Sorter);
 		}
 	}
-*/
+#endif
 }
 
 ////////////////////////////////////////////////////////////////////////////////////////////////////
@@ -1950,6 +1967,7 @@ void STimersView::RebuildTree(bool bResync)
 				ensure(Timer.Id == TimerIndex);
 				const ETimerNodeType Type = Timer.IsGpuTimer ? ETimerNodeType::GpuScope : ETimerNodeType::CpuScope;
 				FTimerNodePtr TimerNodePtr = MakeShared<FTimerNode>(Timer.Id, Timer.Name, Type);
+				TimerNodePtr->SetDefaultSortOrder(TimerIndex + 1);
 				TimerNodes.Add(TimerNodePtr);
 			}
 			ensure(TimerNodes.Num() == TimerCount);
