[Insights]
  - Modules: Split the "Stats" column into 4 separate columns: Discovered, Cached, Resolved and Failed.
  - Modules: Added sorting for all columns.
  - Modules: Improved coloring to highlight modules with failed status.
  - Modules: Implemented workadound for Discovered value to include the Cached value (so now displayed Discovered = Cached + Resolved + Failed).

#prefligh 621623b3104496cff888e9bd

diff --git a/Engine/Source/Developer/TraceInsights/Private/Insights/Widgets/SModulesView.cpp b/Engine/Source/Developer/TraceInsights/Private/Insights/Widgets/SModulesView.cpp
--- a/Engine/Source/Developer/TraceInsights/Private/Insights/Widgets/SModulesView.cpp
+++ b/Engine/Source/Developer/TraceInsights/Private/Insights/Widgets/SModulesView.cpp
@@ -43,6 +43,11 @@
 		return Module;
 	}
 
+	uint64 GetBaseAddress() const
+	{
+		return Module->Base;
+	}
+
 	FText GetAddressRangeText() const
 	{
 		TStringBuilder<64> Sb;
@@ -60,32 +65,60 @@
 		return FText::FromStringView(Module->FullName);
 	}
 
-	FText GetSymbolsFileText() const
+	inline uint32 GetDiscoveredSymbolCount() const
 	{
-		return Module->StatusMessage ? FText::FromStringView(Module->StatusMessage) : FText();
+		// Workaround for the fact that Stats.Discovered does not include the Stats.Cached value.
+		return Module->Stats.Discovered + Module->Stats.Cached;
 	}
 
+	FText GetDiscoveredAsText() const
+	{
+		// Discovered != Cached + Resolved + Failed ?
+		if (GetDiscoveredSymbolCount() != Module->Stats.Cached + Module->Stats.Resolved + Module->Stats.Failed)
+		{
+			return FText::Format(LOCTEXT("DiscoveredStarredFmt", "{0} (!)"), FText::AsNumber(Module->Stats.Discovered + Module->Stats.Cached));
+		}
+		return FText::AsNumber(GetDiscoveredSymbolCount());
+	}
+
+	inline uint32 GetCachedSymbolCount() const
+	{
+		return Module->Stats.Cached;
+	}
+
+	FText GetCachedAsText() const
+	{
+		return FText::AsNumber(GetCachedSymbolCount());
+	}
+
+	inline uint32 GetResolvedSymbolCount() const
+	{
+		return Module->Stats.Resolved;
+	}
+
+	FText GetResolvedAsText() const
+	{
+		return FText::AsNumber(GetResolvedSymbolCount());
+	}
+
+	inline uint32 GetFailedSymbolCount() const
+	{
+		return Module->Stats.Failed;
+	}
+
+	FText GetFailedAsText() const
+	{
+		return FText::AsNumber(GetFailedSymbolCount());
+	}
+
 	FText GetStatusText() const
 	{
 		return FText::FromStringView(ModuleStatusToString(Module->Status.load()));
 	}
 
-	FText GetStatsText() const
+	FText GetSymbolsFileText() const
 	{
-		TStringBuilder<128> Sb;
-		Sb.Appendf(
-		TEXT("Discovered: %u, Cached: %u, Resolved: %u, Failed: %u"),
-			Module->Stats.Discovered.load(),
-			Module->Stats.Cached.load(),
-			Module->Stats.Resolved.load(),
-			Module->Stats.Failed.load()
-		);
-		return FText::FromStringView(Sb.ToView());
+		return Module->StatusMessage ? FText::FromStringView(Module->StatusMessage) : FText();
 	}
 
-	uint64 GetBaseAddress() const
-	{
-		return Module->Base;
-	}
-
 	FText GetToolTipText() const { return FText(); }
@@ -155,13 +188,44 @@
 					.ToolTipText(this, &SModuleRow::GetModuleFullName)
 				];
 		}
-		else if (ColumnName == ModulesViewColumns::SymbolsFileColumnName)
+		else if (ColumnName == ModulesViewColumns::DiscoveredColumnName)
 		{
 			return SNew(SBox)
 				.Padding(FMargin(4.0, 0.0))
 				[
 					SNew(STextBlock)
-					.Text(this, &SModuleRow::GetSymbolsFile)
+					.ColorAndOpacity(this, &SModuleRow::GetDiscoveredColor)
+					.Text(this, &SModuleRow::GetDiscoveredAsText)
 				];
 		}
+		else if (ColumnName == ModulesViewColumns::CachedColumnName)
+		{
+			return SNew(SBox)
+				.Padding(FMargin(4.0, 0.0))
+				[
+					SNew(STextBlock)
+					.ColorAndOpacity(this, &SModuleRow::GetCachedColor)
+					.Text(this, &SModuleRow::GetCachedAsText)
+				];
+		}
+		else if (ColumnName == ModulesViewColumns::ResolvedColumnName)
+		{
+			return SNew(SBox)
+				.Padding(FMargin(4.0, 0.0))
+				[
+					SNew(STextBlock)
+					.ColorAndOpacity(this, &SModuleRow::GetResolvedColor)
+					.Text(this, &SModuleRow::GetResolvedAsText)
+				];
+		}
+		else if (ColumnName == ModulesViewColumns::FailedColumnName)
+		{
+			return SNew(SBox)
+				.Padding(FMargin(4.0, 0.0))
+				[
+					SNew(STextBlock)
+					.ColorAndOpacity(this, &SModuleRow::GetFailedColor)
+					.Text(this, &SModuleRow::GetFailedAsText)
+				];
+		}
 		else if (ColumnName == ModulesViewColumns::StatusColumnName)
@@ -170,17 +234,18 @@
 				.Padding(FMargin(4.0, 0.0))
 				[
 					SNew(STextBlock)
+					.ColorAndOpacity(this, &SModuleRow::GetStatusColor)
 					.Text(this, &SModuleRow::GetStatus)
 				];
 		}
-		else if (ColumnName == ModulesViewColumns::StatsColumnName)
+		else if (ColumnName == ModulesViewColumns::SymbolsFileColumnName)
 		{
 			return SNew(SBox)
 				.Padding(FMargin(4.0, 0.0))
 				[
 					SNew(STextBlock)
-					.ColorAndOpacity(this, &SModuleRow::GetStatsColor)
-					.Text(this, &SModuleRow::GetStats)
+					.ColorAndOpacity(this, &SModuleRow::GetStatusColor)
+					.Text(this, &SModuleRow::GetSymbolsFile)
 				];
 		}
 		else
@@ -189,11 +254,6 @@
 		}
 	}
 
-	FSlateColor GetBackgroundColor() const
-	{
-		return FSlateColor(FLinearColor(0.0f, 0.0f, 0.0f, 0.0f));
-	}
-
 	FText GetAddressRange() const
 	{
 		TSharedPtr<FModule> ModulePin = WeakModule.Pin();
@@ -207,18 +267,6 @@
 		}
 	}
 
-	FSlateColor GetAddressRangeColor() const
-	{
-		if (IsSelected())
-		{
-			return FLinearColor(0.5f, 0.5f, 0.5f, 1.0f);
-		}
-		else
-		{
-			return FLinearColor(0.2f, 0.2f, 0.2f, 1.0f);
-		}
-	}
-
 	FText GetModuleName() const
 	{
 		TSharedPtr<FModule> ModulePin = WeakModule.Pin();
@@ -228,3 +276,16 @@
 		}
 		else
 		{
+			return FText();
+		}
+	}
+
+	FText GetModuleFullName() const
+	{
+		TSharedPtr<FModule> ModulePin = WeakModule.Pin();
+		if (ModulePin.IsValid())
+		{
+			return ModulePin->GetModuleFullNameText();
+		}
+		else
+		{
@@ -232,12 +293,12 @@
 		}
 	}
 
-	FText GetModuleFullName() const
+	FText GetDiscoveredAsText() const
 	{
-		TSharedPtr<FModule> ModulePin = WeakModule.Pin();
+		const TSharedPtr<FModule> ModulePin = WeakModule.Pin();
 		if (ModulePin.IsValid())
 		{
-			return ModulePin->GetModuleFullNameText();
+			return ModulePin->GetDiscoveredAsText();
 		}
 		else
 		{
@@ -245,12 +306,38 @@
 		}
 	}
 
-	FText GetSymbolsFile() const
+	FText GetCachedAsText() const
 	{
-		TSharedPtr<FModule> ModulePin = WeakModule.Pin();
+		const TSharedPtr<FModule> ModulePin = WeakModule.Pin();
 		if (ModulePin.IsValid())
 		{
-			return ModulePin->GetSymbolsFileText();
+			return ModulePin->GetCachedAsText();
 		}
 		else
 		{
+			return FText();
+		}
+	}
+
+	FText GetResolvedAsText() const
+	{
+		const TSharedPtr<FModule> ModulePin = WeakModule.Pin();
+		if (ModulePin.IsValid())
+		{
+			return ModulePin->GetResolvedAsText();
+		}
+		else
+		{
+			return FText();
+		}
+	}
+
+	FText GetFailedAsText() const
+	{
+		const TSharedPtr<FModule> ModulePin = WeakModule.Pin();
+		if (ModulePin.IsValid())
+		{
+			return ModulePin->GetFailedAsText();
+		}
+		else
+		{
@@ -271,12 +358,12 @@
 		}
 	}
 
-	FText GetStats() const
+	FText GetSymbolsFile() const
 	{
-		const TSharedPtr<FModule> ModulePin = WeakModule.Pin();
+		TSharedPtr<FModule> ModulePin = WeakModule.Pin();
 		if (ModulePin.IsValid())
 		{
-			return ModulePin->GetStatsText();
+			return ModulePin->GetSymbolsFileText();
 		}
 		else
 		{
@@ -284,7 +371,38 @@
 		}
 	}
 
-	FSlateColor GetStatsColor() const
+	FText GetRowToolTip() const
 	{
 		TSharedPtr<FModule> ModulePin = WeakModule.Pin();
 		if (ModulePin.IsValid())
+		{
+			return ModulePin->GetToolTipText();
+		}
+		else
+		{
+			return FText();
+		}
+	}
+
+	inline const FLinearColor GetRedColor() const       { return FLinearColor(1.0f, 0.7f, 0.7f, 1.0f); }
+	inline const FLinearColor GetYellowColor() const    { return FLinearColor(1.0f, 0.9f, 0.6f, 1.0f); }
+	inline const FLinearColor GetGreenColor() const     { return FLinearColor(0.7f, 1.0f, 0.7f, 1.0f); }
+	inline const FLinearColor GetWhileColor() const     { return FLinearColor(1.0f, 1.0f, 1.0f, 1.0f); }
+	inline const FLinearColor GetGrayColor() const      { return FLinearColor(0.5f, 0.5f, 0.5f, 1.0f); }
+	inline const FLinearColor GetDarkGrayColor() const  { return FLinearColor(0.2f, 0.2f, 0.2f, 1.0f); }
+	inline const FLinearColor GetBlackColor() const     { return FLinearColor(0.0f, 0.0f, 0.0f, 1.0f); }
+
+	FSlateColor GetBackgroundColor() const
+	{
+		return FSlateColor(FLinearColor(0.0f, 0.0f, 0.0f, 0.0f));
+	}
+
+	FSlateColor GetAddressRangeColor() const
+	{
+		return IsSelected() ? GetGrayColor() : GetDarkGrayColor();
+	}
+
+	FSlateColor GetDiscoveredColor() const
+	{
+		TSharedPtr<FModule> ModulePin = WeakModule.Pin();
+		if (ModulePin.IsValid())
@@ -296,49 +414,130 @@
 			{
 				if (Module->Stats.Cached + Module->Stats.Resolved > Module->Stats.Failed)
 				{
-					return FLinearColor(1.0f, 0.9f, 0.6f, 1.0f);
+					return GetYellowColor();
 				}
 				else
 				{
-					return FLinearColor(1.0f, 0.7f, 0.7f, 1.0f);
+					return GetRedColor();
 				}
 			}
 			else if (Module->Stats.Cached > 0 || Module->Stats.Resolved > 0)
 			{
-				return FLinearColor(0.7f, 1.0f, 0.7f, 1.0f);
+				return GetGreenColor();
 			}
 			else if (Module->Stats.Discovered > 0)
 			{
-				return FLinearColor(1.0f, 1.0f, 1.0f, 1.0f);
+				return GetWhileColor();
 			}
 			else
 			{
-				if (IsSelected())
-				{
-					return FLinearColor(0.5f, 0.5f, 0.5f, 1.0f);
-				}
-				else
-				{
-					return FLinearColor(0.2f, 0.2f, 0.2f, 1.0f);
-				}
+				return IsSelected() ? GetGrayColor() : GetDarkGrayColor();
 			}
 		}
 		else
 		{
-			return FLinearColor(0.0f, 0.0f, 0.0f, 1.0f);
+			return GetBlackColor();
 		}
 	}
 
-	FText GetRowToolTip() const
+	FSlateColor GetCachedColor() const
 	{
 		TSharedPtr<FModule> ModulePin = WeakModule.Pin();
 		if (ModulePin.IsValid())
 		{
-			return ModulePin->GetToolTipText();
+			const TraceServices::FModule* Module = ModulePin->GetModule();
+			check(Module != nullptr);
+
+			if (Module->Stats.Cached > 0)
+			{
+				return GetGreenColor();
+			}
+			else
+			{
+				return IsSelected() ? GetGrayColor() : GetDarkGrayColor();
+			}
 		}
 		else
 		{
-			return FText();
+			return GetBlackColor();
 		}
 	}
 
+	FSlateColor GetResolvedColor() const
+	{
+		TSharedPtr<FModule> ModulePin = WeakModule.Pin();
+		if (ModulePin.IsValid())
+		{
+			const TraceServices::FModule* Module = ModulePin->GetModule();
+			check(Module != nullptr);
+
+			if (Module->Stats.Resolved > 0)
+			{
+				return GetGreenColor();
+			}
+			else
+			{
+				return IsSelected() ? GetGrayColor() : GetDarkGrayColor();
+			}
+		}
+		else
+		{
+			return GetBlackColor();
+		}
+	}
+
+	FSlateColor GetFailedColor() const
+	{
+		TSharedPtr<FModule> ModulePin = WeakModule.Pin();
+		if (ModulePin.IsValid())
+		{
+			const TraceServices::FModule* Module = ModulePin->GetModule();
+			check(Module != nullptr);
+
+			if (Module->Stats.Failed > 0)
+			{
+				return GetRedColor();
+			}
+			else
+			{
+				return IsSelected() ? GetGrayColor() : GetDarkGrayColor();
+			}
+		}
+		else
+		{
+			return GetBlackColor();
+		}
+	}
+
+	FSlateColor GetStatusColor() const
+	{
+		TSharedPtr<FModule> ModulePin = WeakModule.Pin();
+		if (ModulePin.IsValid())
+		{
+			const TraceServices::FModule* Module = ModulePin->GetModule();
+			check(Module != nullptr);
+
+			const TraceServices::EModuleStatus Status = Module->Status;
+			switch (Status)
+			{
+				case TraceServices::EModuleStatus::Pending:
+					return GetYellowColor();
+
+				case TraceServices::EModuleStatus::Loaded:
+					return GetGreenColor();
+
+				case TraceServices::EModuleStatus::VersionMismatch:
+				case TraceServices::EModuleStatus::NotFound:
+				case TraceServices::EModuleStatus::Failed:
+					return GetRedColor();
+
+				default:
+					return IsSelected() ? GetGrayColor() : GetDarkGrayColor();
+			}
+		}
+		else
+		{
+			return GetBlackColor();
+		}
+	}
+
@@ -353,6 +552,8 @@
 
 SModulesView::SModulesView()
 	: LastUpdateTime(0.0)
+	, SortColumn(ModulesViewColumns::DiscoveredColumnName)
+	, SortMode(EColumnSortMode::Descending)
 {
 }
 
@@ -439,23 +640,62 @@
 							+ SHeaderRow::Column(ModulesViewColumns::AddressRangeColumnName)
 							.ManualWidth(320.0f)
 							.DefaultLabel(LOCTEXT("AddressRangeColumn", "Address Range"))
+							.InitialSortMode(EColumnSortMode::Ascending)
+							.SortMode(this, &SModulesView::GetSortModeForColumn, ModulesViewColumns::AddressRangeColumnName)
+							.OnSort(this, &SModulesView::OnSortModeChanged)
 
 							+ SHeaderRow::Column(ModulesViewColumns::ModuleNameColumnName)
 							.ManualWidth(300.0f)
 							.DefaultLabel(LOCTEXT("ModuleNameColumn", "Name"))
+							.InitialSortMode(EColumnSortMode::Ascending)
+							.SortMode(this, &SModulesView::GetSortModeForColumn, ModulesViewColumns::ModuleNameColumnName)
+							.OnSort(this, &SModulesView::OnSortModeChanged)
 
-							+ SHeaderRow::Column(ModulesViewColumns::StatsColumnName)
-							.ManualWidth(350.0f)
-							.DefaultLabel(LOCTEXT("StatsColumn", "Statistics"))
-							.DefaultTooltip(LOCTEXT("StatsColumnTooltip", "Discovered = The total number of symbols discovered.\nCached = The number of symbols loaded from cache (*.ucache file).\nResolved = The number of symbols successfully resolved.\nFailed = The number of symbols failed to resolve."))
+							+ SHeaderRow::Column(ModulesViewColumns::DiscoveredColumnName)
+							.ManualWidth(76.0f)
+							.DefaultLabel(LOCTEXT("DiscoveredColumn", "Discovered"))
+							.DefaultTooltip(LOCTEXT("DiscoveredColumnTooltip", "The total number of symbols discovered in the trace.\nDiscovered = Cached + Resolved + Failed"))
+							.InitialSortMode(EColumnSortMode::Descending)
+							.SortMode(this, &SModulesView::GetSortModeForColumn, ModulesViewColumns::DiscoveredColumnName)
+							.OnSort(this, &SModulesView::OnSortModeChanged)
 
+							+ SHeaderRow::Column(ModulesViewColumns::CachedColumnName)
+							.ManualWidth(76.0f)
+							.DefaultLabel(LOCTEXT("CachedColumn", "Cached"))
+							.DefaultTooltip(LOCTEXT("CachedColumnTooltip", "The number of symbols loaded (pre-resolved) from the cache (*.ucache file)."))
+							.InitialSortMode(EColumnSortMode::Descending)
+							.SortMode(this, &SModulesView::GetSortModeForColumn, ModulesViewColumns::CachedColumnName)
+							.OnSort(this, &SModulesView::OnSortModeChanged)
+
+							+ SHeaderRow::Column(ModulesViewColumns::ResolvedColumnName)
+							.ManualWidth(76.0f)
+							.DefaultLabel(LOCTEXT("ResolvedColumn", "Resolved"))
+							.DefaultTooltip(LOCTEXT("ResolvedColumnTooltip", "The number of symbols successfully resolved in the current analysis session."))
+							.InitialSortMode(EColumnSortMode::Descending)
+							.SortMode(this, &SModulesView::GetSortModeForColumn, ModulesViewColumns::ResolvedColumnName)
+							.OnSort(this, &SModulesView::OnSortModeChanged)
+
+							+ SHeaderRow::Column(ModulesViewColumns::FailedColumnName)
+							.ManualWidth(76.0f)
+							.DefaultLabel(LOCTEXT("FailedColumn", "Failed"))
+							.DefaultTooltip(LOCTEXT("FailedColumnTooltip", "The number of symbols failed to resolve."))
+							.InitialSortMode(EColumnSortMode::Descending)
+							.SortMode(this, &SModulesView::GetSortModeForColumn, ModulesViewColumns::FailedColumnName)
+							.OnSort(this, &SModulesView::OnSortModeChanged)
+
 							+ SHeaderRow::Column(ModulesViewColumns::StatusColumnName)
 							.ManualWidth(100.0f)
 							.DefaultLabel(LOCTEXT("StatusColumn", "Status"))
+							.InitialSortMode(EColumnSortMode::Ascending)
+							.SortMode(this, &SModulesView::GetSortModeForColumn, ModulesViewColumns::StatusColumnName)
+							.OnSort(this, &SModulesView::OnSortModeChanged)
 
 							+ SHeaderRow::Column(ModulesViewColumns::SymbolsFileColumnName)
 							.ManualWidth(800.0f)
 							.DefaultLabel(LOCTEXT("SymbolsFileColumn", "Symbols"))
+							.InitialSortMode(EColumnSortMode::Ascending)
+							.SortMode(this, &SModulesView::GetSortModeForColumn, ModulesViewColumns::SymbolsFileColumnName)
+							.OnSort(this, &SModulesView::OnSortModeChanged)
 						)
 					]
 				]
@@ -499,10 +739,15 @@
 	const TSharedPtr<const IAnalysisSession> Session = FInsightsManager::Get()->GetSession();
 	if (const IModuleProvider* ModuleProvider = Session->ReadProvider<IModuleProvider>("ModuleProvider"))
 	{
+		const int32 ModuleCount = Modules.Num();
 		ModuleProvider->EnumerateModules(Modules.Num(), [this](const TraceServices::FModule& InModule)
 		{
 			Modules.Emplace(MakeShared<FModule>(&InModule));
 		});
+		if (Modules.Num() != ModuleCount)
+		{
+			UpdateSorting();
+		}
 	}
 }
 
@@ -676,6 +921,177 @@
 
 ////////////////////////////////////////////////////////////////////////////////////////////////////
 
+EColumnSortMode::Type SModulesView::GetSortModeForColumn(const FName ColumnId) const
+{
+	return ColumnId == SortColumn ? SortMode : EColumnSortMode::None;
+}
+
+////////////////////////////////////////////////////////////////////////////////////////////////////
+
+void SModulesView::OnSortModeChanged(const EColumnSortPriority::Type SortPriority, const FName& ColumnId, const EColumnSortMode::Type InSortMode)
+{
+	SortColumn = ColumnId;
+	SortMode = InSortMode;
+	UpdateSorting();
+
+	if (ListView.IsValid())
+	{
+		ListView->RebuildList();
+	}
+}
+
+////////////////////////////////////////////////////////////////////////////////////////////////////
+
+void SModulesView::UpdateSorting()
+{
+	if (SortColumn == ModulesViewColumns::AddressRangeColumnName)
+	{
+		if (SortMode == EColumnSortMode::Ascending)
+		{
+			Modules.Sort([](const TSharedPtr<FModule>& A, const TSharedPtr<FModule>& B)
+				{ return A->GetBaseAddress() < B->GetBaseAddress(); });
+		}
+		else
+		{
+			Modules.Sort([](const TSharedPtr<FModule>& A, const TSharedPtr<FModule>& B)
+				{ return A->GetBaseAddress() > B->GetBaseAddress(); });
+		}
+	}
+	else if (SortColumn == ModulesViewColumns::ModuleNameColumnName)
+	{
+		if (SortMode == EColumnSortMode::Ascending)
+		{
+			Modules.Sort([](const TSharedPtr<FModule>& A, const TSharedPtr<FModule>& B)
+				{
+					if (A->GetModule()->Name == nullptr)
+					{
+						return true;
+					}
+					if (B->GetModule()->Name == nullptr)
+					{
+						return false;
+					}
+					return FCString::Strcmp(A->GetModule()->Name, B->GetModule()->Name) < 0;
+				});
+		}
+		else
+		{
+			Modules.Sort([](const TSharedPtr<FModule>& A, const TSharedPtr<FModule>& B)
+				{
+					if (A->GetModule()->Name == nullptr)
+					{
+						return false;
+					}
+					if (B->GetModule()->Name == nullptr)
+					{
+						return true;
+					}
+					return FCString::Strcmp(A->GetModule()->Name, B->GetModule()->Name) > 0;
+				});
+		}
+	}
+	else if (SortColumn == ModulesViewColumns::DiscoveredColumnName)
+	{
+		if (SortMode == EColumnSortMode::Ascending)
+		{
+			Modules.Sort([](const TSharedPtr<FModule>& A, const TSharedPtr<FModule>& B)
+				{ return A->GetDiscoveredSymbolCount() < B->GetDiscoveredSymbolCount(); });
+		}
+		else
+		{
+			Modules.Sort([](const TSharedPtr<FModule>& A, const TSharedPtr<FModule>& B)
+				{ return A->GetDiscoveredSymbolCount() > B->GetDiscoveredSymbolCount(); });
+		}
+	}
+	else if (SortColumn == ModulesViewColumns::CachedColumnName)
+	{
+		if (SortMode == EColumnSortMode::Ascending)
+		{
+			Modules.Sort([](const TSharedPtr<FModule>& A, const TSharedPtr<FModule>& B)
+				{ return A->GetCachedSymbolCount() < B->GetCachedSymbolCount(); });
+		}
+		else
+		{
+			Modules.Sort([](const TSharedPtr<FModule>& A, const TSharedPtr<FModule>& B)
+				{ return A->GetCachedSymbolCount() > B->GetCachedSymbolCount(); });
+		}
+	}
+	else if (SortColumn == ModulesViewColumns::ResolvedColumnName)
+	{
+		if (SortMode == EColumnSortMode::Ascending)
+		{
+			Modules.Sort([](const TSharedPtr<FModule>& A, const TSharedPtr<FModule>& B)
+				{ return A->GetResolvedSymbolCount() < B->GetResolvedSymbolCount(); });
+		}
+		else
+		{
+			Modules.Sort([](const TSharedPtr<FModule>& A, const TSharedPtr<FModule>& B)
+				{ return A->GetResolvedSymbolCount() > B->GetResolvedSymbolCount(); });
+		}
+	}
+	else if (SortColumn == ModulesViewColumns::FailedColumnName)
+	{
+		if (SortMode == EColumnSortMode::Ascending)
+		{
+			Modules.Sort([](const TSharedPtr<FModule>& A, const TSharedPtr<FModule>& B)
+				{ return A->GetFailedSymbolCount() < B->GetFailedSymbolCount(); });
+		}
+		else
+		{
+			Modules.Sort([](const TSharedPtr<FModule>& A, const TSharedPtr<FModule>& B)
+				{ return A->GetFailedSymbolCount() > B->GetFailedSymbolCount(); });
+		}
+	}
+	else if (SortColumn == ModulesViewColumns::StatusColumnName)
+	{
+		if (SortMode == EColumnSortMode::Ascending)
+		{
+			Modules.Sort([](const TSharedPtr<FModule>& A, const TSharedPtr<FModule>& B)
+				{ return A->GetModule()->Status < B->GetModule()->Status; });
+		}
+		else
+		{
+			Modules.Sort([](const TSharedPtr<FModule>& A, const TSharedPtr<FModule>& B)
+				{ return A->GetModule()->Status > B->GetModule()->Status; });
+		}
+	}
+	else if (SortColumn == ModulesViewColumns::SymbolsFileColumnName)
+	{
+		if (SortMode == EColumnSortMode::Ascending)
+		{
+			Modules.Sort([](const TSharedPtr<FModule>& A, const TSharedPtr<FModule>& B)
+				{
+					if (A->GetModule()->StatusMessage == nullptr)
+					{
+						return true;
+					}
+					if (B->GetModule()->StatusMessage == nullptr)
+					{
+						return false;
+					}
+					return FCString::Strcmp(A->GetModule()->StatusMessage, B->GetModule()->StatusMessage) < 0;
+				});
+		}
+		else
+		{
+			Modules.Sort([](const TSharedPtr<FModule>& A, const TSharedPtr<FModule>& B)
+				{
+					if (A->GetModule()->StatusMessage == nullptr)
+					{
+						return false;
+					}
+					if (B->GetModule()->StatusMessage == nullptr)
+					{
+						return true;
+					}
+					return FCString::Strcmp(A->GetModule()->StatusMessage, B->GetModule()->StatusMessage) > 0;
+				});
+		}
+	}
+}
+
+////////////////////////////////////////////////////////////////////////////////////////////////////
+
 } //namespace Insights
 
 #undef LOCTEXT_NAMESPACE
diff --git a/Engine/Source/Developer/TraceInsights/Private/Insights/Widgets/SModulesView.h b/Engine/Source/Developer/TraceInsights/Private/Insights/Widgets/SModulesView.h
--- a/Engine/Source/Developer/TraceInsights/Private/Insights/Widgets/SModulesView.h
+++ b/Engine/Source/Developer/TraceInsights/Private/Insights/Widgets/SModulesView.h
@@ -22,9 +22,12 @@ namespace ModulesViewColumns
 {
 	static const FName AddressRangeColumnName(TEXT("BaseAddress"));
 	static const FName ModuleNameColumnName(TEXT("ModuleName"));
-	static const FName SymbolsFileColumnName(TEXT("SymbolsFile"));
+	static const FName DiscoveredColumnName(TEXT("Discovered"));
+	static const FName CachedColumnName(TEXT("Cached"));
+	static const FName ResolvedColumnName(TEXT("Resolved"));
+	static const FName FailedColumnName(TEXT("Failed"));
 	static const FName StatusColumnName(TEXT("Status"));
-	static const FName StatsColumnName(TEXT("Stats"));
+	static const FName SymbolsFileColumnName(TEXT("SymbolsFile"));
 }
 
 ////////////////////////////////////////////////////////////////////////////////////////////////////
@@ -81,6 +84,12 @@ protected:
 
 	void LoadSymbols_Execute(TSharedPtr<FModule> Module, bool bOpenFile);
 
+	void UpdateSorting();
+
+private:
+	EColumnSortMode::Type GetSortModeForColumn(const FName ColumnId) const;
+	void OnSortModeChanged(const EColumnSortPriority::Type SortPriority, const FName& ColumnId, const EColumnSortMode::Type SortMode);
+
 protected:
 	/** The list view widget. */
 	TSharedPtr<SListView<TSharedPtr<FModule>>> ListView;
@@ -96,6 +105,12 @@ protected:
 
 	/** Last time we queried about discovered modules. */
 	double LastUpdateTime;
+
+	/** The current sorting column. */
+	FName SortColumn;
+
+	/** Sorting mode (ascending or descending) for current sorting column. */
+	EColumnSortMode::Type SortMode;
 };
 
 ////////////////////////////////////////////////////////////////////////////////////////////////////
