[Insights] Session Browser: Added sorting for all columns in the trace sessions list view.

diff --git a/Engine/Source/Developer/TraceInsights/Private/Insights/Widgets/SStartPageWindow.cpp b/Engine/Source/Developer/TraceInsights/Private/Insights/Widgets/SStartPageWindow.cpp
--- a/Engine/Source/Developer/TraceInsights/Private/Insights/Widgets/SStartPageWindow.cpp
+++ b/Engine/Source/Developer/TraceInsights/Private/Insights/Widgets/SStartPageWindow.cpp
@@ -75,70 +75,70 @@ public:
 
 	virtual TSharedRef<SWidget> GenerateWidgetForColumn(const FName& ColumnName) override
 	{
-		if (ColumnName == FName(TEXT("Name")))
+		if (ColumnName == TraceStoreColumns::Name)
 		{
 			return SNew(SBox)
-				.Padding(FMargin(4.0, 0.0))
+				.Padding(FMargin(4.0f, 0.0f))
 				[
 					SNew(STextBlock)
 					.Text(this, &STraceListRow::GetTraceName)
 					.ToolTip(STraceListRow::GetTraceTooltip())
 				];
 		}
-		else if (ColumnName == FName(TEXT("Uri")))
+		else if (ColumnName == TraceStoreColumns::Uri)
 		{
 			return SNew(SBox)
-				.Padding(FMargin(4.0, 0.0))
+				.Padding(FMargin(4.0f, 0.0f))
 				[
 					SNew(STextBlock)
 					.Text(this, &STraceListRow::GetTraceUri)
 					.ToolTip(STraceListRow::GetTraceTooltip())
 				];
 		}
-		else if (ColumnName == FName(TEXT("Platform")))
+		else if (ColumnName == TraceStoreColumns::Platform)
 		{
 			return SNew(SBox)
-				.Padding(FMargin(4.0, 0.0))
+				.Padding(FMargin(4.0f, 0.0f))
 				[
 					SNew(STextBlock)
 					.Text(this, &STraceListRow::GetTracePlatform)
 					.ToolTip(STraceListRow::GetTraceTooltip())
 				];
 		}
-		else if (ColumnName == FName(TEXT("AppName")))
+		else if (ColumnName == TraceStoreColumns::AppName)
 		{
 			return SNew(SBox)
-				.Padding(FMargin(4.0, 0.0))
+				.Padding(FMargin(4.0f, 0.0f))
 				[
 					SNew(STextBlock)
 					.Text(this, &STraceListRow::GetTraceAppName)
 					.ToolTip(STraceListRow::GetTraceTooltip())
 				];
 		}
-		else if (ColumnName == FName(TEXT("BuildConfig")))
+		else if (ColumnName == TraceStoreColumns::BuildConfig)
 		{
 			return SNew(SBox)
-				.Padding(FMargin(4.0, 0.0))
+				.Padding(FMargin(4.0f, 0.0f))
 				[
 					SNew(STextBlock)
 					.Text(this, &STraceListRow::GetTraceBuildConfiguration)
 					.ToolTip(STraceListRow::GetTraceTooltip())
 				];
 		}
-		else if (ColumnName == FName(TEXT("BuildTarget")))
+		else if (ColumnName == TraceStoreColumns::BuildTarget)
 		{
 			return SNew(SBox)
-				.Padding(FMargin(4.0, 0.0))
+				.Padding(FMargin(4.0f, 0.0f))
 				[
 					SNew(STextBlock)
 					.Text(this, &STraceListRow::GetTraceBuildTarget)
 					.ToolTip(STraceListRow::GetTraceTooltip())
 				];
 		}
-		else if (ColumnName == FName(TEXT("Size")))
+		else if (ColumnName == TraceStoreColumns::Size)
 		{
 			return SNew(SBox)
-				.Padding(FMargin(4.0, 0.0))
+				.Padding(FMargin(4.0f, 0.0f))
 				[
 					SNew(STextBlock)
 					.Text(this, &STraceListRow::GetTraceSize)
@@ -146,10 +146,10 @@ public:
 					.ToolTip(STraceListRow::GetTraceTooltip())
 				];
 		}
-		else if (ColumnName == FName(TEXT("Status")))
+		else if (ColumnName == TraceStoreColumns::Status)
 		{
 			return SNew(SBox)
-				.Padding(FMargin(4.0, 0.0))
+				.Padding(FMargin(4.0f, 0.0f))
 				[
 					SNew(STextBlock)
 					.Text(this, &STraceListRow::GetTraceStatus)
@@ -636,6 +636,8 @@ STraceStoreWindow::STraceStoreWindow()
 	, TraceViewModelMap()
 	, TraceListView()
 	, SelectedTrace()
+	, SortColumn(TraceStoreColumns::Date)
+	, SortMode(EColumnSortMode::Ascending)
 	, SplashScreenOverlayFadeTime(0.0f)
 {
 }
@@ -793,37 +795,58 @@ TSharedRef<SWidget> STraceStoreWindow::ConstructSessionsPanel()
 		(
 			SNew(SHeaderRow)
 
-			+ SHeaderRow::Column(FName(TEXT("Name")))
-			.FillWidth(0.25f)
+			+ SHeaderRow::Column(TraceStoreColumns::Name)
 			.DefaultLabel(LOCTEXT("NameColumn", "Name"))
+			.FillWidth(0.25f)
+			.InitialSortMode(EColumnSortMode::Ascending)
+			.SortMode(this, &STraceStoreWindow::GetSortModeForColumn, TraceStoreColumns::Name)
+			.OnSort(this, &STraceStoreWindow::OnSortModeChanged)
 
-			+ SHeaderRow::Column(FName(TEXT("Platform")))
-			.FillWidth(0.1f)
+			+ SHeaderRow::Column(TraceStoreColumns::Platform)
 			.DefaultLabel(LOCTEXT("PlatformColumn", "Platform"))
-
-			+ SHeaderRow::Column(FName(TEXT("AppName")))
 			.FillWidth(0.1f)
-			.DefaultLabel(LOCTEXT("AppNameColumn", "App Name"))
+			.InitialSortMode(EColumnSortMode::Ascending)
+			.SortMode(this, &STraceStoreWindow::GetSortModeForColumn, TraceStoreColumns::Platform)
+			.OnSort(this, &STraceStoreWindow::OnSortModeChanged)
 
-			+ SHeaderRow::Column(FName(TEXT("BuildConfig")))
+			+ SHeaderRow::Column(TraceStoreColumns::AppName)
+			.DefaultLabel(LOCTEXT("AppNameColumn", "App Name"))
 			.FillWidth(0.1f)
-			.DefaultLabel(LOCTEXT("BuildConfigColumn", "Build Config"))
+			.InitialSortMode(EColumnSortMode::Ascending)
+			.SortMode(this, &STraceStoreWindow::GetSortModeForColumn, TraceStoreColumns::AppName)
+			.OnSort(this, &STraceStoreWindow::OnSortModeChanged)
 
-			+ SHeaderRow::Column(FName(TEXT("BuildTarget")))
+			+ SHeaderRow::Column(TraceStoreColumns::BuildConfig)
+			.DefaultLabel(LOCTEXT("BuildConfigColumn", "Build Config"))
 			.FillWidth(0.1f)
+			.InitialSortMode(EColumnSortMode::Ascending)
+			.SortMode(this, &STraceStoreWindow::GetSortModeForColumn, TraceStoreColumns::BuildConfig)
+			.OnSort(this, &STraceStoreWindow::OnSortModeChanged)
+
+			+ SHeaderRow::Column(TraceStoreColumns::BuildTarget)
 			.DefaultLabel(LOCTEXT("BuildTargetColumn", "Build Target"))
+			.FillWidth(0.1f)
+			.InitialSortMode(EColumnSortMode::Ascending)
+			.SortMode(this, &STraceStoreWindow::GetSortModeForColumn, TraceStoreColumns::BuildTarget)
+			.OnSort(this, &STraceStoreWindow::OnSortModeChanged)
 
-			+ SHeaderRow::Column(FName(TEXT("Size")))
+			+ SHeaderRow::Column(TraceStoreColumns::Size)
+			.DefaultLabel(LOCTEXT("SizeColumn", "File Size"))
 			.FixedWidth(100.0f)
 			.HAlignHeader(HAlign_Right)
 			.HAlignCell(HAlign_Right)
-			.DefaultLabel(LOCTEXT("SizeColumn", "File Size"))
+			.InitialSortMode(EColumnSortMode::Descending)
+			.SortMode(this, &STraceStoreWindow::GetSortModeForColumn, TraceStoreColumns::Size)
+			.OnSort(this, &STraceStoreWindow::OnSortModeChanged)
 
-			+ SHeaderRow::Column(FName(TEXT("Status")))
+			+ SHeaderRow::Column(TraceStoreColumns::Status)
+			.DefaultLabel(LOCTEXT("StatusColumn", "Status"))
 			.FixedWidth(60.0f)
 			.HAlignHeader(HAlign_Right)
 			.HAlignCell(HAlign_Right)
-			.DefaultLabel(LOCTEXT("StatusColumn", "Status"))
+			.InitialSortMode(EColumnSortMode::Ascending)
+			.SortMode(this, &STraceStoreWindow::GetSortModeForColumn, TraceStoreColumns::Status)
+			.OnSort(this, &STraceStoreWindow::OnSortModeChanged)
 		);
 
 	return Widget;
@@ -1210,11 +1233,18 @@ void STraceStoreWindow::UpdateTrace(FTraceViewModel& InOutTrace, const Insights:
 
 void STraceStoreWindow::OnTraceListChanged()
 {
-	Algo::SortBy(TraceViewModels, &FTraceViewModel::Timestamp);
+	UpdateSorting();
+	UpdateTraceListView();
+}
 
-	//////////////////////////////////////////////////
-	// TraceViewModels array has changed.
-	// Now we need to rebuild the list in the ListView.
+////////////////////////////////////////////////////////////////////////////////////////////////////
+
+void STraceStoreWindow::UpdateTraceListView()
+{
+	if (!TraceListView)
+	{
+		return;
+	}
 
 	TSharedPtr<FTraceViewModel> NewSelectedTrace;
 	if (SelectedTrace)
@@ -1226,13 +1256,25 @@ void STraceStoreWindow::OnTraceListChanged()
 
 	TraceListView->RebuildList();
 
-	// If no selection, auto select the last (newest) trace.
-	if (!NewSelectedTrace.IsValid() && TraceViewModels.Num() > 0)
+	if ((SortColumn == TraceStoreColumns::Date && SortMode == EColumnSortMode::Ascending) ||
+		(SortColumn == TraceStoreColumns::Status && SortMode == EColumnSortMode::Ascending))
 	{
-		NewSelectedTrace = TraceViewModels.Last();
-	}
+		// If no selection, auto select the last (newest) trace.
+		if (!NewSelectedTrace.IsValid() && TraceViewModels.Num() > 0)
+		{
+			NewSelectedTrace = TraceViewModels.Last();
+		}
 
-	TraceListView->ScrollToBottom();
+		TraceListView->ScrollToBottom();
+	}
+	else
+	{
+		// If no selection, auto select the first trace.
+		if (!NewSelectedTrace.IsValid() && TraceViewModels.Num() > 0)
+		{
+			NewSelectedTrace = TraceViewModels[0];
+		}
+	}
 
 	// Restore selection and ensure it is visible.
 	if (NewSelectedTrace.IsValid())
@@ -1708,6 +1750,186 @@ void STraceStoreWindow::GetExtraCommandLineParams(FString& OutParams) const
 
 ////////////////////////////////////////////////////////////////////////////////////////////////////
 
+EColumnSortMode::Type STraceStoreWindow::GetSortModeForColumn(const FName ColumnId) const
+{
+	return ColumnId == SortColumn ? SortMode : EColumnSortMode::None;
+}
+
+////////////////////////////////////////////////////////////////////////////////////////////////////
+
+void STraceStoreWindow::OnSortModeChanged(const EColumnSortPriority::Type SortPriority, const FName& ColumnId, const EColumnSortMode::Type InSortMode)
+{
+	SortColumn = ColumnId;
+	SortMode = InSortMode;
+	UpdateSorting();
+	UpdateTraceListView();
+}
+
+////////////////////////////////////////////////////////////////////////////////////////////////////
+
+void STraceStoreWindow::UpdateSorting()
+{
+	if (SortColumn == TraceStoreColumns::Date)
+	{
+		if (SortMode == EColumnSortMode::Ascending)
+		{
+			TraceViewModels.Sort([](const TSharedPtr<FTraceViewModel>& A, const TSharedPtr<FTraceViewModel>& B)
+				{ return A->Timestamp < B->Timestamp; });
+		}
+		else
+		{
+			TraceViewModels.Sort([](const TSharedPtr<FTraceViewModel>& A, const TSharedPtr<FTraceViewModel>& B)
+				{ return A->Timestamp > B->Timestamp; });
+		}
+	}
+	else if (SortColumn == TraceStoreColumns::Name)
+	{
+		if (SortMode == EColumnSortMode::Ascending)
+		{
+			TraceViewModels.Sort([](const TSharedPtr<FTraceViewModel>& A, const TSharedPtr<FTraceViewModel>& B)
+				{ return A->Name.CompareTo(B->Name) < 0; });
+		}
+		else
+		{
+			TraceViewModels.Sort([](const TSharedPtr<FTraceViewModel>& A, const TSharedPtr<FTraceViewModel>& B)
+				{ return B->Name.CompareTo(A->Name) < 0; });
+		}
+	}
+	else if (SortColumn == TraceStoreColumns::Uri)
+	{
+		if (SortMode == EColumnSortMode::Ascending)
+		{
+			TraceViewModels.Sort([](const TSharedPtr<FTraceViewModel>& A, const TSharedPtr<FTraceViewModel>& B)
+				{ return A->Uri.CompareTo(B->Uri) < 0; });
+		}
+		else
+		{
+			TraceViewModels.Sort([](const TSharedPtr<FTraceViewModel>& A, const TSharedPtr<FTraceViewModel>& B)
+				{ return B->Uri.CompareTo(A->Uri) < 0; });
+		}
+	}
+	else if (SortColumn == TraceStoreColumns::Platform)
+	{
+		if (SortMode == EColumnSortMode::Ascending)
+		{
+			TraceViewModels.Sort([](const TSharedPtr<FTraceViewModel>& A, const TSharedPtr<FTraceViewModel>& B)
+				{
+					int32 CompareResult = A->Platform.CompareTo(B->Platform);
+					return CompareResult == 0 ? A->Timestamp < B->Timestamp : CompareResult < 0;
+				});
+		}
+		else
+		{
+			TraceViewModels.Sort([](const TSharedPtr<FTraceViewModel>& A, const TSharedPtr<FTraceViewModel>& B)
+				{
+					int32 CompareResult = B->Platform.CompareTo(A->Platform);
+					return CompareResult == 0 ? A->Timestamp < B->Timestamp : CompareResult < 0;
+				});
+		}
+	}
+	else if (SortColumn == TraceStoreColumns::AppName)
+	{
+		if (SortMode == EColumnSortMode::Ascending)
+		{
+			TraceViewModels.Sort([](const TSharedPtr<FTraceViewModel>& A, const TSharedPtr<FTraceViewModel>& B)
+				{
+					int32 CompareResult = A->AppName.CompareTo(B->AppName);
+					return CompareResult == 0 ? A->Timestamp < B->Timestamp : CompareResult < 0;
+				});
+		}
+		else
+		{
+			TraceViewModels.Sort([](const TSharedPtr<FTraceViewModel>& A, const TSharedPtr<FTraceViewModel>& B)
+				{ 
+					int32 CompareResult = B->AppName.CompareTo(A->AppName);
+					return CompareResult == 0 ? A->Timestamp < B->Timestamp : CompareResult < 0;
+				});
+		}
+	}
+	else if (SortColumn == TraceStoreColumns::BuildConfig)
+	{
+		if (SortMode == EColumnSortMode::Ascending)
+		{
+			TraceViewModels.Sort([](const TSharedPtr<FTraceViewModel>& A, const TSharedPtr<FTraceViewModel>& B)
+				{
+					return A->ConfigurationType == B->ConfigurationType ?
+						A->Timestamp < B->Timestamp :
+						A->ConfigurationType < B->ConfigurationType;
+				});
+		}
+		else
+		{
+			TraceViewModels.Sort([](const TSharedPtr<FTraceViewModel>& A, const TSharedPtr<FTraceViewModel>& B)
+				{
+					return A->ConfigurationType == B->ConfigurationType ?
+						A->Timestamp < B->Timestamp :
+						A->ConfigurationType > B->ConfigurationType;
+				});
+		}
+	}
+	else if (SortColumn == TraceStoreColumns::BuildTarget)
+	{
+		if (SortMode == EColumnSortMode::Ascending)
+		{
+			TraceViewModels.Sort([](const TSharedPtr<FTraceViewModel>& A, const TSharedPtr<FTraceViewModel>& B)
+				{
+					return A->TargetType == B->TargetType ?
+						A->Timestamp < B->Timestamp :
+						A->TargetType < B->TargetType;
+				});
+		}
+		else
+		{
+			TraceViewModels.Sort([](const TSharedPtr<FTraceViewModel>& A, const TSharedPtr<FTraceViewModel>& B)
+				{
+					return A->TargetType == B->TargetType ?
+						A->Timestamp < B->Timestamp :
+						A->TargetType > B->TargetType;
+				});
+		}
+	}
+	else if (SortColumn == TraceStoreColumns::Size)
+	{
+		if (SortMode == EColumnSortMode::Ascending)
+		{
+			TraceViewModels.Sort([](const TSharedPtr<FTraceViewModel>& A, const TSharedPtr<FTraceViewModel>& B)
+				{ return A->Size < B->Size; });
+		}
+		else
+		{
+			TraceViewModels.Sort([](const TSharedPtr<FTraceViewModel>& A, const TSharedPtr<FTraceViewModel>& B)
+				{ return A->Size > B->Size; });
+		}
+	}
+	else if (SortColumn == TraceStoreColumns::Status)
+	{
+		if (SortMode == EColumnSortMode::Ascending)
+		{
+			TraceViewModels.Sort([](const TSharedPtr<FTraceViewModel>& A, const TSharedPtr<FTraceViewModel>& B)
+				{
+					return A->bIsLive == B->bIsLive ?
+						A->Timestamp < B->Timestamp :
+						A->bIsLive;
+				});
+		}
+		else
+		{
+			TraceViewModels.Sort([](const TSharedPtr<FTraceViewModel>& A, const TSharedPtr<FTraceViewModel>& B)
+				{
+					return A->bIsLive == B->bIsLive ?
+						A->Timestamp < B->Timestamp :
+						B->bIsLive;
+				});
+		}
+	}
+	else
+	{
+		Algo::SortBy(TraceViewModels, &FTraceViewModel::Timestamp);
+	}
+}
+
+////////////////////////////////////////////////////////////////////////////////////////////////////
+
 #undef LOCTEXT_NAMESPACE
 #define LOCTEXT_NAMESPACE "SConnectionWindow"
 
diff --git a/Engine/Source/Developer/TraceInsights/Private/Insights/Widgets/SStartPageWindow.h b/Engine/Source/Developer/TraceInsights/Private/Insights/Widgets/SStartPageWindow.h
--- a/Engine/Source/Developer/TraceInsights/Private/Insights/Widgets/SStartPageWindow.h
+++ b/Engine/Source/Developer/TraceInsights/Private/Insights/Widgets/SStartPageWindow.h
@@ -25,6 +25,19 @@ namespace Insights
 	struct FStoreBrowserTraceInfo;
 }
 
+namespace TraceStoreColumns
+{
+	static const FName Date(TEXT("Date"));
+	static const FName Name(TEXT("Name"));
+	static const FName Uri(TEXT("Uri"));
+	static const FName Platform(TEXT("Platform"));
+	static const FName AppName(TEXT("AppName"));
+	static const FName BuildConfig(TEXT("BuildConfig"));
+	static const FName BuildTarget(TEXT("BuildTarget"));
+	static const FName Size(TEXT("Size"));
+	static const FName Status(TEXT("Status"));
+}
+
 ////////////////////////////////////////////////////////////////////////////////////////////////////
 
 /** Type definition for shared pointers to instances of SNotificationItem. */
@@ -227,6 +240,12 @@ private:
 	 */
 	virtual FReply OnDragOver(const FGeometry& MyGeometry, const FDragDropEvent& DragDropEvent)  override;
 
+	EColumnSortMode::Type GetSortModeForColumn(const FName ColumnId) const;
+	void OnSortModeChanged(const EColumnSortPriority::Type SortPriority, const FName& ColumnId, const EColumnSortMode::Type InSortMode);
+
+	void UpdateSorting();
+	void UpdateTraceListView();
+
 private:
 	/** Widget for the non-intrusive notifications. */
 	TSharedPtr<SNotificationList> NotificationList;
@@ -264,6 +283,9 @@ private:
 	TSharedPtr<SListView<TSharedPtr<FTraceViewModel>>> TraceListView;
 	TSharedPtr<FTraceViewModel> SelectedTrace;
 
+	FName SortColumn;
+	EColumnSortMode::Type SortMode;
+
 	FString SplashScreenOverlayTraceFile;
 	float SplashScreenOverlayFadeTime;
 
