[Insights] Implemented workaround for the missing brushes used in Automation window (by re-declaring those in InsightsStyle).

diff --git a/Engine/Source/Developer/TraceInsights/Private/Insights/InsightsStyle.cpp b/Engine/Source/Developer/TraceInsights/Private/Insights/InsightsStyle.cpp
--- a/Engine/Source/Developer/TraceInsights/Private/Insights/InsightsStyle.cpp
+++ b/Engine/Source/Developer/TraceInsights/Private/Insights/InsightsStyle.cpp
@@ -5,6 +5,7 @@
 #include "Styling/AppStyle.h"
 #include "Styling/SlateStyleMacros.h"
 #include "Styling/SlateStyleRegistry.h"
+#include "Styling/StyleColors.h"
 #include "Styling/ToolBarStyle.h"
 
 LLM_DEFINE_TAG(Insights);
@@ -80,6 +81,12 @@ void FInsightsStyle::FStyle::SyncParentStyles()
 	const ISlateStyle* ParentStyle = GetParentStyle();
 
 	NormalText = ParentStyle->GetWidgetStyle<FTextBlockStyle>("NormalText");
+	Button = ParentStyle->GetWidgetStyle<FButtonStyle>("Button");
+
+	SelectorColor = ParentStyle->GetSlateColor("SelectorColor");
+	SelectionColor = ParentStyle->GetSlateColor("SelectionColor");
+	SelectionColor_Inactive = ParentStyle->GetSlateColor("SelectionColor_Inactive");
+	SelectionColor_Pressed = ParentStyle->GetSlateColor("SelectionColor_Pressed");
 }
 
 ////////////////////////////////////////////////////////////////////////////////////////////////////
@@ -87,6 +94,7 @@ void FInsightsStyle::FStyle::SyncParentStyles()
 #define EDITOR_IMAGE_BRUSH(RelativePath, ...) IMAGE_BRUSH("../../../Editor/Slate/" RelativePath, __VA_ARGS__)
 #define EDITOR_IMAGE_BRUSH_SVG(RelativePath, ...) IMAGE_BRUSH_SVG("../../../Editor/Slate/" RelativePath, __VA_ARGS__)
 #define EDITOR_BOX_BRUSH(RelativePath, ...) BOX_BRUSH("../../../Editor/Slate/" RelativePath, __VA_ARGS__)
+#define EDITOR_BORDER_BRUSH(RelativePath, ...) BORDER_BRUSH("../../../Editor/Slate/" RelativePath, __VA_ARGS__)
 #define TODO_IMAGE_BRUSH(...) EDITOR_IMAGE_BRUSH_SVG("Starship/Common/StaticMesh", __VA_ARGS__)
 
 void FInsightsStyle::FStyle::Initialize()
@@ -390,6 +398,137 @@ void FInsightsStyle::FStyle::Initialize()
 		Set("SecondaryToolbar2.MinUniformToolbarSize", 32.0f);
 		Set("SecondaryToolbar2.MaxUniformToolbarSize", 32.0f);
 	}
+
+	//////////////////////////////////////////////////
+	// Workaround for Automation styles...
+#if 1
+
+	const FVector2D Icon24x24(24.0f, 24.0f);
+	const FVector2D Icon40x40(40.0f, 40.0f);
+
+	Set("Automation.Header", FTextBlockStyle(NormalText)
+		.SetFont(DEFAULT_FONT("Mono", 12))
+		.SetColorAndOpacity(FLinearColor(FColor(0xffffffff))));
+
+	Set("Automation.Normal", FTextBlockStyle(NormalText)
+		.SetFont(DEFAULT_FONT("Mono", 9))
+		.SetColorAndOpacity(FLinearColor(FColor(0xffaaaaaa))));
+
+	Set("Automation.Warning", FTextBlockStyle(NormalText)
+		.SetFont(DEFAULT_FONT("Mono", 9))
+		.SetColorAndOpacity(FLinearColor(FColor(0xffbbbb44))));
+
+	Set("Automation.Error", FTextBlockStyle(NormalText)
+		.SetFont(DEFAULT_FONT("Mono", 9))
+		.SetColorAndOpacity(FLinearColor(FColor(0xffff0000))));
+
+	Set("Automation.ReportHeader", FTextBlockStyle(NormalText)
+		.SetFont(DEFAULT_FONT("Mono", 10))
+		.SetColorAndOpacity(FLinearColor(FColor(0xffffffff))));
+
+	Set("Automation.Success", new EDITOR_IMAGE_BRUSH("Automation/Success", Icon16x16));
+	Set("Automation.Warning", new EDITOR_IMAGE_BRUSH("Automation/Warning", Icon16x16));
+	Set("Automation.Fail", new EDITOR_IMAGE_BRUSH("Automation/Fail", Icon16x16));
+	Set("Automation.InProcess", new EDITOR_IMAGE_BRUSH("Automation/InProcess", Icon16x16));
+	Set("Automation.NotRun", new EDITOR_IMAGE_BRUSH("Automation/NotRun", Icon16x16, FLinearColor(0.0f, 0.0f, 0.0f, 0.4f)));
+	Set("Automation.Skipped", new EDITOR_IMAGE_BRUSH("Automation/NoSessionWarning", Icon16x16));
+	Set("Automation.ParticipantsWarning", new EDITOR_IMAGE_BRUSH("Automation/ParticipantsWarning", Icon16x16));
+	Set("Automation.Participant", new EDITOR_IMAGE_BRUSH("Automation/Participant", Icon16x16));
+
+	Set("Automation.SmokeTest", new EDITOR_IMAGE_BRUSH("Automation/SmokeTest", Icon16x16));
+	Set("Automation.SmokeTestParent", new EDITOR_IMAGE_BRUSH("Automation/SmokeTestParent", Icon16x16));
+
+	Set("AutomationWindow.RunTests", new EDITOR_IMAGE_BRUSH("Automation/RunTests", Icon40x40));
+	Set("AutomationWindow.RefreshTests", new EDITOR_IMAGE_BRUSH("Automation/RefreshTests", Icon40x40));
+	Set("AutomationWindow.FindWorkers", new EDITOR_IMAGE_BRUSH("Automation/RefreshWorkers", Icon40x40));
+	Set("AutomationWindow.StopTests", new EDITOR_IMAGE_BRUSH("Automation/StopTests", Icon40x40));
+	Set("AutomationWindow.RunTests.Small", new EDITOR_IMAGE_BRUSH("Automation/RunTests", Icon20x20));
+	Set("AutomationWindow.RefreshTests.Small", new EDITOR_IMAGE_BRUSH("Automation/RefreshTests", Icon20x20));
+	Set("AutomationWindow.FindWorkers.Small", new EDITOR_IMAGE_BRUSH("Automation/RefreshWorkers", Icon20x20));
+	Set("AutomationWindow.StopTests.Small", new EDITOR_IMAGE_BRUSH("Automation/StopTests", Icon20x20));
+
+	//filter icons
+	Set("AutomationWindow.ErrorFilter", new EDITOR_IMAGE_BRUSH("Automation/ErrorFilter", Icon40x40));
+	Set("AutomationWindow.WarningFilter", new EDITOR_IMAGE_BRUSH("Automation/WarningFilter", Icon40x40));
+	Set("AutomationWindow.SmokeTestFilter", new EDITOR_IMAGE_BRUSH("Automation/SmokeTestFilter", Icon40x40));
+	Set("AutomationWindow.DeveloperDirectoryContent", new EDITOR_IMAGE_BRUSH("Automation/DeveloperDirectoryContent", Icon40x40));
+	Set("AutomationWindow.ExcludedTestsFilter", new EDITOR_IMAGE_BRUSH("Automation/ExcludedTestsFilter", Icon40x40));
+	Set("AutomationWindow.ErrorFilter.Small", new EDITOR_IMAGE_BRUSH("Automation/ErrorFilter", Icon20x20));
+	Set("AutomationWindow.WarningFilter.Small", new EDITOR_IMAGE_BRUSH("Automation/WarningFilter", Icon20x20));
+	Set("AutomationWindow.SmokeTestFilter.Small", new EDITOR_IMAGE_BRUSH("Automation/SmokeTestFilter", Icon20x20));
+	Set("AutomationWindow.DeveloperDirectoryContent.Small", new EDITOR_IMAGE_BRUSH("Automation/DeveloperDirectoryContent", Icon20x20));
+	Set("AutomationWindow.TrackHistory", new EDITOR_IMAGE_BRUSH("Automation/TrackTestHistory", Icon40x40));
+
+	//device group settings
+	Set("AutomationWindow.GroupSettings", new EDITOR_IMAGE_BRUSH("Automation/Groups", Icon40x40));
+	Set("AutomationWindow.GroupSettings.Small", new EDITOR_IMAGE_BRUSH("Automation/Groups", Icon20x20));
+
+	//test preset icons
+	Set("AutomationWindow.PresetNew", new EDITOR_IMAGE_BRUSH("Icons/icon_add_40x", Icon16x16));
+	Set("AutomationWindow.PresetSave", new EDITOR_IMAGE_BRUSH_SVG("Starship/Common/SaveCurrent", Icon16x16));
+	Set("AutomationWindow.PresetRemove", new EDITOR_IMAGE_BRUSH("Icons/icon_Cascade_DeleteLOD_40x", Icon16x16));
+
+	//test backgrounds
+	Set("AutomationWindow.GameGroupBorder", new EDITOR_BOX_BRUSH("Automation/GameGroupBorder", FMargin(4.0f / 16.0f)));
+	Set("AutomationWindow.EditorGroupBorder", new EDITOR_BOX_BRUSH("Automation/EditorGroupBorder", FMargin(4.0f / 16.0f)));
+
+	Set("MessageLog.ListBorder", new EDITOR_BOX_BRUSH("/Docking/AppTabContentArea", FMargin(4 / 16.0f)));
+
+	Set("Launcher.Platform_Windows", new EDITOR_IMAGE_BRUSH("Launcher/All_Platforms_24x", Icon24x24));
+
+	// Output Log Window
+	{
+		const FSlateColor LogColor_SelectionBackground(EStyleColor::User2);
+		const FSlateColor LogColor_Normal(EStyleColor::User3);
+		const FSlateColor LogColor_Command(EStyleColor::User4);
+
+		const int32 LogFontSize = 9;
+		const FTextBlockStyle NormalLogText = FTextBlockStyle(NormalText)
+			.SetFont(DEFAULT_FONT("Mono", LogFontSize))
+			.SetColorAndOpacity(LogColor_Normal)
+			.SetSelectedBackgroundColor(LogColor_SelectionBackground)
+			.SetHighlightColor(FStyleColors::Black);
+
+		Set("Log.Normal", NormalLogText);
+
+		Set("Log.Command", FTextBlockStyle(NormalLogText)
+			.SetColorAndOpacity(LogColor_Command));
+
+		Set("Log.Warning", FTextBlockStyle(NormalLogText)
+			.SetColorAndOpacity(FStyleColors::Warning));
+
+		Set("Log.Error", FTextBlockStyle(NormalLogText)
+			.SetColorAndOpacity(FStyleColors::Error));
+	}
+
+	{
+		Set("ToggleButton", FButtonStyle(Button)
+			.SetNormal(FSlateNoResource())
+			.SetHovered(EDITOR_BOX_BRUSH("Common/RoundedSelection_16x", 4.0f / 16.0f, SelectionColor))
+			.SetPressed(EDITOR_BOX_BRUSH("Common/RoundedSelection_16x", 4.0f / 16.0f, SelectionColor_Pressed)));
+	}
+
+	{
+		FTextBlockStyle InheritedFromNativeTextStyle = FTextBlockStyle(NormalText)
+			.SetFont(DEFAULT_FONT("Regular", 10));
+
+		Set("Common.InheritedFromNativeTextStyle", InheritedFromNativeTextStyle);
+
+		// Go to native class hyperlink
+		FButtonStyle EditNativeHyperlinkButton = FButtonStyle()
+			.SetNormal(EDITOR_BORDER_BRUSH("Old/HyperlinkDotted", FMargin(0, 0, 0, 3 / 16.0f)))
+			.SetPressed(FSlateNoResource())
+			.SetHovered(EDITOR_BORDER_BRUSH("Old/HyperlinkUnderline", FMargin(0, 0, 0, 3 / 16.0f)));
+		FHyperlinkStyle EditNativeHyperlinkStyle = FHyperlinkStyle()
+			.SetUnderlineStyle(EditNativeHyperlinkButton)
+			.SetTextStyle(InheritedFromNativeTextStyle)
+			.SetPadding(FMargin(0.0f));
+
+		Set("Common.GotoNativeCodeHyperlink", EditNativeHyperlinkStyle);
+	}
+
+#endif
+	//////////////////////////////////////////////////
 }
 
 #undef TODO_IMAGE_BRUSH
diff --git a/Engine/Source/Developer/TraceInsights/Private/Insights/InsightsStyle.h b/Engine/Source/Developer/TraceInsights/Private/Insights/InsightsStyle.h
--- a/Engine/Source/Developer/TraceInsights/Private/Insights/InsightsStyle.h
+++ b/Engine/Source/Developer/TraceInsights/Private/Insights/InsightsStyle.h
@@ -53,6 +53,11 @@ public:
 
 		// Styles inherited from the parent style
 		FTextBlockStyle NormalText;
+		FButtonStyle Button;
+		FSlateColor SelectorColor;
+		FSlateColor SelectionColor;
+		FSlateColor SelectionColor_Inactive;
+		FSlateColor SelectionColor_Pressed;
 	};
 
 private:
