[Insights] Session Browser: Added filtering for the list of trace sessions (by Name, by Platform, by [REDACTED_PERSON], by [REDACTED_PERSON] and by [REDACTED_PERSON]).

diff --git a/Engine/Source/Developer/TraceInsights/Private/Insights/Widgets/SStartPageWindow.cpp b/Engine/Source/Developer/TraceInsights/Private/Insights/Widgets/SStartPageWindow.cpp
--- a/Engine/Source/Developer/TraceInsights/Private/Insights/Widgets/SStartPageWindow.cpp
+++ b/Engine/Source/Developer/TraceInsights/Private/Insights/Widgets/SStartPageWindow.cpp
@@ -623,21 +623,31 @@ STraceStoreWindow::STraceStoreWindow()
 	, DurationActive(0.0f)
 	, ActiveTimerHandle()
 	, MainContentPanel()
-	, LiveSessionCount(0)
-	, bAutoStartAnalysisForLiveSessions(false)
-	, AutoStartedSessions()
-	, AutoStartPlatformFilter()
-	, AutoStartAppNameFilter()
-	, AutoStartConfigurationTypeFilter(EBuildConfiguration::Unknown)
-	, AutoStartTargetTypeFilter(EBuildTargetType::Unknown)
 	, StoreBrowser(new Insights::FStoreBrowser())
 	, TracesChangeSerial(0)
 	, TraceViewModels()
+	, FilteredTraceViewModels()
 	, TraceViewModelMap()
 	, TraceListView()
 	, SelectedTrace()
+	, bIsUserSelectedTrace(false)
+	, Filters()
+	, FilterByNameSearchBox()
+	, FilterByName()
+	, FilterByPlatform()
+	, FilterByAppName()
+	, FilterByBuildConfig()
+	, FilterByBuildTarget()
+	, bFilterStatsTextIsDirty(true)
+	, FilterStatsText()
 	, SortColumn(TraceStoreColumns::Date)
 	, SortMode(EColumnSortMode::Ascending)
+	, bAutoStartAnalysisForLiveSessions(false)
+	, AutoStartedSessions()
+	, AutoStartPlatformFilter()
+	, AutoStartAppNameFilter()
+	, AutoStartConfigurationTypeFilter(EBuildConfiguration::Unknown)
+	, AutoStartTargetTypeFilter(EBuildTargetType::Unknown)
 	, SplashScreenOverlayFadeTime(0.0f)
 {
 }
@@ -702,16 +712,25 @@ void STraceStoreWindow::Construct(const FArguments& InArgs)
 			+ SVerticalBox::Slot()
 			.HAlign(HAlign_Fill)
 			.AutoHeight()
-			.Padding(12.0f, 8.0f, 12.0f, 4.0f)
+			.Padding(12.0f, 8.0f, 12.0f, 0.0f)
 			[
 				ConstructTraceStoreDirectoryPanel()
 			]
 
+			+ SVerticalBox::Slot()
+			.HAlign(HAlign_Fill)
+			.VAlign(VAlign_Fill)
+			.AutoHeight()
+			.Padding(0.0f, 0.0f, 0.0f, 0.0f)
+			[
+				ConstructFiltersToolbar()
+			]
+
 			+ SVerticalBox::Slot()
 			.HAlign(HAlign_Fill)
 			.VAlign(VAlign_Fill)
 			.FillHeight(1.0f)
-			.Padding(3.0f, 4.0f)
+			.Padding(3.0f, 0.0f, 3.0f, 4.0f)
 			[
 				ConstructSessionsPanel()
 			]
@@ -771,6 +790,8 @@ void STraceStoreWindow::Construct(const FArguments& InArgs)
 		.Expose(OverlaySettingsSlot)
 	];
 
+	CreateFilters();
+
 	RefreshTraceList();
 
 	FSlateApplication::Get().SetKeyboardFocus(TraceListView);
@@ -779,6 +800,104 @@ void STraceStoreWindow::Construct(const FArguments& InArgs)
 
 ////////////////////////////////////////////////////////////////////////////////////////////////////
 
+TSharedRef<SWidget> STraceStoreWindow::ConstructFiltersToolbar()
+{
+	FSlimHorizontalToolBarBuilder ToolbarBuilder(TSharedPtr<const FUICommandList>(), FMultiBoxCustomization::None);
+	ToolbarBuilder.SetStyle(&FInsightsStyle::Get(), "SecondaryToolbar");
+
+	ToolbarBuilder.BeginSection("Filters");
+	{
+		// Text Filter (Search Box)
+		ToolbarBuilder.AddWidget(
+			SAssignNew(FilterByNameSearchBox, SSearchBox)
+			.MinDesiredWidth(150.0f)
+			.HintText(LOCTEXT("NameFilter_Hint", "Name"))
+			.ToolTipText(LOCTEXT("NameFilter_Tooltip", "Type here to filter the list of trace sessions by name."))
+			.IsEnabled_Lambda([this]() { return TraceViewModels.Num() > 0; })
+			.OnTextChanged(this, &STraceStoreWindow::FilterByNameSearchBox_OnTextChanged)
+			.DelayChangeNotificationsWhileTyping(true)
+		);
+
+		// Filter by Platform
+		ToolbarBuilder.AddComboButton(
+			FUIAction(),
+			FOnGetContent::CreateSP(this, &STraceStoreWindow::MakePlatformFilterMenu),
+			LOCTEXT("FilterByPlatformText", "Platform"),
+			LOCTEXT("FilterByPlatformToolTip", "Filters the list of trace sessions by platform."),
+			FSlateIcon(FInsightsStyle::GetStyleSetName(), "Icons.Filter.ToolBar"),
+			false
+		);
+
+		// Filter by AppName
+		ToolbarBuilder.AddComboButton(
+			FUIAction(),
+			FOnGetContent::CreateSP(this, &STraceStoreWindow::MakeAppNameFilterMenu),
+			LOCTEXT("FilterByAppNameText", "App Name"),
+			LOCTEXT("FilterByAppNameToolTip", "Filters the list of trace sessions by application name."),
+			FSlateIcon(FInsightsStyle::GetStyleSetName(), "Icons.Filter.ToolBar"),
+			false
+		);
+
+		// Filter by Build Config
+		ToolbarBuilder.AddComboButton(
+			FUIAction(),
+			FOnGetContent::CreateSP(this, &STraceStoreWindow::MakeBuildConfigFilterMenu),
+			LOCTEXT("FilterByBuildConfigText", "Build Config"),
+			LOCTEXT("FilterByBuildConfigToolTip", "Filters the list of trace sessions by build configuration."),
+			FSlateIcon(FInsightsStyle::GetStyleSetName(), "Icons.Filter.ToolBar"),
+			false
+		);
+
+		// Filter by Build Target
+		ToolbarBuilder.AddComboButton(
+			FUIAction(),
+			FOnGetContent::CreateSP(this, &STraceStoreWindow::MakeBuildTargetFilterMenu),
+			LOCTEXT("FilterByBuildTargetText", "Build Target"),
+			LOCTEXT("FilterByBuildTargetToolTip", "Filters the list of trace sessions by build target."),
+			FSlateIcon(FInsightsStyle::GetStyleSetName(), "Icons.Filter.ToolBar"),
+			false
+		);
+	}
+	ToolbarBuilder.EndSection();
+
+	FSlimHorizontalToolBarBuilder RightSideToolbarBuilder(TSharedPtr<const FUICommandList>(), FMultiBoxCustomization::None);
+	RightSideToolbarBuilder.SetStyle(&FInsightsStyle::Get(), "PrimaryToolbar");
+	RightSideToolbarBuilder.BeginSection("FilterStats");
+	{
+		// Filter Stats Text (number and size of filtered trace sessions)
+		RightSideToolbarBuilder.AddWidget(
+			SNew(SBox)
+			.VAlign(VAlign_Bottom)
+			.Padding(FMargin(4.0f, 0.0f, 4.0f, 0.0f))
+			[
+				SNew(STextBlock)
+				.Text(this, &STraceStoreWindow::GetFilterStatsText)
+			]
+		);
+	}
+	RightSideToolbarBuilder.EndSection();
+
+	return SNew(SHorizontalBox)
+
+		+ SHorizontalBox::Slot()
+		.HAlign(HAlign_Fill)
+		.VAlign(VAlign_Center)
+		.FillWidth(1.0f)
+		.Padding(0.0f)
+		[
+			ToolbarBuilder.MakeWidget()
+		]
+
+	+ SHorizontalBox::Slot()
+		.HAlign(HAlign_Right)
+		.VAlign(VAlign_Center)
+		.AutoWidth()
+		.Padding(0.0f)
+		[
+			RightSideToolbarBuilder.MakeWidget()
+		];
+}
+
 TSharedRef<SWidget> STraceStoreWindow::ConstructSessionsPanel()
 {
 	TSharedRef<SWidget> Widget = SAssignNew(TraceListView, SListView<TSharedPtr<FTraceViewModel>>)
@@ -787,7 +906,7 @@ TSharedRef<SWidget> STraceStoreWindow::ConstructSessionsPanel()
 		.SelectionMode(ESelectionMode::Single)
 		.OnSelectionChanged(this, &STraceStoreWindow::TraceList_OnSelectionChanged)
 		.OnMouseButtonDoubleClick(this, &STraceStoreWindow::TraceList_OnMouseButtonDoubleClick)
-		.ListItemsSource(&TraceViewModels)
+		.ListItemsSource(&FilteredTraceViewModels)
 		.OnGenerateRow(this, &STraceStoreWindow::TraceList_OnGenerateRow)
 		.ConsumeMouseWheel(EConsumeMouseWheel::Always)
 		//.OnContextMenuOpening(FOnContextMenuOpening::CreateSP(this, &STraceStoreWindow::TraceList_GetContextMenu))
@@ -796,39 +915,83 @@ TSharedRef<SWidget> STraceStoreWindow::ConstructSessionsPanel()
 			SNew(SHeaderRow)
 
 			+ SHeaderRow::Column(TraceStoreColumns::Name)
-			.DefaultLabel(LOCTEXT("NameColumn", "Name"))
 			.FillWidth(0.25f)
 			.InitialSortMode(EColumnSortMode::Ascending)
 			.SortMode(this, &STraceStoreWindow::GetSortModeForColumn, TraceStoreColumns::Name)
 			.OnSort(this, &STraceStoreWindow::OnSortModeChanged)
+			[
+				SNew(SBox)
+				.VAlign(VAlign_Center)
+				[
+					SNew(STextBlock)
+					.Text(LOCTEXT("NameColumn", "Name"))
+					.ColorAndOpacity_Lambda([this] { return FilterByName->GetRawFilterText().IsEmpty() ? FLinearColor(0.5f, 0.5f, 0.5f, 1.0f) : FLinearColor(0.3f, 0.75f, 1.0f, 1.0f); })
+				]
+			]
 
 			+ SHeaderRow::Column(TraceStoreColumns::Platform)
-			.DefaultLabel(LOCTEXT("PlatformColumn", "Platform"))
 			.FillWidth(0.1f)
 			.InitialSortMode(EColumnSortMode::Ascending)
 			.SortMode(this, &STraceStoreWindow::GetSortModeForColumn, TraceStoreColumns::Platform)
 			.OnSort(this, &STraceStoreWindow::OnSortModeChanged)
+			.OnGetMenuContent(this, &STraceStoreWindow::MakePlatformColumnHeaderMenu)
+			[
+				SNew(SBox)
+				.VAlign(VAlign_Center)
+				[
+					SNew(STextBlock)
+					.Text(LOCTEXT("PlatformColumn", "Platform"))
+					.ColorAndOpacity_Lambda([this] { return FilterByPlatform->IsEmpty() ? FLinearColor(0.5f, 0.5f, 0.5f, 1.0f) : FLinearColor(0.3f, 0.75f, 1.0f, 1.0f); })
+				]
+			]
 
 			+ SHeaderRow::Column(TraceStoreColumns::AppName)
-			.DefaultLabel(LOCTEXT("AppNameColumn", "App Name"))
 			.FillWidth(0.1f)
 			.InitialSortMode(EColumnSortMode::Ascending)
 			.SortMode(this, &STraceStoreWindow::GetSortModeForColumn, TraceStoreColumns::AppName)
 			.OnSort(this, &STraceStoreWindow::OnSortModeChanged)
+			.OnGetMenuContent(this, &STraceStoreWindow::MakeAppNameColumnHeaderMenu)
+			[
+				SNew(SBox)
+				.VAlign(VAlign_Center)
+				[
+					SNew(STextBlock)
+					.Text(LOCTEXT("AppNameColumn", "App Name"))
+					.ColorAndOpacity_Lambda([this] { return FilterByAppName->IsEmpty() ? FLinearColor(0.5f, 0.5f, 0.5f, 1.0f) : FLinearColor(0.3f, 0.75f, 1.0f, 1.0f); })
+				]
+			]
 
 			+ SHeaderRow::Column(TraceStoreColumns::BuildConfig)
-			.DefaultLabel(LOCTEXT("BuildConfigColumn", "Build Config"))
 			.FillWidth(0.1f)
 			.InitialSortMode(EColumnSortMode::Ascending)
 			.SortMode(this, &STraceStoreWindow::GetSortModeForColumn, TraceStoreColumns::BuildConfig)
 			.OnSort(this, &STraceStoreWindow::OnSortModeChanged)
+			.OnGetMenuContent(this, &STraceStoreWindow::MakeBuildConfigColumnHeaderMenu)
+			[
+				SNew(SBox)
+				.VAlign(VAlign_Center)
+				[
+					SNew(STextBlock)
+					.Text(LOCTEXT("BuildConfigColumn", "Build Config"))
+					.ColorAndOpacity_Lambda([this] { return FilterByBuildConfig->IsEmpty() ? FLinearColor(0.5f, 0.5f, 0.5f, 1.0f) : FLinearColor(0.3f, 0.75f, 1.0f, 1.0f); })
+				]
+			]
 
 			+ SHeaderRow::Column(TraceStoreColumns::BuildTarget)
-			.DefaultLabel(LOCTEXT("BuildTargetColumn", "Build Target"))
 			.FillWidth(0.1f)
 			.InitialSortMode(EColumnSortMode::Ascending)
 			.SortMode(this, &STraceStoreWindow::GetSortModeForColumn, TraceStoreColumns::BuildTarget)
 			.OnSort(this, &STraceStoreWindow::OnSortModeChanged)
+			.OnGetMenuContent(this, &STraceStoreWindow::MakeBuildTargetColumnHeaderMenu)
+			[
+				SNew(SBox)
+				.VAlign(VAlign_Center)
+				[
+					SNew(STextBlock)
+					.Text(LOCTEXT("BuildTargetColumn", "Build Target"))
+					.ColorAndOpacity_Lambda([this] { return FilterByBuildTarget->IsEmpty() ? FLinearColor(0.5f, 0.5f, 0.5f, 1.0f) : FLinearColor(0.3f, 0.75f, 1.0f, 1.0f); })
+				]
+			]
 
 			+ SHeaderRow::Column(TraceStoreColumns::Size)
 			.DefaultLabel(LOCTEXT("SizeColumn", "File Size"))
@@ -1233,6 +1396,7 @@ void STraceStoreWindow::UpdateTrace(FTraceViewModel& InOutTrace, const Insights:
 
 void STraceStoreWindow::OnTraceListChanged()
 {
+	UpdateFiltering();
 	UpdateSorting();
 	UpdateTraceListView();
 }
@@ -1247,7 +1411,7 @@ void STraceStoreWindow::UpdateTraceListView()
 	}
 
 	TSharedPtr<FTraceViewModel> NewSelectedTrace;
-	if (SelectedTrace)
+	if (bIsUserSelectedTrace && SelectedTrace)
 	{
 		// Identify the previously selected trace (if stil available) to ensure selection remains unchanged.
 		TSharedPtr<FTraceViewModel>* TracePtrPtr = TraceViewModelMap.Find(SelectedTrace->TraceId);
@@ -1260,9 +1424,9 @@ void STraceStoreWindow::UpdateTraceListView()
 		(SortColumn == TraceStoreColumns::Status && SortMode == EColumnSortMode::Ascending))
 	{
 		// If no selection, auto select the last (newest) trace.
-		if (!NewSelectedTrace.IsValid() && TraceViewModels.Num() > 0)
+		if (!NewSelectedTrace.IsValid() && FilteredTraceViewModels.Num() > 0)
 		{
-			NewSelectedTrace = TraceViewModels.Last();
+			NewSelectedTrace = FilteredTraceViewModels.Last();
 		}
 
 		TraceListView->ScrollToBottom();
@@ -1270,15 +1434,16 @@ void STraceStoreWindow::UpdateTraceListView()
 	else
 	{
 		// If no selection, auto select the first trace.
-		if (!NewSelectedTrace.IsValid() && TraceViewModels.Num() > 0)
+		if (!NewSelectedTrace.IsValid() && FilteredTraceViewModels.Num() > 0)
 		{
-			NewSelectedTrace = TraceViewModels[0];
+			NewSelectedTrace = FilteredTraceViewModels[0];
 		}
 	}
 
 	// Restore selection and ensure it is visible.
 	if (NewSelectedTrace.IsValid())
 	{
+		TraceListView->ClearSelection();
 		TraceListView->SetItemSelection(NewSelectedTrace, true);
 		TraceListView->RequestScrollIntoView(NewSelectedTrace);
 	}
@@ -1288,6 +1453,10 @@ void STraceStoreWindow::UpdateTraceListView()
 
 void STraceStoreWindow::TraceList_OnSelectionChanged(TSharedPtr<FTraceViewModel> TraceSession, ESelectInfo::Type SelectInfo)
 {
+	if (SelectInfo != ESelectInfo::Direct)
+	{
+		bIsUserSelectedTrace = true;
+	}
 	SelectedTrace = TraceSession;
 }
 
@@ -1324,6 +1493,11 @@ void STraceStoreWindow::Tick(const FGeometry& AllottedGeometry, const double InC
 		const uint64 WaitTime = static_cast<uint64>(0.5 / FPlatformTime::GetSecondsPerCycle64()); // 500ms
 		NextTimestamp = Time + WaitTime;
 		RefreshTraceList();
+
+		if (bFilterStatsTextIsDirty)
+		{
+			UpdateFilterStatsText();
+		}
 	}
 
 	TickSplashScreenOverlay(InDeltaTime);
@@ -1425,7 +1599,7 @@ FReply STraceStoreWindow::OnDrop(const FGeometry& MyGeometry, const FDragDropEve
 
 bool STraceStoreWindow::Open_IsEnabled() const
 {
-	return TraceViewModels.Num() > 0;
+	return FilteredTraceViewModels.Num() > 0;
 }
 
 ////////////////////////////////////////////////////////////////////////////////////////////////////
@@ -1561,6 +1735,8 @@ void STraceStoreWindow::LoadTrace(uint32 InTraceId)
 
 TSharedRef<SWidget> STraceStoreWindow::MakeTraceListMenu()
 {
+	FSlateApplication::Get().CloseToolTip();
+
 	RefreshTraceList();
 
 	FMenuBuilder MenuBuilder(/*bInShouldCloseWindowAfterMenuSelection=*/true, nullptr);
@@ -1686,6 +1862,154 @@ TSharedRef<SWidget> STraceStoreWindow::MakeTraceListMenu()
 
 ////////////////////////////////////////////////////////////////////////////////////////////////////
 
+TSharedRef<SWidget> STraceStoreWindow::MakePlatformColumnHeaderMenu()
+{
+	FSlateApplication::Get().CloseToolTip();
+
+	FMenuBuilder MenuBuilder(/*bInShouldCloseWindowAfterMenuSelection=*/true, nullptr);
+
+	MenuBuilder.BeginSection("Filter", LOCTEXT("MenuSection_PlatformFilter", "Platform Filter"));
+	BuildPlatformFilterSubMenu(MenuBuilder);
+	MenuBuilder.EndSection();
+
+	return MenuBuilder.MakeWidget();
+}
+
+////////////////////////////////////////////////////////////////////////////////////////////////////
+
+TSharedRef<SWidget> STraceStoreWindow::MakePlatformFilterMenu()
+{
+	FSlateApplication::Get().CloseToolTip();
+
+	FMenuBuilder MenuBuilder(/*bInShouldCloseWindowAfterMenuSelection=*/true, nullptr);
+
+	MenuBuilder.BeginSection("Filter", LOCTEXT("MenuSection_PlatformFilter", "Platform Filter"));
+	BuildPlatformFilterSubMenu(MenuBuilder);
+	MenuBuilder.EndSection();
+
+	return MenuBuilder.MakeWidget();
+}
+
+////////////////////////////////////////////////////////////////////////////////////////////////////
+
+void STraceStoreWindow::BuildPlatformFilterSubMenu(FMenuBuilder& InMenuBuilder)
+{
+	FilterByPlatform->BuildMenu(InMenuBuilder, *this);
+}
+
+////////////////////////////////////////////////////////////////////////////////////////////////////
+
+TSharedRef<SWidget> STraceStoreWindow::MakeAppNameColumnHeaderMenu()
+{
+	FSlateApplication::Get().CloseToolTip();
+
+	FMenuBuilder MenuBuilder(/*bInShouldCloseWindowAfterMenuSelection=*/true, nullptr);
+
+	MenuBuilder.BeginSection("Filter", LOCTEXT("MenuSection_AppNameFilter", "App Name Filter"));
+	BuildAppNameFilterSubMenu(MenuBuilder);
+	MenuBuilder.EndSection();
+
+	return MenuBuilder.MakeWidget();
+}
+
+////////////////////////////////////////////////////////////////////////////////////////////////////
+
+TSharedRef<SWidget> STraceStoreWindow::MakeAppNameFilterMenu()
+{
+	FSlateApplication::Get().CloseToolTip();
+
+	FMenuBuilder MenuBuilder(/*bInShouldCloseWindowAfterMenuSelection=*/true, nullptr);
+
+	MenuBuilder.BeginSection("Filter", LOCTEXT("MenuSection_AppNameFilter", "App Name Filter"));
+	BuildAppNameFilterSubMenu(MenuBuilder);
+	MenuBuilder.EndSection();
+
+	return MenuBuilder.MakeWidget();
+}
+
+////////////////////////////////////////////////////////////////////////////////////////////////////
+
+void STraceStoreWindow::BuildAppNameFilterSubMenu(FMenuBuilder& InMenuBuilder)
+{
+	FilterByAppName->BuildMenu(InMenuBuilder, *this);
+}
+
+////////////////////////////////////////////////////////////////////////////////////////////////////
+
+TSharedRef<SWidget> STraceStoreWindow::MakeBuildConfigColumnHeaderMenu()
+{
+	FSlateApplication::Get().CloseToolTip();
+
+	FMenuBuilder MenuBuilder(/*bInShouldCloseWindowAfterMenuSelection=*/true, nullptr);
+
+	MenuBuilder.BeginSection("Filter", LOCTEXT("MenuSection_BuildConfigFilter", "Build Config Filter"));
+	BuildBuildConfigFilterSubMenu(MenuBuilder);
+	MenuBuilder.EndSection();
+
+	return MenuBuilder.MakeWidget();
+}
+
+////////////////////////////////////////////////////////////////////////////////////////////////////
+
+TSharedRef<SWidget> STraceStoreWindow::MakeBuildConfigFilterMenu()
+{
+	FSlateApplication::Get().CloseToolTip();
+
+	FMenuBuilder MenuBuilder(/*bInShouldCloseWindowAfterMenuSelection=*/true, nullptr);
+
+	MenuBuilder.BeginSection("Filter", LOCTEXT("MenuSection_BuildConfigFilter", "Build Config Filter"));
+	BuildBuildConfigFilterSubMenu(MenuBuilder);
+	MenuBuilder.EndSection();
+
+	return MenuBuilder.MakeWidget();
+}
+
+////////////////////////////////////////////////////////////////////////////////////////////////////
+
+void STraceStoreWindow::BuildBuildConfigFilterSubMenu(FMenuBuilder& InMenuBuilder)
+{
+	FilterByBuildConfig->BuildMenu(InMenuBuilder, *this);
+}
+
+////////////////////////////////////////////////////////////////////////////////////////////////////
+
+TSharedRef<SWidget> STraceStoreWindow::MakeBuildTargetColumnHeaderMenu()
+{
+	FSlateApplication::Get().CloseToolTip();
+
+	FMenuBuilder MenuBuilder(/*bInShouldCloseWindowAfterMenuSelection=*/true, nullptr);
+
+	MenuBuilder.BeginSection("Filter", LOCTEXT("MenuSection_BuildTargetFilter", "Build Target Filter"));
+	BuildBuildTargetFilterSubMenu(MenuBuilder);
+	MenuBuilder.EndSection();
+
+	return MenuBuilder.MakeWidget();
+}
+
+////////////////////////////////////////////////////////////////////////////////////////////////////
+
+TSharedRef<SWidget> STraceStoreWindow::MakeBuildTargetFilterMenu()
+{
+	FSlateApplication::Get().CloseToolTip();
+
+	FMenuBuilder MenuBuilder(/*bInShouldCloseWindowAfterMenuSelection=*/true, nullptr);
+
+	MenuBuilder.BeginSection("Filter", LOCTEXT("MenuSection_BuildTargetFilter", "Build Target Filter"));
+	BuildBuildTargetFilterSubMenu(MenuBuilder);
+	MenuBuilder.EndSection();
+
+	return MenuBuilder.MakeWidget();
+}
+
+////////////////////////////////////////////////////////////////////////////////////////////////////
+
+void STraceStoreWindow::BuildBuildTargetFilterSubMenu(FMenuBuilder& InMenuBuilder)
+{
+	FilterByBuildTarget->BuildMenu(InMenuBuilder, *this);
+}
+
+////////////////////////////////////////////////////////////////////////////////////////////////////
+
 FText STraceStoreWindow::GetTraceStoreDirectory() const
 {
 	return FText::FromString(FPaths::ConvertRelativePathToFull(FInsightsManager::Get()->GetStoreDir()));
@@ -1748,6 +2072,247 @@ void STraceStoreWindow::GetExtraCommandLineParams(FString& OutParams) const
 	}
 }
 
+////////////////////////////////////////////////////////////////////////////////////////////////////
+// TTraceSetFilter
+
+template<typename TSetType>
+void TTraceSetFilter<TSetType>::BuildMenu(FMenuBuilder& InMenuBuilder, STraceStoreWindow& Window)
+{
+	// Show/Hide All
+	{
+		FUIAction Action;
+		Action.ExecuteAction = FExecuteAction::CreateLambda([this, &Window]()
+			{
+				if (FilterSet.IsEmpty())
+				{
+					FilterSet.Reset();
+					for (const TSharedPtr<FTraceViewModel>& Trace : Window.GetAllAvailableTraces())
+					{
+						FilterSet.Add(GetFilterValueForTrace(*Trace));
+					}
+				}
+				else
+				{
+					FilterSet.Reset();
+				}
+				Window.OnFilterChanged();
+			});
+		Action.GetActionCheckState = FGetActionCheckState::CreateLambda([this]()
+			{
+				return FilterSet.IsEmpty() ? ECheckBoxState::Checked : ECheckBoxState::Unchecked;
+			});
+
+		InMenuBuilder.AddMenuEntry(
+			ToggleAllActionLabel,
+			ToggleAllActionTooltip.IsEmpty() ? TAttribute<FText>() : ToggleAllActionTooltip,
+			FSlateIcon(),
+			Action,
+			NAME_None,
+			EUserInterfaceActionType::ToggleButton
+		);
+	}
+
+	InMenuBuilder.AddSeparator();
+
+	TMap<TSetType, uint32> AllUniqueValues;
+	for (const TSharedPtr<FTraceViewModel>& Trace : Window.GetAllAvailableTraces())
+	{
+		TSetType Value = GetFilterValueForTrace(*Trace);
+		if (AllUniqueValues.Contains(Value))
+		{
+			AllUniqueValues[Value]++;
+		}
+		else
+		{
+			AllUniqueValues.Add(Value, 1);
+		}
+	}
+	AllUniqueValues.KeySort([](const TSetType& A, const TSetType& B) { return A < B; });
+
+	for (const auto& Pair : AllUniqueValues)
+	{
+		const TSetType& Value = Pair.Key;
+		FUIAction Action;
+		Action.ExecuteAction = FExecuteAction::CreateLambda([this, Value, &Window]()
+			{
+				if (FilterSet.Contains(Value))
+				{
+					FilterSet.Remove(Value);
+				}
+				else
+				{
+					FilterSet.Add(Value);
+				}
+				Window.OnFilterChanged();
+			});
+		Action.GetActionCheckState = FGetActionCheckState::CreateLambda([this, Value]()
+			{
+				return FilterSet.Contains(Value) ? ECheckBoxState::Unchecked : ECheckBoxState::Checked;
+			});
+
+		InMenuBuilder.AddMenuEntry(
+			FText::Format(LOCTEXT("FilterValue_Fmt", "{0} ({1})"), ValueToText(Value), FText::AsNumber(Pair.Value)),
+			TAttribute<FText>(), // no tooltip
+			FSlateIcon(),
+			Action,
+			NAME_None,
+			EUserInterfaceActionType::ToggleButton
+		);
+	}
+}
+
+template<typename TSetType>
+TTraceSetFilter<TSetType>::TTraceSetFilter()
+{
+	ToggleAllActionLabel = LOCTEXT("ToggleAll_Label", "Show/Hide All");
+	UndefinedValueLabel = LOCTEXT("UndefinedValueLabel", "N/A");
+}
+
+FTraceFilterByPlatform::FTraceFilterByPlatform()
+{
+	ToggleAllActionTooltip = LOCTEXT("FilterByPlatform_ToggleAll_Tooltip", "Shows or hides traces for all platforms.");
+}
+
+FTraceFilterByAppName::FTraceFilterByAppName()
+{
+	ToggleAllActionTooltip = LOCTEXT("FilterByAppName_ToggleAll_Tooltip", "Shows or hides traces for all app names.");
+}
+
+FTraceFilterByBuildConfig::FTraceFilterByBuildConfig()
+{
+	ToggleAllActionTooltip = LOCTEXT("FilterByBuildConfig_ToggleAll_Tooltip", "Shows or hides traces for all build configurations.");
+}
+
+FTraceFilterByBuildTarget::FTraceFilterByBuildTarget()
+{
+	ToggleAllActionTooltip = LOCTEXT("FilterByBuildConfig_ToggleAll_Tooltip", "Shows or hides traces for all build targets.");
+}
+
+////////////////////////////////////////////////////////////////////////////////////////////////////
+
+void STraceStoreWindow::FilterByNameSearchBox_OnTextChanged(const FText& InFilterText)
+{
+	FilterByName->SetRawFilterText(InFilterText);
+	FilterByNameSearchBox->SetError(FilterByName->GetFilterErrorText());
+
+	UpdateFiltering();
+	UpdateTraceListView();
+}
+
+////////////////////////////////////////////////////////////////////////////////////////////////////
+
+void STraceStoreWindow::OnFilterChanged()
+{
+	UpdateFiltering();
+	UpdateTraceListView();
+}
+
+////////////////////////////////////////////////////////////////////////////////////////////////////
+
+const TArray<TSharedPtr<FTraceViewModel>>& STraceStoreWindow::GetAllAvailableTraces() const
+{
+	return TraceViewModels;
+}
+
+////////////////////////////////////////////////////////////////////////////////////////////////////
+
+void STraceStoreWindow::CreateFilters()
+{
+	Filters = MakeShared<FTraceViewModelFilterCollection>();
+
+	FilterByName = MakeShared<FTraceTextFilter>(FTraceTextFilter::FItemToStringArray::CreateSP(this, &STraceStoreWindow::HandleItemToStringArray));
+	Filters->Add(FilterByName);
+
+	FilterByPlatform = MakeShared<FTraceFilterByPlatform>();
+	Filters->Add(FilterByPlatform);
+
+	FilterByAppName = MakeShared<FTraceFilterByAppName>();
+	Filters->Add(FilterByAppName);
+
+	FilterByBuildConfig = MakeShared<FTraceFilterByBuildConfig>();
+	Filters->Add(FilterByBuildConfig);
+
+	FilterByBuildTarget = MakeShared<FTraceFilterByBuildTarget>();
+	Filters->Add(FilterByBuildTarget);
+}
+
+////////////////////////////////////////////////////////////////////////////////////////////////////
+
+void STraceStoreWindow::HandleItemToStringArray(const FTraceViewModel& InTrace, TArray<FString>& OutSearchStrings) const
+{
+	OutSearchStrings.Add(InTrace.Name.ToString());
+}
+
+////////////////////////////////////////////////////////////////////////////////////////////////////
+
+void STraceStoreWindow::UpdateFiltering()
+{
+	FilteredTraceViewModels.Reset();
+
+	if (FilterByName->GetRawFilterText().IsEmpty() &&
+		FilterByPlatform->IsEmpty() &&
+		FilterByAppName->IsEmpty() &&
+		FilterByBuildConfig->IsEmpty() &&
+		FilterByBuildTarget->IsEmpty())
+	{
+		// No filtering.
+		FilteredTraceViewModels = TraceViewModels;
+	}
+	else
+	{
+		for (const TSharedPtr<FTraceViewModel>& Trace : TraceViewModels)
+		{
+			const bool bIsTraceVisible = Filters->PassesAllFilters(*Trace.Get());
+			if (bIsTraceVisible)
+			{
+				FilteredTraceViewModels.Add(Trace);
+			}
+		}
+	}
+
+	UpdateFilterStatsText();
+}
+
+////////////////////////////////////////////////////////////////////////////////////////////////////
+
+void STraceStoreWindow::UpdateFilterStatsText()
+{
+	bFilterStatsTextIsDirty = false;
+
+	uint64 FilterdTotalSize = 0;
+	for (const TSharedPtr<FTraceViewModel>& Trace : FilteredTraceViewModels)
+	{
+		FilterdTotalSize += Trace->Size;
+		if (Trace->bIsLive)
+		{
+			bFilterStatsTextIsDirty = true;
+		}
+	}
+
+	// When having live sessions, but too many traces, do not further update the stats text on Tick().
+	if (FilteredTraceViewModels.Num() > 1000)
+	{
+		bFilterStatsTextIsDirty = false;
+	}
+
+	FNumberFormattingOptions FormattingOptionsSize;
+	FormattingOptionsSize.MaximumFractionalDigits = 1;
+
+	if (FilteredTraceViewModels.Num() == TraceViewModels.Num())
+	{
+		FilterStatsText = FText::Format(LOCTEXT("FilterStatsText_Fmt1", "{0} trace sessions ({1})"),
+			FText::AsNumber(TraceViewModels.Num()),
+			FText::AsMemory(FilterdTotalSize, &FormattingOptionsSize));
+	}
+	else
+	{
+		FilterStatsText = FText::Format(LOCTEXT("FilterStatsText_Fmt2", "{0} / {1} trace sessions ({2})"),
+			FText::AsNumber(FilteredTraceViewModels.Num()),
+			FText::AsNumber(TraceViewModels.Num()),
+			FText::AsMemory(FilterdTotalSize, &FormattingOptionsSize));
+	}
+}
+
 ////////////////////////////////////////////////////////////////////////////////////////////////////
 
 EColumnSortMode::Type STraceStoreWindow::GetSortModeForColumn(const FName ColumnId) const
@@ -1773,12 +2338,12 @@ void STraceStoreWindow::UpdateSorting()
 	{
 		if (SortMode == EColumnSortMode::Ascending)
 		{
-			TraceViewModels.Sort([](const TSharedPtr<FTraceViewModel>& A, const TSharedPtr<FTraceViewModel>& B)
+			FilteredTraceViewModels.Sort([](const TSharedPtr<FTraceViewModel>& A, const TSharedPtr<FTraceViewModel>& B)
 				{ return A->Timestamp < B->Timestamp; });
 		}
 		else
 		{
-			TraceViewModels.Sort([](const TSharedPtr<FTraceViewModel>& A, const TSharedPtr<FTraceViewModel>& B)
+			FilteredTraceViewModels.Sort([](const TSharedPtr<FTraceViewModel>& A, const TSharedPtr<FTraceViewModel>& B)
 				{ return A->Timestamp > B->Timestamp; });
 		}
 	}
@@ -1786,12 +2351,12 @@ void STraceStoreWindow::UpdateSorting()
 	{
 		if (SortMode == EColumnSortMode::Ascending)
 		{
-			TraceViewModels.Sort([](const TSharedPtr<FTraceViewModel>& A, const TSharedPtr<FTraceViewModel>& B)
+			FilteredTraceViewModels.Sort([](const TSharedPtr<FTraceViewModel>& A, const TSharedPtr<FTraceViewModel>& B)
 				{ return A->Name.CompareTo(B->Name) < 0; });
 		}
 		else
 		{
-			TraceViewModels.Sort([](const TSharedPtr<FTraceViewModel>& A, const TSharedPtr<FTraceViewModel>& B)
+			FilteredTraceViewModels.Sort([](const TSharedPtr<FTraceViewModel>& A, const TSharedPtr<FTraceViewModel>& B)
 				{ return B->Name.CompareTo(A->Name) < 0; });
 		}
 	}
@@ -1799,12 +2364,12 @@ void STraceStoreWindow::UpdateSorting()
 	{
 		if (SortMode == EColumnSortMode::Ascending)
 		{
-			TraceViewModels.Sort([](const TSharedPtr<FTraceViewModel>& A, const TSharedPtr<FTraceViewModel>& B)
+			FilteredTraceViewModels.Sort([](const TSharedPtr<FTraceViewModel>& A, const TSharedPtr<FTraceViewModel>& B)
 				{ return A->Uri.CompareTo(B->Uri) < 0; });
 		}
 		else
 		{
-			TraceViewModels.Sort([](const TSharedPtr<FTraceViewModel>& A, const TSharedPtr<FTraceViewModel>& B)
+			FilteredTraceViewModels.Sort([](const TSharedPtr<FTraceViewModel>& A, const TSharedPtr<FTraceViewModel>& B)
 				{ return B->Uri.CompareTo(A->Uri) < 0; });
 		}
 	}
@@ -1812,7 +2377,7 @@ void STraceStoreWindow::UpdateSorting()
 	{
 		if (SortMode == EColumnSortMode::Ascending)
 		{
-			TraceViewModels.Sort([](const TSharedPtr<FTraceViewModel>& A, const TSharedPtr<FTraceViewModel>& B)
+			FilteredTraceViewModels.Sort([](const TSharedPtr<FTraceViewModel>& A, const TSharedPtr<FTraceViewModel>& B)
 				{
 					int32 CompareResult = A->Platform.CompareTo(B->Platform);
 					return CompareResult == 0 ? A->Timestamp < B->Timestamp : CompareResult < 0;
@@ -1820,7 +2385,7 @@ void STraceStoreWindow::UpdateSorting()
 		}
 		else
 		{
-			TraceViewModels.Sort([](const TSharedPtr<FTraceViewModel>& A, const TSharedPtr<FTraceViewModel>& B)
+			FilteredTraceViewModels.Sort([](const TSharedPtr<FTraceViewModel>& A, const TSharedPtr<FTraceViewModel>& B)
 				{
 					int32 CompareResult = B->Platform.CompareTo(A->Platform);
 					return CompareResult == 0 ? A->Timestamp < B->Timestamp : CompareResult < 0;
@@ -1831,7 +2396,7 @@ void STraceStoreWindow::UpdateSorting()
 	{
 		if (SortMode == EColumnSortMode::Ascending)
 		{
-			TraceViewModels.Sort([](const TSharedPtr<FTraceViewModel>& A, const TSharedPtr<FTraceViewModel>& B)
+			FilteredTraceViewModels.Sort([](const TSharedPtr<FTraceViewModel>& A, const TSharedPtr<FTraceViewModel>& B)
 				{
 					int32 CompareResult = A->AppName.CompareTo(B->AppName);
 					return CompareResult == 0 ? A->Timestamp < B->Timestamp : CompareResult < 0;
@@ -1839,8 +2404,8 @@ void STraceStoreWindow::UpdateSorting()
 		}
 		else
 		{
-			TraceViewModels.Sort([](const TSharedPtr<FTraceViewModel>& A, const TSharedPtr<FTraceViewModel>& B)
-				{ 
+			FilteredTraceViewModels.Sort([](const TSharedPtr<FTraceViewModel>& A, const TSharedPtr<FTraceViewModel>& B)
+				{
 					int32 CompareResult = B->AppName.CompareTo(A->AppName);
 					return CompareResult == 0 ? A->Timestamp < B->Timestamp : CompareResult < 0;
 				});
@@ -1850,7 +2415,7 @@ void STraceStoreWindow::UpdateSorting()
 	{
 		if (SortMode == EColumnSortMode::Ascending)
 		{
-			TraceViewModels.Sort([](const TSharedPtr<FTraceViewModel>& A, const TSharedPtr<FTraceViewModel>& B)
+			FilteredTraceViewModels.Sort([](const TSharedPtr<FTraceViewModel>& A, const TSharedPtr<FTraceViewModel>& B)
 				{
 					return A->ConfigurationType == B->ConfigurationType ?
 						A->Timestamp < B->Timestamp :
@@ -1859,7 +2424,7 @@ void STraceStoreWindow::UpdateSorting()
 		}
 		else
 		{
-			TraceViewModels.Sort([](const TSharedPtr<FTraceViewModel>& A, const TSharedPtr<FTraceViewModel>& B)
+			FilteredTraceViewModels.Sort([](const TSharedPtr<FTraceViewModel>& A, const TSharedPtr<FTraceViewModel>& B)
 				{
 					return A->ConfigurationType == B->ConfigurationType ?
 						A->Timestamp < B->Timestamp :
@@ -1871,7 +2436,7 @@ void STraceStoreWindow::UpdateSorting()
 	{
 		if (SortMode == EColumnSortMode::Ascending)
 		{
-			TraceViewModels.Sort([](const TSharedPtr<FTraceViewModel>& A, const TSharedPtr<FTraceViewModel>& B)
+			FilteredTraceViewModels.Sort([](const TSharedPtr<FTraceViewModel>& A, const TSharedPtr<FTraceViewModel>& B)
 				{
 					return A->TargetType == B->TargetType ?
 						A->Timestamp < B->Timestamp :
@@ -1880,7 +2445,7 @@ void STraceStoreWindow::UpdateSorting()
 		}
 		else
 		{
-			TraceViewModels.Sort([](const TSharedPtr<FTraceViewModel>& A, const TSharedPtr<FTraceViewModel>& B)
+			FilteredTraceViewModels.Sort([](const TSharedPtr<FTraceViewModel>& A, const TSharedPtr<FTraceViewModel>& B)
 				{
 					return A->TargetType == B->TargetType ?
 						A->Timestamp < B->Timestamp :
@@ -1892,12 +2457,12 @@ void STraceStoreWindow::UpdateSorting()
 	{
 		if (SortMode == EColumnSortMode::Ascending)
 		{
-			TraceViewModels.Sort([](const TSharedPtr<FTraceViewModel>& A, const TSharedPtr<FTraceViewModel>& B)
+			FilteredTraceViewModels.Sort([](const TSharedPtr<FTraceViewModel>& A, const TSharedPtr<FTraceViewModel>& B)
 				{ return A->Size < B->Size; });
 		}
 		else
 		{
-			TraceViewModels.Sort([](const TSharedPtr<FTraceViewModel>& A, const TSharedPtr<FTraceViewModel>& B)
+			FilteredTraceViewModels.Sort([](const TSharedPtr<FTraceViewModel>& A, const TSharedPtr<FTraceViewModel>& B)
 				{ return A->Size > B->Size; });
 		}
 	}
@@ -1905,26 +2470,26 @@ void STraceStoreWindow::UpdateSorting()
 	{
 		if (SortMode == EColumnSortMode::Ascending)
 		{
-			TraceViewModels.Sort([](const TSharedPtr<FTraceViewModel>& A, const TSharedPtr<FTraceViewModel>& B)
+			FilteredTraceViewModels.Sort([](const TSharedPtr<FTraceViewModel>& A, const TSharedPtr<FTraceViewModel>& B)
 				{
 					return A->bIsLive == B->bIsLive ?
 						A->Timestamp < B->Timestamp :
-						A->bIsLive;
+						B->bIsLive;
 				});
 		}
 		else
 		{
-			TraceViewModels.Sort([](const TSharedPtr<FTraceViewModel>& A, const TSharedPtr<FTraceViewModel>& B)
+			FilteredTraceViewModels.Sort([](const TSharedPtr<FTraceViewModel>& A, const TSharedPtr<FTraceViewModel>& B)
 				{
 					return A->bIsLive == B->bIsLive ?
 						A->Timestamp < B->Timestamp :
-						B->bIsLive;
+						A->bIsLive;
 				});
 		}
 	}
 	else
 	{
-		Algo::SortBy(TraceViewModels, &FTraceViewModel::Timestamp);
+		Algo::SortBy(FilteredTraceViewModels, &FTraceViewModel::Timestamp);
 	}
 }
 
diff --git a/Engine/Source/Developer/TraceInsights/Private/Insights/Widgets/SStartPageWindow.h b/Engine/Source/Developer/TraceInsights/Private/Insights/Widgets/SStartPageWindow.h
--- a/Engine/Source/Developer/TraceInsights/Private/Insights/Widgets/SStartPageWindow.h
+++ b/Engine/Source/Developer/TraceInsights/Private/Insights/Widgets/SStartPageWindow.h
@@ -3,6 +3,8 @@
 #pragma once
 
 #include "Async/TaskGraphInterfaces.h"
+#include "Misc/FilterCollection.h"
+#include "Misc/TextFilter.h"
 #include "Styling/SlateTypes.h"
 #include "Widgets/DeclarativeSyntaxSupport.h"
 #include "Widgets/SCompoundWidget.h"
@@ -91,6 +93,125 @@ struct FTraceViewModel
 	}
 };
 
+////////////////////////////////////////////////////////////////////////////////////////////////////
+// Filters
+
+/** The filter collection - used for updating the list of trace sessions. */
+typedef TFilterCollection<const FTraceViewModel&> FTraceViewModelFilterCollection;
+
+/** The text based filter - used for updating the list of trace sessions. */
+typedef TTextFilter<const FTraceViewModel&> FTraceTextFilter;
+
+template<typename TSetType>
+class TTraceSetFilter : public IFilter<const FTraceViewModel&>, public TSharedFromThis< TTraceSetFilter<TSetType> >
+{
+public:
+	TTraceSetFilter();
+	virtual ~TTraceSetFilter() {};
+
+	typedef const FTraceViewModel& ItemType;
+
+	/** Broadcasts anytime the restrictions of the Filter changes. */
+	DECLARE_DERIVED_EVENT(TTraceSetFilter, IFilter<ItemType>::FChangedEvent, FChangedEvent);
+	virtual FChangedEvent& OnChanged() override { return ChangedEvent; }
+
+	/** Returns whether the specified Trace passes the Filter's restrictions. */
+	virtual bool PassesFilter(const FTraceViewModel& InTrace) const override
+	{
+		return FilterSet.IsEmpty() || !FilterSet.Contains(GetFilterValueForTrace(InTrace));
+	}
+
+	bool IsEmpty() const { return FilterSet.IsEmpty(); }
+
+	void Reset() { FilterSet.Reset(); }
+
+	virtual void BuildMenu(FMenuBuilder& InMenuBuilder, class STraceStoreWindow& InWindow);
+
+protected:
+	virtual TSetType GetFilterValueForTrace(const FTraceViewModel& InTrace) const = 0;
+	virtual FText ValueToText(const TSetType Value) const = 0;
+
+protected:
+	/**	The event that fires whenever new search terms are provided */
+	FChangedEvent ChangedEvent;
+
+	/** The set of values used to filter */
+	TSet<TSetType> FilterSet;
+
+	FText ToggleAllActionLabel;
+	FText ToggleAllActionTooltip;
+	FText UndefinedValueLabel;
+};
+
+class FTraceFilterByStringSet : public TTraceSetFilter<FString>
+{
+protected:
+	virtual FText ValueToText(const FString InValue) const override
+	{
+		return InValue.IsEmpty() ? UndefinedValueLabel : FText::FromString(InValue);
+	}
+};
+
+class FTraceFilterByPlatform : public FTraceFilterByStringSet
+{
+public:
+	FTraceFilterByPlatform();
+
+protected:
+	virtual FString GetFilterValueForTrace(const FTraceViewModel& InTrace) const override
+	{
+		return InTrace.Platform.ToString();
+	}
+};
+
+class FTraceFilterByAppName : public FTraceFilterByStringSet
+{
+public:
+	FTraceFilterByAppName();
+
+protected:
+	virtual FString GetFilterValueForTrace(const FTraceViewModel& InTrace) const override
+	{
+		return InTrace.AppName.ToString();
+	}
+};
+
+class FTraceFilterByBuildConfig : public TTraceSetFilter<uint8>
+{
+public:
+	FTraceFilterByBuildConfig();
+
+protected:
+	virtual uint8 GetFilterValueForTrace(const FTraceViewModel& InTrace) const override
+	{
+		return (uint8)InTrace.ConfigurationType;
+	}
+
+	virtual FText ValueToText(const uint8 InValue) const override
+	{
+		const TCHAR* Str = LexToString((EBuildConfiguration)InValue);
+		return Str ? FText::FromString(Str) : UndefinedValueLabel;
+	}
+};
+
+class FTraceFilterByBuildTarget : public TTraceSetFilter<uint8>
+{
+public:
+	FTraceFilterByBuildTarget();
+
+protected:
+	virtual uint8 GetFilterValueForTrace(const FTraceViewModel& InTrace) const override
+	{
+		return (uint8)InTrace.TargetType;
+	}
+
+	virtual FText ValueToText(const uint8 InValue) const override
+	{
+		const TCHAR* Str = LexToString((EBuildTargetType)InValue);
+		return Str ? FText::FromString(Str) : UndefinedValueLabel;
+	}
+};
+
 ////////////////////////////////////////////////////////////////////////////////////////////////////
 
 /** Implements the Trace Store window. */
@@ -123,7 +244,11 @@ public:
 	void SetStartProcessWithStompMalloc(bool InValue) { bStartProcessWithStompMalloc = InValue; };
 	bool GetStartProcessWithStompMalloc() const { return bStartProcessWithStompMalloc; };
 
+	void OnFilterChanged();
+	const TArray<TSharedPtr<FTraceViewModel>>& GetAllAvailableTraces() const;
+
 private:
+	TSharedRef<SWidget> ConstructFiltersToolbar();
 	TSharedRef<SWidget> ConstructSessionsPanel();
 	TSharedRef<SWidget> ConstructLoadPanel();
 	TSharedRef<SWidget> ConstructTraceStoreDirectoryPanel();
@@ -160,6 +285,22 @@ private:
 
 	TSharedRef<SWidget> MakeTraceListMenu();
 
+	TSharedRef<SWidget> MakePlatformColumnHeaderMenu();
+	TSharedRef<SWidget> MakePlatformFilterMenu();
+	void BuildPlatformFilterSubMenu(FMenuBuilder& InMenuBuilder);
+
+	TSharedRef<SWidget> MakeAppNameColumnHeaderMenu();
+	TSharedRef<SWidget> MakeAppNameFilterMenu();
+	void BuildAppNameFilterSubMenu(FMenuBuilder& InMenuBuilder);
+
+	TSharedRef<SWidget> MakeBuildConfigColumnHeaderMenu();
+	TSharedRef<SWidget> MakeBuildConfigFilterMenu();
+	void BuildBuildConfigFilterSubMenu(FMenuBuilder& InMenuBuilder);
+
+	TSharedRef<SWidget> MakeBuildTargetColumnHeaderMenu();
+	TSharedRef<SWidget> MakeBuildTargetFilterMenu();
+	void BuildBuildTargetFilterSubMenu(FMenuBuilder& InMenuBuilder);
+
 	FReply RefreshTraces_OnClicked();
 
 	void RefreshTraceList();
@@ -240,10 +381,38 @@ private:
 	 */
 	virtual FReply OnDragOver(const FGeometry& MyGeometry, const FDragDropEvent& DragDropEvent)  override;
 
+	//////////////////////////////////////////////////
+	// Filtering
+
+	void FilterByNameSearchBox_OnTextChanged(const FText& InFilterText);
+
+	FText GetFilterStatsText() const { return FilterStatsText; }
+
+	void CreateFilters();
+
+	/**
+	 * Populates OutSearchStrings with the strings that should be used in searching.
+	 *
+	 * @param GroupOrStatNodePtr - the group and stat node to get a text description from.
+	 * @param OutSearchStrings   - an array of strings to use in searching.
+	 *
+	 */
+	void HandleItemToStringArray(const FTraceViewModel& InTrace, TArray<FString>& OutSearchStrings) const;
+
+	void UpdateFiltering();
+
+	void UpdateFilterStatsText();
+
+	//////////////////////////////////////////////////
+	// Sorting
+
 	EColumnSortMode::Type GetSortModeForColumn(const FName ColumnId) const;
 	void OnSortModeChanged(const EColumnSortPriority::Type SortPriority, const FName& ColumnId, const EColumnSortMode::Type InSortMode);
 
 	void UpdateSorting();
+
+	//////////////////////////////////////////////////
+
 	void UpdateTraceListView();
 
 private:
@@ -264,28 +433,54 @@ private:
 
 	TSharedPtr<SVerticalBox> MainContentPanel;
 
-	int32 LiveSessionCount;
-
-	bool bAutoStartAnalysisForLiveSessions;
-	TArray<uint32> AutoStartedSessions; // tracks sessions that were auto started (in order to not start them again)
-
-	TSharedPtr<SSearchBox> AutoStartPlatformFilter;
-	TSharedPtr<SSearchBox> AutoStartAppNameFilter;
-	EBuildConfiguration AutoStartConfigurationTypeFilter;
-	EBuildTargetType AutoStartTargetTypeFilter;
+	//////////////////////////////////////////////////
 
 	TUniquePtr<Insights::FStoreBrowser> StoreBrowser;
 	uint64 TracesChangeSerial;
 
-	TArray<TSharedPtr<FTraceViewModel>> TraceViewModels;
+	TArray<TSharedPtr<FTraceViewModel>> TraceViewModels; // all available trace view models
+	TArray<TSharedPtr<FTraceViewModel>> FilteredTraceViewModels; // the filtered list of trace view models
 	TMap<uint32, TSharedPtr<FTraceViewModel>> TraceViewModelMap;
 
 	TSharedPtr<SListView<TSharedPtr<FTraceViewModel>>> TraceListView;
 	TSharedPtr<FTraceViewModel> SelectedTrace;
+	bool bIsUserSelectedTrace;
+
+	//////////////////////////////////////////////////
+	// Filtering
+
+	TSharedPtr<FTraceViewModelFilterCollection> Filters;
+
+	TSharedPtr<SSearchBox> FilterByNameSearchBox;
+	TSharedPtr<FTraceTextFilter> FilterByName;
+
+	TSharedPtr<FTraceFilterByPlatform> FilterByPlatform;
+	TSharedPtr<FTraceFilterByAppName> FilterByAppName;
+	TSharedPtr<FTraceFilterByBuildConfig> FilterByBuildConfig;
+	TSharedPtr<FTraceFilterByBuildTarget> FilterByBuildTarget;
+
+	bool bFilterStatsTextIsDirty;
+	FText FilterStatsText;
+
+	//////////////////////////////////////////////////
+	// Sorting
 
 	FName SortColumn;
 	EColumnSortMode::Type SortMode;
 
+	//////////////////////////////////////////////////
+	// Auto-start functionality
+
+	bool bAutoStartAnalysisForLiveSessions;
+	TArray<uint32> AutoStartedSessions; // tracks sessions that were auto started (in order to not start them again)
+
+	TSharedPtr<SSearchBox> AutoStartPlatformFilter;
+	TSharedPtr<SSearchBox> AutoStartAppNameFilter;
+	EBuildConfiguration AutoStartConfigurationTypeFilter;
+	EBuildTargetType AutoStartTargetTypeFilter;
+
+	//////////////////////////////////////////////////
+
 	FString SplashScreenOverlayTraceFile;
 	float SplashScreenOverlayFadeTime;
 
