[Insights] Added "Insights.Enumerate10K" and "Insights.Enumerate100K" tests (that can be run from UnrealInsights on any utrace files), similar with to ones used by "Insights.Enumerate" and "Insights.EnumeratePerformance" from EngineTest (that uses specific utrace files).

diff --git a/Engine/Source/Developer/TraceInsights/Private/Insights/Tests/TimingProfilerTests.cpp b/Engine/Source/Developer/TraceInsights/Private/Insights/Tests/TimingProfilerTests.cpp
--- a/Engine/Source/Developer/TraceInsights/Private/Insights/Tests/TimingProfilerTests.cpp
+++ b/Engine/Source/Developer/TraceInsights/Private/Insights/Tests/TimingProfilerTests.cpp
@@ -101,6 +101,36 @@ bool EnumerateScopesToFile::RunTest(const FString& Parameters)
 
 ////////////////////////////////////////////////////////////////////////////////////////////////////
 
+IMPLEMENT_SIMPLE_AUTOMATION_TEST(InsightsEnumerate10K, "Insights.Enumerate10K", EAutomationTestFlags::ApplicationContextMask | EAutomationTestFlags::ProductFilter)
+bool InsightsEnumerate10K::RunTest(const FString& Parameters)
+{
+	FTimingProfilerTests::FEnumerateTestParams Params;
+	Params.Interval = 0.01;
+	Params.NumEnumerations = 10000;
+
+	FTimingProfilerTests::FCheckValues CheckValues;
+	FTimingProfilerTests::RunEnumerateBenchmark(Params, CheckValues);
+
+	return !HasAnyErrors();
+}
+
+////////////////////////////////////////////////////////////////////////////////////////////////////
+
+IMPLEMENT_SIMPLE_AUTOMATION_TEST(InsightsEnumerate100K, "Insights.Enumerate100K", EAutomationTestFlags::ApplicationContextMask | EAutomationTestFlags::ProductFilter)
+bool InsightsEnumerate100K::RunTest(const FString& Parameters)
+{
+	FTimingProfilerTests::FEnumerateTestParams Params;
+	Params.Interval = 0.01;
+	Params.NumEnumerations = 100000;
+
+	FTimingProfilerTests::FCheckValues CheckValues;
+	FTimingProfilerTests::RunEnumerateBenchmark(Params, CheckValues);
+
+	return !HasAnyErrors();
+}
+
+////////////////////////////////////////////////////////////////////////////////////////////////////
+
 IMPLEMENT_SIMPLE_AUTOMATION_TEST(EnumerateByEndTimeAsyncAllTracks, "Insights.EnumerateByEndTimeAsyncAllTracks", EAutomationTestFlags::ApplicationContextMask | EAutomationTestFlags::ProductFilter)
 bool EnumerateByEndTimeAsyncAllTracks::RunTest(const FString& Parameters)
 {
@@ -222,7 +252,6 @@ bool EnumerateByStartTimeAllTracks5sIntervals::RunTest(const FString& Parameters
 }
 
 #endif // #if !WITH_EDITOR && WITH_DEV_AUTOMATION_TESTS
-
 ////////////////////////////////////////////////////////////////////////////////////////////////////
 
 void FTimingProfilerTests::RunEnumerateBenchmark(const FEnumerateTestParams& InParams, FCheckValues& OutCheckValues)
@@ -279,6 +308,7 @@ void FTimingProfilerTests::RunEnumerateBenchmark(const FEnumerateTestParams& InP
 }
 
 ////////////////////////////////////////////////////////////////////////////////////////////////////
+
 void FTimingProfilerTests::RunEnumerateAsyncBenchmark(const FEnumerateTestParams& InParams, FCheckValues& OutCheckValues)
 {
 	UE_LOG(TimingProfilerTests, Log, TEXT("RUNNING ASYNC ENUMERATE BENCHMARK..."));
@@ -551,4 +581,3 @@ void FTimingProfilerTests::VerifyCheckValues(FAutomationTestBase& Test, FCheckVa
 }
 
 ////////////////////////////////////////////////////////////////////////////////////////////////////
-
