[Insights]
  - Added Exec() API to IUnrealInsightsModule and to IInsightsComponent (i.e. each component can now expose custom commands).
  - Extended "-ExecOnAnalysisCompleteCmd" command line argument to be used also without AutomationTesting, to run custom commands exposed by Insights components.
  - Added custom response file support for "-ExecOnAnalysisCompleteCmd" command line argument (using @=file syntax).
  - Fixed -AutoQuit command line argument to work also without AutomationTesting.

diff --git a/Engine/Source/Developer/TraceInsights/Private/Insights/InsightsManager.cpp b/Engine/Source/Developer/TraceInsights/Private/Insights/InsightsManager.cpp
--- a/Engine/Source/Developer/TraceInsights/Private/Insights/InsightsManager.cpp
+++ b/Engine/Source/Developer/TraceInsights/Private/Insights/InsightsManager.cpp
@@ -7,6 +7,7 @@
 #include "Logging/MessageLog.h"
 #include "MessageLog/Public/MessageLogModule.h"
 #include "Misc/CString.h"
+#include "Misc/FileHelper.h"
 #include "Misc/Paths.h"
 #include "Modules/ModuleManager.h"
 #include "Templates/UniquePtr.h"
@@ -102,6 +103,7 @@ void FAvailabilityCheck::Enable(double InWaitTime)
 // FInsightsManager
 ////////////////////////////////////////////////////////////////////////////////////////////////////
 
+const TCHAR* FInsightsManager::AutoQuitMsg = TEXT("Application is closing because it was started with the AutoQuit parameter and session analysis is complete.");
 const TCHAR* FInsightsManager::AutoQuitMsgOnFail = TEXT("Application is closing because it was started with the AutoQuit parameter and session analysis failed to start.");
 
 TSharedPtr<FInsightsManager> FInsightsManager::Instance = nullptr;
@@ -133,23 +135,11 @@ TSharedPtr<FInsightsManager> FInsightsManager::CreateInstance(TSharedRef<TraceSe
 
 FInsightsManager::FInsightsManager(TSharedRef<TraceServices::IAnalysisService> InTraceAnalysisService,
 								   TSharedRef<TraceServices::IModuleService> InTraceModuleService)
-	: bIsInitialized(false)
-	, bMemUsageLimitHysteresis(false)
-	, MemUsageLimitLastTimestamp(0)
-	, LogListingName(TEXT("UnrealInsights"))
+	: LogListingName(TEXT("UnrealInsights"))
 	, AnalysisService(InTraceAnalysisService)
 	, ModuleService(InTraceModuleService)
-	, StoreDir()
-	, StoreClient()
 	, CommandList(new FUICommandList())
 	, ActionManager(this)
-	, Settings()
-	, bIsDebugInfoEnabled(false)
-	, AnalysisStopwatch()
-	, bIsAnalysisComplete(false)
-	, SessionDuration(0.0)
-	, AnalysisDuration(0.0)
-	, AnalysisSpeedFactor(0.0)
 {
 }
 
@@ -546,6 +536,8 @@ void FInsightsManager::UpdateSessionDuration()
 				*TimeUtils::FormatTimeAuto(AnalysisDuration, 2),
 				AnalysisSpeedFactor,
 				*TimeUtils::FormatTimeAuto(SessionDuration, 2));
+
+			OnSessionAnalysisCompleted();
 		}
 	}
 }
@@ -764,6 +756,7 @@ void FInsightsManager::LoadTrace(uint32 InTraceId, bool InAutoQuit)
 		CurrentTraceId = InTraceId;
 		CurrentTraceFilename = TraceName;
 		OnSessionChanged();
+		bSessionAnalysisCompletedAutoQuit = InAutoQuit;
 	}
 	else if (InAutoQuit)
 	{
@@ -791,6 +784,7 @@ void FInsightsManager::LoadTraceFile(const FString& InTraceFilename, bool InAuto
 		CurrentTraceId = 0;
 		CurrentTraceFilename = InTraceFilename;
 		OnSessionChanged();
+		bSessionAnalysisCompletedAutoQuit = InAutoQuit;
 	}
 	else if (InAutoQuit)
 	{
@@ -885,4 +879,117 @@ void FInsightsManager::OpenSettings()
 
 ////////////////////////////////////////////////////////////////////////////////////////////////////
 
+void FInsightsManager::ScheduleCommand(const FString& InCmd)
+{
+	SessionAnalysisCompletedCmd = InCmd;
+}
+
+////////////////////////////////////////////////////////////////////////////////////////////////////
+
+void FInsightsManager::OnSessionAnalysisCompleted()
+{
+	if (!SessionAnalysisCompletedCmd.IsEmpty())
+	{
+		FOutputDevice& Ar = *GLog;
+		Ar.Logf(TEXT("Executing commands on analysis completed..."));
+		FStopwatch Stopwatch;
+		Stopwatch.Start();
+		IUnrealInsightsModule& TraceInsightsModule = FModuleManager::LoadModuleChecked<IUnrealInsightsModule>("TraceInsights");
+		TraceInsightsModule.Exec(*SessionAnalysisCompletedCmd, Ar);
+		Stopwatch.Stop();
+		Ar.Logf(TEXT("Commands executed in %.3fs."), Stopwatch.GetAccumulatedTime());
+	}
+
+#if !UE_BUILD_SHIPPING && !WITH_EDITOR
+	if (FInsightsTestRunner::Get().IsValid())
+	{
+		// Don't quit now. Let the test runner to execute.
+		bSessionAnalysisCompletedAutoQuit = false;
+	}
+#endif
+
+	if (bSessionAnalysisCompletedAutoQuit)
+	{
+		bSessionAnalysisCompletedAutoQuit = false;
+		RequestEngineExit(AutoQuitMsg);
+	}
+}
+
+////////////////////////////////////////////////////////////////////////////////////////////////////
+
+bool FInsightsManager::Exec(const TCHAR* Cmd, FOutputDevice& Ar)
+{
+	FString ResponseFile;
+	if (FParse::Value(Cmd, TEXT("@="), ResponseFile))
+	{
+		HandleResponseFileCmd(*ResponseFile, Ar);
+		return true;
+	}
+
+	return false;
+}
+
+////////////////////////////////////////////////////////////////////////////////////////////////////
+
+bool FInsightsManager::HandleResponseFileCmd(const TCHAR* ResponseFile, FOutputDevice& Ar)
+{
+	Ar.Logf(TEXT("Executing commands using response file (%s)..."), ResponseFile);
+
+	FString Contents;
+	if (!FFileHelper::LoadFileToString(Contents, &IPlatformFile::GetPlatformPhysical(), ResponseFile))
+	{
+		Ar.Logf(ELogVerbosity::Error, TEXT("Failed to open the response file (%s)."), ResponseFile);
+		return false;
+	}
+
+	if (Contents.IsEmpty())
+	{
+		return true;
+	}
+
+	IUnrealInsightsModule& TraceInsightsModule = FModuleManager::LoadModuleChecked<IUnrealInsightsModule>("TraceInsights");
+
+	TCHAR* StartPos = &Contents[0];
+	TCHAR* EndPos = StartPos + Contents.Len();
+	TCHAR* CrtPos = StartPos;
+	while (CrtPos != EndPos)
+	{
+		uint32 EndOfLine = 0;
+		if (*CrtPos == TEXT('\r'))
+		{
+			if (*(CrtPos + 1) == '\n')
+			{
+				EndOfLine = 2;
+			}
+			else
+			{
+				EndOfLine = 1;
+			}
+		}
+		else if (*CrtPos == TEXT('\n'))
+		{
+			EndOfLine = 1;
+		}
+
+		if (EndOfLine > 0)
+		{
+			const uint32 LineLen = uint32(CrtPos - StartPos);
+			StartPos[LineLen] = TEXT('\0');
+
+			TraceInsightsModule.Exec(StartPos, Ar);
+
+			CrtPos += EndOfLine;
+			StartPos = CrtPos;
+		}
+		else
+		{
+			++CrtPos;
+		}
+	}
+
+	return true;
+}
+
+////////////////////////////////////////////////////////////////////////////////////////////////////
+
 #undef LOCTEXT_NAMESPACE
diff --git a/Engine/Source/Developer/TraceInsights/Private/Insights/InsightsManager.h b/Engine/Source/Developer/TraceInsights/Private/Insights/InsightsManager.h
--- a/Engine/Source/Developer/TraceInsights/Private/Insights/InsightsManager.h
+++ b/Engine/Source/Developer/TraceInsights/Private/Insights/InsightsManager.h
@@ -94,6 +94,7 @@ public:
 	virtual void Shutdown() override;
 	virtual void RegisterMajorTabs(IUnrealInsightsModule& InsightsModule) override;
 	virtual void UnregisterMajorTabs() override;
+	virtual bool Exec(const TCHAR* Cmd, FOutputDevice& Ar) override;
 
 	//////////////////////////////////////////////////
 
@@ -261,6 +262,8 @@ public:
 
 	const FName& GetLogListingName() const { return LogListingName; }
 
+	void ScheduleCommand(const FString& InCmd);
+
 	////////////////////////////////////////////////////////////////////////////////////////////////////
 	// SessionChangedEvent
 
@@ -320,19 +323,22 @@ private:
 	void ResetSession(bool bNotify = true);
 
 	void OnSessionChanged();
+	void OnSessionAnalysisCompleted();
 
 	void SpawnAndActivateTabs();
 
 	void ActivateTimingInsightsTab();
 
+	bool HandleResponseFileCmd(const TCHAR* ResponseFile, FOutputDevice& Ar);
+
 private:
-	bool bIsInitialized;
+	bool bIsInitialized = false;
 
 	/** If true, the "high system memory usage warning" will be disabled until the system memory usage first drops below a certain threshold. */
-	bool bMemUsageLimitHysteresis;
+	bool bMemUsageLimitHysteresis = false;
 
 	/** The timestamp when has occurred the last check for system memory usage. */
-	uint64 MemUsageLimitLastTimestamp;
+	uint64 MemUsageLimitLastTimestamp = 0;
 
 	/** The name of the Unreal Insights log listing. */
 	FName LogListingName;
@@ -356,7 +362,7 @@ private:
 	TSharedPtr<const TraceServices::IAnalysisSession> Session;
 
 	/** The id of the trace being analyzed. */
-	uint32 CurrentTraceId;
+	uint32 CurrentTraceId = 0;
 
 	/** The filename of the trace being analyzed. */
 	FString CurrentTraceFilename;
@@ -383,20 +389,23 @@ private:
 	TWeakPtr<class SSessionInfoWindow> SessionInfoWindow;
 
 	/** If enabled, UI can display additional info for debugging purposes. */
-	bool bIsDebugInfoEnabled;
-
-	FStopwatch AnalysisStopwatch;
-	bool bIsAnalysisComplete;
-	double SessionDuration;
-	double AnalysisDuration;
-	double AnalysisSpeedFactor;
+	bool bIsDebugInfoEnabled = false;
 
 	bool bIsMainTabSet = false;
+	bool bIsAnalysisComplete = false;
+	bool bSessionAnalysisCompletedAutoQuit = false;
+
+	FStopwatch AnalysisStopwatch;
+	double SessionDuration = 0.0;
+	double AnalysisDuration = 0.0;
+	double AnalysisSpeedFactor = 0.0;
 
 	TSharedPtr<FInsightsMenuBuilder> InsightsMenuBuilder;
 	TSharedPtr<FInsightsTestRunner> TestRunner;
 
-private:
+	FString SessionAnalysisCompletedCmd;
+
+	static const TCHAR* AutoQuitMsg;
 	static const TCHAR* AutoQuitMsgOnFail;
 
 	/** A shared pointer to the global instance of the main manager. */
diff --git a/Engine/Source/Developer/TraceInsights/Private/Insights/TraceInsightsModule.cpp b/Engine/Source/Developer/TraceInsights/Private/Insights/TraceInsightsModule.cpp
--- a/Engine/Source/Developer/TraceInsights/Private/Insights/TraceInsightsModule.cpp
+++ b/Engine/Source/Developer/TraceInsights/Private/Insights/TraceInsightsModule.cpp
@@ -328,7 +328,7 @@ void FTraceInsightsModule::CreateSessionViewer(bool bAllowDebugTools)
 	FDisplayMetrics DisplayMetrics;
 	FSlateApplication::Get().GetDisplayMetrics(DisplayMetrics);
 	const float DPIScaleFactor = FPlatformApplicationMisc::GetDPIScaleFactorAtPoint(DisplayMetrics.PrimaryDisplayWorkAreaRect.Left, DisplayMetrics.PrimaryDisplayWorkAreaRect.Top);
-	
+
 	const FVector2D ClientSize(1280.f * DPIScaleFactor, 720.0f * DPIScaleFactor);
 
 	TSharedRef<SWindow> RootWindow = SNew(SWindow)
@@ -377,7 +377,7 @@ void FTraceInsightsModule::CreateSessionViewer(bool bAllowDebugTools)
 	const bool bForceWindowToFront = true;
 	RootWindow->BringToFront(bForceWindowToFront);
 
-#endif
+#endif // !WITH_EDITOR
 }
 
 ////////////////////////////////////////////////////////////////////////////////////////////////////
@@ -556,25 +556,50 @@ void FTraceInsightsModule::OnWindowClosedEvent(const TSharedRef<SWindow>&)
 
 ////////////////////////////////////////////////////////////////////////////////////////////////////
 
-void FTraceInsightsModule::ScheduleCommand(const FString& InCmd)
+void FTraceInsightsModule::InitializeTesting(bool InInitAutomationModules, bool InAutoQuit)
 {
 #if !UE_BUILD_SHIPPING && !WITH_EDITOR
-	FInsightsTestRunner::Get()->ScheduleCommand(InCmd);
+	auto TestRunner = FInsightsTestRunner::CreateInstance();
+
+	TestRunner->SetInitAutomationModules(InInitAutomationModules);
+	TestRunner->SetAutoQuit(InAutoQuit);
+
+	RegisterComponent(TestRunner);
 #endif
 }
 
 ////////////////////////////////////////////////////////////////////////////////////////////////////
 
-void FTraceInsightsModule::InitializeTesting(bool InInitAutomationModules, bool InAutoQuit)
+void FTraceInsightsModule::ScheduleCommand(const FString& InCmd)
 {
+	FString ActualCmd = InCmd;
+	ActualCmd.TrimCharInline(TEXT('\"'), nullptr);
+	ActualCmd.TrimCharInline(TEXT('\''), nullptr);
+
 #if !UE_BUILD_SHIPPING && !WITH_EDITOR
-	auto TestRunner = FInsightsTestRunner::CreateInstance();
+	if (ActualCmd.StartsWith(TEXT("Automation RunTests")))
+	{
+		FInsightsTestRunner::Get()->ScheduleCommand(ActualCmd);
+		return;
+	}
+#endif
 
-	TestRunner->SetInitAutomationModules(InInitAutomationModules);
-	TestRunner->SetAutoQuit(InAutoQuit);
+	FInsightsManager::Get()->ScheduleCommand(ActualCmd);
+}
 
-	RegisterComponent(TestRunner);
-#endif
+////////////////////////////////////////////////////////////////////////////////////////////////////
+
+bool FTraceInsightsModule::Exec(const TCHAR* Cmd, FOutputDevice& Ar)
+{
+	for (TSharedRef<IInsightsComponent>& Component : Components)
+	{
+		if (Component->Exec(Cmd, Ar))
+		{
+			return true;
+		}
+	}
+
+	return false;
 }
 
 ////////////////////////////////////////////////////////////////////////////////////////////////////
diff --git a/Engine/Source/Developer/TraceInsights/Private/Insights/TraceInsightsModule.h b/Engine/Source/Developer/TraceInsights/Private/Insights/TraceInsightsModule.h
--- a/Engine/Source/Developer/TraceInsights/Private/Insights/TraceInsightsModule.h
+++ b/Engine/Source/Developer/TraceInsights/Private/Insights/TraceInsightsModule.h
@@ -74,9 +74,9 @@ public:
 	/** Set the ini path for saving persistent layout data. */
 	virtual void SetUnrealInsightsLayoutIni(const FString& InIniPath) override;
 
-	virtual void ScheduleCommand(const FString& InCmd) override;
-
 	virtual void InitializeTesting(bool InInitAutomationModules, bool InAutoQuit) override;
+	virtual void ScheduleCommand(const FString& InCmd) override;
+	virtual bool Exec(const TCHAR* Cmd, FOutputDevice& Ar) override;
 
 protected:
 	void InitTraceStore();
diff --git a/Engine/Source/Developer/TraceInsights/Public/Insights/IUnrealInsightsModule.h b/Engine/Source/Developer/TraceInsights/Public/Insights/IUnrealInsightsModule.h
--- a/Engine/Source/Developer/TraceInsights/Public/Insights/IUnrealInsightsModule.h
+++ b/Engine/Source/Developer/TraceInsights/Public/Insights/IUnrealInsightsModule.h
@@ -279,6 +279,9 @@ public:
 	 * @param InAutoQuit If true Insights will close after completing session analysis and running any tests started using the ScheduleCommand function.
 	*/
 	virtual void InitializeTesting(bool InInitAutomationModules, bool InAutoQuit) = 0;
+
+	/** Execute command. */
+	virtual bool Exec(const TCHAR* Cmd, FOutputDevice& Ar) = 0;
 };
 
 ////////////////////////////////////////////////////////////////////////////////////////////////////
@@ -299,7 +302,10 @@ public:
 	virtual void UnregisterMajorTabs() = 0;
 
 	/* Called by the TraceInsights module when it receives the OnWindowClosedEvent. Can be used to close any panels that should not persist in the layout. */
-	virtual void OnWindowClosedEvent() {};
+	virtual void OnWindowClosedEvent() {}
+
+	/** Execute command. */
+	virtual bool Exec(const TCHAR* Cmd, FOutputDevice& Ar) { return false; }
 };
 
 ////////////////////////////////////////////////////////////////////////////////////////////////////
diff --git a/Engine/Source/Programs/UnrealInsights/Private/UserInterfaceCommand.cpp b/Engine/Source/Programs/UnrealInsights/Private/UserInterfaceCommand.cpp
--- a/Engine/Source/Programs/UnrealInsights/Private/UserInterfaceCommand.cpp
+++ b/Engine/Source/Programs/UnrealInsights/Private/UserInterfaceCommand.cpp
@@ -263,11 +263,10 @@ void FUserInterfaceCommand::InitializeSlateApplication(bool bOpenTraceFile, cons
 
 	IUnrealInsightsModule& TraceInsightsModule = FModuleManager::LoadModuleChecked<IUnrealInsightsModule>("TraceInsights");
 
-	const uint32 MaxPath = FPlatformMisc::GetMaxPathLength();
-
 	uint32 TraceId = 0;
 	bool bUseTraceId = FParse::Value(FCommandLine::Get(), TEXT("-OpenTraceId="), TraceId);
 
+	const uint32 MaxPath = FPlatformMisc::GetMaxPathLength();
 	TCHAR* StoreHost = new TCHAR[MaxPath + 1];
 	FCString::Strcpy(StoreHost, MaxPath, TEXT("127.0.0.1"));
 	uint32 StorePort = 0;
@@ -293,78 +292,62 @@ void FUserInterfaceCommand::InitializeSlateApplication(bool bOpenTraceFile, cons
 		bUseCustomStoreAddress = true;
 	}
 
-	TCHAR Cmd[1024];
-	bool bExecuteCommand = false;
-	if (FParse::Value(FCommandLine::Get(), TEXT("-ExecOnAnalysisCompleteCmd="), Cmd, 1024, false))
-	{
-		bExecuteCommand = true;
-	}
-
-	//This parameter will cause the application to close when analysis fails to start or completes succesfully
-	const bool bAutoQuit = FParse::Param(FCommandLine::Get(), TEXT("AutoQuit"));
-
-	const bool bNoUI = FParse::Param(FCommandLine::Get(), TEXT("NoUI"));
-
 	const bool bInitializeTesting = FParse::Param(FCommandLine::Get(), TEXT("InsightsTest"));
 
-	if (bUseTraceId)
+	if (bUseTraceId || bOpenTraceFile) // viewer mode
 	{
-		if (bInitializeTesting || bAutoQuit)
+		// This parameter will cause the application to close when analysis fails to start or completes successfully.
+		const bool bAutoQuit = FParse::Param(FCommandLine::Get(), TEXT("AutoQuit"));
+
+		if (bInitializeTesting)
 		{
-			TraceInsightsModule.InitializeTesting(bInitializeTesting, bAutoQuit);
+			const bool bInitAutomationModules = true;
+			TraceInsightsModule.InitializeTesting(bInitAutomationModules, bAutoQuit);
+		}
 
-			if (bExecuteCommand)
-			{
-				TraceInsightsModule.ScheduleCommand(Cmd);
-			}
+		TCHAR Cmd[1024];
+		bool bExecuteCommand = false;
+		if (FParse::Value(FCommandLine::Get(), TEXT("-ExecOnAnalysisCompleteCmd="), Cmd, 1024, false))
+		{
+			bExecuteCommand = true;
+		}
+		if (bExecuteCommand)
+		{
+			TraceInsightsModule.ScheduleCommand(Cmd);
 		}
 
+		const bool bNoUI = FParse::Param(FCommandLine::Get(), TEXT("NoUI"));
 		if (!bNoUI)
 		{
 			TraceInsightsModule.CreateSessionViewer(bAllowDebugTools);
 		}
 
-		TraceInsightsModule.ConnectToStore(StoreHost, StorePort);
-		TraceInsightsModule.StartAnalysisForTrace(TraceId, bAutoQuit);
+		if (bUseTraceId)
+		{
+			TraceInsightsModule.ConnectToStore(StoreHost, StorePort);
+			TraceInsightsModule.StartAnalysisForTrace(TraceId, bAutoQuit);
+		}
+		else
+		{
+			TraceInsightsModule.StartAnalysisForTraceFile(TraceFile, bAutoQuit);
+		}
 	}
-	else
+	else // browser mode
 	{
-		if (bOpenTraceFile)
+		if (bUseCustomStoreAddress)
 		{
-			if (bInitializeTesting || bAutoQuit)
-			{
-				TraceInsightsModule.InitializeTesting(bInitializeTesting, bAutoQuit);
-
-				if (bExecuteCommand)
-				{
-					TraceInsightsModule.ScheduleCommand(Cmd);
-				}
-			}
-
-			if (!bNoUI)
-			{
-				TraceInsightsModule.CreateSessionViewer(bAllowDebugTools);
-			}
-
-			TraceInsightsModule.StartAnalysisForTraceFile(TraceFile, bAutoQuit);
+			TraceInsightsModule.ConnectToStore(StoreHost, StorePort);
 		}
 		else
 		{
-			if (!bUseCustomStoreAddress)
-			{
-				TraceInsightsModule.CreateDefaultStore();
-			}
-			else
-			{
-				TraceInsightsModule.ConnectToStore(StoreHost, StorePort);
-			}
-
-			FCreateSessionBrowserParams Params;
-			Params.bAllowDebugTools = bAllowDebugTools;
-			Params.bInitializeTesting = bInitializeTesting;
-			Params.bStartProcessWithStompMalloc = FParse::Param(FCommandLine::Get(), TEXT("stompmalloc"));
-			TraceInsightsModule.CreateSessionBrowser(Params);
+			TraceInsightsModule.CreateDefaultStore();
 		}
+
+		FCreateSessionBrowserParams Params;
+		Params.bAllowDebugTools = bAllowDebugTools;
+		Params.bInitializeTesting = bInitializeTesting;
+		Params.bStartProcessWithStompMalloc = FParse::Param(FCommandLine::Get(), TEXT("stompmalloc"));
+		TraceInsightsModule.CreateSessionBrowser(Params);
 	}
 
 	delete[] StoreHost;
