[Insights] Added Platform, AppName, Build Config and Build Target of trace session being analyzed to the application title name.

diff --git a/Engine/Source/Developer/TraceInsights/Private/Insights/InsightsManager.cpp b/Engine/Source/Developer/TraceInsights/Private/Insights/InsightsManager.cpp
--- a/Engine/Source/Developer/TraceInsights/Private/Insights/InsightsManager.cpp
+++ b/Engine/Source/Developer/TraceInsights/Private/Insights/InsightsManager.cpp
@@ -12,6 +12,7 @@
 #include "Modules/ModuleManager.h"
 #include "Templates/UniquePtr.h"
 #include "Trace/StoreClient.h"
+#include "TraceServices/Model/Diagnostics.h"
 #include "TraceServices/Model/NetProfiler.h"
 #include "WorkspaceMenuStructure.h"
 #include "WorkspaceMenuStructureModule.h"
@@ -494,8 +495,24 @@ bool FInsightsManager::Tick(float DeltaTime)
 	UpdateSessionDuration();
 
 #if !WITH_EDITOR
+	if (!bIsSessionInfoSet && Session.IsValid())
+	{
+		{
+			TraceServices::FAnalysisSessionReadScope SessionReadScope(*Session.Get());
+			const TraceServices::IDiagnosticsProvider* DiagnosticsProvider = TraceServices::ReadDiagnosticsProvider(*Session.Get());
+			if (DiagnosticsProvider && DiagnosticsProvider->IsSessionInfoAvailable())
+			{
+				bIsSessionInfoSet = true;
+			}
+		}
+		if (bIsSessionInfoSet)
+		{
+			UpdateAppTitle();
+		}
+	}
+
 	CheckMemoryUsage();
-#endif
+#endif // !WITH_EDITOR
 
 	return true;
 }
@@ -611,6 +628,7 @@ void FInsightsManager::ResetSession(bool bNotify)
 		}
 	}
 
+	bIsSessionInfoSet = false;
 	bIsAnalysisComplete = false;
 	SessionDuration = 0.0;
 	AnalysisStopwatch.Restart();
@@ -755,6 +773,7 @@ void FInsightsManager::LoadTrace(uint32 InTraceId, bool InAutoQuit)
 	{
 		CurrentTraceId = InTraceId;
 		CurrentTraceFilename = TraceName;
+		bIsSessionInfoSet = false;
 		OnSessionChanged();
 		bSessionAnalysisCompletedAutoQuit = InAutoQuit;
 	}
@@ -783,6 +802,7 @@ void FInsightsManager::LoadTraceFile(const FString& InTraceFilename, bool InAuto
 	{
 		CurrentTraceId = 0;
 		CurrentTraceFilename = InTraceFilename;
+		bIsSessionInfoSet = false;
 		OnSessionChanged();
 		bSessionAnalysisCompletedAutoQuit = InAutoQuit;
 	}
@@ -858,9 +878,34 @@ void FInsightsManager::UpdateAppTitle()
 		}
 		else
 		{
-			const FString SessionName = FPaths::GetBaseFilename(CurrentTraceFilename);
-			const FText AppTitle = FText::Format(LOCTEXT("UnrealInsightsAppNameFmt", "{0} - Unreal Insights"), FText::FromString(SessionName));
-			RootWindow->SetTitle(AppTitle);
+			bool bWasAppTitleUpdated = false;
+			if (Session.IsValid())
+			{
+				TraceServices::FAnalysisSessionReadScope SessionReadScope(*Session.Get());
+				const TraceServices::IDiagnosticsProvider* DiagnosticsProvider = TraceServices::ReadDiagnosticsProvider(*Session.Get());
+				if (DiagnosticsProvider && DiagnosticsProvider->IsSessionInfoAvailable())
+				{
+					TraceServices::FSessionInfo SessionInfo = DiagnosticsProvider->GetSessionInfo();
+
+					const FString SessionName = FPaths::GetBaseFilename(CurrentTraceFilename);
+					const FText AppTitle = FText::Format(LOCTEXT("UnrealInsightsAppNameFmt2", "{0} - {1} - {2} - {3} - {4} - Unreal Insights"),
+						FText::FromString(SessionName),
+						FText::FromString(SessionInfo.Platform),
+						FText::FromString(SessionInfo.AppName),
+						FText::FromString(LexToString(SessionInfo.ConfigurationType)),
+						FText::FromString(LexToString(SessionInfo.TargetType)));
+					RootWindow->SetTitle(AppTitle);
+
+					bWasAppTitleUpdated = true;
+				}
+			}
+
+			if (!bWasAppTitleUpdated)
+			{
+				const FString SessionName = FPaths::GetBaseFilename(CurrentTraceFilename);
+				const FText AppTitle = FText::Format(LOCTEXT("UnrealInsightsAppNameFmt", "{0} - Unreal Insights"), FText::FromString(SessionName));
+				RootWindow->SetTitle(AppTitle);
+			}
 		}
 	}
 #endif
diff --git a/Engine/Source/Developer/TraceInsights/Private/Insights/InsightsManager.h b/Engine/Source/Developer/TraceInsights/Private/Insights/InsightsManager.h
--- a/Engine/Source/Developer/TraceInsights/Private/Insights/InsightsManager.h
+++ b/Engine/Source/Developer/TraceInsights/Private/Insights/InsightsManager.h
@@ -392,6 +392,8 @@ private:
 	bool bIsDebugInfoEnabled = false;
 
 	bool bIsMainTabSet = false;
+	bool bIsSessionInfoSet = false;
+
 	bool bIsAnalysisComplete = false;
 	bool bSessionAnalysisCompletedAutoQuit = false;
 
