[Insights] Fix Sorting an Allocs Table before a compute is finished will forget expansion state

diff --git a/Engine/Source/Developer/TraceInsights/Private/Insights/Table/Widgets/STableTreeView.cpp b/Engine/Source/Developer/TraceInsights/Private/Insights/Table/Widgets/STableTreeView.cpp
--- a/Engine/Source/Developer/TraceInsights/Private/Insights/Table/Widgets/STableTreeView.cpp
+++ b/Engine/Source/Developer/TraceInsights/Private/Insights/Table/Widgets/STableTreeView.cpp
@@ -2521,8 +2521,10 @@ void STableTreeView::OnPreAsyncUpdate()
 	AsyncUpdateStopwatch.Restart();
 	bIsUpdateRunning = true;
 
-	ExpandedNodes.Empty();
-	TreeView->GetExpandedItems(ExpandedNodes);
+	if (ExpandedNodes.IsEmpty())
+	{
+		TreeView->GetExpandedItems(ExpandedNodes);
+	}
 
 	TreeView->SetTreeItemsSource(&DummyGroupNodes);
 	TreeView_Refresh();
@@ -2558,6 +2560,7 @@ void STableTreeView::OnPostAsyncUpdate()
 			TreeView->SetItemExpansion(Node, true);
 		}
 		NodesToExpand.Empty();
+		ExpandedNodes.Empty();
 
 		// Expand each group node on the first few depths (if it doesn't have too many children).
 		SetExpandValueForChildGroups(Root.Get(), 1000, 4, true);
