[Insights] Added custom commands exposed by TimingProfiler:
  - TimingInsights.ExportThreads file
  - TimingInsights.ExportTimers file
  - TimingInsights.ExportTimingEvents file [-columns="..."] [-threads="..."] [-timers="..."] [-startTime=...] [-endTime=...]

The "..." params defines filters for list of columns, threads or timers; it supports *?-type wildcards.

diff --git a/Engine/Source/Developer/TraceInsights/Private/Insights/TimingProfilerManager.cpp b/Engine/Source/Developer/TraceInsights/Private/Insights/TimingProfilerManager.cpp
--- a/Engine/Source/Developer/TraceInsights/Private/Insights/TimingProfilerManager.cpp
+++ b/Engine/Source/Developer/TraceInsights/Private/Insights/TimingProfilerManager.cpp
@@ -14,6 +14,7 @@
 #include "Insights/InsightsStyle.h"
 #include "Insights/TimingProfilerCommon.h"
 #include "Insights/ViewModels/TimerButterflyAggregation.h"
+#include "Insights/ViewModels/TimingExporter.h"
 #include "Insights/Widgets/SFrameTrack.h"
 #include "Insights/Widgets/SLogView.h"
 #include "Insights/Widgets/SStatsView.h"
@@ -627,4 +628,118 @@ void FTimingProfilerManager::OnWindowClosedEvent()
 
 ////////////////////////////////////////////////////////////////////////////////////////////////////
 
+bool FTimingProfilerManager::Exec(const TCHAR* Cmd, FOutputDevice& Ar)
+{
+	if (FParse::Command(&Cmd, TEXT("TimingInsights.ExportThreads")))
+	{
+		Ar.Logf(TEXT("TimingInsights.ExportThreads %s"), Cmd);
+		check(FInsightsManager::Get().IsValid() && FInsightsManager::Get()->GetSession().IsValid());
+		Insights::FTimingExporter Exporter(*FInsightsManager::Get()->GetSession().Get());
+		Insights::FTimingExporter::FExportThreadsParams Params; // default
+		Exporter.ExportThreadsAsText(Cmd, Params);
+		return true;
+	}
+
+	if (FParse::Command(&Cmd, TEXT("TimingInsights.ExportTimers")))
+	{
+		Ar.Logf(TEXT("TimingInsights.ExportTimers %s"), Cmd);
+		check(FInsightsManager::Get().IsValid() && FInsightsManager::Get()->GetSession().IsValid());
+		Insights::FTimingExporter Exporter(*FInsightsManager::Get()->GetSession().Get());
+		Insights::FTimingExporter::FExportTimersParams Params; // default
+		Exporter.ExportTimersAsText(Cmd, Params);
+		return true;
+	}
+
+	if (FParse::Command(&Cmd, TEXT("TimingInsights.ExportTimingEvents")))
+	{
+		Ar.Logf(TEXT("TimingInsights.ExportTimingEvents %s"), Cmd);
+
+		check(FInsightsManager::Get().IsValid() && FInsightsManager::Get()->GetSession().IsValid());
+		Insights::FTimingExporter Exporter(*FInsightsManager::Get()->GetSession().Get());
+		Insights::FTimingExporter::FExportTimingEventsParams Params; // default (all timing events)
+
+		// These variables needs to be in the same scope with the call to Exporter.ExportTimingEventsAsText().
+		TArray<FName> Columns; // referenced by Params.Columns
+		TSet<uint32> IncludedThreads; // referenced in Params.ThreadFilter lambda function
+		TSet<uint32> IncludedTimers; // referenced in Params.TimingEventFilter lambda function
+
+		//////////////////////////////////////////////////
+
+		const bool bUseEscape = true;
+		FString Filename = FParse::Token(Cmd, bUseEscape);
+		Ar.Logf(TEXT("  Filename: %s"), *Filename);
+
+		while (Cmd && Cmd[0] != TEXT('\0'))
+		{
+			FString Token;
+			if (FParse::Token(Cmd, Token, bUseEscape))
+			{
+				Ar.Logf(TEXT("  Token: %s"), *Token);
+
+				static constexpr TCHAR ColumnsToken[] = TEXT("-columns=");
+				static constexpr TCHAR ThreadsToken[] = TEXT("-threads=");
+				static constexpr TCHAR TimersToken[] = TEXT("-timers=");
+				static constexpr TCHAR StartTimeToken[] = TEXT("-startTime=");
+				static constexpr TCHAR EndTimeToken[] = TEXT("-endTime=");
+
+				if (Token.StartsWith(ColumnsToken))
+				{
+					// Comma-delimited list of column names. Supports *?-type wildcard.
+					// Default: -columns="ThreadId,TimerId,StartTime,EndTime,Depth"
+					// Example: -columns="*" -columns="TimerName,Duration"
+					Token.RightChopInline(UE_ARRAY_COUNT(ColumnsToken) - 1);
+					Token.TrimQuotesInline();
+					Exporter.MakeExportTimingEventsColumnList(Token, Columns);
+					Params.Columns = &Columns;
+				}
+				else if (Token.StartsWith(ThreadsToken))
+				{
+					// Comma-delimited list of thread names. Supports *?-type wildcard.
+					// Default: -threads="*" (all threads; no filter)
+					// Example: -threads="GameThread" -threads="GPU1,GPU2,GameThread,Render*"
+					Token.RightChopInline(UE_ARRAY_COUNT(ThreadsToken) - 1);
+					Token.TrimQuotesInline();
+					Params.ThreadFilter = Exporter.MakeThreadFilterInclusive(Token, IncludedThreads);
+				}
+				else if (Token.StartsWith(TimersToken))
+				{
+					// Comma-delimited list of timer names. Supports *?-type wildcard.
+					// Default: -timers="*" (all timers; no filter)
+					// Example: -timers="A,B,*z"
+					Token.RightChopInline(UE_ARRAY_COUNT(TimersToken) - 1);
+					Token.TrimQuotesInline();
+					Params.TimingEventFilter = Exporter.MakeTimingEventFilterByTimersInclusive(Token, IncludedTimers);
+				}
+				else if (Token.StartsWith(StartTimeToken))
+				{
+					// Default: -startTime=-infinte
+					// Example: -startTime=10.0
+					Token.RightChopInline(UE_ARRAY_COUNT(StartTimeToken) - 1);
+					Params.IntervalStartTime = atof(TCHAR_TO_ANSI(*Token));
+				}
+				else if (Token.StartsWith(EndTimeToken))
+				{
+					// Default: -endTime=+infinte
+					// Example: -endTime=20.0
+					Token.RightChopInline(UE_ARRAY_COUNT(EndTimeToken) - 1);
+					Params.IntervalEndTime = atof(TCHAR_TO_ANSI(*Token));
+				}
+				else
+				{
+					Ar.Logf(ELogVerbosity::Warning, TEXT("Unknown Cmd Param: %s"), *Token);
+				}
+			}
+		}
+
+		//////////////////////////////////////////////////
+
+		Exporter.ExportTimingEventsAsText(Filename, Params);
+		return true;
+	}
+
+	return false;
+}
+
+////////////////////////////////////////////////////////////////////////////////////////////////////
+
 #undef LOCTEXT_NAMESPACE
diff --git a/Engine/Source/Developer/TraceInsights/Private/Insights/TimingProfilerManager.h b/Engine/Source/Developer/TraceInsights/Private/Insights/TimingProfilerManager.h
--- a/Engine/Source/Developer/TraceInsights/Private/Insights/TimingProfilerManager.h
+++ b/Engine/Source/Developer/TraceInsights/Private/Insights/TimingProfilerManager.h
@@ -66,6 +66,7 @@ public:
 	virtual void RegisterMajorTabs(IUnrealInsightsModule& InsightsModule) override;
 	virtual void UnregisterMajorTabs() override;
 	virtual void OnWindowClosedEvent() override;
+	virtual bool Exec(const TCHAR* Cmd, FOutputDevice& Ar) override;
 
 	//////////////////////////////////////////////////
 
diff --git a/Engine/Source/Developer/TraceInsights/Private/Insights/ViewModels/TimingExporter.cpp b/Engine/Source/Developer/TraceInsights/Private/Insights/ViewModels/TimingExporter.cpp
new file mode 100644
--- /dev/null
+++ b/Engine/Source/Developer/TraceInsights/Private/Insights/ViewModels/TimingExporter.cpp
@@ -0,0 +1,761 @@
+// Copyright Epic Games, Inc. All Rights Reserved.
+
+#include "TimingExporter.h"
+#include "HAL/PlatformFileManager.h"
+#include "Logging/MessageLog.h"
+
+// Insights
+#include "Insights/Common/Stopwatch.h"
+#include "Insights/Log.h"
+#include "Insights/ViewModels/ThreadTimingTrack.h"
+#include "Insights/TimingProfilerManager.h"
+
+namespace Insights
+{
+
+#define LOCTEXT_NAMESPACE "FTimingExporter"
+
+////////////////////////////////////////////////////////////////////////////////////////////////////
+// FTimingExporter
+////////////////////////////////////////////////////////////////////////////////////////////////////
+
+FTimingExporter::FTimingExporter(const TraceServices::IAnalysisSession& InSession)
+	: Session(InSession)
+{	
+}
+
+////////////////////////////////////////////////////////////////////////////////////////////////////
+
+FTimingExporter::~FTimingExporter()
+{	
+}
+
+////////////////////////////////////////////////////////////////////////////////////////////////////
+
+IFileHandle* FTimingExporter::OpenExportFile(const TCHAR* InFilename) const
+{
+	IPlatformFile& PlatformFile = FPlatformFileManager::Get().GetPlatformFile();
+
+	FString Path = FPaths::GetPath(InFilename);
+	if (!PlatformFile.DirectoryExists(*Path))
+	{
+		PlatformFile.CreateDirectoryTree(*Path);
+	}
+
+	IFileHandle* ExportFileHandle = PlatformFile.OpenWrite(InFilename);
+
+	if (ExportFileHandle == nullptr)
+	{
+		Error(LOCTEXT("FailedToOpenFile", "Export failed. Failed to open file for write."));
+		return nullptr;
+	}
+
+	return ExportFileHandle;
+}
+
+////////////////////////////////////////////////////////////////////////////////////////////////////
+
+void FTimingExporter::Error(const FText& InMessage) const
+{
+	FName LogListingName = FTimingProfilerManager::Get()->GetLogListingName();
+	FMessageLog ReportMessageLog((LogListingName != NAME_None) ? LogListingName : TEXT("Other"));
+	ReportMessageLog.Error(InMessage);
+	ReportMessageLog.Notify();
+}
+
+////////////////////////////////////////////////////////////////////////////////////////////////////
+
+void FTimingExporter::AppendString(FUtf8StringBuilder& StringBuilder, const TCHAR* Str, UTF8CHAR Separator)
+{
+	if (Str == nullptr || Str[0] == TCHAR('\0'))
+	{
+		// nothing to append
+	}
+	else if (FCString::Strchr(Str, Separator) != nullptr)
+	{
+		if (FCString::Strchr(Str, TEXT('\"')) != nullptr)
+		{
+			FString String = Str;
+			String.ReplaceInline(TEXT("\""), TEXT("\"\""));
+			String = FString::Printf(TEXT("\"%s\""), *String);
+			StringBuilder.Append(TCHAR_TO_UTF8(*String));
+		}
+		else
+		{
+			FString String = FString::Printf(TEXT("\"%s\""), Str);
+			StringBuilder.Append(TCHAR_TO_UTF8(*String));
+		}
+	}
+	else
+	{
+		StringBuilder.Append(TCHAR_TO_UTF8(Str));
+	}
+}
+
+////////////////////////////////////////////////////////////////////////////////////////////////////
+
+int32 FTimingExporter::ExportThreadsAsText(const FString& Filename, FExportThreadsParams& Params) const
+{
+	checkf(Params.Columns == nullptr, TEXT("Custom list of columns is not yet supported!"));
+
+	FStopwatch Stopwatch;
+	Stopwatch.Start();
+
+	IFileHandle* ExportFileHandle = OpenExportFile(*Filename);
+	if (!ExportFileHandle)
+	{
+		return -1;
+	}
+
+	UTF8CHAR Separator = UTF8CHAR('\t');
+	if (Filename.EndsWith(TEXT(".csv")))
+	{
+		Separator = UTF8CHAR(',');
+	}
+	const UTF8CHAR LineEnd = UTF8CHAR('\n');
+
+	FUtf8StringBuilder StringBuilder;
+
+	// Write header.
+	{
+		StringBuilder.Append(UTF8TEXT("Id"));
+		StringBuilder.AppendChar(Separator);
+		StringBuilder.Append(UTF8TEXT("Name"));
+		StringBuilder.AppendChar(Separator);
+		StringBuilder.Append(UTF8TEXT("Group"));
+		StringBuilder.AppendChar(LineEnd);
+
+		ExportFileHandle->Write((const uint8*)StringBuilder.ToString(), StringBuilder.Len() * sizeof(UTF8CHAR));
+	}
+
+	int32 ThreadCount = 0;
+
+	// Write values.
+	{
+		StringBuilder.Reset();
+		StringBuilder.Appendf(UTF8TEXT("%u"), FGpuTimingTrack::Gpu1ThreadId);
+		StringBuilder.AppendChar(Separator);
+		StringBuilder.Append(UTF8TEXT("GPU1"));
+		StringBuilder.AppendChar(Separator);
+		StringBuilder.Append(UTF8TEXT("GPU"));
+		StringBuilder.AppendChar(LineEnd);
+		ExportFileHandle->Write((const uint8*)StringBuilder.ToString(), StringBuilder.Len() * sizeof(UTF8CHAR));
+
+		StringBuilder.Reset();
+		StringBuilder.Appendf(UTF8TEXT("%u"), FGpuTimingTrack::Gpu2ThreadId);
+		StringBuilder.AppendChar(Separator);
+		StringBuilder.Append(UTF8TEXT("GPU2"));
+		StringBuilder.AppendChar(Separator);
+		StringBuilder.Append(UTF8TEXT("GPU"));
+		StringBuilder.AppendChar(LineEnd);
+		ExportFileHandle->Write((const uint8*)StringBuilder.ToString(), StringBuilder.Len() * sizeof(UTF8CHAR));
+
+		// Iterate the CPU threads.
+		{
+			TraceServices::FAnalysisSessionReadScope SessionReadScope(Session);
+			const TraceServices::IThreadProvider& ThreadProvider = TraceServices::ReadThreadProvider(Session);
+			ThreadProvider.EnumerateThreads(
+				[&](const TraceServices::FThreadInfo& ThreadInfo)
+				{
+					StringBuilder.Reset();
+					StringBuilder.Appendf(UTF8TEXT("%u"), ThreadInfo.Id);
+					StringBuilder.AppendChar(Separator);
+					AppendString(StringBuilder, ThreadInfo.Name, Separator);
+					StringBuilder.AppendChar(Separator);
+					AppendString(StringBuilder, ThreadInfo.GroupName, Separator);
+					StringBuilder.AppendChar(LineEnd);
+					ExportFileHandle->Write((const uint8*)StringBuilder.ToString(), StringBuilder.Len() * sizeof(UTF8CHAR));
+					++ThreadCount;
+				});
+		}
+	}
+
+	ExportFileHandle->Flush();
+	delete ExportFileHandle;
+	ExportFileHandle = nullptr;
+
+	Stopwatch.Stop();
+	const double TotalTime = Stopwatch.GetAccumulatedTime();
+	UE_LOG(TraceInsights, Log, TEXT("Exported %d threads to file in %.3fs (\"%s\")."), ThreadCount, TotalTime, *Filename);
+
+	return ThreadCount;
+}
+
+////////////////////////////////////////////////////////////////////////////////////////////////////
+
+int32 FTimingExporter::ExportTimersAsText(const FString& Filename, FExportTimersParams& Params) const
+{
+	checkf(Params.Columns == nullptr, TEXT("Custom list of columns is not yet supported!"));
+
+	FStopwatch Stopwatch;
+	Stopwatch.Start();
+
+	IFileHandle* ExportFileHandle = OpenExportFile(*Filename);
+	if (!ExportFileHandle)
+	{
+		return -1;
+	}
+
+	UTF8CHAR Separator = UTF8CHAR('\t');
+	if (Filename.EndsWith(TEXT(".csv")))
+	{
+		Separator = UTF8CHAR(',');
+	}
+	const UTF8CHAR LineEnd = UTF8CHAR('\n');
+
+	FUtf8StringBuilder StringBuilder;
+
+	// Write header.
+	{
+		StringBuilder.Append(UTF8TEXT("Id"));
+		StringBuilder.AppendChar(Separator);
+		StringBuilder.Append(UTF8TEXT("Type"));
+		StringBuilder.AppendChar(Separator);
+		StringBuilder.Append(UTF8TEXT("Name"));
+		StringBuilder.AppendChar(Separator);
+		StringBuilder.Append(UTF8TEXT("File"));
+		StringBuilder.AppendChar(Separator);
+		StringBuilder.Append(UTF8TEXT("Line"));
+		StringBuilder.AppendChar(LineEnd);
+
+		ExportFileHandle->Write((const uint8*)StringBuilder.ToString(), StringBuilder.Len() * sizeof(UTF8CHAR));
+	}
+
+	uint32 TimerCount = 0;
+
+	// Write values.
+	if (TraceServices::ReadTimingProfilerProvider(Session))
+	{
+		TraceServices::FAnalysisSessionReadScope SessionReadScope(Session);
+
+		const TraceServices::ITimingProfilerProvider& TimingProfilerProvider = *TraceServices::ReadTimingProfilerProvider(Session);
+
+		const TraceServices::ITimingProfilerTimerReader* TimerReader;
+		TimingProfilerProvider.ReadTimers([&TimerReader](const TraceServices::ITimingProfilerTimerReader& Out) { TimerReader = &Out; });
+
+		TimerCount = TimerReader->GetTimerCount();
+		for (uint32 TimerIndex = 0; TimerIndex < TimerCount; ++TimerIndex)
+		{
+			const TraceServices::FTimingProfilerTimer& Timer = *(TimerReader->GetTimer(TimerIndex));
+			StringBuilder.Reset();
+			StringBuilder.Appendf(UTF8TEXT("%u"), Timer.Id);
+			StringBuilder.AppendChar(Separator);
+			StringBuilder.Append(Timer.IsGpuTimer ? UTF8TEXT("GPU") : UTF8TEXT("CPU"));
+			StringBuilder.AppendChar(Separator);
+			AppendString(StringBuilder, Timer.Name, Separator);
+			StringBuilder.AppendChar(Separator);
+			if (Timer.File)
+			{
+				StringBuilder.Append((const UTF8CHAR*)TCHAR_TO_UTF8(Timer.File));
+			}
+			StringBuilder.AppendChar(Separator);
+			StringBuilder.Appendf(UTF8TEXT("%u"), Timer.Line);
+			StringBuilder.AppendChar(LineEnd);
+			ExportFileHandle->Write((const uint8*)StringBuilder.ToString(), StringBuilder.Len() * sizeof(UTF8CHAR));
+		}
+	}
+
+	ExportFileHandle->Flush();
+	delete ExportFileHandle;
+	ExportFileHandle = nullptr;
+
+	Stopwatch.Stop();
+	const double TotalTime = Stopwatch.GetAccumulatedTime();
+	UE_LOG(TraceInsights, Log, TEXT("Exported %u timers to file in %.3fs (\"%s\")."), TimerCount, TotalTime, *Filename);
+
+	return int32(TimerCount);
+}
+
+////////////////////////////////////////////////////////////////////////////////////////////////////
+
+const FName FTimingExporter::ExportTimingEvents_ThreadIdColumn("ThreadId");
+const FName FTimingExporter::ExportTimingEvents_ThreadNameColumn("ThreadName");
+const FName FTimingExporter::ExportTimingEvents_TimerIdColumn("TimerId");
+const FName FTimingExporter::ExportTimingEvents_TimerNameColumn("TimerName");
+const FName FTimingExporter::ExportTimingEvents_StartTimeColumn("StartTime");
+const FName FTimingExporter::ExportTimingEvents_EndTimeColumn("EndTime");
+const FName FTimingExporter::ExportTimingEvents_DurationColumn("Duration");
+const FName FTimingExporter::ExportTimingEvents_DepthColumn("Depth");
+
+////////////////////////////////////////////////////////////////////////////////////////////////////
+
+void FTimingExporter::ExportTimingEvents_InitColumns() const
+{
+	if (ExportTimingEventsColumns.IsEmpty())
+	{
+		ExportTimingEventsColumns.Add(ExportTimingEvents_ThreadIdColumn);
+		ExportTimingEventsColumns.Add(ExportTimingEvents_ThreadNameColumn);
+		ExportTimingEventsColumns.Add(ExportTimingEvents_TimerIdColumn);
+		ExportTimingEventsColumns.Add(ExportTimingEvents_TimerNameColumn);
+		ExportTimingEventsColumns.Add(ExportTimingEvents_StartTimeColumn);
+		ExportTimingEventsColumns.Add(ExportTimingEvents_EndTimeColumn);
+		ExportTimingEventsColumns.Add(ExportTimingEvents_DurationColumn);
+		ExportTimingEventsColumns.Add(ExportTimingEvents_DepthColumn);
+
+		ExportTimingEventsDefaultColumns.Add(ExportTimingEvents_ThreadIdColumn);
+		ExportTimingEventsDefaultColumns.Add(ExportTimingEvents_TimerIdColumn);
+		ExportTimingEventsDefaultColumns.Add(ExportTimingEvents_StartTimeColumn);
+		ExportTimingEventsDefaultColumns.Add(ExportTimingEvents_EndTimeColumn);
+		ExportTimingEventsDefaultColumns.Add(ExportTimingEvents_DepthColumn);
+	}
+}
+
+////////////////////////////////////////////////////////////////////////////////////////////////////
+
+void FTimingExporter::MakeExportTimingEventsColumnList(const FString& InColumnsString, TArray<FName>& OutColumnList) const
+{
+	ExportTimingEvents_InitColumns();
+
+	TArray<FString> Columns;
+	InColumnsString.ParseIntoArray(Columns, TEXT(","), true);
+
+	for (const FString& ColumnWildcard : Columns)
+	{
+		FName ColumnName(*ColumnWildcard);
+		if (ExportTimingEventsColumns.Contains(ColumnName))
+		{
+			OutColumnList.Add(ColumnName);
+		}
+		else
+		{
+			for (const FName& Column : ExportTimingEventsColumns)
+			{
+				if (Column.GetPlainNameString().MatchesWildcard(ColumnWildcard))
+				{
+					OutColumnList.Add(Column);
+				}
+			}
+		}
+	}
+}
+
+////////////////////////////////////////////////////////////////////////////////////////////////////
+
+void FTimingExporter::ExportTimingEvents_WriteHeader(FExportTimingEventsInternalParams& Params) const
+{
+	bool bFirst = true;
+	for (const FName& Column : Params.Columns)
+	{
+		if (ExportTimingEventsColumns.Contains(Column))
+		{
+			if (bFirst)
+			{
+				bFirst = false;
+			}
+			else
+			{
+				Params.StringBuilder.AppendChar(Params.Separator);
+			}
+			Params.StringBuilder.Append(Column.GetPlainNameString());
+		}
+	}
+	Params.StringBuilder.AppendChar(Params.LineEnd);
+
+	Params.ExportFileHandle->Write((const uint8*)Params.StringBuilder.ToString(), Params.StringBuilder.Len() * sizeof(UTF8CHAR));
+}
+
+////////////////////////////////////////////////////////////////////////////////////////////////////
+
+int32 FTimingExporter::ExportTimingEvents_WriteEvents(FExportTimingEventsInternalParams& Params) const
+{
+	TMap<uint32, const TCHAR*> Timers; // only used if exporting the TimerName column
+
+	if (Params.Columns.Contains(ExportTimingEvents_TimerNameColumn))
+	{
+		// Iterate the GPU & CPU timers.
+		if (TraceServices::ReadTimingProfilerProvider(Session))
+		{
+			TraceServices::FAnalysisSessionReadScope SessionReadScope(Session);
+
+			const TraceServices::ITimingProfilerProvider& TimingProfilerProvider = *TraceServices::ReadTimingProfilerProvider(Session);
+
+			const TraceServices::ITimingProfilerTimerReader* TimerReader;
+			TimingProfilerProvider.ReadTimers([&TimerReader, &Timers](const TraceServices::ITimingProfilerTimerReader& Out) { TimerReader = &Out; });
+
+			uint32 TimerCount = TimerReader->GetTimerCount();
+			for (uint32 TimerIndex = 0; TimerIndex < TimerCount; ++TimerIndex)
+			{
+				const TraceServices::FTimingProfilerTimer& Timer = *(TimerReader->GetTimer(TimerIndex));
+				Timers.Add(Timer.Id, Timer.Name);
+			}
+		}
+	}
+
+	int32 TimingEventCount = 0;
+
+	if (TraceServices::ReadTimingProfilerProvider(Session))
+	{
+		TraceServices::FAnalysisSessionReadScope SessionReadScope(Session);
+
+		const TraceServices::ITimingProfilerProvider& TimingProfilerProvider = *TraceServices::ReadTimingProfilerProvider(Session);
+
+		auto TimelineEnumerator = [&Params, &TimingEventCount, &Timers](const TraceServices::ITimingProfilerProvider::Timeline& Timeline)
+		{
+			// Iterate timing events.
+			Timeline.EnumerateEvents(Params.UserParams.IntervalStartTime, Params.UserParams.IntervalEndTime,
+				[&Params, &TimingEventCount, &Timers](double EventStartTime, double EventEndTime, uint32 EventDepth, const TraceServices::FTimingProfilerEvent& Event)
+				{
+					if (!Params.UserParams.TimingEventFilter || Params.UserParams.TimingEventFilter(EventStartTime, EventEndTime, EventDepth, Event))
+					{
+						Params.StringBuilder.Reset();
+
+						bool bFirst = true;
+						for (const FName& Column : Params.Columns)
+						{
+							if (Params.Exporter.ExportTimingEventsColumns.Contains(Column))
+							{
+								if (bFirst)
+								{
+									bFirst = false;
+								}
+								else
+								{
+									Params.StringBuilder.AppendChar(Params.Separator);
+								}
+
+								if (Column == ExportTimingEvents_ThreadIdColumn)
+								{
+									Params.StringBuilder.Appendf(UTF8TEXT("%u"), Params.ThreadId);
+								}
+								else if (Column == ExportTimingEvents_ThreadNameColumn)
+								{
+									AppendString(Params.StringBuilder, Params.ThreadName, Params.Separator);
+								}
+								else if (Column == ExportTimingEvents_TimerIdColumn)
+								{
+									Params.StringBuilder.Appendf(UTF8TEXT("%u"), Event.TimerIndex);
+								}
+								else if (Column == ExportTimingEvents_TimerNameColumn)
+								{
+									const TCHAR* TimerName = Timers.FindRef(Event.TimerIndex);
+									AppendString(Params.StringBuilder, TimerName, Params.Separator);
+								}
+								else if (Column == ExportTimingEvents_StartTimeColumn)
+								{
+									Params.StringBuilder.Appendf(UTF8TEXT("%.9g"), EventStartTime);
+								}
+								else if (Column == ExportTimingEvents_EndTimeColumn)
+								{
+									Params.StringBuilder.Appendf(UTF8TEXT("%.9g"), EventEndTime);
+								}
+								else if (Column == ExportTimingEvents_DurationColumn)
+								{
+									Params.StringBuilder.Appendf(UTF8TEXT("%.9f"), EventEndTime - EventStartTime);
+								}
+								else if (Column == ExportTimingEvents_DepthColumn)
+								{
+									Params.StringBuilder.Appendf(UTF8TEXT("%u"), EventDepth);
+								}
+							}
+						}
+						Params.StringBuilder.AppendChar(Params.LineEnd);
+
+						Params.ExportFileHandle->Write((const uint8*)Params.StringBuilder.ToString(), Params.StringBuilder.Len() * sizeof(UTF8CHAR));
+						++TimingEventCount;
+					}
+
+					return TraceServices::EEventEnumerate::Continue;
+				});
+		};
+
+		// Iterate the GPU timelines.
+		{
+			if (!Params.UserParams.ThreadFilter || Params.UserParams.ThreadFilter(FGpuTimingTrack::Gpu1ThreadId))
+			{
+				Params.ThreadId = FGpuTimingTrack::Gpu1ThreadId;
+				Params.ThreadName = TEXT("GPU1");
+				uint32 GpuTimelineIndex1 = 0;
+				TimingProfilerProvider.GetGpuTimelineIndex(GpuTimelineIndex1);
+				TimingProfilerProvider.ReadTimeline(GpuTimelineIndex1, TimelineEnumerator);
+			}
+
+			if (!Params.UserParams.ThreadFilter || Params.UserParams.ThreadFilter(FGpuTimingTrack::Gpu2ThreadId))
+			{
+				Params.ThreadId = FGpuTimingTrack::Gpu2ThreadId;
+				Params.ThreadName = TEXT("GPU2");
+				uint32 GpuTimelineIndex2 = 0;
+				TimingProfilerProvider.GetGpu2TimelineIndex(GpuTimelineIndex2);
+				TimingProfilerProvider.ReadTimeline(GpuTimelineIndex2, TimelineEnumerator);
+			}
+		}
+
+		// Iterate the CPU threads and their corresponding timelines.
+		const TraceServices::IThreadProvider& ThreadProvider = TraceServices::ReadThreadProvider(Session);
+		ThreadProvider.EnumerateThreads(
+			[&TimelineEnumerator, &Params, &TimingProfilerProvider](const TraceServices::FThreadInfo& ThreadInfo)
+			{
+				if (!Params.UserParams.ThreadFilter || Params.UserParams.ThreadFilter(ThreadInfo.Id))
+				{
+					Params.ThreadId = ThreadInfo.Id;
+					Params.ThreadName = ThreadInfo.Name;
+					uint32 CpuTimelineIndex = 0;
+					TimingProfilerProvider.GetCpuThreadTimelineIndex(ThreadInfo.Id, CpuTimelineIndex);
+					TimingProfilerProvider.ReadTimeline(CpuTimelineIndex, TimelineEnumerator);
+				}
+			});
+	}
+
+	return TimingEventCount;
+}
+
+////////////////////////////////////////////////////////////////////////////////////////////////////
+
+int32 FTimingExporter::ExportTimingEventsAsText(const FString& Filename, FExportTimingEventsParams& Params) const
+{
+	ExportTimingEvents_InitColumns();
+	const TArray<FName>& Columns = Params.Columns ? *Params.Columns : ExportTimingEventsDefaultColumns;
+
+	FStopwatch Stopwatch;
+	Stopwatch.Start();
+
+	IFileHandle* ExportFileHandle = OpenExportFile(*Filename);
+	if (!ExportFileHandle)
+	{
+		return -1;
+	}
+
+	UTF8CHAR Separator = UTF8CHAR('\t');
+	if (Filename.EndsWith(TEXT(".csv")))
+	{
+		Separator = UTF8CHAR(',');
+	}
+	const UTF8CHAR LineEnd = UTF8CHAR('\n');
+
+	FUtf8StringBuilder StringBuilder;
+
+	FExportTimingEventsInternalParams InternalParams = { *this, Params, Columns, ExportFileHandle, Separator, LineEnd, StringBuilder, 0 };
+
+	// Write header.
+	ExportTimingEvents_WriteHeader(InternalParams);
+
+	// Write values.
+	int32 TimingEventCount = ExportTimingEvents_WriteEvents(InternalParams);
+
+	ExportFileHandle->Flush();
+	delete ExportFileHandle;
+	ExportFileHandle = nullptr;
+
+	Stopwatch.Stop();
+	const double TotalTime = Stopwatch.GetAccumulatedTime();
+	UE_LOG(TraceInsights, Log, TEXT("Exported %d timing events to file in %.3fs (\"%s\")."), TimingEventCount, TotalTime, *Filename);
+
+	return TimingEventCount;
+}
+
+////////////////////////////////////////////////////////////////////////////////////////////////////
+
+FTimingExporter::FThreadFilterFunc FTimingExporter::MakeThreadFilterInclusive(const FString& InFilterString, TSet<uint32>& OutIncludedThreads) const
+{
+	if (InFilterString.Len() == 1 && InFilterString[0] == TEXT('*'))
+	{
+		return nullptr;
+	}
+
+	OutIncludedThreads.Reset();
+
+	TMap<FString, uint32> Threads;
+
+	// Add the GPU threads.
+	Threads.Add(FString("GPU"), FGpuTimingTrack::Gpu1ThreadId);
+	Threads.Add(FString("GPU1"), FGpuTimingTrack::Gpu1ThreadId);
+	Threads.Add(FString("GPU2"), FGpuTimingTrack::Gpu2ThreadId);
+
+	// Iterate the CPU threads.
+	{
+		TraceServices::FAnalysisSessionReadScope SessionReadScope(Session);
+		const TraceServices::IThreadProvider& ThreadProvider = TraceServices::ReadThreadProvider(Session);
+		ThreadProvider.EnumerateThreads(
+			[&Threads](const TraceServices::FThreadInfo& ThreadInfo)
+			{
+				Threads.Add(FString(ThreadInfo.Name), ThreadInfo.Id);
+			});
+	}
+
+	TArray<FString> Filter;
+	InFilterString.ParseIntoArray(Filter, TEXT(","), true);
+
+	for (const FString& ThreadWildcard : Filter)
+	{
+		const uint32* Id = Threads.Find(ThreadWildcard);
+		if (Id)
+		{
+			OutIncludedThreads.Add(*Id);
+		}
+		else
+		{
+			for (const auto& KeyValuePair : Threads)
+			{
+				if (KeyValuePair.Key.MatchesWildcard(ThreadWildcard))
+				{
+					OutIncludedThreads.Add(KeyValuePair.Value);
+				}
+			}
+		}
+	}
+
+	return MakeThreadFilterInclusive(OutIncludedThreads);
+}
+
+////////////////////////////////////////////////////////////////////////////////////////////////////
+
+TFunction<bool(uint32)> FTimingExporter::MakeThreadFilterInclusive(const TSet<uint32>& IncludedThreads)
+{
+	if (IncludedThreads.Num() == 0)
+	{
+		return [](uint32 ThreadId) -> bool
+		{
+			return false;
+		};
+	}
+
+	if (IncludedThreads.Num() == 1)
+	{
+		const uint32 IncludedThreadId = IncludedThreads[FSetElementId::FromInteger(0)];
+		return [IncludedThreadId](uint32 ThreadId) -> bool
+		{
+			return ThreadId == IncludedThreadId;
+		};
+	}
+
+	return [&IncludedThreads](uint32 ThreadId) -> bool
+	{
+		return IncludedThreads.Contains(ThreadId);
+	};
+}
+
+////////////////////////////////////////////////////////////////////////////////////////////////////
+
+TFunction<bool(uint32)> FTimingExporter::MakeThreadFilterExclusive(const TSet<uint32>& ExcludedThreads)
+{
+	if (ExcludedThreads.Num() == 0)
+	{
+		return nullptr;
+	}
+
+	if (ExcludedThreads.Num() == 1)
+	{
+		const uint32 ExcludedThreadId = ExcludedThreads[FSetElementId::FromInteger(0)];
+		return [ExcludedThreadId](uint32 ThreadId) -> bool
+		{
+			return ThreadId != ExcludedThreadId;
+		};
+	}
+	
+	return [&ExcludedThreads](uint32 ThreadId) -> bool
+	{
+		return !ExcludedThreads.Contains(ThreadId);
+	};
+}
+
+////////////////////////////////////////////////////////////////////////////////////////////////////
+
+FTimingExporter::FTimingEventFilterFunc FTimingExporter::MakeTimingEventFilterByTimersInclusive(const FString& InFilterString, TSet<uint32>& OutIncludedTimers) const
+{
+	if (InFilterString.Len() == 1 && InFilterString[0] == TEXT('*'))
+	{
+		return nullptr;
+	}
+
+	OutIncludedTimers.Reset();
+	
+	TMap<FString, uint32> Timers;
+
+	// Iterate the GPU & CPU timers.
+	if (TraceServices::ReadTimingProfilerProvider(Session))
+	{
+		TraceServices::FAnalysisSessionReadScope SessionReadScope(Session);
+
+		const TraceServices::ITimingProfilerProvider& TimingProfilerProvider = *TraceServices::ReadTimingProfilerProvider(Session);
+
+		const TraceServices::ITimingProfilerTimerReader* TimerReader;
+		TimingProfilerProvider.ReadTimers([&TimerReader, &Timers](const TraceServices::ITimingProfilerTimerReader& Out) { TimerReader = &Out; });
+
+		uint32 TimerCount = TimerReader->GetTimerCount();
+		for (uint32 TimerIndex = 0; TimerIndex < TimerCount; ++TimerIndex)
+		{
+			const TraceServices::FTimingProfilerTimer& Timer = *(TimerReader->GetTimer(TimerIndex));
+			Timers.Add(FString(Timer.Name), Timer.Id);
+		}
+	}
+
+	TArray<FString> Filter;
+	InFilterString.ParseIntoArray(Filter, TEXT(","), true);
+
+	for (const FString& TimerWildcard : Filter)
+	{
+		const uint32* Id = Timers.Find(TimerWildcard);
+		if (Id)
+		{
+			OutIncludedTimers.Add(*Id);
+		}
+		else
+		{
+			for (const auto& KeyValuePair : Timers)
+			{
+				if (KeyValuePair.Key.MatchesWildcard(TimerWildcard))
+				{
+					OutIncludedTimers.Add(KeyValuePair.Value);
+				}
+			}
+		}
+	}
+
+	return MakeTimingEventFilterByTimersInclusive(OutIncludedTimers);
+}
+
+////////////////////////////////////////////////////////////////////////////////////////////////////
+
+FTimingExporter::FTimingEventFilterFunc FTimingExporter::MakeTimingEventFilterByTimersInclusive(const TSet<uint32>& IncludedTimers)
+{
+	if (IncludedTimers.Num() == 0)
+	{
+		return [](double EventStartTime, double EventEndTime, uint32 EventDepth, const TraceServices::FTimingProfilerEvent& Event) -> bool
+		{
+			return false;
+		};
+	}
+
+	if (IncludedTimers.Num() == 1)
+	{
+		const uint32 IncludedTimerId = IncludedTimers[FSetElementId::FromInteger(0)];
+		return [IncludedTimerId](double EventStartTime, double EventEndTime, uint32 EventDepth, const TraceServices::FTimingProfilerEvent& Event) -> bool
+		{
+			return Event.TimerIndex == IncludedTimerId;
+		};
+	}
+
+	return [&IncludedTimers](double EventStartTime, double EventEndTime, uint32 EventDepth, const TraceServices::FTimingProfilerEvent& Event) -> bool
+	{
+		return IncludedTimers.Contains(Event.TimerIndex);
+	};
+}
+
+////////////////////////////////////////////////////////////////////////////////////////////////////
+
+FTimingExporter::FTimingEventFilterFunc FTimingExporter::MakeTimingEventFilterByTimersExclusive(const TSet<uint32>& ExcludedTimers)
+{
+	if (ExcludedTimers.Num() == 0)
+	{
+		return nullptr;
+	}
+		
+	if (ExcludedTimers.Num() == 1)
+	{
+		const uint32 ExcludedTimerId = ExcludedTimers[FSetElementId::FromInteger(0)];
+		return [ExcludedTimerId](double EventStartTime, double EventEndTime, uint32 EventDepth, const TraceServices::FTimingProfilerEvent& Event) -> bool
+		{
+			return Event.TimerIndex != ExcludedTimerId;
+		};
+	}
+
+	return [&ExcludedTimers](double EventStartTime, double EventEndTime, uint32 EventDepth, const TraceServices::FTimingProfilerEvent& Event) -> bool
+	{
+		return !ExcludedTimers.Contains(Event.TimerIndex);
+	};
+}
+
+////////////////////////////////////////////////////////////////////////////////////////////////////
+
+#undef LOCTEXT_NAMESPACE
+
+} // namespace Insights
diff --git a/Engine/Source/Developer/TraceInsights/Private/Insights/ViewModels/TimingExporter.h b/Engine/Source/Developer/TraceInsights/Private/Insights/ViewModels/TimingExporter.h
new file mode 100644
--- /dev/null
+++ b/Engine/Source/Developer/TraceInsights/Private/Insights/ViewModels/TimingExporter.h
@@ -0,0 +1,167 @@
+// Copyright Epic Games, Inc. All Rights Reserved.
+
+#pragma once
+
+#include "CoreMinimal.h"
+
+#include "TraceServices/Model/TimingProfiler.h"
+#include <limits>
+
+////////////////////////////////////////////////////////////////////////////////////////////////////
+
+class IFileHandle;
+
+////////////////////////////////////////////////////////////////////////////////////////////////////
+
+namespace Insights
+{
+
+////////////////////////////////////////////////////////////////////////////////////////////////////
+
+class FTimingExporter
+{
+public:
+	typedef TFunction<bool(const FName /*ColumnId*/)> FColumnFilterFunc;
+	typedef TFunction<bool(uint32 /*ThreadId*/)> FThreadFilterFunc;
+	typedef TFunction<bool(double /*EventStartTime*/, double /*EventEndTime*/, uint32 /*EventDepth*/, const TraceServices::FTimingProfilerEvent& /*Event*/)> FTimingEventFilterFunc;
+
+	struct FExportThreadsParams
+	{
+		/**
+		 * The list of columns to be exported.
+		 * If nullptr, it uses the default list of columns.
+		 */
+		const TArray<FName>* Columns = nullptr;
+	};
+
+	struct FExportTimersParams
+	{
+		/**
+		 * The list of columns to be exported.
+		 * If nullptr, it uses the default list of columns.
+		 */
+		const TArray<FName>* Columns = nullptr;
+	};
+
+	struct FExportTimingEventsParams
+	{
+		/**
+		 * The list of columns to be exported.
+		 * If nullptr, it uses the default list of columns.
+		 */
+		const TArray<FName>* Columns = nullptr;
+
+		/**
+		 * Filters the threads for which timing events are exported.
+		 * If nullptr, exports timing events from all threads.
+		 */
+		FThreadFilterFunc ThreadFilter = nullptr;
+
+		/**
+		 * Filters the timing events.
+		 * If nullptr, exports all timing events.
+		 */
+		FTimingEventFilterFunc TimingEventFilter = nullptr;
+
+		/**
+		 * Filters the timing events by time.
+		 * Only timning events that intersects the [StartTime, EndTime] interval are exported.
+		 */
+		double IntervalStartTime = -std::numeric_limits<double>::infinity();
+		double IntervalEndTime = +std::numeric_limits<double>::infinity();
+	};
+
+private:
+	typedef TUtf8StringBuilder<1024> FUtf8StringBuilder;
+
+	struct FExportTimingEventsInternalParams
+	{
+		const FTimingExporter& Exporter;
+		const FExportTimingEventsParams& UserParams;
+		const TArray<FName>& Columns;
+		IFileHandle* ExportFileHandle;
+		const UTF8CHAR Separator;
+		const UTF8CHAR LineEnd;
+		TUtf8StringBuilder<1024>& StringBuilder;
+		uint32 ThreadId;
+		const TCHAR* ThreadName;
+	};
+
+public:
+	FTimingExporter(const TraceServices::IAnalysisSession& InSession);
+	virtual ~FTimingExporter();
+
+	int32 ExportThreadsAsText(const FString& Filename, FExportThreadsParams& Params) const;
+	int32 ExportTimersAsText(const FString& Filename, FExportTimersParams& Params) const;
+	int32 ExportTimingEventsAsText(const FString& Filename, FExportTimingEventsParams& Params) const;
+
+	void MakeExportTimingEventsColumnList(const FString& InColumnsString, TArray<FName>& OutColumnList) const;
+
+	//////////////////////////////////////////////////////////////////////
+	// Utilities to make FThreadFilterFunc filters.
+
+	FThreadFilterFunc MakeThreadFilterInclusive(const FString& InFilterString, TSet<uint32>& OutIncludedThreads) const;
+
+	/**
+	 * Makes a FThreadFilterFunc using a set of included list of threads.
+	 * Note: The set is referenced in the returned function.
+	 * @param IncludedThreads The set of thread ids to be accepted by filter.
+	 * @return The FThreadFilterFunc function.
+	 */
+	static FThreadFilterFunc MakeThreadFilterInclusive(const TSet<uint32>& IncludedThreads);
+
+	/**
+	 * Makes a FThreadFilterFunc using a set of excluded list of threads.
+	 * Note: The set is referenced in the returned function.
+	 * @param ExcludedThreads The set of thread ids to be rejected by filter.
+	 * @return The FThreadFilterFunc function.
+	 */
+	static FThreadFilterFunc MakeThreadFilterExclusive(const TSet<uint32>& ExcludedThreads);
+
+	//////////////////////////////////////////////////////////////////////
+	// Utilities to make FTimingEventFilterFunc filters.
+
+	FTimingEventFilterFunc MakeTimingEventFilterByTimersInclusive(const FString& InFilterString, TSet<uint32>& OutIncludedTimers) const;
+
+	/**
+	 * Makes a FTimingEventFilterFunc using a set of included list of timers.
+	 * Note: The set is referenced in the returned function.
+	 * @param IncludedTimers The set of timer ids to be accepted by filter.
+	 * @return The TimingEventFilter function.
+	 */ 
+	static FTimingEventFilterFunc MakeTimingEventFilterByTimersInclusive(const TSet<uint32>& IncludedTimers);
+
+	/**
+	 * Makes a FTimingEventFilterFunc using a set of excluded list of timers.
+	 * Note: The set is referenced in the returned function.
+	 * @param ExcludedTimers The set of timer ids to be rejected by filter.
+	 * @return The TimingEventFilter function. Can be nullptr (i.e. no filter).
+	 */ 
+	static FTimingEventFilterFunc MakeTimingEventFilterByTimersExclusive(const TSet<uint32>& ExcludedTimers);
+
+private:
+	IFileHandle* OpenExportFile(const TCHAR* InFilename) const;
+	void Error(const FText& InMessage) const;
+	static void AppendString(FUtf8StringBuilder& StringBuilder, const TCHAR* String, UTF8CHAR Separator);
+	void ExportTimingEvents_InitColumns() const;
+	void ExportTimingEvents_WriteHeader(FExportTimingEventsInternalParams& Params) const;
+	int32 ExportTimingEvents_WriteEvents(FExportTimingEventsInternalParams& Params) const;
+
+private:
+	const TraceServices::IAnalysisSession& Session;
+	mutable TSet<FName> ExportTimingEventsColumns;
+	mutable TArray<FName> ExportTimingEventsDefaultColumns;
+
+	static const FName ExportTimingEvents_ThreadIdColumn;
+	static const FName ExportTimingEvents_ThreadNameColumn;
+	static const FName ExportTimingEvents_TimerIdColumn;
+	static const FName ExportTimingEvents_TimerNameColumn;
+	static const FName ExportTimingEvents_StartTimeColumn;
+	static const FName ExportTimingEvents_EndTimeColumn;
+	static const FName ExportTimingEvents_DurationColumn;
+	static const FName ExportTimingEvents_DepthColumn;
+};
+
+////////////////////////////////////////////////////////////////////////////////////////////////////
+
+} // namespace Insights
diff --git a/Engine/Source/Developer/TraceInsights/Private/Insights/Widgets/STimersView.cpp b/Engine/Source/Developer/TraceInsights/Private/Insights/Widgets/STimersView.cpp
--- a/Engine/Source/Developer/TraceInsights/Private/Insights/Widgets/STimersView.cpp
+++ b/Engine/Source/Developer/TraceInsights/Private/Insights/Widgets/STimersView.cpp
@@ -38,6 +38,7 @@
 #include "Insights/ViewModels/TimerAggregation.h"
 #include "Insights/ViewModels/TimerNodeHelper.h"
 #include "Insights/ViewModels/TimersViewColumnFactory.h"
+#include "Insights/ViewModels/TimingExporter.h"
 #include "Insights/ViewModels/TimingGraphTrack.h"
 #include "Insights/Widgets/SAsyncOperationStatus.h"
 #include "Insights/Widgets/STimersViewTooltip.h"
@@ -2260,8 +2261,12 @@ void STimersView::ContextMenu_Export_Execute()
 	const FString DialogTitle = LOCTEXT("Export_Title", "Export Aggregated Timer Stats").ToString();
 	const FString DefaultFile = TEXT("TimerStats.tsv");
 	FString Filename;
-	IFileHandle* ExportFileHandle = OpenSaveTextFileDialog(DialogTitle, DefaultFile, Filename);
+	if (!OpenSaveTextFileDialog(DialogTitle, DefaultFile, Filename))
+	{
+		return;
+	}
 
+	IFileHandle* ExportFileHandle = OpenExportFile(*Filename);
 	if (!ExportFileHandle)
 	{
 		return;
@@ -2282,7 +2287,7 @@ void STimersView::ContextMenu_Export_Execute()
 	constexpr TCHAR QuotationMarkBegin = TEXT('\"');
 	constexpr TCHAR QuotationMarkEnd = TEXT('\"');
 
-	TStringBuilder<512> StringBuilder;
+	TStringBuilder<1024> StringBuilder;
 
 	TArray<TSharedRef<Insights::FTableColumn>> VisibleColumns;
 	Table->GetVisibleColumns(VisibleColumns);
@@ -2369,120 +2374,6 @@ void STimersView::ContextMenu_Export_Execute()
 
 ////////////////////////////////////////////////////////////////////////////////////////////////////
 
-namespace Insights {
-
-struct FExportTimingEventsHeaderParams
-{
-	IFileHandle* ExportFileHandle;
-	const UTF8CHAR Separator;
-	const UTF8CHAR LineEnd;
-	TUtf8StringBuilder<512>& StringBuilder;
-};
-
-void ExportTimingEventsHeader(FExportTimingEventsHeaderParams& Params)
-{
-	Params.StringBuilder.Append(UTF8TEXT("ThreadId"));
-	Params.StringBuilder.AppendChar(Params.Separator);
-	Params.StringBuilder.Append(UTF8TEXT("TimerId"));
-	Params.StringBuilder.AppendChar(Params.Separator);
-	Params.StringBuilder.Append(UTF8TEXT("StartTime"));
-	Params.StringBuilder.AppendChar(Params.Separator);
-	Params.StringBuilder.Append(UTF8TEXT("EndTime"));
-	Params.StringBuilder.AppendChar(Params.Separator);
-	Params.StringBuilder.Append(UTF8TEXT("Depth"));
-	Params.StringBuilder.AppendChar(Params.LineEnd);
-
-	Params.ExportFileHandle->Write((const uint8*)Params.StringBuilder.ToString(), Params.StringBuilder.Len() * sizeof(UTF8CHAR));
-}
-
-struct FExportTimingEventsEnumerateParams
-{
-	TSharedPtr<const TraceServices::IAnalysisSession> Session;
-	IFileHandle* ExportFileHandle;
-	const UTF8CHAR Separator;
-	const UTF8CHAR LineEnd;
-	TUtf8StringBuilder<512>& StringBuilder;
-	double IntervalStartTime;
-	double IntervalEndTime;
-	TFunction<bool(uint32 ThreadId)> ThreadFilter;
-	TFunction<bool(double EventStartTime, double EventEndTime, uint32 EventDepth, const TraceServices::FTimingProfilerEvent& Event)> TimingEventFilter;
-	uint32 ThreadId;
-};
-
-void ExportTimingEventsEnumerate(FExportTimingEventsEnumerateParams& Params)
-{
-	if (Params.Session.IsValid() && TraceServices::ReadTimingProfilerProvider(*Params.Session.Get()))
-	{
-		TraceServices::FAnalysisSessionReadScope SessionReadScope(*Params.Session.Get());
-
-		const TraceServices::ITimingProfilerProvider& TimingProfilerProvider = *TraceServices::ReadTimingProfilerProvider(*Params.Session.Get());
-
-		auto TimelineEnumerator = [&Params](const TraceServices::ITimingProfilerProvider::Timeline& Timeline)
-		{
-			// Iterate timing events.
-			Timeline.EnumerateEvents(Params.IntervalStartTime, Params.IntervalEndTime,
-				[&Params](double EventStartTime, double EventEndTime, uint32 EventDepth, const TraceServices::FTimingProfilerEvent& Event)
-				{
-					if (!Params.TimingEventFilter || Params.TimingEventFilter(EventStartTime, EventEndTime, EventDepth, Event))
-					{
-						Params.StringBuilder.Reset();
-						Params.StringBuilder.Appendf(UTF8TEXT("%u"), Params.ThreadId);
-						Params.StringBuilder.AppendChar(Params.Separator);
-						Params.StringBuilder.Appendf(UTF8TEXT("%u"), Event.TimerIndex);
-						Params.StringBuilder.AppendChar(Params.Separator);
-						Params.StringBuilder.Appendf(UTF8TEXT("%.9g"), EventStartTime);
-						Params.StringBuilder.AppendChar(Params.Separator);
-						Params.StringBuilder.Appendf(UTF8TEXT("%.9g"), EventEndTime);
-						Params.StringBuilder.AppendChar(Params.Separator);
-						Params.StringBuilder.Appendf(UTF8TEXT("%u"), EventDepth);
-						Params.StringBuilder.AppendChar(Params.LineEnd);
-
-						Params.ExportFileHandle->Write((const uint8*)Params.StringBuilder.ToString(), Params.StringBuilder.Len() * sizeof(UTF8CHAR));
-					}
-
-					return TraceServices::EEventEnumerate::Continue;
-				});
-		};
-
-		// Iterate the GPU timelines.
-		{
-			if (!Params.ThreadFilter || Params.ThreadFilter(FGpuTimingTrack::Gpu1ThreadId))
-			{
-				Params.ThreadId = FGpuTimingTrack::Gpu1ThreadId;
-				uint32 GpuTimelineIndex1 = 0;
-				TimingProfilerProvider.GetGpuTimelineIndex(GpuTimelineIndex1);
-				TimingProfilerProvider.ReadTimeline(GpuTimelineIndex1, TimelineEnumerator);
-			}
-
-			if (!Params.ThreadFilter || Params.ThreadFilter(FGpuTimingTrack::Gpu2ThreadId))
-			{
-				Params.ThreadId = FGpuTimingTrack::Gpu2ThreadId;
-				uint32 GpuTimelineIndex2 = 0;
-				TimingProfilerProvider.GetGpu2TimelineIndex(GpuTimelineIndex2);
-				TimingProfilerProvider.ReadTimeline(GpuTimelineIndex2, TimelineEnumerator);
-			}
-		}
-
-		// Iterate the CPU threads and their corresponding timelines.
-		const TraceServices::IThreadProvider& ThreadProvider = TraceServices::ReadThreadProvider(*Params.Session.Get());
-		ThreadProvider.EnumerateThreads(
-			[&TimelineEnumerator, &Params, &TimingProfilerProvider](const TraceServices::FThreadInfo& ThreadInfo)
-			{
-				if (!Params.ThreadFilter || Params.ThreadFilter(ThreadInfo.Id))
-				{
-					Params.ThreadId = ThreadInfo.Id;
-					uint32 CpuTimelineIndex = 0;
-					TimingProfilerProvider.GetCpuThreadTimelineIndex(ThreadInfo.Id, CpuTimelineIndex);
-					TimingProfilerProvider.ReadTimeline(CpuTimelineIndex, TimelineEnumerator);
-				}
-			});
-	}
-}
-
-} // namespace Insights
-
-////////////////////////////////////////////////////////////////////////////////////////////////////
-
 bool STimersView::ContextMenu_ExportTimingEventsSelection_CanExecute() const
 {
 	const TArray<FTimerNodePtr> SelectedNodes = TreeView->GetSelectedItems();
@@ -2511,181 +2402,115 @@ void STimersView::AddTimerNodeRecursive(FTimerNodePtr InNode, TSet<uint32>& InOu
 
 void STimersView::ContextMenu_ExportTimingEventsSelection_Execute() const
 {
-	const FString DialogTitle = LOCTEXT("ExportTimingEventsSelection_Title", "Export Timing Events (Selection)").ToString();
-	const FString DefaultFile = TEXT("TimingEvents.tsv");
-	FString Filename;
-	IFileHandle* ExportFileHandle = OpenSaveTextFileDialog(DialogTitle, DefaultFile, Filename);
-
-	if (!ExportFileHandle)
+	if (!Session.IsValid())
 	{
 		return;
 	}
 
-	FStopwatch Stopwatch;
-	Stopwatch.Start();
-
-	UTF8CHAR Separator = UTF8CHAR('\t');
-	if (Filename.EndsWith(TEXT(".csv")))
+	const FString DialogTitle = LOCTEXT("ExportTimingEventsSelection_Title", "Export Timing Events (Selection)").ToString();
+	const FString DefaultFile = TEXT("TimingEvents.tsv");
+	FString Filename;
+	if (!OpenSaveTextFileDialog(DialogTitle, DefaultFile, Filename))
 	{
-		Separator = UTF8CHAR(',');
+		return;
 	}
-	const UTF8CHAR LineEnd = UTF8CHAR('\n');
 
-	TUtf8StringBuilder<512> StringBuilder;
+	Insights::FTimingExporter Exporter(*Session.Get());
+	Insights::FTimingExporter::FExportTimingEventsParams Params; // default columns, all timing events
 
-	// Write header.
-	{
-		Insights::FExportTimingEventsHeaderParams HeaderParams = { ExportFileHandle, Separator, LineEnd, StringBuilder };
-		Insights::ExportTimingEventsHeader(HeaderParams);
-	}
+	////////////////////////////////////////////////////////////////////////////////////////////////////
 
-	// Write values.
-	if (Session.IsValid())
-	{
-		TSharedPtr<STimingProfilerWindow> Wnd = FTimingProfilerManager::Get()->GetProfilerWindow();
-		TSharedPtr<STimingView> TimingView = Wnd.IsValid() ? Wnd->GetTimingView() : nullptr;
+	TSharedPtr<STimingProfilerWindow> Wnd = FTimingProfilerManager::Get()->GetProfilerWindow();
+	TSharedPtr<STimingView> TimingView = Wnd.IsValid() ? Wnd->GetTimingView() : nullptr;
 
-		Insights::FExportTimingEventsEnumerateParams EnumerateParams = { Session, ExportFileHandle, Separator, LineEnd, StringBuilder, 0.0, 0.0, nullptr, nullptr, 0 };
+	////////////////////////////////////////////////////////////////////////////////////////////////////
+	// Filter by thread (visible Gpu/Cpu tracks in TimingView).
 
-		////////////////////////////////////////////////////////////////////////////////////////////////////
-		// Filter by thread (visible Gpu/Cpu tracks in TimingView).
+	// These variables needs to be in the same scope with the call to Exporter.ExportTimingEventsAsText()
+	// (becasue they are referenced in the ThreadFilter lambda function).
+	TSet<uint32> IncludedThreads;
+	TSet<uint32> ExcludedThreads;
 
-		// Filter sets to be used in ExportTimingEventsEnumerate.
-		// It needs to be in the same scope (as it is referenced in the lambda filter)!
-		TSet<uint32> IncludedThreads;
-		TSet<uint32> ExcludedThreads;
+	if (TimingView.IsValid())
+	{
+		// Add available Gpu threads to the ExcludedThreads list.
+		ExcludedThreads.Add(FGpuTimingTrack::Gpu1ThreadId);
+		ExcludedThreads.Add(FGpuTimingTrack::Gpu2ThreadId);
 
-		if (TimingView.IsValid())
+		// Add available Cpu threads to the ExcludedThreads list.
 		{
-			// Add available Gpu threads to the ExcludedThreads list.
-			ExcludedThreads.Add(FGpuTimingTrack::Gpu1ThreadId);
-			ExcludedThreads.Add(FGpuTimingTrack::Gpu2ThreadId);
-
-			// Add available Cpu threads to the ExcludedThreads list.
-			{
-				TraceServices::FAnalysisSessionReadScope SessionReadScope(*Session.Get());
-				const TraceServices::IThreadProvider& ThreadProvider = TraceServices::ReadThreadProvider(*Session.Get());
-				ThreadProvider.EnumerateThreads(
-					[&ExcludedThreads](const TraceServices::FThreadInfo& ThreadInfo)
-					{
-						ExcludedThreads.Add(ThreadInfo.Id);
-					});
-			}
-
-			// Move the threads corresponding to visible Cpu/Gpu tracks to the IncludedThreads list.
-			TimingView->EnumerateAllTracks([&IncludedThreads, &ExcludedThreads](TSharedPtr<FBaseTimingTrack>& Track) -> bool
+			TraceServices::FAnalysisSessionReadScope SessionReadScope(*Session.Get());
+			const TraceServices::IThreadProvider& ThreadProvider = TraceServices::ReadThreadProvider(*Session.Get());
+			ThreadProvider.EnumerateThreads(
+				[&ExcludedThreads](const TraceServices::FThreadInfo& ThreadInfo)
 				{
-					if (Track->IsVisible() && Track->Is<FThreadTimingTrack>())
-					{
-						const uint32 ThreadId = Track->As<FThreadTimingTrack>().GetThreadId();
-						ExcludedThreads.Remove(ThreadId);
-						IncludedThreads.Add(ThreadId);
-					}
-					return true;
+					ExcludedThreads.Add(ThreadInfo.Id);
 				});
+		}
 
-			if (IncludedThreads.Num() < ExcludedThreads.Num())
-			{
-				if (IncludedThreads.Num() == 1)
-				{
-					const uint32 IncludedThreadId = IncludedThreads[FSetElementId::FromInteger(0)];
-					EnumerateParams.ThreadFilter = [IncludedThreadId](uint32 ThreadId) -> bool
-					{
-						return ThreadId == IncludedThreadId;
-					};
-				}
-				else
-				{
-					EnumerateParams.ThreadFilter = [&IncludedThreads](uint32 ThreadId) -> bool
-					{
-						return IncludedThreads.Contains(ThreadId);
-					};
-				}
-			}
-			else
+		// Move the threads corresponding to visible Cpu/Gpu tracks to the IncludedThreads list.
+		TimingView->EnumerateAllTracks([&IncludedThreads, &ExcludedThreads](TSharedPtr<FBaseTimingTrack>& Track) -> bool
 			{
-				if (ExcludedThreads.Num() == 1)
-				{
-					const uint32 ExcludedThreadId = ExcludedThreads[FSetElementId::FromInteger(0)];
-					EnumerateParams.ThreadFilter = [ExcludedThreadId](uint32 ThreadId) -> bool
-					{
-						return ThreadId != ExcludedThreadId;
-					};
-				}
-				else
+				if (Track->IsVisible() && Track->Is<FThreadTimingTrack>())
 				{
-					EnumerateParams.ThreadFilter = [&ExcludedThreads](uint32 ThreadId) -> bool
-					{
-						return !ExcludedThreads.Contains(ThreadId);
-					};
+					const uint32 ThreadId = Track->As<FThreadTimingTrack>().GetThreadId();
+					ExcludedThreads.Remove(ThreadId);
+					IncludedThreads.Add(ThreadId);
 				}
-			}
-		}
-
-		// Debug/test filters.
-		//EnumerateParams.ThreadFilter = [](uint32 ThreadId) -> bool { return ThreadId == 2; };
-
-		////////////////////////////////////////////////////////////////////////////////////////////////////
-		// Filter by timing event (e.g.: by timer, by duration, by depth).
-
-		// Filter set to be used in ExportTimingEventsEnumerate.
-		// It needs to be in the same scope (as it is referenced in the lambda filter)!
-		TSet<uint32> IncludedTimers;
-
-		TArray<FTimerNodePtr> SelectedNodes = TreeView->GetSelectedItems();
-		for (FTimerNodePtr Node : SelectedNodes)
-		{
-			AddTimerNodeRecursive(Node, IncludedTimers);
-		}
+				return true;
+			});
 
-		if (IncludedTimers.Num() == 1)
+		if (IncludedThreads.Num() < ExcludedThreads.Num())
 		{
-			const uint32 IncludedTimerId = IncludedTimers[FSetElementId::FromInteger(0)];
-			EnumerateParams.TimingEventFilter = [IncludedTimerId](double EventStartTime, double EventEndTime, uint32 EventDepth, const TraceServices::FTimingProfilerEvent& Event) -> bool
-			{
-				return Event.TimerIndex == IncludedTimerId;
-			};
+			Params.ThreadFilter = Insights::FTimingExporter::MakeThreadFilterInclusive(IncludedThreads);
 		}
 		else
 		{
-			EnumerateParams.TimingEventFilter = [&IncludedTimers](double EventStartTime, double EventEndTime, uint32 EventDepth, const TraceServices::FTimingProfilerEvent& Event) -> bool
-			{
-				return IncludedTimers.Contains(Event.TimerIndex);
-			};
+			Params.ThreadFilter = Insights::FTimingExporter::MakeThreadFilterExclusive(ExcludedThreads);
 		}
+	}
 
-		// Debug/test filters.
-		//EnumerateParams.TimingEventFilter = [](double EventStartTime, double EventEndTime, uint32 EventDepth, const TraceServices::FTimingProfilerEvent& Event) -> bool { return Event.TimerIndex < 100; };
-		//EnumerateParams.TimingEventFilter = [](double EventStartTime, double EventEndTime, uint32 EventDepth, const TraceServices::FTimingProfilerEvent& Event) -> bool { return EventEndTime - EventStartTime > 0.001; };
+	// Debug/test filters.
+	//Params.ThreadFilter = [](uint32 ThreadId) -> bool { return ThreadId == 2; };
 
-		////////////////////////////////////////////////////////////////////////////////////////////////////
-		// Limit the time interval for enumeration (if a time range selection is made in Timing view).
+	////////////////////////////////////////////////////////////////////////////////////////////////////
+	// Filter by timing event (e.g.: by timer, by duration, by depth).
 
-		EnumerateParams.IntervalStartTime = -std::numeric_limits<double>::infinity();
-		EnumerateParams.IntervalEndTime = +std::numeric_limits<double>::infinity();
-		if (TimingView.IsValid())
-		{
-			const double SelectionStartTime = TimingView->GetSelectionStartTime();
-			const double SelectionEndTime = TimingView->GetSelectionEndTime();
-			if (SelectionStartTime < SelectionEndTime)
-			{
-				EnumerateParams.IntervalStartTime = SelectionStartTime;
-				EnumerateParams.IntervalEndTime = SelectionEndTime;
-			}
-		}
+	// This variable needs to be in the same scope with the call to Exporter.ExportTimingEventsAsText()
+	// (becasue it is referenced in the TimingEventFilter lambda function).
+	TSet<uint32> IncludedTimers;
+
+	TArray<FTimerNodePtr> SelectedNodes = TreeView->GetSelectedItems();
+	for (FTimerNodePtr Node : SelectedNodes)
+	{
+		AddTimerNodeRecursive(Node, IncludedTimers);
+	}
+
+	Params.TimingEventFilter = Insights::FTimingExporter::MakeTimingEventFilterByTimersInclusive(IncludedTimers);
+
+	// Debug/test filters.
+	//Params.TimingEventFilter = [](double EventStartTime, double EventEndTime, uint32 EventDepth, const TraceServices::FTimingProfilerEvent& Event) -> bool { return Event.TimerIndex < 100; };
+	//Params.TimingEventFilter = [](double EventStartTime, double EventEndTime, uint32 EventDepth, const TraceServices::FTimingProfilerEvent& Event) -> bool { return EventEndTime - EventStartTime > 0.001; };
 
-		////////////////////////////////////////////////////////////////////////////////////////////////////
+	////////////////////////////////////////////////////////////////////////////////////////////////////
+	// Limit the time interval for enumeration (if a time range selection is made in Timing view).
 
-		ExportTimingEventsEnumerate(EnumerateParams);
+	Params.IntervalStartTime = -std::numeric_limits<double>::infinity();
+	Params.IntervalEndTime = +std::numeric_limits<double>::infinity();
+	if (TimingView.IsValid())
+	{
+		const double SelectionStartTime = TimingView->GetSelectionStartTime();
+		const double SelectionEndTime = TimingView->GetSelectionEndTime();
+		if (SelectionStartTime < SelectionEndTime)
+		{
+			Params.IntervalStartTime = SelectionStartTime;
+			Params.IntervalEndTime = SelectionEndTime;
+		}
 	}
 
-	ExportFileHandle->Flush();
-	delete ExportFileHandle;
-	ExportFileHandle = nullptr;
+	////////////////////////////////////////////////////////////////////////////////////////////////////
 
-	Stopwatch.Stop();
-	const double TotalTime = Stopwatch.GetAccumulatedTime();
-	UE_LOG(TraceInsights, Log, TEXT("Exported selected timing events to file in %.3fs (\"%s\")."), TotalTime, *Filename);
+	Exporter.ExportTimingEventsAsText(*Filename, Params);
 }
 
 ////////////////////////////////////////////////////////////////////////////////////////////////////
@@ -2699,52 +2524,22 @@ bool STimersView::ContextMenu_ExportTimingEvents_CanExecute() const
 
 void STimersView::ContextMenu_ExportTimingEvents_Execute() const
 {
-	const FString DialogTitle = LOCTEXT("ExportTimingEvents_Title", "Export Timing Events (All)").ToString();
-	const FString DefaultFile = TEXT("TimingEvents.tsv");
-	FString Filename;
-	IFileHandle* ExportFileHandle = OpenSaveTextFileDialog(DialogTitle, DefaultFile, Filename);
-
-	if (!ExportFileHandle)
+	if (!Session.IsValid())
 	{
 		return;
 	}
 
-	FStopwatch Stopwatch;
-	Stopwatch.Start();
-
-	UTF8CHAR Separator = UTF8CHAR('\t');
-	if (Filename.EndsWith(TEXT(".csv")))
-	{
-		Separator = UTF8CHAR(',');
-	}
-	const UTF8CHAR LineEnd = UTF8CHAR('\n');
-
-	TUtf8StringBuilder<512> StringBuilder;
-
-	// Write header.
-	{
-		Insights::FExportTimingEventsHeaderParams HeaderParams = { ExportFileHandle, Separator, LineEnd, StringBuilder };
-		Insights::ExportTimingEventsHeader(HeaderParams);
-	}
-
-	// Write values.
-	if (Session.IsValid())
+	const FString DialogTitle = LOCTEXT("ExportTimingEvents_Title", "Export Timing Events (All)").ToString();
+	const FString DefaultFile = TEXT("TimingEvents.tsv");
+	FString Filename;
+	if (!OpenSaveTextFileDialog(DialogTitle, DefaultFile, Filename))
 	{
-		Insights::FExportTimingEventsEnumerateParams EnumerateParams = { Session, ExportFileHandle, Separator, LineEnd, StringBuilder, 0.0, 0.0, nullptr, nullptr, 0 };
-
-		EnumerateParams.IntervalStartTime = -std::numeric_limits<double>::infinity();
-		EnumerateParams.IntervalEndTime = +std::numeric_limits<double>::infinity();
-
-		Insights::ExportTimingEventsEnumerate(EnumerateParams);
+		return;
 	}
 
-	ExportFileHandle->Flush();
-	delete ExportFileHandle;
-	ExportFileHandle = nullptr;
-
-	Stopwatch.Stop();
-	const double TotalTime = Stopwatch.GetAccumulatedTime();
-	UE_LOG(TraceInsights, Log, TEXT("Exported all timing events to file in %.3fs (\"%s\")."), TotalTime, *Filename);
+	Insights::FTimingExporter Exporter(*Session.Get());
+	Insights::FTimingExporter::FExportTimingEventsParams Params; // default columns, all timing events
+	Exporter.ExportTimingEventsAsText(*Filename, Params);
 }
 
 ////////////////////////////////////////////////////////////////////////////////////////////////////
@@ -2758,87 +2553,22 @@ bool STimersView::ContextMenu_ExportThreads_CanExecute() const
 
 void STimersView::ContextMenu_ExportThreads_Execute() const
 {
-	const FString DialogTitle = LOCTEXT("ExportThreads_Title", "Export Threads").ToString();
-	const FString DefaultFile = TEXT("Threads.tsv");
-	FString Filename;
-	IFileHandle* ExportFileHandle = OpenSaveTextFileDialog(DialogTitle, DefaultFile, Filename);
-
-	if (!ExportFileHandle)
+	if (!Session.IsValid())
 	{
 		return;
 	}
 
-	FStopwatch Stopwatch;
-	Stopwatch.Start();
-
-	UTF8CHAR Separator = UTF8CHAR('\t');
-	if (Filename.EndsWith(TEXT(".csv")))
-	{
-		Separator = UTF8CHAR(',');
-	}
-	const UTF8CHAR LineEnd = UTF8CHAR('\n');
-
-	TUtf8StringBuilder<512> StringBuilder;
-
-	// Write header.
-	{
-		StringBuilder.Append(UTF8TEXT("ThreadId"));
-		StringBuilder.AppendChar(Separator);
-		StringBuilder.Append(UTF8TEXT("Name"));
-		StringBuilder.AppendChar(Separator);
-		StringBuilder.Append(UTF8TEXT("GroupName"));
-		StringBuilder.AppendChar(LineEnd);
-
-		ExportFileHandle->Write((const uint8*)StringBuilder.ToString(), StringBuilder.Len() * sizeof(UTF8CHAR));
-	}
-
-	// Write values.
-	if (Session.IsValid())
+	const FString DialogTitle = LOCTEXT("ExportThreads_Title", "Export Threads").ToString();
+	const FString DefaultFile = TEXT("Threads.tsv");
+	FString Filename;
+	if (!OpenSaveTextFileDialog(DialogTitle, DefaultFile, Filename))
 	{
-		StringBuilder.Reset();
-		StringBuilder.Appendf(UTF8TEXT("%u"), FGpuTimingTrack::Gpu1ThreadId);
-		StringBuilder.AppendChar(Separator);
-		StringBuilder.Append(UTF8TEXT("GPU1"));
-		StringBuilder.AppendChar(Separator);
-		StringBuilder.Append(UTF8TEXT("GPU"));
-		StringBuilder.AppendChar(LineEnd);
-		ExportFileHandle->Write((const uint8*)StringBuilder.ToString(), StringBuilder.Len() * sizeof(UTF8CHAR));
-
-		StringBuilder.Reset();
-		StringBuilder.Appendf(UTF8TEXT("%u"), FGpuTimingTrack::Gpu2ThreadId);
-		StringBuilder.AppendChar(Separator);
-		StringBuilder.Append(UTF8TEXT("GPU2"));
-		StringBuilder.AppendChar(Separator);
-		StringBuilder.Append(UTF8TEXT("GPU"));
-		StringBuilder.AppendChar(LineEnd);
-		ExportFileHandle->Write((const uint8*)StringBuilder.ToString(), StringBuilder.Len() * sizeof(UTF8CHAR));
-
-		// Iterate the CPU threads.
-		{
-			TraceServices::FAnalysisSessionReadScope SessionReadScope(*Session.Get());
-			const TraceServices::IThreadProvider& ThreadProvider = TraceServices::ReadThreadProvider(*Session.Get());
-			ThreadProvider.EnumerateThreads(
-				[&](const TraceServices::FThreadInfo& ThreadInfo)
-				{
-					StringBuilder.Reset();
-					StringBuilder.Appendf(UTF8TEXT("%u"), ThreadInfo.Id);
-					StringBuilder.AppendChar(Separator);
-					StringBuilder.Append(TCHAR_TO_UTF8(ThreadInfo.Name));
-					StringBuilder.AppendChar(Separator);
-					StringBuilder.Append(TCHAR_TO_UTF8(ThreadInfo.GroupName));
-					StringBuilder.AppendChar(LineEnd);
-					ExportFileHandle->Write((const uint8*)StringBuilder.ToString(), StringBuilder.Len() * sizeof(UTF8CHAR));
-				});
-		}
+		return;
 	}
 
-	ExportFileHandle->Flush();
-	delete ExportFileHandle;
-	ExportFileHandle = nullptr;
-
-	Stopwatch.Stop();
-	const double TotalTime = Stopwatch.GetAccumulatedTime();
-	UE_LOG(TraceInsights, Log, TEXT("Exported the thread list to file in %.3fs (\"%s\")."), TotalTime, *Filename);
+	Insights::FTimingExporter Exporter(*Session.Get());
+	Insights::FTimingExporter::FExportThreadsParams Params; // default
+	Exporter.ExportThreadsAsText(*Filename, Params);
 }
 
 ////////////////////////////////////////////////////////////////////////////////////////////////////
@@ -2852,85 +2582,27 @@ bool STimersView::ContextMenu_ExportTimers_CanExecute() const
 
 void STimersView::ContextMenu_ExportTimers_Execute() const
 {
-	const FString DialogTitle = LOCTEXT("ExportTimers_Title", "Export Timers").ToString();
-	const FString DefaultFile = TEXT("Timers.tsv");
-	FString Filename;
-	IFileHandle* ExportFileHandle = OpenSaveTextFileDialog(DialogTitle, DefaultFile, Filename);
-
-	if (!ExportFileHandle)
+	if (!Session.IsValid())
 	{
 		return;
 	}
 
-	FStopwatch Stopwatch;
-	Stopwatch.Start();
-
-	UTF8CHAR Separator = UTF8CHAR('\t');
-	if (Filename.EndsWith(TEXT(".csv")))
-	{
-		Separator = UTF8CHAR(',');
-	}
-	const UTF8CHAR LineEnd = UTF8CHAR('\n');
-
-	TUtf8StringBuilder<512> StringBuilder;
-
-	// Write header.
-	{
-		StringBuilder.Append(UTF8TEXT("TimerId"));
-		StringBuilder.AppendChar(Separator);
-		StringBuilder.Append(UTF8TEXT("Type"));
-		StringBuilder.AppendChar(Separator);
-		StringBuilder.Append(UTF8TEXT("Name"));
-		StringBuilder.AppendChar(Separator);
-		StringBuilder.Append(UTF8TEXT("File"));
-		StringBuilder.AppendChar(Separator);
-		StringBuilder.Append(UTF8TEXT("Line"));
-		StringBuilder.AppendChar(LineEnd);
-
-		ExportFileHandle->Write((const uint8*)StringBuilder.ToString(), StringBuilder.Len() * sizeof(UTF8CHAR));
-	}
-
-	// Write values.
-	if (Session.IsValid() && TraceServices::ReadTimingProfilerProvider(*Session.Get()))
+	const FString DialogTitle = LOCTEXT("ExportTimers_Title", "Export Timers").ToString();
+	const FString DefaultFile = TEXT("Timers.tsv");
+	FString Filename;
+	if (!OpenSaveTextFileDialog(DialogTitle, DefaultFile, Filename))
 	{
-		TraceServices::FAnalysisSessionReadScope SessionReadScope(*Session.Get());
-
-		const TraceServices::ITimingProfilerProvider& TimingProfilerProvider = *TraceServices::ReadTimingProfilerProvider(*Session.Get());
-
-		const TraceServices::ITimingProfilerTimerReader* TimerReader;
-		TimingProfilerProvider.ReadTimers([&TimerReader](const TraceServices::ITimingProfilerTimerReader& Out) { TimerReader = &Out; });
-
-		const uint32 TimerCount = TimerReader->GetTimerCount();
-		for (uint32 TimerIndex = 0; TimerIndex < TimerCount; ++TimerIndex)
-		{
-			const TraceServices::FTimingProfilerTimer& Timer = *(TimerReader->GetTimer(TimerIndex));
-			StringBuilder.Reset();
-			StringBuilder.Appendf(UTF8TEXT("%u"), Timer.Id);
-			StringBuilder.AppendChar(Separator);
-			StringBuilder.Append(Timer.IsGpuTimer ? UTF8TEXT("GPU") : UTF8TEXT("CPU"));
-			StringBuilder.AppendChar(Separator);
-			StringBuilder.Append(TCHAR_TO_UTF8(Timer.Name));
-			StringBuilder.AppendChar(Separator);
-			StringBuilder.Append(Timer.File ? (const UTF8CHAR*)TCHAR_TO_UTF8(Timer.File) : UTF8TEXT(""));
-			StringBuilder.AppendChar(Separator);
-			StringBuilder.Appendf(UTF8TEXT("%u"), Timer.Line);
-			StringBuilder.AppendChar(LineEnd);
-			ExportFileHandle->Write((const uint8*)StringBuilder.ToString(), StringBuilder.Len() * sizeof(UTF8CHAR));
-		}
+		return;
 	}
 
-	ExportFileHandle->Flush();
-	delete ExportFileHandle;
-	ExportFileHandle = nullptr;
-
-	Stopwatch.Stop();
-	const double TotalTime = Stopwatch.GetAccumulatedTime();
-	UE_LOG(TraceInsights, Log, TEXT("Exported the timers list to file in %.3fs (\"%s\")."), TotalTime, *Filename);
+	Insights::FTimingExporter Exporter(*Session.Get());
+	Insights::FTimingExporter::FExportTimersParams Params; // default
+	Exporter.ExportTimersAsText(*Filename, Params);
 }
 
 ////////////////////////////////////////////////////////////////////////////////////////////////////
 
-IFileHandle* STimersView::OpenSaveTextFileDialog(const FString& InDialogTitle, const FString& InDefaultFile, FString& OutFilename) const
+bool STimersView::OpenSaveTextFileDialog(const FString& InDialogTitle, const FString& InDefaultFile, FString& OutFilename) const
 {
 	TArray<FString> SaveFilenames;
 	bool bDialogResult = false;
@@ -2952,12 +2624,18 @@ IFileHandle* STimersView::OpenSaveTextFileDialog(const FString& InDialogTitle, c
 
 	if (!bDialogResult || SaveFilenames.Num() == 0)
 	{
-		return nullptr;
+		return false;
 	}
 
 	OutFilename = SaveFilenames[0];
+	return true;
+}
 
-	IFileHandle* ExportFileHandle = FPlatformFileManager::Get().GetPlatformFile().OpenWrite(*OutFilename);
+////////////////////////////////////////////////////////////////////////////////////////////////////
+
+IFileHandle* STimersView::OpenExportFile(const TCHAR* InFilename) const
+{
+	IFileHandle* ExportFileHandle = FPlatformFileManager::Get().GetPlatformFile().OpenWrite(InFilename);
 
 	if (ExportFileHandle == nullptr)
 	{
diff --git a/Engine/Source/Developer/TraceInsights/Private/Insights/Widgets/STimersView.h b/Engine/Source/Developer/TraceInsights/Private/Insights/Widgets/STimersView.h
--- a/Engine/Source/Developer/TraceInsights/Private/Insights/Widgets/STimersView.h
+++ b/Engine/Source/Developer/TraceInsights/Private/Insights/Widgets/STimersView.h
@@ -139,7 +139,8 @@ private:
 	bool ContextMenu_ExportTimers_CanExecute() const;
 	void ContextMenu_ExportTimers_Execute() const;
 
-	class IFileHandle* OpenSaveTextFileDialog(const FString& InDialogTitle, const FString& InDefaultFile, FString& OutFilename) const;
+	bool OpenSaveTextFileDialog(const FString& InDialogTitle, const FString& InDefaultFile, FString& OutFilename) const;
+	class IFileHandle* OpenExportFile(const TCHAR* InFilename) const;
 
 	bool ContextMenu_OpenSource_CanExecute() const;
 	void ContextMenu_OpenSource_Execute() const;
