[Insights] Changed Ctrl+C for the selected CPU/GPU timing event to copy also the duration and the metadata (if available) of the timing event (i.e. same text as displayed in the Timing view).

diff --git a/Engine/Source/Developer/TraceInsights/Private/Insights/ViewModels/ThreadTimingTrack.cpp b/Engine/Source/Developer/TraceInsights/Private/Insights/ViewModels/ThreadTimingTrack.cpp
--- a/Engine/Source/Developer/TraceInsights/Private/Insights/ViewModels/ThreadTimingTrack.cpp
+++ b/Engine/Source/Developer/TraceInsights/Private/Insights/ViewModels/ThreadTimingTrack.cpp
@@ -1497,8 +1497,30 @@ void FThreadTimingTrack::OnClipboardCopyEvent(const ITimingEvent& InSelectedEven
 		FTimerNodePtr TimerNodePtr = FTimingProfilerManager::Get()->GetTimerNode(TrackEvent.GetTimerId());
 		if (TimerNodePtr)
 		{
+			FString EventName = TimerNodePtr->GetName().ToString();
+
+			FTimingEventsTrackDrawStateBuilder::AppendDurationToEventName(EventName, TrackEvent.GetDuration());
+
+			const uint32 TimerIndex = TrackEvent.GetTimerIndex();
+			if (int32(TimerIndex) < 0) // has metadata?
+			{
+				TSharedPtr<const TraceServices::IAnalysisSession> Session = FInsightsManager::Get()->GetSession();
+				check(Session.IsValid());
+
+				TraceServices::FAnalysisSessionReadScope SessionReadScope(*Session.Get());
+				const TraceServices::ITimingProfilerProvider& TimingProfilerProvider = *TraceServices::ReadTimingProfilerProvider(*Session.Get());
+				const TraceServices::ITimingProfilerTimerReader* TimerReader;
+				TimingProfilerProvider.ReadTimers([&TimerReader](const TraceServices::ITimingProfilerTimerReader& Out) { TimerReader = &Out; });
+
+				TArrayView<const uint8> Metadata = TimerReader->GetMetadata(TimerIndex);
+				if (Metadata.Num() > 0)
+				{
+					AppendMetadataToString(EventName, Metadata);
+				}
+			}
+
 			// Copy name of selected timing event to clipboard.
-			FPlatformApplicationMisc::ClipboardCopy(*TimerNodePtr->GetName().ToString());
+			FPlatformApplicationMisc::ClipboardCopy(*EventName);
 		}
 	}
 }
