[Insights] Fixed the list of trace sessions becoming unsorted when filter is changed.

diff --git a/Engine/Source/Developer/TraceInsights/Private/Insights/Widgets/SStartPageWindow.cpp b/Engine/Source/Developer/TraceInsights/Private/Insights/Widgets/SStartPageWindow.cpp
--- a/Engine/Source/Developer/TraceInsights/Private/Insights/Widgets/SStartPageWindow.cpp
+++ b/Engine/Source/Developer/TraceInsights/Private/Insights/Widgets/SStartPageWindow.cpp
@@ -2194,9 +2194,7 @@ void STraceStoreWindow::FilterByNameSearchBox_OnTextChanged(const FText& InFilte
 {
 	FilterByName->SetRawFilterText(InFilterText);
 	FilterByNameSearchBox->SetError(FilterByName->GetFilterErrorText());
-
-	UpdateFiltering();
-	UpdateTraceListView();
+	OnFilterChanged();
 }
 
 ////////////////////////////////////////////////////////////////////////////////////////////////////
@@ -2204,6 +2202,7 @@ void STraceStoreWindow::FilterByNameSearchBox_OnTextChanged(const FText& InFilte
 void STraceStoreWindow::OnFilterChanged()
 {
 	UpdateFiltering();
+	UpdateSorting();
 	UpdateTraceListView();
 }
 
@@ -2703,7 +2702,7 @@ TSharedRef<SWidget> SConnectionWindow::ConstructConnectPanel()
 				SNew(SButton)
 				.ButtonStyle(&FAppStyle::Get().GetWidgetStyle<FButtonStyle>("PrimaryButton"))
 				.Text(LOCTEXT("Connect", "Connect"))
-				.ToolTipText(LOCTEXT("ConnectToolTip", "Connect the running instance at specified ip with the local trace recorder."))
+				.ToolTipText(LOCTEXT("ConnectToolTip", "Connect the running instance at specified address with the trace recorder."))
 				.OnClicked(this, &SConnectionWindow::Connect_OnClicked)
 				.IsEnabled_Lambda([this]() { return !bIsConnecting; })
 			]
