[Insights] Table Tree views: Changed logging to show how much time is spent in aggregation operation, separated by grouping and filtering. Also made a minor refactoring of STableTreeView code re grouping, filtering and sorting.

diff --git a/Engine/Source/Developer/TraceInsights/Private/Insights/Table/Widgets/STableTreeView.cpp b/Engine/Source/Developer/TraceInsights/Private/Insights/Table/Widgets/STableTreeView.cpp
--- a/Engine/Source/Developer/TraceInsights/Private/Insights/Table/Widgets/STableTreeView.cpp
+++ b/Engine/Source/Developer/TraceInsights/Private/Insights/Table/Widgets/STableTreeView.cpp
@@ -896,9 +896,6 @@ void STableTreeView::Tick(const FGeometry& AllottedGeometry, const double InCurr
 
 void STableTreeView::UpdateTree()
 {
-	FStopwatch Stopwatch;
-	Stopwatch.Start();
-
 	if (bRunInAsyncMode)
 	{
 		if (!bIsUpdateRunning)
@@ -916,19 +913,34 @@ void STableTreeView::UpdateTree()
 	}
 	else
 	{
-		CreateGroups(CurrentGroupings);
-		if (CurrentSorter.IsValid())
-		{
-			SortTreeNodes(CurrentSorter.Get(), ColumnSortMode);
-		}
+		ApplyGrouping();
+		ApplySorting();
 		ApplyFiltering();
 	}
+}
 
-	Stopwatch.Stop();
-	const double TotalTime = Stopwatch.GetAccumulatedTime();
-	if (TotalTime > 0.1)
+////////////////////////////////////////////////////////////////////////////////////////////////////
+// Filtering
+////////////////////////////////////////////////////////////////////////////////////////////////////
+
+void STableTreeView::OnFilteringChanged()
+{
+	if (bRunInAsyncMode)
 	{
-		UE_LOG(TraceInsights, Log, TEXT("[Tree - %s] Tree updated (grouping + sorting + filtering) in %.3fs."), *Table->GetDisplayName().ToString(), TotalTime);
+		if (!bIsUpdateRunning)
+		{
+			OnPreAsyncUpdate();
+
+			InProgressAsyncOperationEvent = StartApplyFiltersTask();
+		}
+		else
+		{
+			CancelCurrentAsyncOp();
+		}
+	}
+	else
+	{
+		ApplyFiltering();
 	}
 }
 
@@ -963,14 +975,15 @@ void STableTreeView::ApplyFiltering()
 		}
 	}
 
-	Stopwatch.Update();
-	const double Time1 = Stopwatch.GetAccumulatedTime();
+	Stopwatch.Stop();
+	const double FilteringTime = Stopwatch.GetAccumulatedTime();
+	if (FilteringTime > 0.1)
+	{
+		UE_LOG(TraceInsights, Log, TEXT("[Tree - %s] Filtering completed in %.3fs."), *Table->GetDisplayName().ToString(), FilteringTime);
+	}
 
 	UpdateAggregatedValues(*Root);
 
-	Stopwatch.Update();
-	const double Time2 = Stopwatch.GetAccumulatedTime();
-
 	// Cannot call TreeView functions from other threads than MainThread and SlateThread.
 	if (!bRunInAsyncMode)
 	{
@@ -1008,13 +1021,6 @@ void STableTreeView::ApplyFiltering()
 		// Request tree refresh
 		TreeView->RequestTreeRefresh();
 	}
-
-	Stopwatch.Stop();
-	const double TotalTime = Stopwatch.GetAccumulatedTime();
-	if (TotalTime > 0.1)
-	{
-		UE_LOG(TraceInsights, Log, TEXT("[Tree - %s] Filtering completed in %.3fs (aggregation done in %.3fs)."), *Table->GetDisplayName().ToString(), TotalTime, Time2 - Time1);
-	}
 }
 
 ////////////////////////////////////////////////////////////////////////////////////////////////////
@@ -1297,14 +1303,43 @@ void STableTreeView::SearchBox_OnTextChanged(const FText& InFilterText)
 {
 	TextFilter->SetRawFilterText(InFilterText);
 	SearchBox->SetError(TextFilter->GetFilterErrorText());
+	OnFilteringChanged();
+}
 
+////////////////////////////////////////////////////////////////////////////////////////////////////
+
+bool STableTreeView::SearchBox_IsEnabled() const
+{
+	return !FilterConfigurator_HasFilters();
+}
+
+////////////////////////////////////////////////////////////////////////////////////////////////////
+
+FText STableTreeView::SearchBox_GetTooltipText() const
+{
+	if (SearchBox_IsEnabled())
+	{
+		return LOCTEXT("SearchBox_ToolTip", "Type here to search the tree hierarchy by item or group name.");
+	}
+
+	return LOCTEXT("SearchBox_Disabled_ToolTip", "Searching the tree hierarchy is disabled when advanced filters are set.");
+}
+
+////////////////////////////////////////////////////////////////////////////////////////////////////
+// Grouping
+////////////////////////////////////////////////////////////////////////////////////////////////////
+
+void STableTreeView::OnGroupingChanged()
+{
 	if (bRunInAsyncMode)
 	{
 		if (!bIsUpdateRunning)
 		{
 			OnPreAsyncUpdate();
 
-			InProgressAsyncOperationEvent = StartApplyFiltersTask();
+			FGraphEventRef CompletedEvent = StartCreateGroupsTask();
+			CompletedEvent = StartSortTreeNodesTask(CompletedEvent);
+			InProgressAsyncOperationEvent = StartApplyFiltersTask(CompletedEvent);
 		}
 		else
 		{
@@ -1313,31 +1348,19 @@ void STableTreeView::SearchBox_OnTextChanged(const FText& InFilterText)
 	}
 	else
 	{
+		ApplyGrouping();
+		ApplySorting();
 		ApplyFiltering();
 	}
 }
 
 ////////////////////////////////////////////////////////////////////////////////////////////////////
 
-bool STableTreeView::SearchBox_IsEnabled() const
-{
-	return !FilterConfigurator_HasFilters();
-}
-
-////////////////////////////////////////////////////////////////////////////////////////////////////
-
-FText STableTreeView::SearchBox_GetTooltipText() const
+void STableTreeView::ApplyGrouping()
 {
-	if (SearchBox_IsEnabled())
-	{
-		return LOCTEXT("SearchBox_ToolTip", "Type here to search the tree hierarchy by item or group name.");
-	}
-
-	return LOCTEXT("SearchBox_Disabled_ToolTip", "Searching the tree hierarchy is disabled when advanced filters are set.");
+	CreateGroups(CurrentGroupings);
 }
 
-////////////////////////////////////////////////////////////////////////////////////////////////////
-// Grouping
 ////////////////////////////////////////////////////////////////////////////////////////////////////
 
 void STableTreeView::CreateGroups(const TArray<TSharedPtr<FTreeNodeGrouping>>& Groupings)
@@ -1347,17 +1370,14 @@ void STableTreeView::CreateGroups(const TArray<TSharedPtr<FTreeNodeGrouping>>& G
 
 	GroupNodesRec(TableTreeNodes, *Root, 0, Groupings);
 
-	Stopwatch.Update();
-	const double Time1 = Stopwatch.GetAccumulatedTime();
-
-	UpdateAggregatedValues(*Root);
-
 	Stopwatch.Stop();
-	const double TotalTime = Stopwatch.GetAccumulatedTime();
-	if (TotalTime > 0.1)
+	const double GroupingTime = Stopwatch.GetAccumulatedTime();
+	if (GroupingTime > 0.1)
 	{
-		UE_LOG(TraceInsights, Log, TEXT("[Tree - %s] Grouping completed in %.3fs (%.3fs + %.3fs)."), *Table->GetDisplayName().ToString(), TotalTime, Time1, TotalTime - Time1);
+		UE_LOG(TraceInsights, Log, TEXT("[Tree - %s] Grouping completed in %.3fs."), *Table->GetDisplayName().ToString(), GroupingTime);
 	}
+
+	UpdateAggregatedValues(*Root);
 }
 
 ////////////////////////////////////////////////////////////////////////////////////////////////////
@@ -1413,6 +1433,9 @@ void STableTreeView::GroupNodesRec(const TArray<FTableTreeNodePtr>& Nodes, FTabl
 
 void STableTreeView::UpdateAggregatedValues(FTableTreeNode& GroupNode)
 {
+	FStopwatch Stopwatch;
+	Stopwatch.Start();
+
 	for (const TSharedRef<FTableColumn>& ColumnRef : Table->GetColumns())
 	{
 		FTableColumn& Column = ColumnRef.Get();
@@ -1470,6 +1493,13 @@ void STableTreeView::UpdateAggregatedValues(FTableTreeNode& GroupNode)
 				break;
 		}
 	}
+
+	Stopwatch.Stop();
+	const double AggregationTime = Stopwatch.GetAccumulatedTime();
+	if (AggregationTime > 0.1)
+	{
+		UE_LOG(TraceInsights, Log, TEXT("[Tree - %s] Aggregation completed in %.3fs."), *Table->GetDisplayName().ToString(), AggregationTime);
+	}
 }
 
 ////////////////////////////////////////////////////////////////////////////////////////////////////
@@ -1636,30 +1666,7 @@ void STableTreeView::PostChangeGroupings()
 
 	TreeViewHeaderRow->RefreshColumns();
 
-	if (bRunInAsyncMode)
-	{
-		if (!bIsUpdateRunning)
-		{
-			OnPreAsyncUpdate();
-
-			FGraphEventRef CompletedEvent = StartCreateGroupsTask();
-			CompletedEvent = StartSortTreeNodesTask(CompletedEvent);
-			InProgressAsyncOperationEvent = StartApplyFiltersTask(CompletedEvent);
-		}
-		else
-		{
-			CancelCurrentAsyncOp();
-		}
-	}
-	else
-	{
-		CreateGroups(CurrentGroupings);
-		if (CurrentSorter.IsValid())
-		{
-			SortTreeNodes(CurrentSorter.Get(), ColumnSortMode);
-		}
-		ApplyFiltering();
-	}
+	OnGroupingChanged();
 
 	RebuildGroupingCrumbs();
 }
@@ -2093,6 +2100,41 @@ void STableTreeView::UpdateCurrentSortingByColumn()
 
 ////////////////////////////////////////////////////////////////////////////////////////////////////
 
+void STableTreeView::OnSortingChanged()
+{
+	if (bRunInAsyncMode)
+	{
+		if (!bIsUpdateRunning)
+		{
+			OnPreAsyncUpdate();
+
+			FGraphEventRef CompletedEvent = StartSortTreeNodesTask();
+			InProgressAsyncOperationEvent = StartApplyFiltersTask(CompletedEvent);
+		}
+		else
+		{
+			CancelCurrentAsyncOp();
+		}
+	}
+	else
+	{
+		ApplySorting();
+		ApplyFiltering();
+	}
+}
+
+////////////////////////////////////////////////////////////////////////////////////////////////////
+
+void STableTreeView::ApplySorting()
+{
+	if (CurrentSorter.IsValid())
+	{
+		SortTreeNodes(CurrentSorter.Get(), ColumnSortMode);
+	}
+}
+
+////////////////////////////////////////////////////////////////////////////////////////////////////
+
 void STableTreeView::SortTreeNodes(ITableCellValueSorter* InSorter, EColumnSortMode::Type InColumnSortMode)
 {
 	FStopwatch Stopwatch;
@@ -2154,28 +2196,7 @@ void STableTreeView::SetSortModeForColumn(const FName& ColumnId, const EColumnSo
 	ColumnSortMode = SortMode;
 	UpdateCurrentSortingByColumn();
 
-	if (bRunInAsyncMode)
-	{
-		if (!bIsUpdateRunning)
-		{
-			OnPreAsyncUpdate();
-
-			FGraphEventRef CompletedEvent = StartSortTreeNodesTask();
-			InProgressAsyncOperationEvent = StartApplyFiltersTask(CompletedEvent);
-		}
-		else
-		{
-			CancelCurrentAsyncOp();
-		}
-	}
-	else
-	{
-		if (CurrentSorter.IsValid())
-		{
-			SortTreeNodes(CurrentSorter.Get(), ColumnSortMode);
-		}
-		ApplyFiltering();
-	}
+	OnSortingChanged();
 }
 
 ////////////////////////////////////////////////////////////////////////////////////////////////////
@@ -2926,22 +2947,7 @@ bool STableTreeView::ApplyAdvancedFilters(const FTableTreeNodePtr& NodePtr)
 
 void STableTreeView::OnAdvancedFiltersChangesCommited()
 {
-	if (bRunInAsyncMode)
-	{
-		if (!bIsUpdateRunning)
-		{
-			OnPreAsyncUpdate();
-			InProgressAsyncOperationEvent = StartApplyFiltersTask();
-		}
-		else
-		{
-			CancelCurrentAsyncOp();
-		}
-	}
-	else
-	{
-		ApplyFiltering();
-	}
+	OnFilteringChanged();
 }
 
 ////////////////////////////////////////////////////////////////////////////////////////////////////
diff --git a/Engine/Source/Developer/TraceInsights/Private/Insights/Table/Widgets/STableTreeView.h b/Engine/Source/Developer/TraceInsights/Private/Insights/Table/Widgets/STableTreeView.h
--- a/Engine/Source/Developer/TraceInsights/Private/Insights/Table/Widgets/STableTreeView.h
+++ b/Engine/Source/Developer/TraceInsights/Private/Insights/Table/Widgets/STableTreeView.h
@@ -225,6 +225,8 @@ protected:
 	////////////////////////////////////////////////////////////////////////////////////////////////////
 	// Filtering
 
+	void OnFilteringChanged();
+
 	/** Populates the group and stat tree with items based on the current data. */
 	void ApplyFiltering();
 
@@ -246,11 +248,36 @@ protected:
 	void CreateGroupings();
 	virtual void InternalCreateGroupings();
 
+	void OnGroupingChanged();
+	void ApplyGrouping();
+
 	void CreateGroups(const TArray<TSharedPtr<FTreeNodeGrouping>>& Groupings);
 	void GroupNodesRec(const TArray<FTableTreeNodePtr>& Nodes, FTableTreeNode& ParentGroup, int32 GroupingDepth, const TArray<TSharedPtr<FTreeNodeGrouping>>& Groupings);
 
+	void RebuildGroupingCrumbs();
+	void OnGroupingCrumbClicked(const TSharedPtr<FTreeNodeGrouping>& InEntry);
+	void BuildGroupingSubMenu_Change(FMenuBuilder& MenuBuilder, const TSharedPtr<FTreeNodeGrouping> CrumbGrouping);
+	void BuildGroupingSubMenu_Add(FMenuBuilder& MenuBuilder, const TSharedPtr<FTreeNodeGrouping> CrumbGrouping);
+	TSharedRef<SWidget> GetGroupingCrumbMenuContent(const TSharedPtr<FTreeNodeGrouping>& CrumbGrouping);
+
+	void PreChangeGroupings();
+	void PostChangeGroupings();
+	int32 GetGroupingDepth(const TSharedPtr<FTreeNodeGrouping>& Grouping) const;
+
+	void GroupingCrumbMenu_Reset_Execute();
+	void GroupingCrumbMenu_Remove_Execute(const TSharedPtr<FTreeNodeGrouping> Grouping);
+	void GroupingCrumbMenu_MoveLeft_Execute(const TSharedPtr<FTreeNodeGrouping> Grouping);
+	void GroupingCrumbMenu_MoveRight_Execute(const TSharedPtr<FTreeNodeGrouping> Grouping);
+	void GroupingCrumbMenu_Change_Execute(const TSharedPtr<FTreeNodeGrouping> OldGrouping, const TSharedPtr<FTreeNodeGrouping> NewGrouping);
+	bool GroupingCrumbMenu_Change_CanExecute(const TSharedPtr<FTreeNodeGrouping> OldGrouping, const TSharedPtr<FTreeNodeGrouping> NewGrouping) const;
+	void GroupingCrumbMenu_Add_Execute(const TSharedPtr<FTreeNodeGrouping> Grouping, const TSharedPtr<FTreeNodeGrouping> AfterGrouping);
+	bool GroupingCrumbMenu_Add_CanExecute(const TSharedPtr<FTreeNodeGrouping> Grouping, const TSharedPtr<FTreeNodeGrouping> AfterGrouping) const;
+
+	////////////////////////////////////////////////////////////////////////////////////////////////////
+	// Aggregation
+
 	void UpdateAggregatedValues(FTableTreeNode& GroupNode);
-	
+
 	template<typename T>
 	static void UpdateAggregationRec(FTableColumn& Column, FTableTreeNode& GroupNode, T InitialAggregatedValue, bool bSetInitialValue, TFunctionRef<T(T, const FTableCellValue&)> ValueGetterFunc)
 	{
@@ -262,7 +289,7 @@ protected:
 			{
 				continue;
 			}
-			
+
 			if (!NodePtr->IsGroup())
 			{
 				const TOptional<FTableCellValue> NodeValue = Column.GetValue(*NodePtr);
@@ -289,34 +316,18 @@ protected:
 		}
 	}
 
-	void RebuildGroupingCrumbs();
-	void OnGroupingCrumbClicked(const TSharedPtr<FTreeNodeGrouping>& InEntry);
-	void BuildGroupingSubMenu_Change(FMenuBuilder& MenuBuilder, const TSharedPtr<FTreeNodeGrouping> CrumbGrouping);
-	void BuildGroupingSubMenu_Add(FMenuBuilder& MenuBuilder, const TSharedPtr<FTreeNodeGrouping> CrumbGrouping);
-	TSharedRef<SWidget> GetGroupingCrumbMenuContent(const TSharedPtr<FTreeNodeGrouping>& CrumbGrouping);
-
-	void PreChangeGroupings();
-	void PostChangeGroupings();
-	int32 GetGroupingDepth(const TSharedPtr<FTreeNodeGrouping>& Grouping) const;
-
-	void GroupingCrumbMenu_Reset_Execute();
-	void GroupingCrumbMenu_Remove_Execute(const TSharedPtr<FTreeNodeGrouping> Grouping);
-	void GroupingCrumbMenu_MoveLeft_Execute(const TSharedPtr<FTreeNodeGrouping> Grouping);
-	void GroupingCrumbMenu_MoveRight_Execute(const TSharedPtr<FTreeNodeGrouping> Grouping);
-	void GroupingCrumbMenu_Change_Execute(const TSharedPtr<FTreeNodeGrouping> OldGrouping, const TSharedPtr<FTreeNodeGrouping> NewGrouping);
-	bool GroupingCrumbMenu_Change_CanExecute(const TSharedPtr<FTreeNodeGrouping> OldGrouping, const TSharedPtr<FTreeNodeGrouping> NewGrouping) const;
-	void GroupingCrumbMenu_Add_Execute(const TSharedPtr<FTreeNodeGrouping> Grouping, const TSharedPtr<FTreeNodeGrouping> AfterGrouping);
-	bool GroupingCrumbMenu_Add_CanExecute(const TSharedPtr<FTreeNodeGrouping> Grouping, const TSharedPtr<FTreeNodeGrouping> AfterGrouping) const;
-
 	////////////////////////////////////////////////////////////////////////////////////////////////////
 	// Sorting
 
 	static const EColumnSortMode::Type GetDefaultColumnSortMode();
 	static const FName GetDefaultColumnBeingSorted();
 
+	void OnSortingChanged();
 	void CreateSortings();
 
 	void UpdateCurrentSortingByColumn();
+
+	void ApplySorting();
 	void SortTreeNodes(ITableCellValueSorter* InSorter, EColumnSortMode::Type InColumnSortMode);
 	void SortTreeNodesRec(FTableTreeNode& GroupNode, const ITableCellValueSorter& Sorter, EColumnSortMode::Type InColumnSortMode);
 
