[Insights] Memory Insights: Improved the performance for sorting the Function column in Allocs Table (and in general, the sorting of FText columns). The sorting for FText columns now caches the GetValue for all tree nodes before sorting. Stress test: Previously, the sorting "by function" of 3.5M allocs took more than two hours, now it complets in ~4min.

diff --git a/Engine/Source/Developer/TraceInsights/Private/Insights/Table/ViewModels/TableCellValueSorter.cpp b/Engine/Source/Developer/TraceInsights/Private/Insights/Table/ViewModels/TableCellValueSorter.cpp
--- a/Engine/Source/Developer/TraceInsights/Private/Insights/Table/ViewModels/TableCellValueSorter.cpp
+++ b/Engine/Source/Developer/TraceInsights/Private/Insights/Table/ViewModels/TableCellValueSorter.cpp
@@ -234,7 +234,7 @@ FSorterByInt64Value::FSorterByInt64Value(TSharedRef<FTableColumn> InColumnRef)
 void FSorterByInt64Value::Sort(TArray<FBaseTreeNodePtr>& NodesToSort, ESortMode SortMode) const
 {
 	// Note: This implementation is faster than using above delegates.
-	//       It caches the values before sorting, in this way minimizes the calls to Column.GetValue().
+	//       It caches the values before sorting, in this way it minimizes the calls to Column.GetValue().
 
 	const FTableColumn& Column = *ColumnRef;
 
@@ -347,7 +347,7 @@ FSorterByFloatValue::FSorterByFloatValue(TSharedRef<FTableColumn> InColumnRef)
 void FSorterByFloatValue::Sort(TArray<FBaseTreeNodePtr>& NodesToSort, ESortMode SortMode) const
 {
 	// Note: This implementation is faster than using above delegates.
-	//       It caches the values before sorting, in this way minimizes the calls to Column.GetValue().
+	//       It caches the values before sorting, in this way it minimizes the calls to Column.GetValue().
 
 	const FTableColumn& Column = *ColumnRef;
 
@@ -460,7 +460,7 @@ FSorterByDoubleValue::FSorterByDoubleValue(TSharedRef<FTableColumn> InColumnRef)
 void FSorterByDoubleValue::Sort(TArray<FBaseTreeNodePtr>& NodesToSort, ESortMode SortMode) const
 {
 	// Note: This implementation is faster than using above delegates.
-	//       It caches the values before sorting, in this way minimizes the calls to Column.GetValue().
+	//       It caches the values before sorting, in this way it minimizes the calls to Column.GetValue().
 
 	const FTableColumn& Column = *ColumnRef;
 
@@ -642,6 +642,107 @@ FSorterByTextValue::FSorterByTextValue(TSharedRef<FTableColumn> InColumnRef)
 
 ////////////////////////////////////////////////////////////////////////////////////////////////////
 
+void FSorterByTextValue::Sort(TArray<FBaseTreeNodePtr>& NodesToSort, ESortMode SortMode) const
+{
+	// Note: This implementation is faster than using above delegates.
+	//       It caches the values before sorting, in this way it minimizes the calls to Column.GetValue().
+
+	const FTableColumn& Column = *ColumnRef;
+
+	const int32 NumElements = NodesToSort.Num();
+
+	struct FSortElement
+	{
+		const FBaseTreeNodePtr NodePtr;
+		TOptional<FTableCellValue> Value;
+	};
+	TArray<FSortElement> ElementsToSort;
+
+	// Cache value for each node.
+	ElementsToSort.Reset(NumElements);
+	for (const FBaseTreeNodePtr& NodePtr : NodesToSort)
+	{
+		ElementsToSort.Add({ NodePtr, Column.GetValue(*NodePtr) });
+	}
+	ensure(ElementsToSort.Num() == NumElements);
+
+	if (SortMode == ESortMode::Ascending)
+	{
+		ElementsToSort.Sort([](const FSortElement& A, const FSortElement& B) -> bool
+			{
+				if (!A.Value.IsSet())
+				{
+					if (!B.Value.IsSet())
+					{
+						INSIGHTS_DEFAULT_SORTING_NODES(A.NodePtr, B.NodePtr)
+					}
+
+					return true;
+				}
+
+				if (!B.Value.IsSet())
+				{
+					return false;
+				}
+
+				const FText& ValueA = A.Value.GetValue().GetText();
+				const FText& ValueB = B.Value.GetValue().GetText();
+
+				const int32 Compare = ValueA.CompareTo(ValueB);
+				if (Compare == 0)
+				{
+					INSIGHTS_DEFAULT_SORTING_NODES(A.NodePtr, B.NodePtr)
+				}
+				else
+				{
+					// Sort by value (ascending).
+					return Compare < 0;
+				}
+			});
+	}
+	else // if (SortMode == ESortMode::Descending)
+	{
+		ElementsToSort.Sort([](const FSortElement& A, const FSortElement& B) -> bool
+			{
+				if (!A.Value.IsSet())
+				{
+					if (!B.Value.IsSet())
+					{
+						INSIGHTS_DEFAULT_SORTING_NODES(A.NodePtr, B.NodePtr)
+					}
+
+					return true;
+				}
+
+				if (!B.Value.IsSet())
+				{
+					return false;
+				}
+
+				const FText& ValueA = A.Value.GetValue().GetText();
+				const FText& ValueB = B.Value.GetValue().GetText();
+
+				const int32 Compare = ValueA.CompareTo(ValueB);
+				if (Compare == 0)
+				{
+					INSIGHTS_DEFAULT_SORTING_NODES(A.NodePtr, B.NodePtr)
+				}
+				else
+				{
+					// Sort by value (descending).
+					return Compare > 0;
+				}
+			});
+	}
+
+	for (int32 Index = 0; Index < NumElements; ++Index)
+	{
+		NodesToSort[Index] = ElementsToSort[Index].NodePtr;
+	}
+}
+
+////////////////////////////////////////////////////////////////////////////////////////////////////
+
 } // namespace Insights
 
 #undef INSIGHTS_DEFAULT_SORTING_NODES
diff --git a/Engine/Source/Developer/TraceInsights/Private/Insights/Table/ViewModels/TableCellValueSorter.h b/Engine/Source/Developer/TraceInsights/Private/Insights/Table/ViewModels/TableCellValueSorter.h
--- a/Engine/Source/Developer/TraceInsights/Private/Insights/Table/ViewModels/TableCellValueSorter.h
+++ b/Engine/Source/Developer/TraceInsights/Private/Insights/Table/ViewModels/TableCellValueSorter.h
@@ -153,6 +153,8 @@ class FSorterByTextValue : public FBaseTableColumnSorter
 {
 public:
 	FSorterByTextValue(TSharedRef<FTableColumn> InColumnRef);
+
+	virtual void Sort(TArray<FBaseTreeNodePtr>& NodesToSort, ESortMode SortMode) const override;
 };
 
 ////////////////////////////////////////////////////////////////////////////////////////////////////
