[Insights] Memory Insights: Added "Empty Callstack" for callstacks traced/resolved with zero stack frames.

diff --git a/Engine/Source/Developer/TraceInsights/Private/Insights/MemoryProfiler/ViewModels/CallstackFormatting.h b/Engine/Source/Developer/TraceInsights/Private/Insights/MemoryProfiler/ViewModels/CallstackFormatting.h
--- a/Engine/Source/Developer/TraceInsights/Private/Insights/MemoryProfiler/ViewModels/CallstackFormatting.h
+++ b/Engine/Source/Developer/TraceInsights/Private/Insights/MemoryProfiler/ViewModels/CallstackFormatting.h
@@ -1,21 +1,39 @@
 ﻿// Copyright Epic Games, Inc. All Rights Reserved.
+
 #pragma once
+
 #include "TraceServices/Model/Callstack.h"
 #include "TraceServices/Model/Modules.h"
 
-/////////////////////////////////////////////////////////////////////
+////////////////////////////////////////////////////////////////////////////////////////////////////
+
 enum EStackFrameFormatFlags : uint8
 {
 	Module			= 1 << 0,
 	FileAndLine		= 1 << 1
 };
 
-/////////////////////////////////////////////////////////////////////
+////////////////////////////////////////////////////////////////////////////////////////////////////
+
+inline const TCHAR* GetCallstackNotAvailableString()
+{
+	return TEXT("Unknown Callstack");
+}
+
+////////////////////////////////////////////////////////////////////////////////////////////////////
+
+inline const TCHAR* GetEmptyCallstackString()
+{
+	return TEXT("Empty Callstack");
+}
+
+////////////////////////////////////////////////////////////////////////////////////////////////////
+
 inline void FormatStackFrame(const TraceServices::FStackFrame& Frame, FStringBuilderBase& OutString, uint8 FormatFlags)
 {
 	using namespace TraceServices;
 	const ESymbolQueryResult Result = Frame.Symbol->GetResult();
-	switch(Result)
+	switch (Result)
 	{
 		case ESymbolQueryResult::OK:
 			if (FormatFlags & (uint8)EStackFrameFormatFlags::Module)
@@ -28,6 +46,7 @@ inline void FormatStackFrame(const TraceServices::FStackFrame& Frame, FStringBui
 				OutString.Appendf(TEXT(" %s(%d)"), Frame.Symbol->File, Frame.Symbol->Line);
 			}
 			break;
+
 		case ESymbolQueryResult::Mismatch:
 		case ESymbolQueryResult::NotFound:
 		case ESymbolQueryResult::NotLoaded:
@@ -40,9 +59,12 @@ inline void FormatStackFrame(const TraceServices::FStackFrame& Frame, FStringBui
 				OutString.Appendf(TEXT("0x%08x (%s)"), Frame.Symbol->Name, Frame.Addr, QueryResultToString((Result)));
 			}
 			break;
+
 		case ESymbolQueryResult::Pending:
 		default:
 			OutString.Append(QueryResultToString(Result));
 			break;
 	}
 }
+
+////////////////////////////////////////////////////////////////////////////////////////////////////
diff --git a/Engine/Source/Developer/TraceInsights/Private/Insights/MemoryProfiler/ViewModels/MemAllocGroupingByCallstack.cpp b/Engine/Source/Developer/TraceInsights/Private/Insights/MemoryProfiler/ViewModels/MemAllocGroupingByCallstack.cpp
--- a/Engine/Source/Developer/TraceInsights/Private/Insights/MemoryProfiler/ViewModels/MemAllocGroupingByCallstack.cpp
+++ b/Engine/Source/Developer/TraceInsights/Private/Insights/MemoryProfiler/ViewModels/MemAllocGroupingByCallstack.cpp
@@ -2,9 +2,9 @@
 
 #include "MemAllocGroupingByCallstack.h"
 #include "TraceServices/Model/Callstack.h"
-#include "CallstackFormatting.h"
 
 // Insights
+#include "Insights/MemoryProfiler/ViewModels/CallstackFormatting.h"
 #include "Insights/MemoryProfiler/ViewModels/MemAllocNode.h"
 
 #define LOCTEXT_NAMESPACE "Insights::FMemAllocGroupingByCallstack"
@@ -53,6 +53,7 @@ void FMemAllocGroupingByCallstack::GroupNodes(const TArray<FTableTreeNodePtr>& N
 	CallstackGroups.Add(Root);
 
 	FTableTreeNode* UnsetGroupPtr = nullptr;
+	FTableTreeNode* EmptyCallstackGroupPtr = nullptr;
 	TMap<const TraceServices::FCallstack*, FCallstackGroup*> GroupMapByCallstack; // Callstack* -> FCallstackGroup*
 
 	for (FTableTreeNodePtr NodePtr : Nodes)
@@ -69,6 +70,7 @@ void FMemAllocGroupingByCallstack::GroupNodes(const TArray<FTableTreeNodePtr>& N
 		}
 
 		FCallstackGroup* GroupPtr = Root;
+		int32 NumFrames = -1;
 
 		const FMemAllocNode& MemAllocNode = static_cast<const FMemAllocNode&>(*NodePtr);
 		const FMemoryAlloc* Alloc = MemAllocNode.GetMemAlloc();
@@ -83,7 +85,8 @@ void FMemAllocGroupingByCallstack::GroupNodes(const TArray<FTableTreeNodePtr>& N
 			}
 			else if (Callstack)
 			{
-				const int32 NumFrames = static_cast<int32>(Callstack->Num());
+				NumFrames = static_cast<int32>(Callstack->Num());
+
 				for (int32 FrameDepth = NumFrames - 1; FrameDepth >= 0; --FrameDepth)
 				{
 					const TraceServices::FStackFrame* Frame = Callstack->Frame(bIsInverted ? NumFrames - FrameDepth - 1 : FrameDepth);
@@ -130,7 +133,10 @@ void FMemAllocGroupingByCallstack::GroupNodes(const TArray<FTableTreeNodePtr>& N
 					}
 				}
 
-				GroupMapByCallstack.Add(Callstack, GroupPtr);
+				if (NumFrames > 0)
+				{
+					GroupMapByCallstack.Add(Callstack, GroupPtr);
+				}
 			}
 		}
 
@@ -138,6 +144,14 @@ void FMemAllocGroupingByCallstack::GroupNodes(const TArray<FTableTreeNodePtr>& N
 		{
 			GroupPtr->Node->AddChildAndSetGroupPtr(NodePtr);
 		}
+		else if (NumFrames == 0)
+		{
+			if (!EmptyCallstackGroupPtr)
+			{
+				EmptyCallstackGroupPtr = CreateEmptyCallstackGroup(InParentTable, ParentGroup);
+			}
+			EmptyCallstackGroupPtr->AddChildAndSetGroupPtr(NodePtr);
+		}
 		else
 		{
 			if (!UnsetGroupPtr)
@@ -226,6 +240,17 @@ FTableTreeNode* FMemAllocGroupingByCallstack::CreateUnsetGroup(TWeakPtr<FTable>
 
 ////////////////////////////////////////////////////////////////////////////////////////////////////
 
+FTableTreeNode* FMemAllocGroupingByCallstack::CreateEmptyCallstackGroup(TWeakPtr<FTable> ParentTable, FTableTreeNode& Parent) const
+{
+	static FName NotAvailableName(GetEmptyCallstackString());
+	FTableTreeNodePtr NodePtr = MakeShared<FTableTreeNode>(NotAvailableName, ParentTable);
+	NodePtr->SetExpansion(false);
+	Parent.AddChildAndSetGroupPtr(NodePtr);
+	return NodePtr.Get();
+}
+
+////////////////////////////////////////////////////////////////////////////////////////////////////
+
 } // namespace Insights
 
 #undef LOCTEXT_NAMESPACE
diff --git a/Engine/Source/Developer/TraceInsights/Private/Insights/MemoryProfiler/ViewModels/MemAllocGroupingByCallstack.h b/Engine/Source/Developer/TraceInsights/Private/Insights/MemoryProfiler/ViewModels/MemAllocGroupingByCallstack.h
--- a/Engine/Source/Developer/TraceInsights/Private/Insights/MemoryProfiler/ViewModels/MemAllocGroupingByCallstack.h
+++ b/Engine/Source/Developer/TraceInsights/Private/Insights/MemoryProfiler/ViewModels/MemAllocGroupingByCallstack.h
@@ -51,6 +51,7 @@ private:
 	FCallstackGroup* CreateGroup(TArray<FCallstackGroup*>& InOutAllCallstackGroup, FCallstackGroup* InParentGroup, const FName InGroupName, TWeakPtr<FTable> InParentTable, const TraceServices::FStackFrame* InFrame) const;
 
 	FTableTreeNode* CreateUnsetGroup(TWeakPtr<FTable> ParentTable, FTableTreeNode& Parent) const;
+	FTableTreeNode* CreateEmptyCallstackGroup(TWeakPtr<FTable> ParentTable, FTableTreeNode& Parent) const;
 
 private:
 	bool bIsInverted;
diff --git a/Engine/Source/Developer/TraceInsights/Private/Insights/MemoryProfiler/ViewModels/MemAllocTable.cpp b/Engine/Source/Developer/TraceInsights/Private/Insights/MemoryProfiler/ViewModels/MemAllocTable.cpp
--- a/Engine/Source/Developer/TraceInsights/Private/Insights/MemoryProfiler/ViewModels/MemAllocTable.cpp
+++ b/Engine/Source/Developer/TraceInsights/Private/Insights/MemoryProfiler/ViewModels/MemAllocTable.cpp
@@ -1,7 +1,6 @@
 // Copyright Epic Games, Inc. All Rights Reserved.
 
 #include "MemAllocTable.h"
-#include "CallstackFormatting.h"
 #include "Containers/StringView.h"
 #include "Internationalization/Regex.h"
 #include "Widgets/SBoxPanel.h"
@@ -9,6 +8,7 @@
 #include "Widgets/Text/STextBlock.h"
 
 // Insights
+#include "Insights/MemoryProfiler/ViewModels/CallstackFormatting.h"
 #include "Insights/MemoryProfiler/ViewModels/MemAllocFilterValueConverter.h"
 #include "Insights/MemoryProfiler/ViewModels/MemAllocNode.h"
 #include "Insights/MemoryProfiler/ViewModels/MemAllocTable.h"
@@ -768,90 +768,92 @@ void FMemAllocTable::AddDefaultColumns()
 					const FMemoryAlloc* Alloc = MemAllocNode.GetMemAlloc();
 					if (Alloc)
 					{
-						static FString NotAvailable = TEXT("Unknown callstack");
 						const TraceServices::FCallstack* Callstack = Alloc->GetCallstack();
 
-						if (Callstack)
+						if (!Callstack)
 						{
-							check(Callstack->Num() > 0);
+							return FTableCellValue(FText::FromString(GetCallstackNotAvailableString()));
+						}
 
-							const TraceServices::FStackFrame* Frame = nullptr;
-							const TraceServices::FStackFrame* FirstUnknownSymbol = nullptr;
-							for (uint32 FrameIndex = 0; FrameIndex < Callstack->Num(); ++FrameIndex)
-							{
-								Frame = Callstack->Frame(FrameIndex);
-								check(Frame != nullptr);
+						if (Callstack->Num() == 0)
+						{
+							return FTableCellValue(FText::FromString(GetEmptyCallstackString()));
+						}
 
-								if (!Frame->Symbol || !Frame->Symbol->Name)
+						const TraceServices::FStackFrame* Frame = nullptr;
+						const TraceServices::FStackFrame* FirstUnknownSymbol = nullptr;
+						for (uint32 FrameIndex = 0; FrameIndex < Callstack->Num(); ++FrameIndex)
+						{
+							Frame = Callstack->Frame(FrameIndex);
+							check(Frame != nullptr);
+
+							if (!Frame->Symbol || !Frame->Symbol->Name)
+							{
+								if (FirstUnknownSymbol == nullptr)
 								{
-									if (FirstUnknownSymbol == nullptr)
-									{
-										FirstUnknownSymbol = Frame;
-									}
-									continue;
+									FirstUnknownSymbol = Frame;
 								}
+								continue;
+							}
 
-								bool bIgnoreSymbol = false;
+							bool bIgnoreSymbol = false;
 
-								// Ignore symbols by function prefix.
-								FStringView IgnoreSymbolsByFunctionName[] =
-								{
-									TEXT("FMemory::"_SV),
-									TEXT("FMallocWrapper::"_SV),
-									TEXT("FMallocPoisonProxy::"_SV),
-									TEXT("Malloc"_SV),
-									TEXT("Realloc"_SV),
-								};
-								for (uint32 StringIndex = 0; StringIndex < UE_ARRAY_COUNT(IgnoreSymbolsByFunctionName); ++StringIndex)
+							// Ignore symbols by function prefix.
+							FStringView IgnoreSymbolsByFunctionName[] =
+							{
+								TEXT("FMemory::"_SV),
+								TEXT("FMallocWrapper::"_SV),
+								TEXT("FMallocPoisonProxy::"_SV),
+								TEXT("Malloc"_SV),
+								TEXT("Realloc"_SV),
+							};
+							for (uint32 StringIndex = 0; StringIndex < UE_ARRAY_COUNT(IgnoreSymbolsByFunctionName); ++StringIndex)
+							{
+								if (FCString::Strnicmp(Frame->Symbol->Name, IgnoreSymbolsByFunctionName[StringIndex].GetData(), IgnoreSymbolsByFunctionName[StringIndex].Len()) == 0)
 								{
-									if (FCString::Strnicmp(Frame->Symbol->Name, IgnoreSymbolsByFunctionName[StringIndex].GetData(), IgnoreSymbolsByFunctionName[StringIndex].Len()) == 0)
-									{
-										bIgnoreSymbol = true;
-										break;
-									}
+									bIgnoreSymbol = true;
+									break;
 								}
+							}
 
 #if 1 // TODO: check perf impact
-								if (!bIgnoreSymbol && Frame->Symbol->File)
+							if (!bIgnoreSymbol && Frame->Symbol->File)
+							{
+								// Ignore symbols by file, specified as RegexPattern strings.
+								FStringView IgnoreSymbolsByFilePath[] =
 								{
-									// Ignore symbols by file, specified as RegexPattern strings.
-									FStringView IgnoreSymbolsByFilePath[] =
-									{
-										TEXT(".*/Containers/.*"_SV),
-									};
-									for (uint32 StringIndex = 0; StringIndex < UE_ARRAY_COUNT(IgnoreSymbolsByFilePath); ++StringIndex)
+									TEXT(".*/Containers/.*"_SV),
+								};
+								for (uint32 StringIndex = 0; StringIndex < UE_ARRAY_COUNT(IgnoreSymbolsByFilePath); ++StringIndex)
+								{
+									const FString Pattern(IgnoreSymbolsByFilePath[StringIndex]);
+									const FRegexPattern RegexPattern(Pattern);
+									FString File(Frame->Symbol->File);
+									File.ReplaceCharInline(TEXT('\\'), TEXT('/'), ESearchCase::CaseSensitive);
+									FRegexMatcher RegexMatcher(RegexPattern, File);
+									if (RegexMatcher.FindNext())
 									{
-										const FString Pattern(IgnoreSymbolsByFilePath[StringIndex]);
-										const FRegexPattern RegexPattern(Pattern);
-										FString File(Frame->Symbol->File);
-										File.ReplaceCharInline(TEXT('\\'), TEXT('/'), ESearchCase::CaseSensitive);
-										FRegexMatcher RegexMatcher(RegexPattern, File);
-										if (RegexMatcher.FindNext())
-										{
-											bIgnoreSymbol = true;
-											break;
-										}
+										bIgnoreSymbol = true;
+										break;
 									}
 								}
+							}
 #endif
 
-								if (!bIgnoreSymbol)
-								{
-									break;
-								}
-							}
-							if (!Frame)
+							if (!bIgnoreSymbol)
 							{
-								check(FirstUnknownSymbol != nullptr);
-								Frame = FirstUnknownSymbol;
+								break;
 							}
-
-							TStringBuilder<1024> Str;
-							FormatStackFrame(*Frame, Str, EStackFrameFormatFlags::Module);
-							return FTableCellValue(FText::FromString(FString(Str)));
+						}
+						if (!Frame)
+						{
+							check(FirstUnknownSymbol != nullptr);
+							Frame = FirstUnknownSymbol;
 						}
 
-						return FTableCellValue(FText::FromString(NotAvailable));
+						TStringBuilder<1024> Str;
+						FormatStackFrame(*Frame, Str, EStackFrameFormatFlags::Module);
+						return FTableCellValue(FText::FromString(FString(Str)));
 					}
 				}
 
diff --git a/Engine/Source/Developer/TraceInsights/Private/Insights/MemoryProfiler/ViewModels/MemoryAlloc.cpp b/Engine/Source/Developer/TraceInsights/Private/Insights/MemoryProfiler/ViewModels/MemoryAlloc.cpp
--- a/Engine/Source/Developer/TraceInsights/Private/Insights/MemoryProfiler/ViewModels/MemoryAlloc.cpp
+++ b/Engine/Source/Developer/TraceInsights/Private/Insights/MemoryProfiler/ViewModels/MemoryAlloc.cpp
@@ -1,7 +1,8 @@
 // Copyright Epic Games, Inc. All Rights Reserved.
 
 #include "MemoryAlloc.h"
-#include "CallstackFormatting.h"
+
+#include "Insights/MemoryProfiler/ViewModels/CallstackFormatting.h"
 
 namespace Insights
 {
@@ -32,29 +33,28 @@ FMemoryAlloc::~FMemoryAlloc()
 
 FText FMemoryAlloc::GetFullCallstack() const
 {
-	TStringBuilder<1024> Tooltip;
-	if (Callstack)
+	if (!Callstack)
 	{
-		check(Callstack->Num() > 0);
-		const uint32 FramesNum = Callstack->Num();
-		for (uint32 Index = 0; Index < FramesNum; ++Index)
-		{
-			const TraceServices::FStackFrame* Frame = Callstack->Frame(Index);
-			check(Frame != nullptr);
-
-			FormatStackFrame(*Frame, Tooltip, EStackFrameFormatFlags::Module | EStackFrameFormatFlags::FileAndLine);
-			if (Index != (FramesNum - 1))
-			{
-				Tooltip << TEXT("\n");
-			}
-		}
+		return FText::FromString(GetCallstackNotAvailableString());
 	}
 
-	if (Tooltip.Len() == 0)
+	if (Callstack->Num() == 0)
 	{
-		Tooltip.Append(TEXT("N/A"));
+		return FText::FromString(GetEmptyCallstackString());
 	}
 
+	TStringBuilder<1024> Tooltip;
+	const uint32 FramesNum = Callstack->Num();
+	for (uint32 Index = 0; Index < FramesNum; ++Index)
+	{
+		const TraceServices::FStackFrame* Frame = Callstack->Frame(Index);
+		check(Frame != nullptr);
+		FormatStackFrame(*Frame, Tooltip, EStackFrameFormatFlags::Module | EStackFrameFormatFlags::FileAndLine);
+		if (Index != (FramesNum - 1))
+		{
+			Tooltip << TEXT("\n");
+		}
+	}
 	return FText::FromString(FString(Tooltip));
 }
 
