[Insights] Memory Insights: Added infinte symbol for End Event Index and Event Distance table columns in Allocs Table. Implemented FInt64ValueFormatterAsUInt32InfinteNumber.

diff --git a/Engine/Source/Developer/TraceInsights/Private/Insights/MemoryProfiler/ViewModels/MemAllocTable.cpp b/Engine/Source/Developer/TraceInsights/Private/Insights/MemoryProfiler/ViewModels/MemAllocTable.cpp
--- a/Engine/Source/Developer/TraceInsights/Private/Insights/MemoryProfiler/ViewModels/MemAllocTable.cpp
+++ b/Engine/Source/Developer/TraceInsights/Private/Insights/MemoryProfiler/ViewModels/MemAllocTable.cpp
@@ -190,7 +190,7 @@ void FMemAllocTable::AddDefaultColumns()
 		TSharedRef<ITableCellValueGetter> Getter = MakeShared<FMemAllocEndEventIndexValueGetter>();
 		Column.SetValueGetter(Getter);
 
-		TSharedRef<ITableCellValueFormatter> Formatter = MakeShared<FInt64ValueFormatterAsNumber>();
+		TSharedRef<ITableCellValueFormatter> Formatter = MakeShared<FInt64ValueFormatterAsUInt32InfinteNumber>();
 		Column.SetValueFormatter(Formatter);
 
 		TSharedRef<ITableCellValueSorter> Sorter = MakeShared<FSorterByInt64Value>(ColumnRef);
@@ -248,7 +248,7 @@ void FMemAllocTable::AddDefaultColumns()
 		TSharedRef<ITableCellValueGetter> Getter = MakeShared<FMemAllocEventDistanceValueGetter>();
 		Column.SetValueGetter(Getter);
 
-		TSharedRef<ITableCellValueFormatter> Formatter = MakeShared<FInt64ValueFormatterAsNumber>();
+		TSharedRef<ITableCellValueFormatter> Formatter = MakeShared<FInt64ValueFormatterAsUInt32InfinteNumber>();
 		Column.SetValueFormatter(Formatter);
 
 		TSharedRef<ITableCellValueSorter> Sorter = MakeShared<FSorterByInt64Value>(ColumnRef);
diff --git a/Engine/Source/Developer/TraceInsights/Private/Insights/MemoryProfiler/ViewModels/MemoryAlloc.h b/Engine/Source/Developer/TraceInsights/Private/Insights/MemoryProfiler/ViewModels/MemoryAlloc.h
--- a/Engine/Source/Developer/TraceInsights/Private/Insights/MemoryProfiler/ViewModels/MemoryAlloc.h
+++ b/Engine/Source/Developer/TraceInsights/Private/Insights/MemoryProfiler/ViewModels/MemoryAlloc.h
@@ -22,7 +22,7 @@ public:
 
 	int64 GetStartEventIndex() const { return int64(StartEventIndex); }
 	int64 GetEndEventIndex() const { return int64(EndEventIndex); }
-	int64 GetEventDistance() const { return int64(EndEventIndex) - int64(StartEventIndex); }
+	int64 GetEventDistance() const { return EndEventIndex == ~0 ? int64(EndEventIndex) : int64(EndEventIndex) - int64(StartEventIndex); }
 	double GetStartTime() const { return StartTime; }
 	double GetEndTime() const { return EndTime; }
 	double GetDuration() const { return EndTime - StartTime; }
diff --git a/Engine/Source/Developer/TraceInsights/Private/Insights/Table/ViewModels/TableCellValueFormatter.cpp b/Engine/Source/Developer/TraceInsights/Private/Insights/Table/ViewModels/TableCellValueFormatter.cpp
--- a/Engine/Source/Developer/TraceInsights/Private/Insights/Table/ViewModels/TableCellValueFormatter.cpp
+++ b/Engine/Source/Developer/TraceInsights/Private/Insights/Table/ViewModels/TableCellValueFormatter.cpp
@@ -79,6 +79,38 @@ FText FBoolValueFormatterAsOnOff::FormatValue(const TOptional<FTableCellValue>&
 
 ////////////////////////////////////////////////////////////////////////////////////////////////////
 
+FText FInt64ValueFormatterAsUInt32InfinteNumber::FormatValue(const TOptional<FTableCellValue>& InValue) const
+{
+	if (InValue.IsSet())
+	{
+		const int64 Value = InValue.GetValue().Int64;
+		if (Value == int64(uint32(~0)))
+		{
+			return LOCTEXT("AsUInt32InfinteNumber_Inf", "∞");
+		}
+		return FText::AsNumber(Value);
+	}
+	return FText::GetEmpty();
+}
+
+////////////////////////////////////////////////////////////////////////////////////////////////////
+
+FText FInt64ValueFormatterAsUInt32InfinteNumber::FormatValueForTooltip(const TOptional<FTableCellValue>& InValue) const
+{
+	if (InValue.IsSet())
+	{
+		const int64 Value = InValue.GetValue().Int64;
+		if (Value == int64(uint32(~0)))
+		{
+			return FText::Format(LOCTEXT("AsUInt32InfinteNumber_Inf_Fmt", "{0} (∞)"), FText::AsNumber(Value));
+		}
+		return FText::AsNumber(Value);
+	}
+	return FText::GetEmpty();
+}
+
+////////////////////////////////////////////////////////////////////////////////////////////////////
+
 FText FInt64ValueFormatterAsMemory::FormatValue(const TOptional<FTableCellValue>& InValue) const
 {
 	if (InValue.IsSet())
diff --git a/Engine/Source/Developer/TraceInsights/Private/Insights/Table/ViewModels/TableCellValueFormatter.h b/Engine/Source/Developer/TraceInsights/Private/Insights/Table/ViewModels/TableCellValueFormatter.h
--- a/Engine/Source/Developer/TraceInsights/Private/Insights/Table/ViewModels/TableCellValueFormatter.h
+++ b/Engine/Source/Developer/TraceInsights/Private/Insights/Table/ViewModels/TableCellValueFormatter.h
@@ -108,6 +108,15 @@ public:
 
 ////////////////////////////////////////////////////////////////////////////////////////////////////
 
+class FInt64ValueFormatterAsUInt32InfinteNumber : public FTableCellValueFormatter
+{
+public:
+	virtual FText FormatValue(const TOptional<FTableCellValue>& InValue) const override;
+	virtual FText FormatValueForTooltip(const TOptional<FTableCellValue>& InValue) const override;
+};
+
+////////////////////////////////////////////////////////////////////////////////////////////////////
+
 class FInt64ValueFormatterAsHex32 : public FTableCellValueFormatter
 {
 public:
