Screenshot tracing initial implementation

diff --git a/Engine/Source/Developer/TraceInsights/Private/Insights/ViewModels/MarkersTimingTrack.cpp b/Engine/Source/Developer/TraceInsights/Private/Insights/ViewModels/MarkersTimingTrack.cpp
--- a/Engine/Source/Developer/TraceInsights/Private/Insights/ViewModels/MarkersTimingTrack.cpp
+++ b/Engine/Source/Developer/TraceInsights/Private/Insights/ViewModels/MarkersTimingTrack.cpp
@@ -5,15 +5,21 @@
 #include "Fonts/FontMeasure.h"
 #include "Framework/Application/SlateApplication.h"
 #include "Framework/MultiBox/MultiBoxBuilder.h"
+#include "IImageWrapper.h"
+#include "IImageWrapperModule.h"
+#include "Modules/ModuleManager.h"
 #include "Styling/AppStyle.h"
 #include "TraceServices/AnalysisService.h"
+#include "TraceServices/Model/Screenshot.h"
 
 // Insights
 #include "Insights/Common/PaintUtils.h"
 #include "Insights/InsightsManager.h"
 #include "Insights/InsightsStyle.h"
 #include "Insights/TimingProfilerManager.h"
+#include "Insights/ViewModels/TimingEvent.h"
 #include "Insights/ViewModels/TimingTrackViewport.h"
+#include "Insights/ViewModels/TooltipDrawState.h"
 
 #include <limits>
 
@@ -33,6 +39,7 @@ FMarkersTimingTrack::FMarkersTimingTrack()
 	//, TimeMarkerTexts()
 	, bUseOnlyBookmarks(true)
 	, BookmarkCategory(nullptr)
+	, ScreenshotCategory(nullptr)
 	, Header(*this)
 	, NumLogMessages(0)
 	, NumDrawBoxes(0)
@@ -61,6 +68,7 @@ void FMarkersTimingTrack::Reset()
 
 	bUseOnlyBookmarks = true;
 	BookmarkCategory = nullptr;
+	ScreenshotCategory = nullptr;
 
 	Header.Reset();
 	Header.SetIsInBackground(true);
@@ -99,7 +107,11 @@ void FMarkersTimingTrack::PreUpdate(const ITimingTrackUpdateContext& Context)
 {
 	if (!BookmarkCategory)
 	{
-		UpdateBookmarkCategory();
+		UpdateCategory(BookmarkCategory, TEXT("LogBookmark"));
+	}
+	if (!ScreenshotCategory)
+	{
+		UpdateCategory(ScreenshotCategory, TEXT("Screenshot"));
 	}
 }
 
@@ -364,9 +376,9 @@ double FMarkersTimingTrack::Snap(double Time, const double SnapTolerance)
 
 ////////////////////////////////////////////////////////////////////////////////////////////////////
 
-void FMarkersTimingTrack::UpdateBookmarkCategory()
+void FMarkersTimingTrack::UpdateCategory(const TraceServices::FLogCategoryInfo*& InOutCategory, const TCHAR* CategoryName)
 {
-	BookmarkCategory = nullptr;
+	InOutCategory = nullptr;
 
 	TSharedPtr<const TraceServices::IAnalysisSession> Session = FInsightsManager::Get()->GetSession();
 	if (Session)
@@ -374,16 +386,128 @@ void FMarkersTimingTrack::UpdateBookmarkCategory()
 		TraceServices::FAnalysisSessionReadScope SessionReadScope(*Session.Get());
 		const TraceServices::ILogProvider& LogProvider = TraceServices::ReadLogProvider(*Session.Get());
 
-		LogProvider.EnumerateCategories([this](const TraceServices::FLogCategoryInfo& Category)
+		LogProvider.EnumerateCategories([this, CategoryName, &InOutCategory](const TraceServices::FLogCategoryInfo& Category)
 		{
-			if (Category.Name && FCString::Strcmp(Category.Name, TEXT("LogBookmark")) == 0)
+			if (Category.Name && FCString::Strcmp(Category.Name, CategoryName) == 0)
 			{
-				BookmarkCategory = &Category;
+				InOutCategory = &Category;
 			}
 		});
 	}
 }
 
+////////////////////////////////////////////////////////////////////////////////////////////////////
+
+const TSharedPtr<const ITimingEvent> FMarkersTimingTrack::GetEvent(float InPosX, float InPosY, const FTimingTrackViewport& Viewport) const
+{
+	TSharedPtr<ITimingEvent> TimingEvent;
+
+	const FTimingViewLayout& Layout = Viewport.GetLayout();
+
+	const float DY = InPosY - GetPosY();
+
+	if (DY >= 0 && DY < GetHeight())
+	{
+		const int32 NumBoxes = TimeMarkerTexts.Num();
+		int32 FoundIndex = Algo::LowerBoundBy(TimeMarkerTexts, InPosX, [](const FTimeMarkerTextInfo& Text)
+		{
+			return Text.X;
+		});
+
+		if (FoundIndex > 0)
+		{
+			--FoundIndex;
+		}
+
+		if (FoundIndex < 0 || FoundIndex >= TimeMarkerTexts.Num())
+		{
+			return TimingEvent;
+		}
+
+		float Width = Viewport.GetWidth();
+		if (FoundIndex + 1 < TimeMarkerTexts.Num())
+		{
+			Width = TimeMarkerTexts[FoundIndex + 1].X;
+		}
+
+		uint16 ScreenshotId = std::numeric_limits<uint16>::max();
+
+		const FTimeMarkerTextInfo& Text = TimeMarkerTexts[FoundIndex];
+
+		TSharedPtr<const TraceServices::IAnalysisSession> Session = FInsightsManager::Get()->GetSession();
+		if (!Session.IsValid())
+		{
+			return TimingEvent;
+		}
+
+		TraceServices::FAnalysisSessionReadScope SessionReadScope(*Session.Get());
+		const TraceServices::ILogProvider& LogProvider = TraceServices::ReadLogProvider(*Session.Get());
+		LogProvider.ReadMessage(Text.LogIndex,
+			[&ScreenshotId, this](const TraceServices::FLogMessageInfo& Message)
+			{
+				if (Message.Category == this->ScreenshotCategory)
+				{
+					ScreenshotId = Message.Line;
+				}
+			});
+
+		if (ScreenshotId == std::numeric_limits<uint16>::max())
+		{
+			return TimingEvent;
+		}
+
+		TimingEvent = MakeShared<FTimingEvent>(SharedThis(this), Viewport.SlateUnitsToTime(Text.X), Viewport.SlateUnitsToTime(Text.X + Width), 0, ScreenshotId);
+	}
+
+	return TimingEvent;
+}
+
+////////////////////////////////////////////////////////////////////////////////////////////////////
+
+void FMarkersTimingTrack::InitTooltip(FTooltipDrawState& InOutTooltip, const ITimingEvent& InTooltipEvent) const
+{
+	InOutTooltip.ResetContent();
+	InOutTooltip.UpdateLayout();
+
+	if (!InTooltipEvent.CheckTrack(this))
+	{
+		return;
+	}
+
+	TSharedPtr<const TraceServices::IAnalysisSession> Session = FInsightsManager::Get()->GetSession();
+	if (!Session.IsValid())
+	{
+		return;
+	}
+
+	TraceServices::FAnalysisSessionReadScope SessionReadScope(*Session.Get());
+	const FTimingEvent& Event = StaticCast<const FTimingEvent&>(InTooltipEvent);
+
+	const TraceServices::IScreenshotProvider& ScreenshotProvider = TraceServices::ReadScreenshotProvider(*Session.Get());
+	TSharedPtr<const TraceServices::FScreenshot> Screenshot = ScreenshotProvider.GetScreenshot((uint16)Event.GetType());
+
+	if (!Screenshot.IsValid())
+	{
+		return;
+	}
+
+	IImageWrapperModule& ImageWrapperModule = FModuleManager::LoadModuleChecked<IImageWrapperModule>(FName("ImageWrapper"));
+	TSharedPtr<IImageWrapper> ImageWrapper = ImageWrapperModule.CreateImageWrapper(EImageFormat::PNG);
+	if (ImageWrapper.IsValid() && ImageWrapper->SetCompressed(Screenshot->Data.GetData(), Screenshot->Data.Num()))
+	{
+		TArray<uint8> RawImage;
+		if (!ImageWrapper->GetRaw(ERGBFormat::BGRA, 8, RawImage))
+		{
+			return;
+		}
+
+		TSharedPtr<FSlateBrush> ImageBrush = FSlateDynamicImageBrush::CreateWithImageData(FName(Screenshot->Name), FVector2D(Screenshot->Width, Screenshot->Height), RawImage);
+		InOutTooltip.SetImage(ImageBrush);
+	}
+
+	InOutTooltip.UpdateLayout();
+}
+
 ////////////////////////////////////////////////////////////////////////////////////////////////////
 // FTimeMarkerTrackBuilder
 ////////////////////////////////////////////////////////////////////////////////////////////////////
@@ -430,7 +554,7 @@ void FTimeMarkerTrackBuilder::AddLogMessage(const TraceServices::FLogMessageInfo
 	}
 
 	check(Message.Category != nullptr);
-	if (!Track.bUseOnlyBookmarks || Message.Category == Track.BookmarkCategory)
+	if (!Track.bUseOnlyBookmarks || Message.Category == Track.BookmarkCategory || Message.Category == Track.ScreenshotCategory)
 	{
 		float X = Viewport.TimeToSlateUnitsRounded(Message.Time);
 		if (X < 0.0f)
@@ -538,6 +662,7 @@ void FTimeMarkerTrackBuilder::Flush(float AvailableTextW)
 				FTimeMarkerTextInfo& TextInfo = Track.TimeMarkerTexts.AddDefaulted_GetRef();
 				TextInfo.X = LastX2 + 2.0f;
 				TextInfo.Color = Color;
+				TextInfo.LogIndex = LastLogIndex;
 				if (LastWholeCharacterIndexCategory >= 0)
 				{
 					TextInfo.Category.AppendChars(*CategoryStr, LastWholeCharacterIndexCategory + 1);
diff --git a/Engine/Source/Developer/TraceInsights/Private/Insights/ViewModels/MarkersTimingTrack.h b/Engine/Source/Developer/TraceInsights/Private/Insights/ViewModels/MarkersTimingTrack.h
--- a/Engine/Source/Developer/TraceInsights/Private/Insights/ViewModels/MarkersTimingTrack.h
+++ b/Engine/Source/Developer/TraceInsights/Private/Insights/ViewModels/MarkersTimingTrack.h
@@ -37,6 +37,7 @@ struct FTimeMarkerTextInfo
 	FLinearColor Color;
 	FString Category; // truncated Category string
 	FString Message; // truncated Message string
+	uint64 LogIndex;
 };
 
 ////////////////////////////////////////////////////////////////////////////////////////////////////
@@ -82,6 +83,8 @@ public:
 	virtual FReply OnMouseButtonDown(const FGeometry& MyGeometry, const FPointerEvent& MouseEvent) override;
 	virtual FReply OnMouseButtonDoubleClick(const FGeometry& MyGeometry, const FPointerEvent& MouseEvent) override;
 	virtual void BuildContextMenu(FMenuBuilder& MenuBuilder) override;
+	virtual const TSharedPtr<const ITimingEvent> GetEvent(float InPosX, float InPosY, const FTimingTrackViewport& Viewport) const override;
+	virtual void InitTooltip(FTooltipDrawState& InOutTooltip, const ITimingEvent& InTooltipEvent) const override;
 
 	double Snap(double Time, double SnapTolerance);
 
@@ -95,7 +98,7 @@ private:
 	void UpdateTrackNameAndHeight();
 	void UpdateDrawState(const ITimingTrackUpdateContext& Context);
 
-	void UpdateBookmarkCategory();
+	void UpdateCategory(const TraceServices::FLogCategoryInfo*& InOutCategory, const TCHAR* CategoryName);
 
 private:
 	TArray<FTimeMarkerBoxInfo> TimeMarkerBoxes;
@@ -103,6 +106,7 @@ private:
 
 	bool bUseOnlyBookmarks; // If true, uses only bookmarks; otherwise it uses all log messages.
 	const TraceServices::FLogCategoryInfo* BookmarkCategory;
+	const TraceServices::FLogCategoryInfo* ScreenshotCategory;
 
 	FTrackHeader Header;
 
diff --git a/Engine/Source/Developer/TraceInsights/Private/Insights/ViewModels/TooltipDrawState.cpp b/Engine/Source/Developer/TraceInsights/Private/Insights/ViewModels/TooltipDrawState.cpp
--- a/Engine/Source/Developer/TraceInsights/Private/Insights/ViewModels/TooltipDrawState.cpp
+++ b/Engine/Source/Developer/TraceInsights/Private/Insights/ViewModels/TooltipDrawState.cpp
@@ -55,6 +55,7 @@ void FTooltipDrawState::Reset()
 void FTooltipDrawState::ResetContent()
 {
 	Texts.Reset();
+	ImageBrush.Reset();
 
 	ValueOffsetX = 0.0f;
 	NewLineY = BorderY;
@@ -125,6 +126,12 @@ void FTooltipDrawState::AddTextLine(const float X, const float Y, const FString&
 
 void FTooltipDrawState::UpdateLayout()
 {
+	if (ImageBrush.IsValid())
+	{
+		DesiredSize = ImageBrush->GetImageSize();
+		return;
+	}
+
 	ValueOffsetX = 0.0f;
 	for (const FDrawTextInfo& TextInfo : Texts)
 	{
@@ -252,6 +259,13 @@ void FTooltipDrawState::Draw(const FDrawContext& DrawContext) const
 		//DrawContext.DrawBox(Position.X, Position.Y, Size.X, Size.Y, BorderBrush, BorderColor.CopyWithNewOpacity(Opacity));
 		//DrawContext.LayerId++;
 
+		if (ImageBrush)
+		{
+			DrawContext.DrawBox(Position.X, Position.Y, DesiredSize.X, DesiredSize.Y, ImageBrush.Get(), FLinearColor::White);
+			DrawContext.LayerId++;
+			return;
+		}
+
 		// Draw cached texts.
 		for (const FDrawTextInfo& TextInfo : Texts)
 		{
diff --git a/Engine/Source/Developer/TraceInsights/Public/Insights/ViewModels/TooltipDrawState.h b/Engine/Source/Developer/TraceInsights/Public/Insights/ViewModels/TooltipDrawState.h
--- a/Engine/Source/Developer/TraceInsights/Public/Insights/ViewModels/TooltipDrawState.h
+++ b/Engine/Source/Developer/TraceInsights/Public/Insights/ViewModels/TooltipDrawState.h
@@ -85,6 +85,8 @@ public:
 	void SetFontScale(float InFontScale) { FontScale = InFontScale; }
 	float GetFontScale() const { return FontScale; }
 
+	void SetImage(TSharedPtr<FSlateBrush> InImageBrush) { ImageBrush = InImageBrush; }
+
 private:
 	const FSlateBrush* WhiteBrush;
 	const FSlateFontInfo Font;
@@ -105,6 +107,7 @@ private:
 	float FontScale;
 
 	TArray<FDrawTextInfo> Texts;
+	TSharedPtr<FSlateBrush> ImageBrush;
 };
 
 ////////////////////////////////////////////////////////////////////////////////////////////////////
diff --git a/Engine/Source/Developer/TraceServices/Private/AnalysisService.cpp b/Engine/Source/Developer/TraceServices/Private/AnalysisService.cpp
--- a/Engine/Source/Developer/TraceServices/Private/AnalysisService.cpp
+++ b/Engine/Source/Developer/TraceServices/Private/AnalysisService.cpp
@@ -19,6 +19,7 @@
 #include "Model/MemoryPrivate.h"
 #include "Model/Channel.h"
 #include "Model/DiagnosticsPrivate.h"
+#include "Model/ScreenshotProviderPrivate.h"
 
 namespace TraceServices
 {
@@ -240,7 +241,10 @@ TSharedPtr<const IAnalysisSession> FAnalysisService::StartAnalysis(uint32 TraceI
 	FChannelProvider* ChannelProvider = new FChannelProvider();
 	Session->AddProvider(FChannelProvider::ProviderName, ChannelProvider);
 
-	Session->AddAnalyzer(new FMiscTraceAnalyzer(*Session, *ThreadProvider, *BookmarkProvider, *LogProvider, *FrameProvider, *ChannelProvider));
+	FScreenshotProvider* ScreenshotProvider = new FScreenshotProvider(*Session);
+	Session->AddProvider(FScreenshotProvider::ProviderName, ScreenshotProvider);
+
+	Session->AddAnalyzer(new FMiscTraceAnalyzer(*Session, *ThreadProvider, *BookmarkProvider, *LogProvider, *FrameProvider, *ChannelProvider, *ScreenshotProvider));
 	Session->AddAnalyzer(new FLogTraceAnalyzer(*Session, *LogProvider));
 
 	ModuleService.OnAnalysisBegin(*Session);
diff --git a/Engine/Source/Developer/TraceServices/Private/Analyzers/MiscTraceAnalysis.cpp b/Engine/Source/Developer/TraceServices/Private/Analyzers/MiscTraceAnalysis.cpp
--- a/Engine/Source/Developer/TraceServices/Private/Analyzers/MiscTraceAnalysis.cpp
+++ b/Engine/Source/Developer/TraceServices/Private/Analyzers/MiscTraceAnalysis.cpp
@@ -5,24 +5,26 @@
 #include "Model/ThreadsPrivate.h"
 #include "Model/BookmarksPrivate.h"
 #include "Model/FramesPrivate.h"
+#include "Model/ScreenshotProviderPrivate.h"
 #include "Model/Channel.h"
 #include "Common/Utils.h"
 
 namespace TraceServices
 {
-
 FMiscTraceAnalyzer::FMiscTraceAnalyzer(IAnalysisSession& InSession,
 									   FThreadProvider& InThreadProvider,
 									   FBookmarkProvider& InBookmarkProvider,
 									   FLogProvider& InLogProvider,
 									   FFrameProvider& InFrameProvider,
-									   FChannelProvider& InChannelProvider)
+									   FChannelProvider& InChannelProvider,
+									   FScreenshotProvider& InScreenshotProvider)
 	: Session(InSession)
 	, ThreadProvider(InThreadProvider)
 	, BookmarkProvider(InBookmarkProvider)
 	, LogProvider(InLogProvider)
 	, FrameProvider(InFrameProvider)
 	, ChannelProvider(InChannelProvider)
+	, ScreenshotProvider(InScreenshotProvider)
 {
 	FLogCategoryInfo& BookmarkLogCategory = LogProvider.GetCategory(FLogProvider::ReservedLogCategory_Bookmark);
 	BookmarkLogCategory.Name = TEXT("LogBookmark");
@@ -48,6 +50,8 @@ void FMiscTraceAnalyzer::OnAnalysisBegin(const FOnAnalysisContext& Context)
 	Builder.RouteEvent(RouteId_EndRenderFrame, "Misc", "EndRenderFrame");
 	Builder.RouteEvent(RouteId_ChannelAnnounce, "Trace", "ChannelAnnounce");
 	Builder.RouteEvent(RouteId_ChannelToggle, "Trace", "ChannelToggle");
+	Builder.RouteEvent(RouteId_ScreenshotHeader, "Misc", "ScreenshotHeader");
+	Builder.RouteEvent(RouteId_ScreenshotChunk, "Misc", "ScreenshotChunk");
 }
 
 void FMiscTraceAnalyzer::OnThreadInfo(const FThreadInfo& ThreadInfo)
@@ -104,6 +108,7 @@ bool FMiscTraceAnalyzer::OnEvent(uint16 RouteId, EStyle Style, const FOnEventCon
 		LogMessageSpec.Verbosity = ELogVerbosity::Log;
 		break;
 	}
+
 	case RouteId_Bookmark:
 	{
 		uint64 BookmarkPoint = EventData.GetValue<uint64>("BookmarkPoint");
@@ -114,6 +119,7 @@ bool FMiscTraceAnalyzer::OnEvent(uint16 RouteId, EStyle Style, const FOnEventCon
 		LogProvider.AppendMessage(BookmarkPoint, Timestamp, FormatArgsView.GetData());
 		break;
 	}
+
 	case RouteId_BeginFrame:
 	{
 		uint64 Cycle = EventData.GetValue<uint64>("Cycle");
@@ -122,6 +128,7 @@ bool FMiscTraceAnalyzer::OnEvent(uint16 RouteId, EStyle Style, const FOnEventCon
 		FrameProvider.BeginFrame(ETraceFrameType(FrameType), Context.EventTime.AsSeconds(Cycle));
 		break;
 	}
+
 	case RouteId_EndFrame:
 	{
 		uint64 Cycle = EventData.GetValue<uint64>("Cycle");
@@ -130,6 +137,7 @@ bool FMiscTraceAnalyzer::OnEvent(uint16 RouteId, EStyle Style, const FOnEventCon
 		FrameProvider.EndFrame(ETraceFrameType(FrameType), Context.EventTime.AsSeconds(Cycle));
 		break;
 	}
+
 	case RouteId_BeginGameFrame:
 	case RouteId_EndGameFrame:
 	case RouteId_BeginRenderFrame:
@@ -167,6 +175,47 @@ bool FMiscTraceAnalyzer::OnEvent(uint16 RouteId, EStyle Style, const FOnEventCon
 		OnChannelToggle(Context);
 		break;
 
+	case RouteId_ScreenshotHeader:
+	{
+		uint32 Id = EventData.GetValue<uint32>("Id");
+		TSharedPtr<FScreenshot> Screenshot = ScreenshotProvider.AddScreenshot(Id);
+		Screenshot->Id = Id;
+
+		EventData.GetString("Name", Screenshot->Name);
+		
+		uint64 Cycle = EventData.GetValue<uint64>("Cycle");
+		Screenshot->Timestamp = Context.EventTime.AsSeconds(Cycle);
+
+		Screenshot->Width = EventData.GetValue<uint32>("Width");
+		Screenshot->Height = EventData.GetValue<uint32>("Height");
+		Screenshot->ChunkNum = EventData.GetValue<uint32>("TotalChunkNum");
+		Screenshot->Size = EventData.GetValue<uint32>("Size");
+		Screenshot->Data.Reserve(Screenshot->Size);
+
+		FLogMessageSpec& LogMessageSpec = LogProvider.GetMessageSpec(Cycle);
+		LogMessageSpec.Category = &LogProvider.GetCategory(FLogProvider::ReservedLogCategory_Screenshot);
+		LogMessageSpec.Category->Name = TEXT("Screenshot");
+		LogMessageSpec.Line = Id;
+		LogMessageSpec.File = nullptr;
+		LogMessageSpec.FormatString = nullptr;
+		LogMessageSpec.Verbosity = ELogVerbosity::Log;
+
+		LogProvider.AppendMessage(Cycle, Screenshot->Timestamp, Screenshot->Name);
+
+		break;
+	}
+
+	case RouteId_ScreenshotChunk:
+	{
+		uint32 Id = EventData.GetValue<uint32>("Id");
+		uint16 ChunkNum = EventData.GetValue<uint16>("ChunkNum");
+		uint16 Size = EventData.GetValue<uint16>("Size");
+		TArrayView<const uint8> Data = EventData.GetArrayView<uint8>("Data");
+
+		ScreenshotProvider.AddScreenshotChunk(Id, ChunkNum, Size, Data);
+
+		break;
+	}
 	// Begin retired events
 	//
 	case RouteId_RegisterGameThread:
@@ -175,6 +224,7 @@ bool FMiscTraceAnalyzer::OnEvent(uint16 RouteId, EStyle Style, const FOnEventCon
 		ThreadProvider.AddGameThread(ThreadId);
 		break;
 	}
+
 	case RouteId_CreateThread:
 	{
 		const uint32 CreatedThreadId = FTraceAnalyzerUtils::GetThreadIdField(Context, "CreatedThreadId");
@@ -189,6 +239,7 @@ bool FMiscTraceAnalyzer::OnEvent(uint16 RouteId, EStyle Style, const FOnEventCon
 		}
 		break;
 	}
+
 	case RouteId_SetThreadGroup:
 	{
 		const TCHAR* GroupName = Session.StoreString(ANSI_TO_TCHAR(reinterpret_cast<const char*>(EventData.GetAttachment())));
@@ -196,6 +247,7 @@ bool FMiscTraceAnalyzer::OnEvent(uint16 RouteId, EStyle Style, const FOnEventCon
 		ThreadProvider.SetThreadGroup(ThreadId, GroupName);
 		break;
 	}
+
 	case RouteId_BeginThreadGroupScope:
 	{
 		const TCHAR* GroupName = Session.StoreString(ANSI_TO_TCHAR(reinterpret_cast<const char*>(EventData.GetAttachment())));
@@ -204,6 +256,7 @@ bool FMiscTraceAnalyzer::OnEvent(uint16 RouteId, EStyle Style, const FOnEventCon
 		ThreadState->ThreadGroupStack.Push(GroupName);
 		break;
 	}
+
 	case RouteId_EndThreadGroupScope:
 	{
 		const uint32 CurrentThreadId = FTraceAnalyzerUtils::GetThreadIdField(Context, "CurrentThreadId");
diff --git a/Engine/Source/Developer/TraceServices/Private/Analyzers/MiscTraceAnalysis.h b/Engine/Source/Developer/TraceServices/Private/Analyzers/MiscTraceAnalysis.h
--- a/Engine/Source/Developer/TraceServices/Private/Analyzers/MiscTraceAnalysis.h
+++ b/Engine/Source/Developer/TraceServices/Private/Analyzers/MiscTraceAnalysis.h
@@ -16,6 +16,7 @@ class FBookmarkProvider;
 class FLogProvider;
 class FFrameProvider;
 class FChannelProvider;
+class FScreenshotProvider;
 
 class FMiscTraceAnalyzer
 	: public UE::Trace::IAnalyzer
@@ -26,7 +27,8 @@ public:
 					   FBookmarkProvider& BookmarkProvider,
 					   FLogProvider& LogProvider,
 					   FFrameProvider& FrameProvider, 
-					   FChannelProvider& ChannelProvider);
+					   FChannelProvider& ChannelProvider,
+					   FScreenshotProvider& ScreenshotProvider);
 	virtual void OnAnalysisBegin(const FOnAnalysisContext& Context) override;
 	virtual void OnThreadInfo(const FThreadInfo& ThreadInfo) override;
 	virtual bool OnEvent(uint16 RouteId, EStyle Style, const FOnEventContext& Context) override;
@@ -49,6 +51,8 @@ private:
 		RouteId_EndRenderFrame,
 		RouteId_ChannelAnnounce,
 		RouteId_ChannelToggle,
+		RouteId_ScreenshotHeader,
+		RouteId_ScreenshotChunk,
 	};
 
 	struct FThreadState
@@ -66,6 +70,7 @@ private:
 	FLogProvider& LogProvider;
 	FFrameProvider& FrameProvider;
 	FChannelProvider& ChannelProvider;
+	FScreenshotProvider& ScreenshotProvider;
 	TMap<uint32, TSharedRef<FThreadState>> ThreadStateMap;
 	uint64 LastFrameCycle[TraceFrameType_Count] = { 0, 0 };
 };
diff --git a/Engine/Source/Developer/TraceServices/Private/Model/Log.cpp b/Engine/Source/Developer/TraceServices/Private/Model/Log.cpp
--- a/Engine/Source/Developer/TraceServices/Private/Model/Log.cpp
+++ b/Engine/Source/Developer/TraceServices/Private/Model/Log.cpp
@@ -89,6 +89,20 @@ void FLogProvider::AppendMessage(uint64 LogPoint, double Time, const uint8* Form
 	}
 }
 
+void FLogProvider::AppendMessage(uint64 LogPoint, double Time, const FString& Message)
+{
+	Session.WriteAccessCheck();
+	FLogMessageSpec** FindSpec = SpecMap.Find(LogPoint);
+	if (FindSpec && (*FindSpec)->Verbosity != ELogVerbosity::SetColor)
+	{
+		FLogMessageInternal& InternalMessage = Messages.PushBack();
+		InternalMessage.Time = Time;
+		InternalMessage.Spec = *FindSpec;
+		InternalMessage.Message = Session.StoreString(Message);
+		Session.UpdateDurationSeconds(Time);
+	}
+}
+
 uint64 FLogProvider::GetMessageCount() const
 {
 	Session.ReadAccessCheck();
diff --git a/Engine/Source/Developer/TraceServices/Private/Model/LogPrivate.h b/Engine/Source/Developer/TraceServices/Private/Model/LogPrivate.h
--- a/Engine/Source/Developer/TraceServices/Private/Model/LogPrivate.h
+++ b/Engine/Source/Developer/TraceServices/Private/Model/LogPrivate.h
@@ -36,7 +36,8 @@ class FLogProvider :
 public:
 	enum
 	{
-		ReservedLogCategory_Bookmark = 0
+		ReservedLogCategory_Bookmark = 0,
+		ReservedLogCategory_Screenshot = 1
 	};
 
 	static const FName ProviderName;
@@ -47,6 +48,7 @@ public:
 	FLogCategoryInfo& GetCategory(uint64 CategoryPointer);
 	FLogMessageSpec& GetMessageSpec(uint64 LogPoint);
 	void AppendMessage(uint64 LogPoint, double Time, const uint8* FormatArgs);
+	void AppendMessage(uint64 LogPoint, double Time, const FString& Message);
 
 	virtual uint64 GetMessageCount() const override;
 	virtual bool ReadMessage(uint64 Index, TFunctionRef<void(const FLogMessageInfo&)> Callback) const override;
diff --git a/Engine/Source/Developer/TraceServices/Private/Model/ScreenshotProvider.cpp b/Engine/Source/Developer/TraceServices/Private/Model/ScreenshotProvider.cpp
new file mode 100644
--- /dev/null
+++ b/Engine/Source/Developer/TraceServices/Private/Model/ScreenshotProvider.cpp
@@ -0,0 +1,58 @@
+// Copyright Epic Games, Inc. All Rights Reserved.
+
+#include "Model/ScreenshotProviderPrivate.h"
+#include "AnalysisServicePrivate.h"
+#include "Common/FormatArgs.h"
+
+namespace TraceServices
+{
+
+const FName FScreenshotProvider::ProviderName("ScreenshotProvider");
+
+FScreenshotProvider::FScreenshotProvider(IAnalysisSession& InSession)
+	: Session(InSession)
+{
+
+} 
+
+TSharedPtr<FScreenshot> FScreenshotProvider::AddScreenshot(uint32 Id)
+{
+	TSharedPtr<FScreenshot> NewScreenshot = Screenshots.Add(Id, MakeShared<FScreenshot>());
+	return NewScreenshot;
+}
+
+void FScreenshotProvider::AddScreenshotChunk(uint32 Id, uint32 ChunkNum, uint16 Size, const TArrayView<const uint8>& ChunkData)
+{
+	TSharedPtr<FScreenshot> *ScreenshotPtr = Screenshots.Find(Id);
+	
+	if (ScreenshotPtr == nullptr)
+	{
+		return;
+	}
+
+	TSharedPtr<FScreenshot> Screenshot = *ScreenshotPtr;
+
+	constexpr uint32 ChunkSize = TNumericLimits<uint16>::Max();
+	check(Screenshot->Data.Num() == ChunkNum * ChunkSize);
+
+	Screenshot->Data.Append(ChunkData.GetData(), Size);
+}
+
+const TSharedPtr<const FScreenshot> FScreenshotProvider::GetScreenshot(uint32 Id) const
+{
+	const TSharedPtr<FScreenshot>* ScreenshotPtr = Screenshots.Find(Id);
+
+	if (ScreenshotPtr == nullptr)
+	{
+		return nullptr;
+	}
+
+	return *ScreenshotPtr;
+}
+
+const IScreenshotProvider& ReadScreenshotProvider(const IAnalysisSession& Session)
+{
+	return *Session.ReadProvider<IScreenshotProvider>(FScreenshotProvider::ProviderName);
+}
+
+} // namespace TraceServices
diff --git a/Engine/Source/Developer/TraceServices/Private/Model/ScreenshotProviderPrivate.h b/Engine/Source/Developer/TraceServices/Private/Model/ScreenshotProviderPrivate.h
new file mode 100644
--- /dev/null
+++ b/Engine/Source/Developer/TraceServices/Private/Model/ScreenshotProviderPrivate.h
@@ -0,0 +1,30 @@
+// Copyright Epic Games, Inc. All Rights Reserved.
+
+#pragma once
+
+#include "TraceServices/AnalysisService.h"
+#include "Templates/SharedPointer.h"
+#include "TraceServices/Model/Screenshot.h"
+
+namespace TraceServices
+{
+
+class FScreenshotProvider
+	: public IScreenshotProvider
+{
+public:
+	static const FName ProviderName;
+
+	explicit FScreenshotProvider(IAnalysisSession& Session);
+	virtual ~FScreenshotProvider() {}
+
+	TSharedPtr<FScreenshot> AddScreenshot(uint32 Id);
+	void AddScreenshotChunk(uint32 Id, uint32 ChunkNum, uint16 Size, const TArrayView<const uint8>& ChunkData);
+	virtual const TSharedPtr<const FScreenshot> GetScreenshot(uint32 Id) const override;
+
+private:
+	TMap<uint32, TSharedPtr<FScreenshot>> Screenshots;
+	IAnalysisSession& Session;
+};
+
+} // namespace TraceServices
diff --git a/Engine/Source/Developer/TraceServices/Public/TraceServices/Model/Screenshot.h b/Engine/Source/Developer/TraceServices/Public/TraceServices/Model/Screenshot.h
new file mode 100644
--- /dev/null
+++ b/Engine/Source/Developer/TraceServices/Public/TraceServices/Model/Screenshot.h
@@ -0,0 +1,34 @@
+// Copyright Epic Games, Inc. All Rights Reserved.
+
+#pragma  once
+
+#include "TraceServices/Model/AnalysisSession.h"
+#include "Templates/Function.h"
+
+namespace TraceServices
+{
+
+struct FScreenshot
+{
+	uint32 Id;
+	FString Name;
+	double Timestamp;
+	uint32 Width;
+	uint32 Height;
+	uint32 ChunkNum;
+	uint32 Size;
+	TArray<uint8> Data;
+};
+
+class IScreenshotProvider
+	: public IProvider
+{
+public:
+	virtual ~IScreenshotProvider() = default;
+
+	virtual const TSharedPtr<const FScreenshot> GetScreenshot(uint32 Id) const = 0 ;
+};
+
+TRACESERVICES_API const IScreenshotProvider& ReadScreenshotProvider(const IAnalysisSession& Session);
+
+} // namespace TraceServices
diff --git a/Engine/Source/Runtime/Core/Private/ProfilingDebugging/MiscTrace.cpp b/Engine/Source/Runtime/Core/Private/ProfilingDebugging/MiscTrace.cpp
--- a/Engine/Source/Runtime/Core/Private/ProfilingDebugging/MiscTrace.cpp
+++ b/Engine/Source/Runtime/Core/Private/ProfilingDebugging/MiscTrace.cpp
@@ -10,6 +10,7 @@
 
 UE_TRACE_CHANNEL(FrameChannel)
 UE_TRACE_CHANNEL(BookmarkChannel)
+UE_TRACE_CHANNEL(ScreenshotChannel)
 
 UE_TRACE_EVENT_BEGIN(Misc, BookmarkSpec, NoSync|Important)
 	UE_TRACE_EVENT_FIELD(const void*, BookmarkPoint)
@@ -34,6 +35,23 @@ UE_TRACE_EVENT_BEGIN(Misc, EndFrame)
 	UE_TRACE_EVENT_FIELD(uint8, FrameType)
 UE_TRACE_EVENT_END()
 
+UE_TRACE_EVENT_BEGIN(Misc, ScreenshotHeader)
+	UE_TRACE_EVENT_FIELD(uint32, Id)
+	UE_TRACE_EVENT_FIELD(UE::Trace::WideString, Name)
+	UE_TRACE_EVENT_FIELD(uint64, Cycle)
+	UE_TRACE_EVENT_FIELD(uint32, Width)
+	UE_TRACE_EVENT_FIELD(uint32, Height)
+	UE_TRACE_EVENT_FIELD(uint32, TotalChunkNum)
+	UE_TRACE_EVENT_FIELD(uint32, Size)
+UE_TRACE_EVENT_END()
+
+UE_TRACE_EVENT_BEGIN(Misc, ScreenshotChunk)
+	UE_TRACE_EVENT_FIELD(uint32, Id)
+	UE_TRACE_EVENT_FIELD(uint32, ChunkNum)
+	UE_TRACE_EVENT_FIELD(uint16, Size)
+	UE_TRACE_EVENT_FIELD(uint8[], Data)
+UE_TRACE_EVENT_END()
+
 void FMiscTrace::OutputBookmarkSpec(const void* BookmarkPoint, const ANSICHAR* File, int32 Line, const TCHAR* Format)
 {
 	uint16 FileNameLen = uint16(strlen(File));
@@ -81,4 +99,39 @@ void FMiscTrace::OutputEndFrame(ETraceFrameType FrameType)
 		<< EndFrame.FrameType((uint8)FrameType);
 }
 
+void FMiscTrace::OutputScreenshot(const TCHAR* Name, uint64 Cycle, uint32 Width, uint32 Height, TArray64<uint8> Data)
+{
+	static std::atomic<uint32> ScreenshotId = 0;
+
+	const uint32 DataSize = (uint32) Data.Num();
+	const uint32 MaxChunkSize = TNumericLimits<uint16>::Max();
+	uint32 ChunkNum = (DataSize + MaxChunkSize - 1) / MaxChunkSize;
+
+	uint32 Id = ScreenshotId.fetch_add(1);
+	UE_TRACE_LOG(Misc, ScreenshotHeader, ScreenshotChannel)
+		<< ScreenshotHeader.Id(Id)
+		<< ScreenshotHeader.Name(Name, uint16(FCString::Strlen(Name)))
+		<< ScreenshotHeader.Cycle(Cycle)
+		<< ScreenshotHeader.Width(Width)
+		<< ScreenshotHeader.Height(Height)
+		<< ScreenshotHeader.TotalChunkNum(ChunkNum)
+		<< ScreenshotHeader.Size(DataSize);
+
+	uint32 RemainingSize = DataSize;
+	for (uint32 Index = 0; Index < ChunkNum; ++Index)
+	{
+		uint16 Size = (uint16) FMath::Min(RemainingSize, MaxChunkSize);
+
+		uint8* ChunkData = Data.GetData() + MaxChunkSize * Index;
+		UE_TRACE_LOG(Misc, ScreenshotChunk, ScreenshotChannel)
+			<< ScreenshotChunk.Id(Id)
+			<< ScreenshotChunk.ChunkNum(Index)
+			<< ScreenshotChunk.Size(Size)
+			<< ScreenshotChunk.Data(ChunkData, Size);
+
+		RemainingSize -= Size;
+	}
+
+	check(RemainingSize == 0);
+}
 #endif
diff --git a/Engine/Source/Runtime/Core/Public/ProfilingDebugging/MiscTrace.h b/Engine/Source/Runtime/Core/Public/ProfilingDebugging/MiscTrace.h
--- a/Engine/Source/Runtime/Core/Public/ProfilingDebugging/MiscTrace.h
+++ b/Engine/Source/Runtime/Core/Public/ProfilingDebugging/MiscTrace.h
@@ -61,6 +61,8 @@ struct FMiscTrace
 	CORE_API static void OutputBeginFrame(ETraceFrameType FrameType);
 	CORE_API static void OutputEndFrame(ETraceFrameType FrameType);
 
+	CORE_API static void OutputScreenshot(const TCHAR* Name, uint64 Cycle, uint32 Width, uint32 Height, TArray64<uint8> Data);
+
 private:
 	CORE_API static void OutputBookmarkInternal(const void* BookmarkPoint, uint16 EncodedFormatArgsSize, uint8* EncodedFormatArgs);
 };
diff --git a/Engine/Source/Runtime/Engine/Private/ProfilingDebugging/TraceScreenshot.cpp b/Engine/Source/Runtime/Engine/Private/ProfilingDebugging/TraceScreenshot.cpp
new file mode 100644
--- /dev/null
+++ b/Engine/Source/Runtime/Engine/Private/ProfilingDebugging/TraceScreenshot.cpp
@@ -0,0 +1,108 @@
+// Copyright Epic Games, Inc. All Rights Reserved.
+#include "ProfilingDebugging/TraceScreenshot.h"
+
+#include "Containers/Map.h"
+#include "Engine/Engine.h"
+#include "Engine/GameViewportClient.h"
+#include "ImageUtils.h"
+#include "HAL/PlatformTime.h"
+#include "HAL/PlatformTLS.h"
+#include "HAL/TlsAutoCleanup.h"
+#include "Misc/DateTime.h"
+#include "ProfilingDebugging/MiscTrace.h"
+#include "Trace/Trace.inl"
+
+static void TraceScreenshotCommandCallback(const TArray<FString>& Args)
+{
+	FString Name;
+	if (Args.Num() > 1)
+	{
+		Name = Args[0];
+	}
+
+	if (FTraceScreenshot::Get() == nullptr)
+	{
+		FTraceScreenshot::CreateInstance();
+	}
+
+	FTraceScreenshot::Get()->RequestScreenshot(Name);
+}
+
+static FAutoConsoleCommand TraceScreenshotCmd(
+	TEXT("Trace.Screenshot"),
+	TEXT("[Name] Takes a screenshot and saves it in the trace."
+		" [Name] is the name of the screenshot."
+	),
+	FConsoleCommandWithArgsDelegate::CreateStatic(TraceScreenshotCommandCallback)
+);
+
+FTraceScreenshot* FTraceScreenshot::Instance = nullptr;
+
+void FTraceScreenshot::CreateInstance()
+{
+	if (Instance == nullptr)
+	{
+		Instance = new FTraceScreenshot();
+	}
+}
+
+FTraceScreenshot::FTraceScreenshot()
+{
+	if (GEngine && GEngine->GameViewport)
+	{
+		FlushRenderingCommands();
+		GEngine->GameViewport->OnScreenshotCaptured().AddRaw(this, &FTraceScreenshot::HandleScreenshotData);
+		FWorldDelegates::LevelRemovedFromWorld.AddRaw(this, &FTraceScreenshot::WorldDestroyed);
+	}
+}
+
+FTraceScreenshot::~FTraceScreenshot()
+{
+}
+
+void FTraceScreenshot::RequestScreenshot(FString Name)
+{
+	ScreenshotName = Name;
+	bool bShowUI = false;
+	FScreenshotRequest::RequestScreenshot(bShowUI);
+}
+
+void FTraceScreenshot::HandleScreenshotData(int32 InSizeX, int32 InSizeY, const TArray<FColor>& InImageData)
+{
+	if (ScreenshotName.IsEmpty())
+	{
+		ScreenshotName = FDateTime::Now().ToString(TEXT("Screenshot_%Y%m%d_%H%M%S"));
+	}
+
+	uint64 Cycles = FPlatformTime::Cycles64();
+	UE_LOG(LogCore, Display, TEXT("Tracing Screenshot \"%s\" taken with size: %d x %d"), *ScreenshotName, InSizeX, InSizeY);
+
+	int32 ResizedX = FMath::Min(640, InSizeX);
+	int32 ResizedY = (InSizeY * ResizedX) / InSizeX;
+
+	TArray<FColor> ResizedImage;
+	ResizedImage.SetNum(ResizedX * ResizedY);
+	FImageUtils::ImageResize(InSizeX, InSizeY, InImageData, ResizedX, ResizedY, ResizedImage, false);
+
+	TArray64<uint8> CompressedBitmap;
+	FImageUtils::PNGCompressImageArray(ResizedX, ResizedY, TArrayView64<const FColor>(ResizedImage.GetData(), ResizedImage.Num()), CompressedBitmap);
+
+	FMiscTrace::OutputScreenshot(*ScreenshotName, Cycles, ResizedX, ResizedY, CompressedBitmap);
+	ScreenshotName.Empty();
+}
+
+void FTraceScreenshot::WorldDestroyed(ULevel* InLevel, UWorld* InWorld)
+{
+	if (InLevel == nullptr && InWorld == World.Get())
+	{
+		UE_LOG(LogCore, Display, TEXT("Tracing screenshot \"%s\" skipped - level was removed from world before we got our screenshot"), *ScreenshotName);
+
+		if (Instance)
+		{
+			delete Instance;
+			Instance = nullptr;
+			GEngine->GameViewport->OnScreenshotCaptured().RemoveAll(this);
+		}
+		FWorldDelegates::LevelRemovedFromWorld.RemoveAll(this);
+	}
+}
\ No newline at end of file
diff --git a/Engine/Source/Runtime/Engine/Public/ProfilingDebugging/TraceScreenshot.h b/Engine/Source/Runtime/Engine/Public/ProfilingDebugging/TraceScreenshot.h
new file mode 100644
--- /dev/null
+++ b/Engine/Source/Runtime/Engine/Public/ProfilingDebugging/TraceScreenshot.h
@@ -0,0 +1,32 @@
+// Copyright Epic Games, Inc. All Rights Reserved.
+
+#pragma once
+
+#include "CoreMinimal.h"
+#include "Engine/World.h"
+
+class ULevel;
+
+class FTraceScreenshot
+{
+public:
+	static FTraceScreenshot* Get() { return Instance; }
+	static void CreateInstance();
+	void RequestScreenshot(FString Name);
+	void HandleScreenshotData(int32 InSizeX, int32 InSizeY, const TArray<FColor>& InImageData);
+	void WorldDestroyed(ULevel* InLevel, UWorld* InWorld);
+
+private:
+	FTraceScreenshot();
+	FTraceScreenshot(const FTraceScreenshot& Other) {}
+	FTraceScreenshot(const FTraceScreenshot&& Other) {}
+	void operator =(const FTraceScreenshot& Other) {}
+	~FTraceScreenshot();
+	void Unbind();
+
+private:
+	static FTraceScreenshot* Instance;
+
+	FString ScreenshotName;
+	TWeakObjectPtr<UWorld> World;
+};
\ No newline at end of file
