[Insights] Context Switches: Changed "Core Tracks" option name into "CPU Core Tracks".

diff --git a/Engine/Source/Developer/TraceInsights/Private/Insights/ContextSwitches/ViewModels/ContextSwitchesSharedState.cpp b/Engine/Source/Developer/TraceInsights/Private/Insights/ContextSwitches/ViewModels/ContextSwitchesSharedState.cpp
--- a/Engine/Source/Developer/TraceInsights/Private/Insights/ContextSwitches/ViewModels/ContextSwitchesSharedState.cpp
+++ b/Engine/Source/Developer/TraceInsights/Private/Insights/ContextSwitches/ViewModels/ContextSwitchesSharedState.cpp
@@ -52,14 +52,53 @@ FContextSwitchesStateCommands::~FContextSwitchesStateCommands()
 PRAGMA_DISABLE_OPTIMIZATION
 void FContextSwitchesStateCommands::RegisterCommands()
 {
-	UI_COMMAND(Command_ShowCoreTracks, "Core Tracks", "Shows/hides the CPU Core tracks.", EUserInterfaceActionType::ToggleButton, FInputChord(EModifierKey::Alt, EKeys::C));
-	UI_COMMAND(Command_ShowContextSwitches, "Context Switches", "Shows/hides the context switches on top of CPU timing tracks.", EUserInterfaceActionType::ToggleButton,FInputChord(EModifierKey::Shift, EKeys::C));
-	UI_COMMAND(Command_ShowOverlays, "Overlays", "Extends the visualization of context switches over the CPU timing tracks.", EUserInterfaceActionType::ToggleButton, FInputChord(EModifierKey::Shift, EKeys::O));
-	UI_COMMAND(Command_ShowExtendedLines, "Extended Lines", "Shows/hides the extended vertical lines at edges of each context switch event.", EUserInterfaceActionType::ToggleButton, FInputChord(EModifierKey::Shift, EKeys::L));
-	UI_COMMAND(Command_NavigateToCpuThreadEvent, "Go To CPU Thread Track", "Selects the context switch event in the corresponding CPU Thread track.", EUserInterfaceActionType::Button,	FInputChord(/*EModifierKey::Control, EKeys::Tab*/));
-	UI_COMMAND(Command_DockCpuThreadTrackToBottom, "Dock CPU Thread Track To Bottom", "Docks the corresponding CPU Thread track to the bottom of the Timing view.",	EUserInterfaceActionType::Button, FInputChord());
-	UI_COMMAND(Command_NavigateToCpuCoreEvent, "Go To CPU Core Track", "Selects the timing event in the corresponding CPU Core track.",	EUserInterfaceActionType::Button, FInputChord(EModifierKey::Control, EKeys::Tab));
-	UI_COMMAND(Command_DockCpuCoreTrackToTop, "Dock CPU Core Track To Top",	"Docks the corresponding CPU Core track to the top of the Timing view.", EUserInterfaceActionType::Button, FInputChord());
+	UI_COMMAND(Command_ShowCoreTracks,
+		"CPU Core Tracks",
+		"Shows/hides the CPU Core tracks.",
+		EUserInterfaceActionType::ToggleButton,
+		FInputChord(EModifierKey::Alt, EKeys::C));
+
+	UI_COMMAND(Command_ShowContextSwitches,
+		"Context Switches",
+		"Shows/hides the context switches on top of CPU timing tracks.",
+		EUserInterfaceActionType::ToggleButton,
+		FInputChord(EModifierKey::Shift, EKeys::C));
+
+	UI_COMMAND(Command_ShowOverlays,
+		"Overlays",
+		"Extends the visualization of context switches over the CPU timing tracks.",
+		EUserInterfaceActionType::ToggleButton,
+		FInputChord(EModifierKey::Shift, EKeys::O));
+
+	UI_COMMAND(Command_ShowExtendedLines,
+		"Extended Lines",
+		"Shows/hides the extended vertical lines at edges of each context switch event.",
+		EUserInterfaceActionType::ToggleButton,
+		FInputChord(EModifierKey::Shift, EKeys::L));
+
+	UI_COMMAND(Command_NavigateToCpuThreadEvent,
+		"Go To CPU Thread Track",
+		"Selects the context switch event in the corresponding CPU Thread track.",
+		EUserInterfaceActionType::Button,
+		FInputChord(/*EModifierKey::Control, EKeys::Tab*/));
+
+	UI_COMMAND(Command_DockCpuThreadTrackToBottom,
+		"Dock CPU Thread Track To Bottom",
+		"Docks the corresponding CPU Thread track to the bottom of the Timing view.",
+		EUserInterfaceActionType::Button,
+		FInputChord());
+
+	UI_COMMAND(Command_NavigateToCpuCoreEvent,
+		"Go To CPU Core Track",
+		"Selects the timing event in the corresponding CPU Core track.",
+		EUserInterfaceActionType::Button,
+		FInputChord(EModifierKey::Control, EKeys::Tab));
+
+	UI_COMMAND(Command_DockCpuCoreTrackToTop,
+		"Dock CPU Core Track To Top",
+		"Docks the corresponding CPU Core track to the top of the Timing view.",
+		EUserInterfaceActionType::Button,
+		FInputChord());
 }
 PRAGMA_ENABLE_OPTIMIZATION
 
