[Insights] Convert the LLM tag filter to use a Suggestion Text Box.
Suggestion text boxes in Insights filter will now display all options when the up or down arrow is pressed and the text box is empty.

diff --git a/Engine/Source/Developer/TraceInsights/Private/Insights/MemoryProfiler/ViewModels/MemAllocTable.cpp b/Engine/Source/Developer/TraceInsights/Private/Insights/MemoryProfiler/ViewModels/MemAllocTable.cpp
--- a/Engine/Source/Developer/TraceInsights/Private/Insights/MemoryProfiler/ViewModels/MemAllocTable.cpp
+++ b/Engine/Source/Developer/TraceInsights/Private/Insights/MemoryProfiler/ViewModels/MemAllocTable.cpp
@@ -684,7 +684,9 @@ void FMemAllocTable::AddDefaultColumns()
 		Column.SetTitleName(LOCTEXT("TagColumnTitle", "LLM Tag"));
 		Column.SetDescription(LOCTEXT("TagColumnDesc", "LLM tag of allocation"));
 
-		Column.SetFlags(ETableColumnFlags::ShouldBeVisible | ETableColumnFlags::CanBeHidden | ETableColumnFlags::CanBeFiltered);
+		// This column is filtered with a custom filter with suggestions 
+		// so we do not mark it here as CanBeFiltered to prevent it from having a default string filter set.
+		Column.SetFlags(ETableColumnFlags::ShouldBeVisible | ETableColumnFlags::CanBeHidden);
 
 		Column.SetHorizontalAlignment(HAlign_Left);
 		Column.SetInitialWidth(120.0f);
diff --git a/Engine/Source/Developer/TraceInsights/Private/Insights/MemoryProfiler/Widgets/SMemAllocTableTreeView.cpp b/Engine/Source/Developer/TraceInsights/Private/Insights/MemoryProfiler/Widgets/SMemAllocTableTreeView.cpp
--- a/Engine/Source/Developer/TraceInsights/Private/Insights/MemoryProfiler/Widgets/SMemAllocTableTreeView.cpp
+++ b/Engine/Source/Developer/TraceInsights/Private/Insights/MemoryProfiler/Widgets/SMemAllocTableTreeView.cpp
@@ -42,6 +42,7 @@ namespace Insights
 {
 
 const int SMemAllocTableTreeView::FullCallStackIndex = 0x0000FFFFF;
+const int SMemAllocTableTreeView::LLMFilterIndex = 0x0000FFFFE;
 
 ////////////////////////////////////////////////////////////////////////////////////////////////////
 
@@ -1141,6 +1142,12 @@ bool SMemAllocTableTreeView::ApplyCustomAdvancedFilters(const FTableTreeNodePtr&
 		Context.SetFilterData<FString>(FullCallStackIndex, MemNodePtr->GetFullCallstack().ToString());
 	}
 
+	if (FilterConfigurator && FilterConfigurator->IsKeyUsed(LLMFilterIndex))
+	{
+		FMemAllocNodePtr MemNodePtr = StaticCastSharedPtr<FMemAllocNode>(NodePtr);
+		Context.SetFilterData<FString>(LLMFilterIndex, MemNodePtr->GetMemAlloc()->GetTag());
+	}
+
 	return true;
 }
 
@@ -1152,6 +1159,38 @@ void SMemAllocTableTreeView::AddCustomAdvancedFilters()
 
 	AvailableFilters->Add(MakeShared<FFilter>(FullCallStackIndex, LOCTEXT("FullCallstack", "Full Callstack"), LOCTEXT("SearchFullCallstack", "Search in all the callstack frames"), EFilterDataType::String, FFilterService::Get()->GetStringOperators()));
 	Context.AddFilterData<FString>(FullCallStackIndex, FString());
+
+	TSharedPtr<FFilterWithSuggestions> LLMCategoryFilter = MakeShared<FFilterWithSuggestions>(static_cast<int32>(LLMFilterIndex), LOCTEXT("LLMTag", "LLM Tag"), LOCTEXT("LLMTag", "LLM Tag"), EFilterDataType::String, FFilterService::Get()->GetStringOperators());
+	Context.AddFilterData<FString>(LLMFilterIndex, FString());
+	LLMCategoryFilter->Callback = [this](const FString& Text, TArray<FString>& OutSuggestions)
+	{
+		this->PopulateLLMTagSuggestionList(Text, OutSuggestions);
+	};
+
+	AvailableFilters->Add(LLMCategoryFilter);
+}
+
+////////////////////////////////////////////////////////////////////////////////////////////////////
+
+void SMemAllocTableTreeView::PopulateLLMTagSuggestionList(const FString& Text, TArray<FString>& OutSuggestions)
+{
+	const TraceServices::IAllocationsProvider* AllocationsProvider = TraceServices::ReadAllocationsProvider(*Session.Get());
+	if (!AllocationsProvider)
+	{
+		return;
+	}
+
+	TraceServices::IAllocationsProvider::FReadScopeLock _(*AllocationsProvider);
+
+	AllocationsProvider->EnumerateTags([&OutSuggestions, &Text](const TCHAR* Display, const TCHAR* FullPath, TraceServices::TagIdType CurrentTag, TraceServices::TagIdType ParentTag)
+	{
+		if (Text.IsEmpty() || FCString::Stristr(Display, *Text))
+		{
+			OutSuggestions.Push(Display);
+		}
+
+		return true;
+	});
 }
 
 ////////////////////////////////////////////////////////////////////////////////////////////////////
diff --git a/Engine/Source/Developer/TraceInsights/Private/Insights/MemoryProfiler/Widgets/SMemAllocTableTreeView.h b/Engine/Source/Developer/TraceInsights/Private/Insights/MemoryProfiler/Widgets/SMemAllocTableTreeView.h
--- a/Engine/Source/Developer/TraceInsights/Private/Insights/MemoryProfiler/Widgets/SMemAllocTableTreeView.h
+++ b/Engine/Source/Developer/TraceInsights/Private/Insights/MemoryProfiler/Widgets/SMemAllocTableTreeView.h
@@ -139,9 +139,11 @@ private:
 	TSharedRef<SWidget> ViewPreset_OnGenerateWidget(TSharedRef<IViewPreset> InPreset);
 	FText ViewPreset_GetSelectedText() const;
 	FText ViewPreset_GetSelectedToolTipText() const;
+	void PopulateLLMTagSuggestionList(const FString& Text, TArray<FString>& OutSuggestions);
 
 private:
 	const static int FullCallStackIndex;
+	const static int LLMFilterIndex;
 	int32 TabIndex = -1;
 	TSharedPtr<FMemoryRuleSpec> Rule = nullptr;
 	double TimeMarkers[4];
diff --git a/Engine/Source/Developer/TraceInsights/Private/Insights/Table/Widgets/STableTreeView.cpp b/Engine/Source/Developer/TraceInsights/Private/Insights/Table/Widgets/STableTreeView.cpp
--- a/Engine/Source/Developer/TraceInsights/Private/Insights/Table/Widgets/STableTreeView.cpp
+++ b/Engine/Source/Developer/TraceInsights/Private/Insights/Table/Widgets/STableTreeView.cpp
@@ -2851,6 +2851,11 @@ FReply STableTreeView::OnAdvancedFiltersClicked()
 
 		for (const TSharedRef<FTableColumn>& Column : Table->GetColumns())
 		{
+			if (!Column->CanBeFiltered())
+			{
+				continue;
+			}
+
 			switch (Column->GetDataType())
 			{
 				case ETableCellDataType::Int64:
@@ -2911,6 +2916,11 @@ bool STableTreeView::ApplyAdvancedFilters(const FTableTreeNodePtr& NodePtr)
 
 	for (const TSharedRef<FTableColumn>& Column : Table->GetColumns())
 	{
+		if (!Column->CanBeFiltered())
+		{
+			continue;
+		}
+
 		switch (Column->GetDataType())
 		{
 			case ETableCellDataType::Int64:
diff --git a/Engine/Source/Developer/TraceInsights/Private/Insights/Widgets/SFilterConfiguratorRow.cpp b/Engine/Source/Developer/TraceInsights/Private/Insights/Widgets/SFilterConfiguratorRow.cpp
--- a/Engine/Source/Developer/TraceInsights/Private/Insights/Widgets/SFilterConfiguratorRow.cpp
+++ b/Engine/Source/Developer/TraceInsights/Private/Insights/Widgets/SFilterConfiguratorRow.cpp
@@ -107,6 +107,7 @@ TSharedRef<SWidget> SFilterConfiguratorRow::GenerateWidgetForColumn(const FName&
 					.ToolTipText(this, &SFilterConfiguratorRow::GetTextBoxTooltipText)
 					.HintText(this, &SFilterConfiguratorRow::GetTextBoxHintText)
 					.OnShowingSuggestions(FOnShowingSuggestions::CreateSP(this, &SFilterConfiguratorRow::SuggestionTextBox_GetSuggestions))
+					.OnShowingHistory(FOnShowingHistory::CreateSP(this, &SFilterConfiguratorRow::SuggestionTextBox_GetHistory))
 				];
 		}
 		else
@@ -536,6 +537,14 @@ void SFilterConfiguratorRow::SuggestionTextBox_GetSuggestions(const FString& Tex
 
 ///////////////////////////////////////////////////////////////////////////////////////////////////
 
+void SFilterConfiguratorRow::SuggestionTextBox_GetHistory(TArray<FString>& Suggestions)
+{
+	// We show all suggestion instead of history when arrow up/down is used.
+	SuggestionTextBox_GetSuggestions(FString(), Suggestions);
+}
+
+///////////////////////////////////////////////////////////////////////////////////////////////////
+
 void SFilterConfiguratorRow::SuggestionTextBox_OnValueChanged(const FText& InNewText)
 {
 	SuggestionTextBoxValue = InNewText.ToString();
diff --git a/Engine/Source/Developer/TraceInsights/Private/Insights/Widgets/SFilterConfiguratorRow.h b/Engine/Source/Developer/TraceInsights/Private/Insights/Widgets/SFilterConfiguratorRow.h
--- a/Engine/Source/Developer/TraceInsights/Private/Insights/Widgets/SFilterConfiguratorRow.h
+++ b/Engine/Source/Developer/TraceInsights/Private/Insights/Widgets/SFilterConfiguratorRow.h
@@ -68,6 +68,8 @@ private:
 	bool TextBox_OnVerifyTextChanged(const FText& InText, FText& OutErrorMessage);
 
 	void SuggestionTextBox_GetSuggestions(const FString& Text, TArray<FString>& Suggestions);
+	void SuggestionTextBox_GetHistory(TArray<FString>& Suggestions);
+
 	void SuggestionTextBox_OnValueChanged(const FText& InNewText);
 	FText SuggestionTextBox_GetValue() const;
 	void SuggestionTextBox_OnValueCommitted(const FText& InNewText, ETextCommit::Type InTextCommit);
diff --git a/Engine/Source/Developer/TraceInsights/Private/Insights/Widgets/STimingView.cpp b/Engine/Source/Developer/TraceInsights/Private/Insights/Widgets/STimingView.cpp
--- a/Engine/Source/Developer/TraceInsights/Private/Insights/Widgets/STimingView.cpp
+++ b/Engine/Source/Developer/TraceInsights/Private/Insights/Widgets/STimingView.cpp
@@ -5047,7 +5047,7 @@ void STimingView::PopulateTrackSuggestionList(const FString& Text, TArray<FStrin
 {
 	for (auto& Entry : AllTracks)
 	{
-		if (Entry.Value->GetName().Contains(Text))
+		if (Text.IsEmpty() || Entry.Value->GetName().Contains(Text))
 		{
 			OutSuggestions.Add(Entry.Value->GetName());
 		}
@@ -5074,6 +5074,11 @@ void STimingView::PopulateTimerNameSuggestionList(const FString& Text, TArray<FS
 			const TraceServices::FTimingProfilerTimer* Timer = TimerReader->GetTimer(TimerIndex);
 			if (Timer && Timer->Name)
 			{
+				if (Text.IsEmpty())
+				{
+					OutSuggestions.Add(Timer->Name);
+					continue;
+				}
 				const TCHAR* FoundString = FCString::Stristr(Timer->Name, *Text);
 				if (FoundString)
 				{
