[Insights]
  - Tree Views: Further improved performance of sorting by FText columns and reduced memory cost while sorting. Added improved sorter for FText columns with available id (uint64).
  - Memory Insights: Improved performance of sorting by Function column by using the new sorter of FText values with id. Stress test (a tree with 3.5M allocs): speed improved from 4min to 27sec; memory usage improved from ~1 GiB to ~100 MiB.

diff --git a/Engine/Source/Developer/TraceInsights/Private/Insights/MemoryProfiler/ViewModels/MemAllocTable.cpp b/Engine/Source/Developer/TraceInsights/Private/Insights/MemoryProfiler/ViewModels/MemAllocTable.cpp
--- a/Engine/Source/Developer/TraceInsights/Private/Insights/MemoryProfiler/ViewModels/MemAllocTable.cpp
+++ b/Engine/Source/Developer/TraceInsights/Private/Insights/MemoryProfiler/ViewModels/MemAllocTable.cpp
@@ -774,12 +774,12 @@ void FMemAllocTable::AddDefaultColumns()
 
 						if (!Callstack)
 						{
-							return FTableCellValue(FText::FromString(GetCallstackNotAvailableString()));
+							return FTableCellValue(FText::FromString(GetCallstackNotAvailableString()), 0);
 						}
 
 						if (Callstack->Num() == 0)
 						{
-							return FTableCellValue(FText::FromString(GetEmptyCallstackString()));
+							return FTableCellValue(FText::FromString(GetEmptyCallstackString()), uint64(Callstack));
 						}
 
 						const TraceServices::FStackFrame* Frame = nullptr;
@@ -855,12 +855,28 @@ void FMemAllocTable::AddDefaultColumns()
 
 						TStringBuilder<1024> Str;
 						FormatStackFrame(*Frame, Str, EStackFrameFormatFlags::Module);
-						return FTableCellValue(FText::FromString(FString(Str)));
+						return FTableCellValue(FText::FromString(FString(Str)), uint64(Callstack));
 					}
 				}
 
 				return TOptional<FTableCellValue>();
 			}
+
+			virtual uint64 GetValueId(const FTableColumn& Column, const FBaseTreeNode& Node) const override
+			{
+				if (!Node.IsGroup()) //if (Node->Is<FMemAllocNode>())
+				{
+					const FMemAllocNode& MemAllocNode = static_cast<const FMemAllocNode&>(Node);
+					const FMemoryAlloc* Alloc = MemAllocNode.GetMemAlloc();
+					if (Alloc)
+					{
+						const TraceServices::FCallstack* Callstack = Alloc->GetCallstack();
+						return uint64(Callstack);
+					}
+				}
+
+				return 0;
+			}
 		};
 		TSharedRef<ITableCellValueGetter> Getter = MakeShared<FFunctionValueGetter>();
 		Column.SetValueGetter(Getter);
@@ -896,7 +912,7 @@ void FMemAllocTable::AddDefaultColumns()
 		TSharedRef<ITableCellValueFormatter> Formatter = MakeShared<FunctionValueFormatter>();
 		Column.SetValueFormatter(Formatter);
 
-		TSharedRef<ITableCellValueSorter> Sorter = MakeShared<FSorterByTextValue>(ColumnRef);
+		TSharedRef<ITableCellValueSorter> Sorter = MakeShared<FSorterByTextValueWithId>(ColumnRef);
 		Column.SetValueSorter(Sorter);
 
 		AddColumn(ColumnRef);
diff --git a/Engine/Source/Developer/TraceInsights/Private/Insights/Table/ViewModels/TableCellValue.h b/Engine/Source/Developer/TraceInsights/Private/Insights/Table/ViewModels/TableCellValue.h
--- a/Engine/Source/Developer/TraceInsights/Private/Insights/Table/ViewModels/TableCellValue.h
+++ b/Engine/Source/Developer/TraceInsights/Private/Insights/Table/ViewModels/TableCellValue.h
@@ -68,8 +68,8 @@ public:
 	explicit FTableCellValue(float Value) : DataType(ETableCellDataType::Float), Float(Value) {}
 	explicit FTableCellValue(double Value) : DataType(ETableCellDataType::Double), Double(Value) {}
 	explicit FTableCellValue(const TCHAR* Value) : DataType(ETableCellDataType::CString), CString(Value) {}
-	explicit FTableCellValue(TSharedPtr<ICustomTableCellValue> Value) : DataType(ETableCellDataType::Custom), Custom(Value) {}
-	explicit FTableCellValue(const FText& Value) : DataType(ETableCellDataType::Text), Custom(MakeShared<FTextCustomTableCellValue>(Value)) {}
+	explicit FTableCellValue(TSharedPtr<ICustomTableCellValue> Value, uint64 Id = 0) : DataType(ETableCellDataType::Custom), ValueId(Id), Custom(Value) {}
+	explicit FTableCellValue(const FText& Value, uint64 Id = 0) : DataType(ETableCellDataType::Text), ValueId(Id), Custom(MakeShared<FTextCustomTableCellValue>(Value)) {}
 
 	bool AsBool() const
 	{
@@ -165,6 +165,11 @@ public:
 		return FText::GetEmpty();
 	}
 
+	uint64 GetValueId() const
+	{
+		return ValueId;
+	}
+
 public:
 	ETableCellDataType DataType;
 
@@ -175,6 +180,7 @@ public:
 		float Float;
 		double Double;
 		const TCHAR* CString; // should be valid for the lieftime of the owner table
+		uint64 ValueId; // only used by Custom types
 	};
 
 	TSharedPtr<ICustomTableCellValue> Custom;
diff --git a/Engine/Source/Developer/TraceInsights/Private/Insights/Table/ViewModels/TableCellValueGetter.h b/Engine/Source/Developer/TraceInsights/Private/Insights/Table/ViewModels/TableCellValueGetter.h
--- a/Engine/Source/Developer/TraceInsights/Private/Insights/Table/ViewModels/TableCellValueGetter.h
+++ b/Engine/Source/Developer/TraceInsights/Private/Insights/Table/ViewModels/TableCellValueGetter.h
@@ -18,6 +18,7 @@ class ITableCellValueGetter
 {
 public:
 	virtual const TOptional<FTableCellValue> GetValue(const FTableColumn& Column, const FBaseTreeNode& Node) const = 0;
+	virtual uint64 GetValueId(const FTableColumn& Column, const FBaseTreeNode& Node) const = 0;
 };
 
 ////////////////////////////////////////////////////////////////////////////////////////////////////
@@ -29,6 +30,7 @@ public:
 	virtual ~FTableCellValueGetter() {}
 
 	virtual const TOptional<FTableCellValue> GetValue(const FTableColumn& Column, const FBaseTreeNode& Node) const override { return TOptional<FTableCellValue>(); }
+	virtual uint64 GetValueId(const FTableColumn& Column, const FBaseTreeNode& Node) const override { return 0; }
 };
 
 ////////////////////////////////////////////////////////////////////////////////////////////////////
diff --git a/Engine/Source/Developer/TraceInsights/Private/Insights/Table/ViewModels/TableCellValueSorter.cpp b/Engine/Source/Developer/TraceInsights/Private/Insights/Table/ViewModels/TableCellValueSorter.cpp
--- a/Engine/Source/Developer/TraceInsights/Private/Insights/Table/ViewModels/TableCellValueSorter.cpp
+++ b/Engine/Source/Developer/TraceInsights/Private/Insights/Table/ViewModels/TableCellValueSorter.cpp
@@ -6,6 +6,15 @@
 
 #define LOCTEXT_NAMESPACE "TableCellValueSorter"
 
+// Default pre-sorting (group nodes sorts above leaf nodes)
+#define INSIGHTS_DEFAULT_PRESORTING_NODES(A, B) \
+	{ \
+		if (A->IsGroup() != B->IsGroup()) \
+		{ \
+			return A->IsGroup(); \
+		} \
+	}
+
 // Default sorting
 #define INSIGHTS_DEFAULT_SORTING_NODES(A, B) return A->GetDefaultSortOrder() < B->GetDefaultSortOrder();
 
@@ -77,12 +86,16 @@ FSorterByName::FSorterByName(TSharedRef<FTableColumn> InColumnRef)
 {
 	AscendingCompareDelegate = [](const FBaseTreeNodePtr& A, const FBaseTreeNodePtr& B) -> bool
 	{
+		INSIGHTS_DEFAULT_PRESORTING_NODES(A, B)
+
 		// Sort by name (ascending).
 		return A->GetName().LexicalLess(B->GetName());
 	};
 
 	DescendingCompareDelegate = [](const FBaseTreeNodePtr& A, const FBaseTreeNodePtr& B) -> bool
 	{
+		INSIGHTS_DEFAULT_PRESORTING_NODES(A, B)
+
 		// Sort by name (descending).
 		return B->GetName().LexicalLess(A->GetName());
 	};
@@ -102,6 +115,8 @@ FSorterByTypeName::FSorterByTypeName(TSharedRef<FTableColumn> InColumnRef)
 {
 	AscendingCompareDelegate = [](const FBaseTreeNodePtr& A, const FBaseTreeNodePtr& B) -> bool
 	{
+		INSIGHTS_DEFAULT_PRESORTING_NODES(A, B)
+
 		const FName& TypeA = A->GetTypeName();
 		const FName& TypeB = B->GetTypeName();
 
@@ -118,6 +133,8 @@ FSorterByTypeName::FSorterByTypeName(TSharedRef<FTableColumn> InColumnRef)
 
 	DescendingCompareDelegate = [](const FBaseTreeNodePtr& A, const FBaseTreeNodePtr& B) -> bool
 	{
+		INSIGHTS_DEFAULT_PRESORTING_NODES(A, B)
+
 		const FName& TypeA = A->GetTypeName();
 		const FName& TypeB = B->GetTypeName();
 
@@ -144,6 +161,8 @@ FSorterByBoolValue::FSorterByBoolValue(TSharedRef<FTableColumn> InColumnRef)
 
 	AscendingCompareDelegate = [&Column = *InColumnRef](const FBaseTreeNodePtr& A, const FBaseTreeNodePtr& B) -> bool
 	{
+		INSIGHTS_DEFAULT_PRESORTING_NODES(A, B)
+
 		const TOptional<FTableCellValue> OptionalValueA = Column.GetValue(*A);
 		const TOptional<FTableCellValue> OptionalValueB = Column.GetValue(*B);
 
@@ -163,6 +182,8 @@ FSorterByBoolValue::FSorterByBoolValue(TSharedRef<FTableColumn> InColumnRef)
 
 	DescendingCompareDelegate = [&Column = * InColumnRef](const FBaseTreeNodePtr& A, const FBaseTreeNodePtr& B) -> bool
 	{
+		INSIGHTS_DEFAULT_PRESORTING_NODES(A, B)
+
 		const TOptional<FTableCellValue> OptionalValueA = Column.GetValue(*A);
 		const TOptional<FTableCellValue> OptionalValueB = Column.GetValue(*B);
 
@@ -189,51 +210,13 @@ FSorterByInt64Value::FSorterByInt64Value(TSharedRef<FTableColumn> InColumnRef)
 	: FBaseTableColumnSorter(InColumnRef)
 {
 	ensure(InColumnRef->GetDataType() == ETableCellDataType::Int64);
-
-	AscendingCompareDelegate = [&Column = *InColumnRef](const FBaseTreeNodePtr& A, const FBaseTreeNodePtr& B) -> bool
-	{
-		const TOptional<FTableCellValue> OptionalValueA = Column.GetValue(*A);
-		const TOptional<FTableCellValue> OptionalValueB = Column.GetValue(*B);
-
-		const int64 ValueA = OptionalValueA.IsSet() ? OptionalValueA.GetValue().Int64 : 0;
-		const int64 ValueB = OptionalValueB.IsSet() ? OptionalValueB.GetValue().Int64 : 0;
-
-		if (ValueA == ValueB)
-		{
-			INSIGHTS_DEFAULT_SORTING_NODES(A, B)
-		}
-		else
-		{
-			// Sort by value (ascending).
-			return ValueA < ValueB;
-		}
-	};
-
-	DescendingCompareDelegate = [&Column = *InColumnRef](const FBaseTreeNodePtr& A, const FBaseTreeNodePtr& B) -> bool
-	{
-		const TOptional<FTableCellValue> OptionalValueA = Column.GetValue(*A);
-		const TOptional<FTableCellValue> OptionalValueB = Column.GetValue(*B);
-
-		const int64 ValueA = OptionalValueA.IsSet() ? OptionalValueA.GetValue().Int64 : 0;
-		const int64 ValueB = OptionalValueB.IsSet() ? OptionalValueB.GetValue().Int64 : 0;
-
-		if (ValueA == ValueB)
-		{
-			INSIGHTS_DEFAULT_SORTING_NODES(A, B)
-		}
-		else
-		{
-			// Sort by value (descending).
-			return ValueB < ValueA;
-		}
-	};
 }
 
 ////////////////////////////////////////////////////////////////////////////////////////////////////
 
 void FSorterByInt64Value::Sort(TArray<FBaseTreeNodePtr>& NodesToSort, ESortMode SortMode) const
 {
-	// Note: This implementation is faster than using above delegates.
+	// Note: This implementation is faster than using delegates.
 	//       It caches the values before sorting, in this way it minimizes the calls to Column.GetValue().
 
 	const FTableColumn& Column = *ColumnRef;
@@ -261,6 +244,8 @@ void FSorterByInt64Value::Sort(TArray<FBaseTreeNodePtr>& NodesToSort, ESortMode
 	{
 		ElementsToSort.Sort([](const FSortElement& A, const FSortElement& B) -> bool
 		{
+			INSIGHTS_DEFAULT_PRESORTING_NODES(A.NodePtr, B.NodePtr)
+
 			if (A.Value == B.Value)
 			{
 				INSIGHTS_DEFAULT_SORTING_NODES(A.NodePtr, B.NodePtr)
@@ -276,6 +261,8 @@ void FSorterByInt64Value::Sort(TArray<FBaseTreeNodePtr>& NodesToSort, ESortMode
 	{
 		ElementsToSort.Sort([](const FSortElement& A, const FSortElement& B) -> bool
 		{
+			INSIGHTS_DEFAULT_PRESORTING_NODES(A.NodePtr, B.NodePtr)
+
 			if (A.Value == B.Value)
 			{
 				INSIGHTS_DEFAULT_SORTING_NODES(A.NodePtr, B.NodePtr)
@@ -302,51 +289,13 @@ FSorterByFloatValue::FSorterByFloatValue(TSharedRef<FTableColumn> InColumnRef)
 	: FBaseTableColumnSorter(InColumnRef)
 {
 	ensure(InColumnRef->GetDataType() == ETableCellDataType::Float);
-
-	AscendingCompareDelegate = [&Column = *InColumnRef](const FBaseTreeNodePtr& A, const FBaseTreeNodePtr& B) -> bool
-	{
-		const TOptional<FTableCellValue> OptionalValueA = Column.GetValue(*A);
-		const TOptional<FTableCellValue> OptionalValueB = Column.GetValue(*B);
-
-		const float ValueA = OptionalValueA.IsSet() ? OptionalValueA.GetValue().Float : 0.0f;
-		const float ValueB = OptionalValueB.IsSet() ? OptionalValueB.GetValue().Float : 0.0f;
-
-		if (ValueA == ValueB)
-		{
-			INSIGHTS_DEFAULT_SORTING_NODES(A, B)
-		}
-		else
-		{
-			// Sort by value (ascending).
-			return ValueA < ValueB;
-		}
-	};
-
-	DescendingCompareDelegate = [&Column = *InColumnRef](const FBaseTreeNodePtr& A, const FBaseTreeNodePtr& B) -> bool
-	{
-		const TOptional<FTableCellValue> OptionalValueA = Column.GetValue(*A);
-		const TOptional<FTableCellValue> OptionalValueB = Column.GetValue(*B);
-
-		const float ValueA = OptionalValueA.IsSet() ? OptionalValueA.GetValue().Float : 0.0f;
-		const float ValueB = OptionalValueB.IsSet() ? OptionalValueB.GetValue().Float : 0.0f;
-
-		if (ValueA == ValueB)
-		{
-			INSIGHTS_DEFAULT_SORTING_NODES(A, B)
-		}
-		else
-		{
-			// Sort by value (descending).
-			return ValueB < ValueA;
-		}
-	};
 }
 
 ////////////////////////////////////////////////////////////////////////////////////////////////////
 
 void FSorterByFloatValue::Sort(TArray<FBaseTreeNodePtr>& NodesToSort, ESortMode SortMode) const
 {
-	// Note: This implementation is faster than using above delegates.
+	// Note: This implementation is faster than using delegates.
 	//       It caches the values before sorting, in this way it minimizes the calls to Column.GetValue().
 
 	const FTableColumn& Column = *ColumnRef;
@@ -374,6 +323,8 @@ void FSorterByFloatValue::Sort(TArray<FBaseTreeNodePtr>& NodesToSort, ESortMode
 	{
 		ElementsToSort.Sort([](const FSortElement& A, const FSortElement& B) -> bool
 		{
+			INSIGHTS_DEFAULT_PRESORTING_NODES(A.NodePtr, B.NodePtr)
+
 			if (A.Value == B.Value)
 			{
 				INSIGHTS_DEFAULT_SORTING_NODES(A.NodePtr, B.NodePtr)
@@ -389,6 +340,8 @@ void FSorterByFloatValue::Sort(TArray<FBaseTreeNodePtr>& NodesToSort, ESortMode
 	{
 		ElementsToSort.Sort([](const FSortElement& A, const FSortElement& B) -> bool
 		{
+			INSIGHTS_DEFAULT_PRESORTING_NODES(A.NodePtr, B.NodePtr)
+
 			if (A.Value == B.Value)
 			{
 				INSIGHTS_DEFAULT_SORTING_NODES(A.NodePtr, B.NodePtr)
@@ -415,51 +368,13 @@ FSorterByDoubleValue::FSorterByDoubleValue(TSharedRef<FTableColumn> InColumnRef)
 	: FBaseTableColumnSorter(InColumnRef)
 {
 	ensure(InColumnRef->GetDataType() == ETableCellDataType::Double);
-
-	AscendingCompareDelegate = [&Column = *InColumnRef](const FBaseTreeNodePtr& A, const FBaseTreeNodePtr& B) -> bool
-	{
-		const TOptional<FTableCellValue> OptionalValueA = Column.GetValue(*A);
-		const TOptional<FTableCellValue> OptionalValueB = Column.GetValue(*B);
-
-		const double ValueA = OptionalValueA.IsSet() ? OptionalValueA.GetValue().Double : 0.0;
-		const double ValueB = OptionalValueB.IsSet() ? OptionalValueB.GetValue().Double : 0.0;
-
-		if (ValueA == ValueB)
-		{
-			INSIGHTS_DEFAULT_SORTING_NODES(A, B)
-		}
-		else
-		{
-			// Sort by value (ascending).
-			return ValueA < ValueB;
-		}
-	};
-
-	DescendingCompareDelegate = [&Column = *InColumnRef](const FBaseTreeNodePtr& A, const FBaseTreeNodePtr& B) -> bool
-	{
-		const TOptional<FTableCellValue> OptionalValueA = Column.GetValue(*A);
-		const TOptional<FTableCellValue> OptionalValueB = Column.GetValue(*B);
-
-		const double ValueA = OptionalValueA.IsSet() ? OptionalValueA.GetValue().Double : 0.0;
-		const double ValueB = OptionalValueB.IsSet() ? OptionalValueB.GetValue().Double : 0.0;
-
-		if (ValueA == ValueB)
-		{
-			INSIGHTS_DEFAULT_SORTING_NODES(A, B)
-		}
-		else
-		{
-			// Sort by value (descending).
-			return ValueB < ValueA;
-		}
-	};
 }
 
 ////////////////////////////////////////////////////////////////////////////////////////////////////
 
 void FSorterByDoubleValue::Sort(TArray<FBaseTreeNodePtr>& NodesToSort, ESortMode SortMode) const
 {
-	// Note: This implementation is faster than using above delegates.
+	// Note: This implementation is faster than using delegates.
 	//       It caches the values before sorting, in this way it minimizes the calls to Column.GetValue().
 
 	const FTableColumn& Column = *ColumnRef;
@@ -487,6 +402,8 @@ void FSorterByDoubleValue::Sort(TArray<FBaseTreeNodePtr>& NodesToSort, ESortMode
 	{
 		ElementsToSort.Sort([](const FSortElement& A, const FSortElement& B) -> bool
 		{
+			INSIGHTS_DEFAULT_PRESORTING_NODES(A.NodePtr, B.NodePtr)
+
 			if (A.Value == B.Value)
 			{
 				INSIGHTS_DEFAULT_SORTING_NODES(A.NodePtr, B.NodePtr)
@@ -502,6 +419,8 @@ void FSorterByDoubleValue::Sort(TArray<FBaseTreeNodePtr>& NodesToSort, ESortMode
 	{
 		ElementsToSort.Sort([](const FSortElement& A, const FSortElement& B) -> bool
 		{
+			INSIGHTS_DEFAULT_PRESORTING_NODES(A.NodePtr, B.NodePtr)
+
 			if (A.Value == B.Value)
 			{
 				INSIGHTS_DEFAULT_SORTING_NODES(A.NodePtr, B.NodePtr)
@@ -531,6 +450,8 @@ FSorterByCStringValue::FSorterByCStringValue(TSharedRef<FTableColumn> InColumnRe
 
 	AscendingCompareDelegate = [&Column = *InColumnRef](const FBaseTreeNodePtr& A, const FBaseTreeNodePtr& B) -> bool
 	{
+		INSIGHTS_DEFAULT_PRESORTING_NODES(A, B)
+
 		const TOptional<FTableCellValue> OptionalValueA = Column.GetValue(*A);
 		const TOptional<FTableCellValue> OptionalValueB = Column.GetValue(*B);
 
@@ -551,6 +472,8 @@ FSorterByCStringValue::FSorterByCStringValue(TSharedRef<FTableColumn> InColumnRe
 
 	DescendingCompareDelegate = [&Column = *InColumnRef](const FBaseTreeNodePtr& A, const FBaseTreeNodePtr& B) -> bool
 	{
+		INSIGHTS_DEFAULT_PRESORTING_NODES(A, B)
+
 		const TOptional<FTableCellValue> OptionalValueA = Column.GetValue(*A);
 		const TOptional<FTableCellValue> OptionalValueB = Column.GetValue(*B);
 
@@ -578,73 +501,120 @@ FSorterByTextValue::FSorterByTextValue(TSharedRef<FTableColumn> InColumnRef)
 	: FBaseTableColumnSorter(InColumnRef)
 {
 	ensure(InColumnRef->GetDataType() == ETableCellDataType::Text);
+}
 
-	AscendingCompareDelegate = [&Column = *InColumnRef](const FBaseTreeNodePtr& A, const FBaseTreeNodePtr& B) -> bool
-	{
-		const TOptional<FTableCellValue> OptionalValueA = Column.GetValue(*A);
-		const TOptional<FTableCellValue> OptionalValueB = Column.GetValue(*B);
+////////////////////////////////////////////////////////////////////////////////////////////////////
 
-		if (!OptionalValueB.IsSet())
-		{
-			return true;
-		}
+void FSorterByTextValue::Sort(TArray<FBaseTreeNodePtr>& NodesToSort, ESortMode SortMode) const
+{
+	// Note: This implementation is faster than using delegates.
+	//       It caches the values before sorting, in this way it minimizes the calls to Column.GetValue().
 
-		if (!OptionalValueA.IsSet())
-		{
-			return false;
-		}
+	const FTableColumn& Column = *ColumnRef;
 
-		const FText& ValueA = OptionalValueA.GetValue().GetText();
-		const FText& ValueB = OptionalValueB.GetValue().GetText();
+	const int32 NumElements = NodesToSort.Num();
 
-		const int32 Compare = ValueA.CompareTo(ValueB);
-		if (Compare == 0)
-		{
-			INSIGHTS_DEFAULT_SORTING_NODES(A, B)
-		}
-		else
-		{
-			// Sort by value (ascending).
-			return Compare < 0;
-		}
+	struct FSortElement
+	{
+		const FBaseTreeNodePtr NodePtr;
+		FSetElementId ValueId;
 	};
+	TArray<FSortElement> ElementsToSort;
 
-	DescendingCompareDelegate = [&Column = *InColumnRef](const FBaseTreeNodePtr& A, const FBaseTreeNodePtr& B) -> bool
-	{
-		const TOptional<FTableCellValue> OptionalValueA = Column.GetValue(*A);
-		const TOptional<FTableCellValue> OptionalValueB = Column.GetValue(*B);
+	TSet<FString> UniqueValues;
 
-		if (!OptionalValueB.IsSet())
+	// Cache value for each node.
+	ElementsToSort.Reset(NumElements);
+	for (const FBaseTreeNodePtr& NodePtr : NodesToSort)
+	{
+		FSetElementId ValueId;
+		TOptional<FTableCellValue> Value = Column.GetValue(*NodePtr);
+		if (Value.IsSet())
 		{
-			return true;
+			ValueId = UniqueValues.Add(Value.GetValue().GetText().ToString());
 		}
+		ElementsToSort.Add({ NodePtr, ValueId });
+	}
+	ensure(ElementsToSort.Num() == NumElements);
 
-		if (!OptionalValueA.IsSet())
+	if (SortMode == ESortMode::Ascending)
+	{
+		ElementsToSort.Sort([&UniqueValues](const FSortElement& A, const FSortElement& B) -> bool
 		{
-			return false;
-		}
+			INSIGHTS_DEFAULT_PRESORTING_NODES(A.NodePtr, B.NodePtr)
 
-		const FText& ValueA = OptionalValueA.GetValue().GetText();
-		const FText& ValueB = OptionalValueB.GetValue().GetText();
+			if (A.ValueId == B.ValueId)
+			{
+				INSIGHTS_DEFAULT_SORTING_NODES(A.NodePtr, B.NodePtr)
+			}
 
-		const int32 Compare = ValueA.CompareTo(ValueB);
-		if (Compare == 0)
-		{
-			INSIGHTS_DEFAULT_SORTING_NODES(A, B)
-		}
-		else
+			if (!A.ValueId.IsValidId())
+			{
+				return true;
+			}
+
+			if (!B.ValueId.IsValidId())
+			{
+				return false;
+			}
+
+			const FString& ValueA = UniqueValues[A.ValueId];
+			const FString& ValueB = UniqueValues[B.ValueId];
+
+			// Sort by value (ascending).
+			return ValueA.Compare(ValueB) <= 0;
+		});
+	}
+	else // if (SortMode == ESortMode::Descending)
+	{
+		ElementsToSort.Sort([&UniqueValues](const FSortElement& A, const FSortElement& B) -> bool
 		{
+			INSIGHTS_DEFAULT_PRESORTING_NODES(A.NodePtr, B.NodePtr)
+
+			if (A.ValueId == B.ValueId)
+			{
+				INSIGHTS_DEFAULT_SORTING_NODES(A.NodePtr, B.NodePtr)
+			}
+
+			if (!A.ValueId.IsValidId())
+			{
+				return true;
+			}
+
+			if (!B.ValueId.IsValidId())
+			{
+				return false;
+			}
+
+			const FString& ValueA = UniqueValues[A.ValueId];
+			const FString& ValueB = UniqueValues[B.ValueId];
+
 			// Sort by value (descending).
-			return Compare > 0;
-		}
-	};
+			return ValueA.Compare(ValueB) >= 0;
+		});
+	}
+
+	for (int32 Index = 0; Index < NumElements; ++Index)
+	{
+		NodesToSort[Index] = ElementsToSort[Index].NodePtr;
+	}
 }
 
+////////////////////////////////////////////////////////////////////////////////////////////////////
+// FSorterByTextValueWithId
 ////////////////////////////////////////////////////////////////////////////////////////////////////
 
-void FSorterByTextValue::Sort(TArray<FBaseTreeNodePtr>& NodesToSort, ESortMode SortMode) const
+FSorterByTextValueWithId::FSorterByTextValueWithId(TSharedRef<FTableColumn> InColumnRef)
+	: FBaseTableColumnSorter(InColumnRef)
+{
+	ensure(InColumnRef->GetDataType() == ETableCellDataType::Text);
+}
+
+////////////////////////////////////////////////////////////////////////////////////////////////////
+
+void FSorterByTextValueWithId::Sort(TArray<FBaseTreeNodePtr>& NodesToSort, ESortMode SortMode) const
 {
-	// Note: This implementation is faster than using above delegates.
+	// Note: This implementation is faster than using delegates.
 	//       It caches the values before sorting, in this way it minimizes the calls to Column.GetValue().
 
 	const FTableColumn& Column = *ColumnRef;
@@ -654,85 +624,93 @@ void FSorterByTextValue::Sort(TArray<FBaseTreeNodePtr>& NodesToSort, ESortMode S
 	struct FSortElement
 	{
 		const FBaseTreeNodePtr NodePtr;
-		TOptional<FTableCellValue> Value;
+		int32 ValueIndex;
 	};
 	TArray<FSortElement> ElementsToSort;
 
+	TArray<FString> UniqueValues;
+	TMap<uint64, int32> UniqueValuesMap; // Id --> Index in UniqueValues
+
 	// Cache value for each node.
 	ElementsToSort.Reset(NumElements);
 	for (const FBaseTreeNodePtr& NodePtr : NodesToSort)
 	{
-		ElementsToSort.Add({ NodePtr, Column.GetValue(*NodePtr) });
+		int32 ValueIndex = -1;
+		uint64 ValueId = Column.GetValueId(*NodePtr);
+		const int32* IndexPtr = UniqueValuesMap.Find(ValueId);
+		if (IndexPtr)
+		{
+			ValueIndex = *IndexPtr;
+		}
+		else
+		{
+			TOptional<FTableCellValue> Value = Column.GetValue(*NodePtr);
+			if (Value.IsSet())
+			{
+				ValueIndex = UniqueValues.Num();
+				UniqueValues.Add(Value.GetValue().GetText().ToString());
+			}
+			UniqueValuesMap.Add(ValueId, ValueIndex);
+		}
+		ElementsToSort.Add({ NodePtr, ValueIndex });
 	}
 	ensure(ElementsToSort.Num() == NumElements);
 
 	if (SortMode == ESortMode::Ascending)
 	{
-		ElementsToSort.Sort([](const FSortElement& A, const FSortElement& B) -> bool
+		ElementsToSort.Sort([&UniqueValues](const FSortElement& A, const FSortElement& B) -> bool
+		{
+			INSIGHTS_DEFAULT_PRESORTING_NODES(A.NodePtr, B.NodePtr)
+
+			if (A.ValueIndex == B.ValueIndex)
+			{
+				INSIGHTS_DEFAULT_SORTING_NODES(A.NodePtr, B.NodePtr)
+			}
+
+			if (A.ValueIndex < 0)
 			{
-				if (!A.Value.IsSet())
-				{
-					if (!B.Value.IsSet())
-					{
-						INSIGHTS_DEFAULT_SORTING_NODES(A.NodePtr, B.NodePtr)
-					}
-
-					return true;
-				}
-
-				if (!B.Value.IsSet())
-				{
-					return false;
-				}
-
-				const FText& ValueA = A.Value.GetValue().GetText();
-				const FText& ValueB = B.Value.GetValue().GetText();
-
-				const int32 Compare = ValueA.CompareTo(ValueB);
-				if (Compare == 0)
-				{
-					INSIGHTS_DEFAULT_SORTING_NODES(A.NodePtr, B.NodePtr)
-				}
-				else
-				{
-					// Sort by value (ascending).
-					return Compare < 0;
-				}
-			});
+				return true;
+			}
+
+			if (B.ValueIndex < 0)
+			{
+				return false;
+			}
+
+			const FString& ValueA = UniqueValues[A.ValueIndex];
+			const FString& ValueB = UniqueValues[B.ValueIndex];
+
+			// Sort by value (ascending).
+			return ValueA.Compare(ValueB) <= 0;
+		});
 	}
 	else // if (SortMode == ESortMode::Descending)
 	{
-		ElementsToSort.Sort([](const FSortElement& A, const FSortElement& B) -> bool
+		ElementsToSort.Sort([&UniqueValues](const FSortElement& A, const FSortElement& B) -> bool
+		{
+			INSIGHTS_DEFAULT_PRESORTING_NODES(A.NodePtr, B.NodePtr)
+
+			if (A.ValueIndex == B.ValueIndex)
+			{
+				INSIGHTS_DEFAULT_SORTING_NODES(A.NodePtr, B.NodePtr)
+			}
+
+			if (A.ValueIndex < 0)
 			{
-				if (!A.Value.IsSet())
-				{
-					if (!B.Value.IsSet())
-					{
-						INSIGHTS_DEFAULT_SORTING_NODES(A.NodePtr, B.NodePtr)
-					}
-
-					return true;
-				}
-
-				if (!B.Value.IsSet())
-				{
-					return false;
-				}
-
-				const FText& ValueA = A.Value.GetValue().GetText();
-				const FText& ValueB = B.Value.GetValue().GetText();
-
-				const int32 Compare = ValueA.CompareTo(ValueB);
-				if (Compare == 0)
-				{
-					INSIGHTS_DEFAULT_SORTING_NODES(A.NodePtr, B.NodePtr)
-				}
-				else
-				{
-					// Sort by value (descending).
-					return Compare > 0;
-				}
-			});
+				return true;
+			}
+
+			if (B.ValueIndex < 0)
+			{
+				return false;
+			}
+
+			const FString& ValueA = UniqueValues[A.ValueIndex];
+			const FString& ValueB = UniqueValues[B.ValueIndex];
+
+			// Sort by value (descending).
+			return ValueA.Compare(ValueB) >= 0;
+		});
 	}
 
 	for (int32 Index = 0; Index < NumElements; ++Index)
@@ -746,4 +724,5 @@ void FSorterByTextValue::Sort(TArray<FBaseTreeNodePtr>& NodesToSort, ESortMode S
 } // namespace Insights
 
 #undef INSIGHTS_DEFAULT_SORTING_NODES
+#undef INSIGHTS_DEFAULT_PRESORTING_NODES
 #undef LOCTEXT_NAMESPACE
diff --git a/Engine/Source/Developer/TraceInsights/Private/Insights/Table/ViewModels/TableCellValueSorter.h b/Engine/Source/Developer/TraceInsights/Private/Insights/Table/ViewModels/TableCellValueSorter.h
--- a/Engine/Source/Developer/TraceInsights/Private/Insights/Table/ViewModels/TableCellValueSorter.h
+++ b/Engine/Source/Developer/TraceInsights/Private/Insights/Table/ViewModels/TableCellValueSorter.h
@@ -159,4 +159,14 @@ public:
 
 ////////////////////////////////////////////////////////////////////////////////////////////////////
 
+class FSorterByTextValueWithId : public FBaseTableColumnSorter
+{
+public:
+	FSorterByTextValueWithId(TSharedRef<FTableColumn> InColumnRef);
+
+	virtual void Sort(TArray<FBaseTreeNodePtr>& NodesToSort, ESortMode SortMode) const override;
+};
+
+////////////////////////////////////////////////////////////////////////////////////////////////////
+
 } // namespace Insights
diff --git a/Engine/Source/Developer/TraceInsights/Private/Insights/Table/ViewModels/TableColumn.cpp b/Engine/Source/Developer/TraceInsights/Private/Insights/Table/ViewModels/TableColumn.cpp
--- a/Engine/Source/Developer/TraceInsights/Private/Insights/Table/ViewModels/TableColumn.cpp
+++ b/Engine/Source/Developer/TraceInsights/Private/Insights/Table/ViewModels/TableColumn.cpp
@@ -27,6 +27,13 @@ const TOptional<FTableCellValue> FTableColumn::GetValue(const FBaseTreeNode& InN
 
 ////////////////////////////////////////////////////////////////////////////////////////////////////
 
+uint64 FTableColumn::GetValueId(const FBaseTreeNode& InNode) const
+{
+	return ValueGetter->GetValueId(*this, InNode);
+}
+
+////////////////////////////////////////////////////////////////////////////////////////////////////
+
 TSharedRef<ITableCellValueFormatter> FTableColumn::GetDefaultValueFormatter()
 {
 	return MakeShared<FTableCellValueFormatter>();
diff --git a/Engine/Source/Developer/TraceInsights/Private/Insights/Table/ViewModels/TableColumn.h b/Engine/Source/Developer/TraceInsights/Private/Insights/Table/ViewModels/TableColumn.h
--- a/Engine/Source/Developer/TraceInsights/Private/Insights/Table/ViewModels/TableColumn.h
+++ b/Engine/Source/Developer/TraceInsights/Private/Insights/Table/ViewModels/TableColumn.h
@@ -164,6 +164,7 @@ public:
 	void SetValueGetter(TSharedRef<ITableCellValueGetter> InValueGetter) { ValueGetter = InValueGetter; }
 
 	const TOptional<FTableCellValue> GetValue(const FBaseTreeNode& InNode) const;
+	uint64 GetValueId(const FBaseTreeNode& InNode) const;
 
 	//////////////////////////////////////////////////
 	// Value Formatter
