[Insights] When all nodes are filtered out or the tree contains no nodes, show a warning banner.

diff --git a/Engine/Source/Developer/TraceInsights/Private/Insights/InsightsStyle.cpp b/Engine/Source/Developer/TraceInsights/Private/Insights/InsightsStyle.cpp
--- a/Engine/Source/Developer/TraceInsights/Private/Insights/InsightsStyle.cpp
+++ b/Engine/Source/Developer/TraceInsights/Private/Insights/InsightsStyle.cpp
@@ -137,6 +137,7 @@ void FInsightsStyle::FStyle::Initialize()
 	Set("Border.R", new EDITOR_BOX_BRUSH("Icons/Profiler/Profiler_Border_R_16x", FMargin(4.0f / 16.0f)));
 
 	Set("Graph.Point", new EDITOR_IMAGE_BRUSH("Old/Graph/ExecutionBubble", Icon16x16));
+	Set("TreeViewBanner.WarningIcon", new CORE_IMAGE_BRUSH_SVG("Starship/Common/alert-circle", Icon20x20, FStyleColors::Warning));
 
 	//////////////////////////////////////////////////
 	// Icons for major components
diff --git a/Engine/Source/Developer/TraceInsights/Private/Insights/Table/Widgets/STableTreeView.cpp b/Engine/Source/Developer/TraceInsights/Private/Insights/Table/Widgets/STableTreeView.cpp
--- a/Engine/Source/Developer/TraceInsights/Private/Insights/Table/Widgets/STableTreeView.cpp
+++ b/Engine/Source/Developer/TraceInsights/Private/Insights/Table/Widgets/STableTreeView.cpp
@@ -227,6 +227,41 @@ void STableTreeView::ConstructWidget(TSharedPtr<FTable> InTablePtr)
 			[
 				SAssignNew(AsyncOperationStatus, Insights::SAsyncOperationStatus, SharedThis(this))
 			]
+
+			+ SOverlay::Slot()
+			.HAlign(HAlign_Fill)
+			.VAlign(VAlign_Top)
+			.Padding(FMargin(0.0f, 30.0f, 0.0f, 0.0f))
+			[
+				SNew(SBorder)
+				.BorderImage(FAppStyle::Get().GetBrush("PopupText.Background"))
+				[
+					SNew(SHorizontalBox)
+					.Visibility_Lambda([this]() -> EVisibility 
+					{
+						return TreeViewBannerText.IsEmpty() ? EVisibility::Collapsed : EVisibility::Visible;
+					})
+
+					+ SHorizontalBox::Slot()
+					.VAlign(VAlign_Center)
+					.AutoWidth()
+					.Padding(4.0f, 0.0f, 4.0f, 0.0f)
+					[
+						SNew(SImage)
+						.Image(FInsightsStyle::GetBrush("TreeViewBanner.WarningIcon"))
+					]
+
+					+ SHorizontalBox::Slot()
+					.VAlign(VAlign_Center)
+					.AutoWidth()
+					.Padding(4.0f, 0.0f, 0.0f, 0.0f)
+					[
+						SNew(STextBlock)
+						.Text(this, &STableTreeView::GetTreeViewBannerText)
+						.ColorAndOpacity(FLinearColor(1.0f, 1.0f, 1.0f, 1.0f))
+					]
+				]
+			]
 		]
 
 		+ SHorizontalBox::Slot()
@@ -975,6 +1010,16 @@ void STableTreeView::ApplyFiltering()
 		}
 	}
 
+	TreeViewBannerText = FText::GetEmpty();
+	if (Root->GetChildren().Num() == 0)
+	{
+		TreeViewBannerText = LOCTEXT("EmptyTreeMsg", "Tree is empty.");
+	}
+	else if (Root->GetFilteredChildren().Num() == 0)
+	{
+		TreeViewBannerText = LOCTEXT("AllNodesFiltered", "All tree nodes are filtered out.");
+	}
+
 	Stopwatch.Stop();
 	const double FilteringTime = Stopwatch.GetAccumulatedTime();
 	if (FilteringTime > 0.1)
@@ -1365,6 +1410,12 @@ void STableTreeView::ApplyGrouping()
 
 void STableTreeView::CreateGroups(const TArray<TSharedPtr<FTreeNodeGrouping>>& Groupings)
 {
+	if (TableTreeNodes.Num() == 0)
+	{
+		Root->ClearChildren();
+		return;
+	}
+
 	FStopwatch Stopwatch;
 	Stopwatch.Start();
 
diff --git a/Engine/Source/Developer/TraceInsights/Private/Insights/Table/Widgets/STableTreeView.h b/Engine/Source/Developer/TraceInsights/Private/Insights/Table/Widgets/STableTreeView.h
--- a/Engine/Source/Developer/TraceInsights/Private/Insights/Table/Widgets/STableTreeView.h
+++ b/Engine/Source/Developer/TraceInsights/Private/Insights/Table/Widgets/STableTreeView.h
@@ -415,6 +415,8 @@ protected:
 	typedef TFunctionRef<void(TArray<Insights::FBaseTreeNodePtr>& InNodes)> WriteToFileCallback;
 	void ExportToFileRec(const FBaseTreeNodePtr& InGroupNode, TArray<Insights::FBaseTreeNodePtr>& InNodes, bool bInExportCollapsed, bool InExportLeafs, WriteToFileCallback Callback);
 
+	FText GetTreeViewBannerText() const { return TreeViewBannerText; }
+
 protected:
 	/** Table view model. */
 	TSharedPtr<Insights::FTable> Table;
@@ -534,6 +536,8 @@ protected:
 	double StatsStartTime;
 	double StatsEndTime;
 
+	FText TreeViewBannerText;
+
 	static constexpr int32 MAX_NUMBER_OF_NODES_TO_EXPAND = 1000 * 1000;
 	static constexpr int32 MAX_DEPTH_TO_EXPAND = 100;
 
