[Insights] Fixed missing style for pinned/unpinned brush.

diff --git a/Engine/Source/Developer/TraceInsights/Private/Insights/InsightsStyle.cpp b/Engine/Source/Developer/TraceInsights/Private/Insights/InsightsStyle.cpp
--- a/Engine/Source/Developer/TraceInsights/Private/Insights/InsightsStyle.cpp
+++ b/Engine/Source/Developer/TraceInsights/Private/Insights/InsightsStyle.cpp
@@ -305,6 +305,9 @@ void FInsightsStyle::FStyle::Initialize()
 
 	Set("Icons.Function", new IMAGE_BRUSH_SVG("Function", Icon16x16));
 
+	Set("Icons.Pinned", new EDITOR_IMAGE_BRUSH_SVG("Starship/Common/Pinned", Icon16x16));
+	Set("Icons.Unpinned", new EDITOR_IMAGE_BRUSH_SVG("Starship/Common/Unpinned", Icon16x16));
+
 	//////////////////////////////////////////////////
 
 	Set("TreeTable.RowBackground", new EDITOR_IMAGE_BRUSH("Old/White", Icon16x16, FLinearColor(1.0f, 1.0f, 1.0f, 0.25f)));
