[Insights] Fixed the leaked tab issue when closing a tab that is docked in sidebar. Workaround is to call RemoveTabFromParent() after RequestCloseTab().

diff --git a/Engine/Source/Developer/TraceInsights/Private/Insights/InsightsManager.cpp b/Engine/Source/Developer/TraceInsights/Private/Insights/InsightsManager.cpp
--- a/Engine/Source/Developer/TraceInsights/Private/Insights/InsightsManager.cpp
+++ b/Engine/Source/Developer/TraceInsights/Private/Insights/InsightsManager.cpp
@@ -676,6 +676,7 @@ void FInsightsManager::SpawnAndActivateTabs()
 			while ((NetworkingProfilerTab = FGlobalTabmanager::Get()->FindExistingLiveTab(TabId)).IsValid())
 			{
 				NetworkingProfilerTab->RequestCloseTab();
+				NetworkingProfilerTab->RemoveTabFromParent();
 			}
 		}
 	}
diff --git a/Engine/Source/Developer/TraceInsights/Private/Insights/LoadingProfiler/Widgets/SLoadingProfilerWindow.cpp b/Engine/Source/Developer/TraceInsights/Private/Insights/LoadingProfiler/Widgets/SLoadingProfilerWindow.cpp
--- a/Engine/Source/Developer/TraceInsights/Private/Insights/LoadingProfiler/Widgets/SLoadingProfilerWindow.cpp
+++ b/Engine/Source/Developer/TraceInsights/Private/Insights/LoadingProfiler/Widgets/SLoadingProfilerWindow.cpp
@@ -69,6 +69,7 @@ SLoadingProfilerWindow::~SLoadingProfilerWindow()
 	for (TSharedPtr<SDockTab>& Tab : LocalOpenTabs)
 	{
 		Tab->RequestCloseTab();
+		Tab->RemoveTabFromParent();
 	}
 	LocalOpenTabs.Reset();
 	check(OpenTabs.IsEmpty());
@@ -678,6 +679,7 @@ void SLoadingProfilerWindow::HideTab(const FName& TabID)
 	if (Tab.IsValid())
 	{
 		Tab->RequestCloseTab();
+		Tab->RemoveTabFromParent();
 	}
 }
 
diff --git a/Engine/Source/Developer/TraceInsights/Private/Insights/MemoryProfiler/Widgets/SMemoryProfilerWindow.cpp b/Engine/Source/Developer/TraceInsights/Private/Insights/MemoryProfiler/Widgets/SMemoryProfilerWindow.cpp
--- a/Engine/Source/Developer/TraceInsights/Private/Insights/MemoryProfiler/Widgets/SMemoryProfilerWindow.cpp
+++ b/Engine/Source/Developer/TraceInsights/Private/Insights/MemoryProfiler/Widgets/SMemoryProfilerWindow.cpp
@@ -73,6 +73,7 @@ SMemoryProfilerWindow::~SMemoryProfilerWindow()
 	for (TSharedPtr<SDockTab>& Tab : LocalOpenTabs)
 	{
 		Tab->RequestCloseTab();
+		Tab->RemoveTabFromParent();
 	}
 	LocalOpenTabs.Reset();
 	check(OpenTabs.IsEmpty());
@@ -574,6 +575,7 @@ void SMemoryProfilerWindow::HideTab(const FName& TabID)
 	if (Tab.IsValid())
 	{
 		Tab->RequestCloseTab();
+		Tab->RemoveTabFromParent();
 	}
 }
 
diff --git a/Engine/Source/Developer/TraceInsights/Private/Insights/NetworkingProfiler/Widgets/SNetworkingProfilerWindow.cpp b/Engine/Source/Developer/TraceInsights/Private/Insights/NetworkingProfiler/Widgets/SNetworkingProfilerWindow.cpp
--- a/Engine/Source/Developer/TraceInsights/Private/Insights/NetworkingProfiler/Widgets/SNetworkingProfilerWindow.cpp
+++ b/Engine/Source/Developer/TraceInsights/Private/Insights/NetworkingProfiler/Widgets/SNetworkingProfilerWindow.cpp
@@ -70,6 +70,7 @@ SNetworkingProfilerWindow::~SNetworkingProfilerWindow()
 	for (TSharedPtr<SDockTab>& Tab : LocalOpenTabs)
 	{
 		Tab->RequestCloseTab();
+		Tab->RemoveTabFromParent();
 	}
 	LocalOpenTabs.Reset();
 	check(OpenTabs.IsEmpty());
@@ -447,6 +448,7 @@ void SNetworkingProfilerWindow::HideTab(const FName& TabID)
 	if (Tab.IsValid())
 	{
 		Tab->RequestCloseTab();
+		Tab->RemoveTabFromParent();
 	}
 }
 
diff --git a/Engine/Source/Developer/TraceInsights/Private/Insights/TaskGraphProfiler/TaskGraphProfilerManager.cpp b/Engine/Source/Developer/TraceInsights/Private/Insights/TaskGraphProfiler/TaskGraphProfilerManager.cpp
--- a/Engine/Source/Developer/TraceInsights/Private/Insights/TaskGraphProfiler/TaskGraphProfilerManager.cpp
+++ b/Engine/Source/Developer/TraceInsights/Private/Insights/TaskGraphProfiler/TaskGraphProfilerManager.cpp
@@ -475,6 +475,7 @@ void FTaskGraphProfilerManager::OnWindowClosedEvent()
 		if (Tab.IsValid())
 		{
 			Tab->RequestCloseTab();
+			Tab->RemoveTabFromParent();
 		}
 	}
 }
diff --git a/Engine/Source/Developer/TraceInsights/Private/Insights/Widgets/STimingProfilerWindow.cpp b/Engine/Source/Developer/TraceInsights/Private/Insights/Widgets/STimingProfilerWindow.cpp
--- a/Engine/Source/Developer/TraceInsights/Private/Insights/Widgets/STimingProfilerWindow.cpp
+++ b/Engine/Source/Developer/TraceInsights/Private/Insights/Widgets/STimingProfilerWindow.cpp
@@ -75,6 +75,7 @@ STimingProfilerWindow::~STimingProfilerWindow()
 	for (TSharedPtr<SDockTab>& Tab : LocalOpenTabs)
 	{
 		Tab->RequestCloseTab();
+		Tab->RemoveTabFromParent();
 	}
 	LocalOpenTabs.Reset();
 	check(OpenTabs.IsEmpty());
@@ -639,6 +640,7 @@ void STimingProfilerWindow::HideTab(const FName& TabID)
 	if (Tab.IsValid())
 	{
 		Tab->RequestCloseTab();
+		Tab->RemoveTabFromParent();
 	}
 }
 
diff --git a/Engine/Source/Developer/TraceInsights/Private/Insights/Widgets/STimingView.cpp b/Engine/Source/Developer/TraceInsights/Private/Insights/Widgets/STimingView.cpp
--- a/Engine/Source/Developer/TraceInsights/Private/Insights/Widgets/STimingView.cpp
+++ b/Engine/Source/Developer/TraceInsights/Private/Insights/Widgets/STimingView.cpp
@@ -4832,6 +4832,7 @@ void STimingView::CloseQuickFindTab()
 		if (QuickFindTab.IsValid())
 		{
 			QuickFindTab->RequestCloseTab();
+			QuickFindTab->RemoveTabFromParent();
 		}
 	}
 }
