[Insights] Fixed invisible characters in TimeUtils.h.

diff --git a/Engine/Source/Developer/TraceInsights/Public/Insights/Common/TimeUtils.h b/Engine/Source/Developer/TraceInsights/Public/Insights/Common/TimeUtils.h
--- a/Engine/Source/Developer/TraceInsights/Public/Insights/Common/TimeUtils.h
+++ b/Engine/Source/Developer/TraceInsights/Public/Insights/Common/TimeUtils.h
@@ -8,7 +8,7 @@
 /*
 double is a 64 bit IEEE 754 double precision Floating Point Number
 1 bit for the sign, 11 bits for the exponent, and 52* bits for the value
-15.5 decimal digits of precision (max value*: ‭4`503`599`627`370`495)
+15.5 decimal digits of precision (max value*: 4`503`599`627`370`495)
 |------------------------------------------------------------------------------
 | 900`000.000`000`001     --> up to 10 days   with    1 ns precision
 |  90`000.000`000`000`1   --> up to 25 hours  with  0.1 ns precision (100 ps)
