[Insights] Fix crash when closing Insights during an async operation in TableTreeView.

diff --git a/Engine/Source/Developer/TraceInsights/Private/Insights/Common/AsyncOperationProgress.h b/Engine/Source/Developer/TraceInsights/Private/Insights/Common/AsyncOperationProgress.h
new file mode 100644
--- /dev/null
+++ b/Engine/Source/Developer/TraceInsights/Private/Insights/Common/AsyncOperationProgress.h
@@ -0,0 +1,29 @@
+// Copyright Epic Games, Inc. All Rights Reserved.
+
+#pragma once
+
+#include "CoreMinimal.h"
+
+namespace Insights
+{
+
+class IAsyncOperationProgress
+{
+public:
+	virtual void CancelCurrentAsyncOp() = 0;
+	virtual void Reset() = 0;
+	virtual bool ShouldCancelAsyncOp() const = 0;
+};
+
+class FAsyncOperationProgress : public IAsyncOperationProgress
+{
+public:
+	virtual void CancelCurrentAsyncOp() override { bCancelCurrentAsyncOp.store(true); }
+	virtual void Reset() override{ bCancelCurrentAsyncOp.store(false); }
+	virtual bool ShouldCancelAsyncOp() const override { return bCancelCurrentAsyncOp.load(); }
+
+private:
+	std::atomic<bool> bCancelCurrentAsyncOp{ false };
+};
+
+} //namespace Insights
diff --git a/Engine/Source/Developer/TraceInsights/Private/Insights/InsightsManager.cpp b/Engine/Source/Developer/TraceInsights/Private/Insights/InsightsManager.cpp
--- a/Engine/Source/Developer/TraceInsights/Private/Insights/InsightsManager.cpp
+++ b/Engine/Source/Developer/TraceInsights/Private/Insights/InsightsManager.cpp
@@ -1038,4 +1038,52 @@ bool FInsightsManager::HandleResponseFileCmd(const TCHAR* ResponseFile, FOutputD
 
 ////////////////////////////////////////////////////////////////////////////////////////////////////
 
+void FInsightsManager::AddInProgressAsyncOp(FGraphEventRef Event, const FString& Name)
+{
+	check(IsInGameThread());
+	if (Event.IsValid())
+	{
+		InProgressAsyncTasks.AddTail(FAsyncTaskData(Event, Name));
+	}
+}
+
+////////////////////////////////////////////////////////////////////////////////////////////////////
+
+void FInsightsManager::RemoveInProgressAsyncOp(FGraphEventRef Event)
+{
+	check(IsInGameThread());
+	if (Event.IsValid())
+	{
+		auto Element = InProgressAsyncTasks.GetHead();
+		while (Element && Element->GetValue().GraphEvent != Event)
+		{
+			Element = Element->GetNextNode();
+		}
+		if (Element)
+		{
+			InProgressAsyncTasks.RemoveNode(Element);
+		}
+	}
+}
+
+////////////////////////////////////////////////////////////////////////////////////////////////////
+
+void FInsightsManager::WaitOnInProgressAsyncOps()
+{
+	check(IsInGameThread());
+	while(!InProgressAsyncTasks.IsEmpty())
+	{
+		auto Entry = InProgressAsyncTasks.GetHead();
+		UE_LOG(TraceInsights, Log, TEXT("Waiting on task: %s."), *(Entry->GetValue().Name));
+		FGraphEventRef Event = Entry->GetValue().GraphEvent;
+		while (!Event->IsComplete())
+		{
+			FTaskGraphInterface::Get().ProcessThreadUntilIdle(ENamedThreads::GameThread);
+			FPlatformProcess::Sleep(0.1f);
+		}
+	}
+}
+
+////////////////////////////////////////////////////////////////////////////////////////////////////
+
 #undef LOCTEXT_NAMESPACE
diff --git a/Engine/Source/Developer/TraceInsights/Private/Insights/InsightsManager.h b/Engine/Source/Developer/TraceInsights/Private/Insights/InsightsManager.h
--- a/Engine/Source/Developer/TraceInsights/Private/Insights/InsightsManager.h
+++ b/Engine/Source/Developer/TraceInsights/Private/Insights/InsightsManager.h
@@ -58,6 +58,21 @@ private:
 	uint64 NextTimestamp = (uint64)-1;
 };
 
+/**
+ * Struct that holds data about in progress async operations
+ */
+struct FAsyncTaskData
+{
+	FString Name;
+	FGraphEventRef GraphEvent;
+
+	FAsyncTaskData(FGraphEventRef InGraphEvent, const FString InName)
+	{
+		Name = InName;
+		GraphEvent = InGraphEvent;
+	}
+};
+
 ////////////////////////////////////////////////////////////////////////////////////////////////////
 /**
  * This class manages following areas:
@@ -264,6 +279,10 @@ public:
 
 	void ScheduleCommand(const FString& InCmd);
 
+	void AddInProgressAsyncOp(FGraphEventRef Event, const FString& Name);
+	void RemoveInProgressAsyncOp(FGraphEventRef Event);
+	void WaitOnInProgressAsyncOps();
+
 	////////////////////////////////////////////////////////////////////////////////////////////////////
 	// SessionChangedEvent
 
@@ -410,6 +429,8 @@ private:
 	static const TCHAR* AutoQuitMsg;
 	static const TCHAR* AutoQuitMsgOnFail;
 
+	TDoubleLinkedList<FAsyncTaskData> InProgressAsyncTasks;
+
 	/** A shared pointer to the global instance of the main manager. */
 	static TSharedPtr<FInsightsManager> Instance;
 };
diff --git a/Engine/Source/Developer/TraceInsights/Private/Insights/MemoryProfiler/ViewModels/MemAllocGroupingByCallstack.cpp b/Engine/Source/Developer/TraceInsights/Private/Insights/MemoryProfiler/ViewModels/MemAllocGroupingByCallstack.cpp
--- a/Engine/Source/Developer/TraceInsights/Private/Insights/MemoryProfiler/ViewModels/MemAllocGroupingByCallstack.cpp
+++ b/Engine/Source/Developer/TraceInsights/Private/Insights/MemoryProfiler/ViewModels/MemAllocGroupingByCallstack.cpp
@@ -4,6 +4,7 @@
 #include "TraceServices/Model/Callstack.h"
 
 // Insights
+#include "Insights/Common/AsyncOperationProgress.h"
 #include "Insights/MemoryProfiler/ViewModels/CallstackFormatting.h"
 #include "Insights/MemoryProfiler/ViewModels/MemAllocNode.h"
 
@@ -42,7 +43,7 @@ FMemAllocGroupingByCallstack::~FMemAllocGroupingByCallstack()
 
 ////////////////////////////////////////////////////////////////////////////////////////////////////
 
-void FMemAllocGroupingByCallstack::GroupNodes(const TArray<FTableTreeNodePtr>& Nodes, FTableTreeNode& ParentGroup, TWeakPtr<FTable> InParentTable, std::atomic<bool>& bCancelGrouping) const
+void FMemAllocGroupingByCallstack::GroupNodes(const TArray<FTableTreeNodePtr>& Nodes, FTableTreeNode& ParentGroup, TWeakPtr<FTable> InParentTable, IAsyncOperationProgress& InAsyncOperationProgress) const
 {
 	ParentGroup.ClearChildren();
 
@@ -58,7 +59,7 @@ void FMemAllocGroupingByCallstack::GroupNodes(const TArray<FTableTreeNodePtr>& N
 
 	for (FTableTreeNodePtr NodePtr : Nodes)
 	{
-		if (bCancelGrouping)
+		if (InAsyncOperationProgress.ShouldCancelAsyncOp())
 		{
 			return;
 		}
diff --git a/Engine/Source/Developer/TraceInsights/Private/Insights/MemoryProfiler/ViewModels/MemAllocGroupingByCallstack.h b/Engine/Source/Developer/TraceInsights/Private/Insights/MemoryProfiler/ViewModels/MemAllocGroupingByCallstack.h
--- a/Engine/Source/Developer/TraceInsights/Private/Insights/MemoryProfiler/ViewModels/MemAllocGroupingByCallstack.h
+++ b/Engine/Source/Developer/TraceInsights/Private/Insights/MemoryProfiler/ViewModels/MemAllocGroupingByCallstack.h
@@ -16,6 +16,8 @@ namespace TraceServices
 namespace Insights
 {
 
+class IAsyncOperationProgress;
+
 ////////////////////////////////////////////////////////////////////////////////////////////////////
 
 class FMemAllocGroupingByCallstack : public FTreeNodeGrouping
@@ -37,7 +39,7 @@ public:
 	FMemAllocGroupingByCallstack(bool bInIsInverted, bool bInIsGroupingByFunction);
 	virtual ~FMemAllocGroupingByCallstack();
 
-	virtual void GroupNodes(const TArray<FTableTreeNodePtr>& Nodes, FTableTreeNode& ParentGroup, TWeakPtr<FTable> InParentTable, std::atomic<bool>& bCancelGrouping) const override;
+	virtual void GroupNodes(const TArray<FTableTreeNodePtr>& Nodes, FTableTreeNode& ParentGroup, TWeakPtr<FTable> InParentTable, IAsyncOperationProgress& InAsyncOperationProgress) const override;
 
 	bool IsInverted() const { return bIsInverted; }
 
diff --git a/Engine/Source/Developer/TraceInsights/Private/Insights/MemoryProfiler/ViewModels/MemAllocGroupingByHeap.cpp b/Engine/Source/Developer/TraceInsights/Private/Insights/MemoryProfiler/ViewModels/MemAllocGroupingByHeap.cpp
--- a/Engine/Source/Developer/TraceInsights/Private/Insights/MemoryProfiler/ViewModels/MemAllocGroupingByHeap.cpp
+++ b/Engine/Source/Developer/TraceInsights/Private/Insights/MemoryProfiler/ViewModels/MemAllocGroupingByHeap.cpp
@@ -5,6 +5,7 @@
 #include "Internationalization/Internationalization.h"
 
 // Insights
+#include "Insights/Common/AsyncOperationProgress.h"
 #include "Insights/MemoryProfiler/ViewModels/MemAllocNode.h"
 
 #define LOCTEXT_NAMESPACE "Insights::FMemAllocGroupingByHeap"
@@ -50,7 +51,7 @@ FTableTreeNodePtr MakeGroupNodeHierarchy(const TraceServices::IAllocationsProvid
 ////////////////////////////////////////////////////////////////////////////////////////////////////
 
 void FMemAllocGroupingByHeap::GroupNodes(const TArray<FTableTreeNodePtr>& Nodes, FTableTreeNode& ParentGroup,
-	TWeakPtr<FTable> InParentTable, std::atomic<bool>& bCancelGrouping) const
+	TWeakPtr<FTable> InParentTable, IAsyncOperationProgress& InAsyncOperationProgress) const
 {
 	ParentGroup.ClearChildren();
 
@@ -66,7 +67,7 @@ void FMemAllocGroupingByHeap::GroupNodes(const TArray<FTableTreeNodePtr>& Nodes,
 	// Add allocations nodes to heaps
 	for (FTableTreeNodePtr NodePtr : Nodes)
 	{
-		if (bCancelGrouping)
+		if (InAsyncOperationProgress.ShouldCancelAsyncOp())
 		{
 			return;
 		}
diff --git a/Engine/Source/Developer/TraceInsights/Private/Insights/MemoryProfiler/ViewModels/MemAllocGroupingByHeap.h b/Engine/Source/Developer/TraceInsights/Private/Insights/MemoryProfiler/ViewModels/MemAllocGroupingByHeap.h
--- a/Engine/Source/Developer/TraceInsights/Private/Insights/MemoryProfiler/ViewModels/MemAllocGroupingByHeap.h
+++ b/Engine/Source/Developer/TraceInsights/Private/Insights/MemoryProfiler/ViewModels/MemAllocGroupingByHeap.h
@@ -17,6 +17,8 @@ namespace TraceServices
 namespace Insights
 {
 
+class IAsyncOperationProgress;
+
 ////////////////////////////////////////////////////////////////////////////////////////////////////
 
 class FMemAllocGroupingByHeap : public FTreeNodeGrouping
@@ -27,7 +29,7 @@ public:
 	FMemAllocGroupingByHeap(const TraceServices::IAllocationsProvider& AllocProvider);
 	virtual ~FMemAllocGroupingByHeap() override;
 
-	virtual void GroupNodes(const TArray<FTableTreeNodePtr>& Nodes, FTableTreeNode& ParentGroup, TWeakPtr<FTable> InParentTable, std::atomic<bool>& bCancelGrouping) const override;
+	virtual void GroupNodes(const TArray<FTableTreeNodePtr>& Nodes, FTableTreeNode& ParentGroup, TWeakPtr<FTable> InParentTable, IAsyncOperationProgress& InAsyncOperationProgress) const override;
 private:
 	const TraceServices::IAllocationsProvider& AllocProvider;
 };
diff --git a/Engine/Source/Developer/TraceInsights/Private/Insights/MemoryProfiler/ViewModels/MemAllocGroupingBySize.cpp b/Engine/Source/Developer/TraceInsights/Private/Insights/MemoryProfiler/ViewModels/MemAllocGroupingBySize.cpp
--- a/Engine/Source/Developer/TraceInsights/Private/Insights/MemoryProfiler/ViewModels/MemAllocGroupingBySize.cpp
+++ b/Engine/Source/Developer/TraceInsights/Private/Insights/MemoryProfiler/ViewModels/MemAllocGroupingBySize.cpp
@@ -3,6 +3,7 @@
 #include "MemAllocGroupingBySize.h"
 
 // Insights
+#include "Insights/Common/AsyncOperationProgress.h"
 #include "Insights/MemoryProfiler/ViewModels/MemAllocNode.h"
 
 #define LOCTEXT_NAMESPACE "Insights::FMemAllocGroupingBySize"
@@ -72,7 +73,7 @@ void FMemAllocGroupingBySize::ResetThresholdsPow2()
 
 ////////////////////////////////////////////////////////////////////////////////////////////////////
 
-void FMemAllocGroupingBySize::GroupNodes(const TArray<FTableTreeNodePtr>& Nodes, FTableTreeNode& ParentGroup, TWeakPtr<FTable> InParentTable, std::atomic<bool>& bCancelGrouping) const
+void FMemAllocGroupingBySize::GroupNodes(const TArray<FTableTreeNodePtr>& Nodes, FTableTreeNode& ParentGroup, TWeakPtr<FTable> InParentTable, IAsyncOperationProgress& InAsyncOperationProgress) const
 {
 	TArray<FTableTreeNodePtr> GroupMap; // Node for each ThresholdIndex
 	GroupMap.AddDefaulted(Thresholds.Num());
@@ -84,7 +85,7 @@ void FMemAllocGroupingBySize::GroupNodes(const TArray<FTableTreeNodePtr>& Nodes,
 
 	for (FTableTreeNodePtr NodePtr : Nodes)
 	{
-		if (bCancelGrouping)
+		if (InAsyncOperationProgress.ShouldCancelAsyncOp())
 		{
 			return;
 		}
diff --git a/Engine/Source/Developer/TraceInsights/Private/Insights/MemoryProfiler/ViewModels/MemAllocGroupingBySize.h b/Engine/Source/Developer/TraceInsights/Private/Insights/MemoryProfiler/ViewModels/MemAllocGroupingBySize.h
--- a/Engine/Source/Developer/TraceInsights/Private/Insights/MemoryProfiler/ViewModels/MemAllocGroupingBySize.h
+++ b/Engine/Source/Developer/TraceInsights/Private/Insights/MemoryProfiler/ViewModels/MemAllocGroupingBySize.h
@@ -11,6 +11,8 @@
 namespace Insights
 {
 
+class IAsyncOperationProgress;
+
 ////////////////////////////////////////////////////////////////////////////////////////////////////
 
 class FMemAllocGroupingBySize : public FTreeNodeGrouping
@@ -34,7 +36,7 @@ public:
 	bool IsPow2() const { return bIsPow2; }
 	void ResetThresholdsPow2();
 
-	virtual void GroupNodes(const TArray<FTableTreeNodePtr>& Nodes, FTableTreeNode& ParentGroup, TWeakPtr<FTable> InParentTable, std::atomic<bool>& bCancelGrouping) const override;
+	virtual void GroupNodes(const TArray<FTableTreeNodePtr>& Nodes, FTableTreeNode& ParentGroup, TWeakPtr<FTable> InParentTable, IAsyncOperationProgress& InAsyncOperationProgress) const override;
 
 private:
 	TArray<FThreshold> Thresholds;
diff --git a/Engine/Source/Developer/TraceInsights/Private/Insights/MemoryProfiler/ViewModels/MemAllocGroupingByTag.cpp b/Engine/Source/Developer/TraceInsights/Private/Insights/MemoryProfiler/ViewModels/MemAllocGroupingByTag.cpp
--- a/Engine/Source/Developer/TraceInsights/Private/Insights/MemoryProfiler/ViewModels/MemAllocGroupingByTag.cpp
+++ b/Engine/Source/Developer/TraceInsights/Private/Insights/MemoryProfiler/ViewModels/MemAllocGroupingByTag.cpp
@@ -2,6 +2,8 @@
 #include "MemAllocGroupingByTag.h"
 #include "MemAllocNode.h"
 
+#include "Insights/Common/AsyncOperationProgress.h"
+
 #define LOCTEXT_NAMESPACE "Insights::FMemAllocGroupingByTag"
 
 namespace Insights
@@ -28,7 +30,7 @@ FMemAllocGroupingByTag::FMemAllocGroupingByTag(const TraceServices::IAllocations
 
 void FMemAllocGroupingByTag::GroupNodes(const TArray<FTableTreeNodePtr>& Nodes,
                                         FTableTreeNode& ParentGroup, TWeakPtr<FTable> InParentTable,
-                                        std::atomic<bool>& bCancelGrouping) const
+                                        IAsyncOperationProgress& InAsyncOperationProgress) const
 {
 	using namespace TraceServices;
 	
@@ -58,7 +60,7 @@ void FMemAllocGroupingByTag::GroupNodes(const TArray<FTableTreeNodePtr>& Nodes,
 	// Add nodes to correct tag group
 	for (FTableTreeNodePtr NodePtr : Nodes)
 	{
-		if (bCancelGrouping)
+		if (InAsyncOperationProgress.ShouldCancelAsyncOp())
 		{
 			return;
 		}
@@ -93,7 +95,7 @@ void FMemAllocGroupingByTag::GroupNodes(const TArray<FTableTreeNodePtr>& Nodes,
 	// Fixup parent relationships and filter out empty nodes
 	for (auto TagParentRel : TagParentFixup)
 	{
-		if (bCancelGrouping)
+		if (InAsyncOperationProgress.ShouldCancelAsyncOp())
 		{
 			return;
 		}
diff --git a/Engine/Source/Developer/TraceInsights/Private/Insights/MemoryProfiler/ViewModels/MemAllocGroupingByTag.h b/Engine/Source/Developer/TraceInsights/Private/Insights/MemoryProfiler/ViewModels/MemAllocGroupingByTag.h
--- a/Engine/Source/Developer/TraceInsights/Private/Insights/MemoryProfiler/ViewModels/MemAllocGroupingByTag.h
+++ b/Engine/Source/Developer/TraceInsights/Private/Insights/MemoryProfiler/ViewModels/MemAllocGroupingByTag.h
@@ -9,8 +9,11 @@
 
 namespace Insights
 {
-	
+
+class InAsyncOperationProgress;
+
 ////////////////////////////////////////////////////////////////////////////////////////////////////
+
 class FMemAllocGroupingByTag : public FTreeNodeGrouping
 {
 	INSIGHTS_DECLARE_RTTI(FMemAllocGroupingByTag, FTreeNodeGrouping);
@@ -18,7 +21,7 @@ class FMemAllocGroupingByTag : public FTreeNodeGrouping
 public:
 	FMemAllocGroupingByTag(const TraceServices::IAllocationsProvider& TagProvider);
 
-	virtual void GroupNodes(const TArray<FTableTreeNodePtr>& Nodes, FTableTreeNode& ParentGroup, TWeakPtr<FTable> InParentTable, std::atomic<bool>& bCancelGrouping) const override;
+	virtual void GroupNodes(const TArray<FTableTreeNodePtr>& Nodes, FTableTreeNode& ParentGroup, TWeakPtr<FTable> InParentTable, IAsyncOperationProgress& InAsyncOperationProgress) const override;
 
 private:
 	const TraceServices::IAllocationsProvider& TagProvider;
diff --git a/Engine/Source/Developer/TraceInsights/Private/Insights/MemoryProfiler/ViewModels/MemTagNodeGroupingAndSorting.cpp b/Engine/Source/Developer/TraceInsights/Private/Insights/MemoryProfiler/ViewModels/MemTagNodeGroupingAndSorting.cpp
--- a/Engine/Source/Developer/TraceInsights/Private/Insights/MemoryProfiler/ViewModels/MemTagNodeGroupingAndSorting.cpp
+++ b/Engine/Source/Developer/TraceInsights/Private/Insights/MemoryProfiler/ViewModels/MemTagNodeGroupingAndSorting.cpp
@@ -10,6 +10,19 @@
 #define INSIGHTS_ENSURE ensure
 //#define INSIGHTS_ENSURE(...)
 
+// Default pre-sorting (group nodes sorts above leaf nodes)
+#define INSIGHTS_DEFAULT_PRESORTING_NODES(A, B) \
+	{ \
+		if (ShouldCancelSort()) \
+		{ \
+			return CancelSort(); \
+		} \
+		if (A->IsGroup() != B->IsGroup()) \
+		{ \
+			return A->IsGroup(); \
+		} \
+	}
+
 // Sort by name (ascending).
 #define INSIGHTS_DEFAULT_SORTING_NODES(A, B) return A->GetName().LexicalLess(B->GetName());
 //#define INSIGHTS_DEFAULT_SORTING_NODES(A, B) return A->GetDefaultSortOrder() < B->GetDefaultSortOrder();
@@ -34,8 +47,10 @@ void FMemTagNodeSortingByType::Sort(TArray<Insights::FBaseTreeNodePtr>& NodesToS
 {
 	if (SortMode == Insights::ESortMode::Ascending)
 	{
-		NodesToSort.Sort([](const Insights::FBaseTreeNodePtr& A, const Insights::FBaseTreeNodePtr& B) -> bool
+		NodesToSort.Sort([this](const Insights::FBaseTreeNodePtr& A, const Insights::FBaseTreeNodePtr& B) -> bool
 		{
+			INSIGHTS_DEFAULT_PRESORTING_NODES(A, B);
+
 			INSIGHTS_ENSURE(A.IsValid() && A->GetTypeName() == FMemTagNode::TypeName);
 			const FMemTagNodePtr MemTagNodeA = StaticCastSharedPtr<FMemTagNode, Insights::FBaseTreeNode>(A);
 
@@ -55,8 +70,10 @@ void FMemTagNodeSortingByType::Sort(TArray<Insights::FBaseTreeNodePtr>& NodesToS
 	}
 	else // if (SortMode == ESortMode::Descending)
 	{
-		NodesToSort.Sort([](const Insights::FBaseTreeNodePtr& A, const Insights::FBaseTreeNodePtr& B) -> bool
+		NodesToSort.Sort([this](const Insights::FBaseTreeNodePtr& A, const Insights::FBaseTreeNodePtr& B) -> bool
 		{
+			INSIGHTS_DEFAULT_PRESORTING_NODES(A, B);
+
 			INSIGHTS_ENSURE(A.IsValid() && A->GetTypeName() == FMemTagNode::TypeName);
 			const FMemTagNodePtr MemTagNodeA = StaticCastSharedPtr<FMemTagNode, Insights::FBaseTreeNode>(A);
 
@@ -96,8 +113,10 @@ void FMemTagNodeSortingByTracker::Sort(TArray<Insights::FBaseTreeNodePtr>& Nodes
 {
 	if (SortMode == Insights::ESortMode::Ascending)
 	{
-		NodesToSort.Sort([](const Insights::FBaseTreeNodePtr& A, const Insights::FBaseTreeNodePtr& B) -> bool
+		NodesToSort.Sort([this](const Insights::FBaseTreeNodePtr& A, const Insights::FBaseTreeNodePtr& B) -> bool
 		{
+			INSIGHTS_DEFAULT_PRESORTING_NODES(A, B);
+
 			INSIGHTS_ENSURE(A.IsValid() && A->GetTypeName() == FMemTagNode::TypeName);
 			const FMemTagNodePtr MemTagNodeA = StaticCastSharedPtr<FMemTagNode, Insights::FBaseTreeNode>(A);
 
@@ -117,8 +136,10 @@ void FMemTagNodeSortingByTracker::Sort(TArray<Insights::FBaseTreeNodePtr>& Nodes
 	}
 	else // if (SortMode == ESortMode::Descending)
 	{
-		NodesToSort.Sort([](const Insights::FBaseTreeNodePtr& A, const Insights::FBaseTreeNodePtr& B) -> bool
+		NodesToSort.Sort([this](const Insights::FBaseTreeNodePtr& A, const Insights::FBaseTreeNodePtr& B) -> bool
 		{
+			INSIGHTS_DEFAULT_PRESORTING_NODES(A, B);
+
 			INSIGHTS_ENSURE(A.IsValid() && A->GetTypeName() == FMemTagNode::TypeName);
 			const FMemTagNodePtr MemTagNodeA = StaticCastSharedPtr<FMemTagNode, Insights::FBaseTreeNode>(A);
 
@@ -158,8 +179,10 @@ void FMemTagNodeSortingByInstanceCount::Sort(TArray<Insights::FBaseTreeNodePtr>&
 {
 	if (SortMode == Insights::ESortMode::Ascending)
 	{
-		NodesToSort.Sort([](const Insights::FBaseTreeNodePtr& A, const Insights::FBaseTreeNodePtr& B) -> bool
+		NodesToSort.Sort([this](const Insights::FBaseTreeNodePtr& A, const Insights::FBaseTreeNodePtr& B) -> bool
 		{
+			INSIGHTS_DEFAULT_PRESORTING_NODES(A, B);
+
 			INSIGHTS_ENSURE(A.IsValid() && A->GetTypeName() == FMemTagNode::TypeName);
 			const FMemTagNodePtr MemTagNodeA = StaticCastSharedPtr<FMemTagNode, Insights::FBaseTreeNode>(A);
 			const uint64 ValueA = MemTagNodeA->GetAggregatedStats().InstanceCount;
@@ -181,8 +204,10 @@ void FMemTagNodeSortingByInstanceCount::Sort(TArray<Insights::FBaseTreeNodePtr>&
 	}
 	else // if (SortMode == ESortMode::Descending)
 	{
-		NodesToSort.Sort([](const Insights::FBaseTreeNodePtr& A, const Insights::FBaseTreeNodePtr& B) -> bool
+		NodesToSort.Sort([this](const Insights::FBaseTreeNodePtr& A, const Insights::FBaseTreeNodePtr& B) -> bool
 		{
+			INSIGHTS_DEFAULT_PRESORTING_NODES(A, B);
+
 			INSIGHTS_ENSURE(A.IsValid() && A->GetTypeName() == FMemTagNode::TypeName);
 			const FMemTagNodePtr MemTagNodeA = StaticCastSharedPtr<FMemTagNode, Insights::FBaseTreeNode>(A);
 			const uint64 ValueA = MemTagNodeA->GetAggregatedStats().InstanceCount;
@@ -207,5 +232,6 @@ void FMemTagNodeSortingByInstanceCount::Sort(TArray<Insights::FBaseTreeNodePtr>&
 ////////////////////////////////////////////////////////////////////////////////////////////////////
 
 #undef INSIGHTS_DEFAULT_SORTING_NODES
+#undef INSIGHTS_DEFAULT_PRESORTING_NODES
 #undef INSIGHTS_ENSURE
 #undef LOCTEXT_NAMESPACE
diff --git a/Engine/Source/Developer/TraceInsights/Private/Insights/Table/ViewModels/TableCellValueSorter.cpp b/Engine/Source/Developer/TraceInsights/Private/Insights/Table/ViewModels/TableCellValueSorter.cpp
--- a/Engine/Source/Developer/TraceInsights/Private/Insights/Table/ViewModels/TableCellValueSorter.cpp
+++ b/Engine/Source/Developer/TraceInsights/Private/Insights/Table/ViewModels/TableCellValueSorter.cpp
@@ -9,6 +9,10 @@
 // Default pre-sorting (group nodes sorts above leaf nodes)
 #define INSIGHTS_DEFAULT_PRESORTING_NODES(A, B) \
 	{ \
+		if (ShouldCancelSort()) \
+		{ \
+			return CancelSort(); \
+		} \
 		if (A->IsGroup() != B->IsGroup()) \
 		{ \
 			return A->IsGroup(); \
@@ -58,6 +62,24 @@ void FTableCellValueSorter::Sort(TArray<FBaseTreeNodePtr>& NodesToSort, ESortMod
 	}
 }
 
+////////////////////////////////////////////////////////////////////////////////////////////////////
+
+bool FTableCellValueSorter::CancelSort() const
+{
+#if PLATFORM_EXCEPTIONS_DISABLED
+	return true;
+#else
+	throw "Cancelling sort";
+#endif
+}
+
+////////////////////////////////////////////////////////////////////////////////////////////////////
+
+bool FTableCellValueSorter::ShouldCancelSort() const
+{
+	return this->AsyncOperationProgress && this->AsyncOperationProgress->ShouldCancelAsyncOp();
+}
+
 ////////////////////////////////////////////////////////////////////////////////////////////////////
 // FBaseTableColumnSorter
 ////////////////////////////////////////////////////////////////////////////////////////////////////
@@ -84,7 +106,7 @@ FSorterByName::FSorterByName(TSharedRef<FTableColumn> InColumnRef)
 		LOCTEXT("Sorter_ByName_Desc", "Sort alphabetically by name."),
 		InColumnRef)
 {
-	AscendingCompareDelegate = [](const FBaseTreeNodePtr& A, const FBaseTreeNodePtr& B) -> bool
+	AscendingCompareDelegate = [this](const FBaseTreeNodePtr& A, const FBaseTreeNodePtr& B) -> bool
 	{
 		INSIGHTS_DEFAULT_PRESORTING_NODES(A, B)
 
@@ -92,7 +114,7 @@ FSorterByName::FSorterByName(TSharedRef<FTableColumn> InColumnRef)
 		return A->GetName().LexicalLess(B->GetName());
 	};
 
-	DescendingCompareDelegate = [](const FBaseTreeNodePtr& A, const FBaseTreeNodePtr& B) -> bool
+	DescendingCompareDelegate = [this](const FBaseTreeNodePtr& A, const FBaseTreeNodePtr& B) -> bool
 	{
 		INSIGHTS_DEFAULT_PRESORTING_NODES(A, B)
 
@@ -113,7 +135,7 @@ FSorterByTypeName::FSorterByTypeName(TSharedRef<FTableColumn> InColumnRef)
 		LOCTEXT("Sorter_ByTypeName_Desc", "Sort by type name."),
 		InColumnRef)
 {
-	AscendingCompareDelegate = [](const FBaseTreeNodePtr& A, const FBaseTreeNodePtr& B) -> bool
+	AscendingCompareDelegate = [this](const FBaseTreeNodePtr& A, const FBaseTreeNodePtr& B) -> bool
 	{
 		INSIGHTS_DEFAULT_PRESORTING_NODES(A, B)
 
@@ -131,7 +153,7 @@ FSorterByTypeName::FSorterByTypeName(TSharedRef<FTableColumn> InColumnRef)
 		}
 	};
 
-	DescendingCompareDelegate = [](const FBaseTreeNodePtr& A, const FBaseTreeNodePtr& B) -> bool
+	DescendingCompareDelegate = [this](const FBaseTreeNodePtr& A, const FBaseTreeNodePtr& B) -> bool
 	{
 		INSIGHTS_DEFAULT_PRESORTING_NODES(A, B)
 
@@ -159,7 +181,7 @@ FSorterByBoolValue::FSorterByBoolValue(TSharedRef<FTableColumn> InColumnRef)
 {
 	ensure(InColumnRef->GetDataType() == ETableCellDataType::Bool);
 
-	AscendingCompareDelegate = [&Column = *InColumnRef](const FBaseTreeNodePtr& A, const FBaseTreeNodePtr& B) -> bool
+	AscendingCompareDelegate = [&Column = *InColumnRef, this](const FBaseTreeNodePtr& A, const FBaseTreeNodePtr& B) -> bool
 	{
 		INSIGHTS_DEFAULT_PRESORTING_NODES(A, B)
 
@@ -180,7 +202,7 @@ FSorterByBoolValue::FSorterByBoolValue(TSharedRef<FTableColumn> InColumnRef)
 		}
 	};
 
-	DescendingCompareDelegate = [&Column = * InColumnRef](const FBaseTreeNodePtr& A, const FBaseTreeNodePtr& B) -> bool
+	DescendingCompareDelegate = [&Column = * InColumnRef, this](const FBaseTreeNodePtr& A, const FBaseTreeNodePtr& B) -> bool
 	{
 		INSIGHTS_DEFAULT_PRESORTING_NODES(A, B)
 
@@ -234,6 +256,11 @@ void FSorterByInt64Value::Sort(TArray<FBaseTreeNodePtr>& NodesToSort, ESortMode
 	ElementsToSort.Reset(NumElements);
 	for (const FBaseTreeNodePtr& NodePtr : NodesToSort)
 	{
+		if (ShouldCancelSort())
+		{
+			return;
+		}
+
 		const TOptional<FTableCellValue> OptionalValue = Column.GetValue(*NodePtr);
 		const int64 Value = OptionalValue.IsSet() ? OptionalValue.GetValue().Int64 : 0;
 		ElementsToSort.Add({ NodePtr, Value });
@@ -242,7 +269,7 @@ void FSorterByInt64Value::Sort(TArray<FBaseTreeNodePtr>& NodesToSort, ESortMode
 
 	if (SortMode == ESortMode::Ascending)
 	{
-		ElementsToSort.Sort([](const FSortElement& A, const FSortElement& B) -> bool
+		ElementsToSort.Sort([this](const FSortElement& A, const FSortElement& B) -> bool
 		{
 			INSIGHTS_DEFAULT_PRESORTING_NODES(A.NodePtr, B.NodePtr)
 
@@ -259,7 +286,7 @@ void FSorterByInt64Value::Sort(TArray<FBaseTreeNodePtr>& NodesToSort, ESortMode
 	}
 	else // if (SortMode == ESortMode::Descending)
 	{
-		ElementsToSort.Sort([](const FSortElement& A, const FSortElement& B) -> bool
+		ElementsToSort.Sort([this](const FSortElement& A, const FSortElement& B) -> bool
 		{
 			INSIGHTS_DEFAULT_PRESORTING_NODES(A.NodePtr, B.NodePtr)
 
@@ -313,6 +340,11 @@ void FSorterByFloatValue::Sort(TArray<FBaseTreeNodePtr>& NodesToSort, ESortMode
 	ElementsToSort.Reset(NumElements);
 	for (const FBaseTreeNodePtr& NodePtr : NodesToSort)
 	{
+		if (ShouldCancelSort())
+		{
+			return;
+		}
+
 		const TOptional<FTableCellValue> OptionalValue = Column.GetValue(*NodePtr);
 		const float Value = OptionalValue.IsSet() ? OptionalValue.GetValue().Float : 0.0f;
 		ElementsToSort.Add({ NodePtr, Value });
@@ -321,7 +353,7 @@ void FSorterByFloatValue::Sort(TArray<FBaseTreeNodePtr>& NodesToSort, ESortMode
 
 	if (SortMode == ESortMode::Ascending)
 	{
-		ElementsToSort.Sort([](const FSortElement& A, const FSortElement& B) -> bool
+		ElementsToSort.Sort([this](const FSortElement& A, const FSortElement& B) -> bool
 		{
 			INSIGHTS_DEFAULT_PRESORTING_NODES(A.NodePtr, B.NodePtr)
 
@@ -338,7 +370,7 @@ void FSorterByFloatValue::Sort(TArray<FBaseTreeNodePtr>& NodesToSort, ESortMode
 	}
 	else // if (SortMode == ESortMode::Descending)
 	{
-		ElementsToSort.Sort([](const FSortElement& A, const FSortElement& B) -> bool
+		ElementsToSort.Sort([this](const FSortElement& A, const FSortElement& B) -> bool
 		{
 			INSIGHTS_DEFAULT_PRESORTING_NODES(A.NodePtr, B.NodePtr)
 
@@ -392,6 +424,11 @@ void FSorterByDoubleValue::Sort(TArray<FBaseTreeNodePtr>& NodesToSort, ESortMode
 	ElementsToSort.Reset(NumElements);
 	for (const FBaseTreeNodePtr& NodePtr : NodesToSort)
 	{
+		if (ShouldCancelSort())
+		{
+			return;
+		}
+
 		const TOptional<FTableCellValue> OptionalValue = Column.GetValue(*NodePtr);
 		const double Value = OptionalValue.IsSet() ? OptionalValue.GetValue().Double : 0.0;
 		ElementsToSort.Add({ NodePtr, Value });
@@ -400,7 +437,7 @@ void FSorterByDoubleValue::Sort(TArray<FBaseTreeNodePtr>& NodesToSort, ESortMode
 
 	if (SortMode == ESortMode::Ascending)
 	{
-		ElementsToSort.Sort([](const FSortElement& A, const FSortElement& B) -> bool
+		ElementsToSort.Sort([this](const FSortElement& A, const FSortElement& B) -> bool
 		{
 			INSIGHTS_DEFAULT_PRESORTING_NODES(A.NodePtr, B.NodePtr)
 
@@ -417,7 +454,7 @@ void FSorterByDoubleValue::Sort(TArray<FBaseTreeNodePtr>& NodesToSort, ESortMode
 	}
 	else // if (SortMode == ESortMode::Descending)
 	{
-		ElementsToSort.Sort([](const FSortElement& A, const FSortElement& B) -> bool
+		ElementsToSort.Sort([this](const FSortElement& A, const FSortElement& B) -> bool
 		{
 			INSIGHTS_DEFAULT_PRESORTING_NODES(A.NodePtr, B.NodePtr)
 
@@ -448,7 +485,7 @@ FSorterByCStringValue::FSorterByCStringValue(TSharedRef<FTableColumn> InColumnRe
 {
 	ensure(InColumnRef->GetDataType() == ETableCellDataType::CString);
 
-	AscendingCompareDelegate = [&Column = *InColumnRef](const FBaseTreeNodePtr& A, const FBaseTreeNodePtr& B) -> bool
+	AscendingCompareDelegate = [&Column = *InColumnRef, this](const FBaseTreeNodePtr& A, const FBaseTreeNodePtr& B) -> bool
 	{
 		INSIGHTS_DEFAULT_PRESORTING_NODES(A, B)
 
@@ -470,7 +507,7 @@ FSorterByCStringValue::FSorterByCStringValue(TSharedRef<FTableColumn> InColumnRe
 		}
 	};
 
-	DescendingCompareDelegate = [&Column = *InColumnRef](const FBaseTreeNodePtr& A, const FBaseTreeNodePtr& B) -> bool
+	DescendingCompareDelegate = [&Column = *InColumnRef, this](const FBaseTreeNodePtr& A, const FBaseTreeNodePtr& B) -> bool
 	{
 		INSIGHTS_DEFAULT_PRESORTING_NODES(A, B)
 
@@ -527,6 +564,11 @@ void FSorterByTextValue::Sort(TArray<FBaseTreeNodePtr>& NodesToSort, ESortMode S
 	ElementsToSort.Reset(NumElements);
 	for (const FBaseTreeNodePtr& NodePtr : NodesToSort)
 	{
+		if (ShouldCancelSort())
+		{
+			return;
+		}
+
 		FSetElementId ValueId;
 		TOptional<FTableCellValue> Value = Column.GetValue(*NodePtr);
 		if (Value.IsSet())
@@ -539,7 +581,7 @@ void FSorterByTextValue::Sort(TArray<FBaseTreeNodePtr>& NodesToSort, ESortMode S
 
 	if (SortMode == ESortMode::Ascending)
 	{
-		ElementsToSort.Sort([&UniqueValues](const FSortElement& A, const FSortElement& B) -> bool
+		ElementsToSort.Sort([&UniqueValues, this](const FSortElement& A, const FSortElement& B) -> bool
 		{
 			INSIGHTS_DEFAULT_PRESORTING_NODES(A.NodePtr, B.NodePtr)
 
@@ -567,7 +609,7 @@ void FSorterByTextValue::Sort(TArray<FBaseTreeNodePtr>& NodesToSort, ESortMode S
 	}
 	else // if (SortMode == ESortMode::Descending)
 	{
-		ElementsToSort.Sort([&UniqueValues](const FSortElement& A, const FSortElement& B) -> bool
+		ElementsToSort.Sort([&UniqueValues, this](const FSortElement& A, const FSortElement& B) -> bool
 		{
 			INSIGHTS_DEFAULT_PRESORTING_NODES(A.NodePtr, B.NodePtr)
 
@@ -635,6 +677,11 @@ void FSorterByTextValueWithId::Sort(TArray<FBaseTreeNodePtr>& NodesToSort, ESort
 	ElementsToSort.Reset(NumElements);
 	for (const FBaseTreeNodePtr& NodePtr : NodesToSort)
 	{
+		if (ShouldCancelSort())
+		{
+			return;
+		}
+
 		int32 ValueIndex = -1;
 		uint64 ValueId = Column.GetValueId(*NodePtr);
 		const int32* IndexPtr = UniqueValuesMap.Find(ValueId);
@@ -658,7 +705,7 @@ void FSorterByTextValueWithId::Sort(TArray<FBaseTreeNodePtr>& NodesToSort, ESort
 
 	if (SortMode == ESortMode::Ascending)
 	{
-		ElementsToSort.Sort([&UniqueValues](const FSortElement& A, const FSortElement& B) -> bool
+		ElementsToSort.Sort([&UniqueValues, this](const FSortElement& A, const FSortElement& B) -> bool
 		{
 			INSIGHTS_DEFAULT_PRESORTING_NODES(A.NodePtr, B.NodePtr)
 
@@ -686,7 +733,7 @@ void FSorterByTextValueWithId::Sort(TArray<FBaseTreeNodePtr>& NodesToSort, ESort
 	}
 	else // if (SortMode == ESortMode::Descending)
 	{
-		ElementsToSort.Sort([&UniqueValues](const FSortElement& A, const FSortElement& B) -> bool
+		ElementsToSort.Sort([&UniqueValues, this](const FSortElement& A, const FSortElement& B) -> bool
 		{
 			INSIGHTS_DEFAULT_PRESORTING_NODES(A.NodePtr, B.NodePtr)
 
diff --git a/Engine/Source/Developer/TraceInsights/Private/Insights/Table/ViewModels/TableCellValueSorter.h b/Engine/Source/Developer/TraceInsights/Private/Insights/Table/ViewModels/TableCellValueSorter.h
--- a/Engine/Source/Developer/TraceInsights/Private/Insights/Table/ViewModels/TableCellValueSorter.h
+++ b/Engine/Source/Developer/TraceInsights/Private/Insights/Table/ViewModels/TableCellValueSorter.h
@@ -4,6 +4,7 @@
 
 #include "CoreMinimal.h"
 
+#include "Insights/Common/AsyncOperationProgress.h"
 #include "Insights/Table/ViewModels/TableCellValue.h"
 #include "Insights/Table/ViewModels/TableColumn.h"
 
@@ -41,6 +42,8 @@ public:
 	virtual FTreeNodeCompareFunc GetTreeNodeCompareDelegate(ESortMode SortMode) const = 0;
 
 	virtual void Sort(TArray<FBaseTreeNodePtr>& NodesToSort, ESortMode SortMode) const = 0;
+
+	virtual void SetAsyncOperationProgress(IAsyncOperationProgress* AsyncOperationProgress) = 0;
 };
 
 ////////////////////////////////////////////////////////////////////////////////////////////////////
@@ -62,6 +65,13 @@ public:
 
 	virtual void Sort(TArray<FBaseTreeNodePtr>& NodesToSort, ESortMode SortMode) const override;
 
+	virtual void SetAsyncOperationProgress(IAsyncOperationProgress* InAsyncOperationProgress) override { AsyncOperationProgress = InAsyncOperationProgress; }
+
+protected:
+	// Attempts to cancel the sort by throwing an exception. If exceptions are not available the return value is meant to be returned from sort predicates to speed up the sort.
+	bool CancelSort() const;
+	bool ShouldCancelSort() const;
+
 protected:
 	FName Name;
 	FText ShortName;
@@ -75,6 +85,8 @@ protected:
 
 	FTreeNodeCompareFunc AscendingCompareDelegate;
 	FTreeNodeCompareFunc DescendingCompareDelegate;
+
+	IAsyncOperationProgress* AsyncOperationProgress = nullptr;
 };
 
 ////////////////////////////////////////////////////////////////////////////////////////////////////
diff --git a/Engine/Source/Developer/TraceInsights/Private/Insights/Table/ViewModels/TreeNodeGrouping.cpp b/Engine/Source/Developer/TraceInsights/Private/Insights/Table/ViewModels/TreeNodeGrouping.cpp
--- a/Engine/Source/Developer/TraceInsights/Private/Insights/Table/ViewModels/TreeNodeGrouping.cpp
+++ b/Engine/Source/Developer/TraceInsights/Private/Insights/Table/ViewModels/TreeNodeGrouping.cpp
@@ -2,6 +2,7 @@
 
 #include "TreeNodeGrouping.h"
 
+#include "Insights/Common/AsyncOperationProgress.h"
 #include "Insights/Table/ViewModels/TableTreeNode.h"
 
 #define LOCTEXT_NAMESPACE "Insights_TreeNode"
@@ -31,7 +32,7 @@ FTreeNodeGrouping::FTreeNodeGrouping(const FText& InShortName, const FText& InTi
 
 ////////////////////////////////////////////////////////////////////////////////////////////////////
 
-void FTreeNodeGrouping::GroupNodes(const TArray<FTableTreeNodePtr>& Nodes, FTableTreeNode& ParentGroup, TWeakPtr<FTable> InParentTable, std::atomic<bool>& bCancelGrouping) const
+void FTreeNodeGrouping::GroupNodes(const TArray<FTableTreeNodePtr>& Nodes, FTableTreeNode& ParentGroup, TWeakPtr<FTable> InParentTable, IAsyncOperationProgress& InAsyncOperationProgress) const
 {
 	TMap<FName, FTableTreeNodePtr> GroupMap;
 
@@ -39,7 +40,7 @@ void FTreeNodeGrouping::GroupNodes(const TArray<FTableTreeNodePtr>& Nodes, FTabl
 
 	for (FTableTreeNodePtr NodePtr : Nodes)
 	{
-		if (bCancelGrouping)
+		if (InAsyncOperationProgress.ShouldCancelAsyncOp())
 		{
 			return;
 		}
@@ -86,7 +87,7 @@ FTreeNodeGroupingFlat::FTreeNodeGroupingFlat()
 
 ////////////////////////////////////////////////////////////////////////////////////////////////////
 
-void FTreeNodeGroupingFlat::GroupNodes(const TArray<FTableTreeNodePtr>& Nodes, FTableTreeNode& ParentGroup, TWeakPtr<FTable> InParentTable, std::atomic<bool>& bCancelGrouping) const
+void FTreeNodeGroupingFlat::GroupNodes(const TArray<FTableTreeNodePtr>& Nodes, FTableTreeNode& ParentGroup, TWeakPtr<FTable> InParentTable, IAsyncOperationProgress& InAsyncOperationProgress) const
 {
 	ParentGroup.ClearChildren(1);
 
@@ -97,7 +98,7 @@ void FTreeNodeGroupingFlat::GroupNodes(const TArray<FTableTreeNodePtr>& Nodes, F
 	GroupPtr->ClearChildren(Nodes.Num());
 	for (FTableTreeNodePtr NodePtr : Nodes)
 	{
-		if (bCancelGrouping)
+		if (InAsyncOperationProgress.ShouldCancelAsyncOp())
 		{
 			return;
 		}
diff --git a/Engine/Source/Developer/TraceInsights/Private/Insights/Table/ViewModels/TreeNodeGrouping.h b/Engine/Source/Developer/TraceInsights/Private/Insights/Table/ViewModels/TreeNodeGrouping.h
--- a/Engine/Source/Developer/TraceInsights/Private/Insights/Table/ViewModels/TreeNodeGrouping.h
+++ b/Engine/Source/Developer/TraceInsights/Private/Insights/Table/ViewModels/TreeNodeGrouping.h
@@ -7,6 +7,7 @@
 #include "Delegates/DelegateCombinations.h"
 
 // Insights
+#include "Insights/Common/AsyncOperationProgress.h"
 #include "Insights/Common/SimpleRtti.h"
 #include "Insights/Table/ViewModels/BaseTreeNode.h"
 #include "Insights/Table/ViewModels/TableColumn.h"
@@ -17,6 +18,7 @@
 namespace Insights
 {
 
+class IAsyncOperationProgress;
 class FTable;
 
 ////////////////////////////////////////////////////////////////////////////////////////////////////
@@ -64,7 +66,7 @@ public:
 	virtual FName GetColumnId() const override { return NAME_None; }
 
 	virtual FTreeNodeGroupInfo GetGroupForNode(const FBaseTreeNodePtr InNode) const { return { FName(), false }; }
-	virtual void GroupNodes(const TArray<FTableTreeNodePtr>& Nodes, FTableTreeNode& ParentGroup, TWeakPtr<FTable> InParentTable, std::atomic<bool>& bCancelGrouping) const;
+	virtual void GroupNodes(const TArray<FTableTreeNodePtr>& Nodes, FTableTreeNode& ParentGroup, TWeakPtr<FTable> InParentTable, IAsyncOperationProgress& InAsyncOperationProgress) const;
 
 protected:
 	FText ShortName;
@@ -85,7 +87,7 @@ public:
 	FTreeNodeGroupingFlat();
 	virtual ~FTreeNodeGroupingFlat() {}
 
-	virtual void GroupNodes(const TArray<FTableTreeNodePtr>& Nodes, FTableTreeNode& ParentGroup, TWeakPtr<FTable> InParentTable, std::atomic<bool>& bCancelGrouping) const override;
+	virtual void GroupNodes(const TArray<FTableTreeNodePtr>& Nodes, FTableTreeNode& ParentGroup, TWeakPtr<FTable> InParentTable, IAsyncOperationProgress& InAsyncOperationProgress) const override;
 };
 
 ////////////////////////////////////////////////////////////////////////////////////////////////////
@@ -118,7 +120,7 @@ public:
 	TTreeNodeGroupingByUniqueValue(TSharedRef<FTableColumn> InColumnRef) : FTreeNodeGroupingByUniqueValue(InColumnRef) {}
 	virtual ~TTreeNodeGroupingByUniqueValue() {}
 
-	virtual void GroupNodes(const TArray<FTableTreeNodePtr>& Nodes, FTableTreeNode& ParentGroup, TWeakPtr<FTable> InParentTable, std::atomic<bool>& bCancelGrouping) const override;
+	virtual void GroupNodes(const TArray<FTableTreeNodePtr>& Nodes, FTableTreeNode& ParentGroup, TWeakPtr<FTable> InParentTable, IAsyncOperationProgress& InAsyncOperationProgress) const override;
 
 private:
 	static Type GetValue(const FTableCellValue& CellValue);
@@ -129,7 +131,7 @@ private:
 };
 
 template<typename Type>
-void TTreeNodeGroupingByUniqueValue<Type>::GroupNodes(const TArray<FTableTreeNodePtr>& Nodes, FTableTreeNode& ParentGroup, TWeakPtr<FTable> InParentTable, std::atomic<bool>& bCancelGrouping) const
+void TTreeNodeGroupingByUniqueValue<Type>::GroupNodes(const TArray<FTableTreeNodePtr>& Nodes, FTableTreeNode& ParentGroup, TWeakPtr<FTable> InParentTable, IAsyncOperationProgress& InAsyncOperationProgress) const
 {
 	TMap<Type, FTableTreeNodePtr> GroupMap;
 	FTableTreeNodePtr UnsetGroupPtr = nullptr;
@@ -138,7 +140,7 @@ void TTreeNodeGroupingByUniqueValue<Type>::GroupNodes(const TArray<FTableTreeNod
 
 	for (FTableTreeNodePtr NodePtr : Nodes)
 	{
-		if (bCancelGrouping)
+		if (InAsyncOperationProgress.ShouldCancelAsyncOp())
 		{
 			return;
 		}
diff --git a/Engine/Source/Developer/TraceInsights/Private/Insights/Table/Widgets/STableTreeView.cpp b/Engine/Source/Developer/TraceInsights/Private/Insights/Table/Widgets/STableTreeView.cpp
--- a/Engine/Source/Developer/TraceInsights/Private/Insights/Table/Widgets/STableTreeView.cpp
+++ b/Engine/Source/Developer/TraceInsights/Private/Insights/Table/Widgets/STableTreeView.cpp
@@ -23,7 +23,6 @@
 #include "Widgets/Views/STableViewBase.h"
 
 // Insights
-#include "Insights/InsightsManager.h"
 #include "Insights/InsightsStyle.h"
 #include "Insights/Log.h"
 #include "Insights/Table/ViewModels/Table.h"
@@ -1090,7 +1089,7 @@ bool STableTreeView::ApplyHierarchicalFilterForNode(FTableTreeNodePtr NodePtr, b
 		int32 NumVisibleChildren = 0;
 		for (int32 Cx = 0; Cx < NumChildren; ++Cx)
 		{
-			if (bCancelCurrentAsyncOp)
+			if (AsyncOperationProgress.ShouldCancelAsyncOp())
 			{
 				break;
 			}
@@ -1137,7 +1136,7 @@ bool STableTreeView::ApplyAdvancedFiltersForNode(FTableTreeNodePtr NodePtr)
 		int32 NumVisibleChildren = 0;
 		for (int32 Cx = 0; Cx < NumChildren; ++Cx)
 		{
-			if (bCancelCurrentAsyncOp)
+			if (AsyncOperationProgress.ShouldCancelAsyncOp())
 			{
 				break;
 			}
@@ -1187,7 +1186,7 @@ bool STableTreeView::MakeSubtreeVisible(FTableTreeNodePtr NodePtr, bool bFilterI
 		bool bShouldExpand = false;
 		for (int32 Cx = 0; Cx < NumChildren; ++Cx)
 		{
-			if (bCancelCurrentAsyncOp)
+			if (AsyncOperationProgress.ShouldCancelAsyncOp())
 			{
 				break;
 			}
@@ -1435,7 +1434,7 @@ void STableTreeView::CreateGroups(const TArray<TSharedPtr<FTreeNodeGrouping>>& G
 
 void STableTreeView::GroupNodesRec(const TArray<FTableTreeNodePtr>& Nodes, FTableTreeNode& ParentGroup, int32 GroupingDepth, const TArray<TSharedPtr<FTreeNodeGrouping>>& Groupings)
 {
-	if (bIsUpdateRunning && bCancelCurrentAsyncOp)
+	if (bIsUpdateRunning && AsyncOperationProgress.ShouldCancelAsyncOp())
 	{
 		return;
 	}
@@ -1447,7 +1446,7 @@ void STableTreeView::GroupNodesRec(const TArray<FTableTreeNodePtr>& Nodes, FTabl
 
 	FTreeNodeGrouping& Grouping = *Groupings[GroupingDepth];
 
-	Grouping.GroupNodes(Nodes, ParentGroup, Table, bCancelCurrentAsyncOp);
+	Grouping.GroupNodes(Nodes, ParentGroup, Table, AsyncOperationProgress);
 
 	for (FBaseTreeNodePtr GroupPtr : ParentGroup.GetChildren())
 	{
@@ -2190,8 +2189,22 @@ void STableTreeView::SortTreeNodes(ITableCellValueSorter* InSorter, EColumnSortM
 {
 	FStopwatch Stopwatch;
 	Stopwatch.Start();
-
-	SortTreeNodesRec(*Root, *InSorter, InColumnSortMode);
+#if !PLATFORM_EXCEPTIONS_DISABLED
+	try
+#endif
+	{
+		InSorter->SetAsyncOperationProgress(&AsyncOperationProgress);
+		SortTreeNodesRec(*Root, *InSorter, InColumnSortMode);
+	}
+#if !PLATFORM_EXCEPTIONS_DISABLED
+	catch (const char* Exception)
+	{
+		if (FCStringAnsi::Strcmp(Exception, "Cancelling sort"))
+		{
+			throw Exception;
+		}
+	}
+#endif
 
 	Stopwatch.Stop();
 	const double TotalTime = Stopwatch.GetAccumulatedTime();
@@ -2216,7 +2229,7 @@ void STableTreeView::SortTreeNodesRec(FTableTreeNode& GroupNode, const ITableCel
 
 	for (FBaseTreeNodePtr ChildPtr : GroupNode.GetChildren())
 	{
-		if (bCancelCurrentAsyncOp)
+		if (AsyncOperationProgress.ShouldCancelAsyncOp())
 		{
 			break;
 		}
@@ -2612,7 +2625,7 @@ void STableTreeView::OnPostAsyncUpdate()
 
 	bIsUpdateRunning = false;
 
-	if (!bCancelCurrentAsyncOp)
+	if (!AsyncOperationProgress.ShouldCancelAsyncOp())
 	{
 		TreeView->SetTreeItemsSource(&FilteredGroupNodes);
 
@@ -2641,7 +2654,7 @@ void STableTreeView::OnPostAsyncUpdate()
 		TreeView_Refresh();
 	}
 
-	bCancelCurrentAsyncOp = false;
+	AsyncOperationProgress.Reset();
 	AsyncUpdateStopwatch.Stop();
 }
 
@@ -2795,7 +2808,8 @@ void STableTreeView::OnClose()
 
 		FGraphEventArray Prerequisites;
 		Prerequisites.Add(InProgressAsyncOperationEvent);
-		TGraphTask<FTableTreeViewAsyncCompleteTask>::CreateTask(&Prerequisites).ConstructAndDispatchWhenReady(SharedThis(this));
+		AsyncCompleteTaskEvent = TGraphTask<FTableTreeViewAsyncCompleteTask>::CreateTask(&Prerequisites).ConstructAndDispatchWhenReady(SharedThis(this));
+		FInsightsManager::Get()->AddInProgressAsyncOp(AsyncCompleteTaskEvent, TEXT("FTableTreeViewAsyncCompleteTask"));
 	}
 }
 
@@ -2887,7 +2901,7 @@ void STableTreeView::CancelCurrentAsyncOp()
 {
 	if (bIsUpdateRunning)
 	{
-		bCancelCurrentAsyncOp = true;
+		AsyncOperationProgress.CancelCurrentAsyncOp();
 	}
 }
 
diff --git a/Engine/Source/Developer/TraceInsights/Private/Insights/Table/Widgets/STableTreeView.h b/Engine/Source/Developer/TraceInsights/Private/Insights/Table/Widgets/STableTreeView.h
--- a/Engine/Source/Developer/TraceInsights/Private/Insights/Table/Widgets/STableTreeView.h
+++ b/Engine/Source/Developer/TraceInsights/Private/Insights/Table/Widgets/STableTreeView.h
@@ -17,8 +17,10 @@
 #include "Widgets/Views/STreeView.h"
 
 // Insights
+#include "Insights/Common/AsyncOperationProgress.h"
 #include "Insights/Common/InsightsAsyncWorkUtils.h"
 #include "Insights/Common/Stopwatch.h"
+#include "Insights/InsightsManager.h"
 #include "Insights/Table/ViewModels/TableColumn.h"
 #include "Insights/Table/ViewModels/TableTreeNode.h"
 #include "Insights/ViewModels/Filters.h"
@@ -69,6 +71,7 @@ class STableTreeView : public SCompoundWidget, public IAsyncOperationStatusProvi
 	friend class FTableTreeViewFilterAsyncTask;
 	friend class FTableTreeViewSortAsyncTask;
 	friend class FTableTreeViewGroupAsyncTask;
+	friend class FTableTreeViewAsyncCompleteTask;
 
 public:
 	/** Default constructor. */
@@ -515,6 +518,7 @@ protected:
 
 	TArray<FTableTreeNodePtr> DummyGroupNodes;
 	FGraphEventRef InProgressAsyncOperationEvent;;
+	FGraphEventRef AsyncCompleteTaskEvent;
 	EAsyncOperationType InProgressAsyncOperations = static_cast<EAsyncOperationType>(0);
 	TSharedPtr<class SAsyncOperationStatus> AsyncOperationStatus;
 	FStopwatch AsyncUpdateStopwatch;
@@ -524,7 +528,7 @@ protected:
 	EColumnSortMode::Type CurrentAsyncOpColumnSortMode;
 	TSharedPtr<FTableTreeNodeTextFilter> CurrentAsyncOpTextFilter;
 	FFilterConfigurator* CurrentAsyncOpFilterConfigurator = nullptr;
-	std::atomic<bool> bCancelCurrentAsyncOp { false };
+	FAsyncOperationProgress AsyncOperationProgress;
 	FGraphEventRef DispatchEvent;
 	TArray<FTableTreeNodePtr> NodesToExpand;
 
@@ -649,10 +653,12 @@ public:
 	void DoTask(ENamedThreads::Type CurrentThread, const FGraphEventRef& MyCompletionGraphEvent)
 	{
 		// The role of this task is to keep the STableTreeView object alive until the task and it's prerequisits are completed and to destroy it on the game thread.
+		FGraphEventRef Event = TableTreeViewPtr->AsyncCompleteTaskEvent;
 		if (TableTreeViewPtr.IsValid())
 		{
 			TableTreeViewPtr.Reset();
 		}
+		FInsightsManager::Get()->RemoveInProgressAsyncOp(Event);
 	}
 
 private:
diff --git a/Engine/Source/Developer/TraceInsights/Private/Insights/TraceInsightsModule.cpp b/Engine/Source/Developer/TraceInsights/Private/Insights/TraceInsightsModule.cpp
--- a/Engine/Source/Developer/TraceInsights/Private/Insights/TraceInsightsModule.cpp
+++ b/Engine/Source/Developer/TraceInsights/Private/Insights/TraceInsightsModule.cpp
@@ -76,6 +76,8 @@ void FTraceInsightsModule::StartupModule()
 PRAGMA_DISABLE_OPTIMIZATION
 void FTraceInsightsModule::ShutdownModule()
 {
+	FInsightsManager::Get()->WaitOnInProgressAsyncOps();
+
 #if !WITH_EDITOR
 	if (FModuleManager::Get().IsModuleLoaded("SourceCodeAccess"))
 	{
