[Insights]
  - Session Browser: Added "Rename..." (F2 key shortcut) in context menu for a trace session. Allows renaming in-place for the selected utrace file. The ucache file is also renamed automatically.
  - Session Browser: Added "Delete" (Del key shortcut) in context menu for a trace session. Deletes the selected utrace file. The ucache file is also deleted automatically.

diff --git a/Engine/Source/Developer/TraceInsights/Private/Insights/InsightsStyle.cpp b/Engine/Source/Developer/TraceInsights/Private/Insights/InsightsStyle.cpp
--- a/Engine/Source/Developer/TraceInsights/Private/Insights/InsightsStyle.cpp
+++ b/Engine/Source/Developer/TraceInsights/Private/Insights/InsightsStyle.cpp
@@ -308,6 +308,9 @@ void FInsightsStyle::FStyle::Initialize()
 	Set("Icons.Pinned", new EDITOR_IMAGE_BRUSH_SVG("Starship/Common/Pinned", Icon16x16));
 	Set("Icons.Unpinned", new EDITOR_IMAGE_BRUSH_SVG("Starship/Common/Unpinned", Icon16x16));
 
+	Set("Icons.Rename", new CORE_IMAGE_BRUSH_SVG("Starship/Common/Rename", Icon16x16));
+	Set("Icons.Delete", new CORE_IMAGE_BRUSH_SVG("Starship/Common/Delete", Icon16x16));
+
 	//////////////////////////////////////////////////
 
 	Set("TreeTable.RowBackground", new EDITOR_IMAGE_BRUSH("Old/White", Icon16x16, FLinearColor(1.0f, 1.0f, 1.0f, 0.25f)));
diff --git a/Engine/Source/Developer/TraceInsights/Private/Insights/Widgets/SStartPageWindow.cpp b/Engine/Source/Developer/TraceInsights/Private/Insights/Widgets/SStartPageWindow.cpp
--- a/Engine/Source/Developer/TraceInsights/Private/Insights/Widgets/SStartPageWindow.cpp
+++ b/Engine/Source/Developer/TraceInsights/Private/Insights/Widgets/SStartPageWindow.cpp
@@ -5,8 +5,11 @@
 #include "DesktopPlatformModule.h"
 #include "Framework/Application/SlateApplication.h"
 #include "Framework/MultiBox/MultiBoxBuilder.h"
+#include "HAL/FileManager.h"
+#include "HAL/PlatformProcess.h"
 #include "Internationalization/Text.h"
 #include "IPAddress.h"
+#include "Misc/MessageDialog.h"
 #include "SlateOptMacros.h"
 #include "SocketSubsystem.h"
 #include "Styling/AppStyle.h"
@@ -77,13 +80,42 @@ public:
 	{
 		if (ColumnName == TraceStoreColumns::Name)
 		{
-			return SNew(SBox)
+			TSharedPtr<SEditableTextBox> RenameTextBox;
+
+			TSharedRef<SWidget> Widget = SNew(SBox)
 				.Padding(FMargin(4.0f, 0.0f))
 				[
-					SNew(STextBlock)
-					.Text(this, &STraceListRow::GetTraceName)
-					.ToolTip(STraceListRow::GetTraceTooltip())
+					SNew(SOverlay)
+
+					+ SOverlay::Slot()
+					[
+						SNew(STextBlock)
+						.Visibility_Lambda([this]() { return IsRenaming() ? EVisibility::Collapsed : EVisibility::Visible; })
+						.Text(this, &STraceListRow::GetTraceName)
+						.ToolTip(STraceListRow::GetTraceTooltip())
+					]
+
+					+ SOverlay::Slot()
+					[
+						SAssignNew(RenameTextBox, SEditableTextBox)
+						.Padding(FMargin(-1.0f, 0.0f, 0.0f, 0.0f))
+						.Visibility_Lambda([this]() { return IsRenaming() ? EVisibility::Visible : EVisibility::Collapsed; })
+						.Text(this, &STraceListRow::GetTraceName)
+						.ClearKeyboardFocusOnCommit(true)
+						.OnTextCommitted(this, &STraceListRow::RenameTextBox_OnValueCommitted)
+						.ToolTip(STraceListRow::GetTraceTooltip())
+					]
 				];
+
+			TSharedPtr<FTraceViewModel> TracePin = WeakTrace.Pin();
+			if (TracePin.IsValid())
+			{
+				// A weak reference is stored in the view model in order to focus the text box
+				// when starts renaming.
+				TracePin->RenameTextBox = RenameTextBox;
+			}
+
+			return Widget;
 		}
 		else if (ColumnName == TraceStoreColumns::Uri)
 		{
@@ -162,6 +194,56 @@ public:
 		}
 	}
 
+	bool IsRenaming() const
+	{
+		TSharedPtr<FTraceViewModel> TracePin = WeakTrace.Pin();
+		if (TracePin.IsValid())
+		{
+			return TracePin->bIsRenaming;
+		}
+		return false;
+	}
+
+	void RenameTextBox_OnValueCommitted(const FText& InText, ETextCommit::Type InCommitType)
+	{
+		TSharedPtr<FTraceViewModel> TracePin = WeakTrace.Pin();
+		if (TracePin.IsValid())
+		{
+			if (InCommitType != ETextCommit::OnCleared &&
+				!TracePin->Name.EqualTo(InText))
+			{
+				FString TraceName = TracePin->Name.ToString();
+				FString NewTraceName = InText.ToString();
+
+				UE_LOG(TraceInsights, Log, TEXT("[TraceStore] Renaming \"%s\" to \"%s\"..."), *TraceName, *NewTraceName);
+
+				FString TraceFile = TracePin->Uri.ToString();
+				if (FPaths::FileExists(TraceFile))
+				{
+					FString NewTraceFile = FPaths::Combine(FPaths::GetPath(TraceFile), NewTraceName + TEXT(".utrace"));
+					if (IFileManager::Get().Move(*NewTraceFile, *TraceFile, true))
+					{
+						UE_LOG(TraceInsights, Verbose, TEXT("[TraceStore] Renamed utrace file (\"%s\")."), *NewTraceFile);
+						TracePin->Name = InText;
+						TracePin->Uri = FText::FromString(NewTraceFile);
+
+						FString CacheFile = FPaths::ChangeExtension(TraceFile, TEXT("ucache"));
+						if (FPaths::FileExists(CacheFile))
+						{
+							FString NewCacheFile = FPaths::Combine(FPaths::GetPath(CacheFile), InText.ToString() + TEXT(".ucache"));
+							if (IFileManager::Get().Move(*NewCacheFile, *CacheFile, true))
+							{
+								UE_LOG(TraceInsights, Verbose, TEXT("[TraceStore] Renamed ucache file (\"%s\")."), *NewCacheFile);
+							}
+						}
+					}
+				}
+			}
+
+			TracePin->bIsRenaming = false;
+		}
+	}
+
 	FText GetTraceIndexAndId() const
 	{
 		TSharedPtr<FTraceViewModel> TracePin = WeakTrace.Pin();
@@ -909,7 +991,7 @@ TSharedRef<SWidget> STraceStoreWindow::ConstructSessionsPanel()
 		.ListItemsSource(&FilteredTraceViewModels)
 		.OnGenerateRow(this, &STraceStoreWindow::TraceList_OnGenerateRow)
 		.ConsumeMouseWheel(EConsumeMouseWheel::Always)
-		//.OnContextMenuOpening(FOnContextMenuOpening::CreateSP(this, &STraceStoreWindow::TraceList_GetContextMenu))
+		.OnContextMenuOpening(FOnContextMenuOpening::CreateSP(this, &STraceStoreWindow::TraceList_GetMenuContent))
 		.HeaderRow
 		(
 			SNew(SHeaderRow)
@@ -1171,6 +1253,122 @@ END_SLATE_FUNCTION_BUILD_OPTIMIZATION
 
 ////////////////////////////////////////////////////////////////////////////////////////////////////
 
+TSharedPtr<SWidget> STraceStoreWindow::TraceList_GetMenuContent()
+{
+	const bool bShouldCloseWindowAfterMenuSelection = true;
+	FMenuBuilder MenuBuilder(bShouldCloseWindowAfterMenuSelection, NULL);
+
+	MenuBuilder.BeginSection("Misc");
+	{
+		FUIAction Action_Rename
+		(
+			FExecuteAction::CreateSP(this, &STraceStoreWindow::RenameTraceFile),
+			FCanExecuteAction::CreateSP(this, &STraceStoreWindow::CanEditTraceFile)
+		);
+		MenuBuilder.AddMenuEntry
+		(
+			LOCTEXT("ContextMenu_Rename", "Rename... \t\tF2"),
+			LOCTEXT("ContextMenu_Rename_ToolTip", "Rename the selected utrace file."),
+			FSlateIcon(FInsightsStyle::GetStyleSetName(), "Icons.Rename"),
+			Action_Rename,
+			NAME_None,
+			EUserInterfaceActionType::Button
+		);
+
+		FUIAction Action_Delete
+		(
+			FExecuteAction::CreateSP(this, &STraceStoreWindow::DeleteTraceFile),
+			FCanExecuteAction::CreateSP(this, &STraceStoreWindow::CanEditTraceFile)
+		);
+		MenuBuilder.AddMenuEntry
+		(
+			LOCTEXT("ContextMenu_Delete", "Delete \t\t\tDel"),
+			LOCTEXT("ContextMenu_Delete_ToolTip", "Delete the selected utrace file."),
+			FSlateIcon(FInsightsStyle::GetStyleSetName(), "Icons.Delete"),
+			Action_Delete,
+			NAME_None,
+			EUserInterfaceActionType::Button
+		);
+	}
+
+	return MenuBuilder.MakeWidget();
+}
+
+////////////////////////////////////////////////////////////////////////////////////////////////////
+
+bool STraceStoreWindow::CanEditTraceFile() const
+{
+	return SelectedTrace.IsValid() && !SelectedTrace->bIsLive;
+}
+
+////////////////////////////////////////////////////////////////////////////////////////////////////
+
+void STraceStoreWindow::RenameTraceFile()
+{
+	if (SelectedTrace.IsValid() && !SelectedTrace->bIsLive)
+	{
+		SelectedTrace->bIsRenaming = true;
+
+		FSlateApplication::Get().CloseToolTip();
+
+		TSharedPtr<SEditableTextBox> RenameTextBox = SelectedTrace->RenameTextBox.Pin();
+		if (RenameTextBox.IsValid())
+		{
+			FSlateApplication::Get().SetKeyboardFocus(RenameTextBox.ToSharedRef(), EFocusCause::SetDirectly);
+		}
+	}
+}
+
+////////////////////////////////////////////////////////////////////////////////////////////////////
+
+void STraceStoreWindow::DeleteTraceFile()
+{
+	if (!SelectedTrace.IsValid() || SelectedTrace->bIsLive)
+	{
+		return;
+	}
+
+	FSlateApplication::Get().CloseToolTip();
+
+	FString TraceFile = SelectedTrace->Uri.ToString();
+
+	//TODO: Make a custom OkCancel modal dialog. See FSlateApplication::Get().AddModalWindow(..).
+	FText Title = LOCTEXT("ConfirmToDeleteTraceFile_Title", "Unreal Insights");
+	FText ConfirmMessage = FText::Format(LOCTEXT("ConfirmToDeleteTraceFile", "You are about to delete the utrace file:\n{0}\n\nPress Ok to continue."),
+		FText::FromString(TraceFile));
+	EAppReturnType::Type OkToDelete = FMessageDialog::Open(EAppMsgType::OkCancel, ConfirmMessage, &Title);
+	if (OkToDelete == EAppReturnType::Cancel)
+	{
+		return;
+	}
+
+	FString TraceName = SelectedTrace->Name.ToString();
+	UE_LOG(TraceInsights, Log, TEXT("[TraceStore] Deleting \"%s\"..."), *TraceName);
+
+	if (FPaths::FileExists(TraceFile))
+	{
+		if (IFileManager::Get().Delete(*TraceFile))
+		{
+			UE_LOG(TraceInsights, Verbose, TEXT("[TraceStore] Deleted utrace file (\"%s\")."), *TraceFile);
+
+			FString CacheFile = FPaths::ChangeExtension(TraceFile, TEXT("ucache"));
+			if (FPaths::FileExists(CacheFile))
+			{
+				if (IFileManager::Get().Delete(*CacheFile))
+				{
+					UE_LOG(TraceInsights, Verbose, TEXT("[TraceStore] Deleted ucache file (\"%s\")."), *CacheFile);
+				}
+			}
+
+			TraceViewModels.Remove(SelectedTrace);
+			TraceViewModelMap.Remove(SelectedTrace->TraceId);
+			OnTraceListChanged();
+		}
+	}
+}
+
+////////////////////////////////////////////////////////////////////////////////////////////////////
+
 void STraceStoreWindow::ShowSplashScreenOverlay()
 {
 	SplashScreenOverlayFadeTime = 3.5f;
@@ -1312,7 +1510,9 @@ void STraceStoreWindow::RefreshTraceList()
 					FTraceViewModel& Trace = *TracePtr;
 					Trace.TraceId = SourceTrace.TraceId;
 					Trace.Name = FText::FromString(SourceTrace.Name);
-					Trace.Uri = FText::FromString(FInsightsManager::Get()->GetStoreDir() + TEXT("/") + SourceTrace.Name + TEXT(".utrace"));
+					FString Uri = FInsightsManager::Get()->GetStoreDir() + TEXT("/") + SourceTrace.Name + TEXT(".utrace");
+					FPaths::NormalizeFilename(Uri);
+					Trace.Uri = FText::FromString(Uri);
 					UpdateTrace(Trace, SourceTrace);
 					TraceViewModels.Add(TracePtr);
 					TraceViewModelMap.Add(SourceTrace.TraceId, TracePtr);
@@ -1542,7 +1742,21 @@ void STraceStoreWindow::OnMouseLeave(const FPointerEvent& MouseEvent)
 
 FReply STraceStoreWindow::OnKeyDown(const FGeometry& MyGeometry, const FKeyEvent& InKeyEvent)
 {
-	return FReply::Unhandled();
+	//TODO: Make commands for Rename and Delete.
+	//return GetCommandList()->ProcessCommandBindings(InKeyEvent) ? FReply::Handled() : FReply::Unhandled();
+
+	if (InKeyEvent.GetKey() == EKeys::F2)
+	{
+		RenameTraceFile();
+		return FReply::Handled();
+	}
+	else if (InKeyEvent.GetKey() == EKeys::Delete)
+	{
+		DeleteTraceFile();
+		return FReply::Handled();
+	}
+
+	return SCompoundWidget::OnKeyDown(MyGeometry, InKeyEvent);
 }
 
 ////////////////////////////////////////////////////////////////////////////////////////////////////
diff --git a/Engine/Source/Developer/TraceInsights/Private/Insights/Widgets/SStartPageWindow.h b/Engine/Source/Developer/TraceInsights/Private/Insights/Widgets/SStartPageWindow.h
--- a/Engine/Source/Developer/TraceInsights/Private/Insights/Widgets/SStartPageWindow.h
+++ b/Engine/Source/Developer/TraceInsights/Private/Insights/Widgets/SStartPageWindow.h
@@ -76,9 +76,12 @@ struct FTraceViewModel
 	EBuildTargetType TargetType = EBuildTargetType::Unknown;
 
 	bool bIsMetadataUpdated = false;
+	bool bIsRenaming = false;
 	bool bIsLive = false;
 	uint32 IpAddress = 0;
 
+	TWeakPtr<SEditableTextBox> RenameTextBox;
+
 	FTraceViewModel() = default;
 
 	static FDateTime ConvertTimestamp(uint64 InTimestamp)
@@ -257,6 +260,12 @@ private:
 	/** Generate a new row for the Traces list view. */
 	TSharedRef<ITableRow> TraceList_OnGenerateRow(TSharedPtr<FTraceViewModel> InTrace, const TSharedRef<STableViewBase>& OwnerTable);
 
+	TSharedPtr<SWidget> TraceList_GetMenuContent();
+
+	bool CanEditTraceFile() const;
+	void RenameTraceFile();
+	void DeleteTraceFile();
+
 	//////////////////////////////////////////////////
 	// "Starting Analysis" Splash Screen
 
