[Insights]
  - Added "Session Browser" in main menu. Opens the Unreal Insights Session Browser window.
  - Added "Open Trace File (new instance)...". Starts analysis for a specified trace file, in a separate Unreal Insights instance.
  - Added "Open Trace File (same instance)...". Starts analysis for a specified trace file, replacing the current analysis session.

diff --git a/Engine/Source/Developer/TraceInsights/Private/Insights/Common/InsightsMenuBuilder.cpp b/Engine/Source/Developer/TraceInsights/Private/Insights/Common/InsightsMenuBuilder.cpp
--- a/Engine/Source/Developer/TraceInsights/Private/Insights/Common/InsightsMenuBuilder.cpp
+++ b/Engine/Source/Developer/TraceInsights/Private/Insights/Common/InsightsMenuBuilder.cpp
@@ -9,6 +9,7 @@
 #include "WorkspaceMenuStructureModule.h"
 
 #include "Insights/InsightsStyle.h"
+#include "Insights/InsightsManager.h"
 
 #define LOCTEXT_NAMESPACE "InsightsMenuBuilder"
 
@@ -51,6 +52,24 @@ TSharedRef<FWorkspaceItem> FInsightsMenuBuilder::GetWindowsGroup()
 void FInsightsMenuBuilder::PopulateMenu(FMenuBuilder& MenuBuilder)
 {
 #if !WITH_EDITOR
+	MenuBuilder.BeginSection("Insights");
+	MenuBuilder.AddMenuEntry(
+		LOCTEXT("OpenSessionBrowser", "Session Browser"),
+		LOCTEXT("OpenSessionBrowser_ToolTip", "Opens the Unreal Insights Session Browser window."),
+		FSlateIcon(FInsightsStyle::GetStyleSetName(), "AppIcon.Small"),
+		FUIAction(FExecuteAction::CreateLambda([] { FInsightsManager::Get()->OpenUnrealInsights(); })));
+	MenuBuilder.AddMenuEntry(
+		LOCTEXT("OpenTraceFile1", "Open Trace File (new instance)..."),
+		LOCTEXT("OpenTraceFile1_ToolTip", "Starts analysis for a specified trace file, in a separate Unreal Insights instance."),
+		FSlateIcon(),
+		FUIAction(FExecuteAction::CreateLambda([] { FInsightsManager::Get()->OpenTraceFile(); })));
+	MenuBuilder.AddMenuEntry(
+		LOCTEXT("OpenTraceFile2", "Open Trace File (same instance)..."),
+		LOCTEXT("OpenTraceFile2_ToolTip", "Starts analysis for a specified trace file, replacing the current analysis session."),
+		FSlateIcon(),
+		FUIAction(FExecuteAction::CreateLambda([] { FInsightsManager::Get()->LoadTraceFile(); })));
+	MenuBuilder.EndSection();
+
 	FGlobalTabmanager::Get()->PopulateLocalTabSpawnerMenu(MenuBuilder);
 
 	static FName WidgetReflectorTabId("WidgetReflector");
@@ -60,7 +79,7 @@ void FInsightsMenuBuilder::PopulateMenu(FMenuBuilder& MenuBuilder)
 		MenuBuilder.BeginSection("WidgetTools");
 		MenuBuilder.AddMenuEntry(
 			LOCTEXT("OpenWidgetReflector", "Widget Reflector"),
-			LOCTEXT("OpenWidgetReflectorToolTip", "Opens the Widget Reflector, a handy tool for diagnosing problems with live widgets."),
+			LOCTEXT("OpenWidgetReflector_ToolTip", "Opens the Widget Reflector, a handy tool for diagnosing problems with live widgets."),
 			FSlateIcon(FAppStyle::Get().GetStyleSetName(), "WidgetReflector.Icon"),
 			FUIAction(FExecuteAction::CreateLambda([=] { FGlobalTabmanager::Get()->TryInvokeTab(WidgetReflectorTabId); })));
 		MenuBuilder.EndSection();
@@ -77,7 +96,7 @@ void FInsightsMenuBuilder::PopulateMenu(FMenuBuilder& MenuBuilder)
 
 		MenuBuilder.AddMenuEntry(
 			LOCTEXT("OpenStarshipSuite", "Starship Test Suite"),
-			LOCTEXT("OpenStarshipSuiteDesc", "Opens the Starship UX test suite."),
+			LOCTEXT("OpenStarshipSuite_ToolTip", "Opens the Starship UX test suite."),
 			FSlateIcon(FInsightsStyle::GetStyleSetName(), "Icons.Test"),
 			OpenStarshipSuiteAction,
 			NAME_None,
diff --git a/Engine/Source/Developer/TraceInsights/Private/Insights/InsightsManager.cpp b/Engine/Source/Developer/TraceInsights/Private/Insights/InsightsManager.cpp
--- a/Engine/Source/Developer/TraceInsights/Private/Insights/InsightsManager.cpp
+++ b/Engine/Source/Developer/TraceInsights/Private/Insights/InsightsManager.cpp
@@ -2,8 +2,10 @@
 
 #include "InsightsManager.h"
 
+#include "DesktopPlatformModule.h"
 #include "Framework/Application/SlateApplication.h"
 #include "HAL/PlatformMemory.h"
+#include "HAL/PlatformProcess.h"
 #include "Logging/MessageLog.h"
 #include "MessageLog/Public/MessageLogModule.h"
 #include "Misc/CString.h"
@@ -710,6 +712,89 @@ void FInsightsManager::ActivateTimingInsightsTab()
 
 ////////////////////////////////////////////////////////////////////////////////////////////////////
 
+bool FInsightsManager::ShowOpenTraceFileDialog(FString& OutTraceFile) const
+{
+	const FString ProfilingDirectory(FPaths::ConvertRelativePathToFull(FInsightsManager::Get()->GetStoreDir()));
+
+	TArray<FString> OutFiles;
+	bool bOpened = false;
+
+	IDesktopPlatform* DesktopPlatform = FDesktopPlatformModule::Get();
+	if (DesktopPlatform != nullptr)
+	{
+		FSlateApplication::Get().CloseToolTip();
+
+		bOpened = DesktopPlatform->OpenFileDialog
+		(
+			FSlateApplication::Get().FindBestParentWindowHandleForDialogs(nullptr),
+			LOCTEXT("LoadTrace_FileDesc", "Open trace file...").ToString(),
+			ProfilingDirectory,
+			TEXT(""),
+			LOCTEXT("LoadTrace_FileFilter", "Trace files (*.utrace)|*.utrace|All files (*.*)|*.*").ToString(),
+			EFileDialogFlags::None,
+			OutFiles
+		);
+	}
+
+	if (bOpened == true && OutFiles.Num() == 1)
+	{
+		OutTraceFile = OutFiles[0];
+		return true;
+	}
+
+	return false;
+}
+
+////////////////////////////////////////////////////////////////////////////////////////////////////
+
+void FInsightsManager::OpenUnrealInsights(const TCHAR* CmdLine) const
+{
+	if (CmdLine == nullptr)
+	{
+		CmdLine = TEXT("");
+	}
+
+	const TCHAR* ExecutablePath = FPlatformProcess::ExecutablePath();
+
+	constexpr bool bLaunchDetached = true;
+	constexpr bool bLaunchHidden = false;
+	constexpr bool bLaunchReallyHidden = false;
+
+	uint32 ProcessID = 0;
+	const int32 PriorityModifier = 0;
+	const TCHAR* OptionalWorkingDirectory = nullptr;
+
+	void* PipeWriteChild = nullptr;
+	void* PipeReadChild = nullptr;
+
+	FProcHandle Handle = FPlatformProcess::CreateProc(ExecutablePath, CmdLine, bLaunchDetached, bLaunchHidden, bLaunchReallyHidden, &ProcessID, PriorityModifier, OptionalWorkingDirectory, PipeWriteChild, PipeReadChild);
+	if (Handle.IsValid())
+	{
+		FPlatformProcess::CloseProc(Handle);
+	}
+}
+
+////////////////////////////////////////////////////////////////////////////////////////////////////
+
+void FInsightsManager::OpenTraceFile() const
+{
+	FString TraceFile;
+	if (FInsightsManager::Get()->ShowOpenTraceFileDialog(TraceFile))
+	{
+		OpenTraceFile(TraceFile);
+	}
+}
+
+////////////////////////////////////////////////////////////////////////////////////////////////////
+
+void FInsightsManager::OpenTraceFile(const FString& InTraceFile) const
+{
+	FString CmdLine = TEXT("-OpenTraceFile=\"") + InTraceFile + TEXT("\"");
+	OpenUnrealInsights(*CmdLine);
+}
+
+////////////////////////////////////////////////////////////////////////////////////////////////////
+
 void FInsightsManager::LoadLastLiveSession()
 {
 	ResetSession();
@@ -786,6 +871,17 @@ void FInsightsManager::LoadTrace(uint32 InTraceId, bool InAutoQuit)
 
 ////////////////////////////////////////////////////////////////////////////////////////////////////
 
+void FInsightsManager::LoadTraceFile()
+{
+	FString TraceFile;
+	if (FInsightsManager::Get()->ShowOpenTraceFileDialog(TraceFile))
+	{
+		LoadTraceFile(TraceFile);
+	}
+}
+
+////////////////////////////////////////////////////////////////////////////////////////////////////
+
 void FInsightsManager::LoadTraceFile(const FString& InTraceFilename, bool InAutoQuit)
 {
 	IPlatformFile& PlatformFile = IPlatformFile::GetPlatformPhysical();
diff --git a/Engine/Source/Developer/TraceInsights/Private/Insights/InsightsManager.h b/Engine/Source/Developer/TraceInsights/Private/Insights/InsightsManager.h
--- a/Engine/Source/Developer/TraceInsights/Private/Insights/InsightsManager.h
+++ b/Engine/Source/Developer/TraceInsights/Private/Insights/InsightsManager.h
@@ -240,18 +240,53 @@ public:
 
 	////////////////////////////////////////////////////////////////////////////////////////////////////
 
-	/** Creates a new analysis session instance and loads the latest available trace that is live. */
+	/**
+	 * Shows the open file dialog for choosing a trace file.
+	 * @param OutTraceFile - The chosen trace file, if successful
+	 * @return True, if successful.
+	 */
+	bool ShowOpenTraceFileDialog(FString& OutTraceFile) const;
+
+	/**
+	 * Starts a new Unreal Insights instance.
+	 * @param CmdLine - The command line passed to the new UnrealInsights.exe process
+	 */
+	void OpenUnrealInsights(const TCHAR* CmdLine = nullptr) const;
+
+	/**
+	 * Shows the open file dialog and starts analysis session for the chosen trace file, in a new Unreal Insights instance.
+	 */
+	void OpenTraceFile() const;
+
+	/**
+	 * Starts analysis session for the specified trace file, in a new Unreal Insights instance.
+	 * @param TraceFilename - The trace file to analyze
+	 */
+	void OpenTraceFile(const FString& TraceFilename) const;
+
+	/**
+	 * Creates a new analysis session instance and loads the latest available trace that is live.
+	 * Replaces the current analysis session.
+	 */
 	void LoadLastLiveSession();
 
 	/**
-	 * Creates a new analysis session instance using specified trace id.
+	 * Creates a new analysis session instance using the specified trace id.
+	 * Replaces the current analysis session.
 	 * @param TraceId - The id of the trace to analyze
 	 * @param InAutoQuit - The Application will close when session analysis is complete or fails to start
 	 */
 	void LoadTrace(uint32 TraceId, bool InAutoQuit = false);
 
+	/**
+	 * Shows the open file dialog and creates a new analysis session instance for the chosen trace file.
+	 * Replaces the current analysis session.
+	 */
+	void LoadTraceFile();
+
 	/**
 	 * Creates a new analysis session instance and loads a trace file from the specified location.
+	 * Replaces the current analysis session.
 	 * @param TraceFilename - The trace file to analyze
 	 * @param InAutoQuit - The Application will close when session analysis is complete or fails to start
 	 */
diff --git a/Engine/Source/Developer/TraceInsights/Private/Insights/Widgets/SStartPageWindow.cpp b/Engine/Source/Developer/TraceInsights/Private/Insights/Widgets/SStartPageWindow.cpp
--- a/Engine/Source/Developer/TraceInsights/Private/Insights/Widgets/SStartPageWindow.cpp
+++ b/Engine/Source/Developer/TraceInsights/Private/Insights/Widgets/SStartPageWindow.cpp
@@ -1587,7 +1587,7 @@ void STraceStoreWindow::UpdateTrace(FTraceViewModel& InOutTrace, const Insights:
 		{
 			UE_LOG(TraceInsights, Log, TEXT("[TraceStore] Auto starting analysis for trace with id 0x%08X..."), InOutTrace.TraceId);
 			AutoStartedSessions.Add(InOutTrace.TraceId);
-			LoadTrace(InOutTrace.TraceId);
+			OpenTraceSession(InOutTrace.TraceId);
 		}
 	}
 }
@@ -1664,7 +1664,7 @@ void STraceStoreWindow::TraceList_OnSelectionChanged(TSharedPtr<FTraceViewModel>
 
 void STraceStoreWindow::TraceList_OnMouseButtonDoubleClick(TSharedPtr<FTraceViewModel> TraceSession)
 {
-	LoadTraceSession(TraceSession);
+	OpenTraceSession(TraceSession);
 }
 
 ////////////////////////////////////////////////////////////////////////////////////////////////////
@@ -1799,7 +1799,7 @@ FReply STraceStoreWindow::OnDrop(const FGeometry& MyGeometry, const FDragDropEve
 				const FString DraggedFileExtension = FPaths::GetExtension(Files[0], true);
 				if (DraggedFileExtension == TEXT(".utrace"))
 				{
-					LoadTraceFile(Files[0]);
+					OpenTraceFile(Files[0]);
 					return FReply::Handled();
 				}
 			}
@@ -1820,98 +1820,55 @@ bool STraceStoreWindow::Open_IsEnabled() const
 
 FReply STraceStoreWindow::Open_OnClicked()
 {
-	LoadTraceSession(SelectedTrace);
+	OpenTraceSession(SelectedTrace);
 	return FReply::Handled();
 }
 
 ////////////////////////////////////////////////////////////////////////////////////////////////////
 
-void STraceStoreWindow::OpenFileDialog()
+void STraceStoreWindow::OpenTraceFile()
 {
-	const FString ProfilingDirectory(FPaths::ConvertRelativePathToFull(FInsightsManager::Get()->GetStoreDir()));
-
-	TArray<FString> OutFiles;
-	bool bOpened = false;
-
-	IDesktopPlatform* DesktopPlatform = FDesktopPlatformModule::Get();
-	if (DesktopPlatform != nullptr)
-	{
-		FSlateApplication::Get().CloseToolTip();
-
-		bOpened = DesktopPlatform->OpenFileDialog
-		(
-			FSlateApplication::Get().FindBestParentWindowHandleForDialogs(nullptr),
-			LOCTEXT("LoadTrace_FileDesc", "Open trace file...").ToString(),
-			ProfilingDirectory,
-			TEXT(""),
-			LOCTEXT("LoadTrace_FileFilter", "Trace files (*.utrace)|*.utrace|All files (*.*)|*.*").ToString(),
-			EFileDialogFlags::None,
-			OutFiles
-		);
-	}
-
-	if (bOpened == true)
+	FString TraceFile;
+	if (FInsightsManager::Get()->ShowOpenTraceFileDialog(TraceFile))
 	{
-		if (OutFiles.Num() == 1)
-		{
-			LoadTraceFile(OutFiles[0]);
-		}
+		OpenTraceFile(TraceFile);
 	}
 }
 
 ////////////////////////////////////////////////////////////////////////////////////////////////////
 
-void STraceStoreWindow::LoadTraceSession(TSharedPtr<FTraceViewModel> InTraceSession)
-{
-	if (InTraceSession.IsValid())
-	{
-		LoadTrace(InTraceSession->TraceId);
-	}
-}
-
-////////////////////////////////////////////////////////////////////////////////////////////////////
-
-void STraceStoreWindow::LoadTraceFile(const FString& InTraceFile)
+void STraceStoreWindow::OpenTraceFile(const FString& InTraceFile)
 {
 	UE_LOG(TraceInsights, Log, TEXT("[TraceStore] Start analysis (in separate process) for trace file: \"%s\""), *InTraceFile);
 
-	const TCHAR* ExecutablePath = FPlatformProcess::ExecutablePath();
-
 	FString CmdLine = TEXT("-OpenTraceFile=\"") + InTraceFile + TEXT("\"");
 
 	FString ExtraCmdParams;
 	GetExtraCommandLineParams(ExtraCmdParams);
 	CmdLine += ExtraCmdParams;
 
-	constexpr bool bLaunchDetached = true;
-	constexpr bool bLaunchHidden = false;
-	constexpr bool bLaunchReallyHidden = false;
+	FInsightsManager::Get()->OpenUnrealInsights(*CmdLine);
 
-	uint32 ProcessID = 0;
-	const int32 PriorityModifier = 0;
-	const TCHAR* OptionalWorkingDirectory = nullptr;
+	SplashScreenOverlayTraceFile = FPaths::GetBaseFilename(InTraceFile);
+	ShowSplashScreenOverlay();
+}
 
-	void* PipeWriteChild = nullptr;
-	void* PipeReadChild = nullptr;
+////////////////////////////////////////////////////////////////////////////////////////////////////
 
-	FProcHandle Handle = FPlatformProcess::CreateProc(ExecutablePath, *CmdLine, bLaunchDetached, bLaunchHidden, bLaunchReallyHidden, &ProcessID, PriorityModifier, OptionalWorkingDirectory, PipeWriteChild, PipeReadChild);
-	if (Handle.IsValid())
+void STraceStoreWindow::OpenTraceSession(TSharedPtr<FTraceViewModel> InTraceSession)
+{
+	if (InTraceSession.IsValid())
 	{
-		FPlatformProcess::CloseProc(Handle);
+		OpenTraceSession(InTraceSession->TraceId);
 	}
-
-	SplashScreenOverlayTraceFile = FPaths::GetBaseFilename(InTraceFile);
-	ShowSplashScreenOverlay();
 }
 
 ////////////////////////////////////////////////////////////////////////////////////////////////////
 
-void STraceStoreWindow::LoadTrace(uint32 InTraceId)
+void STraceStoreWindow::OpenTraceSession(uint32 InTraceId)
 {
 	UE_LOG(TraceInsights, Log, TEXT("[TraceStore] Start analysis (in separate process) for trace id: 0x%08X"), InTraceId);
 
-	const TCHAR* ExecutablePath = FPlatformProcess::ExecutablePath();
-
 	const uint32 StorePort = FInsightsManager::Get()->GetStoreClient()->GetStorePort();
 	FString CmdLine = FString::Printf(TEXT("-OpenTraceId=%d -StorePort=%d"), InTraceId, StorePort);
 
@@ -1919,22 +1876,7 @@ void STraceStoreWindow::LoadTrace(uint32 InTraceId)
 	GetExtraCommandLineParams(ExtraCmdParams);
 	CmdLine += ExtraCmdParams;
 
-	constexpr bool bLaunchDetached = true;
-	constexpr bool bLaunchHidden = false;
-	constexpr bool bLaunchReallyHidden = false;
-
-	uint32 ProcessID = 0;
-	const int32 PriorityModifier = 0;
-	const TCHAR* OptionalWorkingDirectory = nullptr;
-
-	void* PipeWriteChild = nullptr;
-	void* PipeReadChild = nullptr;
-
-	FProcHandle Handle = FPlatformProcess::CreateProc(ExecutablePath, *CmdLine, bLaunchDetached, bLaunchHidden, bLaunchReallyHidden, &ProcessID, PriorityModifier, OptionalWorkingDirectory, PipeWriteChild, PipeReadChild);
-	if (Handle.IsValid())
-	{
-		FPlatformProcess::CloseProc(Handle);
-	}
+	FInsightsManager::Get()->OpenUnrealInsights(*CmdLine);
 
 	TSharedPtr<FTraceViewModel>* TraceSessionPtrPtr = TraceViewModelMap.Find(InTraceId);
 	if (TraceSessionPtrPtr)
@@ -1961,7 +1903,7 @@ TSharedRef<SWidget> STraceStoreWindow::MakeTraceListMenu()
 			LOCTEXT("OpenFileButtonLabel", "Open File..."),
 			LOCTEXT("OpenFileButtonTooltip", "Start analysis for a specified trace file."),
 			FSlateIcon(FAppStyle::Get().GetStyleSetName(), "Icons.FolderOpen"),
-			FUIAction(FExecuteAction::CreateSP(this, &STraceStoreWindow::OpenFileDialog)),
+			FUIAction(FExecuteAction::CreateSP(this, &STraceStoreWindow::OpenTraceFile)),
 			NAME_None,
 			EUserInterfaceActionType::Button
 		);
@@ -1994,7 +1936,7 @@ TSharedRef<SWidget> STraceStoreWindow::MakeTraceListMenu()
 					Label,
 					TAttribute<FText>(), // no tooltip
 					FSlateIcon(),
-					FUIAction(FExecuteAction::CreateSP(this, &STraceStoreWindow::LoadTrace, Trace.TraceId)),
+					FUIAction(FExecuteAction::CreateSP(this, &STraceStoreWindow::OpenTraceSession, Trace.TraceId)),
 					NAME_None,
 					EUserInterfaceActionType::Button
 				);
diff --git a/Engine/Source/Developer/TraceInsights/Private/Insights/Widgets/SStartPageWindow.h b/Engine/Source/Developer/TraceInsights/Private/Insights/Widgets/SStartPageWindow.h
--- a/Engine/Source/Developer/TraceInsights/Private/Insights/Widgets/SStartPageWindow.h
+++ b/Engine/Source/Developer/TraceInsights/Private/Insights/Widgets/SStartPageWindow.h
@@ -283,11 +283,10 @@ private:
 	bool Open_IsEnabled() const;
 	FReply Open_OnClicked();
 
-	void OpenFileDialog();
-
-	void LoadTraceSession(TSharedPtr<FTraceViewModel> InTrace);
-	void LoadTraceFile(const FString& InTraceFile);
-	void LoadTrace(uint32 InTraceId);
+	void OpenTraceFile();
+	void OpenTraceFile(const FString& InTraceFile);
+	void OpenTraceSession(TSharedPtr<FTraceViewModel> InTrace);
+	void OpenTraceSession(uint32 InTraceId);
 
 	//////////////////////////////////////////////////
 	// Traces
