[Insights] Fix hang when closing during an async operation.

diff --git a/Engine/Source/Developer/TraceInsights/Private/Insights/InsightsManager.cpp b/Engine/Source/Developer/TraceInsights/Private/Insights/InsightsManager.cpp
--- a/Engine/Source/Developer/TraceInsights/Private/Insights/InsightsManager.cpp
+++ b/Engine/Source/Developer/TraceInsights/Private/Insights/InsightsManager.cpp
@@ -1137,7 +1137,7 @@ bool FInsightsManager::HandleResponseFileCmd(const TCHAR* ResponseFile, FOutputD
 void FInsightsManager::AddInProgressAsyncOp(FGraphEventRef Event, const FString& Name)
 {
 	check(IsInGameThread());
-	if (Event.IsValid())
+	if (Event.IsValid() && !Event->IsComplete())
 	{
 		InProgressAsyncTasks.AddTail(FAsyncTaskData(Event, Name));
 	}
diff --git a/Engine/Source/Developer/TraceInsights/Private/Insights/Table/Widgets/STableTreeView.cpp b/Engine/Source/Developer/TraceInsights/Private/Insights/Table/Widgets/STableTreeView.cpp
--- a/Engine/Source/Developer/TraceInsights/Private/Insights/Table/Widgets/STableTreeView.cpp
+++ b/Engine/Source/Developer/TraceInsights/Private/Insights/Table/Widgets/STableTreeView.cpp
@@ -2801,7 +2801,7 @@ FGraphEventRef STableTreeView::StartApplyFiltersTask(FGraphEventRef Prerequisite
 
 void STableTreeView::OnClose()
 {
-	if (bIsUpdateRunning && InProgressAsyncOperationEvent.IsValid() && !InProgressAsyncOperationEvent->IsComplete())
+	if (bIsUpdateRunning && !bIsCloseScheduled && InProgressAsyncOperationEvent.IsValid() && !InProgressAsyncOperationEvent->IsComplete())
 	{
 		bIsCloseScheduled = true;
 		CancelCurrentAsyncOp();
