[Insights]
  - Removed workaround for closing a tab.
  - Refactored S*ProfilerWindow implementations to use a base class. Moved all common code into a new SMajorTabWindow class.
  - Added extensibility to layout and main toolbar (and to the right side toolbar) for all profiler windows: Timing Insights, Asset Loading Insights, Memory Insights and Networking Insights.

diff --git a/Engine/Source/Developer/TraceInsights/Private/Insights/InsightsManager.cpp b/Engine/Source/Developer/TraceInsights/Private/Insights/InsightsManager.cpp
--- a/Engine/Source/Developer/TraceInsights/Private/Insights/InsightsManager.cpp
+++ b/Engine/Source/Developer/TraceInsights/Private/Insights/InsightsManager.cpp
@@ -678,7 +678,6 @@ void FInsightsManager::SpawnAndActivateTabs()
 			while ((NetworkingProfilerTab = FGlobalTabmanager::Get()->FindExistingLiveTab(TabId)).IsValid())
 			{
 				NetworkingProfilerTab->RequestCloseTab();
-				NetworkingProfilerTab->RemoveTabFromParent();
 			}
 		}
 	}
diff --git a/Engine/Source/Developer/TraceInsights/Private/Insights/LoadingProfiler/LoadingProfilerManager.cpp b/Engine/Source/Developer/TraceInsights/Private/Insights/LoadingProfiler/LoadingProfilerManager.cpp
--- a/Engine/Source/Developer/TraceInsights/Private/Insights/LoadingProfiler/LoadingProfilerManager.cpp
+++ b/Engine/Source/Developer/TraceInsights/Private/Insights/LoadingProfiler/LoadingProfilerManager.cpp
@@ -6,6 +6,7 @@
 #include "Modules/ModuleManager.h"
 #include "TraceServices/AnalysisService.h"
 #include "TraceServices/Model/LoadTimeProfiler.h"
+#include "Widgets/Docking/SDockTab.h"
 #include "WorkspaceMenuStructure.h"
 #include "WorkspaceMenuStructureModule.h"
 
diff --git a/Engine/Source/Developer/TraceInsights/Private/Insights/LoadingProfiler/Widgets/SLoadingProfilerToolbar.cpp b/Engine/Source/Developer/TraceInsights/Private/Insights/LoadingProfiler/Widgets/SLoadingProfilerToolbar.cpp
--- a/Engine/Source/Developer/TraceInsights/Private/Insights/LoadingProfiler/Widgets/SLoadingProfilerToolbar.cpp
+++ b/Engine/Source/Developer/TraceInsights/Private/Insights/LoadingProfiler/Widgets/SLoadingProfilerToolbar.cpp
@@ -2,6 +2,7 @@
 
 #include "SLoadingProfilerToolbar.h"
 
+#include "Framework/MultiBox/MultiBoxExtender.h"
 #include "Framework/MultiBox/MultiBoxBuilder.h"
 #include "Widgets/SBoxPanel.h"
 
@@ -33,7 +34,7 @@ void SLoadingProfilerToolbar::Construct(const FArguments& InArgs)
 {
 	struct Local
 	{
-		static void FillViewToolbar(FToolBarBuilder& ToolbarBuilder)
+		static void FillViewToolbar(FToolBarBuilder& ToolbarBuilder, const FArguments &InArgs)
 		{
 			ToolbarBuilder.BeginSection("View");
 			{
@@ -57,9 +58,14 @@ void SLoadingProfilerToolbar::Construct(const FArguments& InArgs)
 					FSlateIcon(FInsightsStyle::GetStyleSetName(), "Icons.TableTreeView.ToolBar"));
 			}
 			ToolbarBuilder.EndSection();
+
+			if (InArgs._ToolbarExtender.IsValid())
+			{
+				InArgs._ToolbarExtender->Apply("MainToolbar", EExtensionHook::First, ToolbarBuilder);
+			}
 		}
 
-		static void FillRightSideToolbar(FToolBarBuilder& ToolbarBuilder)
+		static void FillRightSideToolbar(FToolBarBuilder& ToolbarBuilder, const FArguments &InArgs)
 		{
 			ToolbarBuilder.BeginSection("Debug");
 			{
@@ -68,6 +74,11 @@ void SLoadingProfilerToolbar::Construct(const FArguments& InArgs)
 					FSlateIcon(FInsightsStyle::GetStyleSetName(), "Icons.Debug.ToolBar"));
 			}
 			ToolbarBuilder.EndSection();
+
+			if (InArgs._ToolbarExtender.IsValid())
+			{
+				InArgs._ToolbarExtender->Apply("RightSideToolbar", EExtensionHook::First, ToolbarBuilder);
+			}
 		}
 	};
 
@@ -75,11 +86,11 @@ void SLoadingProfilerToolbar::Construct(const FArguments& InArgs)
 
 	FSlimHorizontalToolBarBuilder ToolbarBuilder(CommandList.ToSharedRef(), FMultiBoxCustomization::None);
 	ToolbarBuilder.SetStyle(&FInsightsStyle::Get(), "PrimaryToolbar");
-	Local::FillViewToolbar(ToolbarBuilder);
+	Local::FillViewToolbar(ToolbarBuilder, InArgs);
 
 	FSlimHorizontalToolBarBuilder RightSideToolbarBuilder(CommandList.ToSharedRef(), FMultiBoxCustomization::None);
 	RightSideToolbarBuilder.SetStyle(&FInsightsStyle::Get(), "PrimaryToolbar");
-	Local::FillRightSideToolbar(RightSideToolbarBuilder);
+	Local::FillRightSideToolbar(RightSideToolbarBuilder, InArgs);
 
 	ChildSlot
 	[
diff --git a/Engine/Source/Developer/TraceInsights/Private/Insights/LoadingProfiler/Widgets/SLoadingProfilerToolbar.h b/Engine/Source/Developer/TraceInsights/Private/Insights/LoadingProfiler/Widgets/SLoadingProfilerToolbar.h
--- a/Engine/Source/Developer/TraceInsights/Private/Insights/LoadingProfiler/Widgets/SLoadingProfilerToolbar.h
+++ b/Engine/Source/Developer/TraceInsights/Private/Insights/LoadingProfiler/Widgets/SLoadingProfilerToolbar.h
@@ -5,6 +5,8 @@
 #include "Widgets/DeclarativeSyntaxSupport.h"
 #include "Widgets/SCompoundWidget.h"
 
+class FExtender;
+
 /** Ribbon based toolbar used as a main menu in the Profiler window. */
 class SLoadingProfilerToolbar : public SCompoundWidget
 {
@@ -16,6 +18,7 @@ public:
 	virtual ~SLoadingProfilerToolbar();
 
 	SLATE_BEGIN_ARGS(SLoadingProfilerToolbar) {}
+	SLATE_ARGUMENT(TSharedPtr<FExtender>, ToolbarExtender);
 	SLATE_END_ARGS()
 
 	/**
@@ -24,8 +27,4 @@ public:
 	 * @param InArgs The declaration data for this widget
 	 */
 	void Construct(const FArguments& InArgs);
-
-private:
-	/** Create the UI commands for the toolbar */
-	void CreateCommands();
 };
diff --git a/Engine/Source/Developer/TraceInsights/Private/Insights/LoadingProfiler/Widgets/SLoadingProfilerWindow.cpp b/Engine/Source/Developer/TraceInsights/Private/Insights/LoadingProfiler/Widgets/SLoadingProfilerWindow.cpp
--- a/Engine/Source/Developer/TraceInsights/Private/Insights/LoadingProfiler/Widgets/SLoadingProfilerWindow.cpp
+++ b/Engine/Source/Developer/TraceInsights/Private/Insights/LoadingProfiler/Widgets/SLoadingProfilerWindow.cpp
@@ -3,35 +3,16 @@
 #include "SLoadingProfilerWindow.h"
 
 #include "Framework/Commands/UICommandList.h"
-#include "Framework/Docking/LayoutService.h"
 #include "Framework/Docking/TabManager.h"
 #include "Framework/Docking/WorkspaceItem.h"
-#include "Framework/MultiBox/MultiBoxBuilder.h"
 #include "SlateOptMacros.h"
-#include "Styling/AppStyle.h"
 #include "Widgets/Docking/SDockTab.h"
-#include "Widgets/Images/SImage.h"
-#include "Widgets/Layout/SBorder.h"
-#include "Widgets/Layout/SBox.h"
-#include "Widgets/SBoxPanel.h"
-#include "Widgets/Text/STextBlock.h"
-
-#if WITH_EDITOR
-	#include "EngineAnalytics.h"
-	#include "Runtime/Analytics/Analytics/Public/AnalyticsEventAttribute.h"
-	#include "Runtime/Analytics/Analytics/Public/Interfaces/IAnalyticsProvider.h"
-#endif // WITH_EDITOR
 
 // Insights
-#include "Insights/Common/InsightsMenuBuilder.h"
-#include "Insights/InsightsManager.h"
 #include "Insights/InsightsStyle.h"
 #include "Insights/LoadingProfiler/LoadingProfilerManager.h"
 #include "Insights/LoadingProfiler/Widgets/SLoadingProfilerToolbar.h"
 #include "Insights/Table/Widgets/SUntypedTableTreeView.h"
-#include "Insights/TraceInsightsModule.h"
-#include "Insights/Version.h"
-#include "Insights/Widgets/SInsightsSettings.h"
 #include "Insights/Widgets/STimingView.h"
 
 ////////////////////////////////////////////////////////////////////////////////////////////////////
@@ -50,7 +31,7 @@ const FName FLoadingProfilerTabs::RequestsTreeViewID(TEXT("Requests"));
 ////////////////////////////////////////////////////////////////////////////////////////////////////
 
 SLoadingProfilerWindow::SLoadingProfilerWindow()
-	: DurationActive(0.0f)
+	: SMajorTabWindow(FInsightsManagerTabs::LoadingProfilerTabId)
 	, SelectionStartTime(0.0f)
 	, SelectionEndTime(0.0f)
 {
@@ -60,19 +41,7 @@ SLoadingProfilerWindow::SLoadingProfilerWindow()
 
 SLoadingProfilerWindow::~SLoadingProfilerWindow()
 {
-	// Close all tabs.
-	TArray<TSharedPtr<SDockTab>> LocalOpenTabs;
-	for (TSharedPtr<SDockTab>& Tab : OpenTabs)
-	{
-		LocalOpenTabs.Add(Tab);
-	}
-	for (TSharedPtr<SDockTab>& Tab : LocalOpenTabs)
-	{
-		Tab->RequestCloseTab();
-		Tab->RemoveTabFromParent();
-	}
-	LocalOpenTabs.Reset();
-	check(OpenTabs.IsEmpty());
+	CloseAllOpenTabs();
 
 	check(RequestsTreeView == nullptr);
 	check(ExportDetailsTreeView == nullptr);
@@ -80,13 +49,13 @@ SLoadingProfilerWindow::~SLoadingProfilerWindow()
 	check(ObjectTypeAggregationTreeView == nullptr);
 	check(EventAggregationTreeView == nullptr);
 	check(TimingView == nullptr);
+}
 
-#if WITH_EDITOR
-	if (DurationActive > 0.0f && FEngineAnalytics::IsAvailable())
-	{
-		FEngineAnalytics::GetProvider().RecordEvent(TEXT("Insights.Usage.LoadingProfiler"), FAnalyticsEventAttribute(TEXT("Duration"), DurationActive));
-	}
-#endif // WITH_EDITOR
+////////////////////////////////////////////////////////////////////////////////////////////////////
+
+const TCHAR* SLoadingProfilerWindow::GetAnalyticsEventName() const
+{
+	return TEXT("Insights.Usage.LoadingProfiler");
 }
 
 ////////////////////////////////////////////////////////////////////////////////////////////////////
@@ -281,7 +250,7 @@ TSharedRef<SDockTab> SLoadingProfilerWindow::SpawnTab_TimingView(const FSpawnTab
 	TimingView->SelectTimeInterval(SelectionStartTime, SelectionEndTime - SelectionStartTime);
 
 	DockTab->SetOnTabClosed(SDockTab::FOnTabClosedCallback::CreateRaw(this, &SLoadingProfilerWindow::OnTimingViewTabClosed));
-	OpenTabs.Add(DockTab);
+	AddOpenTab(DockTab);
 
 	return DockTab;
 }
@@ -298,7 +267,7 @@ void SLoadingProfilerWindow::OnTimingViewTabClosed(TSharedRef<SDockTab> TabBeing
 		TimingView = nullptr;
 	}
 
-	OpenTabs.Remove(TabBeingClosed);
+	RemoveOpenTab(TabBeingClosed);
 }
 
 ////////////////////////////////////////////////////////////////////////////////////////////////////
@@ -322,7 +291,7 @@ TSharedRef<SDockTab> SLoadingProfilerWindow::SpawnTab_EventAggregationTreeView(c
 	UpdateEventAggregationTreeView();
 
 	DockTab->SetOnTabClosed(SDockTab::FOnTabClosedCallback::CreateRaw(this, &SLoadingProfilerWindow::OnEventAggregationTreeViewTabClosed));
-	OpenTabs.Add(DockTab);
+	AddOpenTab(DockTab);
 
 	return DockTab;
 }
@@ -335,7 +304,7 @@ void SLoadingProfilerWindow::OnEventAggregationTreeViewTabClosed(TSharedRef<SDoc
 	FLoadingProfilerManager::Get()->SetEventAggregationTreeViewVisible(false);
 	EventAggregationTreeView = nullptr;
 
-	OpenTabs.Remove(TabBeingClosed);
+	RemoveOpenTab(TabBeingClosed);
 }
 
 ////////////////////////////////////////////////////////////////////////////////////////////////////
@@ -359,7 +328,7 @@ TSharedRef<SDockTab> SLoadingProfilerWindow::SpawnTab_ObjectTypeAggregationTreeV
 	UpdateObjectTypeAggregationTreeView();
 
 	DockTab->SetOnTabClosed(SDockTab::FOnTabClosedCallback::CreateRaw(this, &SLoadingProfilerWindow::OnObjectTypeAggregationTreeViewTabClosed));
-	OpenTabs.Add(DockTab);
+	AddOpenTab(DockTab);
 
 	return DockTab;
 }
@@ -372,7 +341,7 @@ void SLoadingProfilerWindow::OnObjectTypeAggregationTreeViewTabClosed(TSharedRef
 	FLoadingProfilerManager::Get()->SetObjectTypeAggregationTreeViewVisible(false);
 	ObjectTypeAggregationTreeView = nullptr;
 
-	OpenTabs.Remove(TabBeingClosed);
+	RemoveOpenTab(TabBeingClosed);
 }
 
 ////////////////////////////////////////////////////////////////////////////////////////////////////
@@ -396,7 +365,7 @@ TSharedRef<SDockTab> SLoadingProfilerWindow::SpawnTab_PackageDetailsTreeView(con
 	UpdatePackageDetailsTreeView();
 
 	DockTab->SetOnTabClosed(SDockTab::FOnTabClosedCallback::CreateRaw(this, &SLoadingProfilerWindow::OnPackageDetailsTreeViewTabClosed));
-	OpenTabs.Add(DockTab);
+	AddOpenTab(DockTab);
 
 	return DockTab;
 }
@@ -409,7 +378,7 @@ void SLoadingProfilerWindow::OnPackageDetailsTreeViewTabClosed(TSharedRef<SDockT
 	FLoadingProfilerManager::Get()->SetPackageDetailsTreeViewVisible(false);
 	PackageDetailsTreeView = nullptr;
 
-	OpenTabs.Remove(TabBeingClosed);
+	RemoveOpenTab(TabBeingClosed);
 }
 
 ////////////////////////////////////////////////////////////////////////////////////////////////////
@@ -433,7 +402,7 @@ TSharedRef<SDockTab> SLoadingProfilerWindow::SpawnTab_ExportDetailsTreeView(cons
 	UpdateExportDetailsTreeView();
 
 	DockTab->SetOnTabClosed(SDockTab::FOnTabClosedCallback::CreateRaw(this, &SLoadingProfilerWindow::OnExportDetailsTreeViewTabClosed));
-	OpenTabs.Add(DockTab);
+	AddOpenTab(DockTab);
 
 	return DockTab;
 }
@@ -446,7 +415,7 @@ void SLoadingProfilerWindow::OnExportDetailsTreeViewTabClosed(TSharedRef<SDockTa
 	FLoadingProfilerManager::Get()->SetExportDetailsTreeViewVisible(false);
 	ExportDetailsTreeView = nullptr;
 
-	OpenTabs.Remove(TabBeingClosed);
+	RemoveOpenTab(TabBeingClosed);
 }
 
 ////////////////////////////////////////////////////////////////////////////////////////////////////
@@ -470,7 +439,7 @@ TSharedRef<SDockTab> SLoadingProfilerWindow::SpawnTab_RequestsTreeView(const FSp
 	UpdateRequestsTreeView();
 
 	DockTab->SetOnTabClosed(SDockTab::FOnTabClosedCallback::CreateRaw(this, &SLoadingProfilerWindow::OnRequestsTreeViewTabClosed));
-	OpenTabs.Add(DockTab);
+	AddOpenTab(DockTab);
 
 	return DockTab;
 }
@@ -483,59 +452,75 @@ void SLoadingProfilerWindow::OnRequestsTreeViewTabClosed(TSharedRef<SDockTab> Ta
 	FLoadingProfilerManager::Get()->SetRequestsTreeViewVisible(false);
 	RequestsTreeView = nullptr;
 
-	OpenTabs.Remove(TabBeingClosed);
+	RemoveOpenTab(TabBeingClosed);
 }
 
 ////////////////////////////////////////////////////////////////////////////////////////////////////
 
-BEGIN_SLATE_FUNCTION_BUILD_OPTIMIZATION
 void SLoadingProfilerWindow::Construct(const FArguments& InArgs, const TSharedRef<SDockTab>& ConstructUnderMajorTab, const TSharedPtr<SWindow>& ConstructUnderWindow)
 {
-	// Create & initialize tab manager.
-	TabManager = FGlobalTabmanager::Get()->NewTabManager(ConstructUnderMajorTab);
-	const auto& PersistLayout = [](const TSharedRef<FTabManager::FLayout>& LayoutToSave)
-	{
-		FLayoutSaveRestore::SaveToConfig(FTraceInsightsModule::GetUnrealInsightsLayoutIni(), LayoutToSave);
-	};
-	TabManager->SetOnPersistLayout(FTabManager::FOnPersistLayout::CreateLambda(PersistLayout));
+	TSharedPtr<FLoadingProfilerManager> LoadingProfilerManager = FLoadingProfilerManager::Get();
+	ensure(LoadingProfilerManager.IsValid());
 
-	TSharedRef<FWorkspaceItem> AppMenuGroup = TabManager->AddLocalWorkspaceMenuCategory(LOCTEXT("LoadingProfilerMenuGroupName", "Asset Loading Insights"));
+	SetCommandList(LoadingProfilerManager->GetCommandList());
 
-	TabManager->RegisterTabSpawner(FLoadingProfilerTabs::TimingViewID, FOnSpawnTab::CreateRaw(this, &SLoadingProfilerWindow::SpawnTab_TimingView))
+	SMajorTabWindow::FArguments Args;
+	SMajorTabWindow::Construct(Args, ConstructUnderMajorTab, ConstructUnderWindow);
+}
+
+////////////////////////////////////////////////////////////////////////////////////////////////////
+
+TSharedRef<FWorkspaceItem> SLoadingProfilerWindow::CreateWorkspaceMenuGroup()
+{
+	return GetTabManager()->AddLocalWorkspaceMenuCategory(LOCTEXT("LoadingProfilerMenuGroupName", "Asset Loading Insights"));
+}
+
+////////////////////////////////////////////////////////////////////////////////////////////////////
+
+void SLoadingProfilerWindow::RegisterTabSpawners()
+{
+	check(GetTabManager().IsValid());
+	FTabManager* TabManagerPtr = GetTabManager().Get();
+	check(GetWorkspaceMenuGroup().IsValid());
+	const TSharedRef<FWorkspaceItem> Group = GetWorkspaceMenuGroup().ToSharedRef();
+
+	TabManagerPtr->RegisterTabSpawner(FLoadingProfilerTabs::TimingViewID, FOnSpawnTab::CreateRaw(this, &SLoadingProfilerWindow::SpawnTab_TimingView))
 		.SetDisplayName(LOCTEXT("TimingViewTabTitle", "Timing View"))
 		.SetIcon(FSlateIcon(FInsightsStyle::GetStyleSetName(), "Icons.TimingView"))
-		.SetGroup(AppMenuGroup);
+		.SetGroup(Group);
 
-	TabManager->RegisterTabSpawner(FLoadingProfilerTabs::EventAggregationTreeViewID, FOnSpawnTab::CreateRaw(this, &SLoadingProfilerWindow::SpawnTab_EventAggregationTreeView))
+	TabManagerPtr->RegisterTabSpawner(FLoadingProfilerTabs::EventAggregationTreeViewID, FOnSpawnTab::CreateRaw(this, &SLoadingProfilerWindow::SpawnTab_EventAggregationTreeView))
 		.SetDisplayName(LOCTEXT("EventAggregationTreeViewTabTitle", "Event Aggregation"))
 		.SetIcon(FSlateIcon(FInsightsStyle::GetStyleSetName(), "Icons.TableTreeView"))
-		.SetGroup(AppMenuGroup);
+		.SetGroup(Group);
 
-	TabManager->RegisterTabSpawner(FLoadingProfilerTabs::ObjectTypeAggregationTreeViewID, FOnSpawnTab::CreateRaw(this, &SLoadingProfilerWindow::SpawnTab_ObjectTypeAggregationTreeView))
+	TabManagerPtr->RegisterTabSpawner(FLoadingProfilerTabs::ObjectTypeAggregationTreeViewID, FOnSpawnTab::CreateRaw(this, &SLoadingProfilerWindow::SpawnTab_ObjectTypeAggregationTreeView))
 		.SetDisplayName(LOCTEXT("ObjectTypeAggregationTreeViewTabTitle", "Object Type Aggregation"))
 		.SetIcon(FSlateIcon(FInsightsStyle::GetStyleSetName(), "Icons.TableTreeView"))
-		.SetGroup(AppMenuGroup);
+		.SetGroup(Group);
 
-	TabManager->RegisterTabSpawner(FLoadingProfilerTabs::PackageDetailsTreeViewID, FOnSpawnTab::CreateRaw(this, &SLoadingProfilerWindow::SpawnTab_PackageDetailsTreeView))
+	TabManagerPtr->RegisterTabSpawner(FLoadingProfilerTabs::PackageDetailsTreeViewID, FOnSpawnTab::CreateRaw(this, &SLoadingProfilerWindow::SpawnTab_PackageDetailsTreeView))
 		.SetDisplayName(LOCTEXT("PackageDetailsTreeViewTabTitle", "Package Details"))
 		.SetIcon(FSlateIcon(FInsightsStyle::GetStyleSetName(), "Icons.TableTreeView"))
-		.SetGroup(AppMenuGroup);
+		.SetGroup(Group);
 
-	TabManager->RegisterTabSpawner(FLoadingProfilerTabs::ExportDetailsTreeViewID, FOnSpawnTab::CreateRaw(this, &SLoadingProfilerWindow::SpawnTab_ExportDetailsTreeView))
+	TabManagerPtr->RegisterTabSpawner(FLoadingProfilerTabs::ExportDetailsTreeViewID, FOnSpawnTab::CreateRaw(this, &SLoadingProfilerWindow::SpawnTab_ExportDetailsTreeView))
 		.SetDisplayName(LOCTEXT("ExportDetailsTreeViewTabTitle", "Export Details"))
 		.SetIcon(FSlateIcon(FInsightsStyle::GetStyleSetName(), "Icons.TableTreeView"))
-		.SetGroup(AppMenuGroup);
+		.SetGroup(Group);
 
-	TabManager->RegisterTabSpawner(FLoadingProfilerTabs::RequestsTreeViewID, FOnSpawnTab::CreateRaw(this, &SLoadingProfilerWindow::SpawnTab_RequestsTreeView))
+	TabManagerPtr->RegisterTabSpawner(FLoadingProfilerTabs::RequestsTreeViewID, FOnSpawnTab::CreateRaw(this, &SLoadingProfilerWindow::SpawnTab_RequestsTreeView))
 		.SetDisplayName(LOCTEXT("RequestsTreeViewTabTitle", "Requests"))
 		.SetIcon(FSlateIcon(FInsightsStyle::GetStyleSetName(), "Icons.TableTreeView"))
-		.SetGroup(AppMenuGroup);
+		.SetGroup(Group);
+}
 
-	TSharedPtr<FLoadingProfilerManager> LoadingProfilerManager = FLoadingProfilerManager::Get();
-	ensure(LoadingProfilerManager.IsValid());
+////////////////////////////////////////////////////////////////////////////////////////////////////
 
-	// Create tab layout.
-	TSharedRef<FTabManager::FLayout> Layout = FTabManager::NewLayout("InsightsLoadingProfilerLayout_v1.2")
+BEGIN_SLATE_FUNCTION_BUILD_OPTIMIZATION
+TSharedRef<FTabManager::FLayout> SLoadingProfilerWindow::CreateDefaultTabLayout() const
+{
+	return FTabManager::NewLayout("InsightsLoadingProfilerLayout_v1.2")
 		->AddArea
 		(
 			FTabManager::NewPrimaryArea()
@@ -559,215 +544,14 @@ void SLoadingProfilerWindow::Construct(const FArguments& InArgs, const TSharedRe
 				->SetForegroundTab(FLoadingProfilerTabs::EventAggregationTreeViewID)
 			)
 		);
-
-	Layout = FLayoutSaveRestore::LoadFromConfig(FTraceInsightsModule::GetUnrealInsightsLayoutIni(), Layout);
-
-	// Create & initialize main menu.
-	FMenuBarBuilder MenuBarBuilder = FMenuBarBuilder(TSharedPtr<FUICommandList>());
-	MenuBarBuilder.AddPullDownMenu(
-		LOCTEXT("MenuLabel", "Menu"),
-		FText::GetEmpty(),
-		FNewMenuDelegate::CreateStatic(&SLoadingProfilerWindow::FillMenu, TabManager),
-		FName(TEXT("Menu"))
-	);
-
-#if !WITH_EDITOR
-	TSharedRef<SWidget> MenuWidget = MenuBarBuilder.MakeWidget();
-	MenuWidget->SetClipping(EWidgetClipping::ClipToBoundsWithoutIntersecting);
-#endif
-
-	ChildSlot
-	[
-		SNew(SOverlay)
-
-#if !WITH_EDITOR
-		// Menu
-		+ SOverlay::Slot()
-		.HAlign(HAlign_Left)
-		.VAlign(VAlign_Top)
-		.Padding(34.0f, -60.0f, 0.0f, 0.0f)
-		[
-			MenuWidget
-		]
-#endif
-
-		// Version
-		+ SOverlay::Slot()
-		.HAlign(HAlign_Right)
-		.VAlign(VAlign_Top)
-		.Padding(0.0f, -16.0f, 4.0f, 0.0f)
-		[
-			SNew(STextBlock)
-			.Clipping(EWidgetClipping::ClipToBoundsWithoutIntersecting)
-			.Text(LOCTEXT("UnrealInsightsVersion", UNREAL_INSIGHTS_VERSION_STRING_EX))
-			.ColorAndOpacity(FLinearColor(0.15f, 0.15f, 0.15f, 1.0f))
-		]
-
-		// Overlay slot for the main window area
-		+ SOverlay::Slot()
-		.HAlign(HAlign_Fill)
-		.VAlign(VAlign_Fill)
-		[
-			SNew(SVerticalBox)
-
-			+ SVerticalBox::Slot()
-			.AutoHeight()
-			.Padding(FMargin(0.0f, 0.0f, 0.0f, 5.0f))
-			[
-				SNew(SLoadingProfilerToolbar)
-			]
-
-			+ SVerticalBox::Slot()
-			.FillHeight(1.0f)
-			[
-				TabManager->RestoreFrom(Layout, ConstructUnderWindow).ToSharedRef()
-			]
-		]
-
-		// Session hint overlay
-		+ SOverlay::Slot()
-		.HAlign(HAlign_Center)
-		.VAlign(VAlign_Center)
-		[
-			SNew(SBorder)
-			.Visibility(this, &SLoadingProfilerWindow::IsSessionOverlayVisible)
-			.BorderImage(FAppStyle::Get().GetBrush("PopupText.Background"))
-			.Padding(8.0f)
-			[
-				SNew(STextBlock)
-				.Text(LOCTEXT("SelectTraceOverlayText", "Please select a trace."))
-			]
-		]
-	];
-
-#if !WITH_EDITOR
-	// Tell tab-manager about the global menu bar.
-	TabManager->SetMenuMultiBox(MenuBarBuilder.GetMultiBox(), MenuWidget);
-#endif
 }
 END_SLATE_FUNCTION_BUILD_OPTIMIZATION
 
 ////////////////////////////////////////////////////////////////////////////////////////////////////
 
-void SLoadingProfilerWindow::FillMenu(FMenuBuilder& MenuBuilder, const TSharedPtr<FTabManager> TabManager)
-{
-	if (!TabManager.IsValid())
-	{
-		return;
-	}
-
-	FInsightsManager::Get()->GetInsightsMenuBuilder()->PopulateMenu(MenuBuilder);
-
-	TabManager->PopulateLocalTabSpawnerMenu(MenuBuilder);
-}
-
-////////////////////////////////////////////////////////////////////////////////////////////////////
-
-void SLoadingProfilerWindow::ShowTab(const FName& TabID)
-{
-	if (TabManager->HasTabSpawner(TabID))
-	{
-		TabManager->TryInvokeTab(TabID);
-	}
-}
-
-////////////////////////////////////////////////////////////////////////////////////////////////////
-
-void SLoadingProfilerWindow::HideTab(const FName& TabID)
+TSharedRef<SWidget> SLoadingProfilerWindow::CreateToolbar(TSharedPtr<FExtender> Extender)
 {
-	TSharedPtr<SDockTab> Tab = TabManager->FindExistingLiveTab(TabID);
-	if (Tab.IsValid())
-	{
-		Tab->RequestCloseTab();
-		Tab->RemoveTabFromParent();
-	}
-}
-
-////////////////////////////////////////////////////////////////////////////////////////////////////
-
-EVisibility SLoadingProfilerWindow::IsSessionOverlayVisible() const
-{
-	if (FInsightsManager::Get()->GetSession().IsValid())
-	{
-		return EVisibility::Hidden;
-	}
-	else
-	{
-		return EVisibility::Visible;
-	}
-}
-
-////////////////////////////////////////////////////////////////////////////////////////////////////
-
-bool SLoadingProfilerWindow::IsProfilerEnabled() const
-{
-	return FInsightsManager::Get()->GetSession().IsValid();
-}
-
-////////////////////////////////////////////////////////////////////////////////////////////////////
-
-EActiveTimerReturnType SLoadingProfilerWindow::UpdateActiveDuration(double InCurrentTime, float InDeltaTime)
-{
-	DurationActive += InDeltaTime;
-
-	// The profiler window will explicitly unregister this active timer when the mouse leaves.
-	return EActiveTimerReturnType::Continue;
-}
-
-////////////////////////////////////////////////////////////////////////////////////////////////////
-
-void SLoadingProfilerWindow::OnMouseEnter(const FGeometry& MyGeometry, const FPointerEvent& MouseEvent)
-{
-	SCompoundWidget::OnMouseEnter(MyGeometry, MouseEvent);
-
-	if (!ActiveTimerHandle.IsValid())
-	{
-		ActiveTimerHandle = RegisterActiveTimer(0.f, FWidgetActiveTimerDelegate::CreateSP(this, &SLoadingProfilerWindow::UpdateActiveDuration));
-	}
-}
-
-////////////////////////////////////////////////////////////////////////////////////////////////////
-
-void SLoadingProfilerWindow::OnMouseLeave(const FPointerEvent& MouseEvent)
-{
-	SCompoundWidget::OnMouseLeave(MouseEvent);
-
-	auto PinnedActiveTimerHandle = ActiveTimerHandle.Pin();
-	if (PinnedActiveTimerHandle.IsValid())
-	{
-		UnRegisterActiveTimer(PinnedActiveTimerHandle.ToSharedRef());
-	}
-}
-
-////////////////////////////////////////////////////////////////////////////////////////////////////
-
-FReply SLoadingProfilerWindow::OnKeyDown(const FGeometry& MyGeometry, const FKeyEvent& InKeyEvent)
-{
-	return FLoadingProfilerManager::Get()->GetCommandList()->ProcessCommandBindings(InKeyEvent) ? FReply::Handled() : FReply::Unhandled();
-}
-
-////////////////////////////////////////////////////////////////////////////////////////////////////
-
-FReply SLoadingProfilerWindow::OnDragOver(const FGeometry& MyGeometry, const FDragDropEvent& DragDropEvent)
-{
-	if (FInsightsManager::Get()->OnDragOver(DragDropEvent))
-	{
-		return FReply::Handled();
-	}
-
-	return SCompoundWidget::OnDragOver(MyGeometry, DragDropEvent);
-}
-
-////////////////////////////////////////////////////////////////////////////////////////////////////
-
-FReply SLoadingProfilerWindow::OnDrop(const FGeometry& MyGeometry, const FDragDropEvent& DragDropEvent)
-{
-	if (FInsightsManager::Get()->OnDrop(DragDropEvent))
-	{
-		return FReply::Handled();
-	}
-
-	return SCompoundWidget::OnDrop(MyGeometry, DragDropEvent);
+	return SNew(SLoadingProfilerToolbar).ToolbarExtender(Extender);
 }
 
 ////////////////////////////////////////////////////////////////////////////////////////////////////
diff --git a/Engine/Source/Developer/TraceInsights/Private/Insights/LoadingProfiler/Widgets/SLoadingProfilerWindow.h b/Engine/Source/Developer/TraceInsights/Private/Insights/LoadingProfiler/Widgets/SLoadingProfilerWindow.h
--- a/Engine/Source/Developer/TraceInsights/Private/Insights/LoadingProfiler/Widgets/SLoadingProfilerWindow.h
+++ b/Engine/Source/Developer/TraceInsights/Private/Insights/LoadingProfiler/Widgets/SLoadingProfilerWindow.h
@@ -3,24 +3,13 @@
 #pragma once
 
 #include "CoreMinimal.h"
-#include "Framework/Docking/TabManager.h"
-#include "Input/Reply.h"
-#include "Layout/Visibility.h"
-#include "Misc/Guid.h"
-#include "SlateFwd.h"
-#include "Widgets/DeclarativeSyntaxSupport.h"
-#include "Widgets/Docking/SDockTab.h"
-#include "Widgets/SCompoundWidget.h"
 
 // Insights
-#include "Insights/InsightsManager.h"
-#include "Insights/ITimingViewSession.h"
+#include "Insights/ITimingViewSession.h" // for Insights::ETimeChangedFlags
+#include "Insights/Widgets/SMajorTabWindow.h"
 
 ////////////////////////////////////////////////////////////////////////////////////////////////////
 
-class FActiveTimerHandle;
-class FMenuBuilder;
-
 class STimingView;
 
 namespace Insights
@@ -43,8 +32,8 @@ struct FLoadingProfilerTabs
 
 ////////////////////////////////////////////////////////////////////////////////////////////////////
 
-/** Implements the timing profiler window. */
-class SLoadingProfilerWindow : public SCompoundWidget
+/** Implements the Asset Loading Insights window. */
+class SLoadingProfilerWindow : public Insights::SMajorTabWindow
 {
 public:
 	/** Default constructor. */
@@ -56,7 +45,8 @@ public:
 	SLATE_BEGIN_ARGS(SLoadingProfilerWindow) {}
 	SLATE_END_ARGS()
 
-	void Reset();
+	virtual void Reset() override;
+
 	void UpdateTableTreeViews();
 	void UpdateEventAggregationTreeView();
 	void UpdateObjectTypeAggregationTreeView();
@@ -67,12 +57,6 @@ public:
 	/** Constructs this widget. */
 	void Construct(const FArguments& InArgs, const TSharedRef<SDockTab>& ConstructUnderMajorTab, const TSharedPtr<SWindow>& ConstructUnderWindow);
 
-	void ShowTab(const FName& TabID);
-	void HideTab(const FName& TabID);
-	void ShowHideTab(const FName& TabID, bool bShow) { bShow ? ShowTab(TabID) : HideTab(TabID); }
-
-	TSharedPtr<FTabManager> GetTabManager() const { return TabManager; }
-
 	TSharedPtr<STimingView> GetTimingView() const { return TimingView; }
 	TSharedPtr<Insights::SUntypedTableTreeView> GetEventAggregationTreeView() const { return EventAggregationTreeView; }
 	TSharedPtr<Insights::SUntypedTableTreeView> GetObjectTypeAggregationTreeView() const { return ObjectTypeAggregationTreeView; }
@@ -80,6 +64,13 @@ public:
 	TSharedPtr<Insights::SUntypedTableTreeView> GetExportDetailsTreeView() const { return ExportDetailsTreeView; }
 	TSharedPtr<Insights::SUntypedTableTreeView> GetRequestsTreeView() const { return RequestsTreeView; }
 
+protected:
+	virtual const TCHAR* GetAnalyticsEventName() const override;
+	virtual TSharedRef<FWorkspaceItem> CreateWorkspaceMenuGroup() override;
+	virtual void RegisterTabSpawners() override;
+	virtual TSharedRef<FTabManager::FLayout> CreateDefaultTabLayout() const override;
+	virtual TSharedRef<SWidget> CreateToolbar(TSharedPtr<FExtender> Extender);
+
 private:
 	TSharedRef<SDockTab> SpawnTab_TimingView(const FSpawnTabArgs& Args);
 	void OnTimingViewTabClosed(TSharedRef<SDockTab> TabBeingClosed);
@@ -99,86 +90,9 @@ private:
 	TSharedRef<SDockTab> SpawnTab_RequestsTreeView(const FSpawnTabArgs& Args);
 	void OnRequestsTreeViewTabClosed(TSharedRef<SDockTab> TabBeingClosed);
 
-	/**
-	 * Fill the main menu with menu items.
-	 *
-	 * @param MenuBuilder The multi-box builder that should be filled with content for this pull-down menu.
-	 * @param TabManager A Tab Manager from which to populate tab spawner menu items.
-	 */
-	static void FillMenu(FMenuBuilder& MenuBuilder, const TSharedPtr<FTabManager> TabManager);
-
-	/** Callback for determining the visibility of the 'Select a session' overlay. */
-	EVisibility IsSessionOverlayVisible() const;
-
-	/** Callback for getting the enabled state of the profiler window. */
-	bool IsProfilerEnabled() const;
-
-	/** Updates the amount of time the profiler has been active. */
-	EActiveTimerReturnType UpdateActiveDuration(double InCurrentTime, float InDeltaTime);
-
-	/**
-	 * Ticks this widget. Override in derived classes, but always call the parent implementation.
-	 *
-	 * @param  AllottedGeometry The space allotted for this widget
-	 * @param  InCurrentTime  Current absolute real time
-	 * @param  InDeltaTime  Real time passed since last tick
-	 */
-	//virtual void Tick(const FGeometry& AllottedGeometry, const double InCurrentTime, const float InDeltaTime) override;
-
-	/**
-	 * The system will use this event to notify a widget that the cursor has entered it. This event is NOT bubbled.
-	 *
-	 * @param MyGeometry The Geometry of the widget receiving the event
-	 * @param MouseEvent Information about the input event
-	 */
-	virtual void OnMouseEnter(const FGeometry& MyGeometry, const FPointerEvent& MouseEvent) override;
-
-	/**
-	 * The system will use this event to notify a widget that the cursor has left it. This event is NOT bubbled.
-	 *
-	 * @param MouseEvent Information about the input event
-	 */
-	virtual void OnMouseLeave(const FPointerEvent& MouseEvent) override;
-
-	/**
-	 * Called after a key is pressed when this widget has focus
-	 *
-	 * @param MyGeometry The Geometry of the widget receiving the event
-	 * @param  InKeyEvent  Key event
-	 *
-	 * @return  Returns whether the event was handled, along with other possible actions
-	 */
-	virtual FReply OnKeyDown(const FGeometry& MyGeometry, const FKeyEvent& InKeyEvent) override;
-
-	/**
-	 * Called when the user is dropping something onto a widget; terminates drag and drop.
-	 *
-	 * @param MyGeometry      The geometry of the widget receiving the event.
-	 * @param DragDropEvent   The drag and drop event.
-	 *
-	 * @return A reply that indicated whether this event was handled.
-	 */
-	virtual FReply OnDrop(const FGeometry& MyGeometry, const FDragDropEvent& DragDropEvent) override;
-
-	/**
-	 * Called during drag and drop when the the mouse is being dragged over a widget.
-	 *
-	 * @param MyGeometry      The geometry of the widget receiving the event.
-	 * @param DragDropEvent   The drag and drop event.
-	 *
-	 * @return A reply that indicated whether this event was handled.
-	 */
-	virtual FReply OnDragOver(const FGeometry& MyGeometry, const FDragDropEvent& DragDropEvent)  override;
-
 	void OnTimeSelectionChanged(Insights::ETimeChangedFlags InFlags, double InStartTime, double InEndTime);
 
 private:
-	/** Holds the tab manager that manages the front-end's tabs. */
-	TSharedPtr<FTabManager> TabManager;
-
-	/** All tabs owned by this window */
-	TSet<TSharedPtr<SDockTab>> OpenTabs;
-
 	/** The Timing view (multi-track) widget */
 	TSharedPtr<STimingView> TimingView;
 
@@ -197,12 +111,6 @@ private:
 	/** The Requests tree view widget */
 	TSharedPtr<Insights::SUntypedTableTreeView> RequestsTreeView;
 
-	/** The handle to the active update duration tick */
-	TWeakPtr<FActiveTimerHandle> ActiveTimerHandle;
-
-	/** The number of seconds the profiler has been active */
-	float DurationActive;
-
 	double SelectionStartTime;
 	double SelectionEndTime;
 };
diff --git a/Engine/Source/Developer/TraceInsights/Private/Insights/MemoryProfiler/MemoryProfilerManager.cpp b/Engine/Source/Developer/TraceInsights/Private/Insights/MemoryProfiler/MemoryProfilerManager.cpp
--- a/Engine/Source/Developer/TraceInsights/Private/Insights/MemoryProfiler/MemoryProfilerManager.cpp
+++ b/Engine/Source/Developer/TraceInsights/Private/Insights/MemoryProfiler/MemoryProfilerManager.cpp
@@ -7,6 +7,7 @@
 #include "TraceServices/AnalysisService.h"
 #include "TraceServices/Model/AllocationsProvider.h"
 #include "TraceServices/Model/Memory.h"
+#include "Widgets/Docking/SDockTab.h"
 #include "WorkspaceMenuStructure.h"
 #include "WorkspaceMenuStructureModule.h"
 
diff --git a/Engine/Source/Developer/TraceInsights/Private/Insights/MemoryProfiler/Widgets/SMemoryProfilerToolbar.cpp b/Engine/Source/Developer/TraceInsights/Private/Insights/MemoryProfiler/Widgets/SMemoryProfilerToolbar.cpp
--- a/Engine/Source/Developer/TraceInsights/Private/Insights/MemoryProfiler/Widgets/SMemoryProfilerToolbar.cpp
+++ b/Engine/Source/Developer/TraceInsights/Private/Insights/MemoryProfiler/Widgets/SMemoryProfilerToolbar.cpp
@@ -2,6 +2,7 @@
 
 #include "SMemoryProfilerToolbar.h"
 
+#include "Framework/MultiBox/MultiBoxExtender.h"
 #include "Framework/MultiBox/MultiBoxBuilder.h"
 #include "Widgets/SBoxPanel.h"
 
@@ -33,7 +34,7 @@ void SMemoryProfilerToolbar::Construct(const FArguments& InArgs)
 {
 	struct Local
 	{
-		static void FillViewToolbar(FToolBarBuilder& ToolbarBuilder)
+		static void FillViewToolbar(FToolBarBuilder& ToolbarBuilder, const FArguments &InArgs)
 		{
 			ToolbarBuilder.BeginSection("View");
 			{
@@ -51,9 +52,14 @@ void SMemoryProfilerToolbar::Construct(const FArguments& InArgs)
 					FSlateIcon(FInsightsStyle::GetStyleSetName(), "Icons.ModulesView.ToolBar"));
 			}
 			ToolbarBuilder.EndSection();
+
+			if (InArgs._ToolbarExtender.IsValid())
+			{
+				InArgs._ToolbarExtender->Apply("MainToolbar", EExtensionHook::First, ToolbarBuilder);
+			}
 		}
 
-		static void FillRightSideToolbar(FToolBarBuilder& ToolbarBuilder)
+		static void FillRightSideToolbar(FToolBarBuilder& ToolbarBuilder, const FArguments &InArgs)
 		{
 			ToolbarBuilder.BeginSection("Debug");
 			{
@@ -62,6 +68,11 @@ void SMemoryProfilerToolbar::Construct(const FArguments& InArgs)
 					FSlateIcon(FInsightsStyle::GetStyleSetName(), "Icons.Debug.ToolBar"));
 			}
 			ToolbarBuilder.EndSection();
+
+			if (InArgs._ToolbarExtender.IsValid())
+			{
+				InArgs._ToolbarExtender->Apply("RightSideToolbar", EExtensionHook::First, ToolbarBuilder);
+			}
 		}
 	};
 
@@ -69,11 +80,11 @@ void SMemoryProfilerToolbar::Construct(const FArguments& InArgs)
 
 	FSlimHorizontalToolBarBuilder ToolbarBuilder(CommandList.ToSharedRef(), FMultiBoxCustomization::None);
 	ToolbarBuilder.SetStyle(&FInsightsStyle::Get(), "PrimaryToolbar");
-	Local::FillViewToolbar(ToolbarBuilder);
+	Local::FillViewToolbar(ToolbarBuilder, InArgs);
 
 	FSlimHorizontalToolBarBuilder RightSideToolbarBuilder(CommandList.ToSharedRef(), FMultiBoxCustomization::None);
 	RightSideToolbarBuilder.SetStyle(&FInsightsStyle::Get(), "PrimaryToolbar");
-	Local::FillRightSideToolbar(RightSideToolbarBuilder);
+	Local::FillRightSideToolbar(RightSideToolbarBuilder, InArgs);
 
 	ChildSlot
 	[
diff --git a/Engine/Source/Developer/TraceInsights/Private/Insights/MemoryProfiler/Widgets/SMemoryProfilerToolbar.h b/Engine/Source/Developer/TraceInsights/Private/Insights/MemoryProfiler/Widgets/SMemoryProfilerToolbar.h
--- a/Engine/Source/Developer/TraceInsights/Private/Insights/MemoryProfiler/Widgets/SMemoryProfilerToolbar.h
+++ b/Engine/Source/Developer/TraceInsights/Private/Insights/MemoryProfiler/Widgets/SMemoryProfilerToolbar.h
@@ -5,6 +5,8 @@
 #include "Widgets/DeclarativeSyntaxSupport.h"
 #include "Widgets/SCompoundWidget.h"
 
+class FExtender;
+
 /** Ribbon based toolbar used as a main menu in the Profiler window. */
 class SMemoryProfilerToolbar : public SCompoundWidget
 {
@@ -16,6 +18,7 @@ public:
 	virtual ~SMemoryProfilerToolbar();
 
 	SLATE_BEGIN_ARGS(SMemoryProfilerToolbar) {}
+	SLATE_ARGUMENT(TSharedPtr<FExtender>, ToolbarExtender);
 	SLATE_END_ARGS()
 
 	/**
@@ -24,8 +27,4 @@ public:
 	 * @param InArgs The declaration data for this widget
 	 */
 	void Construct(const FArguments& InArgs);
-
-private:
-	/** Create the UI commands for the toolbar */
-	void CreateCommands();
 };
diff --git a/Engine/Source/Developer/TraceInsights/Private/Insights/MemoryProfiler/Widgets/SMemoryProfilerWindow.cpp b/Engine/Source/Developer/TraceInsights/Private/Insights/MemoryProfiler/Widgets/SMemoryProfilerWindow.cpp
--- a/Engine/Source/Developer/TraceInsights/Private/Insights/MemoryProfiler/Widgets/SMemoryProfilerWindow.cpp
+++ b/Engine/Source/Developer/TraceInsights/Private/Insights/MemoryProfiler/Widgets/SMemoryProfilerWindow.cpp
@@ -4,28 +4,12 @@
 
 #include "Features/IModularFeatures.h"
 #include "Framework/Commands/UICommandList.h"
-#include "Framework/Docking/LayoutService.h"
 #include "Framework/Docking/TabManager.h"
 #include "Framework/Docking/WorkspaceItem.h"
-#include "Framework/MultiBox/MultiBoxBuilder.h"
 #include "SlateOptMacros.h"
-#include "Styling/AppStyle.h"
 #include "Widgets/Docking/SDockTab.h"
-#include "Widgets/Images/SImage.h"
-#include "Widgets/Layout/SBorder.h"
-#include "Widgets/Layout/SBox.h"
-#include "Widgets/SBoxPanel.h"
-#include "Widgets/Text/STextBlock.h"
-
-#if WITH_EDITOR
-	#include "EngineAnalytics.h"
-	#include "Runtime/Analytics/Analytics/Public/AnalyticsEventAttribute.h"
-	#include "Runtime/Analytics/Analytics/Public/Interfaces/IAnalyticsProvider.h"
-#endif // WITH_EDITOR
 
 // Insights
-#include "Insights/Common/InsightsMenuBuilder.h"
-#include "Insights/InsightsManager.h"
 #include "Insights/InsightsStyle.h"
 #include "Insights/MemoryProfiler/MemoryProfilerManager.h"
 #include "Insights/MemoryProfiler/ViewModels/MemorySharedState.h"
@@ -33,10 +17,7 @@
 #include "Insights/MemoryProfiler/Widgets/SMemInvestigationView.h"
 #include "Insights/MemoryProfiler/Widgets/SMemoryProfilerToolbar.h"
 #include "Insights/MemoryProfiler/Widgets/SMemTagTreeView.h"
-#include "Insights/TraceInsightsModule.h"
-#include "Insights/Version.h"
 #include "Insights/ViewModels/TimeRulerTrack.h"
-#include "Insights/Widgets/SInsightsSettings.h"
 #include "Insights/Widgets/STimingView.h"
 
 ////////////////////////////////////////////////////////////////////////////////////////////////////
@@ -54,8 +35,8 @@ const FName FMemoryProfilerTabs::ModulesViewID(TEXT("Modules"));
 ////////////////////////////////////////////////////////////////////////////////////////////////////
 
 SMemoryProfilerWindow::SMemoryProfilerWindow()
-	: SharedState(MakeShared<FMemorySharedState>())
-	, DurationActive(0.0f)
+	: SMajorTabWindow(FInsightsManagerTabs::MemoryProfilerTabId)
+	, SharedState(MakeShared<FMemorySharedState>())
 {
 	CreateTimingViewMarkers();
 }
@@ -64,31 +45,19 @@ SMemoryProfilerWindow::SMemoryProfilerWindow()
 
 SMemoryProfilerWindow::~SMemoryProfilerWindow()
 {
-	// Close all tabs.
-	TArray<TSharedPtr<SDockTab>> LocalOpenTabs;
-	for (TSharedPtr<SDockTab>& Tab : OpenTabs)
-	{
-		LocalOpenTabs.Add(Tab);
-	}
-	for (TSharedPtr<SDockTab>& Tab : LocalOpenTabs)
-	{
-		Tab->RequestCloseTab();
-		Tab->RemoveTabFromParent();
-	}
-	LocalOpenTabs.Reset();
-	check(OpenTabs.IsEmpty());
+	CloseAllOpenTabs();
 
 	check(ModulesView == nullptr);
 	check(MemTagTreeView == nullptr);
 	check(MemInvestigationView == nullptr);
 	check(TimingView == nullptr);
+}
 
-#if WITH_EDITOR
-	if (DurationActive > 0.0f && FEngineAnalytics::IsAvailable())
-	{
-		FEngineAnalytics::GetProvider().RecordEvent(TEXT("Insights.Usage.MemoryProfiler"), FAnalyticsEventAttribute(TEXT("Duration"), DurationActive));
-	}
-#endif // WITH_EDITOR
+////////////////////////////////////////////////////////////////////////////////////////////////////
+
+const TCHAR* SMemoryProfilerWindow::GetAnalyticsEventName() const
+{
+	return TEXT("Insights.Usage.MemoryProfiler");
 }
 
 ////////////////////////////////////////////////////////////////////////////////////////////////////
@@ -142,7 +111,7 @@ TSharedRef<SDockTab> SMemoryProfilerWindow::SpawnTab_TimingView(const FSpawnTabA
 	TimingView->HideAllDefaultTracks();
 
 	DockTab->SetOnTabClosed(SDockTab::FOnTabClosedCallback::CreateRaw(this, &SMemoryProfilerWindow::OnTimingViewTabClosed));
-	OpenTabs.Add(DockTab);
+	AddOpenTab(DockTab);
 
 	return DockTab;
 }
@@ -164,7 +133,7 @@ void SMemoryProfilerWindow::OnTimingViewTabClosed(TSharedRef<SDockTab> TabBeingC
 	IModularFeatures::Get().UnregisterModularFeature(Insights::TimingViewExtenderFeatureName, &SharedState.Get());
 	SharedState->SetTimingView(nullptr);
 
-	OpenTabs.Remove(TabBeingClosed);
+	RemoveOpenTab(TabBeingClosed);
 }
 
 ////////////////////////////////////////////////////////////////////////////////////////////////////
@@ -182,7 +151,7 @@ TSharedRef<SDockTab> SMemoryProfilerWindow::SpawnTab_MemInvestigationView(const
 		];
 
 	DockTab->SetOnTabClosed(SDockTab::FOnTabClosedCallback::CreateRaw(this, &SMemoryProfilerWindow::OnMemInvestigationViewTabClosed));
-	OpenTabs.Add(DockTab);
+	AddOpenTab(DockTab);
 
 	return DockTab;
 }
@@ -195,7 +164,7 @@ void SMemoryProfilerWindow::OnMemInvestigationViewTabClosed(TSharedRef<SDockTab>
 	FMemoryProfilerManager::Get()->SetMemInvestigationViewVisible(false);
 	MemInvestigationView = nullptr;
 
-	OpenTabs.Remove(TabBeingClosed);
+	RemoveOpenTab(TabBeingClosed);
 }
 
 ////////////////////////////////////////////////////////////////////////////////////////////////////
@@ -213,7 +182,7 @@ TSharedRef<SDockTab> SMemoryProfilerWindow::SpawnTab_MemTagTreeView(const FSpawn
 		];
 
 	DockTab->SetOnTabClosed(SDockTab::FOnTabClosedCallback::CreateRaw(this, &SMemoryProfilerWindow::OnMemTagTreeViewTabClosed));
-	OpenTabs.Add(DockTab);
+	AddOpenTab(DockTab);
 
 	return DockTab;
 }
@@ -226,7 +195,7 @@ void SMemoryProfilerWindow::OnMemTagTreeViewTabClosed(TSharedRef<SDockTab> TabBe
 	FMemoryProfilerManager::Get()->SetMemTagTreeViewVisible(false);
 	MemTagTreeView = nullptr;
 
-	OpenTabs.Remove(TabBeingClosed);
+	RemoveOpenTab(TabBeingClosed);
 }
 
 ////////////////////////////////////////////////////////////////////////////////////////////////////
@@ -254,7 +223,7 @@ TSharedRef<SDockTab> SMemoryProfilerWindow::SpawnTab_MemAllocTableTreeView(const
 	MemAllocTableTreeViews.Add(MemAllocTableTreeView);
 
 	DockTab->SetOnTabClosed(SDockTab::FOnTabClosedCallback::CreateRaw(this, &SMemoryProfilerWindow::OnMemAllocTableTreeViewTabClosed));
-	OpenTabs.Add(DockTab);
+	AddOpenTab(DockTab);
 
 	return DockTab;
 }
@@ -295,12 +264,12 @@ void SMemoryProfilerWindow::OnMemAllocTableTreeViewTabClosed(TSharedRef<SDockTab
 		}
 	}
 
-	TabManager->UnregisterTabSpawner(ClosingTabId);
+	GetTabManager()->UnregisterTabSpawner(ClosingTabId);
 
 	MemAllocTableTreeView->OnClose();
 	MemAllocTableTreeViews.Remove(MemAllocTableTreeView);
 
-	OpenTabs.Remove(TabBeingClosed);
+	RemoveOpenTab(TabBeingClosed);
 }
 
 ////////////////////////////////////////////////////////////////////////////////////////////////////
@@ -325,17 +294,26 @@ void SMemoryProfilerWindow::CloseMemAllocTableTreeTabs()
 BEGIN_SLATE_FUNCTION_BUILD_OPTIMIZATION
 TSharedPtr<Insights::SMemAllocTableTreeView> SMemoryProfilerWindow::ShowMemAllocTableTreeViewTab()
 {
+	if (!GetTabManager().IsValid())
+	{
+		return nullptr;
+	}
+	FTabManager* TabManagerPtr = GetTabManager().Get();
+
 	if (SharedState->GetCurrentQueryTarget()->GetName() == Insights::FQueryTargetWindowSpec::NewWindow)
 	{
 		++LastMemAllocTableTreeViewIndex;
 		FName TabId = FMemoryProfilerTabs::MemAllocTableTreeViewID;
 		TabId.SetNumber(LastMemAllocTableTreeViewIndex);
 
+		check(GetWorkspaceMenuGroup().IsValid());
+		const TSharedRef<FWorkspaceItem> Group = GetWorkspaceMenuGroup().ToSharedRef();
+
 		FText MemAllocTableTreeViewTabDisplayName = FText::Format(LOCTEXT("MemoryProfiler.MemAllocTableTreeViewTabTitle", "Allocs Table {0}"), FText::AsNumber(LastMemAllocTableTreeViewIndex));
-		TabManager->RegisterTabSpawner(TabId, FOnSpawnTab::CreateRaw(this, &SMemoryProfilerWindow::SpawnTab_MemAllocTableTreeView, LastMemAllocTableTreeViewIndex))
+		TabManagerPtr->RegisterTabSpawner(TabId, FOnSpawnTab::CreateRaw(this, &SMemoryProfilerWindow::SpawnTab_MemAllocTableTreeView, LastMemAllocTableTreeViewIndex))
 			.SetDisplayName(MemAllocTableTreeViewTabDisplayName)
 			.SetIcon(FSlateIcon(FInsightsStyle::GetStyleSetName(), "Icons.MemAllocTableTreeView"))
-			.SetGroup(AppMenuGroup.ToSharedRef());
+			.SetGroup(Group);
 
 		TSharedPtr<Insights::FQueryTargetWindowSpec> NewTarget = MakeShared<Insights::FQueryTargetWindowSpec>(TabId, MemAllocTableTreeViewTabDisplayName);
 		SharedState->AddQueryTarget(NewTarget);
@@ -344,9 +322,9 @@ TSharedPtr<Insights::SMemAllocTableTreeView> SMemoryProfilerWindow::ShowMemAlloc
 	}
 
 	FName TabId = SharedState->GetCurrentQueryTarget()->GetName();
-	if (TabManager->HasTabSpawner(TabId))
+	if (TabManagerPtr->HasTabSpawner(TabId))
 	{
-		TSharedPtr<SDockTab> Tab = TabManager->TryInvokeTab(TabId);
+		TSharedPtr<SDockTab> Tab = TabManagerPtr->TryInvokeTab(TabId);
 		if (Tab)
 		{
 			TSharedRef<Insights::SMemAllocTableTreeView> MemAllocTableTreeView = StaticCastSharedRef<Insights::SMemAllocTableTreeView>(Tab->GetContent());
@@ -378,7 +356,7 @@ TSharedRef<SDockTab> SMemoryProfilerWindow::SpawnTab_ModulesView(const FSpawnTab
 		];
 
 	DockTab->SetOnTabClosed(SDockTab::FOnTabClosedCallback::CreateRaw(this, &SMemoryProfilerWindow::OnModulesViewClosed));
-	OpenTabs.Add(DockTab);
+	AddOpenTab(DockTab);
 
 	return DockTab;
 }
@@ -391,49 +369,65 @@ void SMemoryProfilerWindow::OnModulesViewClosed(TSharedRef<SDockTab> TabBeingClo
 	FMemoryProfilerManager::Get()->SetModulesViewVisible(false);
 	ModulesView = nullptr;
 
-	OpenTabs.Remove(TabBeingClosed);
+	RemoveOpenTab(TabBeingClosed);
 }
 
 ////////////////////////////////////////////////////////////////////////////////////////////////////
 
-BEGIN_SLATE_FUNCTION_BUILD_OPTIMIZATION
 void SMemoryProfilerWindow::Construct(const FArguments& InArgs, const TSharedRef<SDockTab>& ConstructUnderMajorTab, const TSharedPtr<SWindow>& ConstructUnderWindow)
 {
-	// Create & initialize tab manager.
-	TabManager = FGlobalTabmanager::Get()->NewTabManager(ConstructUnderMajorTab);
-	const auto& PersistLayout = [](const TSharedRef<FTabManager::FLayout>& LayoutToSave)
-	{
-		FLayoutSaveRestore::SaveToConfig(FTraceInsightsModule::GetUnrealInsightsLayoutIni(), LayoutToSave);
-	};
-	TabManager->SetOnPersistLayout(FTabManager::FOnPersistLayout::CreateLambda(PersistLayout));
+	TSharedPtr<FMemoryProfilerManager> MemoryProfilerManager = FMemoryProfilerManager::Get();
+	ensure(MemoryProfilerManager.IsValid());
+
+	SetCommandList(MemoryProfilerManager->GetCommandList());
+
+	SMajorTabWindow::FArguments Args;
+	SMajorTabWindow::Construct(Args, ConstructUnderMajorTab, ConstructUnderWindow);
+}
+
+////////////////////////////////////////////////////////////////////////////////////////////////////
+
+TSharedRef<FWorkspaceItem> SMemoryProfilerWindow::CreateWorkspaceMenuGroup()
+{
+	return GetTabManager()->AddLocalWorkspaceMenuCategory(LOCTEXT("MemoryProfilerMenuGroupName", "Memory Insights"));
+}
 
-	AppMenuGroup = TabManager->AddLocalWorkspaceMenuCategory(LOCTEXT("MemoryProfilerMenuGroupName", "Memory Insights"));
+////////////////////////////////////////////////////////////////////////////////////////////////////
+
+void SMemoryProfilerWindow::RegisterTabSpawners()
+{
+	check(GetTabManager().IsValid());
+	FTabManager* TabManagerPtr = GetTabManager().Get();
+	check(GetWorkspaceMenuGroup().IsValid());
+	const TSharedRef<FWorkspaceItem> Group = GetWorkspaceMenuGroup().ToSharedRef();
 
-	TabManager->RegisterTabSpawner(FMemoryProfilerTabs::TimingViewID, FOnSpawnTab::CreateRaw(this, &SMemoryProfilerWindow::SpawnTab_TimingView))
+	TabManagerPtr->RegisterTabSpawner(FMemoryProfilerTabs::TimingViewID, FOnSpawnTab::CreateRaw(this, &SMemoryProfilerWindow::SpawnTab_TimingView))
 		.SetDisplayName(LOCTEXT("TimingViewTabTitle", "Timing View"))
 		.SetIcon(FSlateIcon(FInsightsStyle::GetStyleSetName(), "Icons.TimingView"))
-		.SetGroup(AppMenuGroup.ToSharedRef());
+		.SetGroup(Group);
 
-	TabManager->RegisterTabSpawner(FMemoryProfilerTabs::MemInvestigationViewID, FOnSpawnTab::CreateRaw(this, &SMemoryProfilerWindow::SpawnTab_MemInvestigationView))
+	TabManagerPtr->RegisterTabSpawner(FMemoryProfilerTabs::MemInvestigationViewID, FOnSpawnTab::CreateRaw(this, &SMemoryProfilerWindow::SpawnTab_MemInvestigationView))
 		.SetDisplayName(LOCTEXT("MemInvestigationViewTabTitle", "Investigation"))
 		.SetIcon(FSlateIcon(FInsightsStyle::GetStyleSetName(), "Icons.MemInvestigationView"))
-		.SetGroup(AppMenuGroup.ToSharedRef());
+		.SetGroup(Group);
 
-	TabManager->RegisterTabSpawner(FMemoryProfilerTabs::MemTagTreeViewID, FOnSpawnTab::CreateRaw(this, &SMemoryProfilerWindow::SpawnTab_MemTagTreeView))
+	TabManagerPtr->RegisterTabSpawner(FMemoryProfilerTabs::MemTagTreeViewID, FOnSpawnTab::CreateRaw(this, &SMemoryProfilerWindow::SpawnTab_MemTagTreeView))
 		.SetDisplayName(LOCTEXT("MemTagTreeViewTabTitle", "LLM Tags"))
 		.SetIcon(FSlateIcon(FInsightsStyle::GetStyleSetName(), "Icons.MemTagTreeView"))
-		.SetGroup(AppMenuGroup.ToSharedRef());
+		.SetGroup(Group);
 
-	TabManager->RegisterTabSpawner(FMemoryProfilerTabs::ModulesViewID, FOnSpawnTab::CreateRaw(this, &SMemoryProfilerWindow::SpawnTab_ModulesView))
+	TabManagerPtr->RegisterTabSpawner(FMemoryProfilerTabs::ModulesViewID, FOnSpawnTab::CreateRaw(this, &SMemoryProfilerWindow::SpawnTab_ModulesView))
 		.SetDisplayName(LOCTEXT("ModulesViewTabTitle", "Modules"))
 		.SetIcon(FSlateIcon(FInsightsStyle::GetStyleSetName(), "Icons.ModulesView"))
-		.SetGroup(AppMenuGroup.ToSharedRef());
+		.SetGroup(Group);
+}
 
-	TSharedPtr<FMemoryProfilerManager> MemoryProfilerManager = FMemoryProfilerManager::Get();
-	ensure(MemoryProfilerManager.IsValid());
+////////////////////////////////////////////////////////////////////////////////////////////////////
 
-	// Create tab layout.
-	TSharedRef<FTabManager::FLayout> Layout = FTabManager::NewLayout("InsightsMemoryProfilerLayout_v1.2")
+BEGIN_SLATE_FUNCTION_BUILD_OPTIMIZATION
+TSharedRef<FTabManager::FLayout> SMemoryProfilerWindow::CreateDefaultTabLayout() const
+{
+	return FTabManager::NewLayout("InsightsMemoryProfilerLayout_v1.2")
 		->AddArea
 		(
 			FTabManager::NewPrimaryArea()
@@ -455,215 +449,14 @@ void SMemoryProfilerWindow::Construct(const FArguments& InArgs, const TSharedRef
 				->SetForegroundTab(FMemoryProfilerTabs::MemInvestigationViewID)
 			)
 		);
-
-	Layout = FLayoutSaveRestore::LoadFromConfig(FTraceInsightsModule::GetUnrealInsightsLayoutIni(), Layout);
-
-	// Create & initialize main menu.
-	FMenuBarBuilder MenuBarBuilder = FMenuBarBuilder(TSharedPtr<FUICommandList>());
-	MenuBarBuilder.AddPullDownMenu(
-		LOCTEXT("MenuLabel", "Menu"),
-		FText::GetEmpty(),
-		FNewMenuDelegate::CreateStatic(&SMemoryProfilerWindow::FillMenu, TabManager),
-		FName(TEXT("Menu"))
-	);
-
-#if !WITH_EDITOR
-	TSharedRef<SWidget> MenuWidget = MenuBarBuilder.MakeWidget();
-	MenuWidget->SetClipping(EWidgetClipping::ClipToBoundsWithoutIntersecting);
-#endif
-
-	ChildSlot
-	[
-		SNew(SOverlay)
-
-#if !WITH_EDITOR
-		// Menu
-		+ SOverlay::Slot()
-		.HAlign(HAlign_Left)
-		.VAlign(VAlign_Top)
-		.Padding(34.0f, -60.0f, 0.0f, 0.0f)
-		[
-			MenuWidget
-		]
-#endif
-
-		// Version
-		+ SOverlay::Slot()
-		.HAlign(HAlign_Right)
-		.VAlign(VAlign_Top)
-		.Padding(0.0f, -16.0f, 4.0f, 0.0f)
-		[
-			SNew(STextBlock)
-			.Clipping(EWidgetClipping::ClipToBoundsWithoutIntersecting)
-			.Text(LOCTEXT("UnrealInsightsVersion", UNREAL_INSIGHTS_VERSION_STRING_EX))
-			.ColorAndOpacity(FLinearColor(0.15f, 0.15f, 0.15f, 1.0f))
-		]
-
-		// Overlay slot for the main window area
-		+ SOverlay::Slot()
-		.HAlign(HAlign_Fill)
-		.VAlign(VAlign_Fill)
-		[
-			SNew(SVerticalBox)
-
-			+ SVerticalBox::Slot()
-			.AutoHeight()
-			.Padding(FMargin(0.0f, 0.0f, 0.0f, 5.0f))
-			[
-				SNew(SMemoryProfilerToolbar)
-			]
-
-			+ SVerticalBox::Slot()
-			.FillHeight(1.0f)
-			[
-				TabManager->RestoreFrom(Layout, ConstructUnderWindow).ToSharedRef()
-			]
-		]
-
-		// Session hint overlay
-		+ SOverlay::Slot()
-		.HAlign(HAlign_Center)
-		.VAlign(VAlign_Center)
-		[
-			SNew(SBorder)
-			.Visibility(this, &SMemoryProfilerWindow::IsSessionOverlayVisible)
-			.BorderImage(FAppStyle::Get().GetBrush("PopupText.Background"))
-			.Padding(8.0f)
-			[
-				SNew(STextBlock)
-				.Text(LOCTEXT("SelectTraceOverlayText", "Please select a trace."))
-			]
-		]
-	];
-
-#if !WITH_EDITOR
-	// Tell tab-manager about the global menu bar.
-	TabManager->SetMenuMultiBox(MenuBarBuilder.GetMultiBox(), MenuWidget);
-#endif
 }
 END_SLATE_FUNCTION_BUILD_OPTIMIZATION
 
 ////////////////////////////////////////////////////////////////////////////////////////////////////
 
-void SMemoryProfilerWindow::FillMenu(FMenuBuilder& MenuBuilder, const TSharedPtr<FTabManager> TabManager)
+TSharedRef<SWidget> SMemoryProfilerWindow::CreateToolbar(TSharedPtr<FExtender> Extender)
 {
-	if (!TabManager.IsValid())
-	{
-		return;
-	}
-
-	FInsightsManager::Get()->GetInsightsMenuBuilder()->PopulateMenu(MenuBuilder);
-
-	TabManager->PopulateLocalTabSpawnerMenu(MenuBuilder);
-}
-
-////////////////////////////////////////////////////////////////////////////////////////////////////
-
-void SMemoryProfilerWindow::ShowTab(const FName& TabID)
-{
-	if (TabManager->HasTabSpawner(TabID))
-	{
-		TabManager->TryInvokeTab(TabID);
-	}
-}
-
-////////////////////////////////////////////////////////////////////////////////////////////////////
-
-void SMemoryProfilerWindow::HideTab(const FName& TabID)
-{
-	TSharedPtr<SDockTab> Tab = TabManager->FindExistingLiveTab(TabID);
-	if (Tab.IsValid())
-	{
-		Tab->RequestCloseTab();
-		Tab->RemoveTabFromParent();
-	}
-}
-
-////////////////////////////////////////////////////////////////////////////////////////////////////
-
-EVisibility SMemoryProfilerWindow::IsSessionOverlayVisible() const
-{
-	if (FInsightsManager::Get()->GetSession().IsValid())
-	{
-		return EVisibility::Hidden;
-	}
-	else
-	{
-		return EVisibility::Visible;
-	}
-}
-
-////////////////////////////////////////////////////////////////////////////////////////////////////
-
-bool SMemoryProfilerWindow::IsProfilerEnabled() const
-{
-	return FInsightsManager::Get()->GetSession().IsValid();
-}
-
-////////////////////////////////////////////////////////////////////////////////////////////////////
-
-EActiveTimerReturnType SMemoryProfilerWindow::UpdateActiveDuration(double InCurrentTime, float InDeltaTime)
-{
-	DurationActive += InDeltaTime;
-
-	// The profiler window will explicitly unregister this active timer when the mouse leaves.
-	return EActiveTimerReturnType::Continue;
-}
-
-////////////////////////////////////////////////////////////////////////////////////////////////////
-
-void SMemoryProfilerWindow::OnMouseEnter(const FGeometry& MyGeometry, const FPointerEvent& MouseEvent)
-{
-	SCompoundWidget::OnMouseEnter(MyGeometry, MouseEvent);
-
-	if (!ActiveTimerHandle.IsValid())
-	{
-		ActiveTimerHandle = RegisterActiveTimer(0.f, FWidgetActiveTimerDelegate::CreateSP(this, &SMemoryProfilerWindow::UpdateActiveDuration));
-	}
-}
-
-////////////////////////////////////////////////////////////////////////////////////////////////////
-
-void SMemoryProfilerWindow::OnMouseLeave(const FPointerEvent& MouseEvent)
-{
-	SCompoundWidget::OnMouseLeave(MouseEvent);
-
-	auto PinnedActiveTimerHandle = ActiveTimerHandle.Pin();
-	if (PinnedActiveTimerHandle.IsValid())
-	{
-		UnRegisterActiveTimer(PinnedActiveTimerHandle.ToSharedRef());
-	}
-}
-
-////////////////////////////////////////////////////////////////////////////////////////////////////
-
-FReply SMemoryProfilerWindow::OnKeyDown(const FGeometry& MyGeometry, const FKeyEvent& InKeyEvent)
-{
-	return FMemoryProfilerManager::Get()->GetCommandList()->ProcessCommandBindings(InKeyEvent) ? FReply::Handled() : FReply::Unhandled();
-}
-
-////////////////////////////////////////////////////////////////////////////////////////////////////
-
-FReply SMemoryProfilerWindow::OnDragOver(const FGeometry& MyGeometry, const FDragDropEvent& DragDropEvent)
-{
-	if (FInsightsManager::Get()->OnDragOver(DragDropEvent))
-	{
-		return FReply::Handled();
-	}
-
-	return SCompoundWidget::OnDragOver(MyGeometry, DragDropEvent);
-}
-
-////////////////////////////////////////////////////////////////////////////////////////////////////
-
-FReply SMemoryProfilerWindow::OnDrop(const FGeometry& MyGeometry, const FDragDropEvent& DragDropEvent)
-{
-	if (FInsightsManager::Get()->OnDrop(DragDropEvent))
-	{
-		return FReply::Handled();
-	}
-
-	return SCompoundWidget::OnDrop(MyGeometry, DragDropEvent);
+	return SNew(SMemoryProfilerToolbar).ToolbarExtender(Extender);
 }
 
 ////////////////////////////////////////////////////////////////////////////////////////////////////
diff --git a/Engine/Source/Developer/TraceInsights/Private/Insights/MemoryProfiler/Widgets/SMemoryProfilerWindow.h b/Engine/Source/Developer/TraceInsights/Private/Insights/MemoryProfiler/Widgets/SMemoryProfilerWindow.h
--- a/Engine/Source/Developer/TraceInsights/Private/Insights/MemoryProfiler/Widgets/SMemoryProfilerWindow.h
+++ b/Engine/Source/Developer/TraceInsights/Private/Insights/MemoryProfiler/Widgets/SMemoryProfilerWindow.h
@@ -3,26 +3,14 @@
 #pragma once
 
 #include "CoreMinimal.h"
-#include "Framework/Docking/TabManager.h"
-#include "Input/Reply.h"
-#include "Layout/Visibility.h"
-#include "Misc/Guid.h"
-#include "SlateFwd.h"
-#include "Widgets/DeclarativeSyntaxSupport.h"
-#include "Widgets/Docking/SDockTab.h"
-#include "Widgets/SCompoundWidget.h"
 
 // Insights
-#include "Insights/InsightsManager.h"
-#include "Insights/ITimingViewSession.h"
+#include "Insights/ITimingViewSession.h" // for Insights::ETimeChangedFlags
+#include "Insights/Widgets/SMajorTabWindow.h"
 #include "Insights/Widgets/SModulesView.h"
 
 ////////////////////////////////////////////////////////////////////////////////////////////////////
 
-class FActiveTimerHandle;
-class FMenuBuilder;
-class FWorkspaceItem;
-
 class FMemorySharedState;
 class SMemInvestigationView;
 class SMemTagTreeView;
@@ -49,8 +37,8 @@ struct FMemoryProfilerTabs
 
 ////////////////////////////////////////////////////////////////////////////////////////////////////
 
-/** Implements the timing profiler window. */
-class SMemoryProfilerWindow : public SCompoundWidget
+/** Implements the Memory Insights window. */
+class SMemoryProfilerWindow : public Insights::SMajorTabWindow
 {
 public:
 	/** Default constructor. */
@@ -62,23 +50,18 @@ public:
 	SLATE_BEGIN_ARGS(SMemoryProfilerWindow) {}
 	SLATE_END_ARGS()
 
-	void Reset();
+	virtual void Reset() override;
 
 	/** Constructs this widget. */
 	void Construct(const FArguments& InArgs, const TSharedRef<SDockTab>& ConstructUnderMajorTab, const TSharedPtr<SWindow>& ConstructUnderWindow);
 
-	void ShowTab(const FName& TabID);
-	void HideTab(const FName& TabID);
-	void ShowHideTab(const FName& TabID, bool bShow) { bShow ? ShowTab(TabID) : HideTab(TabID); }
-
-	TSharedPtr<Insights::SMemAllocTableTreeView> ShowMemAllocTableTreeViewTab();
-
-	TSharedPtr<FTabManager> GetTabManager() const { return TabManager; }
-
 	TSharedPtr<STimingView> GetTimingView() const { return TimingView; }
 	TSharedPtr<SMemInvestigationView> GetMemInvestigationView() const { return MemInvestigationView; }
 	TSharedPtr<SMemTagTreeView> GetMemTagTreeView() const { return MemTagTreeView; }
 
+	void CloseMemAllocTableTreeTabs();
+	TSharedPtr<Insights::SMemAllocTableTreeView> ShowMemAllocTableTreeViewTab();
+
 	uint32 GetNumCustomTimeMarkers() const { return (uint32)CustomTimeMarkers.Num(); }
 	const TSharedRef<Insights::FTimeMarker>& GetCustomTimeMarker(uint32 Index) const { return CustomTimeMarkers[Index]; }
 	const TArray<TSharedRef<Insights::FTimeMarker>>& GetCustomTimeMarkers() const { return CustomTimeMarkers; }
@@ -86,11 +69,16 @@ public:
 	FMemorySharedState& GetSharedState() { return *SharedState; }
 	const FMemorySharedState& GetSharedState() const { return *SharedState; }
 
-	void CloseMemAllocTableTreeTabs();
-
 	void OnMemoryRuleChanged();
 	void OnTimeMarkerChanged(Insights::ETimeChangedFlags InFlags, TSharedRef<Insights::ITimeMarker> InTimeMarker);
 
+protected:
+	virtual const TCHAR* GetAnalyticsEventName() const override;
+	virtual TSharedRef<FWorkspaceItem> CreateWorkspaceMenuGroup() override;
+	virtual void RegisterTabSpawners() override;
+	virtual TSharedRef<FTabManager::FLayout> CreateDefaultTabLayout() const override;
+	virtual TSharedRef<SWidget> CreateToolbar(TSharedPtr<FExtender> Extender);
+
 private:
 	TSharedRef<SDockTab> SpawnTab_TimingView(const FSpawnTabArgs& Args);
 	void OnTimingViewTabClosed(TSharedRef<SDockTab> TabBeingClosed);
@@ -107,77 +95,6 @@ private:
 	TSharedRef<SDockTab> SpawnTab_ModulesView(const FSpawnTabArgs& Args);
 	void OnModulesViewClosed(TSharedRef<SDockTab> TabBeingClosed);
 
-	/**
-	 * Fill the main menu with menu items.
-	 *
-	 * @param MenuBuilder The multi-box builder that should be filled with content for this pull-down menu.
-	 * @param TabManager A Tab Manager from which to populate tab spawner menu items.
-	 */
-	static void FillMenu(FMenuBuilder& MenuBuilder, const TSharedPtr<FTabManager> TabManager);
-
-	/** Callback for determining the visibility of the 'Select a session' overlay. */
-	EVisibility IsSessionOverlayVisible() const;
-
-	/** Callback for getting the enabled state of the profiler window. */
-	bool IsProfilerEnabled() const;
-
-	/** Updates the amount of time the profiler has been active. */
-	EActiveTimerReturnType UpdateActiveDuration(double InCurrentTime, float InDeltaTime);
-
-	/**
-	 * Ticks this widget. Override in derived classes, but always call the parent implementation.
-	 *
-	 * @param  AllottedGeometry The space allotted for this widget
-	 * @param  InCurrentTime  Current absolute real time
-	 * @param  InDeltaTime  Real time passed since last tick
-	 */
-	//virtual void Tick(const FGeometry& AllottedGeometry, const double InCurrentTime, const float InDeltaTime) override;
-
-	/**
-	 * The system will use this event to notify a widget that the cursor has entered it. This event is NOT bubbled.
-	 *
-	 * @param MyGeometry The Geometry of the widget receiving the event
-	 * @param MouseEvent Information about the input event
-	 */
-	virtual void OnMouseEnter(const FGeometry& MyGeometry, const FPointerEvent& MouseEvent) override;
-
-	/**
-	 * The system will use this event to notify a widget that the cursor has left it. This event is NOT bubbled.
-	 *
-	 * @param MouseEvent Information about the input event
-	 */
-	virtual void OnMouseLeave(const FPointerEvent& MouseEvent) override;
-
-	/**
-	 * Called after a key is pressed when this widget has focus
-	 *
-	 * @param MyGeometry The Geometry of the widget receiving the event
-	 * @param  InKeyEvent  Key event
-	 *
-	 * @return  Returns whether the event was handled, along with other possible actions
-	 */
-	virtual FReply OnKeyDown(const FGeometry& MyGeometry, const FKeyEvent& InKeyEvent) override;
-
-	/**
-	 * Called when the user is dropping something onto a widget; terminates drag and drop.
-	 *
-	 * @param MyGeometry      The geometry of the widget receiving the event.
-	 * @param DragDropEvent   The drag and drop event.
-	 *
-	 * @return A reply that indicated whether this event was handled.
-	 */
-	virtual FReply OnDrop(const FGeometry& MyGeometry, const FDragDropEvent& DragDropEvent) override;
-
-	/**
-	 * Called during drag and drop when the the mouse is being dragged over a widget.
-	 *
-	 * @param MyGeometry      The geometry of the widget receiving the event.
-	 * @param DragDropEvent   The drag and drop event.
-	 *
-	 * @return A reply that indicated whether this event was handled.
-	 */
-	virtual FReply OnDragOver(const FGeometry& MyGeometry, const FDragDropEvent& DragDropEvent)  override;
-
 	void OnTimeSelectionChanged(Insights::ETimeChangedFlags InFlags, double InStartTime, double InEndTime);
 
 	void CreateTimingViewMarkers();
@@ -187,14 +104,6 @@ private:
 private:
 	TSharedRef<FMemorySharedState> SharedState;
 
-	/** Holds the tab manager that manages the front-end's tabs. */
-	TSharedPtr<FTabManager> TabManager;
-
-	TSharedPtr<FWorkspaceItem> AppMenuGroup;
-
-	/** All tabs owned by this window */
-	TSet<TSharedPtr<SDockTab>> OpenTabs;
-
 	/** The Timing view (multi-track) widget */
 	TSharedPtr<STimingView> TimingView;
 
@@ -214,10 +123,4 @@ private:
 
 	const int32 MaxMemAllocTableTreeViews = 4;
 	int32 LastMemAllocTableTreeViewIndex = -1;
-
-	/** The handle to the active update duration tick */
-	TWeakPtr<FActiveTimerHandle> ActiveTimerHandle;
-
-	/** The number of seconds the profiler has been active */
-	float DurationActive;
 };
diff --git a/Engine/Source/Developer/TraceInsights/Private/Insights/NetworkingProfiler/NetworkingProfilerManager.cpp b/Engine/Source/Developer/TraceInsights/Private/Insights/NetworkingProfiler/NetworkingProfilerManager.cpp
--- a/Engine/Source/Developer/TraceInsights/Private/Insights/NetworkingProfiler/NetworkingProfilerManager.cpp
+++ b/Engine/Source/Developer/TraceInsights/Private/Insights/NetworkingProfiler/NetworkingProfilerManager.cpp
@@ -5,6 +5,7 @@
 #include "MessageLog/Public/MessageLogModule.h"
 #include "Modules/ModuleManager.h"
 #include "TraceServices/AnalysisService.h"
+#include "Widgets/Docking/SDockTab.h"
 #include "WorkspaceMenuStructure.h"
 #include "WorkspaceMenuStructureModule.h"
 
diff --git a/Engine/Source/Developer/TraceInsights/Private/Insights/NetworkingProfiler/Widgets/SNetworkingProfilerToolbar.cpp b/Engine/Source/Developer/TraceInsights/Private/Insights/NetworkingProfiler/Widgets/SNetworkingProfilerToolbar.cpp
--- a/Engine/Source/Developer/TraceInsights/Private/Insights/NetworkingProfiler/Widgets/SNetworkingProfilerToolbar.cpp
+++ b/Engine/Source/Developer/TraceInsights/Private/Insights/NetworkingProfiler/Widgets/SNetworkingProfilerToolbar.cpp
@@ -2,6 +2,7 @@
 
 #include "SNetworkingProfilerToolbar.h"
 
+#include "Framework/MultiBox/MultiBoxExtender.h"
 #include "Framework/MultiBox/MultiBoxBuilder.h"
 #include "Widgets/SBoxPanel.h"
 
@@ -34,7 +35,7 @@ void SNetworkingProfilerToolbar::Construct(const FArguments& InArgs, TSharedRef<
 {
 	struct Local
 	{
-		static void FillViewToolbar(FToolBarBuilder& ToolbarBuilder, TSharedRef<SNetworkingProfilerWindow> ProfilerWindow)
+		static void FillViewToolbar(FToolBarBuilder& ToolbarBuilder, const FArguments &InArgs, TSharedRef<SNetworkingProfilerWindow> ProfilerWindow)
 		{
 			ToolbarBuilder.BeginSection("View");
 			{
@@ -68,9 +69,14 @@ void SNetworkingProfilerToolbar::Construct(const FArguments& InArgs, TSharedRef<
 			//		FSlateIcon(FInsightsStyle::GetStyleSetName(), "NewNetworkingInsights.Icon"));
 			//}
 			//ToolbarBuilder.EndSection();
+
+			if (InArgs._ToolbarExtender.IsValid())
+			{
+				InArgs._ToolbarExtender->Apply("MainToolbar", EExtensionHook::First, ToolbarBuilder);
+			}
 		}
 
-		static void FillRightSideToolbar(FToolBarBuilder& ToolbarBuilder)
+		static void FillRightSideToolbar(FToolBarBuilder& ToolbarBuilder, const FArguments &InArgs)
 		{
 			ToolbarBuilder.BeginSection("Debug");
 			{
@@ -79,18 +85,23 @@ void SNetworkingProfilerToolbar::Construct(const FArguments& InArgs, TSharedRef<
 					FSlateIcon(FInsightsStyle::GetStyleSetName(), "Icons.Debug.ToolBar"));
 			}
 			ToolbarBuilder.EndSection();
+
+			if (InArgs._ToolbarExtender.IsValid())
+			{
+				InArgs._ToolbarExtender->Apply("RightSideToolbar", EExtensionHook::First, ToolbarBuilder);
+			}
 		}
 	};
 
-	TSharedPtr<FUICommandList> CommandList = InProfilerWindow->GetCommandList();
+	const TSharedPtr<FUICommandList> CommandList = InProfilerWindow->GetCommandList();
 
 	FSlimHorizontalToolBarBuilder ToolbarBuilder(CommandList.ToSharedRef(), FMultiBoxCustomization::None);
 	ToolbarBuilder.SetStyle(&FInsightsStyle::Get(), "PrimaryToolbar");
-	Local::FillViewToolbar(ToolbarBuilder, InProfilerWindow);
+	Local::FillViewToolbar(ToolbarBuilder, InArgs, InProfilerWindow);
 
 	FSlimHorizontalToolBarBuilder RightSideToolbarBuilder(CommandList.ToSharedRef(), FMultiBoxCustomization::None);
 	RightSideToolbarBuilder.SetStyle(&FInsightsStyle::Get(), "PrimaryToolbar");
-	Local::FillRightSideToolbar(RightSideToolbarBuilder);
+	Local::FillRightSideToolbar(RightSideToolbarBuilder, InArgs);
 
 	ChildSlot
 	[
diff --git a/Engine/Source/Developer/TraceInsights/Private/Insights/NetworkingProfiler/Widgets/SNetworkingProfilerToolbar.h b/Engine/Source/Developer/TraceInsights/Private/Insights/NetworkingProfiler/Widgets/SNetworkingProfilerToolbar.h
--- a/Engine/Source/Developer/TraceInsights/Private/Insights/NetworkingProfiler/Widgets/SNetworkingProfilerToolbar.h
+++ b/Engine/Source/Developer/TraceInsights/Private/Insights/NetworkingProfiler/Widgets/SNetworkingProfilerToolbar.h
@@ -2,10 +2,10 @@
 
 #pragma once
 
-#include "Templates/SharedPointer.h"
 #include "Widgets/DeclarativeSyntaxSupport.h"
 #include "Widgets/SCompoundWidget.h"
 
+class FExtender;
 class SNetworkingProfilerWindow;
 
 class SNetworkingProfilerToolbar : public SCompoundWidget
@@ -18,6 +18,7 @@ public:
 	virtual ~SNetworkingProfilerToolbar();
 
 	SLATE_BEGIN_ARGS(SNetworkingProfilerToolbar) {}
+	SLATE_ARGUMENT(TSharedPtr<FExtender>, ToolbarExtender);
 	SLATE_END_ARGS()
 
 	/**
diff --git a/Engine/Source/Developer/TraceInsights/Private/Insights/NetworkingProfiler/Widgets/SNetworkingProfilerWindow.cpp b/Engine/Source/Developer/TraceInsights/Private/Insights/NetworkingProfiler/Widgets/SNetworkingProfilerWindow.cpp
--- a/Engine/Source/Developer/TraceInsights/Private/Insights/NetworkingProfiler/Widgets/SNetworkingProfilerWindow.cpp
+++ b/Engine/Source/Developer/TraceInsights/Private/Insights/NetworkingProfiler/Widgets/SNetworkingProfilerWindow.cpp
@@ -3,37 +3,18 @@
 #include "SNetworkingProfilerWindow.h"
 
 #include "Framework/Commands/UICommandList.h"
-#include "Framework/Docking/LayoutService.h"
 #include "Framework/Docking/TabManager.h"
 #include "Framework/Docking/WorkspaceItem.h"
-#include "Framework/MultiBox/MultiBoxBuilder.h"
 #include "SlateOptMacros.h"
-#include "Styling/AppStyle.h"
 #include "Widgets/Docking/SDockTab.h"
-#include "Widgets/Images/SImage.h"
-#include "Widgets/Layout/SBorder.h"
-#include "Widgets/Layout/SBox.h"
-#include "Widgets/SBoxPanel.h"
-#include "Widgets/Text/STextBlock.h"
-
-#if WITH_EDITOR
-	#include "EngineAnalytics.h"
-	#include "Runtime/Analytics/Analytics/Public/AnalyticsEventAttribute.h"
-	#include "Runtime/Analytics/Analytics/Public/Interfaces/IAnalyticsProvider.h"
-#endif // WITH_EDITOR
 
 // Insights
-#include "Insights/Common/InsightsMenuBuilder.h"
-#include "Insights/InsightsManager.h"
 #include "Insights/InsightsStyle.h"
 #include "Insights/NetworkingProfiler/NetworkingProfilerManager.h"
 #include "Insights/NetworkingProfiler/Widgets/SNetStatsView.h"
 #include "Insights/NetworkingProfiler/Widgets/SNetworkingProfilerToolbar.h"
 #include "Insights/NetworkingProfiler/Widgets/SPacketContentView.h"
 #include "Insights/NetworkingProfiler/Widgets/SPacketView.h"
-#include "Insights/TraceInsightsModule.h"
-#include "Insights/Version.h"
-#include "Insights/Widgets/SInsightsSettings.h"
 
 ////////////////////////////////////////////////////////////////////////////////////////////////////
 
@@ -48,12 +29,12 @@ const FName FNetworkingProfilerTabs::NetStatsViewID(TEXT("NetStats"));
 ////////////////////////////////////////////////////////////////////////////////////////////////////
 
 SNetworkingProfilerWindow::SNetworkingProfilerWindow()
-	: SelectedPacketStartIndex(0)
+	: SMajorTabWindow(FInsightsManagerTabs::NetworkingProfilerTabId)
+	, SelectedPacketStartIndex(0)
 	, SelectedPacketEndIndex(0)
 	, SelectionStartPosition(0)
 	, SelectionEndPosition(0)
 	, SelectedEventTypeIndex(InvalidEventTypeIndex)
-	, DurationActive(0.0f)
 {
 }
 
@@ -61,30 +42,18 @@ SNetworkingProfilerWindow::SNetworkingProfilerWindow()
 
 SNetworkingProfilerWindow::~SNetworkingProfilerWindow()
 {
-	// Close all tabs.
-	TArray<TSharedPtr<SDockTab>> LocalOpenTabs;
-	for (TSharedPtr<SDockTab>& Tab : OpenTabs)
-	{
-		LocalOpenTabs.Add(Tab);
-	}
-	for (TSharedPtr<SDockTab>& Tab : LocalOpenTabs)
-	{
-		Tab->RequestCloseTab();
-		Tab->RemoveTabFromParent();
-	}
-	LocalOpenTabs.Reset();
-	check(OpenTabs.IsEmpty());
+	CloseAllOpenTabs();
 
 	check(NetStatsView == nullptr);
 	check(PacketContentView == nullptr);
 	check(PacketView == nullptr);
+}
 
-#if WITH_EDITOR
-	if (DurationActive > 0.0f && FEngineAnalytics::IsAvailable())
-	{
-		FEngineAnalytics::GetProvider().RecordEvent(TEXT("Insights.Usage.NetworkingProfiler"), FAnalyticsEventAttribute(TEXT("Duration"), DurationActive));
-	}
-#endif // WITH_EDITOR
+////////////////////////////////////////////////////////////////////////////////////////////////////
+
+const TCHAR* SNetworkingProfilerWindow::GetAnalyticsEventName() const
+{
+	return TEXT("Insights.Usage.NetworkingProfiler");
 }
 
 ////////////////////////////////////////////////////////////////////////////////////////////////////
@@ -134,7 +103,7 @@ TSharedRef<SDockTab> SNetworkingProfilerWindow::SpawnTab_PacketView(const FSpawn
 		];
 
 	DockTab->SetOnTabClosed(SDockTab::FOnTabClosedCallback::CreateRaw(this, &SNetworkingProfilerWindow::OnPacketViewTabClosed));
-	OpenTabs.Add(DockTab);
+	AddOpenTab(DockTab);
 
 	return DockTab;
 }
@@ -146,7 +115,7 @@ void SNetworkingProfilerWindow::OnPacketViewTabClosed(TSharedRef<SDockTab> TabBe
 {
 	PacketView = nullptr;
 
-	OpenTabs.Remove(TabBeingClosed);
+	RemoveOpenTab(TabBeingClosed);
 }
 
 ////////////////////////////////////////////////////////////////////////////////////////////////////
@@ -162,7 +131,7 @@ TSharedRef<SDockTab> SNetworkingProfilerWindow::SpawnTab_PacketContentView(const
 		];
 
 	DockTab->SetOnTabClosed(SDockTab::FOnTabClosedCallback::CreateRaw(this, &SNetworkingProfilerWindow::OnPacketContentViewTabClosed));
-	OpenTabs.Add(DockTab);
+	AddOpenTab(DockTab);
 
 	return DockTab;
 }
@@ -174,7 +143,7 @@ void SNetworkingProfilerWindow::OnPacketContentViewTabClosed(TSharedRef<SDockTab
 {
 	PacketContentView = nullptr;
 
-	OpenTabs.Remove(TabBeingClosed);
+	RemoveOpenTab(TabBeingClosed);
 }
 
 ////////////////////////////////////////////////////////////////////////////////////////////////////
@@ -190,7 +159,7 @@ TSharedRef<SDockTab> SNetworkingProfilerWindow::SpawnTab_NetStatsView(const FSpa
 		];
 
 	DockTab->SetOnTabClosed(SDockTab::FOnTabClosedCallback::CreateRaw(this, &SNetworkingProfilerWindow::OnNetStatsViewTabClosed));
-	OpenTabs.Add(DockTab);
+	AddOpenTab(DockTab);
 
 	return DockTab;
 }
@@ -202,46 +171,61 @@ void SNetworkingProfilerWindow::OnNetStatsViewTabClosed(TSharedRef<SDockTab> Tab
 {
 	NetStatsView = nullptr;
 
-	OpenTabs.Remove(TabBeingClosed);
+	RemoveOpenTab(TabBeingClosed);
 }
 
 ////////////////////////////////////////////////////////////////////////////////////////////////////
 
-BEGIN_SLATE_FUNCTION_BUILD_OPTIMIZATION
 void SNetworkingProfilerWindow::Construct(const FArguments& InArgs, const TSharedRef<SDockTab>& ConstructUnderMajorTab, const TSharedPtr<SWindow>& ConstructUnderWindow)
 {
-	CommandList = MakeShared<FUICommandList>();
+	TSharedPtr<FNetworkingProfilerManager> NetworkingProfilerManager = FNetworkingProfilerManager::Get();
+	ensure(NetworkingProfilerManager.IsValid());
 
-	// Create & initialize tab manager.
-	TabManager = FGlobalTabmanager::Get()->NewTabManager(ConstructUnderMajorTab);
-	const auto& PersistLayout = [](const TSharedRef<FTabManager::FLayout>& LayoutToSave)
-	{
-		FLayoutSaveRestore::SaveToConfig(FTraceInsightsModule::GetUnrealInsightsLayoutIni(), LayoutToSave);
-	};
-	TabManager->SetOnPersistLayout(FTabManager::FOnPersistLayout::CreateLambda(PersistLayout));
+	SetCommandList(MakeShared<FUICommandList>());
+	BindCommands();
 
-	TSharedRef<FWorkspaceItem> AppMenuGroup = TabManager->AddLocalWorkspaceMenuCategory(LOCTEXT("NetworkingProfilerMenuGroupName", "Networking Insights"));
+	SMajorTabWindow::FArguments Args;
+	SMajorTabWindow::Construct(Args, ConstructUnderMajorTab, ConstructUnderWindow);
+}
 
-	TabManager->RegisterTabSpawner(FNetworkingProfilerTabs::PacketViewID, FOnSpawnTab::CreateRaw(this, &SNetworkingProfilerWindow::SpawnTab_PacketView))
+////////////////////////////////////////////////////////////////////////////////////////////////////
+
+TSharedRef<FWorkspaceItem> SNetworkingProfilerWindow::CreateWorkspaceMenuGroup()
+{
+	return GetTabManager()->AddLocalWorkspaceMenuCategory(LOCTEXT("NetworkingProfilerMenuGroupName", "Networking Insights"));
+}
+
+////////////////////////////////////////////////////////////////////////////////////////////////////
+
+void SNetworkingProfilerWindow::RegisterTabSpawners()
+{
+	check(GetTabManager().IsValid());
+	FTabManager* TabManagerPtr = GetTabManager().Get();
+	check(GetWorkspaceMenuGroup().IsValid());
+	const TSharedRef<FWorkspaceItem> Group = GetWorkspaceMenuGroup().ToSharedRef();
+
+	TabManagerPtr->RegisterTabSpawner(FNetworkingProfilerTabs::PacketViewID, FOnSpawnTab::CreateRaw(this, &SNetworkingProfilerWindow::SpawnTab_PacketView))
 		.SetDisplayName(LOCTEXT("PacketViewTabTitle", "Packet View"))
 		.SetIcon(FSlateIcon(FInsightsStyle::GetStyleSetName(), "Icons.PacketView"))
-		.SetGroup(AppMenuGroup);
+		.SetGroup(Group);
 
-	TabManager->RegisterTabSpawner(FNetworkingProfilerTabs::PacketContentViewID, FOnSpawnTab::CreateRaw(this, &SNetworkingProfilerWindow::SpawnTab_PacketContentView))
+	TabManagerPtr->RegisterTabSpawner(FNetworkingProfilerTabs::PacketContentViewID, FOnSpawnTab::CreateRaw(this, &SNetworkingProfilerWindow::SpawnTab_PacketContentView))
 		.SetDisplayName(LOCTEXT("PacketContentViewTabTitle", "Packet Content"))
 		.SetIcon(FSlateIcon(FInsightsStyle::GetStyleSetName(), "Icons.PacketContentView"))
-		.SetGroup(AppMenuGroup);
+		.SetGroup(Group);
 
-	TabManager->RegisterTabSpawner(FNetworkingProfilerTabs::NetStatsViewID, FOnSpawnTab::CreateRaw(this, &SNetworkingProfilerWindow::SpawnTab_NetStatsView))
+	TabManagerPtr->RegisterTabSpawner(FNetworkingProfilerTabs::NetStatsViewID, FOnSpawnTab::CreateRaw(this, &SNetworkingProfilerWindow::SpawnTab_NetStatsView))
 		.SetDisplayName(LOCTEXT("NetStatsViewTabTitle", "Net Stats"))
 		.SetIcon(FSlateIcon(FInsightsStyle::GetStyleSetName(), "Icons.NetStatsView"))
-		.SetGroup(AppMenuGroup);
+		.SetGroup(Group);
+}
 
-	TSharedPtr<FNetworkingProfilerManager> NetworkingProfilerManager = FNetworkingProfilerManager::Get();
-	ensure(NetworkingProfilerManager.IsValid());
+////////////////////////////////////////////////////////////////////////////////////////////////////
 
-	// Create tab layout.
-	TSharedRef<FTabManager::FLayout> Layout = FTabManager::NewLayout("InsightsNetworkingProfilerLayout_v1.2")
+BEGIN_SLATE_FUNCTION_BUILD_OPTIMIZATION
+TSharedRef<FTabManager::FLayout> SNetworkingProfilerWindow::CreateDefaultTabLayout() const
+{
+	return FTabManager::NewLayout("InsightsNetworkingProfilerLayout_v1.2")
 		->AddArea
 		(
 			FTabManager::NewPrimaryArea()
@@ -274,105 +258,28 @@ void SNetworkingProfilerWindow::Construct(const FArguments& InArgs, const TShare
 				->SetForegroundTab(FNetworkingProfilerTabs::NetStatsViewID)
 			)
 		);
+}
+END_SLATE_FUNCTION_BUILD_OPTIMIZATION
 
-	Layout = FLayoutSaveRestore::LoadFromConfig(FTraceInsightsModule::GetUnrealInsightsLayoutIni(), Layout);
-
-	// Create & initialize main menu.
-	FMenuBarBuilder MenuBarBuilder = FMenuBarBuilder(TSharedPtr<FUICommandList>());
-	MenuBarBuilder.AddPullDownMenu(
-		LOCTEXT("MenuLabel", "Menu"),
-		FText::GetEmpty(),
-		FNewMenuDelegate::CreateStatic(&SNetworkingProfilerWindow::FillMenu, TabManager),
-		FName(TEXT("Menu"))
-	);
-
-#if !WITH_EDITOR
-	TSharedRef<SWidget> MenuWidget = MenuBarBuilder.MakeWidget();
-	MenuWidget->SetClipping(EWidgetClipping::ClipToBoundsWithoutIntersecting);
-#endif
-
-	ChildSlot
-	[
-		SNew(SOverlay)
-
-#if !WITH_EDITOR
-		// Menu
-		+ SOverlay::Slot()
-		.HAlign(HAlign_Left)
-		.VAlign(VAlign_Top)
-		.Padding(34.0f, -60.0f, 0.0f, 0.0f)
-		[
-			MenuWidget
-		]
-#endif
-
-		// Version
-		+ SOverlay::Slot()
-		.HAlign(HAlign_Right)
-		.VAlign(VAlign_Top)
-		.Padding(0.0f, -16.0f, 4.0f, 0.0f)
-		[
-			SNew(STextBlock)
-			.Clipping(EWidgetClipping::ClipToBoundsWithoutIntersecting)
-			.Text(LOCTEXT("UnrealInsightsVersion", UNREAL_INSIGHTS_VERSION_STRING_EX))
-			.ColorAndOpacity(FLinearColor(0.15f, 0.15f, 0.15f, 1.0f))
-		]
-
-		// Overlay slot for the main window area
-		+ SOverlay::Slot()
-		.HAlign(HAlign_Fill)
-		.VAlign(VAlign_Fill)
-		[
-			SNew(SVerticalBox)
-
-			+ SVerticalBox::Slot()
-			.AutoHeight()
-			.Padding(FMargin(0.0f, 0.0f, 0.0f, 5.0f))
-			[
-				SNew(SNetworkingProfilerToolbar, SharedThis(this))
-			]
-
-			+ SVerticalBox::Slot()
-			.FillHeight(1.0f)
-			[
-				TabManager->RestoreFrom(Layout, ConstructUnderWindow).ToSharedRef()
-			]
-		]
-
-		// Session hint overlay
-		+ SOverlay::Slot()
-		.HAlign(HAlign_Center)
-		.VAlign(VAlign_Center)
-		[
-			SNew(SBorder)
-			.Visibility(this, &SNetworkingProfilerWindow::IsSessionOverlayVisible)
-			.BorderImage(FAppStyle::Get().GetBrush("PopupText.Background"))
-			.Padding(8.0f)
-			[
-				SNew(STextBlock)
-				.Text(LOCTEXT("SelectTraceOverlayText", "Please select a trace."))
-			]
-		]
-	];
-
-#if !WITH_EDITOR
-	// Tell tab-manager about the global menu bar.
-	TabManager->SetMenuMultiBox(MenuBarBuilder.GetMultiBox(), MenuWidget);
-#endif
+////////////////////////////////////////////////////////////////////////////////////////////////////
 
-	BindCommands();
+TSharedRef<SWidget> SNetworkingProfilerWindow::CreateToolbar(TSharedPtr<FExtender> Extender)
+{
+	return SNew(SNetworkingProfilerToolbar, SharedThis(this))
+		.ToolbarExtender(Extender);
 }
-END_SLATE_FUNCTION_BUILD_OPTIMIZATION
 
 ////////////////////////////////////////////////////////////////////////////////////////////////////
 
 void SNetworkingProfilerWindow::BindCommands()
 {
+	check(GetCommandList().IsValid());
+
 	Map_TogglePacketViewVisibility();
 	Map_TogglePacketContentViewVisibility();
 	Map_ToggleNetStatsViewVisibility();
 
-	CommandList->MapAction(FInsightsCommands::Get().ToggleDebugInfo, FInsightsManager::GetActionManager().ToggleDebugInfo_Custom());
+	GetCommandList()->MapAction(FInsightsCommands::Get().ToggleDebugInfo, FInsightsManager::GetActionManager().ToggleDebugInfo_Custom());
 }
 
 ////////////////////////////////////////////////////////////////////////////////////////////////////
@@ -383,7 +290,7 @@ void SNetworkingProfilerWindow::BindCommands()
 	\
 	void SNetworkingProfilerWindow::Map_##CmdName()\
 	{\
-		CommandList->MapAction(FNetworkingProfilerManager::GetCommands().CmdName, CmdName##_Custom());\
+		GetCommandList()->MapAction(FNetworkingProfilerManager::GetCommands().CmdName, CmdName##_Custom());\
 	}\
 	\
 	const FUIAction SNetworkingProfilerWindow::CmdName##_Custom() \
@@ -418,73 +325,6 @@ IMPLEMENT_TOGGLE_COMMAND(ToggleNetStatsViewVisibility, IsNetStatsViewVisible, Sh
 
 ////////////////////////////////////////////////////////////////////////////////////////////////////
 
-void SNetworkingProfilerWindow::FillMenu(FMenuBuilder& MenuBuilder, const TSharedPtr<FTabManager> TabManager)
-{
-	if (!TabManager.IsValid())
-	{
-		return;
-	}
-
-	FInsightsManager::Get()->GetInsightsMenuBuilder()->PopulateMenu(MenuBuilder);
-
-	TabManager->PopulateLocalTabSpawnerMenu(MenuBuilder);
-}
-
-////////////////////////////////////////////////////////////////////////////////////////////////////
-
-void SNetworkingProfilerWindow::ShowTab(const FName& TabID)
-{
-	if (TabManager->HasTabSpawner(TabID))
-	{
-		TabManager->TryInvokeTab(TabID);
-	}
-}
-
-////////////////////////////////////////////////////////////////////////////////////////////////////
-
-void SNetworkingProfilerWindow::HideTab(const FName& TabID)
-{
-	TSharedPtr<SDockTab> Tab = TabManager->FindExistingLiveTab(TabID);
-	if (Tab.IsValid())
-	{
-		Tab->RequestCloseTab();
-		Tab->RemoveTabFromParent();
-	}
-}
-
-////////////////////////////////////////////////////////////////////////////////////////////////////
-
-EVisibility SNetworkingProfilerWindow::IsSessionOverlayVisible() const
-{
-	if (FInsightsManager::Get()->GetSession().IsValid())
-	{
-		return EVisibility::Hidden;
-	}
-	else
-	{
-		return EVisibility::Visible;
-	}
-}
-
-////////////////////////////////////////////////////////////////////////////////////////////////////
-
-bool SNetworkingProfilerWindow::IsProfilerEnabled() const
-{
-	return FInsightsManager::Get()->GetSession().IsValid();
-}
-
-////////////////////////////////////////////////////////////////////////////////////////////////////
-
-EActiveTimerReturnType SNetworkingProfilerWindow::UpdateActiveDuration(double InCurrentTime, float InDeltaTime)
-{
-	DurationActive += InDeltaTime;
-
-	// The profiler window will explicitly unregister this active timer when the mouse leaves.
-	return EActiveTimerReturnType::Continue;
-}
-
-////////////////////////////////////////////////////////////////////////////////////////////////////
-
 void SNetworkingProfilerWindow::Tick(const FGeometry& AllottedGeometry, const double InCurrentTime, const float InDeltaTime)
 {
 	TSharedPtr<const TraceServices::IAnalysisSession> Session = FInsightsManager::Get()->GetSession();
@@ -671,62 +511,6 @@ void SNetworkingProfilerWindow::UpdateAvailableConnectionModes()
 
 ////////////////////////////////////////////////////////////////////////////////////////////////////
 
-void SNetworkingProfilerWindow::OnMouseEnter(const FGeometry& MyGeometry, const FPointerEvent& MouseEvent)
-{
-	SCompoundWidget::OnMouseEnter(MyGeometry, MouseEvent);
-
-	if (!ActiveTimerHandle.IsValid())
-	{
-		ActiveTimerHandle = RegisterActiveTimer(0.f, FWidgetActiveTimerDelegate::CreateSP(this, &SNetworkingProfilerWindow::UpdateActiveDuration));
-	}
-}
-
-////////////////////////////////////////////////////////////////////////////////////////////////////
-
-void SNetworkingProfilerWindow::OnMouseLeave(const FPointerEvent& MouseEvent)
-{
-	SCompoundWidget::OnMouseLeave(MouseEvent);
-
-	auto PinnedActiveTimerHandle = ActiveTimerHandle.Pin();
-	if (PinnedActiveTimerHandle.IsValid())
-	{
-		UnRegisterActiveTimer(PinnedActiveTimerHandle.ToSharedRef());
-	}
-}
-
-////////////////////////////////////////////////////////////////////////////////////////////////////
-
-FReply SNetworkingProfilerWindow::OnKeyDown(const FGeometry& MyGeometry, const FKeyEvent& InKeyEvent)
-{
-	return GetCommandList()->ProcessCommandBindings(InKeyEvent) ? FReply::Handled() : FReply::Unhandled();
-}
-
-////////////////////////////////////////////////////////////////////////////////////////////////////
-
-FReply SNetworkingProfilerWindow::OnDragOver(const FGeometry& MyGeometry, const FDragDropEvent& DragDropEvent)
-{
-	if (FInsightsManager::Get()->OnDragOver(DragDropEvent))
-	{
-		return FReply::Handled();
-	}
-
-	return SCompoundWidget::OnDragOver(MyGeometry, DragDropEvent);
-}
-
-////////////////////////////////////////////////////////////////////////////////////////////////////
-
-FReply SNetworkingProfilerWindow::OnDrop(const FGeometry& MyGeometry, const FDragDropEvent& DragDropEvent)
-{
-	if (FInsightsManager::Get()->OnDrop(DragDropEvent))
-	{
-		return FReply::Handled();
-	}
-
-	return SCompoundWidget::OnDrop(MyGeometry, DragDropEvent);
-}
-
-////////////////////////////////////////////////////////////////////////////////////////////////////
-
 FText SNetworkingProfilerWindow::FGameInstanceItem::GetText() const
 {
 	if (IsInstanceNameSet())
diff --git a/Engine/Source/Developer/TraceInsights/Private/Insights/NetworkingProfiler/Widgets/SNetworkingProfilerWindow.h b/Engine/Source/Developer/TraceInsights/Private/Insights/NetworkingProfiler/Widgets/SNetworkingProfilerWindow.h
--- a/Engine/Source/Developer/TraceInsights/Private/Insights/NetworkingProfiler/Widgets/SNetworkingProfilerWindow.h
+++ b/Engine/Source/Developer/TraceInsights/Private/Insights/NetworkingProfiler/Widgets/SNetworkingProfilerWindow.h
@@ -3,26 +3,14 @@
 #pragma once
 
 #include "CoreMinimal.h"
-#include "Framework/Docking/TabManager.h"
-#include "Input/Reply.h"
-#include "Layout/Visibility.h"
-#include "Misc/Guid.h"
-#include "SlateFwd.h"
 #include "TraceServices/Model/NetProfiler.h"
-#include "Widgets/DeclarativeSyntaxSupport.h"
-#include "Widgets/Docking/SDockTab.h"
 #include "Widgets/Input/SComboBox.h"
-#include "Widgets/SCompoundWidget.h"
 
 // Insights
-#include "Insights/InsightsManager.h"
+#include "Insights/Widgets/SMajorTabWindow.h"
 
 ////////////////////////////////////////////////////////////////////////////////////////////////////
 
-class FActiveTimerHandle;
-class FMenuBuilder;
-class FUICommandList;
-
 class SPacketView;
 class SPacketContentView;
 class SNetStatsView;
@@ -39,8 +27,8 @@ struct FNetworkingProfilerTabs
 
 ////////////////////////////////////////////////////////////////////////////////////////////////////
 
-/** Implements the timing profiler window. */
-class SNetworkingProfilerWindow : public SCompoundWidget
+/** Implements the Networking Insights window. */
+class SNetworkingProfilerWindow : public Insights::SMajorTabWindow
 {
 private:
 	struct FGameInstanceItem
@@ -101,32 +89,22 @@ public:
 	SLATE_BEGIN_ARGS(SNetworkingProfilerWindow) {}
 	SLATE_END_ARGS()
 
-	void Reset();
+	virtual void Reset() override;
 
 	/** Constructs this widget. */
 	void Construct(const FArguments& InArgs, const TSharedRef<SDockTab>& ConstructUnderMajorTab, const TSharedPtr<SWindow>& ConstructUnderWindow);
 
-	/** @return The viewport command list */
-	const TSharedPtr<FUICommandList> GetCommandList() const { return CommandList; }
-
-	/** @return The tab manager */
-	TSharedPtr<FTabManager> GetTabManager() const { return TabManager; }
-
-	void ShowTab(const FName& TabID);
-	void HideTab(const FName& TabID);
-	void ShowOrHideTab(const FName& TabID, bool bShow) { bShow ? ShowTab(TabID) : HideTab(TabID); }
-
 	TSharedPtr<SPacketView> GetPacketView() const { return PacketView; }
 	const bool IsPacketViewVisible() const { return PacketView.IsValid(); }
-	void ShowOrHidePacketView(const bool bVisibleState) { ShowOrHideTab(FNetworkingProfilerTabs::PacketViewID, bVisibleState); }
+	void ShowOrHidePacketView(const bool bVisibleState) { ShowHideTab(FNetworkingProfilerTabs::PacketViewID, bVisibleState); }
 
 	TSharedPtr<SPacketContentView> GetPacketContentView() const { return PacketContentView; }
 	const bool IsPacketContentViewVisible() const { return PacketContentView.IsValid(); }
-	void ShowOrHidePacketContentView(const bool bVisibleState) { ShowOrHideTab(FNetworkingProfilerTabs::PacketContentViewID, bVisibleState); }
+	void ShowOrHidePacketContentView(const bool bVisibleState) { ShowHideTab(FNetworkingProfilerTabs::PacketContentViewID, bVisibleState); }
 
 	TSharedPtr<SNetStatsView> GetNetStatsView() const { return NetStatsView; }
 	const bool IsNetStatsViewVisible() const { return NetStatsView.IsValid(); }
-	void ShowOrHideNetStatsView(const bool bVisibleState) { ShowOrHideTab(FNetworkingProfilerTabs::NetStatsViewID, bVisibleState); }
+	void ShowOrHideNetStatsView(const bool bVisibleState) { ShowHideTab(FNetworkingProfilerTabs::NetStatsViewID, bVisibleState); }
 
 	const TraceServices::FNetProfilerGameInstance* GetSelectedGameInstance() const { return SelectedGameInstance ? &SelectedGameInstance->GameInstance : nullptr; }
 	uint32 GetSelectedGameInstanceIndex() const { return SelectedGameInstance ? SelectedGameInstance->GetIndex() : 0; }
@@ -142,6 +120,13 @@ public:
 	void SetSelectedBitRange(uint32 StartPos, uint32 EndPos);
 	void SetSelectedEventTypeIndex(uint32 InEventTypeIndex);
 
+protected:
+	virtual const TCHAR* GetAnalyticsEventName() const override;
+	virtual TSharedRef<FWorkspaceItem> CreateWorkspaceMenuGroup() override;
+	virtual void RegisterTabSpawners() override;
+	virtual TSharedRef<FTabManager::FLayout> CreateDefaultTabLayout() const override;
+	virtual TSharedRef<SWidget> CreateToolbar(TSharedPtr<FExtender> Extender);
+
 private:
 	TSharedRef<SDockTab> SpawnTab_PacketView(const FSpawnTabArgs& Args);
 	void OnPacketViewTabClosed(TSharedRef<SDockTab> TabBeingClosed);
@@ -177,23 +162,6 @@ private:\
 
 	void UpdateAggregatedNetStats();
 
-	/**
-	 * Fill the main menu with menu items.
-	 *
-	 * @param MenuBuilder The multi-box builder that should be filled with content for this pull-down menu.
-	 * @param TabManager A Tab Manager from which to populate tab spawner menu items.
-	 */
-	static void FillMenu(FMenuBuilder& MenuBuilder, const TSharedPtr<FTabManager> TabManager);
-
-	/** Callback for determining the visibility of the 'Select a session' overlay. */
-	EVisibility IsSessionOverlayVisible() const;
-
-	/** Callback for getting the enabled state of the profiler window. */
-	bool IsProfilerEnabled() const;
-
-	/** Updates the amount of time the profiler has been active. */
-	EActiveTimerReturnType UpdateActiveDuration(double InCurrentTime, float InDeltaTime);
-
 	//////////////////////////////////////////////////
 
 	/**
@@ -205,51 +173,6 @@ private:\
 	 */
 	virtual void Tick(const FGeometry& AllottedGeometry, const double InCurrentTime, const float InDeltaTime) override;
 
-	/**
-	 * The system will use this event to notify a widget that the cursor has entered it. This event is NOT bubbled.
-	 *
-	 * @param MyGeometry The Geometry of the widget receiving the event
-	 * @param MouseEvent Information about the input event
-	 */
-	virtual void OnMouseEnter(const FGeometry& MyGeometry, const FPointerEvent& MouseEvent) override;
-
-	/**
-	 * The system will use this event to notify a widget that the cursor has left it. This event is NOT bubbled.
-	 *
-	 * @param MouseEvent Information about the input event
-	 */
-	virtual void OnMouseLeave(const FPointerEvent& MouseEvent) override;
-
-	/**
-	 * Called after a key is pressed when this widget has focus
-	 *
-	 * @param MyGeometry The Geometry of the widget receiving the event
-	 * @param  InKeyEvent  Key event
-	 *
-	 * @return  Returns whether the event was handled, along with other possible actions
-	 */
-	virtual FReply OnKeyDown(const FGeometry& MyGeometry, const FKeyEvent& InKeyEvent) override;
-
-	/**
-	 * Called when the user is dropping something onto a widget; terminates drag and drop.
-	 *
-	 * @param MyGeometry      The geometry of the widget receiving the event.
-	 * @param DragDropEvent   The drag and drop event.
-	 *
-	 * @return A reply that indicated whether this event was handled.
-	 */
-	virtual FReply OnDrop(const FGeometry& MyGeometry, const FDragDropEvent& DragDropEvent) override;
-
-	/**
-	 * Called during drag and drop when the the mouse is being dragged over a widget.
-	 *
-	 * @param MyGeometry      The geometry of the widget receiving the event.
-	 * @param DragDropEvent   The drag and drop event.
-	 *
-	 * @return A reply that indicated whether this event was handled.
-	 */
-	virtual FReply OnDragOver(const FGeometry& MyGeometry, const FDragDropEvent& DragDropEvent)  override;
-
 	//////////////////////////////////////////////////
 
 	void UpdateAvailableGameInstances();
@@ -274,15 +197,6 @@ private:\
 	//////////////////////////////////////////////////
 
 private:
-	/** Holds the tab manager that manages the front-end's tabs. */
-	TSharedPtr<FTabManager> TabManager;
-
-	/** All tabs owned by this window */
-	TSet<TSharedPtr<SDockTab>> OpenTabs;
-
-	/** Commandlist used in the window (Maps commands to window specific actions) */
-	TSharedPtr<FUICommandList> CommandList;
-
 	/** The Packet View widget */
 	TSharedPtr<SPacketView> PacketView;
 
@@ -317,10 +231,4 @@ private:
 
 	static const uint32 InvalidEventTypeIndex = uint32(-1);
 	uint32 SelectedEventTypeIndex;
-
-	/** The handle to the active update duration tick */
-	TWeakPtr<FActiveTimerHandle> ActiveTimerHandle;
-
-	/** The number of seconds the profiler has been active */
-	float DurationActive;
 };
diff --git a/Engine/Source/Developer/TraceInsights/Private/Insights/TaskGraphProfiler/TaskGraphProfilerManager.cpp b/Engine/Source/Developer/TraceInsights/Private/Insights/TaskGraphProfiler/TaskGraphProfilerManager.cpp
--- a/Engine/Source/Developer/TraceInsights/Private/Insights/TaskGraphProfiler/TaskGraphProfilerManager.cpp
+++ b/Engine/Source/Developer/TraceInsights/Private/Insights/TaskGraphProfiler/TaskGraphProfilerManager.cpp
@@ -8,6 +8,7 @@
 #include "Modules/ModuleManager.h"
 #include "TraceServices/AnalysisService.h"
 #include "TraceServices/Model/TasksProfiler.h"
+#include "Widgets/Docking/SDockTab.h"
 
 // Insights
 #include "Insights/InsightsStyle.h"
@@ -475,7 +476,6 @@ void FTaskGraphProfilerManager::OnWindowClosedEvent()
 		if (Tab.IsValid())
 		{
 			Tab->RequestCloseTab();
-			Tab->RemoveTabFromParent();
 		}
 	}
 }
diff --git a/Engine/Source/Developer/TraceInsights/Private/Insights/TimingProfilerManager.cpp b/Engine/Source/Developer/TraceInsights/Private/Insights/TimingProfilerManager.cpp
--- a/Engine/Source/Developer/TraceInsights/Private/Insights/TimingProfilerManager.cpp
+++ b/Engine/Source/Developer/TraceInsights/Private/Insights/TimingProfilerManager.cpp
@@ -5,6 +5,7 @@
 #include "MessageLog/Public/MessageLogModule.h"
 #include "Modules/ModuleManager.h"
 #include "TraceServices/AnalysisService.h"
+#include "Widgets/Docking/SDockTab.h"
 #include "WorkspaceMenuStructure.h"
 #include "WorkspaceMenuStructureModule.h"
 
diff --git a/Engine/Source/Developer/TraceInsights/Private/Insights/Widgets/SMajorTabWindow.cpp b/Engine/Source/Developer/TraceInsights/Private/Insights/Widgets/SMajorTabWindow.cpp
new file mode 100644
--- /dev/null
+++ b/Engine/Source/Developer/TraceInsights/Private/Insights/Widgets/SMajorTabWindow.cpp
@@ -0,0 +1,434 @@
+// Copyright Epic Games, Inc. All Rights Reserved.
+
+#include "Insights/Widgets/SMajorTabWindow.h"
+
+#include "Framework/Commands/UICommandList.h"
+#include "Framework/Docking/LayoutExtender.h"
+#include "Framework/Docking/LayoutService.h"
+#include "Framework/Docking/WorkspaceItem.h"
+#include "Framework/MultiBox/MultiBoxBuilder.h"
+#include "Modules/ModuleManager.h"
+#include "SlateOptMacros.h"
+#include "Styling/AppStyle.h"
+#include "Widgets/Docking/SDockTab.h"
+#include "Widgets/Layout/SBorder.h"
+#include "Widgets/Layout/SBox.h"
+#include "Widgets/SBoxPanel.h"
+#include "Widgets/Text/STextBlock.h"
+
+
+#if WITH_EDITOR
+	#include "EngineAnalytics.h"
+	#include "Runtime/Analytics/Analytics/Public/AnalyticsEventAttribute.h"
+	#include "Runtime/Analytics/Analytics/Public/Interfaces/IAnalyticsProvider.h"
+#endif // WITH_EDITOR
+
+// Insights
+#include "Insights/Common/InsightsMenuBuilder.h"
+#include "Insights/InsightsManager.h"
+#include "Insights/InsightsStyle.h"
+#include "Insights/Log.h"
+#include "Insights/TraceInsightsModule.h"
+#include "Insights/Version.h"
+
+////////////////////////////////////////////////////////////////////////////////////////////////////
+
+#define LOCTEXT_NAMESPACE "SMajorTabWindow"
+
+////////////////////////////////////////////////////////////////////////////////////////////////////
+
+namespace Insights
+{
+
+////////////////////////////////////////////////////////////////////////////////////////////////////
+// SMajorTabWindow
+////////////////////////////////////////////////////////////////////////////////////////////////////
+
+SMajorTabWindow::SMajorTabWindow(const FName& InMajorTabId)
+	: MajorTabId(InMajorTabId)
+	, DurationActive(0.0f)
+{
+}
+
+////////////////////////////////////////////////////////////////////////////////////////////////////
+
+SMajorTabWindow::~SMajorTabWindow()
+{
+	CloseAllOpenTabs();
+
+#if WITH_EDITOR
+	if (DurationActive > 0.0f && FEngineAnalytics::IsAvailable())
+	{
+		FString EventName = GetAnalyticsEventName();
+		FEngineAnalytics::GetProvider().RecordEvent(EventName, FAnalyticsEventAttribute(TEXT("Duration"), DurationActive));
+	}
+#endif // WITH_EDITOR
+}
+
+////////////////////////////////////////////////////////////////////////////////////////////////////
+
+const TCHAR* SMajorTabWindow::GetAnalyticsEventName() const
+{
+	return TEXT("Insights.Usage.MajorTabWindow");
+}
+
+////////////////////////////////////////////////////////////////////////////////////////////////////
+
+void SMajorTabWindow::Reset()
+{
+}
+
+////////////////////////////////////////////////////////////////////////////////////////////////////
+
+void SMajorTabWindow::ShowTab(const FName& TabId)
+{
+	if (TabManager->HasTabSpawner(TabId))
+	{
+		TabManager->TryInvokeTab(TabId);
+	}
+}
+
+////////////////////////////////////////////////////////////////////////////////////////////////////
+
+void SMajorTabWindow::HideTab(const FName& TabId)
+{
+	TSharedPtr<SDockTab> Tab = TabManager->FindExistingLiveTab(TabId);
+	if (Tab.IsValid())
+	{
+		Tab->RequestCloseTab();
+	}
+}
+
+////////////////////////////////////////////////////////////////////////////////////////////////////
+
+void SMajorTabWindow::AddOpenTab(const TSharedRef<SDockTab>& DockTab)
+{
+	UE_LOG(TraceInsights, Log, TEXT("[Tab] SHOW %s"), *DockTab->GetContent()->GetTypeAsString());
+	OpenTabs.Add(DockTab);
+}
+
+////////////////////////////////////////////////////////////////////////////////////////////////////
+
+void SMajorTabWindow::RemoveOpenTab(const TSharedRef<SDockTab>& DockTab)
+{
+	UE_LOG(TraceInsights, Log, TEXT("[Tab] HIDE %s"), *DockTab->GetContent()->GetTypeAsString());
+	OpenTabs.Remove(DockTab);
+}
+
+////////////////////////////////////////////////////////////////////////////////////////////////////
+
+void SMajorTabWindow::CloseAllOpenTabs()
+{
+	// Close all tabs.
+	TArray<TSharedPtr<SDockTab>> LocalOpenTabs;
+	for (TSharedPtr<SDockTab>& Tab : OpenTabs)
+	{
+		LocalOpenTabs.Add(Tab);
+	}
+	for (TSharedPtr<SDockTab>& Tab : LocalOpenTabs)
+	{
+		Tab->RequestCloseTab();
+	}
+	LocalOpenTabs.Reset();
+	check(OpenTabs.IsEmpty());
+}
+
+////////////////////////////////////////////////////////////////////////////////////////////////////
+
+TSharedRef<FWorkspaceItem> SMajorTabWindow::CreateWorkspaceMenuGroup()
+{
+	check(TabManager.IsValid());
+	return TabManager->AddLocalWorkspaceMenuCategory(LOCTEXT("MajorTabMenuGroupName", "Insights"));
+}
+
+////////////////////////////////////////////////////////////////////////////////////////////////////
+
+void SMajorTabWindow::RegisterTabSpawners()
+{
+}
+
+////////////////////////////////////////////////////////////////////////////////////////////////////
+
+BEGIN_SLATE_FUNCTION_BUILD_OPTIMIZATION
+void SMajorTabWindow::Construct(const FArguments& InArgs, const TSharedRef<SDockTab>& ConstructUnderMajorTab, const TSharedPtr<SWindow>& ConstructUnderWindow)
+{
+	//////////////////////////////////////////////////
+	// Create & initialize tab manager.
+
+	TabManager = FGlobalTabmanager::Get()->NewTabManager(ConstructUnderMajorTab);
+
+	const auto& PersistLayout = [](const TSharedRef<FTabManager::FLayout>& LayoutToSave)
+	{
+		FLayoutSaveRestore::SaveToConfig(FTraceInsightsModule::GetUnrealInsightsLayoutIni(), LayoutToSave);
+	};
+	TabManager->SetOnPersistLayout(FTabManager::FOnPersistLayout::CreateLambda(PersistLayout));
+
+	//////////////////////////////////////////////////
+
+	TSharedRef<FWorkspaceItem> MenuGroup = CreateWorkspaceMenuGroup();
+	WorkspaceMenuGroup = MenuGroup;
+
+	Extension = MakeShared<FInsightsMajorTabExtender>(TabManager, MenuGroup);
+
+	RegisterTabSpawners();
+
+	//////////////////////////////////////////////////
+
+	FTraceInsightsModule& TraceInsightsModule = FModuleManager::GetModuleChecked<FTraceInsightsModule>("TraceInsights");
+
+	const FOnRegisterMajorTabExtensions* ExtensionDelegate = TraceInsightsModule.FindMajorTabLayoutExtension(GetMajorTabId());
+	if (ExtensionDelegate)
+	{
+		ExtensionDelegate->Broadcast(*Extension);
+	}
+
+	// Register any new minor tabs.
+	for (const FMinorTabConfig& MinorTabConfig : Extension->GetMinorTabs())
+	{
+		FTabSpawnerEntry& TabSpawnerEntry = TabManager->RegisterTabSpawner(MinorTabConfig.TabId, MinorTabConfig.OnSpawnTab, MinorTabConfig.CanSpawnTab);
+
+		TabSpawnerEntry
+		.SetDisplayName(MinorTabConfig.TabLabel)
+		.SetTooltipText(MinorTabConfig.TabTooltip)
+		.SetIcon(MinorTabConfig.TabIcon)
+		.SetReuseTabMethod(MinorTabConfig.OnFindTabToReuse);
+
+		if (MinorTabConfig.WorkspaceGroup.IsValid())
+		{
+			TabSpawnerEntry.SetGroup(MinorTabConfig.WorkspaceGroup.ToSharedRef());
+		}
+	}
+
+	// Check for layout overrides.
+	FInsightsMajorTabConfig TabConfig = TraceInsightsModule.FindMajorTabConfig(GetMajorTabId());
+
+	// Create tab layout.
+	TSharedRef<FTabManager::FLayout> Layout = [&TabConfig, this]() -> TSharedRef<FTabManager::FLayout>
+	{
+		if (TabConfig.Layout.IsValid())
+		{
+			return TabConfig.Layout.ToSharedRef();
+		}
+		else
+		{
+			return CreateDefaultTabLayout();
+		}
+	}();
+
+	Layout->ProcessExtensions(Extension->GetLayoutExtender());
+	Layout = FLayoutSaveRestore::LoadFromConfig(FTraceInsightsModule::GetUnrealInsightsLayoutIni(), Layout);
+
+	// Create & initialize main menu.
+	FMenuBarBuilder MenuBarBuilder = FMenuBarBuilder(TSharedPtr<FUICommandList>(), Extension->GetMenuExtender());
+	MenuBarBuilder.AddPullDownMenu(
+		LOCTEXT("MenuLabel", "Menu"),
+		FText::GetEmpty(),
+		FNewMenuDelegate::CreateStatic(&SMajorTabWindow::FillMenu, TabManager),
+		FName(TEXT("Menu"))
+	);
+
+#if !WITH_EDITOR
+	TSharedRef<SWidget> MenuWidget = MenuBarBuilder.MakeWidget();
+	MenuWidget->SetClipping(EWidgetClipping::ClipToBoundsWithoutIntersecting);
+#endif
+
+	ChildSlot
+	[
+		SNew(SOverlay)
+
+#if !WITH_EDITOR
+		// Menu
+		+ SOverlay::Slot()
+		.HAlign(HAlign_Left)
+		.VAlign(VAlign_Top)
+		.Padding(34.0f, -60.0f, 0.0f, 0.0f)
+		[
+			MenuWidget
+		]
+#endif
+
+		// Version
+		+ SOverlay::Slot()
+		.HAlign(HAlign_Right)
+		.VAlign(VAlign_Top)
+		.Padding(0.0f, -16.0f, 4.0f, 0.0f)
+		[
+			SNew(STextBlock)
+			.Clipping(EWidgetClipping::ClipToBoundsWithoutIntersecting)
+			.Text(LOCTEXT("UnrealInsightsVersion", UNREAL_INSIGHTS_VERSION_STRING_EX))
+			.ColorAndOpacity(FLinearColor(0.15f, 0.15f, 0.15f, 1.0f))
+		]
+
+		// Overlay slot for the main window area
+		+ SOverlay::Slot()
+		.HAlign(HAlign_Fill)
+		.VAlign(VAlign_Fill)
+		[
+			SNew(SVerticalBox)
+
+			+ SVerticalBox::Slot()
+			.AutoHeight()
+			.Padding(FMargin(0.0f, 0.0f, 0.0f, 5.0f))
+			[
+				CreateToolbar(Extension->GetMenuExtender())
+			]
+
+			+ SVerticalBox::Slot()
+			.FillHeight(1.0f)
+			[
+				TabManager->RestoreFrom(Layout, ConstructUnderWindow).ToSharedRef()
+			]
+		]
+
+		// Session hint overlay
+		+ SOverlay::Slot()
+		.HAlign(HAlign_Center)
+		.VAlign(VAlign_Center)
+		[
+			SNew(SBorder)
+			.Visibility_Lambda([this]() { return IsValidSession() ? EVisibility::Collapsed : EVisibility::Visible; })
+			.BorderImage(FAppStyle::Get().GetBrush("PopupText.Background"))
+			.Padding(8.0f)
+			[
+				SNew(STextBlock)
+				.Text(LOCTEXT("SelectTraceOverlayText", "Please select a trace."))
+			]
+		]
+	];
+
+#if !WITH_EDITOR
+	// Tell tab-manager about the global menu bar.
+	TabManager->SetMenuMultiBox(MenuBarBuilder.GetMultiBox(), MenuWidget);
+#endif
+
+	// Tell clients about creation.
+	TraceInsightsModule.OnMajorTabCreated().Broadcast(GetMajorTabId(), TabManager.ToSharedRef());
+}
+END_SLATE_FUNCTION_BUILD_OPTIMIZATION
+
+////////////////////////////////////////////////////////////////////////////////////////////////////
+
+BEGIN_SLATE_FUNCTION_BUILD_OPTIMIZATION
+TSharedRef<FTabManager::FLayout> SMajorTabWindow::CreateDefaultTabLayout() const
+{
+	return FTabManager::NewLayout("InsightsMajorTabLayout_v1.0")
+		->AddArea
+		(
+			FTabManager::NewPrimaryArea()
+			->SetOrientation(Orient_Horizontal)
+			->Split
+			(
+				FTabManager::NewStack()
+				->SetSizeCoefficient(1.0f)
+				->SetHideTabWell(true)
+			)
+		);
+}
+END_SLATE_FUNCTION_BUILD_OPTIMIZATION
+
+////////////////////////////////////////////////////////////////////////////////////////////////////
+
+TSharedRef<SWidget> SMajorTabWindow::CreateToolbar(TSharedPtr<FExtender> Extender)
+{
+	return SNew(SBox);
+}
+
+////////////////////////////////////////////////////////////////////////////////////////////////////
+
+void SMajorTabWindow::FillMenu(FMenuBuilder& MenuBuilder, const TSharedPtr<FTabManager> TabManager)
+{
+	if (!TabManager.IsValid())
+	{
+		return;
+	}
+
+	FInsightsManager::Get()->GetInsightsMenuBuilder()->PopulateMenu(MenuBuilder);
+
+	TabManager->PopulateLocalTabSpawnerMenu(MenuBuilder);
+}
+
+////////////////////////////////////////////////////////////////////////////////////////////////////
+
+bool SMajorTabWindow::IsValidSession() const
+{
+	return FInsightsManager::Get()->GetSession().IsValid();
+}
+
+////////////////////////////////////////////////////////////////////////////////////////////////////
+
+EActiveTimerReturnType SMajorTabWindow::UpdateActiveDuration(double InCurrentTime, float InDeltaTime)
+{
+	DurationActive += InDeltaTime;
+
+	// The window will explicitly unregister this active timer when the mouse leaves.
+	return EActiveTimerReturnType::Continue;
+}
+
+////////////////////////////////////////////////////////////////////////////////////////////////////
+
+void SMajorTabWindow::OnMouseEnter(const FGeometry& MyGeometry, const FPointerEvent& MouseEvent)
+{
+	SCompoundWidget::OnMouseEnter(MyGeometry, MouseEvent);
+
+	if (!ActiveTimerHandle.IsValid())
+	{
+		ActiveTimerHandle = RegisterActiveTimer(0.f, FWidgetActiveTimerDelegate::CreateSP(this, &SMajorTabWindow::UpdateActiveDuration));
+	}
+}
+
+////////////////////////////////////////////////////////////////////////////////////////////////////
+
+void SMajorTabWindow::OnMouseLeave(const FPointerEvent& MouseEvent)
+{
+	SCompoundWidget::OnMouseLeave(MouseEvent);
+
+	auto PinnedActiveTimerHandle = ActiveTimerHandle.Pin();
+	if (PinnedActiveTimerHandle.IsValid())
+	{
+		UnRegisterActiveTimer(PinnedActiveTimerHandle.ToSharedRef());
+	}
+}
+
+////////////////////////////////////////////////////////////////////////////////////////////////////
+
+FReply SMajorTabWindow::OnKeyDown(const FGeometry& MyGeometry, const FKeyEvent& InKeyEvent)
+{
+	if (CommandList.IsValid())
+	{
+		return CommandList->ProcessCommandBindings(InKeyEvent) ? FReply::Handled() : FReply::Unhandled();
+	}
+	else
+	{
+		return SCompoundWidget::OnKeyDown(MyGeometry, InKeyEvent);
+	}
+}
+
+////////////////////////////////////////////////////////////////////////////////////////////////////
+
+FReply SMajorTabWindow::OnDragOver(const FGeometry& MyGeometry, const FDragDropEvent& DragDropEvent)
+{
+	if (FInsightsManager::Get()->OnDragOver(DragDropEvent))
+	{
+		return FReply::Handled();
+	}
+
+	return SCompoundWidget::OnDragOver(MyGeometry, DragDropEvent);
+}
+
+////////////////////////////////////////////////////////////////////////////////////////////////////
+
+FReply SMajorTabWindow::OnDrop(const FGeometry& MyGeometry, const FDragDropEvent& DragDropEvent)
+{
+	if (FInsightsManager::Get()->OnDrop(DragDropEvent))
+	{
+		return FReply::Handled();
+	}
+
+	return SCompoundWidget::OnDrop(MyGeometry, DragDropEvent);
+}
+
+////////////////////////////////////////////////////////////////////////////////////////////////////
+
+} // namespace Insights
+
+#undef LOCTEXT_NAMESPACE
diff --git a/Engine/Source/Developer/TraceInsights/Private/Insights/Widgets/STimingProfilerToolbar.cpp b/Engine/Source/Developer/TraceInsights/Private/Insights/Widgets/STimingProfilerToolbar.cpp
--- a/Engine/Source/Developer/TraceInsights/Private/Insights/Widgets/STimingProfilerToolbar.cpp
+++ b/Engine/Source/Developer/TraceInsights/Private/Insights/Widgets/STimingProfilerToolbar.cpp
@@ -2,6 +2,7 @@
 
 #include "STimingProfilerToolbar.h"
 
+#include "Framework/MultiBox/MultiBoxExtender.h"
 #include "Framework/MultiBox/MultiBoxBuilder.h"
 #include "Widgets/SBoxPanel.h"
 
@@ -58,17 +59,16 @@ void STimingProfilerToolbar::Construct(const FArguments& InArgs)
 				ToolbarBuilder.AddToolBarButton(FTimingProfilerCommands::Get().ToggleLogViewVisibility,
 					NAME_None, TAttribute<FText>(), TAttribute<FText>(),
 					FSlateIcon(FInsightsStyle::GetStyleSetName(), "Icons.LogView.ToolBar"));
-
-				if (InArgs._ToolbarExtender.IsValid())
-				{
-					InArgs._ToolbarExtender->Apply("MainToolbar", EExtensionHook::First, ToolbarBuilder);
-				}
-
 			}
 			ToolbarBuilder.EndSection();
+
+			if (InArgs._ToolbarExtender.IsValid())
+			{
+				InArgs._ToolbarExtender->Apply("MainToolbar", EExtensionHook::First, ToolbarBuilder);
+			}
 		}
 
-		static void FillRightSideToolbar(FToolBarBuilder& ToolbarBuilder)
+		static void FillRightSideToolbar(FToolBarBuilder& ToolbarBuilder, const FArguments &InArgs)
 		{
 			ToolbarBuilder.BeginSection("Debug");
 			{
@@ -77,6 +77,11 @@ void STimingProfilerToolbar::Construct(const FArguments& InArgs)
 					FSlateIcon(FInsightsStyle::GetStyleSetName(), "Icons.Debug.ToolBar"));
 			}
 			ToolbarBuilder.EndSection();
+
+			if (InArgs._ToolbarExtender.IsValid())
+			{
+				InArgs._ToolbarExtender->Apply("RightSideToolbar", EExtensionHook::First, ToolbarBuilder);
+			}
 		}
 	};
 
@@ -88,7 +93,7 @@ void STimingProfilerToolbar::Construct(const FArguments& InArgs)
 
 	FSlimHorizontalToolBarBuilder RightSideToolbarBuilder(CommandList.ToSharedRef(), FMultiBoxCustomization::None);
 	RightSideToolbarBuilder.SetStyle(&FInsightsStyle::Get(), "PrimaryToolbar");
-	Local::FillRightSideToolbar(RightSideToolbarBuilder);
+	Local::FillRightSideToolbar(RightSideToolbarBuilder, InArgs);
 
 	ChildSlot
 	[
diff --git a/Engine/Source/Developer/TraceInsights/Private/Insights/Widgets/STimingProfilerToolbar.h b/Engine/Source/Developer/TraceInsights/Private/Insights/Widgets/STimingProfilerToolbar.h
--- a/Engine/Source/Developer/TraceInsights/Private/Insights/Widgets/STimingProfilerToolbar.h
+++ b/Engine/Source/Developer/TraceInsights/Private/Insights/Widgets/STimingProfilerToolbar.h
@@ -4,7 +4,8 @@
 
 #include "Widgets/DeclarativeSyntaxSupport.h"
 #include "Widgets/SCompoundWidget.h"
-#include "Framework/MultiBox/MultiBoxExtender.h"
+
+class FExtender;
 
 /** Ribbon based toolbar used as a main menu in the Profiler window. */
 class STimingProfilerToolbar : public SCompoundWidget
@@ -26,8 +27,4 @@ public:
 	 * @param InArgs The declaration data for this widget
 	 */
 	void Construct(const FArguments& InArgs);
-
-private:
-	/** Create the UI commands for the toolbar */
-	void CreateCommands();
 };
diff --git a/Engine/Source/Developer/TraceInsights/Private/Insights/Widgets/STimingProfilerWindow.cpp b/Engine/Source/Developer/TraceInsights/Private/Insights/Widgets/STimingProfilerWindow.cpp
--- a/Engine/Source/Developer/TraceInsights/Private/Insights/Widgets/STimingProfilerWindow.cpp
+++ b/Engine/Source/Developer/TraceInsights/Private/Insights/Widgets/STimingProfilerWindow.cpp
@@ -3,36 +3,15 @@
 #include "STimingProfilerWindow.h"
 
 #include "Framework/Commands/UICommandList.h"
-#include "Framework/Docking/LayoutExtender.h"
-#include "Framework/Docking/LayoutService.h"
 #include "Framework/Docking/TabManager.h"
 #include "Framework/Docking/WorkspaceItem.h"
-#include "Framework/MultiBox/MultiBoxBuilder.h"
 #include "SlateOptMacros.h"
-#include "Styling/AppStyle.h"
 #include "Widgets/Docking/SDockTab.h"
-#include "Widgets/Images/SImage.h"
-#include "Widgets/Layout/SBorder.h"
-#include "Widgets/Layout/SBox.h"
-#include "Widgets/SBoxPanel.h"
-#include "Widgets/Text/STextBlock.h"
-#include "Modules/ModuleManager.h"
-
-#if WITH_EDITOR
-	#include "EngineAnalytics.h"
-	#include "Runtime/Analytics/Analytics/Public/AnalyticsEventAttribute.h"
-	#include "Runtime/Analytics/Analytics/Public/Interfaces/IAnalyticsProvider.h"
-#endif // WITH_EDITOR
 
 // Insights
-#include "Insights/Common/InsightsMenuBuilder.h"
-#include "Insights/InsightsManager.h"
 #include "Insights/InsightsStyle.h"
 #include "Insights/TimingProfilerManager.h"
-#include "Insights/TraceInsightsModule.h"
-#include "Insights/Version.h"
 #include "Insights/Widgets/SFrameTrack.h"
-#include "Insights/Widgets/SInsightsSettings.h"
 #include "Insights/Widgets/SLogView.h"
 #include "Insights/Widgets/SStatsView.h"
 #include "Insights/Widgets/STimersView.h"
@@ -55,10 +34,12 @@ const FName FTimingProfilerTabs::CalleesID(TEXT("Callees"));
 const FName FTimingProfilerTabs::StatsCountersID(TEXT("StasCounters"));
 const FName FTimingProfilerTabs::LogViewID(TEXT("LogView"));
 
+////////////////////////////////////////////////////////////////////////////////////////////////////
+// STimingProfilerWindow
 ////////////////////////////////////////////////////////////////////////////////////////////////////
 
 STimingProfilerWindow::STimingProfilerWindow()
-	: DurationActive(0.0f)
+	: SMajorTabWindow(FInsightsManagerTabs::TimingProfilerTabId)
 {
 }
 
@@ -66,19 +47,7 @@ STimingProfilerWindow::STimingProfilerWindow()
 
 STimingProfilerWindow::~STimingProfilerWindow()
 {
-	// Close all tabs.
-	TArray<TSharedPtr<SDockTab>> LocalOpenTabs;
-	for (TSharedPtr<SDockTab>& Tab : OpenTabs)
-	{
-		LocalOpenTabs.Add(Tab);
-	}
-	for (TSharedPtr<SDockTab>& Tab : LocalOpenTabs)
-	{
-		Tab->RequestCloseTab();
-		Tab->RemoveTabFromParent();
-	}
-	LocalOpenTabs.Reset();
-	check(OpenTabs.IsEmpty());
+	CloseAllOpenTabs();
 
 	check(LogView == nullptr);
 	check(StatsView == nullptr);
@@ -87,13 +56,13 @@ STimingProfilerWindow::~STimingProfilerWindow()
 	check(TimersView == nullptr);
 	check(TimingView == nullptr);
 	check(FrameTrack == nullptr);
+}
 
-#if WITH_EDITOR
-	if (DurationActive > 0.0f && FEngineAnalytics::IsAvailable())
-	{
-		FEngineAnalytics::GetProvider().RecordEvent(TEXT("Insights.Usage.TimingProfiler"), FAnalyticsEventAttribute(TEXT("Duration"), DurationActive));
-	}
-#endif // WITH_EDITOR
+////////////////////////////////////////////////////////////////////////////////////////////////////
+
+const TCHAR* STimingProfilerWindow::GetAnalyticsEventName() const
+{
+	return TEXT("Insights.Usage.TimingProfiler");
 }
 
 ////////////////////////////////////////////////////////////////////////////////////////////////////
@@ -151,7 +120,7 @@ TSharedRef<SDockTab> STimingProfilerWindow::SpawnTab_FramesTrack(const FSpawnTab
 		];
 
 	DockTab->SetOnTabClosed(SDockTab::FOnTabClosedCallback::CreateRaw(this, &STimingProfilerWindow::OnFramesTrackTabClosed));
-	OpenTabs.Add(DockTab);
+	AddOpenTab(DockTab);
 
 	return DockTab;
 }
@@ -164,7 +133,7 @@ void STimingProfilerWindow::OnFramesTrackTabClosed(TSharedRef<SDockTab> TabBeing
 	FTimingProfilerManager::Get()->SetFramesTrackVisible(false);
 	FrameTrack = nullptr;
 
-	OpenTabs.Remove(TabBeingClosed);
+	RemoveOpenTab(TabBeingClosed);
 }
 
 ////////////////////////////////////////////////////////////////////////////////////////////////////
@@ -188,7 +157,7 @@ TSharedRef<SDockTab> STimingProfilerWindow::SpawnTab_TimingView(const FSpawnTabA
 	TimingView->SelectTimeInterval(SelectionStartTime, SelectionEndTime - SelectionStartTime);
 
 	DockTab->SetOnTabClosed(SDockTab::FOnTabClosedCallback::CreateRaw(this, &STimingProfilerWindow::OnTimingViewTabClosed));
-	OpenTabs.Add(DockTab);
+	AddOpenTab(DockTab);
 
 	return DockTab;
 }
@@ -205,7 +174,7 @@ void STimingProfilerWindow::OnTimingViewTabClosed(TSharedRef<SDockTab> TabBeingC
 		TimingView = nullptr;
 	}
 
-	OpenTabs.Remove(TabBeingClosed);
+	RemoveOpenTab(TabBeingClosed);
 }
 
 ////////////////////////////////////////////////////////////////////////////////////////////////////
@@ -223,7 +192,7 @@ TSharedRef<SDockTab> STimingProfilerWindow::SpawnTab_Timers(const FSpawnTabArgs&
 		];
 
 	DockTab->SetOnTabClosed(SDockTab::FOnTabClosedCallback::CreateRaw(this, &STimingProfilerWindow::OnTimersTabClosed));
-	OpenTabs.Add(DockTab);
+	AddOpenTab(DockTab);
 
 	return DockTab;
 }
@@ -236,7 +205,7 @@ void STimingProfilerWindow::OnTimersTabClosed(TSharedRef<SDockTab> TabBeingClose
 	FTimingProfilerManager::Get()->SetTimersViewVisible(false);
 	TimersView = nullptr;
 
-	OpenTabs.Remove(TabBeingClosed);
+	RemoveOpenTab(TabBeingClosed);
 }
 
 ////////////////////////////////////////////////////////////////////////////////////////////////////
@@ -254,7 +223,7 @@ TSharedRef<SDockTab> STimingProfilerWindow::SpawnTab_Callers(const FSpawnTabArgs
 		];
 
 	DockTab->SetOnTabClosed(SDockTab::FOnTabClosedCallback::CreateRaw(this, &STimingProfilerWindow::OnCallersTabClosed));
-	OpenTabs.Add(DockTab);
+	AddOpenTab(DockTab);
 
 	return DockTab;
 }
@@ -267,7 +236,7 @@ void STimingProfilerWindow::OnCallersTabClosed(TSharedRef<SDockTab> TabBeingClos
 	FTimingProfilerManager::Get()->SetCallersTreeViewVisible(false);
 	CallersTreeView = nullptr;
 
-	OpenTabs.Remove(TabBeingClosed);
+	RemoveOpenTab(TabBeingClosed);
 }
 
 ////////////////////////////////////////////////////////////////////////////////////////////////////
@@ -285,7 +254,7 @@ TSharedRef<SDockTab> STimingProfilerWindow::SpawnTab_Callees(const FSpawnTabArgs
 		];
 
 	DockTab->SetOnTabClosed(SDockTab::FOnTabClosedCallback::CreateRaw(this, &STimingProfilerWindow::OnCalleesTabClosed));
-	OpenTabs.Add(DockTab);
+	AddOpenTab(DockTab);
 
 	return DockTab;
 }
@@ -298,7 +267,7 @@ void STimingProfilerWindow::OnCalleesTabClosed(TSharedRef<SDockTab> TabBeingClos
 	FTimingProfilerManager::Get()->SetCalleesTreeViewVisible(false);
 	CalleesTreeView = nullptr;
 
-	OpenTabs.Remove(TabBeingClosed);
+	RemoveOpenTab(TabBeingClosed);
 }
 
 ////////////////////////////////////////////////////////////////////////////////////////////////////
@@ -316,7 +285,7 @@ TSharedRef<SDockTab> STimingProfilerWindow::SpawnTab_StatsCounters(const FSpawnT
 		];
 
 	DockTab->SetOnTabClosed(SDockTab::FOnTabClosedCallback::CreateRaw(this, &STimingProfilerWindow::OnStatsCountersTabClosed));
-	OpenTabs.Add(DockTab);
+	AddOpenTab(DockTab);
 
 	return DockTab;
 }
@@ -329,7 +298,7 @@ void STimingProfilerWindow::OnStatsCountersTabClosed(TSharedRef<SDockTab> TabBei
 	FTimingProfilerManager::Get()->SetStatsCountersViewVisible(false);
 	StatsView = nullptr;
 
-	OpenTabs.Remove(TabBeingClosed);
+	RemoveOpenTab(TabBeingClosed);
 }
 
 ////////////////////////////////////////////////////////////////////////////////////////////////////
@@ -347,7 +316,7 @@ TSharedRef<SDockTab> STimingProfilerWindow::SpawnTab_LogView(const FSpawnTabArgs
 		];
 
 	DockTab->SetOnTabClosed(SDockTab::FOnTabClosedCallback::CreateRaw(this, &STimingProfilerWindow::OnLogViewTabClosed));
-	OpenTabs.Add(DockTab);
+	AddOpenTab(DockTab);
 
 	return DockTab;
 }
@@ -360,375 +329,148 @@ void STimingProfilerWindow::OnLogViewTabClosed(TSharedRef<SDockTab> TabBeingClos
 	FTimingProfilerManager::Get()->SetLogViewVisible(false);
 	LogView = nullptr;
 
-	OpenTabs.Remove(TabBeingClosed);
+	RemoveOpenTab(TabBeingClosed);
 }
 
 ////////////////////////////////////////////////////////////////////////////////////////////////////
 
-BEGIN_SLATE_FUNCTION_BUILD_OPTIMIZATION
 void STimingProfilerWindow::Construct(const FArguments& InArgs, const TSharedRef<SDockTab>& ConstructUnderMajorTab, const TSharedPtr<SWindow>& ConstructUnderWindow)
 {
-	// Create & initialize tab manager.
-	TabManager = FGlobalTabmanager::Get()->NewTabManager(ConstructUnderMajorTab);
+	TSharedPtr<FTimingProfilerManager> TimingProfilerManager = FTimingProfilerManager::Get();
+	ensure(TimingProfilerManager.IsValid());
 
-	const auto& PersistLayout = [](const TSharedRef<FTabManager::FLayout>& LayoutToSave)
-	{
-		FLayoutSaveRestore::SaveToConfig(FTraceInsightsModule::GetUnrealInsightsLayoutIni(), LayoutToSave);
-	};
-	TabManager->SetOnPersistLayout(FTabManager::FOnPersistLayout::CreateLambda(PersistLayout));
+	SetCommandList(TimingProfilerManager->GetCommandList());
 
-	TSharedRef<FWorkspaceItem> AppMenuGroup = TabManager->AddLocalWorkspaceMenuCategory(LOCTEXT("TimingProfilerMenuGroupName", "Timing Insights"));
+	SMajorTabWindow::FArguments Args;
+	SMajorTabWindow::Construct(Args, ConstructUnderMajorTab, ConstructUnderWindow);
+}
+
+////////////////////////////////////////////////////////////////////////////////////////////////////
+
+TSharedRef<FWorkspaceItem> STimingProfilerWindow::CreateWorkspaceMenuGroup()
+{
+	return GetTabManager()->AddLocalWorkspaceMenuCategory(LOCTEXT("TimingProfilerMenuGroupName", "Timing Insights"));
+}
 
-	Extension = MakeShared<FInsightsMajorTabExtender>(TabManager, AppMenuGroup);
+////////////////////////////////////////////////////////////////////////////////////////////////////
 
-	TabManager->RegisterTabSpawner(FTimingProfilerTabs::FramesTrackID, FOnSpawnTab::CreateRaw(this, &STimingProfilerWindow::SpawnTab_FramesTrack))
+void STimingProfilerWindow::RegisterTabSpawners()
+{
+	check(GetTabManager().IsValid());
+	FTabManager* TabManagerPtr = GetTabManager().Get();
+	check(GetWorkspaceMenuGroup().IsValid());
+	const TSharedRef<FWorkspaceItem> Group = GetWorkspaceMenuGroup().ToSharedRef();
+
+	TabManagerPtr->RegisterTabSpawner(FTimingProfilerTabs::FramesTrackID, FOnSpawnTab::CreateRaw(this, &STimingProfilerWindow::SpawnTab_FramesTrack))
 		.SetDisplayName(LOCTEXT("FramesTrackTabTitle", "Frames"))
 		.SetIcon(FSlateIcon(FInsightsStyle::GetStyleSetName(), "Icons.FramesTrack"))
-		.SetGroup(AppMenuGroup);
+		.SetGroup(Group);
 
-	TabManager->RegisterTabSpawner(FTimingProfilerTabs::TimingViewID, FOnSpawnTab::CreateRaw(this, &STimingProfilerWindow::SpawnTab_TimingView))
+	TabManagerPtr->RegisterTabSpawner(FTimingProfilerTabs::TimingViewID, FOnSpawnTab::CreateRaw(this, &STimingProfilerWindow::SpawnTab_TimingView))
 		.SetDisplayName(LOCTEXT("TimingViewTabTitle", "Timing View"))
 		.SetIcon(FSlateIcon(FInsightsStyle::GetStyleSetName(), "Icons.TimingView"))
-		.SetGroup(AppMenuGroup);
+		.SetGroup(Group);
 
-	TabManager->RegisterTabSpawner(FTimingProfilerTabs::TimersID, FOnSpawnTab::CreateRaw(this, &STimingProfilerWindow::SpawnTab_Timers))
+	TabManagerPtr->RegisterTabSpawner(FTimingProfilerTabs::TimersID, FOnSpawnTab::CreateRaw(this, &STimingProfilerWindow::SpawnTab_Timers))
 		.SetDisplayName(LOCTEXT("TimersTabTitle", "Timers"))
 		.SetIcon(FSlateIcon(FInsightsStyle::GetStyleSetName(), "Icons.TimersView"))
-		.SetGroup(AppMenuGroup);
+		.SetGroup(Group);
 
-	TabManager->RegisterTabSpawner(FTimingProfilerTabs::CallersID, FOnSpawnTab::CreateRaw(this, &STimingProfilerWindow::SpawnTab_Callers))
+	TabManagerPtr->RegisterTabSpawner(FTimingProfilerTabs::CallersID, FOnSpawnTab::CreateRaw(this, &STimingProfilerWindow::SpawnTab_Callers))
 		.SetDisplayName(LOCTEXT("CallersTabTitle", "Callers"))
 		.SetIcon(FSlateIcon(FInsightsStyle::GetStyleSetName(), "Icons.CallersView"))
-		.SetGroup(AppMenuGroup);
+		.SetGroup(Group);
 
-	TabManager->RegisterTabSpawner(FTimingProfilerTabs::CalleesID, FOnSpawnTab::CreateRaw(this, &STimingProfilerWindow::SpawnTab_Callees))
+	TabManagerPtr->RegisterTabSpawner(FTimingProfilerTabs::CalleesID, FOnSpawnTab::CreateRaw(this, &STimingProfilerWindow::SpawnTab_Callees))
 		.SetDisplayName(LOCTEXT("CalleesTabTitle", "Callees"))
 		.SetIcon(FSlateIcon(FInsightsStyle::GetStyleSetName(), "Icons.CalleesView"))
-		.SetGroup(AppMenuGroup);
+		.SetGroup(Group);
 
-	TabManager->RegisterTabSpawner(FTimingProfilerTabs::StatsCountersID, FOnSpawnTab::CreateRaw(this, &STimingProfilerWindow::SpawnTab_StatsCounters))
+	TabManagerPtr->RegisterTabSpawner(FTimingProfilerTabs::StatsCountersID, FOnSpawnTab::CreateRaw(this, &STimingProfilerWindow::SpawnTab_StatsCounters))
 		.SetDisplayName(LOCTEXT("StatsCountersTabTitle", "Counters"))
 		.SetIcon(FSlateIcon(FInsightsStyle::GetStyleSetName(), "Icons.CountersView"))
-		.SetGroup(AppMenuGroup);
+		.SetGroup(Group);
 
-	TabManager->RegisterTabSpawner(FTimingProfilerTabs::LogViewID, FOnSpawnTab::CreateRaw(this, &STimingProfilerWindow::SpawnTab_LogView))
+	TabManagerPtr->RegisterTabSpawner(FTimingProfilerTabs::LogViewID, FOnSpawnTab::CreateRaw(this, &STimingProfilerWindow::SpawnTab_LogView))
 		.SetDisplayName(LOCTEXT("LogViewTabTitle", "Log View"))
 		.SetIcon(FSlateIcon(FInsightsStyle::GetStyleSetName(), "Icons.LogView"))
-		.SetGroup(AppMenuGroup);
-
-	TSharedPtr<FTimingProfilerManager> TimingProfilerManager = FTimingProfilerManager::Get();
-	ensure(TimingProfilerManager.IsValid());
-
-	// Check for layout overrides.
-	FTraceInsightsModule& TraceInsightsModule = FModuleManager::GetModuleChecked<FTraceInsightsModule>("TraceInsights");
-	FInsightsMajorTabConfig TabConfig = TraceInsightsModule.FindMajorTabConfig(FInsightsManagerTabs::TimingProfilerTabId);
-
-	const FOnRegisterMajorTabExtensions* ExtensionDelegate = TraceInsightsModule.FindMajorTabLayoutExtension(FInsightsManagerTabs::TimingProfilerTabId);
-	if (ExtensionDelegate)
-	{
-		ExtensionDelegate->Broadcast(*Extension);
-	}
-
-	// Register any new minor tabs.
-	for (const FMinorTabConfig& MinorTabConfig : Extension->GetMinorTabs())
-	{
-		FTabSpawnerEntry& TabSpawnerEntry = TabManager->RegisterTabSpawner(MinorTabConfig.TabId, MinorTabConfig.OnSpawnTab, MinorTabConfig.CanSpawnTab);
-
-		TabSpawnerEntry
-		.SetDisplayName(MinorTabConfig.TabLabel)
-		.SetTooltipText(MinorTabConfig.TabTooltip)
-		.SetIcon(MinorTabConfig.TabIcon)
-		.SetReuseTabMethod(MinorTabConfig.OnFindTabToReuse);
-
-		if (MinorTabConfig.WorkspaceGroup.IsValid())
-		{
-			TabSpawnerEntry.SetGroup(MinorTabConfig.WorkspaceGroup.ToSharedRef());
-		}
-	}
-
-	TSharedRef<FTabManager::FLayout> Layout = [&TabConfig]() -> TSharedRef<FTabManager::FLayout>
-	{
-		if (TabConfig.Layout.IsValid())
-		{
-			return TabConfig.Layout.ToSharedRef();
-		}
-		else
-		{
-			// Create tab layout.
-			return FTabManager::NewLayout("InsightsTimingProfilerLayout_v1.2")
-				->AddArea
-				(
-					FTabManager::NewPrimaryArea()
-					->SetOrientation(Orient_Horizontal)
-					->Split
-					(
-						FTabManager::NewSplitter()
-						->SetSizeCoefficient(0.65f)
-						->SetOrientation(Orient_Vertical)
-						->Split
-						(
-							FTabManager::NewStack()
-							->SetSizeCoefficient(0.1f)
-							->SetHideTabWell(true)
-							->AddTab(FTimingProfilerTabs::FramesTrackID, ETabState::OpenedTab)
-						)
-						->Split
-						(
-							FTabManager::NewStack()
-							->SetSizeCoefficient(0.5f)
-							->SetHideTabWell(true)
-							->AddTab(FTimingProfilerTabs::TimingViewID, ETabState::OpenedTab)
-						)
-						->Split
-						(
-							FTabManager::NewStack()
-							->SetSizeCoefficient(0.2f)
-							->SetHideTabWell(true)
-							->AddTab(FTimingProfilerTabs::LogViewID, ETabState::OpenedTab)
-						)
-					)
-					->Split
-					(
-						FTabManager::NewSplitter()
-						->SetSizeCoefficient(0.35f)
-						->SetOrientation(Orient_Vertical)
-						->Split
-						(
-							FTabManager::NewStack()
-							->SetSizeCoefficient(0.67f)
-							->AddTab(FTimingProfilerTabs::TimersID, ETabState::OpenedTab)
-							->AddTab(FTimingProfilerTabs::StatsCountersID, ETabState::OpenedTab)
-							->SetForegroundTab(FTimingProfilerTabs::TimersID)
-						)
-						->Split
-						(
-							FTabManager::NewStack()
-							->SetSizeCoefficient(0.165f)
-							->SetHideTabWell(true)
-							->AddTab(FTimingProfilerTabs::CallersID, ETabState::OpenedTab)
-						)
-						->Split
-						(
-							FTabManager::NewStack()
-							->SetSizeCoefficient(0.165f)
-							->SetHideTabWell(true)
-							->AddTab(FTimingProfilerTabs::CalleesID, ETabState::OpenedTab)
-						)
-					)
-				);
-		}
-	}();
-
-	Layout->ProcessExtensions(Extension->GetLayoutExtender());
-	Layout = FLayoutSaveRestore::LoadFromConfig(FTraceInsightsModule::GetUnrealInsightsLayoutIni(), Layout);
-
-	// Create & initialize main menu.
-	FMenuBarBuilder MenuBarBuilder = FMenuBarBuilder(TSharedPtr<FUICommandList>(), Extension->GetMenuExtender());
-	MenuBarBuilder.AddPullDownMenu(
-		LOCTEXT("MenuLabel", "Menu"),
-		FText::GetEmpty(),
-		FNewMenuDelegate::CreateStatic(&STimingProfilerWindow::FillMenu, TabManager),
-		FName(TEXT("Menu"))
-	);
-
-	TSharedRef<SWidget> MenuWidget = MenuBarBuilder.MakeWidget();
-	MenuWidget->SetClipping(EWidgetClipping::ClipToBoundsWithoutIntersecting);
-
-	ChildSlot
-	[
-		SNew(SOverlay)
-
-#if !WITH_EDITOR
-		// Menu
-		+ SOverlay::Slot()
-		.HAlign(HAlign_Left)
-		.VAlign(VAlign_Top)
-		.Padding(34.0f, -60.0f, 0.0f, 0.0f)
-		[
-			MenuWidget
-		]
-#endif
-
-		// Version
-		+ SOverlay::Slot()
-		.HAlign(HAlign_Right)
-		.VAlign(VAlign_Top)
-		.Padding(0.0f, -16.0f, 4.0f, 0.0f)
-		[
-			SNew(STextBlock)
-			.Clipping(EWidgetClipping::ClipToBoundsWithoutIntersecting)
-			.Text(LOCTEXT("UnrealInsightsVersion", UNREAL_INSIGHTS_VERSION_STRING_EX))
-			.ColorAndOpacity(FLinearColor(0.15f, 0.15f, 0.15f, 1.0f))
-		]
-
-		// Overlay slot for the main window area
-		+ SOverlay::Slot()
-		.HAlign(HAlign_Fill)
-		.VAlign(VAlign_Fill)
-		[
-			SNew(SVerticalBox)
-			
-			+ SVerticalBox::Slot()
-			.AutoHeight()
-			.Padding(FMargin(0.0f, 0.0f, 0.0f, 5.0f))
-			[
-				SNew(STimingProfilerToolbar).ToolbarExtender(Extension->GetMenuExtender())
-			]
-
-			+ SVerticalBox::Slot()
-			.FillHeight(1.0f)
-			[
-				TabManager->RestoreFrom(Layout, ConstructUnderWindow).ToSharedRef()
-			]
-		]
-
-		// Session hint overlay
-		+ SOverlay::Slot()
-		.HAlign(HAlign_Center)
-		.VAlign(VAlign_Center)
-		[
-			SNew(SBorder)
-			.Visibility(this, &STimingProfilerWindow::IsSessionOverlayVisible)
-			.BorderImage(FAppStyle::Get().GetBrush("PopupText.Background"))
-			.Padding(8.0f)
-			[
-				SNew(STextBlock)
-				.Text(LOCTEXT("SelectTraceOverlayText", "Please select a trace."))
-			]
-		]
-	];
-
-#if !WITH_EDITOR
-	// Tell tab-manager about the global menu bar.
-	TabManager->SetMenuMultiBox(MenuBarBuilder.GetMultiBox(), MenuWidget);
-#endif
-
-	// Tell clients about creation
-	TraceInsightsModule.OnMajorTabCreated().Broadcast(FInsightsManagerTabs::TimingProfilerTabId, TabManager.ToSharedRef());
-}
-END_SLATE_FUNCTION_BUILD_OPTIMIZATION
-
-////////////////////////////////////////////////////////////////////////////////////////////////////
-
-void STimingProfilerWindow::FillMenu(FMenuBuilder& MenuBuilder, const TSharedPtr<FTabManager> TabManager)
-{
-	if (!TabManager.IsValid())
-	{
-		return;
-	}
-
-	FInsightsManager::Get()->GetInsightsMenuBuilder()->PopulateMenu(MenuBuilder);
-
-	TabManager->PopulateLocalTabSpawnerMenu(MenuBuilder);
-}
-
-////////////////////////////////////////////////////////////////////////////////////////////////////
-
-void STimingProfilerWindow::ShowTab(const FName& TabID)
-{
-	if (TabManager->HasTabSpawner(TabID))
-	{
-		TabManager->TryInvokeTab(TabID);
-	}
-}
-
-////////////////////////////////////////////////////////////////////////////////////////////////////
-
-void STimingProfilerWindow::HideTab(const FName& TabID)
-{
-	TSharedPtr<SDockTab> Tab = TabManager->FindExistingLiveTab(TabID);
-	if (Tab.IsValid())
-	{
-		Tab->RequestCloseTab();
-		Tab->RemoveTabFromParent();
-	}
-}
-
-////////////////////////////////////////////////////////////////////////////////////////////////////
-
-EVisibility STimingProfilerWindow::IsSessionOverlayVisible() const
-{
-	if (FInsightsManager::Get()->GetSession().IsValid())
-	{
-		return EVisibility::Hidden;
-	}
-	else
-	{
-		return EVisibility::Visible;
-	}
-}
-
-////////////////////////////////////////////////////////////////////////////////////////////////////
-
-bool STimingProfilerWindow::IsProfilerEnabled() const
-{
-	return FInsightsManager::Get()->GetSession().IsValid();
-}
-
-////////////////////////////////////////////////////////////////////////////////////////////////////
-
-EActiveTimerReturnType STimingProfilerWindow::UpdateActiveDuration(double InCurrentTime, float InDeltaTime)
-{
-	DurationActive += InDeltaTime;
-
-	// The profiler window will explicitly unregister this active timer when the mouse leaves.
-	return EActiveTimerReturnType::Continue;
-}
-
-////////////////////////////////////////////////////////////////////////////////////////////////////
-
-void STimingProfilerWindow::OnMouseEnter(const FGeometry& MyGeometry, const FPointerEvent& MouseEvent)
-{
-	SCompoundWidget::OnMouseEnter(MyGeometry, MouseEvent);
-
-	if (!ActiveTimerHandle.IsValid())
-	{
-		ActiveTimerHandle = RegisterActiveTimer(0.f, FWidgetActiveTimerDelegate::CreateSP(this, &STimingProfilerWindow::UpdateActiveDuration));
-	}
-}
-
-////////////////////////////////////////////////////////////////////////////////////////////////////
-
-void STimingProfilerWindow::OnMouseLeave(const FPointerEvent& MouseEvent)
-{
-	SCompoundWidget::OnMouseLeave(MouseEvent);
-
-	auto PinnedActiveTimerHandle = ActiveTimerHandle.Pin();
-	if (PinnedActiveTimerHandle.IsValid())
-	{
-		UnRegisterActiveTimer(PinnedActiveTimerHandle.ToSharedRef());
-	}
-}
-
-////////////////////////////////////////////////////////////////////////////////////////////////////
-
-FReply STimingProfilerWindow::OnKeyDown(const FGeometry& MyGeometry, const FKeyEvent& InKeyEvent)
-{
-	return FTimingProfilerManager::Get()->GetCommandList()->ProcessCommandBindings(InKeyEvent) ? FReply::Handled() : FReply::Unhandled();
+		.SetGroup(Group);
 }
 
 ////////////////////////////////////////////////////////////////////////////////////////////////////
 
-FReply STimingProfilerWindow::OnDragOver(const FGeometry& MyGeometry, const FDragDropEvent& DragDropEvent)
+BEGIN_SLATE_FUNCTION_BUILD_OPTIMIZATION
+TSharedRef<FTabManager::FLayout> STimingProfilerWindow::CreateDefaultTabLayout() const
 {
-	if (FInsightsManager::Get()->OnDragOver(DragDropEvent))
-	{
-		return FReply::Handled();
-	}
-
-	return SCompoundWidget::OnDragOver(MyGeometry, DragDropEvent);
+	return FTabManager::NewLayout("InsightsTimingProfilerLayout_v1.2")
+		->AddArea
+		(
+			FTabManager::NewPrimaryArea()
+			->SetOrientation(Orient_Horizontal)
+			->Split
+			(
+				FTabManager::NewSplitter()
+				->SetSizeCoefficient(0.65f)
+				->SetOrientation(Orient_Vertical)
+				->Split
+				(
+					FTabManager::NewStack()
+					->SetSizeCoefficient(0.1f)
+					->SetHideTabWell(true)
+					->AddTab(FTimingProfilerTabs::FramesTrackID, ETabState::OpenedTab)
+				)
+				->Split
+				(
+					FTabManager::NewStack()
+					->SetSizeCoefficient(0.5f)
+					->SetHideTabWell(true)
+					->AddTab(FTimingProfilerTabs::TimingViewID, ETabState::OpenedTab)
+				)
+				->Split
+				(
+					FTabManager::NewStack()
+					->SetSizeCoefficient(0.2f)
+					->SetHideTabWell(true)
+					->AddTab(FTimingProfilerTabs::LogViewID, ETabState::OpenedTab)
+				)
+			)
+			->Split
+			(
+				FTabManager::NewSplitter()
+				->SetSizeCoefficient(0.35f)
+				->SetOrientation(Orient_Vertical)
+				->Split
+				(
+					FTabManager::NewStack()
+					->SetSizeCoefficient(0.67f)
+					->AddTab(FTimingProfilerTabs::TimersID, ETabState::OpenedTab)
+					->AddTab(FTimingProfilerTabs::StatsCountersID, ETabState::OpenedTab)
+					->SetForegroundTab(FTimingProfilerTabs::TimersID)
+				)
+				->Split
+				(
+					FTabManager::NewStack()
+					->SetSizeCoefficient(0.165f)
+					->SetHideTabWell(true)
+					->AddTab(FTimingProfilerTabs::CallersID, ETabState::OpenedTab)
+				)
+				->Split
+				(
+					FTabManager::NewStack()
+					->SetSizeCoefficient(0.165f)
+					->SetHideTabWell(true)
+					->AddTab(FTimingProfilerTabs::CalleesID, ETabState::OpenedTab)
+				)
+			)
+		);
 }
+END_SLATE_FUNCTION_BUILD_OPTIMIZATION
 
 ////////////////////////////////////////////////////////////////////////////////////////////////////
 
-FReply STimingProfilerWindow::OnDrop(const FGeometry& MyGeometry, const FDragDropEvent& DragDropEvent)
+TSharedRef<SWidget> STimingProfilerWindow::CreateToolbar(TSharedPtr<FExtender> Extender)
 {
-	if (FInsightsManager::Get()->OnDrop(DragDropEvent))
-	{
-		return FReply::Handled();
-	}
-
-	return SCompoundWidget::OnDrop(MyGeometry, DragDropEvent);
+	return SNew(STimingProfilerToolbar).ToolbarExtender(Extender);
 }
 
 ////////////////////////////////////////////////////////////////////////////////////////////////////
diff --git a/Engine/Source/Developer/TraceInsights/Private/Insights/Widgets/STimingProfilerWindow.h b/Engine/Source/Developer/TraceInsights/Private/Insights/Widgets/STimingProfilerWindow.h
--- a/Engine/Source/Developer/TraceInsights/Private/Insights/Widgets/STimingProfilerWindow.h
+++ b/Engine/Source/Developer/TraceInsights/Private/Insights/Widgets/STimingProfilerWindow.h
@@ -3,36 +3,24 @@
 #pragma once
 
 #include "CoreMinimal.h"
-#include "Framework/Docking/TabManager.h"
-#include "Input/Reply.h"
-#include "Layout/Visibility.h"
-#include "Misc/Guid.h"
-#include "SlateFwd.h"
-#include "Widgets/DeclarativeSyntaxSupport.h"
-#include "Widgets/Docking/SDockTab.h"
-#include "Widgets/SCompoundWidget.h"
 
 // Insights
-#include "Insights/InsightsManager.h"
-#include "Insights/ITimingViewSession.h"
+#include "Insights/ITimingViewSession.h" // for Insights::ETimeChangedFlags
+#include "Insights/Widgets/SMajorTabWindow.h"
 
 ////////////////////////////////////////////////////////////////////////////////////////////////////
 
-class FActiveTimerHandle;
-class FMenuBuilder;
-
 class SFrameTrack;
 class SLogView;
 class SStatsView;
 class STimersView;
 class STimerTreeView;
 class STimingView;
-class SVerticalBox;
 
 ////////////////////////////////////////////////////////////////////////////////////////////////////
 
-/** Implements the timing profiler window. */
-class STimingProfilerWindow : public SCompoundWidget
+/** Implements the Timing Insights major tab window. */
+class STimingProfilerWindow : public Insights::SMajorTabWindow
 {
 public:
 	/** Default constructor. */
@@ -44,17 +32,11 @@ public:
 	SLATE_BEGIN_ARGS(STimingProfilerWindow) {}
 	SLATE_END_ARGS()
 
-	void Reset();
+	virtual void Reset() override;
 
 	/** Constructs this widget. */
 	void Construct(const FArguments& InArgs, const TSharedRef<SDockTab>& ConstructUnderMajorTab, const TSharedPtr<SWindow>& ConstructUnderWindow);
 
-	void ShowTab(const FName& TabID);
-	void HideTab(const FName& TabID);
-	void ShowHideTab(const FName& TabID, bool bShow) { bShow ? ShowTab(TabID) : HideTab(TabID); }
-
-	TSharedPtr<FTabManager> GetTabManager() const { return TabManager; }
-
 	TSharedPtr<STimingView> GetTimingView() const { return TimingView; }
 	TSharedPtr<STimersView> GetTimersView() const { return TimersView; }
 	TSharedPtr<STimerTreeView> GetCallersTreeView() const { return CallersTreeView; }
@@ -62,6 +44,13 @@ public:
 	TSharedPtr<SStatsView> GetStatsView() const { return StatsView; }
 	TSharedPtr<SLogView> GetLogView() const { return LogView; }
 
+protected:
+	virtual const TCHAR* GetAnalyticsEventName() const override;
+	virtual TSharedRef<FWorkspaceItem> CreateWorkspaceMenuGroup() override;
+	virtual void RegisterTabSpawners() override;
+	virtual TSharedRef<FTabManager::FLayout> CreateDefaultTabLayout() const override;
+	virtual TSharedRef<SWidget> CreateToolbar(TSharedPtr<FExtender> Extender);
+
 private:
 	TSharedRef<SDockTab> SpawnTab_FramesTrack(const FSpawnTabArgs& Args);
 	void OnFramesTrackTabClosed(TSharedRef<SDockTab> TabBeingClosed);
@@ -84,86 +73,9 @@ private:
 	TSharedRef<SDockTab> SpawnTab_LogView(const FSpawnTabArgs& Args);
 	void OnLogViewTabClosed(TSharedRef<SDockTab> TabBeingClosed);
 
-	/**
-	 * Fill the main menu with menu items.
-	 *
-	 * @param MenuBuilder The multi-box builder that should be filled with content for this pull-down menu.
-	 * @param TabManager A Tab Manager from which to populate tab spawner menu items.
-	 */
-	static void FillMenu(FMenuBuilder& MenuBuilder, const TSharedPtr<FTabManager> TabManager);
-
-	/** Callback for determining the visibility of the 'Select a session' overlay. */
-	EVisibility IsSessionOverlayVisible() const;
-
-	/** Callback for getting the enabled state of the profiler window. */
-	bool IsProfilerEnabled() const;
-
-	/** Updates the amount of time the profiler has been active. */
-	EActiveTimerReturnType UpdateActiveDuration(double InCurrentTime, float InDeltaTime);
-
-	/**
-	 * Ticks this widget. Override in derived classes, but always call the parent implementation.
-	 *
-	 * @param  AllottedGeometry The space allotted for this widget
-	 * @param  InCurrentTime  Current absolute real time
-	 * @param  InDeltaTime  Real time passed since last tick
-	 */
-	//virtual void Tick(const FGeometry& AllottedGeometry, const double InCurrentTime, const float InDeltaTime) override;
-
-	/**
-	 * The system will use this event to notify a widget that the cursor has entered it. This event is NOT bubbled.
-	 *
-	 * @param MyGeometry The Geometry of the widget receiving the event
-	 * @param MouseEvent Information about the input event
-	 */
-	virtual void OnMouseEnter(const FGeometry& MyGeometry, const FPointerEvent& MouseEvent) override;
-
-	/**
-	 * The system will use this event to notify a widget that the cursor has left it. This event is NOT bubbled.
-	 *
-	 * @param MouseEvent Information about the input event
-	 */
-	virtual void OnMouseLeave(const FPointerEvent& MouseEvent) override;
-
-	/**
-	 * Called after a key is pressed when this widget has focus
-	 *
-	 * @param MyGeometry The Geometry of the widget receiving the event
-	 * @param  InKeyEvent  Key event
-	 *
-	 * @return  Returns whether the event was handled, along with other possible actions
-	 */
-	virtual FReply OnKeyDown(const FGeometry& MyGeometry, const FKeyEvent& InKeyEvent) override;
-
-	/**
-	 * Called when the user is dropping something onto a widget; terminates drag and drop.
-	 *
-	 * @param MyGeometry      The geometry of the widget receiving the event.
-	 * @param DragDropEvent   The drag and drop event.
-	 *
-	 * @return A reply that indicated whether this event was handled.
-	 */
-	virtual FReply OnDrop(const FGeometry& MyGeometry, const FDragDropEvent& DragDropEvent) override;
-
-	/**
-	 * Called during drag and drop when the the mouse is being dragged over a widget.
-	 *
-	 * @param MyGeometry      The geometry of the widget receiving the event.
-	 * @param DragDropEvent   The drag and drop event.
-	 *
-	 * @return A reply that indicated whether this event was handled.
-	 */
-	virtual FReply OnDragOver(const FGeometry& MyGeometry, const FDragDropEvent& DragDropEvent)  override;
-
 	void OnTimeSelectionChanged(Insights::ETimeChangedFlags InFlags, double InStartTime, double InEndTime);
 
 private:
-	/** Holds the tab manager that manages the front-end's tabs. */
-	TSharedPtr<FTabManager> TabManager;
-
-	/** All tabs owned by this window */
-	TSet<TSharedPtr<SDockTab>> OpenTabs;
-
 	/** The Frame track widget */
 	TSharedPtr<SFrameTrack> FrameTrack;
 
@@ -184,13 +96,4 @@ private:
 
 	/** The Log view widget */
 	TSharedPtr<SLogView> LogView;
-
-	/** The handle to the active update duration tick */
-	TWeakPtr<FActiveTimerHandle> ActiveTimerHandle;
-
-	/** The number of seconds the profiler has been active */
-	float DurationActive;
-
-	/** Tab specific slate extender structure, has the lifetime of this widget */
-	TSharedPtr<FInsightsMajorTabExtender> Extension;
 };
diff --git a/Engine/Source/Developer/TraceInsights/Private/Insights/Widgets/STimingView.cpp b/Engine/Source/Developer/TraceInsights/Private/Insights/Widgets/STimingView.cpp
--- a/Engine/Source/Developer/TraceInsights/Private/Insights/Widgets/STimingView.cpp
+++ b/Engine/Source/Developer/TraceInsights/Private/Insights/Widgets/STimingView.cpp
@@ -4832,7 +4832,6 @@ void STimingView::CloseQuickFindTab()
 		if (QuickFindTab.IsValid())
 		{
 			QuickFindTab->RequestCloseTab();
-			QuickFindTab->RemoveTabFromParent();
 		}
 	}
 }
diff --git a/Engine/Source/Developer/TraceInsights/Public/Insights/Widgets/SMajorTabWindow.h b/Engine/Source/Developer/TraceInsights/Public/Insights/Widgets/SMajorTabWindow.h
new file mode 100644
--- /dev/null
+++ b/Engine/Source/Developer/TraceInsights/Public/Insights/Widgets/SMajorTabWindow.h
@@ -0,0 +1,153 @@
+// Copyright Epic Games, Inc. All Rights Reserved.
+
+#pragma once
+
+#include "CoreMinimal.h"
+#include "Framework/Docking/TabManager.h"
+#include "Input/Reply.h"
+#include "Widgets/DeclarativeSyntaxSupport.h"
+#include "Widgets/SCompoundWidget.h"
+
+// Insights
+#include "Insights/IUnrealInsightsModule.h" // for FInsightsMajorTabExtender
+
+////////////////////////////////////////////////////////////////////////////////////////////////////
+
+class FActiveTimerHandle;
+class FMenuBuilder;
+class FUICommandList;
+class SDockTab;
+class SWindow;
+
+////////////////////////////////////////////////////////////////////////////////////////////////////
+
+namespace Insights
+{
+
+/** Implements the base class for a major tab window. */
+class SMajorTabWindow : public SCompoundWidget
+{
+public:
+	/** Default constructor. */
+	SMajorTabWindow(const FName& InMajorTabId);
+
+	/** Virtual destructor. */
+	virtual ~SMajorTabWindow();
+
+	SLATE_BEGIN_ARGS(SMajorTabWindow) {}
+	SLATE_END_ARGS()
+
+	const FName& GetMajorTabId() const { return MajorTabId; }
+	TSharedPtr<FTabManager> GetTabManager() const { return TabManager; }
+	TSharedPtr<FWorkspaceItem> GetWorkspaceMenuGroup() const { return WorkspaceMenuGroup; }
+	const TSharedPtr<FUICommandList> GetCommandList() const { return CommandList; }
+
+	virtual void Reset();
+
+	/** Constructs this widget. */
+	void Construct(const FArguments& InArgs, const TSharedRef<SDockTab>& ConstructUnderMajorTab, const TSharedPtr<SWindow>& ConstructUnderWindow);
+
+	void ShowTab(const FName& TabId);
+	void HideTab(const FName& TabId);
+	void ShowHideTab(const FName& TabId, bool bShow) { bShow ? ShowTab(TabId) : HideTab(TabId); }
+	void CloseAllOpenTabs();
+
+protected:
+	virtual const TCHAR* GetAnalyticsEventName() const;
+	virtual TSharedRef<FWorkspaceItem> CreateWorkspaceMenuGroup();
+	virtual void RegisterTabSpawners();
+	virtual TSharedRef<FTabManager::FLayout> CreateDefaultTabLayout() const;
+	virtual TSharedRef<SWidget> CreateToolbar(TSharedPtr<FExtender> Extender);
+
+	void SetCommandList(const TSharedPtr<FUICommandList> InCommandList) { CommandList = InCommandList; }
+
+	void AddOpenTab(const TSharedRef<SDockTab>& DockTab);
+	void RemoveOpenTab(const TSharedRef<SDockTab>& DockTab);
+
+private:
+	/**
+	 * Fill the main menu with menu items.
+	 *
+	 * @param MenuBuilder The multi-box builder that should be filled with content for this pull-down menu.
+	 * @param TabManager A Tab Manager from which to populate tab spawner menu items.
+	 */
+	static void FillMenu(FMenuBuilder& MenuBuilder, const TSharedPtr<FTabManager> TabManager);
+
+	/** Returns true if the current analysis session is valid. */
+	bool IsValidSession() const;
+
+	/** Updates the amount of time the window has been active. */
+	EActiveTimerReturnType UpdateActiveDuration(double InCurrentTime, float InDeltaTime);
+
+	/**
+	 * The system will use this event to notify a widget that the cursor has entered it. This event is NOT bubbled.
+	 *
+	 * @param MyGeometry The Geometry of the widget receiving the event
+	 * @param MouseEvent Information about the input event
+	 */
+	virtual void OnMouseEnter(const FGeometry& MyGeometry, const FPointerEvent& MouseEvent) override;
+
+	/**
+	 * The system will use this event to notify a widget that the cursor has left it. This event is NOT bubbled.
+	 *
+	 * @param MouseEvent Information about the input event
+	 */
+	virtual void OnMouseLeave(const FPointerEvent& MouseEvent) override;
+
+	/**
+	 * Called after a key is pressed when this widget has focus
+	 *
+	 * @param MyGeometry The Geometry of the widget receiving the event
+	 * @param  InKeyEvent  Key event
+	 *
+	 * @return  Returns whether the event was handled, along with other possible actions
+	 */
+	virtual FReply OnKeyDown(const FGeometry& MyGeometry, const FKeyEvent& InKeyEvent) override;
+
+	/**
+	 * Called when the user is dropping something onto a widget; terminates drag and drop.
+	 *
+	 * @param MyGeometry      The geometry of the widget receiving the event.
+	 * @param DragDropEvent   The drag and drop event.
+	 *
+	 * @return A reply that indicated whether this event was handled.
+	 */
+	virtual FReply OnDrop(const FGeometry& MyGeometry, const FDragDropEvent& DragDropEvent) override;
+
+	/**
+	 * Called during drag and drop when the the mouse is being dragged over a widget.
+	 *
+	 * @param MyGeometry      The geometry of the widget receiving the event.
+	 * @param DragDropEvent   The drag and drop event.
+	 *
+	 * @return A reply that indicated whether this event was handled.
+	 */
+	virtual FReply OnDragOver(const FGeometry& MyGeometry, const FDragDropEvent& DragDropEvent)  override;
+
+private:
+	/** The major tab id for this window. */
+	FName MajorTabId;
+
+	/** Holds the tab manager that manages the front-end's tabs. */
+	TSharedPtr<FTabManager> TabManager;
+
+	/** The menu group. */
+	TSharedPtr<FWorkspaceItem> WorkspaceMenuGroup;
+
+	/** All tabs owned by this window */
+	TSet<TSharedPtr<SDockTab>> OpenTabs;
+
+	/** Command list used in the window. Maps commands to window specific actions. */
+	TSharedPtr<FUICommandList> CommandList;
+
+	/** The handle to the active update duration tick */
+	TWeakPtr<FActiveTimerHandle> ActiveTimerHandle;
+
+	/** The number of seconds the window has been active */
+	float DurationActive;
+
+	/** Tab specific slate extender structure, has the lifetime of this widget */
+	TSharedPtr<FInsightsMajorTabExtender> Extension;
+};
+
+} // namespace Insights
