[Insights]
  - Added "Auto Open Live Trace" toggle option in the main menu. If enabled, the analysis starts automatically for each new live trace session, replacing the current analysis session.
  - Moved the two "Open Trace File" options (from the main menu) into a sub-menu. Also added folder-open icon for these actions.

diff --git a/Engine/Source/Developer/TraceInsights/Private/Insights/Common/InsightsMenuBuilder.cpp b/Engine/Source/Developer/TraceInsights/Private/Insights/Common/InsightsMenuBuilder.cpp
--- a/Engine/Source/Developer/TraceInsights/Private/Insights/Common/InsightsMenuBuilder.cpp
+++ b/Engine/Source/Developer/TraceInsights/Private/Insights/Common/InsightsMenuBuilder.cpp
@@ -58,16 +58,24 @@ void FInsightsMenuBuilder::PopulateMenu(FMenuBuilder& MenuBuilder)
 		LOCTEXT("OpenSessionBrowser_ToolTip", "Opens the Unreal Insights Session Browser window."),
 		FSlateIcon(FInsightsStyle::GetStyleSetName(), "AppIcon.Small"),
 		FUIAction(FExecuteAction::CreateLambda([] { FInsightsManager::Get()->OpenUnrealInsights(); })));
+	MenuBuilder.AddSubMenu(
+		LOCTEXT("OpenTraceFile_SubMenu", "Open Trace File"),
+		LOCTEXT("OpenTraceFile_SubMenu_Desc", "Starts analysis for a specified trace file."),
+		FNewMenuDelegate::CreateSP(this, &FInsightsMenuBuilder::BuildOpenTraceFileSubMenu),
+		false,
+		FSlateIcon(FAppStyle::Get().GetStyleSetName(), "Icons.FolderOpen")
+	);
 	MenuBuilder.AddMenuEntry(
-		LOCTEXT("OpenTraceFile1", "Open Trace File (new instance)..."),
-		LOCTEXT("OpenTraceFile1_ToolTip", "Starts analysis for a specified trace file, in a separate Unreal Insights instance."),
-		FSlateIcon(),
-		FUIAction(FExecuteAction::CreateLambda([] { FInsightsManager::Get()->OpenTraceFile(); })));
-	MenuBuilder.AddMenuEntry(
-		LOCTEXT("OpenTraceFile2", "Open Trace File (same instance)..."),
-		LOCTEXT("OpenTraceFile2_ToolTip", "Starts analysis for a specified trace file, replacing the current analysis session."),
+		LOCTEXT("AutoOpenLiveTrace", "Auto Open Live Trace"),
+		LOCTEXT("AutoOpenLiveTrace_ToolTip", "If enabled, the analysis starts automatically for each new live trace session, replacing the current analysis session."),
 		FSlateIcon(),
-		FUIAction(FExecuteAction::CreateLambda([] { FInsightsManager::Get()->LoadTraceFile(); })));
+		FUIAction(
+			FExecuteAction::CreateLambda([]() { FInsightsManager::Get()->ToggleAutoLoadLiveSession(); }),
+			FCanExecuteAction(),
+			FIsActionChecked::CreateLambda([]() -> bool { return FInsightsManager::Get()->IsAutoLoadLiveSessionEnabled(); })
+		),
+		NAME_None,
+		EUserInterfaceActionType::ToggleButton);
 	MenuBuilder.EndSection();
 
 	FGlobalTabmanager::Get()->PopulateLocalTabSpawnerMenu(MenuBuilder);
@@ -109,6 +117,23 @@ void FInsightsMenuBuilder::PopulateMenu(FMenuBuilder& MenuBuilder)
 
 ////////////////////////////////////////////////////////////////////////////////////////////////////
 
+void FInsightsMenuBuilder::BuildOpenTraceFileSubMenu(FMenuBuilder& MenuBuilder)
+{
+	MenuBuilder.AddMenuEntry(
+		LOCTEXT("OpenTraceFile1", "Open in New Instance..."),
+		LOCTEXT("OpenTraceFile1_ToolTip", "Starts analysis for a specified trace file, in a separate Unreal Insights instance."),
+		FSlateIcon(FAppStyle::Get().GetStyleSetName(), "Icons.FolderOpen"),
+		FUIAction(FExecuteAction::CreateLambda([] { FInsightsManager::Get()->OpenTraceFile(); })));
+
+	MenuBuilder.AddMenuEntry(
+		LOCTEXT("OpenTraceFile2", "Open in Same Instance..."),
+		LOCTEXT("OpenTraceFile2_ToolTip", "Starts analysis for a specified trace file, replacing the current analysis session."),
+		FSlateIcon(FAppStyle::Get().GetStyleSetName(), "Icons.FolderOpen"),
+		FUIAction(FExecuteAction::CreateLambda([] { FInsightsManager::Get()->LoadTraceFile(); })));
+}
+
+////////////////////////////////////////////////////////////////////////////////////////////////////
+
 void FInsightsMenuBuilder::AddMenuEntry(
 	FMenuBuilder& InOutMenuBuilder,
 	const FUIAction& InAction,
diff --git a/Engine/Source/Developer/TraceInsights/Private/Insights/Common/InsightsMenuBuilder.h b/Engine/Source/Developer/TraceInsights/Private/Insights/Common/InsightsMenuBuilder.h
--- a/Engine/Source/Developer/TraceInsights/Private/Insights/Common/InsightsMenuBuilder.h
+++ b/Engine/Source/Developer/TraceInsights/Private/Insights/Common/InsightsMenuBuilder.h
@@ -10,7 +10,7 @@ class FMenuBuilder;
 
 struct FUIAction;
 
-class FInsightsMenuBuilder
+class FInsightsMenuBuilder : public TSharedFromThis<FInsightsMenuBuilder>
 {
 public:
 	FInsightsMenuBuilder();
@@ -18,7 +18,8 @@ public:
 	TSharedRef<FWorkspaceItem> GetInsightsToolsGroup();
 	TSharedRef<FWorkspaceItem> GetWindowsGroup();
 
-	void PopulateMenu(FMenuBuilder& MenuBuilder);
+	void PopulateMenu(FMenuBuilder& InOutMenuBuilder);
+	void BuildOpenTraceFileSubMenu(FMenuBuilder& InOutMenuBuilder);
 
 	// Adds a menu entry with a custom key binding text.
 	static void AddMenuEntry(
diff --git a/Engine/Source/Developer/TraceInsights/Private/Insights/InsightsManager.cpp b/Engine/Source/Developer/TraceInsights/Private/Insights/InsightsManager.cpp
--- a/Engine/Source/Developer/TraceInsights/Private/Insights/InsightsManager.cpp
+++ b/Engine/Source/Developer/TraceInsights/Private/Insights/InsightsManager.cpp
@@ -494,6 +494,8 @@ FInsightsSettings& FInsightsManager::GetSettings()
 
 bool FInsightsManager::Tick(float DeltaTime)
 {
+	AutoLoadLiveSession();
+
 	UpdateSessionDuration();
 
 #if !WITH_EDITOR
@@ -794,28 +796,55 @@ void FInsightsManager::OpenTraceFile(const FString& InTraceFile) const
 
 ////////////////////////////////////////////////////////////////////////////////////////////////////
 
-void FInsightsManager::LoadLastLiveSession()
+void FInsightsManager::AutoLoadLiveSession()
 {
-	ResetSession();
+	if (!bIsAutoLoadLiveSessionEnabled)
+	{
+		return;
+	}
 
 	if (!StoreClient.IsValid())
 	{
 		return;
 	}
 
-	const int32 SessionCount = StoreClient->GetSessionCount();
-	if (SessionCount == 0)
+	const uint32 SessionCount = StoreClient->GetSessionCount();
+	for (uint32 SessionIndex = 0; SessionIndex < SessionCount; ++SessionIndex)
 	{
-		return;
+		const UE::Trace::FStoreClient::FSessionInfo* SessionInfo = StoreClient->GetSessionInfo(SessionIndex);
+		if (SessionInfo)
+		{
+			const uint32 TraceId = SessionInfo->GetTraceId();
+			if (TraceId != CurrentTraceId && !AutoLoadedTraceIds.Contains(TraceId))
+			{
+				AutoLoadedTraceIds.Add(TraceId);
+				LoadTrace(SessionInfo->GetTraceId());
+				break;
+			}
+		}
 	}
+}
 
-	const UE::Trace::FStoreClient::FSessionInfo* SessionInfo = StoreClient->GetSessionInfo(SessionCount - 1);
-	if (SessionInfo == nullptr)
+////////////////////////////////////////////////////////////////////////////////////////////////////
+
+void FInsightsManager::LoadLastLiveSession()
+{
+	ResetSession();
+
+	if (!StoreClient.IsValid())
 	{
 		return;
 	}
 
-	LoadTrace(SessionInfo->GetTraceId());
+	const uint32 SessionCount = StoreClient->GetSessionCount();
+	if (SessionCount != 0)
+	{
+		const UE::Trace::FStoreClient::FSessionInfo* SessionInfo = StoreClient->GetSessionInfo(SessionCount - 1);
+		if (SessionInfo)
+		{
+			LoadTrace(SessionInfo->GetTraceId());
+		}
+	}
 }
 
 ////////////////////////////////////////////////////////////////////////////////////////////////////
@@ -824,7 +853,7 @@ void FInsightsManager::LoadTrace(uint32 InTraceId, bool InAutoQuit)
 {
 	ResetSession();
 
-	if (StoreClient == nullptr)
+	if (!StoreClient.IsValid())
 	{
 		if (InAutoQuit)
 		{
diff --git a/Engine/Source/Developer/TraceInsights/Private/Insights/InsightsManager.h b/Engine/Source/Developer/TraceInsights/Private/Insights/InsightsManager.h
--- a/Engine/Source/Developer/TraceInsights/Private/Insights/InsightsManager.h
+++ b/Engine/Source/Developer/TraceInsights/Private/Insights/InsightsManager.h
@@ -264,6 +264,10 @@ public:
 	 */
 	void OpenTraceFile(const FString& TraceFilename) const;
 
+	void ToggleAutoLoadLiveSession() { bIsAutoLoadLiveSessionEnabled = !bIsAutoLoadLiveSessionEnabled; }
+	bool IsAutoLoadLiveSessionEnabled() const { return bIsAutoLoadLiveSessionEnabled; }
+	void AutoLoadLiveSession();
+
 	/**
 	 * Creates a new analysis session instance and loads the latest available trace that is live.
 	 * Replaces the current analysis session.
@@ -451,6 +455,9 @@ private:
 	bool bIsAnalysisComplete = false;
 	bool bSessionAnalysisCompletedAutoQuit = false;
 
+	bool bIsAutoLoadLiveSessionEnabled = false;
+	TSet<uint32> AutoLoadedTraceIds; // list of trace ids for the auto loaded live sessions
+
 	FStopwatch AnalysisStopwatch;
 	double SessionDuration = 0.0;
 	double AnalysisDuration = 0.0;
