[Insights] Added toggle button to choose btween searching/filtering the list of trace session by Name or by [REDACTED_PERSON]. Also added highlights for searched text (also in tooltip, when searching by cmd line).

diff --git a/Engine/Source/Developer/TraceInsights/Private/Insights/InsightsStyle.cpp b/Engine/Source/Developer/TraceInsights/Private/Insights/InsightsStyle.cpp
--- a/Engine/Source/Developer/TraceInsights/Private/Insights/InsightsStyle.cpp
+++ b/Engine/Source/Developer/TraceInsights/Private/Insights/InsightsStyle.cpp
@@ -123,6 +123,8 @@ void FInsightsStyle::FStyle::Initialize()
 	//////////////////////////////////////////////////
 
 	Set("WhiteBrush", new FSlateColorBrush(FLinearColor::White));
+	Set("DarkGreenBrush", new FSlateColorBrush(FLinearColor(0.0f, 0.25f, 0.0f, 1.0f)));
+
 	Set("SingleBorder", new FSlateBorderBrush(NAME_None, FMargin(1.0f)));
 	Set("DoubleBorder", new FSlateBorderBrush(NAME_None, FMargin(2.0f)));
 
@@ -268,6 +270,7 @@ void FInsightsStyle::FStyle::Initialize()
 	Set("Icons.ResetToDefault", new EDITOR_IMAGE_BRUSH("Icons/Profiler/profiler_ResetToDefault_32x", Icon16x16));
 	Set("Icons.DiffersFromDefault", new EDITOR_IMAGE_BRUSH_SVG("Starship/Common/ResetToDefault", Icon16x16));
 
+	Set("Icons.Console", new EDITOR_IMAGE_BRUSH_SVG("Starship/Common/Console", Icon16x16));
 	//Set("Icons.Filter", new CORE_IMAGE_BRUSH_SVG("Starship/Common/filter", Icon16x16));	//-> use FAppStyle "Icons.Filter"
 	Set("Icons.Filter.ToolBar", new CORE_IMAGE_BRUSH_SVG("Starship/Common/filter", Icon20x20));
 	Set("Icons.Find", new EDITOR_IMAGE_BRUSH_SVG("Starship/Common/TraceDataFiltering", Icon16x16));
diff --git a/Engine/Source/Developer/TraceInsights/Private/Insights/Widgets/SStartPageWindow.cpp b/Engine/Source/Developer/TraceInsights/Private/Insights/Widgets/SStartPageWindow.cpp
--- a/Engine/Source/Developer/TraceInsights/Private/Insights/Widgets/SStartPageWindow.cpp
+++ b/Engine/Source/Developer/TraceInsights/Private/Insights/Widgets/SStartPageWindow.cpp
@@ -92,6 +92,9 @@ public:
 						SNew(STextBlock)
 						.Visibility_Lambda([this]() { return IsRenaming() ? EVisibility::Collapsed : EVisibility::Visible; })
 						.Text(this, &STraceListRow::GetTraceName)
+						.HighlightText(this, &STraceListRow::GetTraceNameHighlightText)
+						.HighlightColor(FLinearColor(0.75f, 0.75f, 0.75f, 1.0f))
+						.HighlightShape(FInsightsStyle::Get().GetBrush("DarkGreenBrush"))
 						.ToolTip(STraceListRow::GetTraceTooltip())
 					]
 
@@ -271,6 +274,29 @@ public:
 		}
 	}
 
+	FText GetTraceNameHighlightText() const
+	{
+		TSharedPtr<STraceStoreWindow> ParentWidgetPin = WeakParentWidget.Pin();
+		if (ParentWidgetPin.IsValid())
+		{
+			if (!ParentWidgetPin->bSearchByCommandLine && ParentWidgetPin->FilterByNameSearchBox.IsValid())
+			{
+				const FText SearchText = ParentWidgetPin->FilterByNameSearchBox->GetText();
+				bool bQuotesRemoved = false;
+				const FString SearchString = SearchText.ToString().TrimQuotes(&bQuotesRemoved);
+				if (bQuotesRemoved)
+				{
+					return FText::FromString(SearchString);
+				}
+				else
+				{
+					return SearchText;
+				}
+			}
+		}
+		return FText::GetEmpty();
+	}
+
 	FText GetTraceUri() const
 	{
 		TSharedPtr<FTraceViewModel> TracePin = WeakTrace.Pin();
@@ -323,6 +349,29 @@ public:
 		}
 	}
 
+	FText GetTraceCommandLineHighlightText() const
+	{
+		TSharedPtr<STraceStoreWindow> ParentWidgetPin = WeakParentWidget.Pin();
+		if (ParentWidgetPin.IsValid())
+		{
+			if (ParentWidgetPin->bSearchByCommandLine && ParentWidgetPin->FilterByNameSearchBox.IsValid())
+			{
+				const FText SearchText = ParentWidgetPin->FilterByNameSearchBox->GetText();
+				bool bQuotesRemoved = false;
+				const FString SearchString = SearchText.ToString().TrimQuotes(&bQuotesRemoved);
+				if (bQuotesRemoved)
+				{
+					return FText::FromString(SearchString);
+				}
+				else
+				{
+					return SearchText;
+				}
+			}
+		}
+		return FText::GetEmpty();
+	}
+
 	FText GetTraceBranch() const
 	{
 		TSharedPtr<FTraceViewModel> TracePin = WeakTrace.Pin();
@@ -622,10 +671,10 @@ public:
 			int32 Row = 0;
 			AddGridPanelRow(GridPanel, Row++, LOCTEXT("TraceTooltip_Platform", "Platform:"), &STraceListRow::GetTracePlatform);
 			AddGridPanelRow(GridPanel, Row++, LOCTEXT("TraceTooltip_AppName", "App Name:"), &STraceListRow::GetTraceAppName);
-			AddGridPanelRow(GridPanel, Row++, LOCTEXT("TraceTooltip_CommandLine", "Command Line:"), &STraceListRow::GetTraceCommandLine);
+			AddGridPanelRow(GridPanel, Row++, LOCTEXT("TraceTooltip_CommandLine", "Command Line:"), &STraceListRow::GetTraceCommandLine, &STraceListRow::GetTraceCommandLineHighlightText);
 			AddGridPanelRow(GridPanel, Row++, LOCTEXT("TraceTooltip_Branch", "Branch:"), &STraceListRow::GetTraceBranch);
 			AddGridPanelRow(GridPanel, Row++, LOCTEXT("TraceTooltip_BuildVersion", "Build Version:"), &STraceListRow::GetTraceBuildVersion);
-			AddGridPanelRow(GridPanel, Row++, LOCTEXT("TraceTooltip_Changelist", "Changelist:"), &STraceListRow::GetTraceChangelist, &STraceListRow::TraceChangelistVisibility);
+			AddGridPanelRow(GridPanel, Row++, LOCTEXT("TraceTooltip_Changelist", "Changelist:"), &STraceListRow::GetTraceChangelist, nullptr, &STraceListRow::TraceChangelistVisibility);
 			AddGridPanelRow(GridPanel, Row++, LOCTEXT("TraceTooltip_BuildConfig", "Build Config:"), &STraceListRow::GetTraceBuildConfiguration);
 			AddGridPanelRow(GridPanel, Row++, LOCTEXT("TraceTooltip_BuildTarget", "Build Target:"), &STraceListRow::GetTraceBuildTarget);
 			AddGridPanelRow(GridPanel, Row++, LOCTEXT("TraceTooltip_Timestamp", "Timestamp:"), &STraceListRow::GetTraceTimestampForTooltip);
@@ -650,6 +699,7 @@ public:
 private:
 	void AddGridPanelRow(TSharedPtr<SGridPanel> Grid, int32 Row, const FText& InHeaderText,
 		typename TAttribute<FText>::FGetter::template TConstMethodPtr<STraceListRow> InValueTextFn,
+		typename TAttribute<FText>::FGetter::template TConstMethodPtr<STraceListRow> InHighlightTextFn = nullptr,
 		typename TAttribute<EVisibility>::FGetter::template TConstMethodPtr<STraceListRow> InVisibilityFn = nullptr) const
 	{
 		SGridPanel::FSlot* Slot0 = nullptr;
@@ -664,17 +714,37 @@ private:
 			];
 
 		SGridPanel::FSlot* Slot1 = nullptr;
-		Grid->AddSlot(1, Row)
-			.Expose(Slot1)
-			.Padding(2.0f)
-			.HAlign(HAlign_Left)
-			[
-				SNew(STextBlock)
-				.Text(this, InValueTextFn)
-				.WrapTextAt(1024.0f)
-				.WrappingPolicy(ETextWrappingPolicy::AllowPerCharacterWrapping)
-				.ColorAndOpacity(FSlateColor(EStyleColor::Foreground))
-			];
+
+		if (InHighlightTextFn)
+		{
+			Grid->AddSlot(1, Row)
+				.Expose(Slot1)
+				.Padding(2.0f)
+				.HAlign(HAlign_Left)
+				[
+					SNew(STextBlock)
+					.Text(this, InValueTextFn)
+					.HighlightText(this, InHighlightTextFn)
+					.HighlightShape(FInsightsStyle::Get().GetBrush("DarkGreenBrush"))
+					.WrapTextAt(1024.0f)
+					.WrappingPolicy(ETextWrappingPolicy::AllowPerCharacterWrapping)
+					.ColorAndOpacity(FSlateColor(EStyleColor::Foreground))
+				];
+		}
+		else
+		{
+			Grid->AddSlot(1, Row)
+				.Expose(Slot1)
+				.Padding(2.0f)
+				.HAlign(HAlign_Left)
+				[
+					SNew(STextBlock)
+					.Text(this, InValueTextFn)
+					.WrapTextAt(1024.0f)
+					.WrappingPolicy(ETextWrappingPolicy::AllowPerCharacterWrapping)
+					.ColorAndOpacity(FSlateColor(EStyleColor::Foreground))
+				];
+		}
 
 		if (InVisibilityFn)
 		{
@@ -714,6 +784,7 @@ STraceStoreWindow::STraceStoreWindow()
 	, SelectedTrace()
 	, bIsUserSelectedTrace(false)
 	, Filters()
+	, bSearchByCommandLine(false)
 	, FilterByNameSearchBox()
 	, FilterByName()
 	, FilterByPlatform()
@@ -889,12 +960,36 @@ TSharedRef<SWidget> STraceStoreWindow::ConstructFiltersToolbar()
 
 	ToolbarBuilder.BeginSection("Filters");
 	{
+		// Toggle between filtering the list of trace sessions by name or by command line
+		ToolbarBuilder.AddWidget(
+			SNew(SCheckBox)
+			.Style(FAppStyle::Get(), "ToggleButtonCheckbox")
+			.HAlign(HAlign_Center)
+			.Padding(3.0f)
+			.OnCheckStateChanged_Lambda([this](ECheckBoxState NewState) { bSearchByCommandLine = (NewState == ECheckBoxState::Checked); OnFilterChanged(); })
+			.IsChecked_Lambda([this]() { return bSearchByCommandLine ? ECheckBoxState::Checked : ECheckBoxState::Unchecked; })
+			.ToolTipText(LOCTEXT("ToggleNameFilter_Tooltip", "Toggle between filtering the list of trace sessions by name or by command line."))
+			[
+				SNew(SImage)
+				.Image(FInsightsStyle::Get().GetBrush("Icons.Console"))
+			]);
+
 		// Text Filter (Search Box)
 		ToolbarBuilder.AddWidget(
 			SAssignNew(FilterByNameSearchBox, SSearchBox)
 			.MinDesiredWidth(150.0f)
-			.HintText(LOCTEXT("NameFilter_Hint", "Name"))
-			.ToolTipText(LOCTEXT("NameFilter_Tooltip", "Type here to filter the list of trace sessions by name."))
+			.HintText_Lambda([this]()
+				{
+					return bSearchByCommandLine ?
+						LOCTEXT("CmdLineFilter_Hint", "Command Line") :
+						LOCTEXT("NameFilter_Hint", "Name");
+				})
+			.ToolTipText_Lambda([this]()
+				{
+					return bSearchByCommandLine ?
+						LOCTEXT("CmdLineFilter_Tooltip", "Type here to filter the list of trace sessions by command line.") :
+						LOCTEXT("NameFilter_Tooltip", "Type here to filter the list of trace sessions by name.");
+				})
 			.IsEnabled_Lambda([this]() { return TraceViewModels.Num() > 0; })
 			.OnTextChanged(this, &STraceStoreWindow::FilterByNameSearchBox_OnTextChanged)
 			.DelayChangeNotificationsWhileTyping(true)
@@ -2395,7 +2490,14 @@ void STraceStoreWindow::CreateFilters()
 
 void STraceStoreWindow::HandleItemToStringArray(const FTraceViewModel& InTrace, TArray<FString>& OutSearchStrings) const
 {
-	OutSearchStrings.Add(InTrace.Name.ToString());
+	if (bSearchByCommandLine)
+	{
+		OutSearchStrings.Add(InTrace.CommandLine.ToString());
+	}
+	else
+	{
+		OutSearchStrings.Add(InTrace.Name.ToString());
+	}
 }
 
 ////////////////////////////////////////////////////////////////////////////////////////////////////
diff --git a/Engine/Source/Developer/TraceInsights/Private/Insights/Widgets/SStartPageWindow.h b/Engine/Source/Developer/TraceInsights/Private/Insights/Widgets/SStartPageWindow.h
--- a/Engine/Source/Developer/TraceInsights/Private/Insights/Widgets/SStartPageWindow.h
+++ b/Engine/Source/Developer/TraceInsights/Private/Insights/Widgets/SStartPageWindow.h
@@ -220,6 +220,8 @@ protected:
 /** Implements the Trace Store window. */
 class STraceStoreWindow : public SCompoundWidget
 {
+	friend class STraceListRow;
+
 public:
 	/** Default constructor. */
 	STraceStoreWindow();
@@ -459,6 +461,7 @@ private:
 
 	TSharedPtr<FTraceViewModelFilterCollection> Filters;
 
+	bool bSearchByCommandLine;
 	TSharedPtr<SSearchBox> FilterByNameSearchBox;
 	TSharedPtr<FTraceTextFilter> FilterByName;
 
