[Insights]
  - Fixed UE-150996, log message after "TimingInsights.ExportThreads" execute command not counting correctly the total number of threads (ignoring the two gpu threads).
  - Fixed UE-150992, executing commands from a response file was ignoring the last command if that last line was not EOL terminated.
  - Added support to comment a command (i.e. ignore execution) in a response file by using # as the first char on a line.

diff --git a/Engine/Source/Developer/TraceInsights/Private/Insights/InsightsManager.cpp b/Engine/Source/Developer/TraceInsights/Private/Insights/InsightsManager.cpp
--- a/Engine/Source/Developer/TraceInsights/Private/Insights/InsightsManager.cpp
+++ b/Engine/Source/Developer/TraceInsights/Private/Insights/InsightsManager.cpp
@@ -1144,9 +1144,11 @@ bool FInsightsManager::HandleResponseFileCmd(const TCHAR* ResponseFile, FOutputD
 		if (EndOfLine > 0)
 		{
 			const uint32 LineLen = uint32(CrtPos - StartPos);
-			StartPos[LineLen] = TEXT('\0');
-
-			TraceInsightsModule.Exec(StartPos, Ar);
+			if (LineLen > 0 && *StartPos != TEXT('#'))
+			{
+				StartPos[LineLen] = TEXT('\0');
+				TraceInsightsModule.Exec(StartPos, Ar);
+			}
 
 			CrtPos += EndOfLine;
 			StartPos = CrtPos;
@@ -1156,6 +1158,12 @@ bool FInsightsManager::HandleResponseFileCmd(const TCHAR* ResponseFile, FOutputD
 			++CrtPos;
 		}
 	}
+	const uint32 LineLen = uint32(CrtPos - StartPos);
+	if (LineLen > 0 && *StartPos != TEXT('#'))
+	{
+		StartPos[LineLen] = TEXT('\0');
+		TraceInsightsModule.Exec(StartPos, Ar);
+	}
 
 	return true;
 }
diff --git a/Engine/Source/Developer/TraceInsights/Private/Insights/ViewModels/TimingExporter.cpp b/Engine/Source/Developer/TraceInsights/Private/Insights/ViewModels/TimingExporter.cpp
--- a/Engine/Source/Developer/TraceInsights/Private/Insights/ViewModels/TimingExporter.cpp
+++ b/Engine/Source/Developer/TraceInsights/Private/Insights/ViewModels/TimingExporter.cpp
@@ -140,6 +140,7 @@ int32 FTimingExporter::ExportThreadsAsText(const FString& Filename, FExportThrea
 		StringBuilder.Append(UTF8TEXT("GPU"));
 		StringBuilder.AppendChar(LineEnd);
 		ExportFileHandle->Write((const uint8*)StringBuilder.ToString(), StringBuilder.Len() * sizeof(UTF8CHAR));
+		++ThreadCount;
 
 		StringBuilder.Reset();
 		StringBuilder.Appendf(UTF8TEXT("%u"), FGpuTimingTrack::Gpu2ThreadId);
@@ -149,6 +150,7 @@ int32 FTimingExporter::ExportThreadsAsText(const FString& Filename, FExportThrea
 		StringBuilder.Append(UTF8TEXT("GPU"));
 		StringBuilder.AppendChar(LineEnd);
 		ExportFileHandle->Write((const uint8*)StringBuilder.ToString(), StringBuilder.Len() * sizeof(UTF8CHAR));
+		++ThreadCount;
 
 		// Iterate the CPU threads.
 		{
