[Insights]
  - Session Browser: Improved metadata update by prioritizing the newer trace sessions. Also enabled parallel update for metadata (~10x faster to update metadata for large number of traces).
  - Added a critical section for synchronizing the use of FStoreClient exposed by [REDACTED_PERSON] (as FStoreClient API is not thread safe).
  - Session Browser: Added F5 key shortcut to refresh the list of trace sessions (debug functionality).

diff --git a/Engine/Source/Developer/TraceInsights/Private/Insights/InsightsManager.h b/Engine/Source/Developer/TraceInsights/Private/Insights/InsightsManager.h
--- a/Engine/Source/Developer/TraceInsights/Private/Insights/InsightsManager.h
+++ b/Engine/Source/Developer/TraceInsights/Private/Insights/InsightsManager.h
@@ -121,6 +121,7 @@ public:
 
 	bool ConnectToStore(const TCHAR* Host, uint32 Port=0);
 	UE::Trace::FStoreClient* GetStoreClient() const { return StoreClient.Get(); }
+	FCriticalSection& GetStoreClientCriticalSection() const { return StoreClientCriticalSection; }
 
 	/** @return an instance of the trace analysis session. */
 	TSharedPtr<const TraceServices::IAnalysisSession> GetSession() const;
@@ -413,9 +414,12 @@ private:
 	/** The location of the trace files managed by the trace store. */
 	FString StoreDir;
 
-	/** The client used to connect to the trace store. */
+	/** The client used to connect to the trace store. It is not thread safe! */
 	TUniquePtr<UE::Trace::FStoreClient> StoreClient;
 
+	/** CriticalSection for using the store client's API. */
+	mutable FCriticalSection StoreClientCriticalSection;
+
 	/** The trace analysis session. */
 	TSharedPtr<const TraceServices::IAnalysisSession> Session;
 
diff --git a/Engine/Source/Developer/TraceInsights/Private/Insights/StoreService/StoreBrowser.cpp b/Engine/Source/Developer/TraceInsights/Private/Insights/StoreService/StoreBrowser.cpp
--- a/Engine/Source/Developer/TraceInsights/Private/Insights/StoreService/StoreBrowser.cpp
+++ b/Engine/Source/Developer/TraceInsights/Private/Insights/StoreService/StoreBrowser.cpp
@@ -2,6 +2,7 @@
 
 #include "StoreBrowser.h"
 
+#include "Async/ParallelFor.h"
 #include "Trace/Analysis.h"
 #include "Trace/StoreClient.h"
 
@@ -78,14 +79,36 @@ void FStoreBrowser::Exit()
 
 ////////////////////////////////////////////////////////////////////////////////////////////////////
 
+void FStoreBrowser::Refresh()
+{
+	if (!IsLocked())
+	{
+		FScopeLock Lock(&TracesCriticalSection);
+
+		StoreChangeSerial = 0;
+		TracesChangeSerial = 0;
+		Traces.Reset();
+		TraceMap.Reset();
+		LiveTraceMap.Reset();
+	}
+}
+
+////////////////////////////////////////////////////////////////////////////////////////////////////
+
 UE::Trace::FStoreClient* FStoreBrowser::GetStoreClient() const
 {
-	// TODO: thread safety !?
 	return FInsightsManager::Get()->GetStoreClient();
 }
 
 ////////////////////////////////////////////////////////////////////////////////////////////////////
 
+FCriticalSection& FStoreBrowser::GetStoreClientCriticalSection() const
+{
+	return FInsightsManager::Get()->GetStoreClientCriticalSection();
+}
+
+////////////////////////////////////////////////////////////////////////////////////////////////////
+
 void FStoreBrowser::UpdateTraces()
 {
 	UE::Trace::FStoreClient* StoreClient = GetStoreClient();
@@ -98,6 +121,8 @@ void FStoreBrowser::UpdateTraces()
 	FStopwatch StopwatchTotal;
 	StopwatchTotal.Start();
 
+	FScopeLock StoreClientLock(&GetStoreClientCriticalSection());
+
 	// Check if the list of trace sessions has changed.
 	{
 		const UE::Trace::FStoreClient::FStatus* Status = StoreClient->GetStatus();
@@ -261,6 +286,8 @@ void FStoreBrowser::UpdateTraces()
 		}
 	}
 
+	StoreClientLock.Unlock();
+
 	StopwatchTotal.Update();
 	const uint64 Step3 = StopwatchTotal.GetAccumulatedTimeMs();
 
@@ -269,21 +296,49 @@ void FStoreBrowser::UpdateTraces()
 		FStopwatch Stopwatch;
 		Stopwatch.Start();
 
-		int32 MetadataUpdateCount = 0; // for debugging
+		struct FStoreBrowserTraceInfoBox
+		{
+			FStoreBrowserTraceInfo* Ptr;
+		};
+		TArray<FStoreBrowserTraceInfoBox> TracesToUpdate;
 
 		for (TSharedPtr<FStoreBrowserTraceInfo>& TracePtr : Traces)
 		{
 			if (TracePtr->MetadataUpdateCount > 0)
 			{
-				MetadataUpdateCount++;
-				UpdateMetadata(*TracePtr);
+				TracesToUpdate.Add({ TracePtr.Get() });
 			}
 		}
 
+		// Sort descending by timestamp (i.e. to update the newer traces first).
+		TracesToUpdate.Sort([](const FStoreBrowserTraceInfoBox& A, const FStoreBrowserTraceInfoBox& B)
+			{
+				return A.Ptr->Timestamp.GetTicks() > B.Ptr->Timestamp.GetTicks();
+			});
+
+#if 0
+		for (const FStoreBrowserTraceInfoBox& Trace : TracesToUpdate)
+		{
+			UpdateMetadata(*Trace.Ptr);
+		}
+#else
+		if (TracesToUpdate.Num() > 0)
+		{
+			ParallelFor(TracesToUpdate.Num(), [this, &TracesToUpdate](uint32 Index)
+				{
+					UpdateMetadata(*TracesToUpdate[Index].Ptr);
+				},
+				EParallelForFlags::BackgroundPriority);
+		}
+#endif
+
 		Stopwatch.Stop();
-		if (MetadataUpdateCount > 0)
+		if (TracesToUpdate.Num() > 0)
 		{
-			UE_LOG(TraceInsights, Log, TEXT("[StoreBrowser] Metadata updated in %llu ms for %d trace(s)."), Stopwatch.GetAccumulatedTimeMs(), MetadataUpdateCount);
+			UE_LOG(TraceInsights, Log, TEXT("[StoreBrowser] Metadata updated in %llu ms for %d trace(s) (~%llu ms/trace)."),
+				Stopwatch.GetAccumulatedTimeMs(),
+				TracesToUpdate.Num(),
+				Stopwatch.GetAccumulatedTimeMs() / TracesToUpdate.Num());
 		}
 	}
 
@@ -291,7 +346,8 @@ void FStoreBrowser::UpdateTraces()
 	const uint64 TotalTime = StopwatchTotal.GetAccumulatedTimeMs();
 	if (TotalTime > 5)
 	{
-		UE_LOG(TraceInsights, Log, TEXT("[StoreBrowser] Updated in %llu ms (%llu + %llu + %llu + %llu)."), TotalTime, Step1, Step2 - Step1, Step3 - Step2, TotalTime - Step3);
+		UE_LOG(TraceInsights, Log, TEXT("[StoreBrowser] Updated in %llu ms (%llu + %llu + %llu + %llu)."),
+			TotalTime, Step1, Step2 - Step1, Step3 - Step2, TotalTime - Step3);
 	}
 }
 
@@ -320,7 +376,9 @@ void FStoreBrowser::UpdateMetadata(FStoreBrowserTraceInfo& Trace)
 
 	FPlatformProcess::SleepNoStats(0.0f);
 
+	FScopeLock StoreClientLock(&GetStoreClientCriticalSection());
 	UE::Trace::FStoreClient::FTraceData TraceData = StoreClient->ReadTrace(Trace.TraceId);
+	StoreClientLock.Unlock();
 	if (!TraceData)
 	{
 		return;
@@ -410,6 +468,10 @@ void FStoreBrowser::UpdateMetadata(FStoreBrowserTraceInfo& Trace)
 			Trace.MetadataUpdateCount = 0; // no more updates
 		}
 	}
+
+	DataStream.Stopwatch.Stop();
+	UE_LOG(TraceInsights, Log, TEXT("[StoreBrowser] Metadata updated in %llu ms for trace \"%s\""),
+		DataStream.Stopwatch.GetAccumulatedTimeMs(), *Trace.Name);
 }
 
 ////////////////////////////////////////////////////////////////////////////////////////////////////
diff --git a/Engine/Source/Developer/TraceInsights/Private/Insights/StoreService/StoreBrowser.h b/Engine/Source/Developer/TraceInsights/Private/Insights/StoreService/StoreBrowser.h
--- a/Engine/Source/Developer/TraceInsights/Private/Insights/StoreService/StoreBrowser.h
+++ b/Engine/Source/Developer/TraceInsights/Private/Insights/StoreService/StoreBrowser.h
@@ -87,8 +87,11 @@ public:
 	const TMap<uint32, TSharedPtr<FStoreBrowserTraceInfo>>& GetLockedTraceMap() const { check(bTracesLocked); return TraceMap; }
 	void Unlock() { check(bTracesLocked); bTracesLocked = false; TracesCriticalSection.Unlock(); }
 
+	void Refresh();
+
 private:
 	UE::Trace::FStoreClient* GetStoreClient() const;
+	FCriticalSection& GetStoreClientCriticalSection() const;
 
 	void UpdateTraces();
 	void ResetTraces();
diff --git a/Engine/Source/Developer/TraceInsights/Private/Insights/Widgets/SStartPageWindow.cpp b/Engine/Source/Developer/TraceInsights/Private/Insights/Widgets/SStartPageWindow.cpp
--- a/Engine/Source/Developer/TraceInsights/Private/Insights/Widgets/SStartPageWindow.cpp
+++ b/Engine/Source/Developer/TraceInsights/Private/Insights/Widgets/SStartPageWindow.cpp
@@ -1537,12 +1537,6 @@ FReply STraceStoreWindow::RefreshTraces_OnClicked()
 
 void STraceStoreWindow::RefreshTraceList()
 {
-	UE::Trace::FStoreClient* StoreClient = FInsightsManager::Get()->GetStoreClient();
-	if (StoreClient == nullptr)
-	{
-		return;
-	}
-
 	FStopwatch StopwatchTotal;
 	StopwatchTotal.Start();
 
@@ -1840,7 +1834,17 @@ FReply STraceStoreWindow::OnKeyDown(const FGeometry& MyGeometry, const FKeyEvent
 	//TODO: Make commands for Rename and Delete.
 	//return GetCommandList()->ProcessCommandBindings(InKeyEvent) ? FReply::Handled() : FReply::Unhandled();
 
-	if (InKeyEvent.GetKey() == EKeys::F2)
+	if (InKeyEvent.GetKey() == EKeys::F5) // refresh metadata for all trace sessions
+	{
+		StoreBrowser->Refresh();
+
+		TracesChangeSerial = 0;
+		TraceViewModels.Reset();
+		TraceViewModelMap.Reset();
+		OnTraceListChanged();
+		return FReply::Handled();
+	}
+	else if (InKeyEvent.GetKey() == EKeys::F2)
 	{
 		RenameTraceFile();
 		return FReply::Handled();
@@ -1964,7 +1968,17 @@ void STraceStoreWindow::OpenTraceSession(uint32 InTraceId)
 {
 	UE_LOG(TraceInsights, Log, TEXT("[TraceStore] Start analysis (in separate process) for trace id: 0x%08X"), InTraceId);
 
-	const uint32 StorePort = FInsightsManager::Get()->GetStoreClient()->GetStorePort();
+	UE::Trace::FStoreClient* StoreClient = FInsightsManager::Get()->GetStoreClient();
+	if (!StoreClient)
+	{
+		return;
+	}
+
+	uint32 StorePort = 0;
+	{
+		FScopeLock StoreClientLock(&FInsightsManager::Get()->GetStoreClientCriticalSection());
+		StorePort = StoreClient->GetStorePort();
+	}
 	FString CmdLine = FString::Printf(TEXT("-OpenTraceId=%d -StorePort=%d"), InTraceId, StorePort);
 
 	FString ExtraCmdParams;
@@ -2007,35 +2021,31 @@ TSharedRef<SWidget> STraceStoreWindow::MakeTraceListMenu()
 
 	MenuBuilder.BeginSection("AvailableTraces", LOCTEXT("TraceListMenu_Section_AvailableTraces", "Top Most Recently Created Traces"));
 	{
-		UE::Trace::FStoreClient* StoreClient = FInsightsManager::Get()->GetStoreClient();
-		if (StoreClient != nullptr)
-		{
-			// Make a copy of the trace list (to allow list view to be sorted by other criteria).
-			TArray<TSharedPtr<FTraceViewModel>> SortedTraces(TraceViewModels);
-			Algo::SortBy(SortedTraces, &FTraceViewModel::Timestamp);
-
-			int32 TraceCountLimit = 10; // top 10
+		// Make a copy of the trace list (to allow list view to be sorted by other criteria).
+		TArray<TSharedPtr<FTraceViewModel>> SortedTraces(TraceViewModels);
+		Algo::SortBy(SortedTraces, &FTraceViewModel::Timestamp);
 
-			// Iterate in reverse order as we want most recently created traces first.
-			for (int32 TraceIndex = SortedTraces.Num() - 1; TraceIndex >= 0 && TraceCountLimit > 0; --TraceIndex, --TraceCountLimit)
-			{
-				const FTraceViewModel& Trace = *SortedTraces[TraceIndex];
+		int32 TraceCountLimit = 10; // top 10
 
-				FText Label = Trace.Name;
-				if (Trace.bIsLive)
-				{
-					Label = FText::Format(LOCTEXT("LiveTraceTextFmt", "{0} (LIVE!)"), Label);
-				}
+		// Iterate in reverse order as we want most recently created traces first.
+		for (int32 TraceIndex = SortedTraces.Num() - 1; TraceIndex >= 0 && TraceCountLimit > 0; --TraceIndex, --TraceCountLimit)
+		{
+			const FTraceViewModel& Trace = *SortedTraces[TraceIndex];
 
-				MenuBuilder.AddMenuEntry(
-					Label,
-					TAttribute<FText>(), // no tooltip
-					FSlateIcon(),
-					FUIAction(FExecuteAction::CreateSP(this, &STraceStoreWindow::OpenTraceSession, Trace.TraceId)),
-					NAME_None,
-					EUserInterfaceActionType::Button
-				);
+			FText Label = Trace.Name;
+			if (Trace.bIsLive)
+			{
+				Label = FText::Format(LOCTEXT("LiveTraceTextFmt", "{0} (LIVE!)"), Label);
 			}
+
+			MenuBuilder.AddMenuEntry(
+				Label,
+				TAttribute<FText>(), // no tooltip
+				FSlateIcon(),
+				FUIAction(FExecuteAction::CreateSP(this, &STraceStoreWindow::OpenTraceSession, Trace.TraceId)),
+				NAME_None,
+				EUserInterfaceActionType::Button
+			);
 		}
 	}
 	MenuBuilder.EndSection();
