[Insights]
  - Memory Insights: Added "Open Source in Visual Studio" for any callstack frame of any allocation item, in the Allocs Table/Tree view (even if the tree is not grouped by callstack).

diff --git a/Engine/Source/Developer/TraceInsights/Private/Insights/MemoryProfiler/Widgets/SMemAllocTableTreeView.cpp b/Engine/Source/Developer/TraceInsights/Private/Insights/MemoryProfiler/Widgets/SMemAllocTableTreeView.cpp
--- a/Engine/Source/Developer/TraceInsights/Private/Insights/MemoryProfiler/Widgets/SMemAllocTableTreeView.cpp
+++ b/Engine/Source/Developer/TraceInsights/Private/Insights/MemoryProfiler/Widgets/SMemAllocTableTreeView.cpp
@@ -1288,37 +1288,246 @@ ECheckBoxState SMemAllocTableTreeView::CallstackGroupingByFunction_IsChecked() c
 
 ////////////////////////////////////////////////////////////////////////////////////////////////////
 
+TSharedPtr<FMemAllocNode> SMemAllocTableTreeView::GetSingleSelectedMemAllocNode() const
+{
+	if (TreeView->GetNumItemsSelected() == 1)
+	{
+		FMemAllocNodePtr SelectedTreeNode = StaticCastSharedPtr<FMemAllocNode>(TreeView->GetSelectedItems()[0]);
+		if (SelectedTreeNode.IsValid() && !SelectedTreeNode->IsGroup())
+		{
+			return SelectedTreeNode;
+		}
+	}
+	return nullptr;
+}
+
+////////////////////////////////////////////////////////////////////////////////////////////////////
+
 void SMemAllocTableTreeView::ExtendMenu(FMenuBuilder& MenuBuilder)
 {
-	MenuBuilder.BeginSection("CallstackFrame", LOCTEXT("ContextMenu_Section_CallstackFrame", "Callstack Frame"));
+	ISourceCodeAccessModule& SourceCodeAccessModule = FModuleManager::LoadModuleChecked<ISourceCodeAccessModule>("SourceCodeAccess");
+	ISourceCodeAccessor& SourceCodeAccessor = SourceCodeAccessModule.GetAccessor();
+
+	FMemAllocNodePtr SingleSelectedMemAllocNode = GetSingleSelectedMemAllocNode();
+	if (SingleSelectedMemAllocNode.IsValid() && CountSourceFiles(*SingleSelectedMemAllocNode) > 0)
 	{
-		ISourceCodeAccessModule& SourceCodeAccessModule = FModuleManager::LoadModuleChecked<ISourceCodeAccessModule>("SourceCodeAccess");
-		ISourceCodeAccessor& SourceCodeAccessor = SourceCodeAccessModule.GetAccessor();
+		MenuBuilder.BeginSection("Allocation", LOCTEXT("ContextMenu_Section_OpenSource", "Allocation"));
+		{
+			FText ItemLabel = FText::Format(LOCTEXT("ContextMenu_Open_SubMenu", "Open in {0}"), SourceCodeAccessor.GetNameText());
+			FText ItemToolTip = FText::Format(LOCTEXT("ContextMenu_Open_Desc_SubMenu", "Open source file of selected callstack frame in {0}."), SourceCodeAccessor.GetNameText());
+
+			MenuBuilder.AddSubMenu
+			(
+				ItemLabel,
+				ItemToolTip,
+				FNewMenuDelegate::CreateSP(this, &SMemAllocTableTreeView::BuildOpenSourceSubMenu),
+				false,
+				FSlateIcon(FAppStyle::GetAppStyleSetName(), SourceCodeAccessor.GetOpenIconName())
+			);
+		}
+		MenuBuilder.EndSection();
+	}
+	else
+	{
+		MenuBuilder.BeginSection("CallstackFrame", LOCTEXT("ContextMenu_Section_CallstackFrame", "Callstack Frame"));
+		{
+			FText ItemLabel = FText::Format(LOCTEXT("ContextMenu_Open", "Open in {0}"), SourceCodeAccessor.GetNameText());
+			FText FileName = GetSelectedCallstackFrameFileName();
+			FText ItemToolTip = FText::Format(LOCTEXT("ContextMenu_Open_Desc", "Open source file of selected callstack frame in {0}.\n{1}"), SourceCodeAccessor.GetNameText(), FileName);
+
+			FUIAction Action_OpenIDE
+			(
+				FExecuteAction::CreateSP(this, &SMemAllocTableTreeView::OpenCallstackFrameSourceFileInIDE),
+				FCanExecuteAction::CreateSP(this, &SMemAllocTableTreeView::CanOpenCallstackFrameSourceFileInIDE)
+			);
+			MenuBuilder.AddMenuEntry
+			(
+				ItemLabel,
+				ItemToolTip,
+				FSlateIcon(FAppStyle::GetAppStyleSetName(), SourceCodeAccessor.GetOpenIconName()),
+				Action_OpenIDE,
+				NAME_None,
+				EUserInterfaceActionType::Button
+			);
+		}
+		MenuBuilder.EndSection();
+	}
+}
 
-		FText ItemLabel = FText::Format(LOCTEXT("ContextMenu_Open", "Open in {0}"), SourceCodeAccessor.GetNameText());
-		FText FileName = GetSelectedCallstackFrameFileName();
-		FText ItemToolTip = FText::Format(LOCTEXT("ContextMenu_Open_Desc", "Open source file of selected callstack frame in {0}.\n{1}"), SourceCodeAccessor.GetNameText(), FileName);
-
-		FUIAction Action_OpenIDE
-		(
-			FExecuteAction::CreateSP(this, &SMemAllocTableTreeView::OpenCallstackFrameSourceFileInIDE),
-			FCanExecuteAction::CreateSP(this, &SMemAllocTableTreeView::CanOpenCallstackFrameSourceFileInIDE)
-		);
-		MenuBuilder.AddMenuEntry
-		(
-			ItemLabel,
-			ItemToolTip,
-			FSlateIcon(FAppStyle::GetAppStyleSetName(), SourceCodeAccessor.GetOpenIconName()),
-			Action_OpenIDE,
-			NAME_None,
-			EUserInterfaceActionType::Button
-		);
+////////////////////////////////////////////////////////////////////////////////////////////////////
+
+uint32 SMemAllocTableTreeView::CountSourceFiles(FMemAllocNode& MemAllocNode)
+{
+	if (MemAllocNode.IsGroup())
+	{
+		return 0;
+	}
+
+	const FMemoryAlloc* Alloc = MemAllocNode.GetMemAlloc();
+	if (!Alloc || !Alloc->Callstack)
+	{
+		return 0;
+	}
+
+	uint32 NumSourceFiles = 0;
+	const uint32 NumCallstackFrames = Alloc->Callstack->Num();
+	for (uint32 FrameIndex = 0; FrameIndex < NumCallstackFrames; ++FrameIndex)
+	{
+		const TraceServices::FStackFrame* Frame = Alloc->Callstack->Frame(FrameIndex);
+		if (Frame && Frame->Symbol && Frame->Symbol->File)
+		{
+			++NumSourceFiles;
+		}
+	}
+	return NumSourceFiles;
+}
+
+////////////////////////////////////////////////////////////////////////////////////////////////////
+
+void SMemAllocTableTreeView::BuildOpenSourceSubMenu(FMenuBuilder& MenuBuilder)
+{
+	MenuBuilder.BeginSection("OpenSource");
+	{
+		uint32 NumSourceFiles = 0;
+
+		FMemAllocNodePtr MemAllocNode = GetSingleSelectedMemAllocNode();
+		if (MemAllocNode.IsValid())
+		{
+			const FMemoryAlloc* Alloc = MemAllocNode->GetMemAlloc();
+			if (Alloc && Alloc->Callstack)
+			{
+				ISourceCodeAccessModule& SourceCodeAccessModule = FModuleManager::LoadModuleChecked<ISourceCodeAccessModule>("SourceCodeAccess");
+				ISourceCodeAccessor& SourceCodeAccessor = SourceCodeAccessModule.GetAccessor();
+
+				const uint32 NumCallstackFrames = Alloc->Callstack->Num();
+				for (uint32 FrameIndex = 0; FrameIndex < NumCallstackFrames; ++FrameIndex)
+				{
+					const TraceServices::FStackFrame* Frame = Alloc->Callstack->Frame(FrameIndex);
+					if (Frame && Frame->Symbol && Frame->Symbol->File)
+					{
+						FText ItemLabel;
+						FText ItemToolTip;
+						if (Frame->Symbol->GetResult() == TraceServices::ESymbolQueryResult::OK)
+						{
+							FText FileName;
+							constexpr int32 MaxFileNameLen = 120;
+							const int32 FileNameLen = FCString::Strlen(Frame->Symbol->File);
+							if (FileNameLen > MaxFileNameLen)
+							{
+								FString FileNameStr = TEXT("...") + FString(MaxFileNameLen, Frame->Symbol->File + (FileNameLen - MaxFileNameLen));
+								FileName = FText::FromString(FileNameStr);
+							}
+							else
+							{
+								FileName = FText::FromString(Frame->Symbol->File);
+							}
+
+							FText SymbolName;
+							constexpr int32 MaxSymbolNameLen = 100;
+							const int32 SymbolNameLen = FCString::Strlen(Frame->Symbol->Name);
+							if (SymbolNameLen > MaxSymbolNameLen)
+							{
+								FString SymbolNameStr = TEXT("...") + FString(MaxSymbolNameLen, Frame->Symbol->Name + (SymbolNameLen - MaxSymbolNameLen));
+								SymbolName = FText::FromString(SymbolNameStr);
+							}
+							else
+							{
+								SymbolName = FText::FromString(Frame->Symbol->Name);
+							}
+
+							ItemLabel = FText::Format(LOCTEXT("ContextMenu_OpenSource_Fmt1", "{0} ({1}) \u2192 {2}"),
+								FileName,
+								FText::AsNumber(Frame->Symbol->Line),
+								SymbolName);
+
+							ItemToolTip = FText::Format(LOCTEXT("ContextMenu_OpenSource_Desc_Fmt1", "Open source file of selected callstack frame in {0}.\n{1} (line {2})\n\u2192 {3}"),
+								SourceCodeAccessor.GetNameText(),
+								FText::FromString(Frame->Symbol->File),
+								FText::AsNumber(Frame->Symbol->Line),
+								FText::FromString(Frame->Symbol->Name));
+						}
+						else
+						{
+							ItemLabel = FText::Format(LOCTEXT("ContextMenu_OpenSource_Fmt2", "{0} ({1}) \u2192 {2}"),
+								FText::FromString(Frame->Symbol->Module),
+								FText::FromString(FString::Printf(TEXT("0x%X"), Frame->Addr)),
+								FText::FromString(TraceServices::QueryResultToString(Frame->Symbol->GetResult())));
+
+							ItemToolTip = FText::Format(LOCTEXT("ContextMenu_OpenSource_Desc_Fmt2", "Open source file of selected callstack frame in {0}."),
+								SourceCodeAccessor.GetNameText());
+						}
+
+						const bool bFileExists = FPaths::FileExists(Frame->Symbol->File);
+
+						FUIAction Action_OpenIDE
+						(
+							FExecuteAction::CreateSP(this, &SMemAllocTableTreeView::OpenSourceFileInIDE, Frame->Symbol->File, uint32(Frame->Symbol->Line)),
+							FCanExecuteAction::CreateLambda([bFileExists]() { return bFileExists; })
+						);
+						MenuBuilder.AddMenuEntry
+						(
+							ItemLabel,
+							ItemToolTip,
+							FSlateIcon(),
+							Action_OpenIDE,
+							NAME_None,
+							EUserInterfaceActionType::Button
+						);
+
+						++NumSourceFiles;
+					}
+				}
+			}
+		}
+
+		if (NumSourceFiles == 0)
+		{
+			struct FLocal
+			{
+				static bool ReturnFalse()
+				{
+					return false;
+				}
+			};
+
+			FUIAction DummyUIAction;
+			DummyUIAction.CanExecuteAction = FCanExecuteAction::CreateStatic(&FLocal::ReturnFalse);
+
+			MenuBuilder.AddMenuEntry
+			(
+				LOCTEXT("ContextMenu_OpenSourceNA", "No Source File Available"),
+				TAttribute<FText>(),
+				FSlateIcon(),
+				DummyUIAction,
+				NAME_None,
+				EUserInterfaceActionType::None
+			);
+		}
 	}
 	MenuBuilder.EndSection();
 }
 
 ////////////////////////////////////////////////////////////////////////////////////////////////////
 
+void SMemAllocTableTreeView::OpenSourceFileInIDE(const TCHAR* InFile, uint32 Line) const
+{
+	ISourceCodeAccessModule& SourceCodeAccessModule = FModuleManager::LoadModuleChecked<ISourceCodeAccessModule>("SourceCodeAccess");
+
+	const FString File = InFile;
+
+	if (FPaths::FileExists(File))
+	{
+		ISourceCodeAccessor& SourceCodeAccessor = SourceCodeAccessModule.GetAccessor();
+		SourceCodeAccessor.OpenFileAtLine(File, Line);
+	}
+	else
+	{
+		SourceCodeAccessModule.OnOpenFileFailed().Broadcast(File);
+	}
+}
+
+////////////////////////////////////////////////////////////////////////////////////////////////////
+
 bool SMemAllocTableTreeView::CanOpenCallstackFrameSourceFileInIDE() const
 {
 	if (TreeView->GetNumItemsSelected() != 1)
@@ -1343,19 +1552,7 @@ void SMemAllocTableTreeView::OpenCallstackFrameSourceFileInIDE()
 
 			if (Frame->Symbol && Frame->Symbol->File)
 			{
-				const FString File = Frame->Symbol->File;
-				const uint32 Line = Frame->Symbol->Line;
-
-				ISourceCodeAccessModule& SourceCodeAccessModule = FModuleManager::LoadModuleChecked<ISourceCodeAccessModule>("SourceCodeAccess");
-				if (FPaths::FileExists(File))
-				{
-					ISourceCodeAccessor& SourceCodeAccessor = SourceCodeAccessModule.GetAccessor();
-					SourceCodeAccessor.OpenFileAtLine(File, Line);
-				}
-				else
-				{
-					SourceCodeAccessModule.OnOpenFileFailed().Broadcast(File);
-				}
+				OpenSourceFileInIDE(Frame->Symbol->File, Frame->Symbol->Line);
 			}
 		}
 	}
diff --git a/Engine/Source/Developer/TraceInsights/Private/Insights/MemoryProfiler/Widgets/SMemAllocTableTreeView.h b/Engine/Source/Developer/TraceInsights/Private/Insights/MemoryProfiler/Widgets/SMemAllocTableTreeView.h
--- a/Engine/Source/Developer/TraceInsights/Private/Insights/MemoryProfiler/Widgets/SMemAllocTableTreeView.h
+++ b/Engine/Source/Developer/TraceInsights/Private/Insights/MemoryProfiler/Widgets/SMemAllocTableTreeView.h
@@ -14,6 +14,7 @@
 namespace Insights
 {
 
+class FMemAllocNode;
 class FMemoryRuleSpec;
 
 ////////////////////////////////////////////////////////////////////////////////////////////////////
@@ -105,10 +106,15 @@ public:
 protected:
 	virtual void InternalCreateGroupings() override;
 
-	virtual void ExtendMenu(FMenuBuilder& Menu) override;
+	TSharedPtr<FMemAllocNode> GetSingleSelectedMemAllocNode() const;
+	uint32 CountSourceFiles(FMemAllocNode& MemAllocNode);
+
+	virtual void ExtendMenu(FMenuBuilder& MenuBuilder) override;
+	void BuildOpenSourceSubMenu(FMenuBuilder& MenuBuilder);
 	bool CanOpenCallstackFrameSourceFileInIDE() const;
 	void OpenCallstackFrameSourceFileInIDE();
 	FText GetSelectedCallstackFrameFileName() const;
+	void OpenSourceFileInIDE(const TCHAR* File, uint32 Line) const;
 
 private:
 	void OnQueryInvalidated();
