[Insights]
  - Table Tree view: Added "SameValue" aggregation for the table/tree columns of CString type. Aggregated value is equal with children value if all children have exactly same value, otherwise aggregated value is empty.
  - Memory Insights: Enabled "SameValue" aggregation for the Mem Tag column of the Mem Allocs table/tree view.

diff --git a/Engine/Source/Developer/TraceInsights/Private/Insights/MemoryProfiler/ViewModels/MemAllocTable.cpp b/Engine/Source/Developer/TraceInsights/Private/Insights/MemoryProfiler/ViewModels/MemAllocTable.cpp
--- a/Engine/Source/Developer/TraceInsights/Private/Insights/MemoryProfiler/ViewModels/MemAllocTable.cpp
+++ b/Engine/Source/Developer/TraceInsights/Private/Insights/MemoryProfiler/ViewModels/MemAllocTable.cpp
@@ -728,6 +728,8 @@ void FMemAllocTable::AddDefaultColumns()
 		TSharedRef<ITableCellValueSorter> Sorter = MakeShared<FSorterByCStringValue>(ColumnRef);
 		Column.SetValueSorter(Sorter);
 
+		Column.SetAggregation(ETableColumnAggregation::SameValue);
+
 		AddColumn(ColumnRef);
 	}
 	//////////////////////////////////////////////////
diff --git a/Engine/Source/Developer/TraceInsights/Private/Insights/Table/ViewModels/TableColumn.h b/Engine/Source/Developer/TraceInsights/Private/Insights/Table/ViewModels/TableColumn.h
--- a/Engine/Source/Developer/TraceInsights/Private/Insights/Table/ViewModels/TableColumn.h
+++ b/Engine/Source/Developer/TraceInsights/Private/Insights/Table/ViewModels/TableColumn.h
@@ -42,6 +42,7 @@ enum class ETableColumnAggregation : uint32
 	Max,
 	//Average,
 	//Median,
+	SameValue
 };
 
 ////////////////////////////////////////////////////////////////////////////////////////////////////
diff --git a/Engine/Source/Developer/TraceInsights/Private/Insights/Table/Widgets/STableTreeView.cpp b/Engine/Source/Developer/TraceInsights/Private/Insights/Table/Widgets/STableTreeView.cpp
--- a/Engine/Source/Developer/TraceInsights/Private/Insights/Table/Widgets/STableTreeView.cpp
+++ b/Engine/Source/Developer/TraceInsights/Private/Insights/Table/Widgets/STableTreeView.cpp
@@ -236,7 +236,7 @@ void STableTreeView::ConstructWidget(TSharedPtr<FTable> InTablePtr)
 				.BorderImage(FAppStyle::Get().GetBrush("PopupText.Background"))
 				[
 					SNew(SHorizontalBox)
-					.Visibility_Lambda([this]() -> EVisibility 
+					.Visibility_Lambda([this]() -> EVisibility
 					{
 						return TreeViewBannerText.IsEmpty() ? EVisibility::Collapsed : EVisibility::Visible;
 					})
@@ -1494,51 +1494,64 @@ void STableTreeView::UpdateAggregatedValues(FTableTreeNode& GroupNode)
 			case ETableColumnAggregation::Sum:
 				if (Column.GetDataType() == ETableCellDataType::Float || Column.GetDataType() == ETableCellDataType::Double)
 				{
-					STableTreeView::UpdateAggregationRec<double>(Column, GroupNode, 0, true, [](double InValue, const FTableCellValue& InTableCellValue)
-					{
-						return InValue + InTableCellValue.AsDouble();
-					});
+					STableTreeView::UpdateAggregationRec<double>(Column, GroupNode, 0, true,
+						[](double InValue, const FTableCellValue& InTableCellValue)
+						{
+							return InValue + InTableCellValue.AsDouble();
+						});
 				}
 				else
 				{
-					STableTreeView::UpdateAggregationRec<int64>(Column, GroupNode, 0, true, [](int64 InValue, const FTableCellValue& InTableCellValue)
-					{
-						return InValue + InTableCellValue.AsInt64();
-					});
+					STableTreeView::UpdateAggregationRec<int64>(Column, GroupNode, 0, true,
+						[](int64 InValue, const FTableCellValue& InTableCellValue)
+						{
+							return InValue + InTableCellValue.AsInt64();
+						});
 				}
 				break;
 
 			case ETableColumnAggregation::Min:
 				if (Column.GetDataType() == ETableCellDataType::Float || Column.GetDataType() == ETableCellDataType::Double)
 				{
-					STableTreeView::UpdateAggregationRec<double>(Column, GroupNode, std::numeric_limits<double>::max(), false, [](double InValue, const FTableCellValue& InTableCellValue)
-					{
-						return FMath::Min(InValue, InTableCellValue.AsDouble());
-					});
+					STableTreeView::UpdateAggregationRec<double>(Column, GroupNode, std::numeric_limits<double>::max(), false,
+						[](double InValue, const FTableCellValue& InTableCellValue)
+						{
+							return FMath::Min(InValue, InTableCellValue.AsDouble());
+						});
 				}
 				else
 				{
-					STableTreeView::UpdateAggregationRec<int64>(Column, GroupNode, std::numeric_limits<int64>::max(), false, [](int64 InValue, const FTableCellValue& InTableCellValue)
-					{
-						return FMath::Min(InValue, InTableCellValue.AsInt64());
-					});
+					STableTreeView::UpdateAggregationRec<int64>(Column, GroupNode, std::numeric_limits<int64>::max(), false,
+						[](int64 InValue, const FTableCellValue& InTableCellValue)
+						{
+							return FMath::Min(InValue, InTableCellValue.AsInt64());
+						});
 				}
 				break;
 
 			case ETableColumnAggregation::Max:
 				if (Column.GetDataType() == ETableCellDataType::Float || Column.GetDataType() == ETableCellDataType::Double)
 				{
-					STableTreeView::UpdateAggregationRec<double>(Column, GroupNode, std::numeric_limits<double>::lowest(), false, [](double InValue, const FTableCellValue& InTableCellValue)
-					{
-						return FMath::Max(InValue, InTableCellValue.AsDouble());
-					});
+					STableTreeView::UpdateAggregationRec<double>(Column, GroupNode, std::numeric_limits<double>::lowest(), false,
+						[](double InValue, const FTableCellValue& InTableCellValue)
+						{
+							return FMath::Max(InValue, InTableCellValue.AsDouble());
+						});
 				}
 				else
 				{
-					STableTreeView::UpdateAggregationRec<int64>(Column, GroupNode, std::numeric_limits<int64>::min(), false, [](int64 InValue, const FTableCellValue& InTableCellValue)
-					{
-						return FMath::Max(InValue, InTableCellValue.AsInt64());
-					});
+					STableTreeView::UpdateAggregationRec<int64>(Column, GroupNode, std::numeric_limits<int64>::min(), false,
+						[](int64 InValue, const FTableCellValue& InTableCellValue)
+						{
+							return FMath::Max(InValue, InTableCellValue.AsInt64());
+						});
+				}
+				break;
+
+			case ETableColumnAggregation::SameValue:
+				if (Column.GetDataType() == ETableCellDataType::CString)
+				{
+					STableTreeView::UpdateCStringSameValueAggregationRec(Column, GroupNode);
 				}
 				break;
 		}
@@ -1554,6 +1567,107 @@ void STableTreeView::UpdateAggregatedValues(FTableTreeNode& GroupNode)
 
 ////////////////////////////////////////////////////////////////////////////////////////////////////
 
+void STableTreeView::UpdateCStringSameValueAggregationRec(FTableColumn& Column, FTableTreeNode& GroupNode)
+{
+	// Update child group nodes first.
+	for (FBaseTreeNodePtr NodePtr : GroupNode.GetChildren())
+	{
+		if (NodePtr->IsFiltered())
+		{
+			continue;
+		}
+		if (NodePtr->IsGroup())
+		{
+			FTableTreeNode& TableNode = *(FTableTreeNode*)NodePtr.Get();
+			UpdateCStringSameValueAggregationRec(Column, TableNode);
+		}
+	}
+
+	const TCHAR* AggregatedValue = nullptr;
+
+	// Find the first child node.
+	for (FBaseTreeNodePtr NodePtr : GroupNode.GetChildren())
+	{
+		if (NodePtr->IsFiltered())
+		{
+			continue;
+		}
+
+		if (!NodePtr->IsGroup())
+		{
+			const TOptional<FTableCellValue> NodeValue = Column.GetValue(*NodePtr);
+			if (NodeValue.IsSet() &&
+				NodeValue.GetValue().DataType == ETableCellDataType::CString)
+			{
+				AggregatedValue = NodeValue.GetValue().CString;
+			}
+		}
+		else
+		{
+			FTableTreeNode& TableNode = *(FTableTreeNode*)NodePtr.Get();
+			if (TableNode.HasAggregatedValue(Column.GetId()))
+			{
+				const FTableCellValue& ChildGroupAggregatedValue = TableNode.GetAggregatedValue(Column.GetId());
+				if (ChildGroupAggregatedValue.DataType == ETableCellDataType::CString)
+				{
+					AggregatedValue = ChildGroupAggregatedValue.CString;
+				}
+			}
+		}
+
+		break;
+	}
+
+	if (AggregatedValue != nullptr)
+	{
+		// Check if all other children have the same value as the first node.
+		for (FBaseTreeNodePtr NodePtr : GroupNode.GetChildren())
+		{
+			if (NodePtr->IsFiltered())
+			{
+				continue;
+			}
+
+			if (!NodePtr->IsGroup())
+			{
+				const TOptional<FTableCellValue> NodeValue = Column.GetValue(*NodePtr);
+				if (!NodeValue.IsSet() ||
+					NodeValue.GetValue().DataType != ETableCellDataType::CString ||
+					AggregatedValue != NodeValue.GetValue().CString)
+				{
+					AggregatedValue = nullptr;
+					break;
+				}
+			}
+			else
+			{
+				FTableTreeNode& TableNode = *(FTableTreeNode*)NodePtr.Get();
+				if (TableNode.HasAggregatedValue(Column.GetId()))
+				{
+					const FTableCellValue& ChildGroupAggregatedValue = TableNode.GetAggregatedValue(Column.GetId());
+					if (ChildGroupAggregatedValue.DataType != ETableCellDataType::CString ||
+						AggregatedValue != ChildGroupAggregatedValue.CString)
+					{
+						AggregatedValue = nullptr;
+						break;
+					}
+				}
+			}
+		}
+	}
+
+	if (AggregatedValue != nullptr)
+	{
+		GroupNode.AddAggregatedValue(Column.GetId(), FTableCellValue(AggregatedValue));
+	}
+	else
+	{
+		GroupNode.ResetAggregatedValues(Column.GetId());
+	}
+}
+
+////////////////////////////////////////////////////////////////////////////////////////////////////
+
 void STableTreeView::CreateGroupings()
 {
 	AvailableGroupings.Reset();
diff --git a/Engine/Source/Developer/TraceInsights/Private/Insights/Table/Widgets/STableTreeView.h b/Engine/Source/Developer/TraceInsights/Private/Insights/Table/Widgets/STableTreeView.h
--- a/Engine/Source/Developer/TraceInsights/Private/Insights/Table/Widgets/STableTreeView.h
+++ b/Engine/Source/Developer/TraceInsights/Private/Insights/Table/Widgets/STableTreeView.h
@@ -319,6 +319,8 @@ protected:
 		}
 	}
 
+	static void UpdateCStringSameValueAggregationRec(FTableColumn& Column, FTableTreeNode& GroupNode);
+
 	////////////////////////////////////////////////////////////////////////////////////////////////////
 	// Sorting
 
