[Insights]
  - Timing view: Fixed track invalidation when a track filter is changed (using "Filter Track..." from context menu). Also improved code related to filters attached to each track, avoiding unnecessary duplications.

diff --git a/Engine/Source/Developer/TraceInsights/Private/Insights/ViewModels/ThreadTimingTrack.cpp b/Engine/Source/Developer/TraceInsights/Private/Insights/ViewModels/ThreadTimingTrack.cpp
--- a/Engine/Source/Developer/TraceInsights/Private/Insights/ViewModels/ThreadTimingTrack.cpp
+++ b/Engine/Source/Developer/TraceInsights/Private/Insights/ViewModels/ThreadTimingTrack.cpp
@@ -1743,13 +1743,20 @@ void FThreadTimingTrack::OnFilterTrackClicked()
 		AvailableFilters->Add(MakeShared<FFilter>(static_cast<int32>(EFilterField::EndTime), LOCTEXT("EndTime", "End Time"), LOCTEXT("EndTime", "End Time"), EFilterDataType::Double, FFilterService::Get()->GetDoubleOperators()));
 		AvailableFilters->Add(MakeShared<FFilter>(static_cast<int32>(EFilterField::Duration), LOCTEXT("Duration", "Duration"), LOCTEXT("Duration", "Duration"), EFilterDataType::Double, FFilterService::Get()->GetDoubleOperators()));
 		AvailableFilters->Add(MakeShared<FFilter>(static_cast<int32>(EFilterField::TimerId), LOCTEXT("TimerId", "Timer Id"), LOCTEXT("TimerId", "Timer Id"), EFilterDataType::Int64, FFilterService::Get()->GetIntegerOperators()));
+	}
+	else
+	{
+		FilterConfigurator->GetOnChangesCommitedEvent().Remove(OnFilterChangesCommitedHandle);
 
-		OnFilterChangesCommitedHandle = FilterConfigurator->GetOnChangesCommitedEvent().AddLambda([this]()
-			{
-				this->SetDirtyFlag();
-			});
+		// Make a copy, so it will not affect other tracks that shares same filter.
+		FilterConfigurator = MakeShared<FFilterConfigurator>(*FilterConfigurator);
 	}
 
+	OnFilterChangesCommitedHandle = FilterConfigurator->GetOnChangesCommitedEvent().AddLambda([this]()
+		{
+			this->SetDirtyFlag();
+		});
+
 	FFilterService::Get()->CreateFilterConfiguratorWidget(FilterConfigurator);
 }
 
@@ -1794,16 +1801,11 @@ int32 FThreadTimingTrack::GetDepthAt(double Time) const
 
 void FThreadTimingTrack::SetFilterConfigurator(TSharedPtr<Insights::FFilterConfigurator> InFilterConfigurator)
 {
-	if (InFilterConfigurator.IsValid())
-	{
-		FilterConfigurator = MakeShared<Insights::FFilterConfigurator>(*InFilterConfigurator);
-	}
-	else
+	if (FilterConfigurator != InFilterConfigurator)
 	{
-		FilterConfigurator.Reset();
+		FilterConfigurator = InFilterConfigurator;
+		SetDirtyFlag();
 	}
-
-	SetDirtyFlag();
 }
 
 ////////////////////////////////////////////////////////////////////////////////////////////////////
diff --git a/Engine/Source/Developer/TraceInsights/Private/Insights/Widgets/STimingView.cpp b/Engine/Source/Developer/TraceInsights/Private/Insights/Widgets/STimingView.cpp
--- a/Engine/Source/Developer/TraceInsights/Private/Insights/Widgets/STimingView.cpp
+++ b/Engine/Source/Developer/TraceInsights/Private/Insights/Widgets/STimingView.cpp
@@ -5021,11 +5021,10 @@ void STimingView::FindLastEvent()
 
 void STimingView::FilterAllTracks()
 {
+	FilterConfigurator = MakeShared<Insights::FFilterConfigurator>(*QuickFindVm->GetFilterConfigurator());
 	for (auto& Entry : AllTracks)
 	{
-		FilterConfigurator = MakeShared<Insights::FFilterConfigurator>(*QuickFindVm->GetFilterConfigurator());
-		Entry.Value->SetFilterConfigurator(QuickFindVm->GetFilterConfigurator());
-		Entry.Value->SetDirtyFlag();
+		Entry.Value->SetFilterConfigurator(FilterConfigurator);
 	}
 }
 
@@ -5033,11 +5032,10 @@ void STimingView::FilterAllTracks()
 
 void STimingView::ClearFilters()
 {
+	FilterConfigurator.Reset();
 	for (auto& Entry : AllTracks)
 	{
-		FilterConfigurator.Reset();
 		Entry.Value->SetFilterConfigurator(nullptr);
-		Entry.Value->SetDirtyFlag();
 	}
 }
 
