Cooking Profiler Prototype

diff --git a/Engine/Source/Developer/TraceInsights/Private/Insights/CookProfiler/CookProfilerManager.cpp b/Engine/Source/Developer/TraceInsights/Private/Insights/CookProfiler/CookProfilerManager.cpp
new file mode 100644
--- /dev/null
+++ b/Engine/Source/Developer/TraceInsights/Private/Insights/CookProfiler/CookProfilerManager.cpp
@@ -0,0 +1,251 @@
+// Copyright Epic Games, Inc. All Rights Reserved.
+
+#include "CookProfilerManager.h"
+
+#include "Features/IModularFeatures.h"
+#include "Framework/Docking/TabManager.h"
+#include "Modules/ModuleManager.h"
+#include "TraceServices/AnalysisService.h"
+#include "TraceServices/Model/CookProfilerProvider.h"
+#include "Widgets/Docking/SDockTab.h"
+
+// Insights
+#include "Insights/InsightsStyle.h"
+#include "Insights/CookProfiler/ViewModels/PackageTable.h"
+#include "Insights/CookProfiler/Widgets/SPackageTableTreeView.h"
+#include "Insights/TimingProfilerManager.h"
+#include "Insights/Widgets/STimingProfilerWindow.h"
+
+////////////////////////////////////////////////////////////////////////////////////////////////////
+
+#define LOCTEXT_NAMESPACE "CookProfilerManager"
+
+namespace Insights
+{
+
+const FName FCookProfilerTabs::PackageTableTreeViewTabID(TEXT("PackageTableTreeView"));
+
+////////////////////////////////////////////////////////////////////////////////////////////////////
+
+TSharedPtr<FCookProfilerManager> FCookProfilerManager::Instance = nullptr;
+
+////////////////////////////////////////////////////////////////////////////////////////////////////
+
+TSharedPtr<FCookProfilerManager> FCookProfilerManager::Get()
+{
+	return FCookProfilerManager::Instance;
+}
+
+////////////////////////////////////////////////////////////////////////////////////////////////////
+
+TSharedPtr<FCookProfilerManager> FCookProfilerManager::CreateInstance()
+{
+	ensure(!FCookProfilerManager::Instance.IsValid());
+	if (FCookProfilerManager::Instance.IsValid())
+	{
+		FCookProfilerManager::Instance.Reset();
+	}
+
+	FCookProfilerManager::Instance = MakeShared<FCookProfilerManager>(FInsightsManager::Get()->GetCommandList());
+
+	return FCookProfilerManager::Instance;
+}
+
+////////////////////////////////////////////////////////////////////////////////////////////////////
+
+FCookProfilerManager::FCookProfilerManager(TSharedRef<FUICommandList> InCommandList)
+	: bIsInitialized(false)
+	, bIsAvailable(false)
+{
+}
+
+////////////////////////////////////////////////////////////////////////////////////////////////////
+
+void FCookProfilerManager::Initialize(IUnrealInsightsModule& InsightsModule)
+{
+	ensure(!bIsInitialized);
+	if (bIsInitialized)
+	{
+		return;
+	}
+	bIsInitialized = true;
+
+	// Register tick functions.
+	OnTick = FTickerDelegate::CreateSP(this, &FCookProfilerManager::Tick);
+	OnTickHandle = FTSTicker::GetCoreTicker().AddTicker(OnTick, 0.0f);
+
+	FOnRegisterMajorTabExtensions* TimingProfilerLayoutExtension = InsightsModule.FindMajorTabLayoutExtension(FInsightsManagerTabs::TimingProfilerTabId);
+	if (TimingProfilerLayoutExtension)
+	{
+		TimingProfilerLayoutExtension->AddRaw(this, &FCookProfilerManager::RegisterTimingProfilerLayoutExtensions);
+	}
+
+	FInsightsManager::Get()->GetSessionChangedEvent().AddSP(this, &FCookProfilerManager::OnSessionChanged);
+	OnSessionChanged();
+}
+
+////////////////////////////////////////////////////////////////////////////////////////////////////
+
+void FCookProfilerManager::Shutdown()
+{
+	if (!bIsInitialized)
+	{
+		return;
+	}
+	bIsInitialized = false;
+
+	FInsightsManager::Get()->GetSessionChangedEvent().RemoveAll(this);
+
+	// Unregister tick function.
+	FTSTicker::GetCoreTicker().RemoveTicker(OnTickHandle);
+
+	FCookProfilerManager::Instance.Reset();
+}
+
+////////////////////////////////////////////////////////////////////////////////////////////////////
+
+FCookProfilerManager::~FCookProfilerManager()
+{
+	ensure(!bIsInitialized);
+}
+
+////////////////////////////////////////////////////////////////////////////////////////////////////
+
+void FCookProfilerManager::RegisterMajorTabs(IUnrealInsightsModule& InsightsModule)
+{
+}
+
+////////////////////////////////////////////////////////////////////////////////////////////////////
+
+void FCookProfilerManager::UnregisterMajorTabs()
+{
+}
+
+////////////////////////////////////////////////////////////////////////////////////////////////////
+
+bool FCookProfilerManager::Tick(float DeltaTime)
+{
+	// Check if session has task events (to spawn the tab), but not too often.
+	if (!bIsAvailable && AvailabilityCheck.Tick())
+	{
+		bool bShouldBeAvailable = false;
+
+		TSharedPtr<const TraceServices::IAnalysisSession> Session = FInsightsManager::Get()->GetSession();
+		if (Session.IsValid())
+		{
+			const TraceServices::ICookProfilerProvider* CookProvider = TraceServices::ReadCookProfilerProvider(*Session.Get());
+			TraceServices::FProviderReadScopeLock ProviderReadScope(*CookProvider);
+
+			TSharedPtr<FTabManager> TabManagerShared = TimingTabManager.Pin();
+			if (CookProvider && CookProvider->GetNumPackages() > 0 && TabManagerShared.IsValid())
+			{
+				bIsAvailable = true;
+				TabManagerShared->TryInvokeTab(FCookProfilerTabs::PackageTableTreeViewTabID);
+			}
+
+			if (Session->IsAnalysisComplete())
+			{
+				// Never check again during this session.
+				AvailabilityCheck.Disable();
+			}
+		}
+		else
+		{
+			// Do not check again until the next session changed event (see OnSessionChanged).
+			AvailabilityCheck.Disable();
+		}
+	}
+
+	return true;
+}
+
+////////////////////////////////////////////////////////////////////////////////////////////////////
+
+void FCookProfilerManager::OnSessionChanged()
+{
+	bIsAvailable = false;
+	if (FInsightsManager::Get()->GetSession().IsValid())
+	{
+		AvailabilityCheck.Enable(0.5);
+	}
+	else
+	{
+		AvailabilityCheck.Disable();
+	}
+}
+
+////////////////////////////////////////////////////////////////////////////////////////////////////
+
+void FCookProfilerManager::RegisterTimingProfilerLayoutExtensions(FInsightsMajorTabExtender& InOutExtender)
+{
+	TimingTabManager = InOutExtender.GetTabManager();
+
+	FMinorTabConfig& MinorTabConfig = InOutExtender.AddMinorTabConfig();
+	MinorTabConfig.TabId = FCookProfilerTabs::PackageTableTreeViewTabID;
+	MinorTabConfig.TabLabel = LOCTEXT("PackageTableTreeViewTabTitle", "Package");
+	MinorTabConfig.TabTooltip = LOCTEXT("PackageTableTreeViewTabTitleTooltip", "Opens the Package Tree View tab, that allows cook profilling.");
+	MinorTabConfig.TabIcon = FSlateIcon(FInsightsStyle::GetStyleSetName(), "Icons.TasksView");
+	MinorTabConfig.OnSpawnTab = FOnSpawnTab::CreateRaw(this, &FCookProfilerManager::SpawnTab_PackageTableTreeView);
+	MinorTabConfig.CanSpawnTab = FCanSpawnTab::CreateRaw(this, &FCookProfilerManager::CanSpawnTab_PackageTableTreeView);
+	MinorTabConfig.WorkspaceGroup = InOutExtender.GetWorkspaceGroup();
+
+	InOutExtender.GetLayoutExtender().ExtendLayout(FTimingProfilerTabs::StatsCountersID
+		, ELayoutExtensionPosition::After
+		, FTabManager::FTab(FCookProfilerTabs::PackageTableTreeViewTabID, ETabState::ClosedTab));
+}
+
+////////////////////////////////////////////////////////////////////////////////////////////////////
+
+TSharedRef<SDockTab> FCookProfilerManager::SpawnTab_PackageTableTreeView(const FSpawnTabArgs& Args)
+{
+	TSharedRef<FPackageTable> PackageTable = MakeShared<FPackageTable>();
+	PackageTable->Reset();
+
+	const TSharedRef<SDockTab> DockTab = SNew(SDockTab)
+		.ShouldAutosize(false)
+		.TabRole(ETabRole::PanelTab)
+		[
+			SAssignNew(PackageTableTreeView, SPackageTableTreeView, PackageTable)
+		];
+
+	DockTab->SetOnTabClosed(SDockTab::FOnTabClosedCallback::CreateRaw(this, &FCookProfilerManager::OnPackageTableTreeViewTabClosed));
+
+	return DockTab;
+}
+
+////////////////////////////////////////////////////////////////////////////////////////////////////
+
+bool FCookProfilerManager::CanSpawnTab_PackageTableTreeView(const FSpawnTabArgs& Args)
+{
+	return bIsAvailable;
+}
+
+////////////////////////////////////////////////////////////////////////////////////////////////////
+
+void FCookProfilerManager::OnPackageTableTreeViewTabClosed(TSharedRef<SDockTab> TabBeingClosed)
+{
+	PackageTableTreeView.Reset();
+}
+
+////////////////////////////////////////////////////////////////////////////////////////////////////
+
+void FCookProfilerManager::OnWindowClosedEvent()
+{
+	TSharedPtr<FTabManager> TimingTabManagerSharedPtr = TimingTabManager.Pin();
+
+	if (TimingTabManagerSharedPtr.IsValid())
+	{
+		TSharedPtr<SDockTab> Tab = TimingTabManagerSharedPtr->FindExistingLiveTab(FCookProfilerTabs::PackageTableTreeViewTabID);
+		if (Tab.IsValid())
+		{
+			Tab->RequestCloseTab();
+			Tab->RemoveTabFromParent();
+		}
+	}
+}
+
+////////////////////////////////////////////////////////////////////////////////////////////////////
+
+} // namespace Insights
+
+#undef LOCTEXT_NAMESPACE
diff --git a/Engine/Source/Developer/TraceInsights/Private/Insights/CookProfiler/CookProfilerManager.h b/Engine/Source/Developer/TraceInsights/Private/Insights/CookProfiler/CookProfilerManager.h
new file mode 100644
--- /dev/null
+++ b/Engine/Source/Developer/TraceInsights/Private/Insights/CookProfiler/CookProfilerManager.h
@@ -0,0 +1,96 @@
+// Copyright Epic Games, Inc. All Rights Reserved.
+
+#pragma once
+
+#include "Containers/Ticker.h"
+#include "CoreMinimal.h"
+#include "Framework/Commands/UICommandList.h"
+
+// Insights
+#include "Insights/InsightsManager.h"
+#include "Insights/IUnrealInsightsModule.h"
+
+namespace Insights
+{
+
+class SPackageTableTreeView;
+
+struct FCookProfilerTabs
+{
+	// Tab identifiers
+	static const FName PackageTableTreeViewTabID;
+};
+
+////////////////////////////////////////////////////////////////////////////////////////////////////
+/**
+ * This class manages the Cooking Profiler state and settings.
+ */
+class FCookProfilerManager : public TSharedFromThis<FCookProfilerManager>, public IInsightsComponent
+{
+public:
+	/** Creates the Cook Profiler manager, only one instance can exist. */
+	FCookProfilerManager(TSharedRef<FUICommandList> InCommandList);
+
+	/** Virtual destructor. */
+	virtual ~FCookProfilerManager();
+
+	/** Creates an instance of the CookProfiler manager. */
+	static TSharedPtr<FCookProfilerManager> CreateInstance();
+
+	/**
+	 * @return the global instance of the Cook Profiler manager.
+	 * This is an internal singleton and cannot be used outside TraceInsights.
+	 * For external use:
+	 *     IUnrealInsightsModule& Module = FModuleManager::Get().LoadModuleChecked<IUnrealInsightsModule>("TraceInsights");
+	 *     Module.GetCookProfilerManager();
+	 */
+	static TSharedPtr<FCookProfilerManager> Get();
+
+	//////////////////////////////////////////////////
+	// IInsightsComponent
+
+	virtual void Initialize(IUnrealInsightsModule& InsightsModule) override;
+	virtual void Shutdown() override;
+	virtual void RegisterMajorTabs(IUnrealInsightsModule& InsightsModule) override;
+	virtual void UnregisterMajorTabs() override;
+	virtual void OnWindowClosedEvent() override;
+
+	//////////////////////////////////////////////////
+
+	TSharedRef<SDockTab> SpawnTab_PackageTableTreeView(const FSpawnTabArgs& Args);
+	bool CanSpawnTab_PackageTableTreeView(const FSpawnTabArgs& Args);
+	void OnPackageTableTreeViewTabClosed(TSharedRef<SDockTab> TabBeingClosed);
+
+	//////////////////////////////////////////////////
+
+	bool IsAvailable() { return bIsAvailable; }
+
+	void OnSessionChanged();
+
+private:
+	/** Updates this manager, done through FCoreTicker. */
+	bool Tick(float DeltaTime);
+
+	void RegisterTimingProfilerLayoutExtensions(FInsightsMajorTabExtender& InOutExtender);
+
+private:
+	bool bIsInitialized;
+	bool bIsAvailable;
+	FAvailabilityCheck AvailabilityCheck;
+
+	/** The delegate to be invoked when this manager ticks. */
+	FTickerDelegate OnTick;
+
+	/** Handle to the registered OnTick. */
+	FTSTicker::FDelegateHandle OnTickHandle;
+
+	/** A shared pointer to the global instance of the Cook Profiler manager. */
+	static TSharedPtr<FCookProfilerManager> Instance;
+
+	TWeakPtr<FTabManager> TimingTabManager;
+
+	TSharedPtr<SPackageTableTreeView> PackageTableTreeView;
+};
+
+} // namespace Insights
+
diff --git a/Engine/Source/Developer/TraceInsights/Private/Insights/CookProfiler/ViewModels/PackageEntry.cpp b/Engine/Source/Developer/TraceInsights/Private/Insights/CookProfiler/ViewModels/PackageEntry.cpp
new file mode 100644
--- /dev/null
+++ b/Engine/Source/Developer/TraceInsights/Private/Insights/CookProfiler/ViewModels/PackageEntry.cpp
@@ -0,0 +1,22 @@
+// Copyright Epic Games, Inc. All Rights Reserved.
+
+#include "PackageEntry.h"
+
+namespace Insights
+{
+
+////////////////////////////////////////////////////////////////////////////////////////////////////
+
+FPackageEntry::FPackageEntry(const TraceServices::FPackageData& PackageData)
+	: Id(PackageData.Id)
+	, Name(PackageData.Name)
+	, LoadTime(PackageData.LoadTime)
+	, SaveTime(PackageData.SaveTime)
+	, BeginCacheForCookedPlatformData(PackageData.BeginCacheForCookedPlatformData)
+	, IsCachedCookedPlatformDataLoaded(PackageData.IsCachedCookedPlatformDataLoaded)
+{
+}
+
+////////////////////////////////////////////////////////////////////////////////////////////////////
+
+} // namespace Insights
diff --git a/Engine/Source/Developer/TraceInsights/Private/Insights/CookProfiler/ViewModels/PackageEntry.h b/Engine/Source/Developer/TraceInsights/Private/Insights/CookProfiler/ViewModels/PackageEntry.h
new file mode 100644
--- /dev/null
+++ b/Engine/Source/Developer/TraceInsights/Private/Insights/CookProfiler/ViewModels/PackageEntry.h
@@ -0,0 +1,40 @@
+// Copyright Epic Games, Inc. All Rights Reserved.
+
+#pragma once
+
+#include "CoreMinimal.h"
+
+#include "TraceServices/Model/CookProfilerProvider.h"
+
+namespace Insights
+{
+
+////////////////////////////////////////////////////////////////////////////////////////////////////
+
+class FPackageEntry
+{
+	friend class SPackageTableTreeView;
+
+public:
+	FPackageEntry(const TraceServices::FPackageData &PackageData);
+	~FPackageEntry() {}
+
+	uint64 GetId() const { return Id; }
+	const TCHAR* GetName() const { return Name; }
+	const double GetLoadTime() const { return LoadTime; }
+	const double GetSaveTime() const { return SaveTime; }
+	const double GetBeginCacheForCookedPlatformData() const { return BeginCacheForCookedPlatformData; }
+	const double GetIsCachedCookedPlatformDataLoaded() const { return IsCachedCookedPlatformDataLoaded; }
+
+private:
+	uint64 Id;
+	const TCHAR* Name; 
+	double LoadTime; 
+	double SaveTime;
+	double BeginCacheForCookedPlatformData;
+	double IsCachedCookedPlatformDataLoaded;
+};
+
+////////////////////////////////////////////////////////////////////////////////////////////////////
+
+} // namespace Insights
diff --git a/Engine/Source/Developer/TraceInsights/Private/Insights/CookProfiler/ViewModels/PackageNode.cpp b/Engine/Source/Developer/TraceInsights/Private/Insights/CookProfiler/ViewModels/PackageNode.cpp
new file mode 100644
--- /dev/null
+++ b/Engine/Source/Developer/TraceInsights/Private/Insights/CookProfiler/ViewModels/PackageNode.cpp
@@ -0,0 +1,14 @@
+// Copyright Epic Games, Inc. All Rights Reserved.
+
+#include "PackageNode.h"
+
+#define LOCTEXT_NAMESPACE "PackageNode"
+
+namespace Insights
+{
+
+const FName FPackageNode::TypeName(TEXT("FPackageNode"));
+
+} // namespace Insights
+
+#undef LOCTEXT_NAMESPACE
diff --git a/Engine/Source/Developer/TraceInsights/Private/Insights/CookProfiler/ViewModels/PackageNode.h b/Engine/Source/Developer/TraceInsights/Private/Insights/CookProfiler/ViewModels/PackageNode.h
new file mode 100644
--- /dev/null
+++ b/Engine/Source/Developer/TraceInsights/Private/Insights/CookProfiler/ViewModels/PackageNode.h
@@ -0,0 +1,93 @@
+// Copyright Epic Games, Inc. All Rights Reserved.
+
+#pragma once
+
+#include "CoreMinimal.h"
+
+// Insights
+#include "Insights/CookProfiler/ViewModels/PackageTable.h"
+#include "Insights/CookProfiler/ViewModels/PackageEntry.h"
+#include "Insights/Table/ViewModels/TableTreeNode.h"
+
+////////////////////////////////////////////////////////////////////////////////////////////////////
+
+namespace Insights
+{
+
+enum class EPackageNodeType
+{
+	/** The TaskNode is an allocation node. */
+	Package,
+
+	/** The TaskNode is a group node. */
+	Group,
+
+	/** Invalid enum type, may be used as a number of enumerations. */
+	InvalidOrMax,
+};
+
+////////////////////////////////////////////////////////////////////////////////////////////////////
+
+class FPackageNode;
+
+/** Type definition for shared pointers to instances of FTaskNode. */
+typedef TSharedPtr<class FPackageNode> FPackageNodePtr;
+
+/** Type definition for shared references to instances of FTaskNode. */
+typedef TSharedRef<class FPackageNode> FPackageNodeRef;
+
+/** Type definition for shared references to const instances of FTaskNode. */
+typedef TSharedRef<const class FPackageNode> FPackageNodeRefConst;
+
+/** Type definition for weak references to instances of FTaskNode. */
+typedef TWeakPtr<class FTaskNode> FPackageNodeWeak;
+
+////////////////////////////////////////////////////////////////////////////////////////////////////
+/**
+ * Class used to store information about a package node (used in the SPackageTableTreeView).
+ */
+class FPackageNode : public FTableTreeNode
+{
+public:
+	static const FName TypeName;
+
+public:
+	/** Initialization constructor for the Task node. */
+	explicit FPackageNode(const FName InName, TWeakPtr<FPackageTable> InParentTable, int32 InRowIndex)
+		: FTableTreeNode(InName, InParentTable, InRowIndex)
+		, Type(EPackageNodeType::Package)
+	{
+	}
+
+	/** Initialization constructor for the group node. */
+	explicit FPackageNode(const FName InGroupName, TWeakPtr<FPackageTable> InParentTable)
+		: FTableTreeNode(InGroupName, InParentTable)
+		, Type(EPackageNodeType::Group)
+	{
+	}
+
+	virtual const FName& GetTypeName() const override { return TypeName; }
+
+	/**
+	 * @return a type of this Task node or ETaskNodeType::Group for group nodes.
+	 */
+	EPackageNodeType GetType() const { return Type; }
+
+	FPackageTable& GetPackageTableChecked() const
+	{
+		const TSharedPtr<FTable>& TablePin = GetParentTable().Pin();
+		check(TablePin.IsValid());
+		return *StaticCastSharedPtr<FPackageTable>(TablePin);
+	}
+
+	bool IsValidPackage() const { return GetPackageTableChecked().IsValidRowIndex(GetRowIndex()); }
+	const FPackageEntry* GetPackage() const { return GetPackageTableChecked().GetPackage(GetRowIndex()); }
+	const FPackageEntry& GetPackageChecked() const { return GetPackageTableChecked().GetPackageChecked(GetRowIndex()); }
+
+private:
+	const EPackageNodeType Type;
+};
+
+////////////////////////////////////////////////////////////////////////////////////////////////////
+
+} // namespace Insights
diff --git a/Engine/Source/Developer/TraceInsights/Private/Insights/CookProfiler/ViewModels/PackageTable.cpp b/Engine/Source/Developer/TraceInsights/Private/Insights/CookProfiler/ViewModels/PackageTable.cpp
new file mode 100644
--- /dev/null
+++ b/Engine/Source/Developer/TraceInsights/Private/Insights/CookProfiler/ViewModels/PackageTable.cpp
@@ -0,0 +1,272 @@
+// Copyright Epic Games, Inc. All Rights Reserved.
+
+#include "PackageTable.h"
+
+// Insights
+#include "Insights/CookProfiler/ViewModels/PackageNode.h"
+#include "Insights/CookProfiler/ViewModels/PackageTable.h"
+#include "Insights/Table/ViewModels/TableCellValueFormatter.h"
+#include "Insights/Table/ViewModels/TableCellValueGetter.h"
+#include "Insights/Table/ViewModels/TableCellValueSorter.h"
+#include "Insights/Table/ViewModels/TableColumn.h"
+
+#define LOCTEXT_NAMESPACE "Insights::FPackageTable"
+
+namespace Insights
+{
+
+////////////////////////////////////////////////////////////////////////////////////////////////////
+// Column identifiers
+
+const FName FPackageTableColumns::IdColumnId(TEXT("Id"));
+const FName FPackageTableColumns::NameColumnId(TEXT("Name"));
+const FName FPackageTableColumns::LoadTimeColumnId(TEXT("LoadTime"));
+const FName FPackageTableColumns::SaveTimeColumnId(TEXT("SaveTime"));
+const FName FPackageTableColumns::BeginCacheForCookedPlatformDataTimeColumnId(TEXT("BeginCacheForCookedPlatformDataTimeColumnId"));
+const FName FPackageTableColumns::GetIsCachedCookedPlatformDataLoaded(TEXT("GetIsCachedCookedPlatformDataLoaded"));
+
+////////////////////////////////////////////////////////////////////////////////////////////////////
+typedef FTableCellValue(*PackageFieldGetter)(const FTableColumn&, const FPackageEntry&);
+
+template<PackageFieldGetter Getter>
+class FPackageColumnValueGetter : public FTableCellValueGetter
+{
+public:
+	virtual const TOptional<FTableCellValue> GetValue(const FTableColumn& Column, const FBaseTreeNode& Node) const override
+	{
+		if (Node.IsGroup())
+		{
+			const FTableTreeNode& NodePtr = static_cast<const FTableTreeNode&>(Node);
+			if (NodePtr.HasAggregatedValue(Column.GetId()))
+			{
+				return NodePtr.GetAggregatedValue(Column.GetId());
+			}
+		}
+		else //if (Node->Is<FackageNode>())
+		{
+			const FPackageNode& PackageNode = static_cast<const FPackageNode&>(Node);
+			const FPackageEntry* Package = PackageNode.GetPackage();
+			if (Package)
+			{
+				return Getter(Column, *Package);
+			}
+		}
+
+		return TOptional<FTableCellValue>();
+	}
+};
+
+////////////////////////////////////////////////////////////////////////////////////////////////////
+
+struct DefaultPackageFieldGetterFuncts
+{
+	static FTableCellValue GetId(const FTableColumn& Column, const FPackageEntry& Package) { return FTableCellValue((int64)Package.GetId());	}
+	static FTableCellValue GetName(const FTableColumn& Column, const FPackageEntry& Package) { return FTableCellValue((int64)Package.GetName());	}
+	static FTableCellValue GetLoadTime(const FTableColumn& Column, const FPackageEntry& Package) { return FTableCellValue((double)Package.GetLoadTime());	}
+	static FTableCellValue GetSaveTime(const FTableColumn& Column, const FPackageEntry& Package) { return FTableCellValue((double)Package.GetSaveTime());	}
+	static FTableCellValue GetBeginCacheForCookedPlatformData(const FTableColumn& Column, const FPackageEntry& Package) { return FTableCellValue((double)Package.GetBeginCacheForCookedPlatformData());	}
+	static FTableCellValue GetIsCachedCookedPlatformDataLoaded(const FTableColumn& Column, const FPackageEntry& Package) { return FTableCellValue((double)Package.GetIsCachedCookedPlatformDataLoaded());	}
+};
+
+////////////////////////////////////////////////////////////////////////////////////////////////////
+// FTaskTable
+////////////////////////////////////////////////////////////////////////////////////////////////////
+
+FPackageTable::FPackageTable()
+{
+}
+
+////////////////////////////////////////////////////////////////////////////////////////////////////
+
+FPackageTable::~FPackageTable()
+{
+}
+
+////////////////////////////////////////////////////////////////////////////////////////////////////
+
+void FPackageTable::Reset()
+{
+	FTable::Reset();
+
+	AddDefaultColumns();
+}
+
+////////////////////////////////////////////////////////////////////////////////////////////////////
+
+void FPackageTable::AddDefaultColumns()
+{
+	//////////////////////////////////////////////////
+	// Hierarchy Column
+	{
+		const int32 HierarchyColumnIndex = -1;
+		const TCHAR* HierarchyColumnName = nullptr;
+		AddHierarchyColumn(HierarchyColumnIndex, HierarchyColumnName);
+
+		const TSharedRef<FTableColumn>& ColumnRef = GetColumns()[0];
+		ColumnRef->SetInitialWidth(200.0f);
+		ColumnRef->SetShortName(LOCTEXT("PackageColumnName", "Hierarchy"));
+		ColumnRef->SetTitleName(LOCTEXT("PackageColumnTitle", "Package Hierarchy"));
+		ColumnRef->SetDescription(LOCTEXT("PackageColumnDesc", "Hierarchy of the package's tree"));
+	}
+
+	int32 ColumnIndex = 0;
+
+	//////////////////////////////////////////////////
+	// Id Column
+	{
+		TSharedRef<FTableColumn> ColumnRef = MakeShared<FTableColumn>(FPackageTableColumns::IdColumnId);
+		FTableColumn& Column = *ColumnRef;
+
+		Column.SetIndex(ColumnIndex++);
+
+		Column.SetShortName(LOCTEXT("CreatedTimestampColumnName", "Id"));
+		Column.SetTitleName(LOCTEXT("CreatedTimestampColumnTitle", "Id"));
+		Column.SetDescription(LOCTEXT("CreatedTimestampColumnDesc", "The id of the package."));
+
+		Column.SetFlags(ETableColumnFlags::ShouldBeVisible | ETableColumnFlags::CanBeHidden | ETableColumnFlags::CanBeFiltered);
+
+		Column.SetHorizontalAlignment(HAlign_Left);
+		Column.SetInitialWidth(100.0f);
+
+		Column.SetDataType(ETableCellDataType::Int64);
+
+		TSharedRef<ITableCellValueGetter> Getter = MakeShared<FPackageColumnValueGetter<DefaultPackageFieldGetterFuncts::GetId>>();
+		Column.SetValueGetter(Getter);
+
+		TSharedRef<ITableCellValueFormatter> Formatter = MakeShared<FInt64ValueFormatterAsNumber>();
+		Column.SetValueFormatter(Formatter);
+
+		TSharedRef<ITableCellValueSorter> Sorter = MakeShared<FSorterByInt64Value>(ColumnRef);
+		Column.SetValueSorter(Sorter);
+
+		Column.SetAggregation(ETableColumnAggregation::None);
+
+		AddColumn(ColumnRef);
+	}
+	//////////////////////////////////////////////////
+	// Load Time Column
+	{
+		TSharedRef<FTableColumn> ColumnRef = MakeShared<FTableColumn>(FPackageTableColumns::LoadTimeColumnId);
+		FTableColumn& Column = *ColumnRef;
+
+		Column.SetIndex(ColumnIndex++);
+
+		Column.SetShortName(LOCTEXT("LoadTimeColumnName", "LoadTime"));
+		Column.SetTitleName(LOCTEXT("LoadTimeColumnTitle", "LoadTime"));
+		Column.SetDescription(LOCTEXT("LoadTimeColumnDesc", "The time it took to load the package."));
+
+		Column.SetFlags(ETableColumnFlags::ShouldBeVisible | ETableColumnFlags::CanBeHidden | ETableColumnFlags::CanBeFiltered);
+
+		Column.SetHorizontalAlignment(HAlign_Left);
+		Column.SetInitialWidth(100.0f);
+
+		Column.SetDataType(ETableCellDataType::Double);
+
+		TSharedRef<ITableCellValueGetter> Getter = MakeShared<FPackageColumnValueGetter<DefaultPackageFieldGetterFuncts::GetLoadTime>>();
+		Column.SetValueGetter(Getter);
+
+		TSharedRef<ITableCellValueFormatter> Formatter = MakeShared<FDoubleValueFormatterAsTimeAuto>();
+		Column.SetValueFormatter(Formatter);
+
+		TSharedRef<ITableCellValueSorter> Sorter = MakeShared<FSorterByDoubleValue>(ColumnRef);
+		Column.SetValueSorter(Sorter);
+
+		AddColumn(ColumnRef);
+	}
+	//////////////////////////////////////////////////
+	// Save Time Column
+	{
+		TSharedRef<FTableColumn> ColumnRef = MakeShared<FTableColumn>(FPackageTableColumns::SaveTimeColumnId);
+		FTableColumn& Column = *ColumnRef;
+
+		Column.SetIndex(ColumnIndex++);
+
+		Column.SetShortName(LOCTEXT("SaveTimeColumnName", "SaveTime"));
+		Column.SetTitleName(LOCTEXT("SaveTimeColumnTitle", "SaveTime"));
+		Column.SetDescription(LOCTEXT("SaveTimeColumnDesc", "The time it took to save the package."));
+
+		Column.SetFlags(ETableColumnFlags::ShouldBeVisible | ETableColumnFlags::CanBeHidden | ETableColumnFlags::CanBeFiltered);
+
+		Column.SetHorizontalAlignment(HAlign_Left);
+		Column.SetInitialWidth(100.0f);
+
+		Column.SetDataType(ETableCellDataType::Double);
+
+		TSharedRef<ITableCellValueGetter> Getter = MakeShared<FPackageColumnValueGetter<DefaultPackageFieldGetterFuncts::GetSaveTime>>();
+		Column.SetValueGetter(Getter);
+
+		TSharedRef<ITableCellValueFormatter> Formatter = MakeShared<FDoubleValueFormatterAsTimeAuto>();
+		Column.SetValueFormatter(Formatter);
+
+		TSharedRef<ITableCellValueSorter> Sorter = MakeShared<FSorterByDoubleValue>(ColumnRef);
+		Column.SetValueSorter(Sorter);
+
+		AddColumn(ColumnRef);
+	}
+	//////////////////////////////////////////////////
+	// BeginCacheForCookedPlatformData Time Column
+	{
+		TSharedRef<FTableColumn> ColumnRef = MakeShared<FTableColumn>(FPackageTableColumns::BeginCacheForCookedPlatformDataTimeColumnId);
+		FTableColumn& Column = *ColumnRef;
+
+		Column.SetIndex(ColumnIndex++);
+
+		Column.SetShortName(LOCTEXT("BeginCacheForCookedPlatformDataColumnName", "BeginCache"));
+		Column.SetTitleName(LOCTEXT("BeginCacheForCookedPlatformDataColumnTitle", "BeginCacheForCookedPlatformData"));
+		Column.SetDescription(LOCTEXT("BeginCacheForCookedPlatformDataColumnDesc", "The total time spent in the BeginCacheForCookedPlatformData function for the package."));
+
+		Column.SetFlags(ETableColumnFlags::ShouldBeVisible | ETableColumnFlags::CanBeHidden | ETableColumnFlags::CanBeFiltered);
+
+		Column.SetHorizontalAlignment(HAlign_Left);
+		Column.SetInitialWidth(100.0f);
+
+		Column.SetDataType(ETableCellDataType::Double);
+
+		TSharedRef<ITableCellValueGetter> Getter = MakeShared<FPackageColumnValueGetter<DefaultPackageFieldGetterFuncts::GetBeginCacheForCookedPlatformData>>();
+		Column.SetValueGetter(Getter);
+
+		TSharedRef<ITableCellValueFormatter> Formatter = MakeShared<FDoubleValueFormatterAsTimeAuto>();
+		Column.SetValueFormatter(Formatter);
+
+		TSharedRef<ITableCellValueSorter> Sorter = MakeShared<FSorterByDoubleValue>(ColumnRef);
+		Column.SetValueSorter(Sorter);
+
+		AddColumn(ColumnRef);
+	}	
+	//////////////////////////////////////////////////
+	// BeginCacheForCookedPlatformData Time Column
+	{
+		TSharedRef<FTableColumn> ColumnRef = MakeShared<FTableColumn>(FPackageTableColumns::GetIsCachedCookedPlatformDataLoaded);
+		FTableColumn& Column = *ColumnRef;
+
+		Column.SetIndex(ColumnIndex++);
+
+		Column.SetShortName(LOCTEXT("GetIsCachedCookedPlatformDataLoadedColumnName", "IsCachedCooked"));
+		Column.SetTitleName(LOCTEXT("GetIsCachedCookedPlatformDataLoadedColumnTitle", "IsCachedCookedPlatformDataLoaded"));
+		Column.SetDescription(LOCTEXT("GetIsCachedCookedPlatformDataLoadedColumnDesc", "The total time spent in the IsCachedCookedPlatformDataLoaded function for the package."));
+
+		Column.SetFlags(ETableColumnFlags::ShouldBeVisible | ETableColumnFlags::CanBeHidden | ETableColumnFlags::CanBeFiltered);
+
+		Column.SetHorizontalAlignment(HAlign_Left);
+		Column.SetInitialWidth(100.0f);
+
+		Column.SetDataType(ETableCellDataType::Double);
+
+		TSharedRef<ITableCellValueGetter> Getter = MakeShared<FPackageColumnValueGetter<DefaultPackageFieldGetterFuncts::GetIsCachedCookedPlatformDataLoaded>>();
+		Column.SetValueGetter(Getter);
+
+		TSharedRef<ITableCellValueFormatter> Formatter = MakeShared<FDoubleValueFormatterAsTimeAuto>();
+		Column.SetValueFormatter(Formatter);
+
+		TSharedRef<ITableCellValueSorter> Sorter = MakeShared<FSorterByDoubleValue>(ColumnRef);
+		Column.SetValueSorter(Sorter);
+
+		AddColumn(ColumnRef);
+	}
+}
+
+////////////////////////////////////////////////////////////////////////////////////////////////////
+
+} // namespace Insights
+
+#undef LOCTEXT_NAMESPACE
diff --git a/Engine/Source/Developer/TraceInsights/Private/Insights/CookProfiler/ViewModels/PackageTable.h b/Engine/Source/Developer/TraceInsights/Private/Insights/CookProfiler/ViewModels/PackageTable.h
new file mode 100644
--- /dev/null
+++ b/Engine/Source/Developer/TraceInsights/Private/Insights/CookProfiler/ViewModels/PackageTable.h
@@ -0,0 +1,55 @@
+// Copyright Epic Games, Inc. All Rights Reserved.
+
+#pragma once
+
+#include "CoreMinimal.h"
+
+// Insights
+#include "Insights/CookProfiler/ViewModels/PackageEntry.h"
+#include "Insights/Table/ViewModels/Table.h"
+
+namespace Insights
+{
+
+class FTableColumn;
+
+////////////////////////////////////////////////////////////////////////////////////////////////////
+
+// Column identifiers
+struct FPackageTableColumns
+{
+	static const FName IdColumnId;
+	static const FName NameColumnId;
+	static const FName LoadTimeColumnId;
+	static const FName SaveTimeColumnId;
+	static const FName BeginCacheForCookedPlatformDataTimeColumnId;
+	static const FName GetIsCachedCookedPlatformDataLoaded;
+};
+
+////////////////////////////////////////////////////////////////////////////////////////////////////
+
+class FPackageTable : public FTable
+{
+public:
+	FPackageTable();
+	virtual ~FPackageTable();
+
+	virtual void Reset();
+
+	TArray<FPackageEntry>& GetPackageEntries() { return PackageEntries; }
+	const TArray<FPackageEntry>& GetPackageEntries() const { return PackageEntries; }
+
+	bool IsValidRowIndex(int32 InIndex) const { return InIndex >= 0 && InIndex < PackageEntries.Num(); }
+	const FPackageEntry* GetPackage(int32 InIndex) const { return IsValidRowIndex(InIndex) ? &PackageEntries[InIndex] : nullptr; }
+	const FPackageEntry& GetPackageChecked(int32 InIndex) const { check(IsValidRowIndex(InIndex)); return PackageEntries[InIndex]; }
+
+private:
+	void AddDefaultColumns();
+
+private:
+	TArray<FPackageEntry> PackageEntries;
+};
+
+////////////////////////////////////////////////////////////////////////////////////////////////////
+
+} // namespace Insights
diff --git a/Engine/Source/Developer/TraceInsights/Private/Insights/CookProfiler/Widgets/SPackageTableTreeView.cpp b/Engine/Source/Developer/TraceInsights/Private/Insights/CookProfiler/Widgets/SPackageTableTreeView.cpp
new file mode 100644
--- /dev/null
+++ b/Engine/Source/Developer/TraceInsights/Private/Insights/CookProfiler/Widgets/SPackageTableTreeView.cpp
@@ -0,0 +1,238 @@
+// Copyright Epic Games, Inc. All Rights Reserved.
+
+#include "SPackageTableTreeView.h"
+
+#include "Framework/MultiBox/MultiBoxBuilder.h"
+#include "SlateOptMacros.h"
+#include "TraceServices/AnalysisService.h"
+
+// Insights
+#include "Insights/InsightsStyle.h"
+#include "Insights/CookProfiler/CookProfilerManager.h"
+#include "Insights/CookProfiler/ViewModels/PackageEntry.h"
+#include "Insights/CookProfiler/ViewModels/PackageNode.h"
+#include "Insights/Table/ViewModels/TableColumn.h"
+#include "Insights/TimingProfilerManager.h"
+#include "Insights/ViewModels/FilterConfigurator.h"
+#include "Insights/ViewModels/ThreadTimingTrack.h"
+#include "Insights/Widgets/STimingProfilerWindow.h"
+#include "Insights/Widgets/STimingView.h"
+
+#include <limits>
+
+#define LOCTEXT_NAMESPACE "SPackageTableTreeView"
+
+using namespace TraceServices;
+
+namespace Insights
+{
+
+////////////////////////////////////////////////////////////////////////////////////////////////////
+// FTaskTableTreeViewCommands
+////////////////////////////////////////////////////////////////////////////////////////////////////
+
+class FPackageTableTreeViewCommands : public TCommands<FPackageTableTreeViewCommands>
+{
+public:
+	FPackageTableTreeViewCommands()
+	: TCommands<FPackageTableTreeViewCommands>(
+		TEXT("PackageTableTreeViewCommands"),
+		NSLOCTEXT("Contexts", "PackageTableTreeViewCommands", "Insights - Package Table Tree View"),
+		NAME_None,
+		FInsightsStyle::GetStyleSetName())
+	{
+	}
+
+	virtual ~FPackageTableTreeViewCommands()
+	{
+	}
+
+	// UI_COMMAND takes long for the compiler to optimize
+	PRAGMA_DISABLE_OPTIMIZATION
+	virtual void RegisterCommands() override
+	{
+	}
+	PRAGMA_ENABLE_OPTIMIZATION
+};
+
+////////////////////////////////////////////////////////////////////////////////////////////////////
+
+SPackageTableTreeView::SPackageTableTreeView()
+{
+	bRunInAsyncMode = true;
+}
+
+////////////////////////////////////////////////////////////////////////////////////////////////////
+
+SPackageTableTreeView::~SPackageTableTreeView()
+{
+}
+
+////////////////////////////////////////////////////////////////////////////////////////////////////
+
+void SPackageTableTreeView::Construct(const FArguments& InArgs, TSharedPtr<FPackageTable> InTablePtr)
+{
+	ConstructWidget(InTablePtr);
+
+	AddCommmands();
+}
+
+////////////////////////////////////////////////////////////////////////////////////////////////////
+
+void SPackageTableTreeView::ExtendMenu(FMenuBuilder& MenuBuilder)
+{
+}
+
+void SPackageTableTreeView::AddCommmands()
+{
+	FPackageTableTreeViewCommands::Register();
+}
+
+////////////////////////////////////////////////////////////////////////////////////////////////////
+
+void SPackageTableTreeView::Reset()
+{
+	STableTreeView::Reset();
+}
+
+////////////////////////////////////////////////////////////////////////////////////////////////////
+
+void SPackageTableTreeView::Tick(const FGeometry& AllottedGeometry, const double InCurrentTime, const float InDeltaTime)
+{
+	STableTreeView::Tick(AllottedGeometry, InCurrentTime, InDeltaTime);
+
+	if (!bIsUpdateRunning)
+	{
+		RebuildTree(false);
+	}
+}
+
+////////////////////////////////////////////////////////////////////////////////////////////////////
+
+void SPackageTableTreeView::RebuildTree(bool bResync)
+{
+	if (bDataLoaded == true)
+	{
+		return;
+	}
+
+	TraceServices::FAnalysisSessionReadScope SessionReadScope(*Session.Get());
+
+	if (!Session->IsAnalysisComplete())
+	{
+		return;
+	}
+
+	TSharedPtr<FPackageTable> PackageTable = GetPackageTable();
+	TArray<FPackageEntry>& Packages = PackageTable->GetPackageEntries();
+	Packages.Empty();
+	TableTreeNodes.Empty();
+
+	const TraceServices::ICookProfilerProvider* CookProvider = TraceServices::ReadCookProfilerProvider(*Session.Get());
+
+	if (CookProvider)
+	{
+		TraceServices::FProviderReadScopeLock ProviderReadScope(*CookProvider);
+		TArray<FTableTreeNodePtr>* Nodes = &TableTreeNodes;
+
+		CookProvider->EnumeratePackages(0, 0, [&Packages, &PackageTable, Nodes](const TraceServices::FPackageData& Package)
+			{
+				Packages.Emplace(Package);
+				uint32 Index = Packages.Num() - 1;
+				FName NodeName(Package.Name);
+				FPackageNodePtr NodePtr = MakeShared<FPackageNode>(NodeName, PackageTable, Index);
+				Nodes->Add(NodePtr);
+				return true;
+			});
+	}
+
+	bDataLoaded = true;
+
+	UpdateTree();
+	TreeView->RebuildList();
+}
+
+////////////////////////////////////////////////////////////////////////////////////////////////////
+
+bool SPackageTableTreeView::IsRunning() const
+{
+	return STableTreeView::IsRunning();
+}
+
+////////////////////////////////////////////////////////////////////////////////////////////////////
+
+double SPackageTableTreeView::GetAllOperationsDuration()
+{
+	return STableTreeView::GetAllOperationsDuration();
+}
+
+////////////////////////////////////////////////////////////////////////////////////////////////////
+
+FText SPackageTableTreeView::GetCurrentOperationName() const
+{
+	return STableTreeView::GetCurrentOperationName();
+}
+
+////////////////////////////////////////////////////////////////////////////////////////////////////
+
+BEGIN_SLATE_FUNCTION_BUILD_OPTIMIZATION
+TSharedPtr<SWidget> SPackageTableTreeView::ConstructToolbar()
+{
+	return nullptr;
+}
+
+////////////////////////////////////////////////////////////////////////////////////////////////////
+
+TSharedPtr<SWidget> SPackageTableTreeView::ConstructFooter()
+{
+	return nullptr;
+}
+END_SLATE_FUNCTION_BUILD_OPTIMIZATION
+
+////////////////////////////////////////////////////////////////////////////////////////////////////
+
+void SPackageTableTreeView::ApplyColumnConfig(const TArrayView<FColumnConfig>& Preset)
+{
+	for (const TSharedRef<FTableColumn>& ColumnRef : Table->GetColumns())
+	{
+		FTableColumn& Column = ColumnRef.Get();
+		for (const FColumnConfig& Config : Preset)
+		{
+			if (Column.GetId() == Config.ColumnId)
+			{
+				if (Config.bIsVisible)
+				{
+					ShowColumn(Column);
+					if (Config.Width > 0.0f)
+					{
+						TreeViewHeaderRow->SetColumnWidth(Column.GetId(), Config.Width);
+					}
+				}
+				else
+				{
+					HideColumn(Column);
+				}
+			}
+		}
+	}
+}
+
+////////////////////////////////////////////////////////////////////////////////////////////////////
+
+void SPackageTableTreeView::InternalCreateGroupings()
+{
+	STableTreeView::InternalCreateGroupings();
+}
+
+////////////////////////////////////////////////////////////////////////////////////////////////////
+
+void SPackageTableTreeView::TreeView_OnMouseButtonDoubleClick(FTableTreeNodePtr TreeNode)
+{
+	STableTreeView::TreeView_OnMouseButtonDoubleClick(TreeNode);
+}
+
+////////////////////////////////////////////////////////////////////////////////////////////////////
+
+} // namespace Insights
+
+#undef LOCTEXT_NAMESPACE
diff --git a/Engine/Source/Developer/TraceInsights/Private/Insights/CookProfiler/Widgets/SPackageTableTreeView.h b/Engine/Source/Developer/TraceInsights/Private/Insights/CookProfiler/Widgets/SPackageTableTreeView.h
new file mode 100644
--- /dev/null
+++ b/Engine/Source/Developer/TraceInsights/Private/Insights/CookProfiler/Widgets/SPackageTableTreeView.h
@@ -0,0 +1,94 @@
+// Copyright Epic Games, Inc. All Rights Reserved.
+
+#pragma once
+
+#include "CoreMinimal.h"
+#include "TraceServices/Model/CookProfilerProvider.h"
+
+// Insights
+#include "Insights/CookProfiler/ViewModels/PackageTable.h"
+#include "Insights/Table/Widgets/STableTreeView.h"
+
+class FMenuBuilder;
+
+namespace Insights
+{
+
+////////////////////////////////////////////////////////////////////////////////////////////////////
+
+class SPackageTableTreeView : public STableTreeView
+{
+private:
+	struct FColumnConfig
+	{
+		const FName& ColumnId;
+		bool bIsVisible;
+		float Width;
+	};
+
+public:
+	/** Default constructor. */
+	SPackageTableTreeView();
+
+	/** Virtual destructor. */
+	virtual ~SPackageTableTreeView();
+
+	SLATE_BEGIN_ARGS(SPackageTableTreeView) {}
+	SLATE_END_ARGS()
+
+	/**
+	 * Construct this widget
+	 * @param InArgs - The declaration data for this widget
+	 */
+	void Construct(const FArguments& InArgs, TSharedPtr<FPackageTable> InTablePtr);
+
+	virtual TSharedPtr<SWidget> ConstructToolbar() override;
+	virtual TSharedPtr<SWidget> ConstructFooter() override;
+
+	TSharedPtr<FPackageTable> GetPackageTable() { return StaticCastSharedPtr<FPackageTable>(GetTable()); }
+	const TSharedPtr<FPackageTable> GetPackageTable() const { return StaticCastSharedPtr<FPackageTable>(GetTable()); }
+
+	virtual void Reset();
+
+	/**
+	 * Ticks this widget.  Override in derived classes, but always call the parent implementation.
+	 *
+	 * @param  AllottedGeometry The space allotted for this widget
+	 * @param  InCurrentTime  Current absolute real time
+	 * @param  InDeltaTime  Real time passed since last tick
+	 */
+	virtual void Tick(const FGeometry& AllottedGeometry, const double InCurrentTime, const float InDeltaTime) override;
+
+	/**
+	 * Rebuilds the tree (if necessary).
+	 * @param bResync - If true, it forces a resync even if the list did not changed since last sync.
+	 */
+	virtual void RebuildTree(bool bResync);
+
+	////////////////////////////////////////////////////////////////////////////////////////////////////
+	// IAsyncOperationStatusProvider implementation
+
+	virtual bool IsRunning() const override;
+	virtual double GetAllOperationsDuration() override;
+	virtual FText GetCurrentOperationName() const override;
+
+	////////////////////////////////////////////////////////////////////////////////////////////////////
+
+	virtual void TreeView_OnMouseButtonDoubleClick(FTableTreeNodePtr TreeNode) override;
+
+protected:
+	virtual void InternalCreateGroupings() override;
+
+	virtual void ExtendMenu(FMenuBuilder& MenuBuilder) override;
+
+private:
+	void ApplyColumnConfig(const TArrayView<FColumnConfig>& Preset);
+	void AddCommmands();
+
+private:
+	bool bDataLoaded = false;
+};
+
+////////////////////////////////////////////////////////////////////////////////////////////////////
+
+} // namespace Insights
diff --git a/Engine/Source/Developer/TraceInsights/Private/Insights/TraceInsightsModule.cpp b/Engine/Source/Developer/TraceInsights/Private/Insights/TraceInsightsModule.cpp
--- a/Engine/Source/Developer/TraceInsights/Private/Insights/TraceInsightsModule.cpp
+++ b/Engine/Source/Developer/TraceInsights/Private/Insights/TraceInsightsModule.cpp
@@ -17,6 +17,7 @@
 
 // Insights
 #include "Insights/ContextSwitches/ContextSwitchesProfilerManager.h"
+#include "Insights/CookProfiler/CookProfilerManager.h"
 #include "Insights/InsightsManager.h"
 #include "Insights/InsightsStyle.h"
 #include "Insights/IUnrealInsightsModule.h"
@@ -62,6 +63,7 @@ void FTraceInsightsModule::StartupModule()
 	RegisterComponent(FMemoryProfilerManager::CreateInstance());
 	RegisterComponent(Insights::FTaskGraphProfilerManager::CreateInstance());
 	RegisterComponent(Insights::FContextSwitchesProfilerManager::CreateInstance());
+	RegisterComponent(Insights::FCookProfilerManager::CreateInstance());
 
 #if !WITH_EDITOR
 	ISourceCodeAccessModule& SourceCodeAccessModule = FModuleManager::LoadModuleChecked<ISourceCodeAccessModule>("SourceCodeAccess");
diff --git a/Engine/Source/Developer/TraceServices/Private/Analyzers/CookAnalysis.cpp b/Engine/Source/Developer/TraceServices/Private/Analyzers/CookAnalysis.cpp
new file mode 100644
--- /dev/null
+++ b/Engine/Source/Developer/TraceServices/Private/Analyzers/CookAnalysis.cpp
@@ -0,0 +1,95 @@
+// Copyright Epic Games, Inc. All Rights Reserved.
+#include "CookAnalysis.h"
+
+#include "Analyzers/MiscTraceAnalysis.h"
+#include "Common/FormatArgs.h"
+#include "Common/Utils.h"
+#include "Model/CookProfilerProviderPrivate.h"
+#include "ProfilingDebugging/CookStats.h"
+
+#include <limits>
+
+namespace TraceServices
+{
+
+FCookAnalyzer::FCookAnalyzer(IAnalysisSession& InSession, FCookProfilerProvider& InCookProfilerProvider)
+	: Session(InSession)
+	, CookProfilerProvider(InCookProfilerProvider)
+{
+}
+
+FCookAnalyzer::~FCookAnalyzer()
+{
+}
+
+void FCookAnalyzer::OnAnalysisBegin(const FOnAnalysisContext& Context)
+{
+	auto& Builder = Context.InterfaceBuilder;
+
+	Builder.RouteEvent(RouteId_Package, "CookTrace", "Package");
+	Builder.RouteEvent(RouteId_PackageStat, "CookTrace", "PackageStat");
+}
+
+bool FCookAnalyzer::OnEvent(uint16 RouteId, EStyle Style, const FOnEventContext& Context)
+{
+	const auto& EventData = Context.EventData;
+	switch (RouteId)
+	{
+	case RouteId_Package:
+	{
+		uint64 Id = EventData.GetValue<uint64>("Id");
+		FString Name;
+		EventData.GetString("Name", Name);
+
+		TraceServices::FProviderEditScopeLock ProviderEditScope(CookProfilerProvider);
+		CookProfilerProvider.AddPackage(Id, Name);
+		break;
+	}
+	case RouteId_PackageStat:
+	{
+		uint64 Id = EventData.GetValue<uint64>("Id");
+		double Value = Context.EventTime.AsSecondsAbsolute(EventData.GetValue<uint64>("Duration"));
+		EPackageEventStatType StatType = (EPackageEventStatType) EventData.GetValue<uint8>("StatType");
+
+		TraceServices::FProviderEditScopeLock ProviderEditScope(CookProfilerProvider);
+		FPackageData* Package = CookProfilerProvider.EditPackage(Id);
+		if (Package)
+		{
+			switch (StatType)
+			{
+			case EPackageEventStatType::LoadPackage:
+			{
+				Package->LoadTime = Value;
+				break;
+			}
+			case EPackageEventStatType::SavePackage:
+			{
+				Package->SaveTime = Value;
+				break;
+			}
+			case EPackageEventStatType::BeginCacheForCookedPlatformData:
+			{
+				// A BeginCacheForCookedPlatformData event is received for each asset in the package. We add the values to get the total time for the package.
+				Package->BeginCacheForCookedPlatformData += Value;
+				break;
+			}
+			case EPackageEventStatType::IsCachedCookedPlatformDataLoaded:
+			{
+				// A IsCachedCookedPlatformDataLoaded event is received for each asset in the package. We add the values to get the total time for the package.
+				Package->IsCachedCookedPlatformDataLoaded += Value;
+				break;
+			}
+			default:
+			{
+				break;
+			}
+			}
+		}
+		break;
+	}
+	}
+
+	return true;
+}
+
+} // namespace TraceServices
diff --git a/Engine/Source/Developer/TraceServices/Private/Analyzers/CookAnalysis.h b/Engine/Source/Developer/TraceServices/Private/Analyzers/CookAnalysis.h
new file mode 100644
--- /dev/null
+++ b/Engine/Source/Developer/TraceServices/Private/Analyzers/CookAnalysis.h
@@ -0,0 +1,32 @@
+// Copyright Epic Games, Inc. All Rights Reserved.
+
+#pragma once
+
+#include "Trace/Analyzer.h"
+#include "AnalysisServicePrivate.h"
+
+namespace TraceServices
+{
+class FCookProfilerProvider;
+
+class FCookAnalyzer : public UE::Trace::IAnalyzer
+{
+public:
+	FCookAnalyzer(IAnalysisSession& Session, FCookProfilerProvider& CookProfilerProvider);
+	virtual ~FCookAnalyzer();
+
+	virtual void OnAnalysisBegin(const FOnAnalysisContext& Context) override;
+	virtual bool OnEvent(uint16 RouteId, EStyle Style, const FOnEventContext& Context) override;
+	
+private:
+	enum : uint16
+	{
+		RouteId_Package,
+		RouteId_PackageStat,
+	};
+
+	IAnalysisSession& Session;
+	FCookProfilerProvider& CookProfilerProvider;
+};
+
+} // namespace TraceServices
diff --git a/Engine/Source/Developer/TraceServices/Private/Common/ProviderLock.cpp b/Engine/Source/Developer/TraceServices/Private/Common/ProviderLock.cpp
new file mode 100644
--- /dev/null
+++ b/Engine/Source/Developer/TraceServices/Private/Common/ProviderLock.cpp
@@ -0,0 +1,59 @@
+// Copyright Epic Games, Inc. All Rights Reserved.
+
+#include "Common/ProviderLock.h"
+
+namespace TraceServices
+{
+void FProviderLock::ReadAccessCheck(const FProviderLock* CurrentProviderLock, const int32& CurrentReadProviderLockCount, const int32& CurrentWriteAllocationsProviderLockCount) const
+{
+	checkf(CurrentProviderLock == this && (CurrentReadProviderLockCount > 0 || CurrentWriteAllocationsProviderLockCount > 0),
+		TEXT("Trying to READ from allocations provider outside of a READ scope"));
+}
+
+void FProviderLock::WriteAccessCheck(const int32& CurrentWriteAllocationsProviderLockCount) const
+{
+	checkf(CurrentWriteAllocationsProviderLockCount > 0,
+		TEXT("Trying to WRITE to provider outside of an EDIT/WRITE scope"));
+}
+
+void FProviderLock::BeginRead(FProviderLock*& CurrentProviderLock, int32& CurrentReadProviderLockCount, const int32& WriteAllocationsProviderLockCount)
+{
+	checkf(WriteAllocationsProviderLockCount == 0, TEXT("Trying to lock allocations provider for READ while holding EDIT/WRITE access"));
+	if (CurrentReadProviderLockCount++ == 0)
+	{
+		CurrentProviderLock = this;
+		RWLock.ReadLock();
+	}
+}
+
+void FProviderLock::EndRead(FProviderLock*& CurrentProviderLock, int32& CurrentReadProviderLockCount)
+{
+	check(CurrentReadProviderLockCount > 0);
+	if (--CurrentReadProviderLockCount == 0)
+	{
+		RWLock.ReadUnlock();
+		CurrentProviderLock = nullptr;
+	}
+}
+
+void FProviderLock::BeginWrite(FProviderLock*& CurrentProviderLock, const int32& CurrentReadProviderLockCount, int32& WriteAllocationsProviderLockCount)
+{
+	check(!CurrentProviderLock || CurrentProviderLock == this);
+	checkf(CurrentReadProviderLockCount == 0, TEXT("Trying to lock allocations provider for EDIT/WRITE while holding READ access"));
+	if (WriteAllocationsProviderLockCount++ == 0)
+	{
+		CurrentProviderLock = this;
+		RWLock.WriteLock();
+	}
+}
+
+void FProviderLock::EndWrite(FProviderLock*& CurrentProviderLock, int32& WriteAllocationsProviderLockCount)
+{
+	check(WriteAllocationsProviderLockCount > 0);
+	if (--WriteAllocationsProviderLockCount == 0)
+	{
+		RWLock.WriteUnlock();
+		CurrentProviderLock = nullptr;
+	}
+}
+} // namespace TraceServices
diff --git a/Engine/Source/Developer/TraceServices/Private/Model/CookProfilerProvider.cpp b/Engine/Source/Developer/TraceServices/Private/Model/CookProfilerProvider.cpp
new file mode 100644
--- /dev/null
+++ b/Engine/Source/Developer/TraceServices/Private/Model/CookProfilerProvider.cpp
@@ -0,0 +1,84 @@
+// Copyright Epic Games, Inc. All Rights Reserved.
+
+#include "TraceServices/Model/CookProfilerProvider.h"
+#include "CookProfilerProviderPrivate.h"
+#include "AnalysisServicePrivate.h"
+
+namespace TraceServices
+{
+thread_local FProviderLock* GThreadCurrentCookProviderLock;
+thread_local int32 GThreadCurrentReadCookProviderLockCount;
+thread_local int32 GThreadCurrentWriteCookProviderLockCount;
+
+void FCookProfilerProvider::BeginEdit() const
+{
+	Lock.BeginWrite(GThreadCurrentCookProviderLock, GThreadCurrentReadCookProviderLockCount, GThreadCurrentWriteCookProviderLockCount);
+}
+
+void FCookProfilerProvider::EndEdit() const
+{
+	Lock.EndWrite(GThreadCurrentCookProviderLock, GThreadCurrentWriteCookProviderLockCount);
+}
+
+void FCookProfilerProvider::EditAccessCheck() const
+{
+	Lock.WriteAccessCheck(GThreadCurrentReadCookProviderLockCount);
+}
+
+void FCookProfilerProvider::BeginRead() const
+{
+	Lock.BeginRead(GThreadCurrentCookProviderLock, GThreadCurrentReadCookProviderLockCount, GThreadCurrentWriteCookProviderLockCount);
+}
+
+void FCookProfilerProvider::EndRead() const
+{
+	Lock.EndRead(GThreadCurrentCookProviderLock, GThreadCurrentReadCookProviderLockCount);
+}
+
+void FCookProfilerProvider::ReadAccessCheck() const
+{
+	Lock.ReadAccessCheck(GThreadCurrentCookProviderLock, GThreadCurrentReadCookProviderLockCount, GThreadCurrentWriteCookProviderLockCount);
+}
+
+FCookProfilerProvider::FCookProfilerProvider(IAnalysisSession& InSession)
+	: Session(InSession)
+{
+	
+}
+
+void FCookProfilerProvider::AddPackage(uint64 Id, TStringView<TCHAR> Name)
+{
+	PackageIdToIndexMap.Add(Id, Packages.Num());
+	Packages.Emplace(Id, Session.StoreString(Name));
+}
+
+void FCookProfilerProvider::EnumeratePackages(double StartTime, double EndTime, EnumeratePackagesCallback Callback) const
+{
+	for (const FPackageData& Package : Packages)
+	{
+		if (Callback(Package) == false)
+		{
+			break;
+		}
+	}
+}
+
+uint32 FCookProfilerProvider::GetNumPackages() const
+{
+	return (uint32) Packages.Num();
+}
+
+FPackageData* FCookProfilerProvider::EditPackage(uint64 Id)
+{
+	uint32* Index = PackageIdToIndexMap.Find(Id);
+
+	if (Index && *Index < Packages.Num())
+	{
+		FPackageData& Package = Packages[*Index];
+		return &Package;
+	}
+
+	return nullptr;
+}
+
+} // namespace TraceServices
diff --git a/Engine/Source/Developer/TraceServices/Private/Model/CookProfilerProviderPrivate.h b/Engine/Source/Developer/TraceServices/Private/Model/CookProfilerProviderPrivate.h
new file mode 100644
--- /dev/null
+++ b/Engine/Source/Developer/TraceServices/Private/Model/CookProfilerProviderPrivate.h
@@ -0,0 +1,42 @@
+// Copyright Epic Games, Inc. All Rights Reserved.
+
+#pragma once
+
+#include "TraceServices/Model/CookProfilerProvider.h"
+#include "TraceServices/AnalysisService.h"
+
+namespace TraceServices
+{
+
+class FCookProfilerProvider
+	: public ICookProfilerProvider
+{
+public:
+	virtual void BeginEdit() const override;
+	virtual void EndEdit() const override;
+	virtual void EditAccessCheck() const override;
+
+	virtual void BeginRead() const override;
+	virtual void EndRead() const override;
+	virtual void ReadAccessCheck() const override;
+
+	explicit FCookProfilerProvider(IAnalysisSession& Session);
+	virtual ~FCookProfilerProvider() {}
+
+	// Functions to add or edit data.
+	void AddPackage(uint64 Id, TStringView<TCHAR> Name);
+	FPackageData* EditPackage(uint64 Id);
+
+	// Functions to query the data.
+	virtual void EnumeratePackages(double StartTime, double EndTime, EnumeratePackagesCallback Callback) const override;
+	virtual uint32 GetNumPackages() const;
+
+private:
+	IAnalysisSession& Session;
+	TMap<uint64, uint32> PackageIdToIndexMap;
+	TArray64<FPackageData> Packages;
+
+	mutable FProviderLock Lock;
+};
+
+} // namespace TraceServices
diff --git a/Engine/Source/Developer/TraceServices/Private/Modules/CookProfilerModule.cpp b/Engine/Source/Developer/TraceServices/Private/Modules/CookProfilerModule.cpp
new file mode 100644
--- /dev/null
+++ b/Engine/Source/Developer/TraceServices/Private/Modules/CookProfilerModule.cpp
@@ -0,0 +1,37 @@
+// Copyright Epic Games, Inc. All Rights Reserved.
+
+#include "CookProfilerModule.h"
+#include "Analyzers/CookAnalysis.h"
+#include "AnalysisServicePrivate.h"
+#include "Model/CookProfilerProviderPrivate.h"
+
+namespace TraceServices
+{
+
+static const FName CookProfilerModuleName("TraceModule_CookProfiler");
+static const FName CookProfilerProviderName("CookProfilerProvider");
+
+void FCookProfilerModule::GetModuleInfo(FModuleInfo& OutModuleInfo)
+{
+	OutModuleInfo.Name = CookProfilerProviderName;
+	OutModuleInfo.DisplayName = TEXT("Cook Profiling");
+}
+
+void FCookProfilerModule::OnAnalysisBegin(IAnalysisSession& Session)
+{
+	FCookProfilerProvider* CookProfilerProvider = new FCookProfilerProvider(Session);
+	Session.AddProvider(CookProfilerProviderName, CookProfilerProvider);
+	Session.AddAnalyzer(new FCookAnalyzer(Session, *CookProfilerProvider));
+}
+
+void FCookProfilerModule::GetLoggers(TArray<const TCHAR *>& OutLoggers)
+{
+	OutLoggers.Add(TEXT("CookProfiler"));
+}
+
+const ICookProfilerProvider* ReadCookProfilerProvider(const IAnalysisSession& Session)
+{
+	return Session.ReadProvider<ICookProfilerProvider>(CookProfilerProviderName);
+}
+
+} // namespace TraceServices
diff --git a/Engine/Source/Developer/TraceServices/Private/Modules/CookProfilerModule.h b/Engine/Source/Developer/TraceServices/Private/Modules/CookProfilerModule.h
new file mode 100644
--- /dev/null
+++ b/Engine/Source/Developer/TraceServices/Private/Modules/CookProfilerModule.h
@@ -0,0 +1,27 @@
+// Copyright Epic Games, Inc. All Rights Reserved.
+
+#pragma once
+
+#include "TraceServices/ModuleService.h"
+
+namespace TraceServices
+{
+
+class FCookProfilerModule
+	: public IModule
+{
+public:
+	#if WITH_EDITOR
+	virtual bool ShouldBeEnabledByDefault() const override { return false; }
+	#endif
+	virtual void GetModuleInfo(FModuleInfo& OutModuleInfo) override;
+	virtual void OnAnalysisBegin(IAnalysisSession& Session) override;
+	virtual void GenerateReports(const IAnalysisSession& Session, const TCHAR* CmdLine, const TCHAR* OutputDirectory) override {}
+	virtual void GetLoggers(TArray<const TCHAR*>& OutLoggers) override;
+	virtual const TCHAR* GetCommandLineArgument() override
+	{
+		return TEXT("cooktrace");
+	}
+};
+
+} // namespace TraceServices
diff --git a/Engine/Source/Developer/TraceServices/Private/TraceServicesModule.cpp b/Engine/Source/Developer/TraceServices/Private/TraceServicesModule.cpp
--- a/Engine/Source/Developer/TraceServices/Private/TraceServicesModule.cpp
+++ b/Engine/Source/Developer/TraceServices/Private/TraceServicesModule.cpp
@@ -5,6 +5,7 @@
 #include "AnalysisServicePrivate.h"
 #include "ModuleServicePrivate.h"
 #include "Features/IModularFeatures.h"
+#include "Modules/CookProfilerModule.h"
 #include "Modules/TimingProfilerModule.h"
 #include "Modules/LoadTimeProfilerModule.h"
 #include "Modules/StatsModule.h"
@@ -43,6 +44,7 @@ private:
 	TraceServices::FDiagnosticsModule DiagnosticsModule;
 	TraceServices::FPlatformEventsModule PlatformEventsModule;
 	TraceServices::FTasksModule TasksModule;
+	TraceServices::FCookProfilerModule CookProfilingModule;
 };
 
 TSharedPtr<TraceServices::IAnalysisService> FTraceServicesModule::GetAnalysisService()
@@ -91,10 +93,12 @@ void FTraceServicesModule::StartupModule()
 	IModularFeatures::Get().RegisterModularFeature(TraceServices::ModuleFeatureName, &LoadTimeProfilerModule);
 	IModularFeatures::Get().RegisterModularFeature(TraceServices::ModuleFeatureName, &MemoryModule);
 	IModularFeatures::Get().RegisterModularFeature(TraceServices::ModuleFeatureName, &TasksModule);
+	IModularFeatures::Get().RegisterModularFeature(TraceServices::ModuleFeatureName, &CookProfilingModule);
 }
 
 void FTraceServicesModule::ShutdownModule()
 {
+	IModularFeatures::Get().UnregisterModularFeature(TraceServices::ModuleFeatureName, &CookProfilingModule);
 	IModularFeatures::Get().UnregisterModularFeature(TraceServices::ModuleFeatureName, &TasksModule);
 	IModularFeatures::Get().UnregisterModularFeature(TraceServices::ModuleFeatureName, &MemoryModule);
 	IModularFeatures::Get().UnregisterModularFeature(TraceServices::ModuleFeatureName, &LoadTimeProfilerModule);
diff --git a/Engine/Source/Developer/TraceServices/Public/Common/ProviderLock.h b/Engine/Source/Developer/TraceServices/Public/Common/ProviderLock.h
new file mode 100644
--- /dev/null
+++ b/Engine/Source/Developer/TraceServices/Public/Common/ProviderLock.h
@@ -0,0 +1,60 @@
+// Copyright Epic Games, Inc. All Rights Reserved.
+
+#pragma once
+
+#include "TraceServices/Model/AnalysisSession.h"
+
+#include "HAL/CriticalSection.h"
+
+namespace TraceServices
+{
+
+struct TRACESERVICES_API FProviderEditScopeLock
+{
+	FProviderEditScopeLock(const IProvider& InProvider)
+		: Provider(InProvider)
+	{
+		Provider.BeginEdit();
+	}
+
+	~FProviderEditScopeLock()
+	{
+		Provider.EndEdit();
+	}
+
+	const IProvider& Provider;
+};
+
+struct TRACESERVICES_API FProviderReadScopeLock
+{
+	FProviderReadScopeLock(const IProvider& InProvider)
+		: Provider(InProvider)
+	{
+		Provider.BeginRead();
+	}
+
+	~FProviderReadScopeLock()
+	{
+		Provider.EndRead();
+	}
+
+	const IProvider& Provider;
+};
+
+class FProviderLock
+{
+public:
+	void ReadAccessCheck(const FProviderLock* CurrentProviderLock, const int32& CurrentReadProviderLockCount, const int32& CurrentWriteAllocationsProviderLockCount) const;
+	void WriteAccessCheck(const int32& CurrentWriteProviderLockCount) const;
+
+	void BeginRead(FProviderLock*& CurrentProviderLock, int32& CurrentReadProviderLockCount, const int32& WriteAllocationsProviderLockCount);
+	void EndRead(FProviderLock*& CurrentProviderLock, int32& CurrentReadProviderLockCount);
+
+	void BeginWrite(FProviderLock*& CurrentProviderLock, const int32& CurrentReadProviderLockCount, int32& WriteAllocationsProviderLockCount);
+	void EndWrite(FProviderLock*& CurrentProviderLock, int32& WriteAllocationsProviderLockCount);
+
+private:
+	FRWLock RWLock;
+};
+
+} // namespace TraceServices
diff --git a/Engine/Source/Developer/TraceServices/Public/TraceServices/Model/AnalysisSession.h b/Engine/Source/Developer/TraceServices/Public/TraceServices/Model/AnalysisSession.h
--- a/Engine/Source/Developer/TraceServices/Public/TraceServices/Model/AnalysisSession.h
+++ b/Engine/Source/Developer/TraceServices/Public/TraceServices/Model/AnalysisSession.h
@@ -4,6 +4,7 @@
 
 #include "CoreTypes.h"
 #include "Containers/StringFwd.h"
+#include "Misc/AssertionMacros.h"
 
 class FName;
 
@@ -25,6 +26,14 @@ class IProvider
 {
 public:
 	virtual ~IProvider() = default;
+
+	virtual void BeginEdit() const { unimplemented(); }
+	virtual void EndEdit() const { unimplemented(); }
+	virtual void EditAccessCheck() const { unimplemented(); }
+
+	virtual void BeginRead() const { unimplemented(); }
+	virtual void EndRead() const { unimplemented(); }
+	virtual void ReadAccessCheck() const { unimplemented(); }
 };
 
 class IAnalysisSession
diff --git a/Engine/Source/Developer/TraceServices/Public/TraceServices/Model/CookProfilerProvider.h b/Engine/Source/Developer/TraceServices/Public/TraceServices/Model/CookProfilerProvider.h
new file mode 100644
--- /dev/null
+++ b/Engine/Source/Developer/TraceServices/Public/TraceServices/Model/CookProfilerProvider.h
@@ -0,0 +1,38 @@
+// Copyright Epic Games, Inc. All Rights Reserved.
+
+#pragma once
+
+#include "TraceServices/Model/AnalysisSession.h"
+#include "Common/ProviderLock.h"
+#include "Templates/Function.h"
+
+namespace TraceServices
+{
+
+struct FPackageData
+{
+	FPackageData(uint64 InId, const TCHAR* InName)
+		: Id(InId)
+		, Name(InName)
+	{}
+	uint64 Id = -1;
+	const TCHAR* Name;
+	double LoadTime = 0;
+	double SaveTime = 0;
+	double BeginCacheForCookedPlatformData = 0; // BeginCacheForCookedPlatformData is the name of the function from the cooker.
+	double IsCachedCookedPlatformDataLoaded = 0; // IsCachedCookedPlatformDataLoaded is the name of the function from the cooker.
+};
+
+class ICookProfilerProvider
+	: public IProvider
+{
+public:
+	typedef TFunctionRef<bool(const FPackageData& /*Package*/)> EnumeratePackagesCallback;
+
+	virtual ~ICookProfilerProvider() = default;
+	virtual void EnumeratePackages(double StartTime, double EndTime, EnumeratePackagesCallback Callback) const = 0;
+	virtual uint32 GetNumPackages() const = 0;
+};
+
+TRACESERVICES_API const ICookProfilerProvider* ReadCookProfilerProvider(const IAnalysisSession& Session);
+} // namespace TraceServices
diff --git a/Engine/Source/Editor/UnrealEd/Private/CookOnTheFlyServer.cpp b/Engine/Source/Editor/UnrealEd/Private/CookOnTheFlyServer.cpp
--- a/Engine/Source/Editor/UnrealEd/Private/CookOnTheFlyServer.cpp
+++ b/Engine/Source/Editor/UnrealEd/Private/CookOnTheFlyServer.cpp
@@ -2962,6 +2962,7 @@ bool UCookOnTheFlyServer::LoadPackageForCooking(UE::Cook::FPackageData& PackageD
 	UE::Cook::FPackageData* ReportingPackageData)
 {
 	UE_SCOPED_HIERARCHICAL_COOKTIMER_AND_DURATION(LoadPackageForCooking, DetailedCookStats::TickCookOnTheSideLoadPackagesTimeSec);
+	UE_SCOPED_COOK_STAT(PackageData.GetPackageName().ToUnstableInt(), EPackageEventStatType::LoadPackage);
 
 	check(PackageTracker->LoadingPackageData == nullptr);
 	PackageTracker->LoadingPackageData = &PackageData;
diff --git a/Engine/Source/Editor/UnrealEd/Private/Cooker/CookPackageData.cpp b/Engine/Source/Editor/UnrealEd/Private/Cooker/CookPackageData.cpp
--- a/Engine/Source/Editor/UnrealEd/Private/Cooker/CookPackageData.cpp
+++ b/Engine/Source/Editor/UnrealEd/Private/Cooker/CookPackageData.cpp
@@ -244,8 +244,7 @@ void FPackageData::SetRequestData(const TArrayView<const ITargetPlatform* const>
 	AddCompletionCallback(MoveTemp(InCompletionCallback));
 	if (Instigator.Category == EInstigator::NotYetRequested)
 	{
-		Instigator = MoveTemp(InInstigator);
-		PackageDatas.DebugInstigator(*this);
+		OnPackageDataFirstRequested(MoveTemp(InInstigator));
 	}
 }
 
@@ -774,6 +773,13 @@ void FPackageData::OnExitHasPackage()
 	SetPackage(nullptr);
 }
 
+void FPackageData::OnPackageDataFirstRequested(FInstigator&& InInstigator)
+{
+	TracePackage(GetPackageName().ToUnstableInt(), GetPackageName().ToString());
+	Instigator = MoveTemp(InInstigator);
+	PackageDatas.DebugInstigator(*this);
+}
+
 void FPackageData::SetState(EPackageState NextState)
 {
 	State = static_cast<uint32>(NextState);
diff --git a/Engine/Source/Editor/UnrealEd/Private/Cooker/CookPackageData.h b/Engine/Source/Editor/UnrealEd/Private/Cooker/CookPackageData.h
--- a/Engine/Source/Editor/UnrealEd/Private/Cooker/CookPackageData.h
+++ b/Engine/Source/Editor/UnrealEd/Private/Cooker/CookPackageData.h
@@ -527,6 +527,8 @@ private:
 	void OnExitHasPackage();
 	void OnEnterHasPackage();
 
+	void OnPackageDataFirstRequested(FInstigator&& InInstigator);
+
 	TUniquePtr<FGeneratorPackage> GeneratorPackage;
 	FGeneratorPackage* GeneratedOwner;
 	/** Data for each platform that has been interacted with by *this. */
diff --git a/Engine/Source/Editor/UnrealEd/Private/Cooker/CookProfiling.cpp b/Engine/Source/Editor/UnrealEd/Private/Cooker/CookProfiling.cpp
--- a/Engine/Source/Editor/UnrealEd/Private/Cooker/CookProfiling.cpp
+++ b/Engine/Source/Editor/UnrealEd/Private/Cooker/CookProfiling.cpp
@@ -10,8 +10,6 @@
 #if OUTPUT_COOKTIMING
 #include <Containers/AllocatorFixedSizeFreeList.h>
 
-UE_TRACE_CHANNEL_DEFINE(CookChannel)
-
 struct FHierarchicalTimerInfo
 {
 public:
@@ -169,6 +167,7 @@ void ClearHierarchyTimers()
 {
 	RootTimerInfo.ClearChildren();
 }
+
 #endif
 
 #if ENABLE_COOK_STATS
diff --git a/Engine/Source/Editor/UnrealEd/Private/Cooker/CookProfiling.h b/Engine/Source/Editor/UnrealEd/Private/Cooker/CookProfiling.h
--- a/Engine/Source/Editor/UnrealEd/Private/Cooker/CookProfiling.h
+++ b/Engine/Source/Editor/UnrealEd/Private/Cooker/CookProfiling.h
@@ -38,8 +38,6 @@ private:
 	FHierarchicalTimerInfo* PrevTimerInfo;
 };
 
-UE_TRACE_CHANNEL_EXTERN(CookChannel)
-
 #define UE_CREATE_HIERARCHICAL_COOKTIMER(name, incrementScope) \
 			FScopeTimer PREPROCESSOR_JOIN(__HierarchicalCookTimerScope, __LINE__)(__COUNTER__, #name, incrementScope); 
 #define UE_CREATE_TEXT_HIERARCHICAL_COOKTIMER(name, incrementScope) \
diff --git a/Engine/Source/Editor/UnrealEd/Private/Cooker/CookTypes.h b/Engine/Source/Editor/UnrealEd/Private/Cooker/CookTypes.h
--- a/Engine/Source/Editor/UnrealEd/Private/Cooker/CookTypes.h
+++ b/Engine/Source/Editor/UnrealEd/Private/Cooker/CookTypes.h
@@ -226,12 +226,14 @@ namespace UE::Cook
 inline void RouteBeginCacheForCookedPlatformData(UObject* Obj, const ITargetPlatform* TargetPlatform)
 {
 	UE_SCOPED_TEXT_COOKTIMER(*WriteToString<128>(GetClassTraceScope(Obj), TEXT("_BeginCacheForCookedPlatformData")));
+	UE_SCOPED_COOK_STAT(Obj->GetPackage()->GetFName().ToUnstableInt(), EPackageEventStatType::BeginCacheForCookedPlatformData);
 	Obj->BeginCacheForCookedPlatformData(TargetPlatform);
 }
 
 inline bool RouteIsCachedCookedPlatformDataLoaded(UObject* Obj, const ITargetPlatform* TargetPlatform)
 {
 	UE_SCOPED_TEXT_COOKTIMER(*WriteToString<128>(GetClassTraceScope(Obj), TEXT("_IsCachedCookedPlatformDataLoaded")));
+	UE_SCOPED_COOK_STAT(Obj->GetPackage()->GetFName().ToUnstableInt(), EPackageEventStatType::IsCachedCookedPlatformDataLoaded);
 	return Obj->IsCachedCookedPlatformDataLoaded(TargetPlatform);
 }
 
diff --git a/Engine/Source/Runtime/Core/Private/ProfilingDebugging/CookStats.cpp b/Engine/Source/Runtime/Core/Private/ProfilingDebugging/CookStats.cpp
--- a/Engine/Source/Runtime/Core/Private/ProfilingDebugging/CookStats.cpp
+++ b/Engine/Source/Runtime/Core/Private/ProfilingDebugging/CookStats.cpp
@@ -2,12 +2,45 @@
 
 #include "ProfilingDebugging/CookStats.h"
 
+#include "Trace/Trace.inl"
 
 #if ENABLE_COOK_STATS
+
+UE_TRACE_CHANNEL_DEFINE(CookChannel)
+
+UE_TRACE_EVENT_BEGIN(CookTrace, Package)
+	UE_TRACE_EVENT_FIELD(uint64, Id)
+	UE_TRACE_EVENT_FIELD(UE::Trace::WideString, Name)
+	UE_TRACE_EVENT_FIELD(uint64, Cycle)
+UE_TRACE_EVENT_END()
+
+UE_TRACE_EVENT_BEGIN(CookTrace, PackageStat)
+	UE_TRACE_EVENT_FIELD(uint64, Id)
+	UE_TRACE_EVENT_FIELD(uint64, Duration)
+	UE_TRACE_EVENT_FIELD(uint8, StatType)
+UE_TRACE_EVENT_END()
+
 CORE_API FCookStatsManager::FGatherCookStatsDelegate FCookStatsManager::CookStatsCallbacks;
 
 CORE_API void FCookStatsManager::LogCookStats(AddStatFuncRef AddStat)
 {
 	CookStatsCallbacks.Broadcast(AddStat);
 }
+
+void TracePackage(uint64 InId, FString InName)
+{
+	UE_TRACE_LOG(CookTrace, Package, CookChannel)
+		<< Package.Id(InId)
+		<< Package.Name(*InName, uint16(InName.Len()))
+		<< Package.Cycle(FPlatformTime::Cycles64());
+}
+
+void TracePackageStat(uint64 InId, uint64 Duration, EPackageEventStatType StatType)
+{
+	UE_TRACE_LOG(CookTrace, PackageStat, CookChannel)
+		<< PackageStat.Id(InId)
+		<< PackageStat.Duration(Duration)
+		<< PackageStat.StatType((uint8)StatType);
+}
+
 #endif
diff --git a/Engine/Source/Runtime/Core/Public/ProfilingDebugging/CookStats.h b/Engine/Source/Runtime/Core/Public/ProfilingDebugging/CookStats.h
--- a/Engine/Source/Runtime/Core/Public/ProfilingDebugging/CookStats.h
+++ b/Engine/Source/Runtime/Core/Public/ProfilingDebugging/CookStats.h
@@ -19,6 +19,14 @@
 template <typename FuncType> class TFunctionRef;
 template <typename T> struct TDecay;
 
+enum class EPackageEventStatType : uint8
+{
+	LoadPackage = 0,
+	SavePackage = 1,
+	BeginCacheForCookedPlatformData = 2,
+	IsCachedCookedPlatformDataLoaded = 3,
+};
+
 #if ENABLE_COOK_STATS
 /**
  * Centralizes the system to gather stats from a cook that need to be collected at the core/engine level.
@@ -423,6 +431,42 @@ struct FDerivedDataUsageStats
 /** any expression inside this block will be compiled out if ENABLE_COOK_STATS is not true. Shorthand for the multi-line #if ENABLE_COOK_STATS guard. */
 #define COOK_STAT(...) __VA_ARGS__
 
+CORE_API UE_TRACE_CHANNEL_EXTERN(CookChannel);
+
+CORE_API void TracePackage(uint64 InId, FString InName);
+CORE_API void TracePackageStat(uint64 InId, uint64 Duration, EPackageEventStatType StatType);
+
+struct FScopedCookStat
+{
+	FScopedCookStat(const FScopedCookStat&) = delete;
+	FScopedCookStat(FScopedCookStat&&) = delete;
+
+	FScopedCookStat(uint64 Id, EPackageEventStatType InStatType)
+		: Id(Id)
+		, StatType(InStatType)
+	{
+		StartTimestamp = FPlatformTime::Cycles64();
+	}
+
+	~FScopedCookStat()
+	{
+		TracePackageStat(Id, FPlatformTime::Cycles64() - StartTimestamp, StatType);
+	}
+
+private:
+	uint64 Id;
+	uint64 StartTimestamp;
+	EPackageEventStatType StatType;
+};
+
+#define UE_SCOPED_COOK_STAT(Id, StatType) \
+	FScopedCookStat  PREPROCESSOR_JOIN(__CookTimerScope, __LINE__)(Id, StatType);
+
 #else
 #define COOK_STAT(...)
+#define UE_SCOPED_COOK_STAT(...)
+
+#define TracePackage(...)
+#define TracePackageStat(...)
+
 #endif
diff --git a/Engine/Source/Runtime/CoreUObject/Private/UObject/SavePackage2.cpp b/Engine/Source/Runtime/CoreUObject/Private/UObject/SavePackage2.cpp
--- a/Engine/Source/Runtime/CoreUObject/Private/UObject/SavePackage2.cpp
+++ b/Engine/Source/Runtime/CoreUObject/Private/UObject/SavePackage2.cpp
@@ -2537,6 +2537,8 @@ FSavePackageResultStruct UPackage::Save2(UPackage* InPackage, UObject* InAsset,
 	COOK_STAT(FScopedDurationTimer FuncSaveTimer(FSavePackageStats::SavePackageTimeSec));
 	COOK_STAT(FSavePackageStats::NumPackagesSaved++);
 	SCOPED_SAVETIMER(UPackage_Save2);
+	UE_SCOPED_COOK_STAT(InPackage->GetFName().ToUnstableInt(), EPackageEventStatType::SavePackage);
+
 	FSaveContext SaveContext(InPackage, InAsset, InFilename, SaveArgs);
 
 	// Create the slow task dialog if needed
