Check if CookProvider is null before dereferencing

Work around for this, as likely we should be able to assume this is not null

diff --git a/Engine/Source/Developer/TraceInsights/Private/Insights/CookProfiler/CookProfilerManager.cpp b/Engine/Source/Developer/TraceInsights/Private/Insights/CookProfiler/CookProfilerManager.cpp
--- a/Engine/Source/Developer/TraceInsights/Private/Insights/CookProfiler/CookProfilerManager.cpp
+++ b/Engine/Source/Developer/TraceInsights/Private/Insights/CookProfiler/CookProfilerManager.cpp
@@ -134,19 +134,23 @@ bool FCookProfilerManager::Tick(float DeltaTime)
 		if (Session.IsValid())
 		{
 			const TraceServices::ICookProfilerProvider* CookProvider = TraceServices::ReadCookProfilerProvider(*Session.Get());
-			TraceServices::FProviderReadScopeLock ProviderReadScope(*CookProvider);
 
-			TSharedPtr<FTabManager> TabManagerShared = TimingTabManager.Pin();
-			if (CookProvider && CookProvider->GetNumPackages() > 0 && TabManagerShared.IsValid())
+			if (CookProvider)
 			{
-				bIsAvailable = true;
-				TabManagerShared->TryInvokeTab(FCookProfilerTabs::PackageTableTreeViewTabID);
-			}
-
-			if (Session->IsAnalysisComplete())
-			{
-				// Never check again during this session.
-				AvailabilityCheck.Disable();
+				TraceServices::FProviderReadScopeLock ProviderReadScope(*CookProvider);
+
+				TSharedPtr<FTabManager> TabManagerShared = TimingTabManager.Pin();
+				if (CookProvider && CookProvider->GetNumPackages() > 0 && TabManagerShared.IsValid())
+				{
+					bIsAvailable = true;
+					TabManagerShared->TryInvokeTab(FCookProfilerTabs::PackageTableTreeViewTabID);
+				}
+
+				if (Session->IsAnalysisComplete())
+				{
+					// Never check again during this session.
+					AvailabilityCheck.Disable();
+				}
 			}
 		}
 		else
