[Insights] Added "Insights" mem scope in various places to capture untagged memory when Insights is running in Editor (i.e. as Animation Insights).

diff --git a/Engine/Source/Developer/TraceInsights/Private/Insights/InsightsStyle.cpp b/Engine/Source/Developer/TraceInsights/Private/Insights/InsightsStyle.cpp
--- a/Engine/Source/Developer/TraceInsights/Private/Insights/InsightsStyle.cpp
+++ b/Engine/Source/Developer/TraceInsights/Private/Insights/InsightsStyle.cpp
@@ -10,8 +10,6 @@
 #include "Styling/StyleColors.h"
 #include "Styling/ToolBarStyle.h"
 
-LLM_DEFINE_TAG(Insights);
-
 ////////////////////////////////////////////////////////////////////////////////////////////////////
 // FInsightsStyle
 ////////////////////////////////////////////////////////////////////////////////////////////////////
@@ -24,7 +22,7 @@
 {
 	LLM_SCOPE_BYNAME(TEXT("Insights/Style"));
 
-	// The core style must be initialized before the editor style
+	// The core style must be initialized before the Insights style
 	FSlateApplication::InitializeCoreStyle();
 
 	if (!StyleInstance.IsValid())
diff --git a/Engine/Source/Developer/TraceInsights/Private/Insights/TraceInsightsModule.cpp b/Engine/Source/Developer/TraceInsights/Private/Insights/TraceInsightsModule.cpp
--- a/Engine/Source/Developer/TraceInsights/Private/Insights/TraceInsightsModule.cpp
+++ b/Engine/Source/Developer/TraceInsights/Private/Insights/TraceInsightsModule.cpp
@@ -4,6 +4,7 @@
 #include "Framework/Application/SlateApplication.h"
 #include "Framework/Docking/LayoutService.h"
 #include "Framework/Notifications/NotificationManager.h"
+#include "HAL/LowLevelMemTracker.h"
 #include "HAL/PlatformApplicationMisc.h"
 #include "ISourceCodeAccessModule.h"
 #include "ISourceCodeAccessor.h"
@@ -34,6 +35,8 @@
 
 DEFINE_LOG_CATEGORY(TraceInsights);
 
+LLM_DEFINE_TAG(Insights);
+
 IMPLEMENT_MODULE(FTraceInsightsModule, TraceInsights);
 
 FString FTraceInsightsModule::UnrealInsightsLayoutIni;
@@ -44,6 +47,8 @@
 
 void FTraceInsightsModule::StartupModule()
 {
+	LLM_SCOPE_BYTAG(Insights);
+
 	ITraceServicesModule& TraceServicesModule = FModuleManager::LoadModuleChecked<ITraceServicesModule>("TraceServices");
 	TraceAnalysisService = TraceServicesModule.GetAnalysisService();
 	TraceModuleService = TraceServicesModule.GetModuleService();
@@ -78,6 +83,8 @@
 PRAGMA_DISABLE_OPTIMIZATION
 void FTraceInsightsModule::ShutdownModule()
 {
+	LLM_SCOPE_BYTAG(Insights);
+
 	FInsightsManager::Get()->WaitOnInProgressAsyncOps();
 
 #if !WITH_EDITOR
@@ -150,6 +157,7 @@
 {
 	if (Component.IsValid())
 	{
+		LLM_SCOPE_BYTAG(Insights);
 		Components.Add(Component.ToSharedRef());
 		Component->Initialize(*this);
 	}
@@ -161,6 +169,7 @@
 {
 	if (Component.IsValid())
 	{
+		LLM_SCOPE_BYTAG(Insights);
 		Component->Shutdown();
 		Components.Remove(Component.ToSharedRef());
 	}
diff --git a/Engine/Source/Developer/TraceInsights/Private/Insights/ViewModels/LogMessage.cpp b/Engine/Source/Developer/TraceInsights/Private/Insights/ViewModels/LogMessage.cpp
--- a/Engine/Source/Developer/TraceInsights/Private/Insights/ViewModels/LogMessage.cpp
+++ b/Engine/Source/Developer/TraceInsights/Private/Insights/ViewModels/LogMessage.cpp
@@ -8,6 +8,7 @@
 
 // Insights
 #include "Insights/Common/TimeUtils.h"
+#include "Insights/Log.h"
 
 ////////////////////////////////////////////////////////////////////////////////////////////////////
 // FLogMessageRecord
@@ -126,6 +127,7 @@
 		const TraceServices::ILogProvider& LogProvider = TraceServices::ReadLogProvider(*Session.Get());
 		LogProvider.ReadMessage(Index, [this, Index](const TraceServices::FLogMessageInfo& Message)
 		{
+			LLM_SCOPE_BYTAG(Insights);
 			FScopeLock Lock(&CriticalSection);
 			FLogMessageRecord Entry(Message);
 			Map.Add(Index, MoveTemp(Entry));
diff --git a/Engine/Source/Developer/TraceInsights/Private/Insights/ViewModels/ThreadTimingTrack.cpp b/Engine/Source/Developer/TraceInsights/Private/Insights/ViewModels/ThreadTimingTrack.cpp
--- a/Engine/Source/Developer/TraceInsights/Private/Insights/ViewModels/ThreadTimingTrack.cpp
+++ b/Engine/Source/Developer/TraceInsights/Private/Insights/ViewModels/ThreadTimingTrack.cpp
@@ -18,6 +18,7 @@
 #include "Insights/InsightsManager.h"
 #include "Insights/InsightsStyle.h"
 #include "Insights/ITimingViewSession.h"
+#include "Insights/Log.h"
 #include "Insights/TimingProfilerManager.h"
 #include "Insights/ViewModels/Filters.h"
 #include "Insights/ViewModels/FilterConfigurator.h"
@@ -467,6 +468,8 @@
 			TimingProfilerTimelineCount = CurrentTimingProfilerTimelineCount;
 			LoadTimeProfilerTimelineCount = CurrentLoadTimeProfilerTimelineCount;
 
+			LLM_SCOPE_BYTAG(Insights);
+
 			// Check if we have a GPU track.
 			if (!GpuTrack.IsValid())
 			{
@@ -1734,6 +1737,8 @@
 
 void FThreadTimingTrack::OnFilterTrackClicked()
 {
+	LLM_SCOPE_BYTAG(Insights);
+
 	if (!FilterConfigurator.IsValid())
 	{
 		FilterConfigurator = MakeShared<FFilterConfigurator>();
diff --git a/Engine/Source/Developer/TraceInsights/Private/Insights/Widgets/SFrameTrack.cpp b/Engine/Source/Developer/TraceInsights/Private/Insights/Widgets/SFrameTrack.cpp
--- a/Engine/Source/Developer/TraceInsights/Private/Insights/Widgets/SFrameTrack.cpp
+++ b/Engine/Source/Developer/TraceInsights/Private/Insights/Widgets/SFrameTrack.cpp
@@ -19,6 +19,7 @@
 #include "Insights/Common/TimeUtils.h"
 #include "Insights/InsightsManager.h"
 #include "Insights/InsightsStyle.h"
+#include "Insights/Log.h"
 #include "Insights/TimingProfilerCommon.h"
 #include "Insights/TimingProfilerManager.h"
 #include "Insights/ViewModels/FrameTrackHelper.h"
@@ -228,6 +229,7 @@
 	}
 	else
 	{
+		LLM_SCOPE_BYTAG(Insights);
 		TSharedRef<FFrameTrackSeries> SeriesRef = MakeShared<FFrameTrackSeries>(FrameType);
 		SeriesMap.Add(FrameType, SeriesRef);
 		return SeriesRef;
@@ -281,6 +283,7 @@
 		{
 			TSharedPtr<FFrameTrackSeries> SeriesPtr = FindOrAddSeries(FrameType);
 
+			LLM_SCOPE_BYTAG(Insights);
 			FFrameTrackSeriesBuilder Builder(*SeriesPtr, Viewport);
 
 			FramesProvider.EnumerateFrames(static_cast<ETraceFrameType>(FrameType), StartIndex, EndIndex, [&Builder](const TraceServices::FFrame& Frame)
@@ -348,6 +351,7 @@
 
 							if (Y >= TopY && Y < BottomY)
 							{
+								LLM_SCOPE_BYTAG(Insights);
 								return FFrameTrackSampleRef(SeriesPtr, MakeShared<FFrameTrackSample>(Sample));
 							}
 						}
diff --git a/Engine/Source/Developer/TraceInsights/Private/Insights/Widgets/SLogView.cpp b/Engine/Source/Developer/TraceInsights/Private/Insights/Widgets/SLogView.cpp
--- a/Engine/Source/Developer/TraceInsights/Private/Insights/Widgets/SLogView.cpp
+++ b/Engine/Source/Developer/TraceInsights/Private/Insights/Widgets/SLogView.cpp
@@ -741,6 +741,8 @@
 
 void SLogView::Tick(const FGeometry& AllottedGeometry, const double InCurrentTime, const float InDeltaTime)
 {
+	LLM_SCOPE_BYTAG(Insights);
+
 	int32 NewMessageCount = 0;
 
 	TSharedPtr<const TraceServices::IAnalysisSession> Session = FInsightsManager::Get()->GetSession();
diff --git a/Engine/Source/Developer/TraceInsights/Private/Insights/Widgets/STimingView.cpp b/Engine/Source/Developer/TraceInsights/Private/Insights/Widgets/STimingView.cpp
--- a/Engine/Source/Developer/TraceInsights/Private/Insights/Widgets/STimingView.cpp
+++ b/Engine/Source/Developer/TraceInsights/Private/Insights/Widgets/STimingView.cpp
@@ -45,6 +45,7 @@
 #include "Insights/InsightsStyle.h"
 #include "Insights/LoadingProfiler/LoadingProfilerManager.h"
 #include "Insights/LoadingProfiler/Widgets/SLoadingProfilerWindow.h"
+#include "Insights/Log.h"
 #include "Insights/Table/Widgets/STableTreeView.h"
 #include "Insights/TaskGraphProfiler/TaskGraphProfilerManager.h"
 #include "Insights/Tests/TimingProfilerTests.h"
@@ -307,6 +308,8 @@
 
 void STimingView::Reset(bool bIsFirstReset)
 {
+	LLM_SCOPE_BYTAG(Insights);
+
 	const FInsightsSettings& Settings = FInsightsManager::Get()->GetSettings();
 
 	if (!bIsFirstReset)
@@ -520,6 +523,8 @@
 	FStopwatch TickStopwatch;
 	TickStopwatch.Start();
 
+	LLM_SCOPE_BYTAG(Insights);
+
 	ThisGeometry = AllottedGeometry;
 
 	Tooltip.SetFontScale(AllottedGeometry.Scale);
@@ -3933,6 +3938,7 @@
 	}
 	else
 	{
+		LLM_SCOPE_BYTAG(Insights);
 		TSharedRef<FTimingEventFilterByEventType> NewEventFilter = MakeShared<FTimingEventFilterByEventType>(EventType);
 		NewEventFilter->SetFilterByTrackTypeName(true);
 		NewEventFilter->SetTrackTypeName(FThreadTimingTrack::GetStaticTypeName());
@@ -4733,6 +4739,8 @@
 {
 	using namespace Insights;
 
+	LLM_SCOPE_BYTAG(Insights);
+
 	if (!QuickFindVm.IsValid())
 	{
 		TSharedPtr<FFilterConfigurator> NewFilterConfigurator = MakeShared<FFilterConfigurator>();
@@ -5021,6 +5029,7 @@
 
 void STimingView::FilterAllTracks()
 {
+	LLM_SCOPE_BYTAG(Insights);
 	FilterConfigurator = MakeShared<Insights::FFilterConfigurator>(*QuickFindVm->GetFilterConfigurator());
 	for (auto& Entry : AllTracks)
 	{
@@ -5032,6 +5041,7 @@
 
 void STimingView::ClearFilters()
 {
+	LLM_SCOPE_BYTAG(Insights);
 	FilterConfigurator.Reset();
 	for (auto& Entry : AllTracks)
 	{
diff --git a/Engine/Source/Developer/TraceInsights/Public/Insights/Log.h b/Engine/Source/Developer/TraceInsights/Public/Insights/Log.h
--- a/Engine/Source/Developer/TraceInsights/Public/Insights/Log.h
+++ b/Engine/Source/Developer/TraceInsights/Public/Insights/Log.h
@@ -1,8 +1,11 @@
-// Copyright Epic Games, Inc. All Rights Reserved.
-
-#pragma once
-
-#include "CoreMinimal.h"
-#include "Logging/LogMacros.h"
-
-DECLARE_LOG_CATEGORY_EXTERN(TraceInsights, Verbose, All);
+// Copyright Epic Games, Inc. All Rights Reserved.
+
+#pragma once
+
+#include "CoreMinimal.h"
+#include "HAL/LowLevelMemTracker.h"
+#include "Logging/LogMacros.h"
+
+DECLARE_LOG_CATEGORY_EXTERN(TraceInsights, Verbose, All);
+
+LLM_DECLARE_TAG(Insights);
