[Insights] Resize Screenshots to a max width of 640 pixels and a max heigth of 480 when displaying them.

diff --git a/Engine/Source/Developer/TraceInsights/Private/Insights/ViewModels/MarkersTimingTrack.cpp b/Engine/Source/Developer/TraceInsights/Private/Insights/ViewModels/MarkersTimingTrack.cpp
--- a/Engine/Source/Developer/TraceInsights/Private/Insights/ViewModels/MarkersTimingTrack.cpp
+++ b/Engine/Source/Developer/TraceInsights/Private/Insights/ViewModels/MarkersTimingTrack.cpp
@@ -492,16 +492,38 @@
 	}
 
 	IImageWrapperModule& ImageWrapperModule = FModuleManager::LoadModuleChecked<IImageWrapperModule>(FName("ImageWrapper"));
-	TSharedPtr<IImageWrapper> ImageWrapper = ImageWrapperModule.CreateImageWrapper(EImageFormat::PNG);
-	if (ImageWrapper.IsValid() && ImageWrapper->SetCompressed(Screenshot->Data.GetData(), Screenshot->Data.Num()))
+	FImage Image;
+	if (ImageWrapperModule.DecompressImage(Screenshot->Data.GetData(), Screenshot->Size, Image))
 	{
-		TArray<uint8> RawImage;
-		if (!ImageWrapper->GetRaw(ERGBFormat::BGRA, 8, RawImage))
+		constexpr int32 MAX_WIDTH = 640;
+		constexpr int32 MAX_HEIGHT = 480;
+
+		int32 ResizedX = Screenshot->Width;
+		int32 ResizedY = Screenshot->Height;
+
+		if (ResizedX > MAX_WIDTH)
 		{
-			return;
+			ResizedY = (ResizedY * MAX_WIDTH) / ResizedX;
+			ResizedX = MAX_WIDTH;
 		}
 
-		TSharedPtr<FSlateBrush> ImageBrush = FSlateDynamicImageBrush::CreateWithImageData(FName(Screenshot->Name), FVector2D(Screenshot->Width, Screenshot->Height), RawImage);
+		if (ResizedY > MAX_HEIGHT)
+		{
+			ResizedX = (ResizedX * MAX_HEIGHT) / ResizedY;
+			ResizedY = MAX_HEIGHT;
+		}
+
+		TSharedPtr<FSlateBrush> ImageBrush;
+		if (Screenshot->Width != ResizedX || Screenshot->Height != ResizedY)
+		{
+			FImage ResizedImage;
+			Image.ResizeTo(ResizedImage, ResizedX, ResizedY, ERawImageFormat::BGRA8, EGammaSpace::sRGB);
+			ImageBrush = FSlateDynamicImageBrush::CreateWithImageData(FName(Screenshot->Name), FVector2D(ResizedX, ResizedY), TArray<uint8>(ResizedImage.RawData));
+		}
+		else
+		{
+			ImageBrush = FSlateDynamicImageBrush::CreateWithImageData(FName(Screenshot->Name), FVector2D(Screenshot->Width, Screenshot->Height), TArray<uint8>(Image.RawData));
+		}
 		InOutTooltip.SetImage(ImageBrush);
 	}
 
diff --git a/Engine/Source/Developer/TraceInsights/TraceInsights.Build.cs b/Engine/Source/Developer/TraceInsights/TraceInsights.Build.cs
--- a/Engine/Source/Developer/TraceInsights/TraceInsights.Build.cs
+++ b/Engine/Source/Developer/TraceInsights/TraceInsights.Build.cs
@@ -60,6 +60,7 @@
 
 		PrivateDependencyModuleNames.AddRange(
 			new string[] {
+				"ImageCore",
 				"SlateCore",
 			}
 		);
