[Insights]
  - Tree Table views: Fixed missing tree items when multiple groupings are added in a generic Tree Table view (like in Memory Insights \ Allocs Table view). Groupings are now recursively applied correctly also when a grouping creates a hierarchy of nodes (not only a flat list of nodes).
  - Tree Table views: Improved performance of filtering and grouping.

diff --git a/Engine/Source/Developer/TraceInsights/Private/Insights/MemoryProfiler/ViewModels/MemAllocGroupingByHeap.cpp b/Engine/Source/Developer/TraceInsights/Private/Insights/MemoryProfiler/ViewModels/MemAllocGroupingByHeap.cpp
--- a/Engine/Source/Developer/TraceInsights/Private/Insights/MemoryProfiler/ViewModels/MemAllocGroupingByHeap.cpp
+++ b/Engine/Source/Developer/TraceInsights/Private/Insights/MemoryProfiler/ViewModels/MemAllocGroupingByHeap.cpp
@@ -12,9 +12,9 @@
 
 namespace Insights
 {
-	
+
 INSIGHTS_IMPLEMENT_RTTI(FMemAllocGroupingByHeap)
-	
+
 ////////////////////////////////////////////////////////////////////////////////////////////////////
 
 FMemAllocGroupingByHeap::FMemAllocGroupingByHeap(const TraceServices::IAllocationsProvider& InAllocProvider)
@@ -58,7 +58,7 @@ void FMemAllocGroupingByHeap::GroupNodes(const TArray<FTableTreeNodePtr>& Nodes,
 	// Build heap hierarchy
 	TArray<FTableTreeNodePtr> HeapNodes;
 	HeapNodes.AddZeroed(256);
-	
+
 	AllocProvider.EnumerateRootHeaps([&](uint8 Id, const TraceServices::IAllocationsProvider::FHeapSpec& Spec)
 	{
 		ParentGroup.AddChildAndSetGroupPtr(MakeGroupNodeHierarchy(Spec, InParentTable, HeapNodes));
@@ -71,7 +71,13 @@ void FMemAllocGroupingByHeap::GroupNodes(const TArray<FTableTreeNodePtr>& Nodes,
 		{
 			return;
 		}
-		
+
+		if (NodePtr->IsGroup())
+		{
+			ParentGroup.AddChildAndSetGroupPtr(NodePtr);
+			continue;
+		}
+
 		const FMemAllocNode& MemAllocNode = static_cast<const FMemAllocNode&>(*NodePtr);
 		const FMemoryAlloc* Alloc = MemAllocNode.GetMemAlloc();
 
@@ -87,8 +93,9 @@ void FMemAllocGroupingByHeap::GroupNodes(const TArray<FTableTreeNodePtr>& Nodes,
 		}
 	}
 }
-	
+
 ////////////////////////////////////////////////////////////////////////////////////////////////////
+
 } // namespace Insights
 
-#undef LOCTEXT_NAMESPACE
\ No newline at end of file
+#undef LOCTEXT_NAMESPACE
diff --git a/Engine/Source/Developer/TraceInsights/Private/Insights/Table/ViewModels/BaseTreeNode.h b/Engine/Source/Developer/TraceInsights/Private/Insights/Table/ViewModels/BaseTreeNode.h
--- a/Engine/Source/Developer/TraceInsights/Private/Insights/Table/ViewModels/BaseTreeNode.h
+++ b/Engine/Source/Developer/TraceInsights/Private/Insights/Table/ViewModels/BaseTreeNode.h
@@ -172,16 +172,42 @@ public:
 		GroupData->Children.Add(ChildPtr);
 	}
 
-	/** Clears children. */
-	void ClearChildren(int32 NewSize = 0)
+	void SetGroupPtrForAllChildren()
+	{
+		FBaseTreeNodeWeak ThisNode = AsShared();
+		for (FBaseTreeNodePtr& NodePtr : GroupData->Children)
+		{
+			NodePtr->GroupPtr = ThisNode;
+		}
+	}
+
+	void ResetGroupPtrForAllChildren()
 	{
 		for (FBaseTreeNodePtr& NodePtr : GroupData->Children)
 		{
 			NodePtr->GroupPtr = nullptr;
 		}
+	}
+
+	/** Clears children. */
+	FORCEINLINE_DEBUGGABLE void ClearChildren(int32 NewSize = 0)
+	{
+		ResetGroupPtrForAllChildren();
 		GroupData->Children.Reset(NewSize);
 	}
 
+	FORCEINLINE_DEBUGGABLE void SwapChildren(TArray<FBaseTreeNodePtr>& NewChildren)
+	{
+		ResetGroupPtrForAllChildren();
+		Swap(NewChildren, GroupData->Children);
+		SetGroupPtrForAllChildren();
+	}
+
+	FORCEINLINE_DEBUGGABLE void SwapChildrenFast(TArray<FBaseTreeNodePtr>& NewChildren)
+	{
+		Swap(NewChildren, GroupData->Children);
+	}
+
 	/** Adds specified child to the filtered children. */
 	FORCEINLINE_DEBUGGABLE void AddFilteredChild(const FBaseTreeNodePtr& ChildPtr)
 	{
@@ -189,9 +215,9 @@ public:
 	}
 
 	/** Clears filtered children. */
-	void ClearFilteredChildren()
+	void ClearFilteredChildren(int32 NewSize = 0)
 	{
-		GroupData->FilteredChildren.Reset();
+		GroupData->FilteredChildren.Reset(NewSize);
 	}
 
 	bool IsExpanded() const { return GroupData->bIsExpanded; }
diff --git a/Engine/Source/Developer/TraceInsights/Private/Insights/Table/Widgets/STableTreeView.cpp b/Engine/Source/Developer/TraceInsights/Private/Insights/Table/Widgets/STableTreeView.cpp
--- a/Engine/Source/Developer/TraceInsights/Private/Insights/Table/Widgets/STableTreeView.cpp
+++ b/Engine/Source/Developer/TraceInsights/Private/Insights/Table/Widgets/STableTreeView.cpp
@@ -1075,17 +1075,18 @@ bool STableTreeView::ApplyHierarchicalFilterForNode(FTableTreeNodePtr NodePtr, b
 
 	if (NodePtr->IsGroup())
 	{
-		// If a group node passes the filter, all child nodes will be shown
+		// If a group node passes the filter, all child nodes will be shown.
 		if (bIsNodeVisible)
 		{
 			MakeSubtreeVisible(NodePtr, bFilterIsEmpty);
 			return true;
 		}
 
-		NodePtr->ClearFilteredChildren();
-
 		const TArray<FBaseTreeNodePtr>& GroupChildren = NodePtr->GetChildren();
 		const int32 NumChildren = GroupChildren.Num();
+
+		NodePtr->ClearFilteredChildren(bFilterIsEmpty ? NumChildren : 0);
+
 		int32 NumVisibleChildren = 0;
 		for (int32 Cx = 0; Cx < NumChildren; ++Cx)
 		{
@@ -1093,6 +1094,7 @@ bool STableTreeView::ApplyHierarchicalFilterForNode(FTableTreeNodePtr NodePtr, b
 			{
 				break;
 			}
+
 			// Add a child.
 			const FTableTreeNodePtr& ChildNodePtr = StaticCastSharedPtr<FTableTreeNode>(GroupChildren[Cx]);
 			if (ApplyHierarchicalFilterForNode(ChildNodePtr, bFilterIsEmpty))
@@ -1103,7 +1105,7 @@ bool STableTreeView::ApplyHierarchicalFilterForNode(FTableTreeNodePtr NodePtr, b
 		}
 
 		const bool bIsGroupNodeVisible = bIsNodeVisible || NumVisibleChildren > 0;
-		if (!bFilterIsEmpty && bIsGroupNodeVisible )
+		if (!bFilterIsEmpty && bIsGroupNodeVisible)
 		{
 			if (bRunInAsyncMode)
 			{
@@ -1129,10 +1131,11 @@ bool STableTreeView::ApplyAdvancedFiltersForNode(FTableTreeNodePtr NodePtr)
 {
 	if (NodePtr->IsGroup())
 	{
-		NodePtr->ClearFilteredChildren();
-
 		const TArray<FBaseTreeNodePtr>& GroupChildren = NodePtr->GetChildren();
 		const int32 NumChildren = GroupChildren.Num();
+
+		NodePtr->ClearFilteredChildren();
+
 		int32 NumVisibleChildren = 0;
 		for (int32 Cx = 0; Cx < NumChildren; ++Cx)
 		{
@@ -1140,6 +1143,7 @@ bool STableTreeView::ApplyAdvancedFiltersForNode(FTableTreeNodePtr NodePtr)
 			{
 				break;
 			}
+
 			// Add a child.
 			const FTableTreeNodePtr& ChildNodePtr = StaticCastSharedPtr<FTableTreeNode>(GroupChildren[Cx]);
 			if (ApplyAdvancedFiltersForNode(ChildNodePtr))
@@ -1178,10 +1182,11 @@ bool STableTreeView::MakeSubtreeVisible(FTableTreeNodePtr NodePtr, bool bFilterI
 	bool bPassesNonEmptyFilter = !bFilterIsEmpty && Filters->PassesAllFilters(NodePtr);
 	if (NodePtr->IsGroup())
 	{
-		NodePtr->ClearFilteredChildren();
-
 		const TArray<FBaseTreeNodePtr>& GroupChildren = NodePtr->GetChildren();
 		const int32 NumChildren = GroupChildren.Num();
+
+		NodePtr->ClearFilteredChildren(bFilterIsEmpty ? NumChildren : 0);
+
 		int32 NumVisibleChildren = 0;
 		bool bShouldExpand = false;
 		for (int32 Cx = 0; Cx < NumChildren; ++Cx)
@@ -1407,6 +1412,7 @@ void STableTreeView::ApplyGrouping()
 
 ////////////////////////////////////////////////////////////////////////////////////////////////////
 
+// Note: This function is also called from async thread, with CurrentAsyncOpGroupings as parameter.
 void STableTreeView::CreateGroups(const TArray<TSharedPtr<FTreeNodeGrouping>>& Groupings)
 {
 	if (TableTreeNodes.Num() == 0)
@@ -1418,7 +1424,10 @@ void STableTreeView::CreateGroups(const TArray<TSharedPtr<FTreeNodeGrouping>>& G
 	FStopwatch Stopwatch;
 	Stopwatch.Start();
 
-	GroupNodesRec(TableTreeNodes, *Root, 0, Groupings);
+	if (Groupings.Num() > 0)
+	{
+		GroupNodesRec(TableTreeNodes, *Root, 0, Groupings);
+	}
 
 	Stopwatch.Stop();
 	const double GroupingTime = Stopwatch.GetAccumulatedTime();
@@ -1434,47 +1443,73 @@ void STableTreeView::CreateGroups(const TArray<TSharedPtr<FTreeNodeGrouping>>& G
 
 void STableTreeView::GroupNodesRec(const TArray<FTableTreeNodePtr>& Nodes, FTableTreeNode& ParentGroup, int32 GroupingDepth, const TArray<TSharedPtr<FTreeNodeGrouping>>& Groupings)
 {
-	if (bIsUpdateRunning && AsyncOperationProgress.ShouldCancelAsyncOp())
+	if (AsyncOperationProgress.ShouldCancelAsyncOp())
 	{
 		return;
 	}
 
-	if (Groupings.Num() == 0)
+	if (!ensure(GroupingDepth < Groupings.Num()))
 	{
 		return;
 	}
 
-	FTreeNodeGrouping& Grouping = *Groupings[GroupingDepth];
+	TSet<FBaseTreeNode*> OldGroupNodes;
 
-	Grouping.GroupNodes(Nodes, ParentGroup, Table, AsyncOperationProgress);
-
-	for (FBaseTreeNodePtr GroupPtr : ParentGroup.GetChildren())
+	// Apply current grouping recursively to the (old) group nodes.
+	for (FBaseTreeNodePtr NodePtr : Nodes)
 	{
-		ensure(GroupPtr->IsGroup());
-		FTableTreeNodePtr TableTreeGroupPtr = StaticCastSharedPtr<FTableTreeNode>(GroupPtr);
-		if (!bRunInAsyncMode)
+		if (NodePtr->IsGroup())
 		{
-			TreeView->SetItemExpansion(TableTreeGroupPtr, TableTreeGroupPtr->IsExpanded());
+			OldGroupNodes.Add(NodePtr.Get());
+
+			FTableTreeNode& Group = static_cast<FTableTreeNode&>(*NodePtr);
+
+			// Extract child nodes.
+			TArray<FTableTreeNodePtr> ChildNodes;
+			Group.SwapChildrenFast(reinterpret_cast<TArray<FBaseTreeNodePtr>&>(ChildNodes));
+
+			GroupNodesRec(ChildNodes, Group, GroupingDepth, Groupings);
 		}
 	}
 
-	if (GroupingDepth < Groupings.Num() - 1)
+	ParentGroup.ClearChildren();
+
+	// Group the current list of nodes.
+	FTreeNodeGrouping& Grouping = *Groupings[GroupingDepth];
+	Grouping.GroupNodes(Nodes, ParentGroup, Table, AsyncOperationProgress);
+
+	if (AsyncOperationProgress.ShouldCancelAsyncOp())
 	{
-		TArray<FTableTreeNodePtr> ChildNodes;
+		return;
+	}
 
+	if (GroupingDepth < Groupings.Num() - 1)
+	{
+		// Apply the next grouping recursively to the new group nodes.
 		for (FBaseTreeNodePtr GroupPtr : ParentGroup.GetChildren())
 		{
 			ensure(GroupPtr->IsGroup());
-			FTableTreeNode& Group = *StaticCastSharedPtr<FTableTreeNode>(GroupPtr);
-
-			// Make a copy of the child nodes.
-			ChildNodes.Reset();
-			for (FBaseTreeNodePtr ChildPtr : Group.GetChildren())
+			if (!OldGroupNodes.Contains(GroupPtr.Get()))
 			{
-				ChildNodes.Add(StaticCastSharedPtr<FTableTreeNode>(ChildPtr));
+				FTableTreeNode& Group = static_cast<FTableTreeNode&>(*GroupPtr);
+
+				// Extract child nodes.
+				TArray<FTableTreeNodePtr> ChildNodes;
+				Group.SwapChildrenFast(reinterpret_cast<TArray<FBaseTreeNodePtr>&>(ChildNodes));
+
+				GroupNodesRec(ChildNodes, Group, GroupingDepth + 1, Groupings);
 			}
+		}
+	}
 
-			GroupNodesRec(ChildNodes, Group, GroupingDepth + 1, Groupings);
+	if (!bRunInAsyncMode)
+	{
+		// Expand group nodes.
+		for (FBaseTreeNodePtr GroupPtr : ParentGroup.GetChildren())
+		{
+			ensure(GroupPtr->IsGroup());
+			FTableTreeNodePtr TableTreeGroupPtr = StaticCastSharedPtr<FTableTreeNode>(GroupPtr);
+			TreeView->SetItemExpansion(TableTreeGroupPtr, TableTreeGroupPtr->IsExpanded());
 		}
 	}
 }
