Updating style files with new icon paths

diff --git a/Engine/Plugins/Slate/SlateInsights/Source/SlateInsights/Private/SlateInsightsStyle.h b/Engine/Plugins/Slate/SlateInsights/Source/SlateInsights/Private/SlateInsightsStyle.h
--- a/Engine/Plugins/Slate/SlateInsights/Source/SlateInsights/Private/SlateInsightsStyle.h
+++ b/Engine/Plugins/Slate/SlateInsights/Source/SlateInsights/Private/SlateInsightsStyle.h
@@ -6,14 +6,13 @@
 #include "Styling/SlateStyle.h"
 #include "Styling/SlateStyleRegistry.h"
 #include "Styling/SlateTypes.h"
+#include "Styling/SlateStyleMacros.h"
 
 namespace UE
 {
 namespace SlateInsights
 {
 
-#define IMAGE_BRUSH(RelativePath, ...) FSlateImageBrush(RootToContentDir(RelativePath, TEXT(".png")), __VA_ARGS__)
-
 class FSlateInsightsStyle final : public FSlateStyleSet
 {
 public:
@@ -24,7 +23,7 @@ public:
 		SetContentRoot(FPaths::EngineContentDir() / TEXT("Editor/Slate"));
 		SetCoreContentRoot(FPaths::EngineContentDir() / TEXT("Slate"));
 
-		Set("SlateProfiler.Icon.Small", new IMAGE_BRUSH("Icons/Profiler/profiler_stats_40x", FVector2D(16.0f, 16.0f)));
+		Set("SlateProfiler.Icon.Small", new CORE_IMAGE_BRUSH("Icons/Profiler/profiler_stats_40x", FVector2D(16.0f, 16.0f)));
 
 		Set("SlateGraph.Color.WidgetCount", FLinearColor(FColorList::Aquamarine));
 		Set("SlateGraph.Color.TickCount", FLinearColor(FColorList::BronzeII));
diff --git a/Engine/Plugins/VirtualProduction/SequencerPlaylists/Source/SequencerPlaylists/Private/SequencerPlaylistsStyle.cpp b/Engine/Plugins/VirtualProduction/SequencerPlaylists/Source/SequencerPlaylists/Private/SequencerPlaylistsStyle.cpp
--- a/Engine/Plugins/VirtualProduction/SequencerPlaylists/Source/SequencerPlaylists/Private/SequencerPlaylistsStyle.cpp
+++ b/Engine/Plugins/VirtualProduction/SequencerPlaylists/Source/SequencerPlaylists/Private/SequencerPlaylistsStyle.cpp
@@ -12,6 +12,7 @@
 
 
 #define RootToContentDir Style->RootToContentDir
+#define RootToCoreContentDir Style->RootToCoreContentDir
 
 
 TSharedPtr<FSlateStyleSet> FSequencerPlaylistsStyle::StyleInstance = nullptr;
@@ -244,13 +245,14 @@ TSharedRef<FSlateStyleSet> FSequencerPlaylistsStyle::Create()
 		// Engine/Content/Editor/Slate/...
 		{
 			Style->SetContentRoot(EngineEditorSlateDir);
+			Style->SetCoreContentRoot(FPaths::EngineContentDir() / TEXT("Slate"));
 
 			Style->Set("SequencerPlaylists.SavePlaylist", new IMAGE_BRUSH_SVG("Starship/Common/SaveCurrent", Icon16x16));
 
 			Style->Set("SequencerPlaylists.Play", new IMAGE_BRUSH_SVG("Starship/Common/play", Icon20x20));
 			Style->Set("SequencerPlaylists.Play.Small", new IMAGE_BRUSH_SVG("Starship/Common/play", Icon12x12));
-			Style->Set("SequencerPlaylists.Stop", new IMAGE_BRUSH_SVG("Starship/MainToolbar/stop", Icon20x20));
-			Style->Set("SequencerPlaylists.Stop.Small", new IMAGE_BRUSH_SVG("Starship/MainToolbar/stop", Icon12x12));
+			Style->Set("SequencerPlaylists.Stop", new IMAGE_BRUSH_SVG("Starship/Common/stop", Icon20x20));
+			Style->Set("SequencerPlaylists.Stop.Small", new IMAGE_BRUSH_SVG("Starship/Common/stop", Icon12x12));
 			Style->Set("SequencerPlaylists.Reset", new IMAGE_BRUSH_SVG("Starship/Common/Reset", Icon20x20));
 			Style->Set("SequencerPlaylists.Reset.Small", new IMAGE_BRUSH_SVG("Starship/Common/Reset", Icon12x12));
 		}
diff --git a/Engine/Source/Developer/LogVisualizer/Private/LogVisualizerStyle.cpp b/Engine/Source/Developer/LogVisualizer/Private/LogVisualizerStyle.cpp
--- a/Engine/Source/Developer/LogVisualizer/Private/LogVisualizerStyle.cpp
+++ b/Engine/Source/Developer/LogVisualizer/Private/LogVisualizerStyle.cpp
@@ -10,6 +10,7 @@
 
 // This is to fix the issue that SlateStyleMacros like IMAGE_BRUSH look for RootToContentDir but StyleSet->RootToContentDir is how this style is set up
 #define RootToContentDir Style.RootToContentDir
+#define RootToCoreContentDir Style.RootToCoreContentDir
 
 TSharedPtr< FSlateStyleSet > FLogVisualizerStyle::StyleInstance = nullptr;
 
@@ -66,8 +67,8 @@ TSharedRef< FSlateStyleSet > FLogVisualizerStyle::Create()
 		Style.Set("LogVisualizerApp.TabIcon", new IMAGE_BRUSH_SVG("Starship/Common/VisualLogger", Icon16x16));
 		
 		Style.Set("LogVisualizer.LogBar.Background", new BOX_BRUSH("Common/ProgressBar_Background", FMargin(5.f / 12.f)));
-		Style.Set("LogVisualizer.LogBar.Selected", new BOX_BRUSH("Common/TaskGraph_Selected", FMargin(5.f / 12.f)));
-		Style.Set("LogVisualizer.LogBar.EntryDefault", new BOX_BRUSH("Common/TaskGraph_Mono", FMargin(5.f / 12.f)));
+		Style.Set("LogVisualizer.LogBar.Selected", new CORE_BOX_BRUSH("Common/TaskGraph_Selected", FMargin(5.f / 12.f)));
+		Style.Set("LogVisualizer.LogBar.EntryDefault", new CORE_BOX_BRUSH("Common/TaskGraph_Mono", FMargin(5.f / 12.f)));
 		Style.Set("LogVisualizer.LogBar.TimeMark", new BOX_BRUSH("Icons/LV_BarMark", FMargin(5.f / 12.f)));
 
 		//Style.Set("ToolPanel.GroupBorder", FAppStyle::Get().GetBrush("ToolPanel.GroupBorder"));
@@ -150,7 +151,7 @@ TSharedRef< FSlateStyleSet > FLogVisualizerStyle::Create()
 
 	// Filters
 	{
-		Style.Set("Filters.FilterIcon", new IMAGE_BRUSH("Icons/Profiler/Profiler_Filter_Events_16x", Icon16x16));
+		Style.Set("Filters.FilterIcon", new CORE_IMAGE_BRUSH("Icons/Profiler/Profiler_Filter_Events_16x", Icon16x16));
 		Style.Set("Filters.Style", FAppStyle::Get().GetWidgetStyle<FComboButtonStyle>("ToolbarComboButton"));
 		Style.Set("ContentBrowser.FilterButtonBorder", new BOX_BRUSH("Common/RoundedSelection_16x", FMargin(4.0f / 16.0f)));
 		//FSlateBrush *FilterButtonBorder = FAppStyle::Get().GetBrush("ToolbarComboButton");
diff --git a/Engine/Source/Developer/TraceInsights/Private/Insights/InsightsStyle.cpp b/Engine/Source/Developer/TraceInsights/Private/Insights/InsightsStyle.cpp
--- a/Engine/Source/Developer/TraceInsights/Private/Insights/InsightsStyle.cpp
+++ b/Engine/Source/Developer/TraceInsights/Private/Insights/InsightsStyle.cpp
@@ -134,9 +134,9 @@
 
 	Set("RoundedBackground", new FSlateRoundedBoxBrush(FLinearColor(1.0f, 1.0f, 1.0f, 1.0f), Icon16x16));
 
-	Set("Border.TB", new EDITOR_BOX_BRUSH("Icons/Profiler/Profiler_Border_TB_16x", FMargin(4.0f / 16.0f)));
-	Set("Border.L", new EDITOR_BOX_BRUSH("Icons/Profiler/Profiler_Border_L_16x", FMargin(4.0f / 16.0f)));
-	Set("Border.R", new EDITOR_BOX_BRUSH("Icons/Profiler/Profiler_Border_R_16x", FMargin(4.0f / 16.0f)));
+	Set("Border.TB", new CORE_BOX_BRUSH("Icons/Profiler/Profiler_Border_TB_16x", FMargin(4.0f / 16.0f)));
+	Set("Border.L", new CORE_BOX_BRUSH("Icons/Profiler/Profiler_Border_L_16x", FMargin(4.0f / 16.0f)));
+	Set("Border.R", new CORE_BOX_BRUSH("Icons/Profiler/Profiler_Border_R_16x", FMargin(4.0f / 16.0f)));
 
 	Set("Graph.Point", new EDITOR_IMAGE_BRUSH("Old/Graph/ExecutionBubble", Icon16x16));
 	Set("TreeViewBanner.WarningIcon", new CORE_IMAGE_BRUSH_SVG("Starship/Common/alert-circle", Icon20x20, FStyleColors::Warning));
@@ -234,15 +234,15 @@
 	//////////////////////////////////////////////////
 	// Tasks
 
-	Set("Icons.GoToTask", new EDITOR_IMAGE_BRUSH("Icons/Profiler/profiler_ViewColumn_32x", Icon16x16));
-	Set("Icons.ShowTaskCriticalPath", new EDITOR_IMAGE_BRUSH("Icons/Profiler/profiler_HotPath_32x", Icon16x16));
-	Set("Icons.ShowTaskTransitions", new EDITOR_IMAGE_BRUSH("Icons/Profiler/profiler_Calls_32x", Icon16x16));
-	Set("Icons.ShowTaskConnections", new EDITOR_IMAGE_BRUSH("Icons/Profiler/profiler_Calls_32x", Icon16x16));
-	Set("Icons.ShowTaskPrerequisites", new EDITOR_IMAGE_BRUSH("Icons/Profiler/profiler_Calls_32x", Icon16x16));
-	Set("Icons.ShowTaskSubsequents", new EDITOR_IMAGE_BRUSH("Icons/Profiler/profiler_Calls_32x", Icon16x16));
-	Set("Icons.ShowNestedTasks", new EDITOR_IMAGE_BRUSH("Icons/Profiler/profiler_Calls_32x", Icon16x16));
-	Set("Icons.ShowTaskTrack", new EDITOR_IMAGE_BRUSH("Icons/Profiler/profiler_Calls_32x", Icon16x16));
-	Set("Icons.ShowDetailedTaskTrackInfo", new EDITOR_IMAGE_BRUSH("Icons/Profiler/profiler_Calls_32x", Icon16x16));
+	Set("Icons.GoToTask", new CORE_IMAGE_BRUSH("Icons/Profiler/profiler_ViewColumn_32x", Icon16x16));
+	Set("Icons.ShowTaskCriticalPath", new CORE_IMAGE_BRUSH("Icons/Profiler/profiler_HotPath_32x", Icon16x16));
+	Set("Icons.ShowTaskTransitions", new CORE_IMAGE_BRUSH("Icons/Profiler/profiler_Calls_32x", Icon16x16));
+	Set("Icons.ShowTaskConnections", new CORE_IMAGE_BRUSH("Icons/Profiler/profiler_Calls_32x", Icon16x16));
+	Set("Icons.ShowTaskPrerequisites", new CORE_IMAGE_BRUSH("Icons/Profiler/profiler_Calls_32x", Icon16x16));
+	Set("Icons.ShowTaskSubsequents", new CORE_IMAGE_BRUSH("Icons/Profiler/profiler_Calls_32x", Icon16x16));
+	Set("Icons.ShowNestedTasks", new CORE_IMAGE_BRUSH("Icons/Profiler/profiler_Calls_32x", Icon16x16));
+	Set("Icons.ShowTaskTrack", new CORE_IMAGE_BRUSH("Icons/Profiler/profiler_Calls_32x", Icon16x16));
+	Set("Icons.ShowDetailedTaskTrackInfo", new CORE_IMAGE_BRUSH("Icons/Profiler/profiler_Calls_32x", Icon16x16));
 
 	//////////////////////////////////////////////////
 
@@ -267,7 +267,7 @@
 
 	//////////////////////////////////////////////////
 
-	Set("Icons.ResetToDefault", new EDITOR_IMAGE_BRUSH("Icons/Profiler/profiler_ResetToDefault_32x", Icon16x16));
+	Set("Icons.ResetToDefault", new CORE_IMAGE_BRUSH("Icons/Profiler/profiler_ResetToDefault_32x", Icon16x16));
 	Set("Icons.DiffersFromDefault", new EDITOR_IMAGE_BRUSH_SVG("Starship/Common/ResetToDefault", Icon16x16));
 
 	Set("Icons.Console", new EDITOR_IMAGE_BRUSH_SVG("Starship/Common/Console", Icon16x16));
@@ -282,19 +282,19 @@
 	//Set("Icons.FolderOpen", new CORE_IMAGE_BRUSH_SVG("Starship/Common/folder-open", Icon16x16));		//-> use FAppStyle "Icons.FolderOpen"
 	//Set("Icons.FolderClosed", new CORE_IMAGE_BRUSH_SVG("Starship/Common/folder-closed", Icon16x16));	//-> use FAppStyle "Icons.FolderClosed"
 
-	Set("Icons.SortBy", new EDITOR_IMAGE_BRUSH("Icons/Profiler/profiler_SortBy_32x", Icon16x16));
+	Set("Icons.SortBy", new CORE_IMAGE_BRUSH("Icons/Profiler/profiler_SortBy_32x", Icon16x16));
 	//Set("Icons.SortAscending", new CORE_IMAGE_BRUSH_SVG("Starship/Common/SortUp", Icon16x16));	//-> use FAppStyle "Icons.SortUp"
 	//Set("Icons.SortDescending", new CORE_IMAGE_BRUSH_SVG("Starship/Common/SortDown", Icon16x16));	//-> use FAppStyle "Icons.SortDown"
 
-	Set("Icons.ViewColumn", new EDITOR_IMAGE_BRUSH("Icons/Profiler/profiler_ViewColumn_32x", Icon16x16));
-	Set("Icons.ResetColumn", new EDITOR_IMAGE_BRUSH("Icons/Profiler/profiler_ResetColumn_32x", Icon16x16));
+	Set("Icons.ViewColumn", new CORE_IMAGE_BRUSH("Icons/Profiler/profiler_ViewColumn_32x", Icon16x16));
+	Set("Icons.ResetColumn", new CORE_IMAGE_BRUSH("Icons/Profiler/profiler_ResetColumn_32x", Icon16x16));
 
-	Set("Icons.ExpandAll", new EDITOR_IMAGE_BRUSH("Icons/Profiler/profiler_ExpandAll_32x", Icon16x16));
-	Set("Icons.CollapseAll", new EDITOR_IMAGE_BRUSH("Icons/Profiler/profiler_CollapseAll_32x", Icon16x16));
-	Set("Icons.ExpandSelection", new EDITOR_IMAGE_BRUSH("Icons/Profiler/profiler_ExpandSelection_32x", Icon16x16));
-	Set("Icons.CollapseSelection", new EDITOR_IMAGE_BRUSH("Icons/Profiler/profiler_CollapseSelection_32x", Icon16x16));
+	Set("Icons.ExpandAll", new CORE_IMAGE_BRUSH("Icons/Profiler/profiler_ExpandAll_32x", Icon16x16));
+	Set("Icons.CollapseAll", new CORE_IMAGE_BRUSH("Icons/Profiler/profiler_CollapseAll_32x", Icon16x16));
+	Set("Icons.ExpandSelection", new CORE_IMAGE_BRUSH("Icons/Profiler/profiler_ExpandSelection_32x", Icon16x16));
+	Set("Icons.CollapseSelection", new CORE_IMAGE_BRUSH("Icons/Profiler/profiler_CollapseSelection_32x", Icon16x16));
 
-	Set("Icons.ToggleShowGraphSeries", new EDITOR_IMAGE_BRUSH("Icons/Profiler/profiler_ShowGraphData_32x", Icon16x16));
+	Set("Icons.ToggleShowGraphSeries", new CORE_IMAGE_BRUSH("Icons/Profiler/profiler_ShowGraphData_32x", Icon16x16));
 
 	Set("Icons.TestAutomation", new EDITOR_IMAGE_BRUSH_SVG("Starship/Common/TestAutomation", Icon16x16));
 	Set("Icons.Test", new EDITOR_IMAGE_BRUSH_SVG("Starship/Common/Test", Icon16x16));
diff --git a/Engine/Source/Editor/EditorStyle/Private/StarshipStyle.cpp b/Engine/Source/Editor/EditorStyle/Private/StarshipStyle.cpp
--- a/Engine/Source/Editor/EditorStyle/Private/StarshipStyle.cpp
+++ b/Engine/Source/Editor/EditorStyle/Private/StarshipStyle.cpp
@@ -993,7 +993,7 @@ void FStarshipEditorStyle::FStyle::SetupGeneralStyles()
 	// Session Frontend Window
 	{
 		Set("SessionFrontEnd.TabIcon", new IMAGE_BRUSH_SVG( "Starship/Common/SessionFrontend", Icon16x16 ) );
-		Set("SessionFrontEnd.Tabs.Tools", new IMAGE_BRUSH( "/Icons/icon_tab_Tools_16x", Icon16x16 ) );
+		Set("SessionFrontEnd.Tabs.Tools", new CORE_IMAGE_BRUSH( "/Icons/icon_tab_Tools_16x", Icon16x16 ) );
 	}
 
 	// Undo History Window
@@ -3029,8 +3029,8 @@ void FStarshipEditorStyle::FStyle::SetupProfilerStyle()
 		Set( "ProfilerCommand.ToggleDataCapture", new CORE_IMAGE_BRUSH( "Icons/Profiler/Profiler_Data_Capture_40x", Icon40x40 ) );
 		Set( "ProfilerCommand.ToggleDataCapture.Small", new CORE_IMAGE_BRUSH( "Icons/Profiler/Profiler_Data_Capture_40x", Icon20x20 ) );
 
-		Set( "ProfilerCommand.ToggleDataCapture.Checked", new CORE_IMAGE_BRUSH( "Icons/icon_stop_40x", Icon40x40 ) );
-		Set( "ProfilerCommand.ToggleDataCapture.Checked.Small", new CORE_IMAGE_BRUSH( "Icons/icon_stop_40x", Icon20x20 ) );
+		Set( "ProfilerCommand.ToggleDataCapture.Checked", new IMAGE_BRUSH( "Icons/icon_stop_40x", Icon40x40 ) );
+		Set( "ProfilerCommand.ToggleDataCapture.Checked.Small", new IMAGE_BRUSH( "Icons/icon_stop_40x", Icon20x20 ) );
 
 		Set( "ProfilerCommand.ToggleShowDataGraph", new CORE_IMAGE_BRUSH( "Icons/Profiler/profiler_ShowGraphData_32x", Icon32x32 ) );
 		Set( "ProfilerCommand.OpenEventGraph", new CORE_IMAGE_BRUSH( "Icons/Profiler/profiler_OpenEventGraph_32x", Icon16x16 ) );
@@ -5253,8 +5253,8 @@ void FStarshipEditorStyle::FStyle::SetupPersonaStyle()
 		Set( "PlayWorld.SingleFrameAdvance", new IMAGE_BRUSH_SVG( "Starship/MainToolbar/SingleFrameAdvance", Icon40x40 ) );
 		Set( "PlayWorld.SingleFrameAdvance.Small", new IMAGE_BRUSH_SVG( "Starship/MainToolbar/SingleFrameAdvance", Icon20x20 ) );
 
-		Set( "PlayWorld.StopPlaySession", new IMAGE_BRUSH_SVG("Starship/MainToolbar/stop", Icon40x40));
-		Set( "PlayWorld.StopPlaySession.Small", new IMAGE_BRUSH_SVG("Starship/MainToolbar/stop", Icon20x20));
+		Set( "PlayWorld.StopPlaySession", new CORE_IMAGE_BRUSH_SVG("Starship/Common/stop", Icon40x40));
+		Set( "PlayWorld.StopPlaySession.Small", new CORE_IMAGE_BRUSH_SVG("Starship/Common/stop", Icon20x20));
 
 		Set("PlayWorld.LateJoinSession", new IMAGE_BRUSH_SVG("Starship/MainToolbar/simulate", Icon40x40));
 		Set("PlayWorld.LateJoinSession.Small", new IMAGE_BRUSH_SVG("Starship/MainToolbar/simulate", Icon20x20));
@@ -6099,7 +6099,7 @@ void FStarshipEditorStyle::FStyle::SetupContentBrowserStyle()
 		Set( "ContentBrowser.ImportPackage", new IMAGE_BRUSH( "Icons/icon_Import_40x", Icon25x25 ) );
 	
 		// Asset Context Menu
-		Set( "ContentBrowser.AssetActions", new IMAGE_BRUSH( "Icons/icon_tab_Tools_16x", Icon16x16 ) );
+		Set( "ContentBrowser.AssetActions", new CORE_IMAGE_BRUSH( "Icons/icon_tab_Tools_16x", Icon16x16 ) );
 		Set( "ContentBrowser.AssetActions.Edit", new IMAGE_BRUSH( "Icons/Edit/icon_Edit_16x", Icon16x16 ) );
 		Set( "ContentBrowser.AssetActions.Delete", new IMAGE_BRUSH( "Icons/icon_delete_16px", Icon16x16, FLinearColor( 0.4f, 0.5f, 0.7f, 1.0f ) ) );
 		//Set( "ContentBrowser.AssetActions.Delete", new IMAGE_BRUSH( "Icons/Edit/icon_Edit_Delete_16x", Icon16x16) );
@@ -7151,7 +7151,7 @@ void FStarshipEditorStyle::FStyle::SetupAutomationStyles()
 		Set( "DeviceDetails.Reboot", new IMAGE_BRUSH_SVG("Starship/DeviceManager/Reboot", Icon20x20));
 
 		Set( "DeviceDetails.TabIcon", new IMAGE_BRUSH_SVG( "Starship/Common/DeviceManager", Icon16x16 ) );
-		Set( "DeviceDetails.Tabs.Tools", new IMAGE_BRUSH( "/Icons/icon_tab_Tools_16x", Icon16x16 ) );
+		Set( "DeviceDetails.Tabs.Tools", new CORE_IMAGE_BRUSH( "/Icons/icon_tab_Tools_16x", Icon16x16 ) );
 		Set( "DeviceDetails.Tabs.ProfileEditor", new IMAGE_BRUSH_SVG( "Starship/Common/DeviceProfiles", Icon16x16 ) );
 		Set( "DeviceDetails.Tabs.ProfileEditorSingleProfile", new IMAGE_BRUSH( "/Icons/icon_tab_DeviceProfileEditor_16x", Icon16x16 ) );
 
