[Insights] Fixed missing icon for text filter toggle (Session Browser) and updated icons for "add/remove graph series" actions in context menu of Timers and Counters.

diff --git a/Engine/Source/Developer/TraceInsights/Private/Insights/InsightsStyle.cpp b/Engine/Source/Developer/TraceInsights/Private/Insights/InsightsStyle.cpp
--- a/Engine/Source/Developer/TraceInsights/Private/Insights/InsightsStyle.cpp
+++ b/Engine/Source/Developer/TraceInsights/Private/Insights/InsightsStyle.cpp
@@ -270,7 +270,7 @@ void FInsightsStyle::FStyle::Initialize()
 	Set("Icons.ResetToDefault", new CORE_IMAGE_BRUSH("Icons/Profiler/profiler_ResetToDefault_32x", Icon16x16));
 	Set("Icons.DiffersFromDefault", new EDITOR_IMAGE_BRUSH_SVG("Starship/Common/ResetToDefault", Icon16x16));
 
-	Set("Icons.Console", new EDITOR_IMAGE_BRUSH_SVG("Starship/Common/Console", Icon16x16));
+	Set("Icons.Console", new CORE_IMAGE_BRUSH_SVG("Starship/Common/Console", Icon16x16));
 	//Set("Icons.Filter", new CORE_IMAGE_BRUSH_SVG("Starship/Common/filter", Icon16x16));	//-> use FAppStyle "Icons.Filter"
 	Set("Icons.Filter.ToolBar", new CORE_IMAGE_BRUSH_SVG("Starship/Common/filter", Icon20x20));
 	Set("Icons.Find", new EDITOR_IMAGE_BRUSH_SVG("Starship/Common/TraceDataFiltering", Icon16x16));
@@ -294,7 +294,8 @@ void FInsightsStyle::FStyle::Initialize()
 	Set("Icons.ExpandSelection", new CORE_IMAGE_BRUSH("Icons/Profiler/profiler_ExpandSelection_32x", Icon16x16));
 	Set("Icons.CollapseSelection", new CORE_IMAGE_BRUSH("Icons/Profiler/profiler_CollapseSelection_32x", Icon16x16));
 
-	Set("Icons.ToggleShowGraphSeries", new CORE_IMAGE_BRUSH("Icons/Profiler/profiler_ShowGraphData_32x", Icon16x16));
+	Set("Icons.AddGraphSeries", new CORE_IMAGE_BRUSH_SVG("Starship/Common/plus", Icon16x16));
+	Set("Icons.RemoveGraphSeries", new CORE_IMAGE_BRUSH_SVG("Starship/Common/close", Icon16x16));
 
 	Set("Icons.TestAutomation", new EDITOR_IMAGE_BRUSH_SVG("Starship/Common/TestAutomation", Icon16x16));
 	Set("Icons.Test", new EDITOR_IMAGE_BRUSH_SVG("Starship/Common/Test", Icon16x16));
diff --git a/Engine/Source/Developer/TraceInsights/Private/Insights/Widgets/SStatsView.cpp b/Engine/Source/Developer/TraceInsights/Private/Insights/Widgets/SStatsView.cpp
--- a/Engine/Source/Developer/TraceInsights/Private/Insights/Widgets/SStatsView.cpp
+++ b/Engine/Source/Developer/TraceInsights/Private/Insights/Widgets/SStatsView.cpp
@@ -409,8 +409,8 @@ TSharedPtr<SWidget> SStatsView::TreeView_GetMenuContent()
 				MenuBuilder.AddMenuEntry
 				(
 					LOCTEXT("ContextMenu_RemoveFromGraphTrack", "Remove series from graph track"),
-					LOCTEXT("ContextMenu_RemoveFromGraphTrack_Desc", "Remove the series containing event instances of the selected counter from the timing graph track."),
-					FSlateIcon(FInsightsStyle::GetStyleSetName(), "Icons.ToggleShowGraphSeries"),
+					LOCTEXT("ContextMenu_RemoveFromGraphTrack_Desc", "Removes the series containing event instances of the selected counter from the Main Graph track."),
+					FSlateIcon(FInsightsStyle::GetStyleSetName(), "Icons.RemoveGraphSeries"),
 					Action_ToggleCounterInGraphTrack,
 					NAME_None,
 					EUserInterfaceActionType::Button
@@ -421,8 +421,8 @@ TSharedPtr<SWidget> SStatsView::TreeView_GetMenuContent()
 				MenuBuilder.AddMenuEntry
 				(
 					LOCTEXT("ContextMenu_AddToGraphTrack", "Add series to graph track"),
-					LOCTEXT("ContextMenu_AddToGraphTrack_Desc", "Add a series containing event instances of the selected counter to the timing graph track."),
-					FSlateIcon(FInsightsStyle::GetStyleSetName(), "Icons.ToggleShowGraphSeries"),
+					LOCTEXT("ContextMenu_AddToGraphTrack_Desc", "Adds a series containing event instances of the selected counter to the Main Graph track."),
+					FSlateIcon(FInsightsStyle::GetStyleSetName(), "Icons.AddGraphSeries"),
 					Action_ToggleCounterInGraphTrack,
 					NAME_None,
 					EUserInterfaceActionType::Button
diff --git a/Engine/Source/Developer/TraceInsights/Private/Insights/Widgets/STimersView.cpp b/Engine/Source/Developer/TraceInsights/Private/Insights/Widgets/STimersView.cpp
--- a/Engine/Source/Developer/TraceInsights/Private/Insights/Widgets/STimersView.cpp
+++ b/Engine/Source/Developer/TraceInsights/Private/Insights/Widgets/STimersView.cpp
@@ -479,8 +479,8 @@ TSharedPtr<SWidget> STimersView::TreeView_GetMenuContent()
 				MenuBuilder.AddMenuEntry
 				(
 					LOCTEXT("ContextMenu_RemoveFromGraphTrack", "Remove series from graph track"),
-					LOCTEXT("ContextMenu_RemoveFromGraphTrack_Desc", "Remove the series containing event instances of the selected timer from the timing graph track."),
-					FSlateIcon(FInsightsStyle::GetStyleSetName(), "Icons.ToggleShowGraphSeries"),
+					LOCTEXT("ContextMenu_RemoveFromGraphTrack_Desc", "Removes the series containing event instances of the selected timer from the Main Graph track."),
+					FSlateIcon(FInsightsStyle::GetStyleSetName(), "Icons.RemoveGraphSeries"),
 					Action_ToggleTimerInGraphTrack,
 					NAME_None,
 					EUserInterfaceActionType::Button
@@ -491,8 +491,8 @@ TSharedPtr<SWidget> STimersView::TreeView_GetMenuContent()
 				MenuBuilder.AddMenuEntry
 				(
 					LOCTEXT("ContextMenu_AddToGraphTrack", "Add series to graph track"),
-					LOCTEXT("ContextMenu_AddToGraphTrack_Desc", "Add a series containing event instances of the selected timer to the timing graph track."),
-					FSlateIcon(FInsightsStyle::GetStyleSetName(), "Icons.ToggleShowGraphSeries"),
+					LOCTEXT("ContextMenu_AddToGraphTrack_Desc", "Adds a series containing event instances of the selected timer to the Main Graph track."),
+					FSlateIcon(FInsightsStyle::GetStyleSetName(), "Icons.AddGraphSeries"),
 					Action_ToggleTimerInGraphTrack,
 					NAME_None,
 					EUserInterfaceActionType::Button
