Network Insights
- Added support to trace simple stats assosciated with network trace data using macros UE_NET_TRACE_PACKET_STATSCOUNTER, UE_NET_TRACE_FRAME_STATSCOUNTER
- Runtime parts + some examples of hooked up stats
- New NetStatsCounters tab in NetworkInsights

Cherrypick of 18824234 + a few other changelists

diff --git a/Engine/Source/Developer/TraceInsights/Private/Insights/NetworkingProfiler/NetworkingProfilerCommands.cpp b/Engine/Source/Developer/TraceInsights/Private/Insights/NetworkingProfiler/NetworkingProfilerCommands.cpp
--- a/Engine/Source/Developer/TraceInsights/Private/Insights/NetworkingProfiler/NetworkingProfilerCommands.cpp
+++ b/Engine/Source/Developer/TraceInsights/Private/Insights/NetworkingProfiler/NetworkingProfilerCommands.cpp
@@ -68,6 +68,12 @@ void FNetworkingProfilerCommands::RegisterCommands()
 		"Toggles the visibility of the Net Stats view.",
 		EUserInterfaceActionType::ToggleButton,
 		FInputChord());
+
+	UI_COMMAND(ToggleNetStatsCountersViewVisibility,
+		"Net Stats Counters",
+		"Toggles the visibility of the Net Stats view.",
+		EUserInterfaceActionType::ToggleButton,
+		FInputChord());
 }
 PRAGMA_ENABLE_OPTIMIZATION
 
diff --git a/Engine/Source/Developer/TraceInsights/Private/Insights/NetworkingProfiler/NetworkingProfilerCommands.h b/Engine/Source/Developer/TraceInsights/Private/Insights/NetworkingProfiler/NetworkingProfilerCommands.h
--- a/Engine/Source/Developer/TraceInsights/Private/Insights/NetworkingProfiler/NetworkingProfilerCommands.h
+++ b/Engine/Source/Developer/TraceInsights/Private/Insights/NetworkingProfiler/NetworkingProfilerCommands.h
@@ -37,6 +37,9 @@ public:
 
 	/** Toggles visibility for the Net Stats view. Per profiler window, custom command. */
 	TSharedPtr<FUICommandInfo> ToggleNetStatsViewVisibility;
+
+	/** Toggles visibility for the Net Stats Counters view. Per profiler window, custom command. */
+	TSharedPtr<FUICommandInfo> ToggleNetStatsCountersViewVisibility;
 };
 
 ////////////////////////////////////////////////////////////////////////////////////////////////////
diff --git a/Engine/Source/Developer/TraceInsights/Private/Insights/NetworkingProfiler/ViewModels/NetStatsCounterGroupingAndSorting.cpp b/Engine/Source/Developer/TraceInsights/Private/Insights/NetworkingProfiler/ViewModels/NetStatsCounterGroupingAndSorting.cpp
new file mode 100644
--- /dev/null
+++ b/Engine/Source/Developer/TraceInsights/Private/Insights/NetworkingProfiler/ViewModels/NetStatsCounterGroupingAndSorting.cpp
@@ -0,0 +1,144 @@
+// Copyright Epic Games, Inc. All Rights Reserved.
+
+#include "NetStatsCounterGroupingAndSorting.h"
+
+#include "Insights/Table/ViewModels/TableColumn.h"
+#include "Insights/NetworkingProfiler/ViewModels/NetStatsCounterNodeHelper.h"
+
+#define LOCTEXT_NAMESPACE "NetStatsCounterNode"
+
+// Sort by name (ascending).
+#define INSIGHTS_DEFAULT_SORTING_NODES(A, B) return A->GetName().LexicalLess(B->GetName());
+
+////////////////////////////////////////////////////////////////////////////////////////////////////
+// Sorting by Event Type
+////////////////////////////////////////////////////////////////////////////////////////////////////
+
+FNetStatsCounterNodeSortingByEventType::FNetStatsCounterNodeSortingByEventType(TSharedRef<Insights::FTableColumn> InColumnRef)
+	: Insights::FTableCellValueSorter(
+		FName(TEXT("ByEventType")),
+		LOCTEXT("Sorting_ByEventType_Name", "By Type"),
+		LOCTEXT("Sorting_ByEventType_Title", "Sort By Type"),
+		LOCTEXT("Sorting_ByEventType_Desc", "Sort by event type."),
+		InColumnRef)
+{
+}
+
+////////////////////////////////////////////////////////////////////////////////////////////////////
+
+void FNetStatsCounterNodeSortingByEventType::Sort(TArray<Insights::FBaseTreeNodePtr>& NodesToSort, Insights::ESortMode SortMode) const
+{
+	if (SortMode == Insights::ESortMode::Ascending)
+	{
+		NodesToSort.Sort([](const Insights::FBaseTreeNodePtr& A, const Insights::FBaseTreeNodePtr& B) -> bool
+		{
+			ensure(A.IsValid() && A->GetTypeName() == FNetStatsCounterNode::TypeName);
+			const FNetStatsCounterNodePtr NetStatsCounterNodeA = StaticCastSharedPtr<FNetStatsCounterNode, Insights::FBaseTreeNode>(A);
+
+			ensure(B.IsValid() && B->GetTypeName() == FNetStatsCounterNode::TypeName);
+			const FNetStatsCounterNodePtr NetStatsCounterNodeB = StaticCastSharedPtr<FNetStatsCounterNode, Insights::FBaseTreeNode>(B);
+
+			if (NetStatsCounterNodeA->GetType() == NetStatsCounterNodeB->GetType())
+			{
+				INSIGHTS_DEFAULT_SORTING_NODES(A, B)
+			}
+			else
+			{
+				// Sort by net event type (ascending).
+				return NetStatsCounterNodeA->GetType() < NetStatsCounterNodeB->GetType();
+			}
+		});
+	}
+	else // if (SortMode == ESortMode::Descending)
+	{
+		NodesToSort.Sort([](const Insights::FBaseTreeNodePtr& A, const Insights::FBaseTreeNodePtr& B) -> bool
+		{
+			ensure(A.IsValid() && A->GetTypeName() == FNetStatsCounterNode::TypeName);
+			const FNetStatsCounterNodePtr NetStatsCounterNodeA = StaticCastSharedPtr<FNetStatsCounterNode, Insights::FBaseTreeNode>(A);
+
+			ensure(B.IsValid() && B->GetTypeName() == FNetStatsCounterNode::TypeName);
+			const FNetStatsCounterNodePtr NetStatsCounterNodeB = StaticCastSharedPtr<FNetStatsCounterNode, Insights::FBaseTreeNode>(B);
+
+			if (NetStatsCounterNodeA->GetType() == NetStatsCounterNodeB->GetType())
+			{
+				INSIGHTS_DEFAULT_SORTING_NODES(A, B)
+			}
+			else
+			{
+				// Sort by net event type (descending).
+				return NetStatsCounterNodeB->GetType() < NetStatsCounterNodeA->GetType();
+			}
+		});
+	}
+}
+
+////////////////////////////////////////////////////////////////////////////////////////////////////
+// Sort by Sum
+////////////////////////////////////////////////////////////////////////////////////////////////////
+
+FNetStatsCounterNodeSortingBySum::FNetStatsCounterNodeSortingBySum(TSharedRef<Insights::FTableColumn> InColumnRef)
+	: Insights::FTableCellValueSorter(
+		FName(TEXT("BySum")),
+		LOCTEXT("Sorting_BySum_Name", "By Sum"),
+		LOCTEXT("Sorting_BySum_Title", "Sort By Sum"),
+		LOCTEXT("Sorting_BySum_Desc", "Sort by aggregated Sum."),
+		InColumnRef)
+{
+}
+
+//////////////////////////////////////////////////////////////////////////////////////////////////
+
+void FNetStatsCounterNodeSortingBySum::Sort(TArray<Insights::FBaseTreeNodePtr>& NodesToSort, Insights::ESortMode SortMode) const
+{
+	if (SortMode == Insights::ESortMode::Ascending)
+	{
+		NodesToSort.Sort([](const Insights::FBaseTreeNodePtr& A, const Insights::FBaseTreeNodePtr& B) -> bool
+		{
+			ensure(A.IsValid() && A->GetTypeName() == FNetStatsCounterNode::TypeName);
+			const FNetStatsCounterNodePtr NetStatsCounterNodeA = StaticCastSharedPtr<FNetStatsCounterNode, Insights::FBaseTreeNode>(A);
+			const uint32 ValueA = NetStatsCounterNodeA->GetAggregatedStats().Sum;
+
+			ensure(B.IsValid() && B->GetTypeName() == FNetStatsCounterNode::TypeName);
+			const FNetStatsCounterNodePtr NetStatsCounterNodeB = StaticCastSharedPtr<FNetStatsCounterNode, Insights::FBaseTreeNode>(B);
+			const uint32 ValueB = NetStatsCounterNodeB->GetAggregatedStats().Sum;
+
+			if (ValueA == ValueB)
+			{
+				INSIGHTS_DEFAULT_SORTING_NODES(A, B)
+			}
+			else
+			{
+				// Sort by Sum (ascending).
+				return ValueA < ValueB;
+			}
+		});
+	}
+	else // if (SortMode == ESortMode::Descending)
+	{
+		NodesToSort.Sort([](const Insights::FBaseTreeNodePtr& A, const Insights::FBaseTreeNodePtr& B) -> bool
+		{
+			ensure(A.IsValid() && A->GetTypeName() == FNetStatsCounterNode::TypeName);
+			const FNetStatsCounterNodePtr NetStatsCounterNodeA = StaticCastSharedPtr<FNetStatsCounterNode, Insights::FBaseTreeNode>(A);
+			const uint32 ValueA = NetStatsCounterNodeA->GetAggregatedStats().Sum;
+
+			ensure(B.IsValid() && B->GetTypeName() == FNetStatsCounterNode::TypeName);
+			const FNetStatsCounterNodePtr NetStatsCounterNodeB = StaticCastSharedPtr<FNetStatsCounterNode, Insights::FBaseTreeNode>(B);
+			const uint32 ValueB = NetStatsCounterNodeB->GetAggregatedStats().Sum;
+
+			if (ValueA == ValueB)
+			{
+				INSIGHTS_DEFAULT_SORTING_NODES(A, B)
+			}
+			else
+			{
+				// Sort by Sum (descending).
+				return ValueB < ValueA;
+			}
+		});
+	}
+}
+
+////////////////////////////////////////////////////////////////////////////////////////////////////
+
+#undef INSIGHTS_DEFAULT_SORTING_NODES
+#undef LOCTEXT_NAMESPACE
diff --git a/Engine/Source/Developer/TraceInsights/Private/Insights/NetworkingProfiler/ViewModels/NetStatsCounterGroupingAndSorting.h b/Engine/Source/Developer/TraceInsights/Private/Insights/NetworkingProfiler/ViewModels/NetStatsCounterGroupingAndSorting.h
new file mode 100644
--- /dev/null
+++ b/Engine/Source/Developer/TraceInsights/Private/Insights/NetworkingProfiler/ViewModels/NetStatsCounterGroupingAndSorting.h
@@ -0,0 +1,58 @@
+// Copyright Epic Games, Inc. All Rights Reserved.
+
+#pragma once
+
+#include "CoreMinimal.h"
+
+// Insights
+#include "Insights/NetworkingProfiler/ViewModels/NetStatsCounterNode.h"
+#include "Insights/Table/ViewModels/TableColumn.h"
+#include "Insights/Table/ViewModels/TableCellValueSorter.h"
+#include "Insights/Table/ViewModels/TreeNodeGrouping.h"
+
+////////////////////////////////////////////////////////////////////////////////////////////////////
+// Sorters
+////////////////////////////////////////////////////////////////////////////////////////////////////
+
+class FNetStatsCounterNodeSortingByEventType: public Insights::FTableCellValueSorter
+{
+public:
+	FNetStatsCounterNodeSortingByEventType(TSharedRef<Insights::FTableColumn> InColumnRef);
+
+	virtual void Sort(TArray<Insights::FBaseTreeNodePtr>& NodesToSort, Insights::ESortMode SortMode) const override;
+};
+
+////////////////////////////////////////////////////////////////////////////////////////////////////
+
+class FNetStatsCounterNodeSortingBySum : public Insights::FTableCellValueSorter
+{
+public:
+	FNetStatsCounterNodeSortingBySum(TSharedRef<Insights::FTableColumn> InColumnRef);
+
+	virtual void Sort(TArray<Insights::FBaseTreeNodePtr>& NodesToSort, Insights::ESortMode SortMode) const override;
+};
+
+////////////////////////////////////////////////////////////////////////////////////////////////////
+// Organizers
+////////////////////////////////////////////////////////////////////////////////////////////////////
+
+/** Enumerates types of grouping or sorting for the NetStatsCounter-nodes. */
+enum class ENetStatsCounterGroupingMode
+{
+	/** Creates a single group for all */
+	Flat,
+
+	/** Creates one group for one letter. */
+	ByName,
+
+	/** Creates one group for each type. */
+	ByType,
+
+	/** Invalid enum type, may be used as a number of enumerations. */
+	InvalidOrMax,
+};
+
+/** Type definition for shared pointers to instances of ENetStatsCounterGroupingMode. */
+typedef TSharedPtr<ENetStatsCounterGroupingMode> ENetStatsCounterGroupingModePtr;
+
+////////////////////////////////////////////////////////////////////////////////////////////////////
diff --git a/Engine/Source/Developer/TraceInsights/Private/Insights/NetworkingProfiler/ViewModels/NetStatsCounterNode.cpp b/Engine/Source/Developer/TraceInsights/Private/Insights/NetworkingProfiler/ViewModels/NetStatsCounterNode.cpp
new file mode 100644
--- /dev/null
+++ b/Engine/Source/Developer/TraceInsights/Private/Insights/NetworkingProfiler/ViewModels/NetStatsCounterNode.cpp
@@ -0,0 +1,69 @@
+// Copyright Epic Games, Inc. All Rights Reserved.
+
+#include "NetStatsCounterNode.h"
+
+// Insights
+#include "Insights/Common/TimeUtils.h"
+
+#define LOCTEXT_NAMESPACE "NetStatsNode"
+
+const FName FNetStatsCounterNode::TypeName(TEXT("FNetStatsCounterNode"));
+
+////////////////////////////////////////////////////////////////////////////////////////////////////
+
+void FNetStatsCounterNode::ResetAggregatedStats()
+{
+	AggregatedStats.Reset();
+}
+
+////////////////////////////////////////////////////////////////////////////////////////////////////
+
+const FText FNetStatsCounterNode::FormatAggregatedStatsValue(uint32 Value) const
+{
+	if (AggregatedStats.Count > 0)
+	{
+		return FText::AsNumber(Value);
+	}
+
+	return LOCTEXT("AggregatedStatsNA", "N/A");
+}
+
+////////////////////////////////////////////////////////////////////////////////////////////////////
+
+const FText FNetStatsCounterNode::GetTextForAggregatedStatsSum(bool bForTooltip) const
+{
+	return FormatAggregatedStatsValue(AggregatedStats.Sum);
+}
+
+////////////////////////////////////////////////////////////////////////////////////////////////////
+
+const FText FNetStatsCounterNode::GetTextForAggregatedStatsMin(bool bForTooltip) const
+{
+	return FormatAggregatedStatsValue(AggregatedStats.Min);
+}
+
+////////////////////////////////////////////////////////////////////////////////////////////////////
+
+const FText FNetStatsCounterNode::GetTextForAggregatedStatsMax(bool bForTooltip) const
+{
+	return FormatAggregatedStatsValue(AggregatedStats.Max);
+}
+
+////////////////////////////////////////////////////////////////////////////////////////////////////
+
+const FText FNetStatsCounterNode::GetTextForAggregatedStatsAverage(bool bForTooltip) const
+{
+	return FormatAggregatedStatsValue(AggregatedStats.Average);
+}
+
+////////////////////////////////////////////////////////////////////////////////////////////////////
+
+void FNetStatsCounterNode::SetAggregatedStats(const FNetProfilerAggregatedStatsCounterStats& InAggregatedStats)
+{
+	ensure(InAggregatedStats.StatsCounterTypeIndex == AggregatedStats.StatsCounterTypeIndex);
+	AggregatedStats = InAggregatedStats;
+}
+
+////////////////////////////////////////////////////////////////////////////////////////////////////
+
+#undef LOCTEXT_NAMESPACE
diff --git a/Engine/Source/Developer/TraceInsights/Private/Insights/NetworkingProfiler/ViewModels/NetStatsCounterNode.h b/Engine/Source/Developer/TraceInsights/Private/Insights/NetworkingProfiler/ViewModels/NetStatsCounterNode.h
new file mode 100644
--- /dev/null
+++ b/Engine/Source/Developer/TraceInsights/Private/Insights/NetworkingProfiler/ViewModels/NetStatsCounterNode.h
@@ -0,0 +1,125 @@
+// Copyright Epic Games, Inc. All Rights Reserved.
+
+#pragma once
+
+#include "CoreMinimal.h"
+
+// Insights
+#include "Insights/Table/ViewModels/BaseTreeNode.h"
+#include "TraceServices/Model/NetProfiler.h"
+
+////////////////////////////////////////////////////////////////////////////////////////////////////
+
+enum class ENetStatsCounterNodeType
+{
+	/** A FrameStats node. */
+	FrameStats,
+
+	/** A PacketStats node. */
+	PacketStats,
+
+	/** A group node. */
+	Group,
+
+	/** Invalid enum type, may be used as a number of enumerations. */
+	InvalidOrMax,
+};
+
+////////////////////////////////////////////////////////////////////////////////////////////////////
+
+class FNetStatsCounterNode;
+
+/** Type definition for shared pointers to instances of FNetStatsCounterNode. */
+typedef TSharedPtr<class FNetStatsCounterNode> FNetStatsCounterNodePtr;
+
+/** Type definition for shared references to instances of FNetStatsCounterNode. */
+typedef TSharedRef<class FNetStatsCounterNode> FNetStatsCounterNodeRef;
+
+/** Type definition for shared references to const instances of FNetStatsCounterNode. */
+typedef TSharedRef<const class FNetStatsCounterNode> FNetStatsCounterNodeRefConst;
+
+/** Type definition for weak references to instances of FNetStatsCounterNode. */
+typedef TWeakPtr<class FNetStatsCounterNode> FNetStatsCounterNodeWeak;
+
+////////////////////////////////////////////////////////////////////////////////////////////////////
+
+/**
+ * Class used to store information about a net stats counter node (used in SNetStatsCountersView).
+ */
+class FNetStatsCounterNode : public Insights::FBaseTreeNode
+{
+public:
+	using FNetProfilerAggregatedStatsCounterStats = TraceServices::FNetProfilerAggregatedStatsCounterStats;
+	using ENetProfilerStatsCounterType = TraceServices::ENetProfilerStatsCounterType;
+
+	static const FName TypeName;
+	static constexpr uint32 InvalidCounterIndex = uint32(-1);
+
+public:
+	/** Initialization constructor for the stats node. */
+	FNetStatsCounterNode(uint32 InStatsCounterTypeIndex, const FName InName, ENetProfilerStatsCounterType CounterType)
+	: FBaseTreeNode(InName, false)
+	, Type(CounterType == ENetProfilerStatsCounterType::Packet ? ENetStatsCounterNodeType::PacketStats : ENetStatsCounterNodeType::FrameStats)
+	{
+		AggregatedStats.StatsCounterTypeIndex = InStatsCounterTypeIndex;
+		ResetAggregatedStats();
+	}
+
+	/** Initialization constructor for the group node. */
+	explicit FNetStatsCounterNode(const FName InGroupName)
+	: FBaseTreeNode(InGroupName, true)
+	, Type(ENetStatsCounterNodeType::Group)
+	{
+		AggregatedStats.StatsCounterTypeIndex = InvalidCounterIndex;
+		ResetAggregatedStats();
+	}
+
+	virtual const FName& GetTypeName() const override { return TypeName; }
+
+	/**
+	 * @return the counter id as provided by analyzer.
+	 */
+	uint32 GetCounterTypeIndex() const { return AggregatedStats.StatsCounterTypeIndex; }
+
+	/**
+	 * @return the type of this node or ENetStatsCounterNodeType::Group for group nodes.
+	 */
+	ENetStatsCounterNodeType GetType() const { return Type; }
+
+	/**
+	 * @return the aggregated stats of this stats counter.
+	 */
+	const FNetProfilerAggregatedStatsCounterStats& GetAggregatedStats() const { return AggregatedStats; }
+	FNetProfilerAggregatedStatsCounterStats& GetAggregatedStats() { return AggregatedStats; }
+
+	void ResetAggregatedStats();
+
+	void SetAggregatedStats(const FNetProfilerAggregatedStatsCounterStats& AggregatedStats);
+
+	const FText GetTextForAggregatedStatsSum(bool bForTooltip = false) const;
+	const FText GetTextForAggregatedStatsMin(bool bForTooltip = false) const;
+	const FText GetTextForAggregatedStatsMax(bool bForTooltip = false) const;
+	const FText GetTextForAggregatedStatsAverage(bool bForTooltip = false) const;
+
+	/** Sorts children using the specified class instance. */
+	template<typename TSortingClass>
+	void SortChildren(const TSortingClass& Instance)
+	{
+		auto Projection = [](Insights::FBaseTreeNodePtr Node) -> FNetStatsCounterNodePtr
+		{
+			return StaticCastSharedPtr<FNetStatsCounterNode, Insights::FBaseTreeNode>(Node);
+		};
+		Algo::SortBy(GetChildrenMutable(), Projection, Instance);
+	}
+
+private:
+	const FText FormatAggregatedStatsValue(uint32 Value) const;
+
+	/** The type of this node. */
+	const ENetStatsCounterNodeType Type;
+
+	/** Aggregated stats. */
+	FNetProfilerAggregatedStatsCounterStats AggregatedStats;
+};
+
+////////////////////////////////////////////////////////////////////////////////////////////////////
diff --git a/Engine/Source/Developer/TraceInsights/Private/Insights/NetworkingProfiler/ViewModels/NetStatsCounterNodeHelper.cpp b/Engine/Source/Developer/TraceInsights/Private/Insights/NetworkingProfiler/ViewModels/NetStatsCounterNodeHelper.cpp
new file mode 100644
--- /dev/null
+++ b/Engine/Source/Developer/TraceInsights/Private/Insights/NetworkingProfiler/ViewModels/NetStatsCounterNodeHelper.cpp
@@ -0,0 +1,126 @@
+// Copyright Epic Games, Inc. All Rights Reserved.
+
+#include "NetStatsCounterNodeHelper.h"
+// Insights
+#include "Insights/InsightsStyle.h"
+
+#define LOCTEXT_NAMESPACE "NetStatsCounterNode"
+
+////////////////////////////////////////////////////////////////////////////////////////////////////
+// NetStatsCounterNode Type Helper
+////////////////////////////////////////////////////////////////////////////////////////////////////
+
+FText NetStatsCounterNodeTypeHelper::ToText(const ENetStatsCounterNodeType NodeType)
+{
+	static_assert(static_cast<int>(ENetStatsCounterNodeType::InvalidOrMax) == 3, "Not all cases are handled in switch below!?");
+	switch (NodeType)
+	{
+		case ENetStatsCounterNodeType::FrameStats:	return LOCTEXT("Stats_Name_FrameStats", "FrameStats");
+		case ENetStatsCounterNodeType::PacketStats:	return LOCTEXT("Stats_Name_PacketStats", "PacketStats");
+		case ENetStatsCounterNodeType::Group:		return LOCTEXT("Stats_Name_Group", "Group");
+		default:						return LOCTEXT("InvalidOrMax", "InvalidOrMax");
+	}
+}
+
+////////////////////////////////////////////////////////////////////////////////////////////////////
+
+FText NetStatsCounterNodeTypeHelper::ToDescription(const ENetStatsCounterNodeType NodeType)
+{
+	static_assert(static_cast<int>(ENetStatsCounterNodeType::InvalidOrMax) == 3, "Not all cases are handled in switch below!?");
+	switch (NodeType)
+	{
+		case ENetStatsCounterNodeType::FrameStats:	return LOCTEXT("Stats_Desc_FrameStats", "FrameStats node");
+		case ENetStatsCounterNodeType::PacketStats:	return LOCTEXT("Stats_Desc_PacketStats", "PacketStats node");
+		case ENetStatsCounterNodeType::Group:		return LOCTEXT("Stats_Desc_Group", "Group node");
+		default:						return LOCTEXT("InvalidOrMax", "InvalidOrMax");
+	}
+}
+
+////////////////////////////////////////////////////////////////////////////////////////////////////
+
+FName NetStatsCounterNodeTypeHelper::ToBrushName(const ENetStatsCounterNodeType NodeType)
+{
+	static_assert(static_cast<int>(ENetStatsCounterNodeType::InvalidOrMax) == 3, "Not all cases are handled in switch below!?");
+	switch (NodeType)
+	{
+		case ENetStatsCounterNodeType::FrameStats:
+		case ENetStatsCounterNodeType::PacketStats:	
+			return TEXT("Profiler.FiltersAndPresets.StatTypeIcon");
+		case ENetStatsCounterNodeType::Group:
+			return TEXT("Profiler.Misc.GenericGroup");
+		default:
+			return NAME_None;
+	}
+}
+
+////////////////////////////////////////////////////////////////////////////////////////////////////
+
+const FSlateBrush* NetStatsCounterNodeTypeHelper::GetIconForGroup()
+{
+	return FInsightsStyle::GetBrush(TEXT("Icons.Group.TreeItem"));
+}
+
+////////////////////////////////////////////////////////////////////////////////////////////////////
+
+const FSlateBrush* NetStatsCounterNodeTypeHelper::GetIcon(const ENetStatsCounterNodeType NodeType)
+{
+	static_assert(static_cast<int>(ENetStatsCounterNodeType::InvalidOrMax) == 3, "Not all cases are handled in switch below!?");
+	switch (NodeType)
+	{
+		case ENetStatsCounterNodeType::FrameStats:
+		case ENetStatsCounterNodeType::PacketStats:
+			return FInsightsStyle::GetBrush(TEXT("Icons.NetEvent.TreeItem"));
+		case ENetStatsCounterNodeType::Group:
+			return FInsightsStyle::GetBrush(TEXT("Icons.Group.TreeItem"));
+		default:
+			return nullptr;
+	}
+}
+
+////////////////////////////////////////////////////////////////////////////////////////////////////
+// NetStatsCounterNode Grouping Helper
+////////////////////////////////////////////////////////////////////////////////////////////////////
+
+FText NetStatsCounterNodeGroupingHelper::ToText(const ENetStatsCounterGroupingMode GroupingMode)
+{
+	static_assert(static_cast<int>(ENetStatsCounterGroupingMode::InvalidOrMax) == 3, "Not all cases are handled in switch below!?");
+	switch (GroupingMode)
+	{
+	case ENetStatsCounterGroupingMode::Flat:		return LOCTEXT("Grouping_Name_Flat", "Flat");
+	case ENetStatsCounterGroupingMode::ByName:		return LOCTEXT("Grouping_Name_ByName", "Name");
+	case ENetStatsCounterGroupingMode::ByType:		return LOCTEXT("Grouping_Name_ByType", "Event Type");
+	default:								return LOCTEXT("InvalidOrMax", "InvalidOrMax");
+	}
+}
+
+////////////////////////////////////////////////////////////////////////////////////////////////////
+
+FText NetStatsCounterNodeGroupingHelper::ToDescription(const ENetStatsCounterGroupingMode GroupingMode)
+{
+	static_assert(static_cast<int>(ENetStatsCounterGroupingMode::InvalidOrMax) == 3, "Not all cases are handled in switch below!?");
+	switch (GroupingMode)
+	{
+	case ENetStatsCounterGroupingMode::Flat:		return LOCTEXT("Grouping_Desc_Flat", "Creates a single group. Includes all counters.");
+	case ENetStatsCounterGroupingMode::ByName:		return LOCTEXT("Grouping_Desc_ByName", "Creates one group for one letter.");
+	case ENetStatsCounterGroupingMode::ByType:		return LOCTEXT("Grouping_Desc_ByType", "Creates one group for each net event type.");
+	default:								return LOCTEXT("InvalidOrMax", "InvalidOrMax");
+	}
+}
+
+////////////////////////////////////////////////////////////////////////////////////////////////////
+
+FName NetStatsCounterNodeGroupingHelper::ToBrushName(const ENetStatsCounterGroupingMode GroupingMode)
+{
+	static_assert(static_cast<int>(ENetStatsCounterGroupingMode::InvalidOrMax) == 3, "Not all cases are handled in switch below!?");
+	switch (GroupingMode)
+	{
+	case ENetStatsCounterGroupingMode::Flat:		return TEXT("Profiler.FiltersAndPresets.GroupNameIcon"); //TODO: "Icons.Grouping.Flat"
+	case ENetStatsCounterGroupingMode::ByName:		return TEXT("Profiler.FiltersAndPresets.GroupNameIcon"); //TODO: "Icons.Grouping.ByName"
+	case ENetStatsCounterGroupingMode::ByType:		return TEXT("Profiler.FiltersAndPresets.StatTypeIcon"); //TODO
+	default:								return NAME_None;
+	}
+}
+
+////////////////////////////////////////////////////////////////////////////////////////////////////
+
+#undef LOCTEXT_NAMESPACE
diff --git a/Engine/Source/Developer/TraceInsights/Private/Insights/NetworkingProfiler/ViewModels/NetStatsCounterNodeHelper.h b/Engine/Source/Developer/TraceInsights/Private/Insights/NetworkingProfiler/ViewModels/NetStatsCounterNodeHelper.h
new file mode 100644
--- /dev/null
+++ b/Engine/Source/Developer/TraceInsights/Private/Insights/NetworkingProfiler/ViewModels/NetStatsCounterNodeHelper.h
@@ -0,0 +1,68 @@
+// Copyright Epic Games, Inc. All Rights Reserved.
+
+#pragma once
+
+#include "CoreMinimal.h"
+#include "Styling/SlateBrush.h"
+
+// Insights
+#include "Insights/NetworkingProfiler/ViewModels/NetStatsCounterNode.h"
+#include "Insights/NetworkingProfiler/ViewModels/NetStatsCounterGroupingAndSorting.h"
+
+////////////////////////////////////////////////////////////////////////////////////////////////////
+
+/** Helper struct that contains helper functions and classes for EStatsNodeType enum. */
+struct NetStatsCounterNodeTypeHelper
+{
+	/**
+	 * @param Type - The value to get the text for.
+	 *
+	 * @return text representation of the specified EStatsNodeType value.
+	 */
+	static FText ToText(const ENetStatsCounterNodeType Type);
+
+	/**
+	 * @param Type - The value to get the text for.
+	 *
+	 * @return text representation with more detailed explanation of the specified EStatsNodeType value.
+	 */
+	static FText ToDescription(const ENetStatsCounterNodeType Type);
+
+	/**
+	 * @param Type - The value to get the brush name for.
+	 *
+	 * @return brush name of the specified EStatsNodeType value.
+	 */
+	static FName ToBrushName(const ENetStatsCounterNodeType Type);
+
+	static const FSlateBrush* GetIconForGroup();
+	static const FSlateBrush* GetIcon(const ENetStatsCounterNodeType Type);
+};
+
+/** Helper struct that contains grouping static functions and classes. */
+struct NetStatsCounterNodeGroupingHelper
+{
+	/**
+	 * @param NetEventGroupingMode - The value to get the text for.
+	 *
+	 * @return text representation of the specified ENetStatsCounterGroupingMode value.
+	 */
+	static FText ToText(const ENetStatsCounterGroupingMode NetEventGroupingMode);
+
+	/**
+	 * @param NetEventGroupingMode - The value to get the text for.
+	 *
+	 * @return text representation with more detailed explanation of the specified ENetStatsCounterGroupingMode value.
+	 */
+	static FText ToDescription(const ENetStatsCounterGroupingMode NetEventGroupingMode);
+
+	/**
+	 * @param NetEventGroupingMode - The value to get the brush name for.
+	 *
+	 * @return brush name of the specified ENetStatsCounterGroupingMode value.
+	 */
+	static FName ToBrushName(const ENetStatsCounterGroupingMode NetEventGroupingMode);
+};
+
+////////////////////////////////////////////////////////////////////////////////////////////////////
+
diff --git a/Engine/Source/Developer/TraceInsights/Private/Insights/NetworkingProfiler/ViewModels/NetStatsCountersViewColumnFactory.cpp b/Engine/Source/Developer/TraceInsights/Private/Insights/NetworkingProfiler/ViewModels/NetStatsCountersViewColumnFactory.cpp
new file mode 100644
--- /dev/null
+++ b/Engine/Source/Developer/TraceInsights/Private/Insights/NetworkingProfiler/ViewModels/NetStatsCountersViewColumnFactory.cpp
@@ -0,0 +1,331 @@
+// Copyright Epic Games, Inc. All Rights Reserved.
+
+#include "NetStatsCountersViewColumnFactory.h"
+
+// Insights
+#include "Insights/Table/ViewModels/TableCellValueFormatter.h"
+#include "Insights/Table/ViewModels/TableCellValueGetter.h"
+#include "Insights/Table/ViewModels/TableCellValueSorter.h"
+#include "Insights/NetworkingProfiler/ViewModels/NetStatsCounterGroupingAndSorting.h"
+#include "Insights/NetworkingProfiler/ViewModels/NetStatsCounterNodeHelper.h"
+
+#define LOCTEXT_NAMESPACE "SNetStatsCountersView"
+
+////////////////////////////////////////////////////////////////////////////////////////////////////
+// Column identifiers
+
+const FName FNetStatsCountersViewColumns::NameColumnID(TEXT("Name")); // TEXT("_Hierarchy")
+const FName FNetStatsCountersViewColumns::TypeColumnID(TEXT("Type"));
+const FName FNetStatsCountersViewColumns::InstanceCountColumnID(TEXT("Count"));
+
+const FName FNetStatsCountersViewColumns::SumColumnID(TEXT("Sum"));
+const FName FNetStatsCountersViewColumns::MaxCountColumnID(TEXT("Max"));
+const FName FNetStatsCountersViewColumns::AverageCountColumnID(TEXT("Avg"));
+
+////////////////////////////////////////////////////////////////////////////////////////////////////
+
+void FNetStatsCountersViewColumnFactory::CreateNetStatsCountersViewColumns(TArray<TSharedRef<Insights::FTableColumn>>& Columns)
+{
+	Columns.Reset();
+
+	Columns.Add(CreateNameColumn());
+	Columns.Add(CreateTypeColumn());
+	Columns.Add(CreateInstanceCountColumn());
+
+	Columns.Add(CreateSumColumn());
+	Columns.Add(CreateMaxCountColumn());
+	Columns.Add(CreateAverageCountColumn());
+}
+
+////////////////////////////////////////////////////////////////////////////////////////////////////
+
+TSharedRef<Insights::FTableColumn> FNetStatsCountersViewColumnFactory::CreateNameColumn()
+{
+	using namespace Insights;
+
+	TSharedRef<FTableColumn> ColumnRef = MakeShared<FTableColumn>(FNetStatsCountersViewColumns::NameColumnID);
+	FTableColumn& Column = *ColumnRef;
+
+	Column.SetShortName(LOCTEXT("Name_ColumnName", "Name"));
+	Column.SetTitleName(LOCTEXT("Name_ColumnTitle", "NetStatsCounter or Group Name"));
+	Column.SetDescription(LOCTEXT("Name_ColumnDesc", "Name of the timer or group"));
+
+	Column.SetFlags(ETableColumnFlags::ShouldBeVisible |
+					ETableColumnFlags::CanBeFiltered |
+					ETableColumnFlags::IsHierarchy);
+
+	Column.SetHorizontalAlignment(HAlign_Left);
+	Column.SetInitialWidth(206.0f);
+	Column.SetMinWidth(42.0f);
+
+	Column.SetDataType(ETableCellDataType::Text);
+
+	TSharedRef<ITableCellValueGetter> Getter = MakeShared<FDisplayNameValueGetter>();
+	Column.SetValueGetter(Getter);
+
+	TSharedRef<ITableCellValueFormatter> Formatter = MakeShared<FTextValueFormatter>();
+	Column.SetValueFormatter(Formatter);
+
+	TSharedRef<ITableCellValueSorter> Sorter = MakeShared<FSorterByName>(ColumnRef);
+	Column.SetValueSorter(Sorter);
+
+	return ColumnRef;
+}
+
+////////////////////////////////////////////////////////////////////////////////////////////////////
+
+TSharedRef<Insights::FTableColumn> FNetStatsCountersViewColumnFactory::CreateTypeColumn()
+{
+	using namespace Insights;
+
+	TSharedRef<FTableColumn> ColumnRef = MakeShared<FTableColumn>(FNetStatsCountersViewColumns::TypeColumnID);
+	FTableColumn& Column = *ColumnRef;
+
+	Column.SetShortName(LOCTEXT("Type_ColumnName", "Type"));
+	Column.SetTitleName(LOCTEXT("Type_ColumnTitle", "Type"));
+	Column.SetDescription(LOCTEXT("Type_ColumnDesc", "Type of counter or group"));
+
+	Column.SetFlags(ETableColumnFlags::CanBeHidden |
+					ETableColumnFlags::CanBeFiltered);
+
+	Column.SetHorizontalAlignment(HAlign_Left);
+	Column.SetInitialWidth(60.0f);
+
+	Column.SetDataType(ETableCellDataType::Text);
+
+	class FNetStatsCounterTypeValueGetter : public FTableCellValueGetter
+	{
+	public:
+		virtual const TOptional<FTableCellValue> GetValue(const FTableColumn& Column, const FBaseTreeNode& Node) const
+		{
+			ensure(Column.GetId() == FNetStatsCountersViewColumns::TypeColumnID);
+			const FNetStatsCounterNode& NetStatsCounterNode = static_cast<const FNetStatsCounterNode&>(Node);
+			return TOptional<FTableCellValue>(FTableCellValue(NetStatsCounterNodeTypeHelper::ToText(NetStatsCounterNode.GetType())));
+		}
+	};
+
+	TSharedRef<ITableCellValueGetter> Getter = MakeShared<FNetStatsCounterTypeValueGetter>();
+	Column.SetValueGetter(Getter);
+
+	TSharedRef<ITableCellValueFormatter> Formatter = MakeShared<FTextValueFormatter>();
+	Column.SetValueFormatter(Formatter);
+
+	TSharedRef<ITableCellValueSorter> Sorter = MakeShared<FNetStatsCounterNodeSortingByEventType>(ColumnRef);
+	Column.SetValueSorter(Sorter);
+
+	return ColumnRef;
+}
+
+////////////////////////////////////////////////////////////////////////////////////////////////////
+
+TSharedRef<Insights::FTableColumn> FNetStatsCountersViewColumnFactory::CreateInstanceCountColumn()
+{
+	using namespace Insights;
+
+	TSharedRef<FTableColumn> ColumnRef = MakeShared<FTableColumn>(FNetStatsCountersViewColumns::InstanceCountColumnID);
+	FTableColumn& Column = *ColumnRef;
+
+	Column.SetShortName(LOCTEXT("InstanceCount_ColumnName", "Count"));
+	Column.SetTitleName(LOCTEXT("InstanceCount_ColumnTitle", "Count"));
+	Column.SetDescription(LOCTEXT("InstanceCount_ColumnDesc", "Number of counters in the selected range"));
+
+	Column.SetFlags(ETableColumnFlags::CanBeHidden |
+//					ETableColumnFlags::ShouldBeVisible |
+					ETableColumnFlags::CanBeFiltered);
+
+	Column.SetHorizontalAlignment(HAlign_Right);
+	Column.SetInitialWidth(60.0f);
+
+	Column.SetDataType(ETableCellDataType::Int64);
+
+	class FInstanceCountValueGetter : public FTableCellValueGetter
+	{
+	public:
+		virtual const TOptional<FTableCellValue> GetValue(const FTableColumn& Column, const FBaseTreeNode& Node) const
+		{
+			ensure(Column.GetId() == FNetStatsCountersViewColumns::InstanceCountColumnID);
+			const FNetStatsCounterNode& NetStatsCounterNode = static_cast<const FNetStatsCounterNode&>(Node);
+			return TOptional<FTableCellValue>(FTableCellValue(static_cast<int64>(NetStatsCounterNode.GetAggregatedStats().Count)));
+		}
+	};
+
+	TSharedRef<ITableCellValueGetter> Getter = MakeShared<FInstanceCountValueGetter>();
+	Column.SetValueGetter(Getter);
+
+	TSharedRef<ITableCellValueFormatter> Formatter = MakeShared<FInt64ValueFormatterAsNumber>();
+	Column.SetValueFormatter(Formatter);
+
+	return ColumnRef;
+}
+
+////////////////////////////////////////////////////////////////////////////////////////////////////
+// Inclusive  Columns
+////////////////////////////////////////////////////////////////////////////////////////////////////
+
+TSharedRef<Insights::FTableColumn> FNetStatsCountersViewColumnFactory::CreateSumColumn()
+{
+	using namespace Insights;
+
+	TSharedRef<FTableColumn> ColumnRef = MakeShared<FTableColumn>(FNetStatsCountersViewColumns::SumColumnID);
+	FTableColumn& Column = *ColumnRef;
+
+	Column.SetShortName(LOCTEXT("Sum_ColumnName", "Sum"));
+	Column.SetTitleName(LOCTEXT("Sum_ColumnTitle", "Sum"));
+	Column.SetDescription(LOCTEXT("Sum_ColumnDesc", "Total sum of selected NetStatsCounters"));
+
+	Column.SetFlags(ETableColumnFlags::CanBeHidden |
+					ETableColumnFlags::ShouldBeVisible |
+					ETableColumnFlags::CanBeFiltered);
+
+	Column.SetHorizontalAlignment(HAlign_Right);
+	Column.SetInitialWidth(SumColumnInitialWidth);
+
+	Column.SetDataType(ETableCellDataType::Int64);
+	Column.SetAggregation(ETableColumnAggregation::Sum);
+
+	class FSumValueGetter : public FTableCellValueGetter
+	{
+	public:
+		virtual const TOptional<FTableCellValue> GetValue(const FTableColumn& Column, const FBaseTreeNode& Node) const
+		{
+			ensure(Column.GetId() == FNetStatsCountersViewColumns::SumColumnID);
+			const FNetStatsCounterNode& NetStatsCounterNode = static_cast<const FNetStatsCounterNode&>(Node);
+			return TOptional<FTableCellValue>(FTableCellValue(static_cast<int64>(NetStatsCounterNode.GetAggregatedStats().Sum)));
+		}
+	};
+
+	TSharedRef<ITableCellValueGetter> Getter = MakeShared<FSumValueGetter>();
+	Column.SetValueGetter(Getter);
+
+	class FSumFormatter : public FTableCellValueFormatter
+	{
+	public:
+		virtual FText FormatValue(const FTableColumn& Column, const FBaseTreeNode& Node) const override
+		{
+			ensure(Column.GetId() == FNetStatsCountersViewColumns::SumColumnID);
+			const FNetStatsCounterNode& NetStatsCounterNode = static_cast<const FNetStatsCounterNode&>(Node);
+			return NetStatsCounterNode.GetTextForAggregatedStatsSum();
+		}
+	};
+	TSharedRef<ITableCellValueFormatter> Formatter = MakeShared<FSumFormatter>();
+	Column.SetValueFormatter(Formatter);
+
+	TSharedRef<ITableCellValueSorter> Sorter = MakeShared<FNetStatsCounterNodeSortingBySum>(ColumnRef);
+	Column.SetValueSorter(Sorter);
+
+	return ColumnRef;
+}
+
+////////////////////////////////////////////////////////////////////////////////////////////////////
+
+TSharedRef<Insights::FTableColumn> FNetStatsCountersViewColumnFactory::CreateMaxCountColumn()
+{
+	using namespace Insights;
+
+	TSharedRef<FTableColumn> ColumnRef = MakeShared<FTableColumn>(FNetStatsCountersViewColumns::MaxCountColumnID);
+	FTableColumn& Column = *ColumnRef;
+
+	Column.SetShortName(LOCTEXT("MaxCount_ColumnName", "Max"));
+	Column.SetTitleName(LOCTEXT("MaxCount_ColumnTitle", "Max Count"));
+	Column.SetDescription(LOCTEXT("MaxCount_ColumnDesc", "Maximum count of NetStatsCounters in the selected range"));
+
+	Column.SetFlags(ETableColumnFlags::CanBeHidden |
+					ETableColumnFlags::ShouldBeVisible |
+					ETableColumnFlags::CanBeFiltered);
+
+	Column.SetHorizontalAlignment(HAlign_Right);
+	Column.SetInitialWidth(CountColumnInitialWidth);
+
+	Column.SetDataType(ETableCellDataType::Int64);
+
+	class FMaxCountValueGetter : public FTableCellValueGetter
+	{
+	public:
+		virtual const TOptional<FTableCellValue> GetValue(const FTableColumn& Column, const FBaseTreeNode& Node) const
+		{
+			ensure(Column.GetId() == FNetStatsCountersViewColumns::MaxCountColumnID);
+			const FNetStatsCounterNode& NetStatsCounterNode = static_cast<const FNetStatsCounterNode&>(Node);
+			return TOptional<FTableCellValue>(FTableCellValue(static_cast<int64>(NetStatsCounterNode.GetAggregatedStats().Max)));
+		}
+	};
+
+	TSharedRef<ITableCellValueGetter> Getter = MakeShared<FMaxCountValueGetter>();
+	Column.SetValueGetter(Getter);
+
+	class FMaxCountFormatter : public FTableCellValueFormatter
+	{
+	public:
+		virtual FText FormatValue(const FTableColumn& Column, const FBaseTreeNode& Node) const override
+		{
+			ensure(Column.GetId() == FNetStatsCountersViewColumns::MaxCountColumnID);
+			const FNetStatsCounterNode& NetStatsCounterNode = static_cast<const FNetStatsCounterNode&>(Node);
+			return NetStatsCounterNode.GetTextForAggregatedStatsMax();
+		}
+	};
+	TSharedRef<ITableCellValueFormatter> Formatter = MakeShared<FMaxCountFormatter>();
+	Column.SetValueFormatter(Formatter);
+
+	TSharedRef<ITableCellValueSorter> Sorter = MakeShared<FSorterByInt64Value>(ColumnRef);
+	Column.SetValueSorter(Sorter);
+
+	return ColumnRef;
+}
+
+////////////////////////////////////////////////////////////////////////////////////////////////////
+
+TSharedRef<Insights::FTableColumn> FNetStatsCountersViewColumnFactory::CreateAverageCountColumn()
+{
+	using namespace Insights;
+
+	TSharedRef<FTableColumn> ColumnRef = MakeShared<FTableColumn>(FNetStatsCountersViewColumns::AverageCountColumnID);
+	FTableColumn& Column = *ColumnRef;
+
+	Column.SetShortName(LOCTEXT("AvgInclusive_ColumnName", "Avg"));
+	Column.SetTitleName(LOCTEXT("AvgInclusive_ColumnTitle", "Average count"));
+	Column.SetDescription(LOCTEXT("AvgInclusive_ColumnDesc", "Average count in selected range"));
+
+	Column.SetFlags(ETableColumnFlags::CanBeHidden |
+					ETableColumnFlags::ShouldBeVisible |
+					ETableColumnFlags::CanBeFiltered);
+
+	Column.SetHorizontalAlignment(HAlign_Right);
+	Column.SetInitialWidth(CountColumnInitialWidth);
+
+	Column.SetDataType(ETableCellDataType::Int64);
+
+	class FAverageCountValueGetter : public FTableCellValueGetter
+	{
+	public:
+		virtual const TOptional<FTableCellValue> GetValue(const FTableColumn& Column, const FBaseTreeNode& Node) const
+		{
+			ensure(Column.GetId() == FNetStatsCountersViewColumns::AverageCountColumnID);
+			const FNetStatsCounterNode& NetStatsCounterNode = static_cast<const FNetStatsCounterNode&>(Node);
+			return TOptional<FTableCellValue>(FTableCellValue(static_cast<int64>(NetStatsCounterNode.GetAggregatedStats().Average)));
+		}
+	};
+
+	TSharedRef<ITableCellValueGetter> Getter = MakeShared<FAverageCountValueGetter>();
+	Column.SetValueGetter(Getter);
+
+	class FAverageCountFormatter : public FTableCellValueFormatter
+	{
+	public:
+		virtual FText FormatValue(const FTableColumn& Column, const FBaseTreeNode& Node) const override
+		{
+			ensure(Column.GetId() == FNetStatsCountersViewColumns::AverageCountColumnID);
+			const FNetStatsCounterNode& NetStatsCounterNode = static_cast<const FNetStatsCounterNode&>(Node);
+			return NetStatsCounterNode.GetTextForAggregatedStatsAverage();
+		}
+	};
+	TSharedRef<ITableCellValueFormatter> Formatter = MakeShared<FAverageCountFormatter>();
+	Column.SetValueFormatter(Formatter);
+
+	TSharedRef<ITableCellValueSorter> Sorter = MakeShared<FSorterByInt64Value>(ColumnRef);
+	Column.SetValueSorter(Sorter);
+
+	return ColumnRef;
+}
+
+////////////////////////////////////////////////////////////////////////////////////////////////////
+
+#undef LOCTEXT_NAMESPACE
diff --git a/Engine/Source/Developer/TraceInsights/Private/Insights/NetworkingProfiler/ViewModels/NetStatsCountersViewColumnFactory.h b/Engine/Source/Developer/TraceInsights/Private/Insights/NetworkingProfiler/ViewModels/NetStatsCountersViewColumnFactory.h
new file mode 100644
--- /dev/null
+++ b/Engine/Source/Developer/TraceInsights/Private/Insights/NetworkingProfiler/ViewModels/NetStatsCountersViewColumnFactory.h
@@ -0,0 +1,50 @@
+// Copyright Epic Games, Inc. All Rights Reserved.
+
+#pragma once
+
+#include "CoreMinimal.h"
+
+////////////////////////////////////////////////////////////////////////////////////////////////////
+
+struct FNetStatsCountersViewColumns
+{
+	//////////////////////////////////////////////////
+	// Column identifiers
+
+	static const FName NameColumnID;
+	static const FName TypeColumnID;
+	static const FName InstanceCountColumnID;
+
+	static const FName SumColumnID;
+	static const FName MaxCountColumnID;
+	static const FName AverageCountColumnID;
+};
+
+////////////////////////////////////////////////////////////////////////////////////////////////////
+
+namespace Insights
+{
+	class FTableColumn;
+}
+
+////////////////////////////////////////////////////////////////////////////////////////////////////
+
+struct FNetStatsCountersViewColumnFactory
+{
+public:
+	static void CreateNetStatsCountersViewColumns(TArray<TSharedRef<Insights::FTableColumn>>& Columns);
+
+	static TSharedRef<Insights::FTableColumn> CreateNameColumn();
+	static TSharedRef<Insights::FTableColumn> CreateTypeColumn();
+	static TSharedRef<Insights::FTableColumn> CreateInstanceCountColumn();
+
+	static TSharedRef<Insights::FTableColumn> CreateSumColumn();
+	static TSharedRef<Insights::FTableColumn> CreateMaxCountColumn();
+	static TSharedRef<Insights::FTableColumn> CreateAverageCountColumn();
+
+private:
+	static constexpr float SumColumnInitialWidth = 60.0f;
+	static constexpr float CountColumnInitialWidth = 50.0f;
+};
+
+////////////////////////////////////////////////////////////////////////////////////////////////////
diff --git a/Engine/Source/Developer/TraceInsights/Private/Insights/NetworkingProfiler/ViewModels/NetStatsViewColumnFactory.cpp b/Engine/Source/Developer/TraceInsights/Private/Insights/NetworkingProfiler/ViewModels/NetStatsViewColumnFactory.cpp
--- a/Engine/Source/Developer/TraceInsights/Private/Insights/NetworkingProfiler/ViewModels/NetStatsViewColumnFactory.cpp
+++ b/Engine/Source/Developer/TraceInsights/Private/Insights/NetworkingProfiler/ViewModels/NetStatsViewColumnFactory.cpp
@@ -144,7 +144,7 @@ TSharedRef<Insights::FTableColumn> FNetStatsViewColumnFactory::CreateLevelColumn
 	Column.SetDescription(LOCTEXT("Level_ColumnDesc", "Level of net event instances"));
 
 	Column.SetFlags(ETableColumnFlags::CanBeHidden |
-					ETableColumnFlags::ShouldBeVisible |
+					//ETableColumnFlags::ShouldBeVisible |
 					ETableColumnFlags::CanBeFiltered);
 
 	Column.SetHorizontalAlignment(HAlign_Right);
@@ -286,7 +286,7 @@ TSharedRef<Insights::FTableColumn> FNetStatsViewColumnFactory::CreateMaxInclusiv
 	Column.SetDescription(LOCTEXT("MaxInclusive_ColumnDesc", "Maximum inclusive size (bits) of selected net event instances"));
 
 	Column.SetFlags(ETableColumnFlags::CanBeHidden |
-					//ETableColumnFlags::ShouldBeVisible |
+					ETableColumnFlags::ShouldBeVisible |
 					ETableColumnFlags::CanBeFiltered);
 
 	Column.SetHorizontalAlignment(HAlign_Right);
@@ -332,7 +332,7 @@ TSharedRef<Insights::FTableColumn> FNetStatsViewColumnFactory::CreateAverageIncl
 	Column.SetDescription(LOCTEXT("AvgInclusive_ColumnDesc", "Average inclusive size (bits) of selected net event instances"));
 
 	Column.SetFlags(ETableColumnFlags::CanBeHidden |
-					//ETableColumnFlags::ShouldBeVisible |
+					ETableColumnFlags::ShouldBeVisible |
 					ETableColumnFlags::CanBeFiltered);
 
 	Column.SetHorizontalAlignment(HAlign_Right);
@@ -428,7 +428,7 @@ TSharedRef<Insights::FTableColumn> FNetStatsViewColumnFactory::CreateMaxExclusiv
 	Column.SetDescription(LOCTEXT("MaxExclusive_ColumnDesc", "Maximum exclusive size (bits) of selected net event instances"));
 
 	Column.SetFlags(ETableColumnFlags::CanBeHidden |
-					//ETableColumnFlags::ShouldBeVisible |
+					ETableColumnFlags::ShouldBeVisible |
 					ETableColumnFlags::CanBeFiltered);
 
 	Column.SetHorizontalAlignment(HAlign_Right);
diff --git a/Engine/Source/Developer/TraceInsights/Private/Insights/NetworkingProfiler/ViewModels/PacketViewDrawHelper.cpp b/Engine/Source/Developer/TraceInsights/Private/Insights/NetworkingProfiler/ViewModels/PacketViewDrawHelper.cpp
--- a/Engine/Source/Developer/TraceInsights/Private/Insights/NetworkingProfiler/ViewModels/PacketViewDrawHelper.cpp
+++ b/Engine/Source/Developer/TraceInsights/Private/Insights/NetworkingProfiler/ViewModels/PacketViewDrawHelper.cpp
@@ -10,6 +10,7 @@
 #include "Insights/Common/PaintUtils.h"
 #include "Insights/InsightsStyle.h"
 #include "Insights/NetworkingProfiler/ViewModels/PacketViewport.h"
+#include "Insights/NetworkingProfiler/ViewModels/PacketContentViewDrawHelper.h"
 #include "Insights/ViewModels/DrawHelpers.h"
 
 #include <limits>
@@ -100,6 +101,11 @@ FNetworkPacketAggregatedSample* FNetworkPacketSeriesBuilder::AddPacket(const int
 	return nullptr;
 }
 
+void FNetworkPacketSeriesBuilder::SetHighlightEventTypeIndex(int32 EventTypeIndex)
+{
+	Series.HighlightEventTypeIndex = EventTypeIndex;
+}
+
 ////////////////////////////////////////////////////////////////////////////////////////////////////
 // FPacketViewDrawHelper
 ////////////////////////////////////////////////////////////////////////////////////////////////////
@@ -174,6 +180,8 @@ void FPacketViewDrawHelper::DrawCached(const FNetworkPacketSeries& Series) const
 	const float ViewHeight = FMath::RoundToFloat(Viewport.GetHeight());
 	const float BaselineY = FMath::RoundToFloat(ViewportY.GetOffsetForValue(0.0));
 
+	const FLinearColor ColorFilterMatch = FPacketContentViewDrawHelper::GetColorByType(Series.HighlightEventTypeIndex);
+
 	for (int32 SampleIndex = 0; SampleIndex < NumSamples; SampleIndex++)
 	{
 		const FNetworkPacketAggregatedSample& Sample = Series.Samples[SampleIndex];
@@ -189,6 +197,10 @@ void FPacketViewDrawHelper::DrawCached(const FNetworkPacketSeries& Series) const
 		//const float ValueY = FMath::RoundToFloat(ViewportY.GetOffsetForValue(static_cast<double>(Sample.LargestPacket.ContentSizeInBits)));
 		const float ValueY = FMath::RoundToFloat(ViewportY.GetOffsetForValue(static_cast<double>(Sample.LargestPacket.TotalSizeInBytes * 8)));
 
+		const float FilterMatchContentValueY = FMath::RoundToFloat(ViewportY.GetOffsetForValue(static_cast<double>(Sample.FilterMatchHighlightSizeInBits)));
+		const float FilterMatchContentH = FilterMatchContentValueY - BaselineY;
+		const float FilterMatchContentY = ViewHeight - FilterMatchContentH;
+
 		const float H = ValueY - BaselineY;
 		const float Y = ViewHeight - H;
 
@@ -203,6 +215,11 @@ void FPacketViewDrawHelper::DrawCached(const FNetworkPacketSeries& Series) const
 		{
 			DrawContext.DrawBox(X + 1.0f, Y + 1.0f, SampleW - 2.0f, H - 2.0f, WhiteBrush, ColorFill);
 
+			if (Sample.FilterMatchHighlightSizeInBits > 0U)
+			{
+				DrawContext.DrawBox(X + 1.0f, FilterMatchContentY + 1.0f, SampleW - 2.0f, FilterMatchContentH - 2.0f, WhiteBrush, ColorFilterMatch);
+			}
+
 			// Draw border.
 			DrawContext.DrawBox(X, Y, 1.0, H, WhiteBrush, ColorBorder);
 			DrawContext.DrawBox(X + SampleW - 1.0f, Y, 1.0, H, WhiteBrush, ColorBorder);
@@ -212,6 +229,12 @@ void FPacketViewDrawHelper::DrawCached(const FNetworkPacketSeries& Series) const
 		else
 		{
 			DrawContext.DrawBox(X, Y, SampleW, H, WhiteBrush, ColorBorder);
+
+			if (Sample.FilterMatchHighlightSizeInBits > 0U)
+			{
+				DrawContext.DrawBox(X, FilterMatchContentY, SampleW, FilterMatchContentH, WhiteBrush, ColorFilterMatch);
+			}
+
 		}
 	}
 
diff --git a/Engine/Source/Developer/TraceInsights/Private/Insights/NetworkingProfiler/ViewModels/PacketViewDrawHelper.h b/Engine/Source/Developer/TraceInsights/Private/Insights/NetworkingProfiler/ViewModels/PacketViewDrawHelper.h
--- a/Engine/Source/Developer/TraceInsights/Private/Insights/NetworkingProfiler/ViewModels/PacketViewDrawHelper.h
+++ b/Engine/Source/Developer/TraceInsights/Private/Insights/NetworkingProfiler/ViewModels/PacketViewDrawHelper.h
@@ -83,6 +83,8 @@ struct FNetworkPacketAggregatedSample
 
 	bool bAtLeastOnePacketMatchesFilter;
 
+	uint32 FilterMatchHighlightSizeInBits;
+
 	FNetworkPacketAggregatedSample()
 		: NumPackets(0)
 		, StartTime(DBL_MAX)
@@ -90,6 +92,7 @@ struct FNetworkPacketAggregatedSample
 		, AggregatedStatus(TraceServices::ENetProfilerDeliveryStatus::Unknown)
 		, LargestPacket()
 		, bAtLeastOnePacketMatchesFilter(true)
+		, FilterMatchHighlightSizeInBits(0)
 	{}
 
 	FNetworkPacketAggregatedSample(const FNetworkPacketAggregatedSample&) = default;
@@ -123,9 +126,12 @@ struct FNetworkPacketSeries
 
 	TArray<FNetworkPacketAggregatedSample> Samples;
 
+	int32 HighlightEventTypeIndex;
+
 	FNetworkPacketSeries()
 		: NumAggregatedPackets(0)
 		, Samples()
+		, HighlightEventTypeIndex(-1)
 	{
 	}
 
@@ -133,6 +139,7 @@ struct FNetworkPacketSeries
 	{
 		NumAggregatedPackets = 0;
 		Samples.Reset();
+		HighlightEventTypeIndex = -1;
 	}
 };
 
@@ -153,6 +160,8 @@ public:
 
 	int32 GetNumAddedPackets() const { return NumAddedPackets; }
 
+	void SetHighlightEventTypeIndex(int32 EventTypeIndex);
+
 private:
 	FNetworkPacketSeries& Series; // series to update
 	const FPacketViewport& Viewport;
diff --git a/Engine/Source/Developer/TraceInsights/Private/Insights/NetworkingProfiler/Widgets/SNetStatsCountersTableCell.cpp b/Engine/Source/Developer/TraceInsights/Private/Insights/NetworkingProfiler/Widgets/SNetStatsCountersTableCell.cpp
new file mode 100644
--- /dev/null
+++ b/Engine/Source/Developer/TraceInsights/Private/Insights/NetworkingProfiler/Widgets/SNetStatsCountersTableCell.cpp
@@ -0,0 +1,159 @@
+// Copyright Epic Games, Inc. All Rights Reserved.
+
+#include "SNetStatsCountersTableCell.h"
+
+#include "SlateOptMacros.h"
+#include "Widgets/Images/SImage.h"
+#include "Widgets/Input/SButton.h"
+#include "Widgets/SBoxPanel.h"
+#include "Widgets/SOverlay.h"
+#include "Widgets/SToolTip.h"
+#include "Widgets/Text/STextBlock.h"
+#include "Widgets/Views/SExpanderArrow.h"
+
+// Insights
+#include "Insights/InsightsStyle.h"
+#include "Insights/Table/ViewModels/Table.h"
+#include "Insights/Table/ViewModels/TableColumn.h"
+#include "Insights/NetworkingProfiler/Widgets/SNetStatsCountersTableRow.h"
+
+#define LOCTEXT_NAMESPACE "SNetStatsCountersView"
+
+////////////////////////////////////////////////////////////////////////////////////////////////////
+
+BEGIN_SLATE_FUNCTION_BUILD_OPTIMIZATION
+
+void SNetStatsCountersTableCell::Construct(const FArguments& InArgs, const TSharedRef<ITableRow>& TableRow)
+{
+	TablePtr = InArgs._TablePtr;
+	ColumnPtr = InArgs._ColumnPtr;
+	NetStatsCounterNodePtr = InArgs._NetStatsCounterNodePtr;
+
+	ensure(TablePtr.IsValid());
+	ensure(ColumnPtr.IsValid());
+	ensure(NetStatsCounterNodePtr.IsValid());
+
+	SetHoveredCellDelegate = InArgs._OnSetHoveredCell;
+
+	ChildSlot
+	[
+		GenerateWidgetForColumn(InArgs, TableRow)
+	];
+}
+
+////////////////////////////////////////////////////////////////////////////////////////////////////
+
+TSharedRef<SWidget> SNetStatsCountersTableCell::GenerateWidgetForColumn(const FArguments& InArgs, const TSharedRef<ITableRow>& TableRow)
+{
+	if (InArgs._IsNameColumn)
+	{
+		return GenerateWidgetForNameColumn(InArgs, TableRow);
+	}
+	else
+	{
+		return GenerateWidgetForStatsColumn(InArgs, TableRow);
+	}
+}
+
+////////////////////////////////////////////////////////////////////////////////////////////////////
+
+TSharedRef<SWidget> SNetStatsCountersTableCell::GenerateWidgetForNameColumn(const FArguments& InArgs, const TSharedRef<ITableRow>& TableRow)
+{
+	return
+		SNew(SHorizontalBox)
+
+		+ SHorizontalBox::Slot()
+		.AutoWidth()
+		.HAlign(HAlign_Right)
+		.VAlign(VAlign_Center)
+		[
+			SNew(SExpanderArrow, TableRow)
+		]
+
+		// Info icon + tooltip
+		+ SHorizontalBox::Slot()
+		.AutoWidth()
+		.HAlign(HAlign_Center)
+		.VAlign(VAlign_Center)
+		[
+			SNew(SImage)
+			.Visibility(this, &SNetStatsCountersTableCell::GetHintIconVisibility)
+			.Image(FInsightsStyle::GetBrush("Icons.Hint.TreeItem"))
+			.ToolTip(GetRowToolTip(TableRow))
+		]
+
+		// Name
+		+ SHorizontalBox::Slot()
+		.AutoWidth()
+		.VAlign(VAlign_Center)
+		.HAlign(HAlign_Left)
+		.Padding(FMargin(2.0f, 0.0f))
+		[
+			SNew(STextBlock)
+			.Text(this, &SNetStatsCountersTableCell::GetDisplayName)
+			.HighlightText(InArgs._HighlightText)
+			.TextStyle(FInsightsStyle::Get(), TEXT("TreeTable.Tooltip"))
+			.ColorAndOpacity(this, &SNetStatsCountersTableCell::GetColorAndOpacity)
+			.ShadowColorAndOpacity(this, &SNetStatsCountersTableCell::GetShadowColorAndOpacity)
+		]
+
+		// Name Suffix
+		+ SHorizontalBox::Slot()
+		.AutoWidth()
+		.VAlign(VAlign_Center)
+		.HAlign(HAlign_Left)
+		.Padding(FMargin(2.0f, 0.0f))
+		[
+			SNew(STextBlock)
+			.Visibility(this, &SNetStatsCountersTableCell::HasExtraDisplayName)
+			.Text(this, &SNetStatsCountersTableCell::GetExtraDisplayName)
+			.TextStyle(FInsightsStyle::Get(), TEXT("TreeTable.Tooltip"))
+			.ColorAndOpacity(this, &SNetStatsCountersTableCell::GetExtraColorAndOpacity)
+			.ShadowColorAndOpacity(this, &SNetStatsCountersTableCell::GetShadowColorAndOpacity)
+		]
+	;
+}
+
+////////////////////////////////////////////////////////////////////////////////////////////////////
+
+TSharedPtr<IToolTip> SNetStatsCountersTableCell::GetRowToolTip(const TSharedRef<ITableRow>& TableRow) const
+{
+	TSharedRef<SNetStatsCountersTableRow> Row = StaticCastSharedRef<SNetStatsCountersTableRow, ITableRow>(TableRow);
+	return Row->GetRowToolTip();
+}
+
+////////////////////////////////////////////////////////////////////////////////////////////////////
+
+FText SNetStatsCountersTableCell::GetValueAsText() const
+{
+	return NetStatsCounterNodePtr->IsGroup() ? FText::GetEmpty() : ColumnPtr->GetValueAsText(*NetStatsCounterNodePtr);
+}
+
+////////////////////////////////////////////////////////////////////////////////////////////////////
+
+TSharedRef<SWidget> SNetStatsCountersTableCell::GenerateWidgetForStatsColumn(const FArguments& InArgs, const TSharedRef<ITableRow>& TableRow)
+{
+	return
+		SNew(SHorizontalBox)
+
+		// Value
+		+ SHorizontalBox::Slot()
+		.FillWidth(1.0f)
+		.VAlign(VAlign_Center)
+		.HAlign(ColumnPtr->GetHorizontalAlignment())
+		.Padding(FMargin(2.0f, 0.0f))
+		[
+			SNew(STextBlock)
+			.Text(this, &SNetStatsCountersTableCell::GetValueAsText)
+			.TextStyle(FInsightsStyle::Get(), TEXT("TreeTable.Tooltip"))
+			.ColorAndOpacity(this, &SNetStatsCountersTableCell::GetStatsColorAndOpacity)
+			.ShadowColorAndOpacity(this, &SNetStatsCountersTableCell::GetShadowColorAndOpacity)
+		]
+	;
+}
+
+END_SLATE_FUNCTION_BUILD_OPTIMIZATION
+
+////////////////////////////////////////////////////////////////////////////////////////////////////
+
+#undef LOCTEXT_NAMESPACE
diff --git a/Engine/Source/Developer/TraceInsights/Private/Insights/NetworkingProfiler/Widgets/SNetStatsCountersTableCell.h b/Engine/Source/Developer/TraceInsights/Private/Insights/NetworkingProfiler/Widgets/SNetStatsCountersTableCell.h
new file mode 100644
--- /dev/null
+++ b/Engine/Source/Developer/TraceInsights/Private/Insights/NetworkingProfiler/Widgets/SNetStatsCountersTableCell.h
@@ -0,0 +1,171 @@
+// Copyright Epic Games, Inc. All Rights Reserved.
+
+#pragma once
+
+#include "CoreMinimal.h"
+#include "Widgets/DeclarativeSyntaxSupport.h"
+#include "Widgets/SCompoundWidget.h"
+
+// Insights
+#include "Insights/NetworkingProfiler/ViewModels/NetStatsCounterNode.h"
+
+class ITableRow;
+
+namespace Insights
+{
+	class FTable;
+	class FTableColumn;
+}
+
+DECLARE_DELEGATE_ThreeParams(FSetHoveredNetStatsCountersTableCell, TSharedPtr<Insights::FTable> /*TablePtr*/, TSharedPtr<Insights::FTableColumn> /*ColumnPtr*/, FNetStatsCounterNodePtr /*NetStatsCounterNodePtr*/);
+
+class SNetStatsCountersTableCell : public SCompoundWidget
+{
+public:
+	SLATE_BEGIN_ARGS(SNetStatsCountersTableCell) {}
+		SLATE_EVENT(FSetHoveredNetStatsCountersTableCell, OnSetHoveredCell)
+		SLATE_ATTRIBUTE(FText, HighlightText)
+		SLATE_ARGUMENT(TSharedPtr<Insights::FTable>, TablePtr)
+		SLATE_ARGUMENT(TSharedPtr<Insights::FTableColumn>, ColumnPtr)
+		SLATE_ARGUMENT(FNetStatsCounterNodePtr, NetStatsCounterNodePtr)
+		SLATE_ARGUMENT(bool, IsNameColumn)
+	SLATE_END_ARGS()
+
+	void Construct(const FArguments& InArgs, const TSharedRef<ITableRow>& TableRow);
+
+protected:
+	TSharedRef<SWidget> GenerateWidgetForColumn(const FArguments& InArgs, const TSharedRef<ITableRow>& TableRow);
+	TSharedRef<SWidget> GenerateWidgetForNameColumn(const FArguments& InArgs, const TSharedRef<ITableRow>& TableRow);
+	TSharedRef<SWidget> GenerateWidgetForStatsColumn(const FArguments& InArgs, const TSharedRef<ITableRow>& TableRow);
+
+	/**
+	 * The system will use this event to notify a widget that the cursor has entered it. This event is NOT bubbled.
+	 *
+	 * @param MyGeometry The Geometry of the widget receiving the event
+	 * @param MouseEvent Information about the input event
+	 */
+	virtual void OnMouseEnter(const FGeometry& MyGeometry, const FPointerEvent& MouseEvent) override
+	{
+		SCompoundWidget::OnMouseEnter(MyGeometry, MouseEvent);
+		SetHoveredCellDelegate.ExecuteIfBound(TablePtr, ColumnPtr, NetStatsCounterNodePtr);
+	}
+
+	/**
+	 * The system will use this event to notify a widget that the cursor has left it. This event is NOT bubbled.
+	 *
+	 * @param MouseEvent Information about the input event
+	 */
+	virtual void OnMouseLeave(const FPointerEvent& MouseEvent) override
+	{
+		SCompoundWidget::OnMouseLeave(MouseEvent);
+		SetHoveredCellDelegate.ExecuteIfBound(nullptr, nullptr, nullptr);
+	}
+
+	/**
+	 * Called during drag and drop when the drag enters a widget.
+	 *
+	 * Enter/Leave events in slate are meant as lightweight notifications.
+	 * So we do not want to capture mouse or set focus in response to these.
+	 * However, OnDragEnter must also support external APIs (e.g. OLE Drag/Drop)
+	 * Those require that we let them know whether we can handle the content
+	 * being dragged OnDragEnter.
+	 *
+	 * The concession is to return a can_handled/cannot_handle
+	 * boolean rather than a full FReply.
+	 *
+	 * @param MyGeometry      The geometry of the widget receiving the event.
+	 * @param DragDropEvent   The drag and drop event.
+	 *
+	 * @return A reply that indicated whether the contents of the DragDropEvent can potentially be processed by this widget.
+	 */
+	virtual void OnDragEnter(const FGeometry& MyGeometry, const FDragDropEvent& DragDropEvent) override
+	{
+		SCompoundWidget::OnDragEnter(MyGeometry, DragDropEvent);
+		SetHoveredCellDelegate.ExecuteIfBound(TablePtr, ColumnPtr, NetStatsCounterNodePtr);
+	}
+
+	/**
+	 * Called during drag and drop when the drag leaves a widget.
+	 *
+	 * @param DragDropEvent   The drag and drop event.
+	 */
+	virtual void OnDragLeave(const FDragDropEvent& DragDropEvent)  override
+	{
+		SCompoundWidget::OnDragLeave(DragDropEvent);
+		SetHoveredCellDelegate.ExecuteIfBound(nullptr, nullptr, nullptr);
+	}
+
+	EVisibility GetHintIconVisibility() const
+	{
+		return IsHovered() ? EVisibility::Visible : EVisibility::Hidden;
+	}
+
+	TSharedPtr<class IToolTip> GetRowToolTip(const TSharedRef<ITableRow>& TableRow) const;
+
+	FText GetDisplayName() const
+	{
+		return NetStatsCounterNodePtr->GetDisplayName();
+	}
+
+	FText GetExtraDisplayName() const
+	{
+		return NetStatsCounterNodePtr->GetExtraDisplayName();
+	}
+
+	EVisibility HasExtraDisplayName() const
+	{
+		return NetStatsCounterNodePtr->HasExtraDisplayName() ? EVisibility::Visible : EVisibility::Collapsed;
+	}
+
+	FText GetValueAsText() const;
+
+	FSlateColor GetColorAndOpacity() const
+	{
+		const FLinearColor TextColor =
+			NetStatsCounterNodePtr->IsFiltered() ?
+				FLinearColor(1.0f, 1.0f, 1.0f, 0.5f) :
+				FLinearColor(1.0f, 1.0f, 1.0f, 1.0f);
+		return TextColor;
+	}
+
+	FSlateColor GetExtraColorAndOpacity() const
+	{
+		const FLinearColor TextColor =
+			NetStatsCounterNodePtr->IsFiltered() ?
+				FLinearColor(0.5f, 0.5f, 0.5f, 0.5f) :
+				FLinearColor(0.5f, 0.5f, 0.5f, 1.0f);
+		return TextColor;
+	}
+
+	FSlateColor GetStatsColorAndOpacity() const
+	{
+		const FLinearColor TextColor =
+			NetStatsCounterNodePtr->IsFiltered() ?
+				FLinearColor(1.0f, 1.0f, 1.0f, 0.5f) :
+			NetStatsCounterNodePtr->GetAggregatedStats().Count == 0 ?
+				FLinearColor(1.0f, 1.0f, 1.0f, 0.6f) :
+				FLinearColor(1.0f, 1.0f, 1.0f, 1.0f);
+		return TextColor;
+	}
+
+	FLinearColor GetShadowColorAndOpacity() const
+	{
+		const FLinearColor ShadowColor =
+			NetStatsCounterNodePtr->IsFiltered() ?
+				FLinearColor(0.f, 0.f, 0.f, 0.25f) :
+				FLinearColor(0.0f, 0.0f, 0.0f, 0.5f);
+		return ShadowColor;
+	}
+
+protected:
+	/** A shared pointer to the table view model. */
+	TSharedPtr<Insights::FTable> TablePtr; // TODO: TSharedRef
+
+	/** A shared pointer to the table column view model. */
+	TSharedPtr<Insights::FTableColumn> ColumnPtr; // TODO: TSharedRef
+
+	/** A shared pointer to the timer node. */
+	FNetStatsCounterNodePtr NetStatsCounterNodePtr; // TODO: TSharedRef
+
+	FSetHoveredNetStatsCountersTableCell SetHoveredCellDelegate;
+};
diff --git a/Engine/Source/Developer/TraceInsights/Private/Insights/NetworkingProfiler/Widgets/SNetStatsCountersTableRow.cpp b/Engine/Source/Developer/TraceInsights/Private/Insights/NetworkingProfiler/Widgets/SNetStatsCountersTableRow.cpp
new file mode 100644
--- /dev/null
+++ b/Engine/Source/Developer/TraceInsights/Private/Insights/NetworkingProfiler/Widgets/SNetStatsCountersTableRow.cpp
@@ -0,0 +1,191 @@
+// Copyright Epic Games, Inc. All Rights Reserved.
+
+#include "SNetStatsCountersTableRow.h"
+
+#include "SlateOptMacros.h"
+#include "Widgets/SOverlay.h"
+#include "Widgets/Images/SImage.h"
+
+// Insights
+#include "Insights/InsightsStyle.h"
+#include "Insights/Common/TimeUtils.h"
+#include "Insights/Table/ViewModels/Table.h"
+#include "Insights/Table/ViewModels/TableColumn.h"
+#include "Insights/NetworkingProfiler/Widgets/SNetStatsCountersTableCell.h"
+#include "Insights/NetworkingProfiler/Widgets/SNetStatsCountersViewTooltip.h"
+
+#define LOCTEXT_NAMESPACE "SNetStatsView"
+
+////////////////////////////////////////////////////////////////////////////////////////////////////
+
+BEGIN_SLATE_FUNCTION_BUILD_OPTIMIZATION
+
+void SNetStatsCountersTableRow::Construct(const FArguments& InArgs, const TSharedRef<STableViewBase>& InOwnerTableView)
+{
+	OnShouldBeEnabled = InArgs._OnShouldBeEnabled;
+	IsColumnVisibleDelegate = InArgs._OnIsColumnVisible;
+	GetColumnOutlineHAlignmentDelegate = InArgs._OnGetColumnOutlineHAlignmentDelegate;
+	SetHoveredCellDelegate = InArgs._OnSetHoveredCell;
+
+	HighlightText = InArgs._HighlightText;
+	HighlightedNodeName = InArgs._HighlightedNodeName;
+
+	TablePtr = InArgs._TablePtr;
+	NetStatsCounterNodePtr = InArgs._NetStatsCounterNodePtr;
+
+	RowToolTip = MakeShared<SNetStatsCounterTableRowToolTip>(NetStatsCounterNodePtr);
+
+	SetEnabled(TAttribute<bool>(this, &SNetStatsCountersTableRow::HandleShouldBeEnabled));
+
+	SMultiColumnTableRow<FNetStatsCounterNodePtr>::Construct(SMultiColumnTableRow<FNetStatsCounterNodePtr>::FArguments(), InOwnerTableView);
+}
+
+////////////////////////////////////////////////////////////////////////////////////////////////////
+
+TSharedRef<SWidget> SNetStatsCountersTableRow::GenerateWidgetForColumn(const FName& ColumnId)
+{
+	TSharedPtr<Insights::FTableColumn> ColumnPtr = TablePtr->FindColumnChecked(ColumnId);
+
+	return
+		SNew(SOverlay)
+		.Visibility(EVisibility::SelfHitTestInvisible)
+
+		+SOverlay::Slot()
+		.Padding(0.0f)
+		[
+			SNew(SImage)
+			.Image(FInsightsStyle::GetBrush("TreeTable.RowBackground"))
+			.ColorAndOpacity(this, &SNetStatsCountersTableRow::GetBackgroundColorAndOpacity)
+		]
+
+		+SOverlay::Slot()
+		.Padding(0.0f)
+		[
+			SNew(SImage)
+			.Image(this, &SNetStatsCountersTableRow::GetOutlineBrush, ColumnId)
+			.ColorAndOpacity(this, &SNetStatsCountersTableRow::GetOutlineColorAndOpacity)
+		]
+
+		+SOverlay::Slot()
+		[
+			SNew(SNetStatsCountersTableCell, SharedThis(this))
+			.Visibility(this, &SNetStatsCountersTableRow::IsColumnVisible, ColumnId)
+			.TablePtr(TablePtr)
+			.ColumnPtr(ColumnPtr)
+			.NetStatsCounterNodePtr(NetStatsCounterNodePtr)
+			.HighlightText(HighlightText)
+			.IsNameColumn(ColumnPtr->IsHierarchy())
+			.OnSetHoveredCell(this, &SNetStatsCountersTableRow::OnSetHoveredCell)
+		];
+}
+
+END_SLATE_FUNCTION_BUILD_OPTIMIZATION
+
+////////////////////////////////////////////////////////////////////////////////////////////////////
+
+FReply SNetStatsCountersTableRow::OnDragDetected(const FGeometry& MyGeometry, const FPointerEvent& MouseEvent)
+{
+	return SMultiColumnTableRow<FNetStatsCounterNodePtr>::OnDragDetected(MyGeometry, MouseEvent);
+}
+
+////////////////////////////////////////////////////////////////////////////////////////////////////
+
+TSharedRef<IToolTip> SNetStatsCountersTableRow::GetRowToolTip() const
+{
+	return RowToolTip.ToSharedRef();
+}
+
+////////////////////////////////////////////////////////////////////////////////////////////////////
+
+void SNetStatsCountersTableRow::InvalidateContent()
+{
+	RowToolTip->InvalidateWidget();
+}
+
+////////////////////////////////////////////////////////////////////////////////////////////////////
+
+FSlateColor SNetStatsCountersTableRow::GetBackgroundColorAndOpacity() const
+{
+	return GetBackgroundColorAndOpacity(NetStatsCounterNodePtr->GetAggregatedStats().Count);
+}
+
+////////////////////////////////////////////////////////////////////////////////////////////////////
+
+FSlateColor SNetStatsCountersTableRow::GetBackgroundColorAndOpacity(uint32 Size) const
+{
+	const FLinearColor Color =	Size > 1000 ? FLinearColor(0.3f, 0.0f, 0.0f, 1.0f) :
+								Size > 100  ? FLinearColor(0.3f, 0.1f, 0.0f, 1.0f) :
+								Size > 10   ? FLinearColor(0.0f, 0.1f, 0.0f, 1.0f) :
+								              FLinearColor(0.0f, 0.0f, 0.0f, 1.0f);
+	return Color;
+}
+
+////////////////////////////////////////////////////////////////////////////////////////////////////
+
+FSlateColor SNetStatsCountersTableRow::GetOutlineColorAndOpacity() const
+{
+	const FLinearColor NoColor(0.0f, 0.0f, 0.0f, 0.0f);
+	const bool bShouldBeHighlighted = NetStatsCounterNodePtr->GetName() == HighlightedNodeName.Get();
+	const FLinearColor OutlineColorAndOpacity = bShouldBeHighlighted ? FLinearColor(FColorList::SlateBlue) : NoColor;
+	return OutlineColorAndOpacity;
+}
+
+////////////////////////////////////////////////////////////////////////////////////////////////////
+
+const FSlateBrush* SNetStatsCountersTableRow::GetOutlineBrush(const FName ColumnId) const
+{
+	EHorizontalAlignment HAlign = HAlign_Center;
+	if (IsColumnVisibleDelegate.IsBound())
+	{
+		HAlign = GetColumnOutlineHAlignmentDelegate.Execute(ColumnId);
+	}
+
+	return FInsightsStyle::GetOutlineBrush(HAlign);
+}
+
+////////////////////////////////////////////////////////////////////////////////////////////////////
+
+bool SNetStatsCountersTableRow::HandleShouldBeEnabled() const
+{
+	bool bResult = false;
+
+	if (NetStatsCounterNodePtr->IsGroup())
+	{
+		bResult = true;
+	}
+	else
+	{
+		if (OnShouldBeEnabled.IsBound())
+		{
+			bResult = OnShouldBeEnabled.Execute(NetStatsCounterNodePtr);
+		}
+		return true;
+	}
+
+	return bResult;
+}
+
+////////////////////////////////////////////////////////////////////////////////////////////////////
+
+EVisibility SNetStatsCountersTableRow::IsColumnVisible(const FName ColumnId) const
+{
+	EVisibility Result = EVisibility::Collapsed;
+
+	if (IsColumnVisibleDelegate.IsBound())
+	{
+		Result = IsColumnVisibleDelegate.Execute(ColumnId) ? EVisibility::Visible : EVisibility::Collapsed;
+	}
+
+	return Result;
+}
+
+////////////////////////////////////////////////////////////////////////////////////////////////////
+
+void SNetStatsCountersTableRow::OnSetHoveredCell(TSharedPtr<Insights::FTable> InTablePtr, TSharedPtr<Insights::FTableColumn> InColumnPtr, FNetStatsCounterNodePtr InNetStatsCounterNodePtr)
+{
+	SetHoveredCellDelegate.ExecuteIfBound(InTablePtr, InColumnPtr, InNetStatsCounterNodePtr);
+}
+
+////////////////////////////////////////////////////////////////////////////////////////////////////
+
+#undef LOCTEXT_NAMESPACE
diff --git a/Engine/Source/Developer/TraceInsights/Private/Insights/NetworkingProfiler/Widgets/SNetStatsCountersTableRow.h b/Engine/Source/Developer/TraceInsights/Private/Insights/NetworkingProfiler/Widgets/SNetStatsCountersTableRow.h
new file mode 100644
--- /dev/null
+++ b/Engine/Source/Developer/TraceInsights/Private/Insights/NetworkingProfiler/Widgets/SNetStatsCountersTableRow.h
@@ -0,0 +1,96 @@
+// Copyright Epic Games, Inc. All Rights Reserved.
+
+#pragma once
+
+#include "CoreMinimal.h"
+#include "Widgets/DeclarativeSyntaxSupport.h"
+#include "Widgets/Views/SHeaderRow.h"
+#include "Widgets/Views/STableViewBase.h"
+#include "Widgets/Views/STableRow.h"
+#include "Widgets/Views/STreeView.h"
+
+// Insights
+#include "Insights/NetworkingProfiler/ViewModels/NetStatsCounterNodeHelper.h"
+
+class IToolTip;
+class SNetStatsCounterTableRowToolTip;
+
+namespace Insights
+{
+	class FTable;
+	class FTableColumn;
+}
+
+DECLARE_DELEGATE_RetVal_OneParam(bool, FNetStatsCounterNodeShouldBeEnabledDelegate, FNetStatsCounterNodePtr /*NodePtr*/);
+DECLARE_DELEGATE_RetVal_OneParam(bool, FIsColumnVisibleDelegate, const FName /*ColumnId*/);
+DECLARE_DELEGATE_RetVal_OneParam(EHorizontalAlignment, FGetColumnOutlineHAlignmentDelegate, const FName /*ColumnId*/);
+DECLARE_DELEGATE_ThreeParams(FSetHoveredNetStatsCounterTableCell, TSharedPtr<Insights::FTable> /*TablePtr*/, TSharedPtr<Insights::FTableColumn> /*ColumnPtr*/, FNetStatsCounterNodePtr /*NetStatsCounterNodePtr*/);
+
+/** Widget that represents a table row in the tree control. Generates widgets for each column on demand. */
+class SNetStatsCountersTableRow : public SMultiColumnTableRow<FNetStatsCounterNodePtr>
+{
+public:
+	SLATE_BEGIN_ARGS(SNetStatsCountersTableRow) {}
+		SLATE_EVENT(FNetStatsCounterNodeShouldBeEnabledDelegate, OnShouldBeEnabled)
+		SLATE_EVENT(FIsColumnVisibleDelegate, OnIsColumnVisible)
+		SLATE_EVENT(FGetColumnOutlineHAlignmentDelegate, OnGetColumnOutlineHAlignmentDelegate)
+		SLATE_EVENT(FSetHoveredNetStatsCounterTableCell, OnSetHoveredCell)
+		SLATE_ATTRIBUTE(FText, HighlightText)
+		SLATE_ATTRIBUTE(FName, HighlightedNodeName)
+		SLATE_ARGUMENT(TSharedPtr<Insights::FTable>, TablePtr)
+		SLATE_ARGUMENT(FNetStatsCounterNodePtr, NetStatsCounterNodePtr)
+	SLATE_END_ARGS()
+
+public:
+	void Construct(const FArguments& InArgs, const TSharedRef<STableViewBase>& InOwnerTableView);
+
+	virtual TSharedRef<SWidget> GenerateWidgetForColumn(const FName& ColumnId) override;
+
+	/**
+	 * Called when Slate detects that a widget started to be dragged.
+	 * Usage:
+	 * A widget can ask Slate to detect a drag.
+	 * OnMouseDown() reply with FReply::Handled().DetectDrag(SharedThis(this)).
+	 * Slate will either send an OnDragDetected() event or do nothing.
+	 * If the user releases a mouse button or leaves the widget before
+	 * a drag is triggered (maybe user started at the very edge) then no event will be
+	 * sent.
+	 *
+	 * @param  InMyGeometry  Widget geometry
+	 * @param  InMouseEvent  MouseMove that triggered the drag
+	 *
+	 */
+	virtual FReply OnDragDetected(const FGeometry& MyGeometry, const FPointerEvent& MouseEvent) override;
+
+	TSharedRef<IToolTip> GetRowToolTip() const;
+	void InvalidateContent();
+
+protected:
+	FSlateColor GetBackgroundColorAndOpacity() const;
+	FSlateColor GetBackgroundColorAndOpacity(uint32 Size) const;
+	FSlateColor GetOutlineColorAndOpacity() const;
+	const FSlateBrush* GetOutlineBrush(const FName ColumnId) const;
+	bool HandleShouldBeEnabled() const;
+	EVisibility IsColumnVisible(const FName ColumnId) const;
+	void OnSetHoveredCell(TSharedPtr<Insights::FTable> InTablePtr, TSharedPtr<Insights::FTableColumn> InColumnPtr, FNetStatsCounterNodePtr InNetStatsCounterNodePtr);
+
+protected:
+	/** A shared pointer to the table view model. */
+	TSharedPtr<Insights::FTable> TablePtr;
+
+	/** Data context for this table row. */
+	FNetStatsCounterNodePtr NetStatsCounterNodePtr;
+
+	FNetStatsCounterNodeShouldBeEnabledDelegate OnShouldBeEnabled;
+	FIsColumnVisibleDelegate IsColumnVisibleDelegate;
+	FSetHoveredNetStatsCounterTableCell SetHoveredCellDelegate;
+	FGetColumnOutlineHAlignmentDelegate GetColumnOutlineHAlignmentDelegate;
+
+	/** Text to be highlighted on timer name. */
+	TAttribute<FText> HighlightText;
+
+	/** Name of the timer node that should be drawn as highlighted. */
+	TAttribute<FName> HighlightedNodeName;
+
+	TSharedPtr<SNetStatsCounterTableRowToolTip> RowToolTip;
+};
diff --git a/Engine/Source/Developer/TraceInsights/Private/Insights/NetworkingProfiler/Widgets/SNetStatsCountersView.cpp b/Engine/Source/Developer/TraceInsights/Private/Insights/NetworkingProfiler/Widgets/SNetStatsCountersView.cpp
new file mode 100644
--- /dev/null
+++ b/Engine/Source/Developer/TraceInsights/Private/Insights/NetworkingProfiler/Widgets/SNetStatsCountersView.cpp
@@ -0,0 +1,1614 @@
+// Copyright Epic Games, Inc. All Rights Reserved.
+
+#include "SNetStatsCountersView.h"
+
+#include "Framework/MultiBox/MultiBoxBuilder.h"
+#include "SlateOptMacros.h"
+#include "Templates/UniquePtr.h"
+#include "TraceServices/AnalysisService.h"
+#include "TraceServices/Model/NetProfiler.h"
+#include "Widgets/Input/SCheckBox.h"
+#include "Widgets/Layout/SScrollBox.h"
+#include "Widgets/Input/SSearchBox.h"
+#include "Widgets/Images/SImage.h"
+#include "Widgets/Layout/SGridPanel.h"
+#include "Widgets/Layout/SSeparator.h"
+#include "Widgets/SToolTip.h"
+#include "Widgets/Views/STableViewBase.h"
+
+// Insights
+#include "Insights/InsightsStyle.h"
+#include "Insights/Table/ViewModels/Table.h"
+#include "Insights/Table/ViewModels/TableColumn.h"
+#include "Insights/NetworkingProfiler/NetworkingProfilerManager.h"
+#include "Insights/NetworkingProfiler/ViewModels/NetStatsCountersViewColumnFactory.h"
+#include "Insights/NetworkingProfiler/ViewModels/NetStatsCounterNodeHelper.h"
+#include "Insights/NetworkingProfiler/Widgets/SNetStatsCountersViewTooltip.h"
+#include "Insights/NetworkingProfiler/Widgets/SNetStatsCountersTableRow.h"
+#include "Insights/NetworkingProfiler/Widgets/SNetworkingProfilerWindow.h"
+#include "Insights/NetworkingProfiler/Widgets/SPacketContentView.h"
+
+#define LOCTEXT_NAMESPACE "SNetStatsCountersView"
+
+////////////////////////////////////////////////////////////////////////////////////////////////////
+
+SNetStatsCountersView::SNetStatsCountersView()
+	: ProfilerWindow()
+	, Table(MakeShared<Insights::FTable>())
+	, bExpansionSaved(false)
+	, bFilterOutZeroCountStatsCounters(true)
+	, GroupingMode(ENetStatsCounterGroupingMode::ByType)
+	, AvailableSorters()
+	, CurrentSorter(nullptr)
+	, ColumnBeingSorted(GetDefaultColumnBeingSorted())
+	, ColumnSortMode(GetDefaultColumnSortMode())
+	, NextTimestamp(0)
+	, ObjectsChangeCount(0)
+	, GameInstanceIndex(0)
+	, ConnectionIndex(0)
+	, ConnectionMode(TraceServices::ENetProfilerConnectionMode::Outgoing)
+	, StatsPacketStartIndex(0)
+	, StatsPacketEndIndex(0)
+{
+	FMemory::Memset(bStatsCounterTypeIsVisible, 1);
+}
+
+////////////////////////////////////////////////////////////////////////////////////////////////////
+
+SNetStatsCountersView::~SNetStatsCountersView()
+{
+	// Remove ourselves from the Insights manager.
+	if (FInsightsManager::Get().IsValid())
+	{
+		FInsightsManager::Get()->GetSessionChangedEvent().RemoveAll(this);
+	}
+}
+
+////////////////////////////////////////////////////////////////////////////////////////////////////
+
+BEGIN_SLATE_FUNCTION_BUILD_OPTIMIZATION
+void SNetStatsCountersView::Construct(const FArguments& InArgs, TSharedPtr<SNetworkingProfilerWindow> InProfilerWindow)
+{
+	ProfilerWindow = InProfilerWindow;
+
+	SAssignNew(ExternalScrollbar, SScrollBar)
+	.AlwaysShowScrollbar(true);
+
+	ChildSlot
+	[
+		SNew(SVerticalBox)
+
+		+ SVerticalBox::Slot()
+		.VAlign(VAlign_Center)
+		.AutoHeight()
+		.Padding(2.0f, 2.0f, 2.0f, 2.0f)
+		[
+			SNew(SVerticalBox)
+
+			+ SVerticalBox::Slot()
+			.VAlign(VAlign_Center)
+			.Padding(2.0f)
+			.AutoHeight()
+			[
+				SNew(SHorizontalBox)
+
+				// Search box
+				+ SHorizontalBox::Slot()				
+				.Padding(2.0f)
+				.FillWidth(1.0f)
+				.VAlign(VAlign_Center)
+				[
+					SAssignNew(SearchBox, SSearchBox)
+					.HintText(LOCTEXT("SearchBoxHint", "Search NetStatsCounters or groups"))
+					.OnTextChanged(this, &SNetStatsCountersView::SearchBox_OnTextChanged)
+					.IsEnabled(this, &SNetStatsCountersView::SearchBox_IsEnabled)
+					.ToolTipText(LOCTEXT("FilterSearchHint", "Type here to search NetStatsCounters or groups"))
+				]
+
+				// Filter out net stats counter types with zero count
+				+ SHorizontalBox::Slot()
+				.VAlign(VAlign_Center)
+				.Padding(2.0f)
+				.AutoWidth()
+				[
+					SNew(SCheckBox)
+					.Style(FAppStyle::Get(), "ToggleButtonCheckbox")
+					.HAlign(HAlign_Center)
+					.Padding(3.0f)
+					.OnCheckStateChanged(this, &SNetStatsCountersView::FilterOutZeroCountStatsCounters_OnCheckStateChanged)
+					.IsChecked(this, &SNetStatsCountersView::FilterOutZeroCountStatsCounters_IsChecked)
+					.ToolTipText(LOCTEXT("FilterOutZeroCountStatsCounters_Tooltip", "Filter out the stats counters having zero total instance counts (aggregated stats)."))
+					[
+						SNew(SImage)
+						.Image(FInsightsStyle::Get().GetBrush("Icons.ZeroCountFilter"))
+					]
+				]
+			]
+
+			// Group by
+			+ SVerticalBox::Slot()
+			.VAlign(VAlign_Center)
+			.Padding(2.0f)
+			.AutoHeight()
+			[
+				SNew(SHorizontalBox)
+
+				+ SHorizontalBox::Slot()
+				.AutoWidth()
+				.Padding(0.0f, 0.0f, 4.0f, 0.0f)
+				.VAlign(VAlign_Center)
+				[
+					SNew(STextBlock)
+					.Text(LOCTEXT("GroupByText", "Group by"))
+				]
+
+				+ SHorizontalBox::Slot()
+				.AutoWidth()
+				.VAlign(VAlign_Center)
+				[
+					SNew(SBox)
+					.MinDesiredWidth(128.0f)
+					[
+						SAssignNew(GroupByComboBox, SComboBox<TSharedPtr<ENetStatsCounterGroupingMode>>)
+						.ToolTipText(this, &SNetStatsCountersView::GroupBy_GetSelectedTooltipText)
+						.OptionsSource(&GroupByOptionsSource)
+						.OnSelectionChanged(this, &SNetStatsCountersView::GroupBy_OnSelectionChanged)
+						.OnGenerateWidget(this, &SNetStatsCountersView::GroupBy_OnGenerateWidget)
+						[
+							SNew(STextBlock)
+							.Text(this, &SNetStatsCountersView::GroupBy_GetSelectedText)
+						]
+					]
+				]
+			]
+		]
+
+		// Tree view
+		+ SVerticalBox::Slot()
+		.FillHeight(1.0f)
+		.Padding(0.0f, 2.0f, 0.0f, 0.0f)
+		[
+			SNew(SHorizontalBox)
+
+			+ SHorizontalBox::Slot()
+			.FillWidth(1.0f)
+			.Padding(0.0f)
+			[
+				SAssignNew(TreeView, STreeView<FNetStatsCounterNodePtr>)
+				.ExternalScrollbar(ExternalScrollbar)
+				.SelectionMode(ESelectionMode::Multi)
+				.TreeItemsSource(&FilteredGroupNodes)
+				.OnGetChildren(this, &SNetStatsCountersView::TreeView_OnGetChildren)
+				.OnGenerateRow(this, &SNetStatsCountersView::TreeView_OnGenerateRow)
+				.OnSelectionChanged(this, &SNetStatsCountersView::TreeView_OnSelectionChanged)
+				.OnMouseButtonDoubleClick(this, &SNetStatsCountersView::TreeView_OnMouseButtonDoubleClick)
+				.OnContextMenuOpening(FOnContextMenuOpening::CreateSP(this, &SNetStatsCountersView::TreeView_GetMenuContent))
+				.ItemHeight(16.0f)
+				.HeaderRow
+				(
+					SAssignNew(TreeViewHeaderRow, SHeaderRow)
+					.Visibility(EVisibility::Visible)
+				)
+			]
+
+			+ SHorizontalBox::Slot()
+			.AutoWidth()
+			.Padding(0.0f)
+			[
+				SNew(SBox)
+				.WidthOverride(FOptionalSize(13.0f))
+				[
+					ExternalScrollbar.ToSharedRef()
+				]
+			]
+		]
+	];
+
+	InitializeAndShowHeaderColumns();
+	//BindCommands();
+
+	// Create the search filters: text based, type based etc.
+	TextFilter = MakeShared<FNetStatsCounterNodeTextFilter>(FNetStatsCounterNodeTextFilter::FItemToStringArray::CreateSP(this, &SNetStatsCountersView::HandleItemToStringArray));
+	Filters = MakeShared<FNetStatsCounterNodeFilterCollection>();
+	Filters->Add(TextFilter);
+
+	CreateGroupByOptionsSources();
+	CreateSortings();
+
+	// Register ourselves with the Insights manager.
+	FInsightsManager::Get()->GetSessionChangedEvent().AddSP(this, &SNetStatsCountersView::InsightsManager_OnSessionChanged);
+
+	// Update the Session (i.e. when analysis session was already started).
+	InsightsManager_OnSessionChanged();
+}
+END_SLATE_FUNCTION_BUILD_OPTIMIZATION
+
+////////////////////////////////////////////////////////////////////////////////////////////////////
+
+TSharedPtr<SWidget> SNetStatsCountersView::TreeView_GetMenuContent()
+{
+	const TArray<FNetStatsCounterNodePtr> SelectedNodes = TreeView->GetSelectedItems();
+	const int32 NumSelectedNodes = SelectedNodes.Num();
+	FNetStatsCounterNodePtr SelectedNode = NumSelectedNodes ? SelectedNodes[0] : nullptr;
+
+	FText SelectionStr;
+
+	if (NumSelectedNodes == 0)
+	{
+		SelectionStr = LOCTEXT("NothingSelected", "Nothing selected");
+	}
+	else if (NumSelectedNodes == 1)
+	{
+		FString ItemName = SelectedNode->GetName().ToString();
+		const int32 MaxStringLen = 64;
+		if (ItemName.Len() > MaxStringLen)
+		{
+			ItemName = ItemName.Left(MaxStringLen) + TEXT("...");
+		}
+		SelectionStr = FText::FromString(ItemName);
+	}
+	else
+	{
+		SelectionStr = FText::Format(LOCTEXT("MultipleSelection_Fmt", "{0} selected items"), FText::AsNumber(NumSelectedNodes));
+	}
+
+	const bool bShouldCloseWindowAfterMenuSelection = true;
+	FMenuBuilder MenuBuilder(bShouldCloseWindowAfterMenuSelection, NULL);
+
+	// Selection menu
+	MenuBuilder.BeginSection("Selection", LOCTEXT("ContextMenu_Header_Selection", "Selection"));
+	{
+		struct FLocal
+		{
+			static bool ReturnFalse()
+			{
+				return false;
+			}
+		};
+
+		FUIAction DummyUIAction;
+		DummyUIAction.CanExecuteAction = FCanExecuteAction::CreateStatic(&FLocal::ReturnFalse);
+		MenuBuilder.AddMenuEntry
+		(
+			SelectionStr,
+			LOCTEXT("ContextMenu_Selection", "Currently selected items"),
+			FSlateIcon(),
+			DummyUIAction,
+			NAME_None,
+			EUserInterfaceActionType::Button
+		);
+	}
+	MenuBuilder.EndSection();
+
+	MenuBuilder.BeginSection("Misc", LOCTEXT("ContextMenu_Header_Misc", "Miscellaneous"));
+	{
+		/*TODO
+		MenuBuilder.AddMenuEntry
+		(
+			FTimersViewCommands::Get().Command_CopyToClipboard,
+			NAME_None,
+			TAttribute<FText>(),
+			TAttribute<FText>(),
+			FSlateIcon(FAppStyle::Get().GetStyleSetName(), "GenericCommands.Copy")
+		);
+		*/
+
+		MenuBuilder.AddSubMenu
+		(
+			LOCTEXT("ContextMenu_Header_Misc_Sort", "Sort By"),
+			LOCTEXT("ContextMenu_Header_Misc_Sort_Desc", "Sort by column"),
+			FNewMenuDelegate::CreateSP(this, &SNetStatsCountersView::TreeView_BuildSortByMenu),
+			false,
+			FSlateIcon(FInsightsStyle::GetStyleSetName(), "Icons.SortBy")
+		);
+	}
+	MenuBuilder.EndSection();
+
+	MenuBuilder.BeginSection("Columns", LOCTEXT("ContextMenu_Header_Columns", "Columns"));
+	{
+		MenuBuilder.AddSubMenu
+		(
+			LOCTEXT("ContextMenu_Header_Columns_View", "View Column"),
+			LOCTEXT("ContextMenu_Header_Columns_View_Desc", "Hides or shows columns"),
+			FNewMenuDelegate::CreateSP(this, &SNetStatsCountersView::TreeView_BuildViewColumnMenu),
+			false,
+			FSlateIcon(FInsightsStyle::GetStyleSetName(), "Icons.ViewColumn")
+		);
+
+		FUIAction Action_ShowAllColumns
+		(
+			FExecuteAction::CreateSP(this, &SNetStatsCountersView::ContextMenu_ShowAllColumns_Execute),
+			FCanExecuteAction::CreateSP(this, &SNetStatsCountersView::ContextMenu_ShowAllColumns_CanExecute)
+		);
+		MenuBuilder.AddMenuEntry
+		(
+			LOCTEXT("ContextMenu_Header_Columns_ShowAllColumns", "Show All Columns"),
+			LOCTEXT("ContextMenu_Header_Columns_ShowAllColumns_Desc", "Resets tree view to show all columns."),
+			FSlateIcon(FInsightsStyle::GetStyleSetName(), "Icons.ResetColumn"),
+			Action_ShowAllColumns,
+			NAME_None,
+			EUserInterfaceActionType::Button
+		);
+
+		FUIAction Action_ResetColumns
+		(
+			FExecuteAction::CreateSP(this, &SNetStatsCountersView::ContextMenu_ResetColumns_Execute),
+			FCanExecuteAction::CreateSP(this, &SNetStatsCountersView::ContextMenu_ResetColumns_CanExecute)
+		);
+		MenuBuilder.AddMenuEntry
+		(
+			LOCTEXT("ContextMenu_Header_Columns_ResetColumns", "Reset Columns to Default"),
+			LOCTEXT("ContextMenu_Header_Columns_ResetColumns_Desc", "Resets columns to default."),
+			FSlateIcon(FInsightsStyle::GetStyleSetName(), "Icons.ResetColumn"),
+			Action_ResetColumns,
+			NAME_None,
+			EUserInterfaceActionType::Button
+		);
+	}
+	MenuBuilder.EndSection();
+
+	return MenuBuilder.MakeWidget();
+}
+
+////////////////////////////////////////////////////////////////////////////////////////////////////
+
+void SNetStatsCountersView::TreeView_BuildSortByMenu(FMenuBuilder& MenuBuilder)
+{
+	// TODO: Refactor later @see TSharedPtr<SWidget> SCascadePreviewViewportToolBar::GenerateViewMenu() const
+
+	MenuBuilder.BeginSection("ColumnName", LOCTEXT("ContextMenu_Header_Misc_ColumnName", "Column Name"));
+
+	for (const TSharedRef<Insights::FTableColumn>& ColumnRef : Table->GetColumns())
+	{
+		const Insights::FTableColumn& Column = *ColumnRef;
+
+		if (Column.IsVisible() && Column.CanBeSorted())
+		{
+			FUIAction Action_SortByColumn
+			(
+				FExecuteAction::CreateSP(this, &SNetStatsCountersView::ContextMenu_SortByColumn_Execute, Column.GetId()),
+				FCanExecuteAction::CreateSP(this, &SNetStatsCountersView::ContextMenu_SortByColumn_CanExecute, Column.GetId()),
+				FIsActionChecked::CreateSP(this, &SNetStatsCountersView::ContextMenu_SortByColumn_IsChecked, Column.GetId())
+			);
+			MenuBuilder.AddMenuEntry
+			(
+				Column.GetTitleName(),
+				Column.GetDescription(),
+				FSlateIcon(),
+				Action_SortByColumn,
+				NAME_None,
+				EUserInterfaceActionType::RadioButton
+			);
+		}
+	}
+
+	MenuBuilder.EndSection();
+
+	//-----------------------------------------------------------------------------
+
+	MenuBuilder.BeginSection("SortMode", LOCTEXT("ContextMenu_Header_Misc_Sort_SortMode", "Sort Mode"));
+	{
+		FUIAction Action_SortAscending
+		(
+			FExecuteAction::CreateSP(this, &SNetStatsCountersView::ContextMenu_SortMode_Execute, EColumnSortMode::Ascending),
+			FCanExecuteAction::CreateSP(this, &SNetStatsCountersView::ContextMenu_SortMode_CanExecute, EColumnSortMode::Ascending),
+			FIsActionChecked::CreateSP(this, &SNetStatsCountersView::ContextMenu_SortMode_IsChecked, EColumnSortMode::Ascending)
+		);
+		MenuBuilder.AddMenuEntry
+		(
+			LOCTEXT("ContextMenu_Header_Misc_Sort_SortAscending", "Sort Ascending"),
+			LOCTEXT("ContextMenu_Header_Misc_Sort_SortAscending_Desc", "Sorts ascending"),
+			FSlateIcon(FAppStyle::Get().GetStyleSetName(), "Icons.SortUp"),
+			Action_SortAscending,
+			NAME_None,
+			EUserInterfaceActionType::RadioButton
+		);
+
+		FUIAction Action_SortDescending
+		(
+			FExecuteAction::CreateSP(this, &SNetStatsCountersView::ContextMenu_SortMode_Execute, EColumnSortMode::Descending),
+			FCanExecuteAction::CreateSP(this, &SNetStatsCountersView::ContextMenu_SortMode_CanExecute, EColumnSortMode::Descending),
+			FIsActionChecked::CreateSP(this, &SNetStatsCountersView::ContextMenu_SortMode_IsChecked, EColumnSortMode::Descending)
+		);
+		MenuBuilder.AddMenuEntry
+		(
+			LOCTEXT("ContextMenu_Header_Misc_Sort_SortDescending", "Sort Descending"),
+			LOCTEXT("ContextMenu_Header_Misc_Sort_SortDescending_Desc", "Sorts descending"),
+			FSlateIcon(FAppStyle::Get().GetStyleSetName(), "Icons.SortDown"),
+			Action_SortDescending,
+			NAME_None,
+			EUserInterfaceActionType::RadioButton
+		);
+	}
+	MenuBuilder.EndSection();
+}
+
+////////////////////////////////////////////////////////////////////////////////////////////////////
+
+void SNetStatsCountersView::TreeView_BuildViewColumnMenu(FMenuBuilder& MenuBuilder)
+{
+	MenuBuilder.BeginSection("ViewColumn", LOCTEXT("ContextMenu_Header_Columns_View", "View Column"));
+
+	for (const TSharedRef<Insights::FTableColumn>& ColumnRef : Table->GetColumns())
+	{
+		const Insights::FTableColumn& Column = *ColumnRef;
+
+		FUIAction Action_ToggleColumn
+		(
+			FExecuteAction::CreateSP(this, &SNetStatsCountersView::ToggleColumnVisibility, Column.GetId()),
+			FCanExecuteAction::CreateSP(this, &SNetStatsCountersView::CanToggleColumnVisibility, Column.GetId()),
+			FIsActionChecked::CreateSP(this, &SNetStatsCountersView::IsColumnVisible, Column.GetId())
+		);
+		MenuBuilder.AddMenuEntry
+		(
+			Column.GetTitleName(),
+			Column.GetDescription(),
+			FSlateIcon(),
+			Action_ToggleColumn,
+			NAME_None,
+			EUserInterfaceActionType::ToggleButton
+		);
+	}
+
+	MenuBuilder.EndSection();
+}
+
+////////////////////////////////////////////////////////////////////////////////////////////////////
+
+void SNetStatsCountersView::InitializeAndShowHeaderColumns()
+{
+	// Create columns.
+	TArray<TSharedRef<Insights::FTableColumn>> Columns;
+	FNetStatsCountersViewColumnFactory::CreateNetStatsCountersViewColumns(Columns);
+	Table->SetColumns(Columns);
+
+	// Show columns.
+	for (const TSharedRef<Insights::FTableColumn>& ColumnRef : Table->GetColumns())
+	{
+		if (ColumnRef->ShouldBeVisible())
+		{
+			ShowColumn(ColumnRef->GetId());
+		}
+	}
+}
+
+////////////////////////////////////////////////////////////////////////////////////////////////////
+
+FText SNetStatsCountersView::GetColumnHeaderText(const FName ColumnId) const
+{
+	const Insights::FTableColumn& Column = *Table->FindColumnChecked(ColumnId);
+	return Column.GetShortName();
+}
+
+////////////////////////////////////////////////////////////////////////////////////////////////////
+
+TSharedRef<SWidget> SNetStatsCountersView::TreeViewHeaderRow_GenerateColumnMenu(const Insights::FTableColumn& Column)
+{
+	bool bIsMenuVisible = false;
+
+	const bool bShouldCloseWindowAfterMenuSelection = true;
+	FMenuBuilder MenuBuilder(bShouldCloseWindowAfterMenuSelection, NULL);
+	{
+		if (Column.CanBeHidden())
+		{
+			MenuBuilder.BeginSection("Column", LOCTEXT("TreeViewHeaderRow_Header_Column", "Column"));
+
+			FUIAction Action_HideColumn
+			(
+				FExecuteAction::CreateSP(this, &SNetStatsCountersView::HideColumn, Column.GetId()),
+				FCanExecuteAction::CreateSP(this, &SNetStatsCountersView::CanHideColumn, Column.GetId())
+			);
+			MenuBuilder.AddMenuEntry
+			(
+				LOCTEXT("TreeViewHeaderRow_HideColumn", "Hide"),
+				LOCTEXT("TreeViewHeaderRow_HideColumn_Desc", "Hides the selected column"),
+				FSlateIcon(),
+				Action_HideColumn,
+				NAME_None,
+				EUserInterfaceActionType::Button
+			);
+
+			bIsMenuVisible = true;
+			MenuBuilder.EndSection();
+		}
+
+		if (Column.CanBeSorted())
+		{
+			MenuBuilder.BeginSection("SortMode", LOCTEXT("ContextMenu_Header_Misc_Sort_SortMode", "Sort Mode"));
+
+			FUIAction Action_SortAscending
+			(
+				FExecuteAction::CreateSP(this, &SNetStatsCountersView::HeaderMenu_SortMode_Execute, Column.GetId(), EColumnSortMode::Ascending),
+				FCanExecuteAction::CreateSP(this, &SNetStatsCountersView::HeaderMenu_SortMode_CanExecute, Column.GetId(), EColumnSortMode::Ascending),
+				FIsActionChecked::CreateSP(this, &SNetStatsCountersView::HeaderMenu_SortMode_IsChecked, Column.GetId(), EColumnSortMode::Ascending)
+			);
+			MenuBuilder.AddMenuEntry
+			(
+				LOCTEXT("ContextMenu_Header_Misc_Sort_SortAscending", "Sort Ascending"),
+				LOCTEXT("ContextMenu_Header_Misc_Sort_SortAscending_Desc", "Sorts ascending"),
+				FSlateIcon(FAppStyle::Get().GetStyleSetName(), "Icons.SortUp"),
+				Action_SortAscending,
+				NAME_None,
+				EUserInterfaceActionType::RadioButton
+			);
+
+			FUIAction Action_SortDescending
+			(
+				FExecuteAction::CreateSP(this, &SNetStatsCountersView::HeaderMenu_SortMode_Execute, Column.GetId(), EColumnSortMode::Descending),
+				FCanExecuteAction::CreateSP(this, &SNetStatsCountersView::HeaderMenu_SortMode_CanExecute, Column.GetId(), EColumnSortMode::Descending),
+				FIsActionChecked::CreateSP(this, &SNetStatsCountersView::HeaderMenu_SortMode_IsChecked, Column.GetId(), EColumnSortMode::Descending)
+			);
+			MenuBuilder.AddMenuEntry
+			(
+				LOCTEXT("ContextMenu_Header_Misc_Sort_SortDescending", "Sort Descending"),
+				LOCTEXT("ContextMenu_Header_Misc_Sort_SortDescending_Desc", "Sorts descending"),
+				FSlateIcon(FAppStyle::Get().GetStyleSetName(), "Icons.SortDown"),
+				Action_SortDescending,
+				NAME_None,
+				EUserInterfaceActionType::RadioButton
+			);
+
+			bIsMenuVisible = true;
+			MenuBuilder.EndSection();
+		}
+
+		//if (Column.CanBeFiltered())
+		//{
+		//	MenuBuilder.BeginSection("FilterMode", LOCTEXT("ContextMenu_Header_Misc_Filter_FilterMode", "Filter Mode"));
+		//	bIsMenuVisible = true;
+		//	MenuBuilder.EndSection();
+		//}
+	}
+
+	/*
+	TODO:
+	- Show top ten
+	- Show top bottom
+	- Filter by list (avg, median, 10%, 90%, etc.)
+	- Text box for filtering for each column instead of one text box used for filtering
+	- Grouping button for flat view modes (show at most X groups, show all groups for names)
+	*/
+
+	return bIsMenuVisible ? MenuBuilder.MakeWidget() : (TSharedRef<SWidget>)SNullWidget::NullWidget;
+}
+
+////////////////////////////////////////////////////////////////////////////////////////////////////
+
+void SNetStatsCountersView::InsightsManager_OnSessionChanged()
+{
+	TSharedPtr<const TraceServices::IAnalysisSession> NewSession = FInsightsManager::Get()->GetSession();
+
+	if (NewSession != Session)
+	{
+		Session = NewSession;
+		Reset();
+	}
+	else
+	{
+		UpdateTree();
+	}
+}
+
+////////////////////////////////////////////////////////////////////////////////////////////////////
+
+void SNetStatsCountersView::UpdateTree()
+{
+	FStopwatch Stopwatch;
+	Stopwatch.Start();
+
+	CreateGroups();
+
+	Stopwatch.Update();
+	const double Time1 = Stopwatch.GetAccumulatedTime();
+
+	SortTreeNodes();
+
+	Stopwatch.Update();
+	const double Time2 = Stopwatch.GetAccumulatedTime();
+
+	ApplyFiltering();
+
+	Stopwatch.Stop();
+	const double TotalTime = Stopwatch.GetAccumulatedTime();
+	if (TotalTime > 0.1)
+	{
+		UE_LOG(NetworkingProfiler, Log, TEXT("[NetStats] Tree view updated in %.3fs (%d events) --> G:%.3fs + S:%.3fs + F:%.3fs"),
+			TotalTime, NetStatsCounterNodes.Num(), Time1, Time2 - Time1, TotalTime - Time2);
+	}
+}
+
+////////////////////////////////////////////////////////////////////////////////////////////////////
+
+void SNetStatsCountersView::ApplyFiltering()
+{
+	FilteredGroupNodes.Reset();
+
+	// Apply filter to all groups and its children.
+	const int32 NumGroups = GroupNodes.Num();
+	for (int32 GroupIndex = 0; GroupIndex < NumGroups; ++GroupIndex)
+	{
+		FNetStatsCounterNodePtr& GroupPtr = GroupNodes[GroupIndex];
+		GroupPtr->ClearFilteredChildren();
+		const bool bIsGroupVisible = Filters->PassesAllFilters(GroupPtr);
+
+		const TArray<Insights::FBaseTreeNodePtr>& GroupChildren = GroupPtr->GetChildren();
+		int32 NumVisibleChildren = 0;
+		for (const Insights::FBaseTreeNodePtr& ChildPtr : GroupChildren)
+		{
+			const FNetStatsCounterNodePtr& NodePtr = StaticCastSharedPtr<FNetStatsCounterNode, Insights::FBaseTreeNode>(ChildPtr);
+
+			const bool bIsChildVisible = (!bFilterOutZeroCountStatsCounters || NodePtr->GetAggregatedStats().Count > 0)
+									  && bStatsCounterTypeIsVisible[static_cast<int>(NodePtr->GetType())]
+									  && Filters->PassesAllFilters(NodePtr);
+			if (bIsChildVisible)
+			{
+				// Add a child
+				GroupPtr->AddFilteredChild(NodePtr);
+				NumVisibleChildren++;
+			}
+		}
+
+		if (bIsGroupVisible || NumVisibleChildren > 0)
+		{
+			// Add a group.
+			FilteredGroupNodes.Add(GroupPtr);
+			GroupPtr->SetExpansion(true);
+		}
+		else
+		{
+			GroupPtr->SetExpansion(false);
+		}
+	}
+
+	// Only expand net event nodes if we have a text filter.
+	const bool bNonEmptyTextFilter = !TextFilter->GetRawFilterText().IsEmpty();
+	if (bNonEmptyTextFilter)
+	{
+		if (!bExpansionSaved)
+		{
+			ExpandedNodes.Empty();
+			TreeView->GetExpandedItems(ExpandedNodes);
+			bExpansionSaved = true;
+		}
+
+		for (const FNetStatsCounterNodePtr& GroupPtr : FilteredGroupNodes)
+		{
+			TreeView->SetItemExpansion(GroupPtr, GroupPtr->IsExpanded());
+		}
+	}
+	else
+	{
+		if (bExpansionSaved)
+		{
+			// Restore previously expanded nodes when the text filter is disabled.
+			TreeView->ClearExpandedItems();
+			for (auto It = ExpandedNodes.CreateConstIterator(); It; ++It)
+			{
+				TreeView->SetItemExpansion(*It, true);
+			}
+			bExpansionSaved = false;
+		}
+	}
+
+	// Request tree refresh
+	TreeView->RequestTreeRefresh();
+}
+
+////////////////////////////////////////////////////////////////////////////////////////////////////
+
+void SNetStatsCountersView::HandleItemToStringArray(const FNetStatsCounterNodePtr& FNetStatsCounterNodePtr, TArray<FString>& OutSearchStrings) const
+{
+	OutSearchStrings.Add(FNetStatsCounterNodePtr->GetName().GetPlainNameString());
+}
+
+////////////////////////////////////////////////////////////////////////////////////////////////////
+
+TSharedRef<SWidget> SNetStatsCountersView::GetToggleButtonForNetStatsCounterType(const ENetStatsCounterNodeType NodeType)
+{
+	return SNew(SCheckBox)
+		.Style(FAppStyle::Get(), "ToggleButtonCheckbox")
+		.Padding(FMargin(4.0f, 2.0f, 4.0f, 2.0f)).Padding(2.0f)
+		.OnCheckStateChanged(this, &SNetStatsCountersView::FilterByNetStatsCounterType_OnCheckStateChanged, NodeType)
+		.IsChecked(this, &SNetStatsCountersView::FilterByNetStatsCounterType_IsChecked, NodeType)
+		.ToolTipText(NetStatsCounterNodeTypeHelper::ToDescription(NodeType))
+		[
+			SNew(SHorizontalBox)
+
+			//+ SHorizontalBox::Slot()
+			//.AutoWidth()
+			//.VAlign(VAlign_Center)
+			//[
+			//	SNew(SImage)
+			//	.Image(NetStatsCounterNodeTypeHelper::GetIcon(NodeType))
+			//]
+
+			+ SHorizontalBox::Slot()
+			.VAlign(VAlign_Center)
+			[
+				SNew(STextBlock)
+				.Text(NetStatsCounterNodeTypeHelper::ToText(NodeType))
+			]
+		];
+}
+
+////////////////////////////////////////////////////////////////////////////////////////////////////
+
+void SNetStatsCountersView::FilterOutZeroCountStatsCounters_OnCheckStateChanged(ECheckBoxState NewState)
+{
+	bFilterOutZeroCountStatsCounters = (NewState == ECheckBoxState::Checked);
+	ApplyFiltering();
+}
+
+////////////////////////////////////////////////////////////////////////////////////////////////////
+
+ECheckBoxState SNetStatsCountersView::FilterOutZeroCountStatsCounters_IsChecked() const
+{
+	return bFilterOutZeroCountStatsCounters ? ECheckBoxState::Checked : ECheckBoxState::Unchecked;
+}
+
+////////////////////////////////////////////////////////////////////////////////////////////////////
+
+void SNetStatsCountersView::FilterByNetStatsCounterType_OnCheckStateChanged(ECheckBoxState NewState, const ENetStatsCounterNodeType InStatType)
+{
+	bStatsCounterTypeIsVisible[static_cast<int>(InStatType)] = (NewState == ECheckBoxState::Checked);
+	ApplyFiltering();
+}
+
+////////////////////////////////////////////////////////////////////////////////////////////////////
+
+ECheckBoxState SNetStatsCountersView::FilterByNetStatsCounterType_IsChecked(const ENetStatsCounterNodeType InStatType) const
+{
+	return bStatsCounterTypeIsVisible[static_cast<int>(InStatType)] ? ECheckBoxState::Checked : ECheckBoxState::Unchecked;
+}
+
+////////////////////////////////////////////////////////////////////////////////////////////////////
+// TreeView
+////////////////////////////////////////////////////////////////////////////////////////////////////
+
+void SNetStatsCountersView::TreeView_Refresh()
+{
+	if (TreeView.IsValid())
+	{
+		TreeView->RequestTreeRefresh();
+	}
+}
+
+////////////////////////////////////////////////////////////////////////////////////////////////////
+
+void SNetStatsCountersView::TreeView_OnSelectionChanged(FNetStatsCounterNodePtr SelectedItem, ESelectInfo::Type SelectInfo)
+{
+}
+
+////////////////////////////////////////////////////////////////////////////////////////////////////
+
+void SNetStatsCountersView::TreeView_OnGetChildren(FNetStatsCounterNodePtr InParent, TArray<FNetStatsCounterNodePtr>& OutChildren)
+{
+	const TArray<Insights::FBaseTreeNodePtr>& Children = InParent->GetFilteredChildren();
+	OutChildren.Reset(Children.Num());
+	for (const Insights::FBaseTreeNodePtr& Child : Children)
+	{
+		OutChildren.Add(StaticCastSharedPtr<FNetStatsCounterNode, Insights::FBaseTreeNode>(Child));
+	}
+}
+
+////////////////////////////////////////////////////////////////////////////////////////////////////
+
+void SNetStatsCountersView::TreeView_OnMouseButtonDoubleClick(FNetStatsCounterNodePtr NetEventNodePtr)
+{
+	if (NetEventNodePtr->IsGroup())
+	{
+		const bool bIsGroupExpanded = TreeView->IsItemExpanded(NetEventNodePtr);
+		TreeView->SetItemExpansion(NetEventNodePtr, !bIsGroupExpanded);
+	}
+}
+
+////////////////////////////////////////////////////////////////////////////////////////////////////
+// Tree View's Table Row
+////////////////////////////////////////////////////////////////////////////////////////////////////
+
+TSharedRef<ITableRow> SNetStatsCountersView::TreeView_OnGenerateRow(FNetStatsCounterNodePtr NetStatsCounterNodePtr, const TSharedRef<STableViewBase>& OwnerTable)
+{
+	TSharedRef<ITableRow> TableRow =
+		SNew(SNetStatsCountersTableRow, OwnerTable)
+		.OnShouldBeEnabled(this, &SNetStatsCountersView::TableRow_ShouldBeEnabled)
+		.OnIsColumnVisible(this, &SNetStatsCountersView::IsColumnVisible)
+		.OnSetHoveredCell(this, &SNetStatsCountersView::TableRow_SetHoveredCell)
+		.OnGetColumnOutlineHAlignmentDelegate(this, &SNetStatsCountersView::TableRow_GetColumnOutlineHAlignment)
+		.HighlightText(this, &SNetStatsCountersView::TableRow_GetHighlightText)
+		.HighlightedNodeName(this, &SNetStatsCountersView::TableRow_GetHighlightedNodeName)
+		.TablePtr(Table)
+		.NetStatsCounterNodePtr(NetStatsCounterNodePtr);
+
+	return TableRow;
+}
+
+////////////////////////////////////////////////////////////////////////////////////////////////////
+
+bool SNetStatsCountersView::TableRow_ShouldBeEnabled(FNetStatsCounterNodePtr NodePtr) const
+{
+	return true;
+}
+
+////////////////////////////////////////////////////////////////////////////////////////////////////
+
+void SNetStatsCountersView::TableRow_SetHoveredCell(TSharedPtr<Insights::FTable> InTablePtr, TSharedPtr<Insights::FTableColumn> InColumnPtr, FNetStatsCounterNodePtr InNodePtr)
+{
+	HoveredColumnId = InColumnPtr ? InColumnPtr->GetId() : FName();
+
+	const bool bIsAnyMenusVisible = FSlateApplication::Get().AnyMenusVisible();
+	if (!HasMouseCapture() && !bIsAnyMenusVisible)
+	{
+		HoveredNodePtr = InNodePtr;
+	}
+}
+
+////////////////////////////////////////////////////////////////////////////////////////////////////
+
+EHorizontalAlignment SNetStatsCountersView::TableRow_GetColumnOutlineHAlignment(const FName ColumnId) const
+{
+	const TIndirectArray<SHeaderRow::FColumn>& Columns = TreeViewHeaderRow->GetColumns();
+	const int32 LastColumnIdx = Columns.Num() - 1;
+
+	// First column
+	if (Columns[0].ColumnId == ColumnId)
+	{
+		return HAlign_Left;
+	}
+	// Last column
+	else if (Columns[LastColumnIdx].ColumnId == ColumnId)
+	{
+		return HAlign_Right;
+	}
+	// Middle columns
+	{
+		return HAlign_Center;
+	}
+}
+
+////////////////////////////////////////////////////////////////////////////////////////////////////
+
+FText SNetStatsCountersView::TableRow_GetHighlightText() const
+{
+	return SearchBox->GetText();
+}
+
+////////////////////////////////////////////////////////////////////////////////////////////////////
+
+FName SNetStatsCountersView::TableRow_GetHighlightedNodeName() const
+{
+	return HighlightedNodeName;
+}
+
+////////////////////////////////////////////////////////////////////////////////////////////////////
+// SearchBox
+////////////////////////////////////////////////////////////////////////////////////////////////////
+
+void SNetStatsCountersView::SearchBox_OnTextChanged(const FText& InFilterText)
+{
+	TextFilter->SetRawFilterText(InFilterText);
+	SearchBox->SetError(TextFilter->GetFilterErrorText());
+	ApplyFiltering();
+}
+
+////////////////////////////////////////////////////////////////////////////////////////////////////
+
+bool SNetStatsCountersView::SearchBox_IsEnabled() const
+{
+	return NetStatsCounterNodes.Num() > 0;
+}
+
+////////////////////////////////////////////////////////////////////////////////////////////////////
+// Grouping
+////////////////////////////////////////////////////////////////////////////////////////////////////
+
+void SNetStatsCountersView::CreateGroups()
+{
+	if (GroupingMode == ENetStatsCounterGroupingMode::Flat)
+	{
+		GroupNodes.Reset();
+
+		const FName GroupName(TEXT("All"));
+		FNetStatsCounterNodePtr GroupPtr = MakeShared<FNetStatsCounterNode>(GroupName);
+		GroupNodes.Add(GroupPtr);
+
+		for (const FNetStatsCounterNodePtr& NodePtr : NetStatsCounterNodes)
+		{
+			GroupPtr->AddChildAndSetGroupPtr(NodePtr);
+		}
+		TreeView->SetItemExpansion(GroupPtr, true);
+	}
+	// Creates one group for each stat type.
+	else if (GroupingMode == ENetStatsCounterGroupingMode::ByType)
+	{
+		TMap<ENetStatsCounterNodeType, FNetStatsCounterNodePtr> GroupNodeSet;
+		for (const FNetStatsCounterNodePtr& NodePtr : NetStatsCounterNodes)
+		{
+			const ENetStatsCounterNodeType NodeType = NodePtr->GetType();
+			FNetStatsCounterNodePtr GroupPtr = GroupNodeSet.FindRef(NodeType);
+			if (!GroupPtr)
+			{
+				const FName GroupName = *NetStatsCounterNodeTypeHelper::ToText(NodeType).ToString();
+				GroupPtr = GroupNodeSet.Add(NodeType, MakeShared<FNetStatsCounterNode>(GroupName));
+			}
+			GroupPtr->AddChildAndSetGroupPtr(NodePtr);
+			TreeView->SetItemExpansion(GroupPtr, true);
+		}
+		GroupNodeSet.KeySort([](const ENetStatsCounterNodeType& A, const ENetStatsCounterNodeType& B) { return A < B; }); // sort groups by type
+		GroupNodeSet.GenerateValueArray(GroupNodes);
+	}
+	// Creates one group for one letter.
+	else if (GroupingMode == ENetStatsCounterGroupingMode::ByName)
+	{
+		TMap<TCHAR, FNetStatsCounterNodePtr> GroupNodeSet;
+		for (const FNetStatsCounterNodePtr& NodePtr : NetStatsCounterNodes)
+		{
+			FString FirstLetterStr(NodePtr->GetName().GetPlainNameString().Left(1).ToUpper());
+			const TCHAR FirstLetter = FirstLetterStr[0];
+			FNetStatsCounterNodePtr GroupPtr = GroupNodeSet.FindRef(FirstLetter);
+			if (!GroupPtr)
+			{
+				const FName GroupName(FirstLetterStr);
+				GroupPtr = GroupNodeSet.Add(FirstLetter, MakeShared<FNetStatsCounterNode>(GroupName));
+			}
+			GroupPtr->AddChildAndSetGroupPtr(NodePtr);
+		}
+		GroupNodeSet.KeySort([](const TCHAR& A, const TCHAR& B) { return A < B; }); // sort groups alphabetically
+		GroupNodeSet.GenerateValueArray(GroupNodes);
+	}
+}
+
+////////////////////////////////////////////////////////////////////////////////////////////////////
+
+void SNetStatsCountersView::CreateGroupByOptionsSources()
+{
+	GroupByOptionsSource.Reset(3);
+
+	// Must be added in order of elements in the ENetStatsCounterGroupingMode.
+	GroupByOptionsSource.Add(MakeShared<ENetStatsCounterGroupingMode>(ENetStatsCounterGroupingMode::ByType));
+	GroupByOptionsSource.Add(MakeShared<ENetStatsCounterGroupingMode>(ENetStatsCounterGroupingMode::Flat));
+	GroupByOptionsSource.Add(MakeShared<ENetStatsCounterGroupingMode>(ENetStatsCounterGroupingMode::ByName));
+	
+	ENetStatsCounterGroupingModePtr* GroupingModePtrPtr = GroupByOptionsSource.FindByPredicate([&](const ENetStatsCounterGroupingModePtr InGroupingModePtr) { return *InGroupingModePtr == GroupingMode; });
+	if (GroupingModePtrPtr != nullptr)
+	{
+		GroupByComboBox->SetSelectedItem(*GroupingModePtrPtr);
+	}
+
+	GroupByComboBox->RefreshOptions();
+}
+
+////////////////////////////////////////////////////////////////////////////////////////////////////
+
+void SNetStatsCountersView::GroupBy_OnSelectionChanged(TSharedPtr<ENetStatsCounterGroupingMode> NewGroupingMode, ESelectInfo::Type SelectInfo)
+{
+	if (SelectInfo != ESelectInfo::Direct)
+	{
+		GroupingMode = *NewGroupingMode;
+
+		CreateGroups();
+		SortTreeNodes();
+		ApplyFiltering();
+	}
+}
+
+////////////////////////////////////////////////////////////////////////////////////////////////////
+
+TSharedRef<SWidget> SNetStatsCountersView::GroupBy_OnGenerateWidget(TSharedPtr<ENetStatsCounterGroupingMode> InGroupingMode) const
+{
+	return SNew(STextBlock)
+		.Text(NetStatsCounterNodeGroupingHelper::ToText(*InGroupingMode))
+		.ToolTipText(NetStatsCounterNodeGroupingHelper::ToDescription(*InGroupingMode));
+}
+
+////////////////////////////////////////////////////////////////////////////////////////////////////
+
+FText SNetStatsCountersView::GroupBy_GetSelectedText() const
+{
+	return NetStatsCounterNodeGroupingHelper::ToText(GroupingMode);
+}
+
+////////////////////////////////////////////////////////////////////////////////////////////////////
+
+FText SNetStatsCountersView::GroupBy_GetSelectedTooltipText() const
+{
+	return NetStatsCounterNodeGroupingHelper::ToDescription(GroupingMode);
+}
+
+////////////////////////////////////////////////////////////////////////////////////////////////////
+// Sorting
+////////////////////////////////////////////////////////////////////////////////////////////////////
+
+const FName SNetStatsCountersView::GetDefaultColumnBeingSorted()
+{
+	return FNetStatsCountersViewColumns::SumColumnID;
+}
+
+////////////////////////////////////////////////////////////////////////////////////////////////////
+
+const EColumnSortMode::Type SNetStatsCountersView::GetDefaultColumnSortMode()
+{
+	return EColumnSortMode::Descending;
+}
+
+////////////////////////////////////////////////////////////////////////////////////////////////////
+
+void SNetStatsCountersView::CreateSortings()
+{
+	AvailableSorters.Reset();
+	CurrentSorter = nullptr;
+
+	for (const TSharedRef<Insights::FTableColumn>& ColumnRef : Table->GetColumns())
+	{
+		if (ColumnRef->CanBeSorted())
+		{
+			TSharedPtr<Insights::ITableCellValueSorter> SorterPtr = ColumnRef->GetValueSorter();
+			if (ensure(SorterPtr.IsValid()))
+			{
+				AvailableSorters.Add(SorterPtr);
+			}
+		}
+	}
+
+	UpdateCurrentSortingByColumn();
+}
+
+////////////////////////////////////////////////////////////////////////////////////////////////////
+
+void SNetStatsCountersView::UpdateCurrentSortingByColumn()
+{
+	TSharedPtr<Insights::FTableColumn> ColumnPtr = Table->FindColumn(ColumnBeingSorted);
+	CurrentSorter = ColumnPtr.IsValid() ? ColumnPtr->GetValueSorter() : nullptr;
+}
+
+////////////////////////////////////////////////////////////////////////////////////////////////////
+
+void SNetStatsCountersView::SortTreeNodes()
+{
+	if (CurrentSorter.IsValid())
+	{
+		for (FNetStatsCounterNodePtr& Root : GroupNodes)
+		{
+			SortTreeNodesRec(*Root, *CurrentSorter);
+		}
+	}
+}
+
+////////////////////////////////////////////////////////////////////////////////////////////////////
+
+void SNetStatsCountersView::SortTreeNodesRec(FNetStatsCounterNode& Node, const Insights::ITableCellValueSorter& Sorter)
+{
+	if (ColumnSortMode == EColumnSortMode::Type::Descending)
+	{
+		Node.SortChildrenDescending(Sorter);
+	}
+	else // if (ColumnSortMode == EColumnSortMode::Type::Ascending)
+	{
+		Node.SortChildrenAscending(Sorter);
+	}
+
+	for (Insights::FBaseTreeNodePtr ChildPtr : Node.GetChildren())
+	{
+		if (ChildPtr->GetChildren().Num() > 0)
+		{
+			SortTreeNodesRec(*StaticCastSharedPtr<FNetStatsCounterNode>(ChildPtr), Sorter);
+		}
+	}
+}
+
+////////////////////////////////////////////////////////////////////////////////////////////////////
+
+EColumnSortMode::Type SNetStatsCountersView::GetSortModeForColumn(const FName ColumnId) const
+{
+	if (ColumnBeingSorted != ColumnId)
+	{
+		return EColumnSortMode::None;
+	}
+
+	return ColumnSortMode;
+}
+
+////////////////////////////////////////////////////////////////////////////////////////////////////
+
+void SNetStatsCountersView::SetSortModeForColumn(const FName& ColumnId, const EColumnSortMode::Type SortMode)
+{
+	ColumnBeingSorted = ColumnId;
+	ColumnSortMode = SortMode;
+	UpdateCurrentSortingByColumn();
+
+	SortTreeNodes();
+	ApplyFiltering();
+}
+
+////////////////////////////////////////////////////////////////////////////////////////////////////
+
+void SNetStatsCountersView::OnSortModeChanged(const EColumnSortPriority::Type SortPriority, const FName& ColumnId, const EColumnSortMode::Type SortMode)
+{
+	SetSortModeForColumn(ColumnId, SortMode);
+	TreeView_Refresh();
+}
+
+////////////////////////////////////////////////////////////////////////////////////////////////////
+// SortMode action (HeaderMenu)
+////////////////////////////////////////////////////////////////////////////////////////////////////
+
+bool SNetStatsCountersView::HeaderMenu_SortMode_IsChecked(const FName ColumnId, const EColumnSortMode::Type InSortMode)
+{
+	return ColumnBeingSorted == ColumnId && ColumnSortMode == InSortMode;
+}
+
+////////////////////////////////////////////////////////////////////////////////////////////////////
+
+bool SNetStatsCountersView::HeaderMenu_SortMode_CanExecute(const FName ColumnId, const EColumnSortMode::Type InSortMode) const
+{
+	const Insights::FTableColumn& Column = *Table->FindColumnChecked(ColumnId);
+	return Column.CanBeSorted();
+}
+
+////////////////////////////////////////////////////////////////////////////////////////////////////
+
+void SNetStatsCountersView::HeaderMenu_SortMode_Execute(const FName ColumnId, const EColumnSortMode::Type InSortMode)
+{
+	SetSortModeForColumn(ColumnId, InSortMode);
+	TreeView_Refresh();
+}
+
+////////////////////////////////////////////////////////////////////////////////////////////////////
+// SortMode action (ContextMenu)
+////////////////////////////////////////////////////////////////////////////////////////////////////
+
+bool SNetStatsCountersView::ContextMenu_SortMode_IsChecked(const EColumnSortMode::Type InSortMode)
+{
+	return ColumnSortMode == InSortMode;
+}
+
+////////////////////////////////////////////////////////////////////////////////////////////////////
+
+bool SNetStatsCountersView::ContextMenu_SortMode_CanExecute(const EColumnSortMode::Type InSortMode) const
+{
+	return true; //ColumnSortMode != InSortMode;
+}
+
+////////////////////////////////////////////////////////////////////////////////////////////////////
+
+void SNetStatsCountersView::ContextMenu_SortMode_Execute(const EColumnSortMode::Type InSortMode)
+{
+	SetSortModeForColumn(ColumnBeingSorted, InSortMode);
+	TreeView_Refresh();
+}
+
+////////////////////////////////////////////////////////////////////////////////////////////////////
+// SortByColumn action (ContextMenu)
+////////////////////////////////////////////////////////////////////////////////////////////////////
+
+bool SNetStatsCountersView::ContextMenu_SortByColumn_IsChecked(const FName ColumnId)
+{
+	return ColumnId == ColumnBeingSorted;
+}
+
+////////////////////////////////////////////////////////////////////////////////////////////////////
+
+bool SNetStatsCountersView::ContextMenu_SortByColumn_CanExecute(const FName ColumnId) const
+{
+	return true; //ColumnId != ColumnBeingSorted;
+}
+
+////////////////////////////////////////////////////////////////////////////////////////////////////
+
+void SNetStatsCountersView::ContextMenu_SortByColumn_Execute(const FName ColumnId)
+{
+	SetSortModeForColumn(ColumnId, EColumnSortMode::Descending);
+	TreeView_Refresh();
+}
+
+////////////////////////////////////////////////////////////////////////////////////////////////////
+// ShowColumn action
+////////////////////////////////////////////////////////////////////////////////////////////////////
+
+bool SNetStatsCountersView::CanShowColumn(const FName ColumnId) const
+{
+	return true;
+}
+
+////////////////////////////////////////////////////////////////////////////////////////////////////
+
+void SNetStatsCountersView::ShowColumn(const FName ColumnId)
+{
+	Insights::FTableColumn& Column = *Table->FindColumnChecked(ColumnId);
+	Column.Show();
+
+	SHeaderRow::FColumn::FArguments ColumnArgs;
+	ColumnArgs
+		.ColumnId(Column.GetId())
+		.DefaultLabel(Column.GetShortName())
+		.ToolTip(SNetStatsCountersViewTooltip::GetColumnTooltip(Column))
+		.HAlignHeader(Column.GetHorizontalAlignment())
+		.VAlignHeader(VAlign_Center)
+		.HeaderContentPadding(FMargin(2.0f))
+		.HAlignCell(HAlign_Fill)
+		.VAlignCell(VAlign_Fill)
+		.SortMode(this, &SNetStatsCountersView::GetSortModeForColumn, Column.GetId())
+		.OnSort(this, &SNetStatsCountersView::OnSortModeChanged)
+		.FillWidth(Column.GetInitialWidth())
+		//.FixedWidth(Column.IsFixedWidth() ? Column.GetInitialWidth() : TOptional<float>())
+		.HeaderContent()
+		[
+			SNew(SBox)
+			.HeightOverride(24.0f)
+			.Padding(FMargin(0.0f))
+			.VAlign(VAlign_Center)
+			[
+				SNew(STextBlock)
+				.Text(this, &SNetStatsCountersView::GetColumnHeaderText, Column.GetId())
+			]
+		]
+		.MenuContent()
+		[
+			TreeViewHeaderRow_GenerateColumnMenu(Column)
+		];
+
+	int32 ColumnIndex = 0;
+	const int32 NewColumnPosition = Table->GetColumnPositionIndex(ColumnId);
+	const int32 NumColumns = TreeViewHeaderRow->GetColumns().Num();
+	for (; ColumnIndex < NumColumns; ColumnIndex++)
+	{
+		const SHeaderRow::FColumn& CurrentColumn = TreeViewHeaderRow->GetColumns()[ColumnIndex];
+		const int32 CurrentColumnPosition = Table->GetColumnPositionIndex(CurrentColumn.ColumnId);
+		if (NewColumnPosition < CurrentColumnPosition)
+		{
+			break;
+		}
+	}
+
+	TreeViewHeaderRow->InsertColumn(ColumnArgs, ColumnIndex);
+}
+
+////////////////////////////////////////////////////////////////////////////////////////////////////
+// HideColumn action
+////////////////////////////////////////////////////////////////////////////////////////////////////
+
+bool SNetStatsCountersView::CanHideColumn(const FName ColumnId) const
+{
+	const Insights::FTableColumn& Column = *Table->FindColumnChecked(ColumnId);
+	return Column.CanBeHidden();
+}
+
+////////////////////////////////////////////////////////////////////////////////////////////////////
+
+void SNetStatsCountersView::HideColumn(const FName ColumnId)
+{
+	Insights::FTableColumn& Column = *Table->FindColumnChecked(ColumnId);
+	Column.Hide();
+
+	TreeViewHeaderRow->RemoveColumn(ColumnId);
+}
+
+////////////////////////////////////////////////////////////////////////////////////////////////////
+// ToggleColumn action
+////////////////////////////////////////////////////////////////////////////////////////////////////
+
+bool SNetStatsCountersView::IsColumnVisible(const FName ColumnId) const
+{
+	const Insights::FTableColumn& Column = *Table->FindColumnChecked(ColumnId);
+	return Column.IsVisible();
+}
+
+////////////////////////////////////////////////////////////////////////////////////////////////////
+
+bool SNetStatsCountersView::CanToggleColumnVisibility(const FName ColumnId) const
+{
+	const Insights::FTableColumn& Column = *Table->FindColumnChecked(ColumnId);
+	return !Column.IsVisible() || Column.CanBeHidden();
+}
+
+////////////////////////////////////////////////////////////////////////////////////////////////////
+
+void SNetStatsCountersView::ToggleColumnVisibility(const FName ColumnId)
+{
+	const Insights::FTableColumn& Column = *Table->FindColumnChecked(ColumnId);
+	if (Column.IsVisible())
+	{
+		HideColumn(ColumnId);
+	}
+	else
+	{
+		ShowColumn(ColumnId);
+	}
+}
+
+////////////////////////////////////////////////////////////////////////////////////////////////////
+// "Show All Columns" action (ContextMenu)
+////////////////////////////////////////////////////////////////////////////////////////////////////
+
+bool SNetStatsCountersView::ContextMenu_ShowAllColumns_CanExecute() const
+{
+	return true;
+}
+
+////////////////////////////////////////////////////////////////////////////////////////////////////
+
+void SNetStatsCountersView::ContextMenu_ShowAllColumns_Execute()
+{
+	ColumnBeingSorted = GetDefaultColumnBeingSorted();
+	ColumnSortMode = GetDefaultColumnSortMode();
+	UpdateCurrentSortingByColumn();
+
+	for (const TSharedRef<Insights::FTableColumn>& ColumnRef : Table->GetColumns())
+	{
+		const Insights::FTableColumn& Column = *ColumnRef;
+
+		if (!Column.IsVisible())
+		{
+			ShowColumn(Column.GetId());
+		}
+	}
+}
+
+////////////////////////////////////////////////////////////////////////////////////////////////////
+// ResetColumns action (ContextMenu)
+////////////////////////////////////////////////////////////////////////////////////////////////////
+
+bool SNetStatsCountersView::ContextMenu_ResetColumns_CanExecute() const
+{
+	return true;
+}
+
+////////////////////////////////////////////////////////////////////////////////////////////////////
+
+void SNetStatsCountersView::ContextMenu_ResetColumns_Execute()
+{
+	ColumnBeingSorted = GetDefaultColumnBeingSorted();
+	ColumnSortMode = GetDefaultColumnSortMode();
+	UpdateCurrentSortingByColumn();
+
+	for (const TSharedRef<Insights::FTableColumn>& ColumnRef : Table->GetColumns())
+	{
+		const Insights::FTableColumn& Column = *ColumnRef;
+
+		if (Column.ShouldBeVisible() && !Column.IsVisible())
+		{
+			ShowColumn(Column.GetId());
+		}
+		else if (!Column.ShouldBeVisible() && Column.IsVisible())
+		{
+			HideColumn(Column.GetId());
+		}
+	}
+}
+
+////////////////////////////////////////////////////////////////////////////////////////////////////
+
+void SNetStatsCountersView::Reset()
+{
+	GameInstanceIndex = 0;
+	ConnectionIndex = 0;
+	ConnectionMode = TraceServices::ENetProfilerConnectionMode::Outgoing;
+
+	StatsPacketStartIndex = 0;
+	StatsPacketEndIndex = 0;
+
+	RebuildTree(true);
+}
+
+////////////////////////////////////////////////////////////////////////////////////////////////////
+
+void SNetStatsCountersView::Tick(const FGeometry& AllottedGeometry, const double InCurrentTime, const float InDeltaTime)
+{
+	// Check if we need to update the lists of NetStatsCounters, but not too often.
+	const uint64 Time = FPlatformTime::Cycles64();
+	if (Time > NextTimestamp)
+	{
+		RebuildTree(false);
+
+		const uint64 WaitTime = static_cast<uint64>(0.5 / FPlatformTime::GetSecondsPerCycle64()); // 500ms
+		NextTimestamp = Time + WaitTime;
+	}
+}
+
+////////////////////////////////////////////////////////////////////////////////////////////////////
+
+void SNetStatsCountersView::RebuildTree(bool bResync)
+{
+	FStopwatch SyncStopwatch;
+	FStopwatch Stopwatch;
+	Stopwatch.Start();
+
+	if (bResync)
+	{
+		NetStatsCounterNodes.Empty();
+	}
+
+	const uint32 PreviousNodeCount = NetStatsCounterNodes.Num();
+
+	SyncStopwatch.Start();
+	if (Session.IsValid())
+	{
+		TraceServices::FAnalysisSessionReadScope SessionReadScope(*Session.Get());
+		const TraceServices::INetProfilerProvider* NetProfilerProvider = TraceServices::ReadNetProfilerProvider(*Session.Get());
+		if (NetProfilerProvider)
+		{
+			NetProfilerProvider->ReadNetStatsCounterTypes([this, &bResync, NetProfilerProvider](const TraceServices::FNetProfilerStatsCounterType* NetStatsCounterTypes, uint64 InNetStatsCounterTypeCount)
+			{
+				const uint32 NetStatsCounterTypeCount = static_cast<uint32>(InNetStatsCounterTypeCount);
+				if (NetStatsCounterTypeCount != NetStatsCounterNodes.Num())
+				{
+					NetStatsCounterNodes.Empty(NetStatsCounterTypeCount);
+					for (uint32 Index = 0; Index < NetStatsCounterTypeCount; ++Index)
+					{
+						const TraceServices::FNetProfilerStatsCounterType& NetStatsCounterType = NetStatsCounterTypes[Index];
+						ensure(NetStatsCounterType.StatsCounterTypeIndex == Index);
+						const TCHAR* NamePtr = nullptr;
+						NetProfilerProvider->ReadName(NetStatsCounterType.NameIndex, [&NamePtr](const TraceServices::FNetProfilerName& InName)
+						{
+							NamePtr = InName.Name;
+						});
+						const FName Name(NamePtr);
+						FNetStatsCounterNodePtr NetEventNodePtr = MakeShared<FNetStatsCounterNode>(NetStatsCounterType.StatsCounterTypeIndex, Name, NetStatsCounterType.Type);
+						NetStatsCounterNodes.Add(NetEventNodePtr);
+					}
+				}
+			});
+		}
+	}
+	SyncStopwatch.Stop();
+
+	if (bResync || NetStatsCounterNodes.Num() != PreviousNodeCount)
+	{
+		UpdateTree();
+		UpdateStatsInternal();
+
+		// Save selection.
+		TArray<FNetStatsCounterNodePtr> SelectedItems;
+		TreeView->GetSelectedItems(SelectedItems);
+
+		TreeView->RebuildList();
+
+		// Restore selection.
+		if (SelectedItems.Num() > 0)
+		{
+			TreeView->ClearSelection();
+			for (FNetStatsCounterNodePtr& NodePtr : SelectedItems)
+			{
+				NodePtr = GetNetStatsCounterNode(NodePtr->GetCounterTypeIndex());
+			}
+			SelectedItems.RemoveAll([](const FNetStatsCounterNodePtr& NodePtr) { return !NodePtr.IsValid(); });
+			if (SelectedItems.Num() > 0)
+			{
+				TreeView->SetItemSelection(SelectedItems, true);
+				TreeView->RequestScrollIntoView(SelectedItems.Last());
+			}
+		}
+	}
+
+	Stopwatch.Stop();
+	const double TotalTime = Stopwatch.GetAccumulatedTime();
+	if (TotalTime > 0.01)
+	{
+		const double SyncTime = SyncStopwatch.GetAccumulatedTime();
+		UE_LOG(NetworkingProfiler, Log, TEXT("[NetStatsCounters] Tree view rebuilt in %.4fs (%.4fs + %.4fs) --> %d net stats counters (%d added)"),
+			TotalTime, SyncTime, TotalTime - SyncTime, NetStatsCounterNodes.Num(), NetStatsCounterNodes.Num() - PreviousNodeCount);
+	}
+}
+
+////////////////////////////////////////////////////////////////////////////////////////////////////
+
+void SNetStatsCountersView::ResetStats()
+{
+	GameInstanceIndex = 0;
+	ConnectionIndex = 0;
+	ConnectionMode = TraceServices::ENetProfilerConnectionMode::Outgoing;
+
+	StatsPacketStartIndex = 0;
+	StatsPacketEndIndex = 0;
+
+	UpdateStatsInternal();
+}
+
+////////////////////////////////////////////////////////////////////////////////////////////////////
+
+void SNetStatsCountersView::UpdateStats(uint32 InGameInstanceIndex, uint32 InConnectionIndex, TraceServices::ENetProfilerConnectionMode InConnectionMode, uint32 InStatsPacketStartIndex, uint32 InStatsPacketEndIndex)
+{
+	GameInstanceIndex = InGameInstanceIndex;
+	ConnectionIndex = InConnectionIndex;
+	ConnectionMode = InConnectionMode;
+
+	StatsPacketStartIndex = InStatsPacketStartIndex;
+	StatsPacketEndIndex = InStatsPacketEndIndex;
+
+	UpdateStatsInternal();
+}
+
+////////////////////////////////////////////////////////////////////////////////////////////////////
+
+void SNetStatsCountersView::UpdateStatsInternal()
+{
+	if (StatsPacketStartIndex >= StatsPacketEndIndex)
+	{
+		// keep previous aggregated stats
+		return;
+	}
+
+	FStopwatch AggregationStopwatch;
+	FStopwatch Stopwatch;
+	Stopwatch.Start();
+
+	for (const FNetStatsCounterNodePtr& NetEventNodePtr : NetStatsCounterNodes)
+	{
+		NetEventNodePtr->ResetAggregatedStats();
+	}
+
+	if (Session.IsValid())
+	{
+		TUniquePtr<TraceServices::ITable<TraceServices::FNetProfilerAggregatedStatsCounterStats>> AggregationResultTable;
+
+		AggregationStopwatch.Start();
+		{
+			TraceServices::FAnalysisSessionReadScope SessionReadScope(*Session.Get());
+			const TraceServices::INetProfilerProvider* NetProfilerProvider = TraceServices::ReadNetProfilerProvider(*Session.Get());
+			if (NetProfilerProvider)
+			{
+				// CreateAggregation requires [PacketStartIndex, PacketEndIndex] as inclusive interval
+				AggregationResultTable.Reset(NetProfilerProvider->CreateStatsCountersAggregation(ConnectionIndex, ConnectionMode, StatsPacketStartIndex, StatsPacketEndIndex - 1));
+			}
+		}
+		AggregationStopwatch.Stop();
+
+		if (AggregationResultTable.IsValid())
+		{
+			TUniquePtr<TraceServices::ITableReader<TraceServices::FNetProfilerAggregatedStatsCounterStats>> TableReader(AggregationResultTable->CreateReader());
+			while (TableReader->IsValid())
+			{
+				const TraceServices::FNetProfilerAggregatedStatsCounterStats* Row = TableReader->GetCurrentRow();
+				FNetStatsCounterNodePtr NetStatsCounterNodePtr = GetNetStatsCounterNode(Row->StatsCounterTypeIndex);
+				if (NetStatsCounterNodePtr)
+				{
+					NetStatsCounterNodePtr->SetAggregatedStats(*Row);
+
+					TSharedPtr<ITableRow> TableRowPtr = TreeView->WidgetFromItem(NetStatsCounterNodePtr);
+					if (TableRowPtr.IsValid())
+					{
+						TSharedPtr<SNetStatsCountersTableRow> RowPtr = StaticCastSharedPtr<SNetStatsCountersTableRow, ITableRow>(TableRowPtr);
+						RowPtr->InvalidateContent();
+					}
+				}
+				TableReader->NextRow();
+			}
+		}
+	}
+
+	UpdateTree();
+
+	const TArray<FNetStatsCounterNodePtr> SelectedNodes = TreeView->GetSelectedItems();
+	if (SelectedNodes.Num() > 0)
+	{
+		TreeView->RequestScrollIntoView(SelectedNodes[0]);
+	}
+
+	Stopwatch.Stop();
+	const double TotalTime = Stopwatch.GetAccumulatedTime();
+	const double AggregationTime = AggregationStopwatch.GetAccumulatedTime();
+	UE_LOG(NetworkingProfiler, Log, TEXT("[NetStats] Aggregated stats updated in %.4fs (%.4fs + %.4fs)"),
+		TotalTime, AggregationTime, TotalTime - AggregationTime);
+}
+
+////////////////////////////////////////////////////////////////////////////////////////////////////
+
+FNetStatsCounterNodePtr SNetStatsCountersView::GetNetStatsCounterNode(uint32 EventTypeIndex) const
+{
+	return (EventTypeIndex < (uint32)NetStatsCounterNodes.Num()) ? NetStatsCounterNodes[EventTypeIndex] : nullptr;
+}
+
+////////////////////////////////////////////////////////////////////////////////////////////////////
+
+void SNetStatsCountersView::SelectNetStatsCounterNode(uint32 EventTypeIndex)
+{
+	FNetStatsCounterNodePtr NodePtr = GetNetStatsCounterNode(EventTypeIndex);
+	if (NodePtr)
+	{
+		TreeView->SetSelection(NodePtr);
+		TreeView->RequestScrollIntoView(NodePtr);
+	}
+}
+
+////////////////////////////////////////////////////////////////////////////////////////////////////
+
+#undef LOCTEXT_NAMESPACE
diff --git a/Engine/Source/Developer/TraceInsights/Private/Insights/NetworkingProfiler/Widgets/SNetStatsCountersView.h b/Engine/Source/Developer/TraceInsights/Private/Insights/NetworkingProfiler/Widgets/SNetStatsCountersView.h
new file mode 100644
--- /dev/null
+++ b/Engine/Source/Developer/TraceInsights/Private/Insights/NetworkingProfiler/Widgets/SNetStatsCountersView.h
@@ -0,0 +1,366 @@
+// Copyright Epic Games, Inc. All Rights Reserved.
+
+#pragma once
+
+#include "CoreMinimal.h"
+#include "Misc/FilterCollection.h"
+#include "Misc/TextFilter.h"
+#include "SlateFwd.h"
+#include "Widgets/DeclarativeSyntaxSupport.h"
+#include "Widgets/Input/SComboBox.h"
+#include "Widgets/SCompoundWidget.h"
+#include "Widgets/SWidget.h"
+#include "Widgets/Views/SHeaderRow.h"
+#include "Widgets/Views/STableRow.h"
+#include "Widgets/Views/STableViewBase.h"
+#include "Widgets/Views/STreeView.h"
+
+// Insights
+#include "Insights/NetworkingProfiler/ViewModels/NetStatsCounterGroupingAndSorting.h"
+#include "Insights/NetworkingProfiler/ViewModels/NetStatsCounterNode.h"
+#include "TraceServices/Model/NetProfiler.h"
+
+////////////////////////////////////////////////////////////////////////////////////////////////////
+
+class FMenuBuilder;
+class SNetworkingProfilerWindow;
+
+namespace TraceServices
+{
+	class IAnalysisSession;
+}
+
+namespace Insights
+{
+	class FTable;
+	class FTableColumn;
+	class ITableCellValueSorter;
+}
+
+////////////////////////////////////////////////////////////////////////////////////////////////////
+
+/** The filter collection - used for updating the list of tree nodes. */
+typedef TFilterCollection<const FNetStatsCounterNodePtr&> FNetStatsCounterNodeFilterCollection;
+
+/** The text based filter - used for updating the list of tree nodes. */
+typedef TTextFilter<const FNetStatsCounterNodePtr&> FNetStatsCounterNodeTextFilter;
+
+////////////////////////////////////////////////////////////////////////////////////////////////////
+/**
+ * A custom widget used to display the list of NetStatsCounters and their aggregated stats.
+ */
+class SNetStatsCountersView : public SCompoundWidget
+{
+public:
+	/** Default constructor. */
+	SNetStatsCountersView();
+
+	/** Virtual destructor. */
+	virtual ~SNetStatsCountersView();
+
+	SLATE_BEGIN_ARGS(SNetStatsCountersView) {}
+	SLATE_END_ARGS()
+
+	/**
+	 * Construct this widget
+	 * @param InArgs - The declaration data for this widget
+	 */
+	void Construct(const FArguments& InArgs, TSharedPtr<SNetworkingProfilerWindow> InProfilerWindow);
+
+	TSharedPtr<Insights::FTable> GetTable() const { return Table; }
+
+	/**
+	 * Ticks this widget. Override in derived classes, but always call the parent implementation.
+	 *
+	 * @param  AllottedGeometry - the space allotted for this widget
+	 * @param  InCurrentTime    - current absolute real time
+	 * @param  InDeltaTime      - real time passed since last tick
+	 */
+	virtual void Tick(const FGeometry& AllottedGeometry, const double InCurrentTime, const float InDeltaTime) override;
+
+	void Reset();
+
+	/**
+	 * Rebuilds the tree (if necessary).
+	 * @param bResync - If true, it forces a resync with list of NetStatsCounters from Analysis, even if the list did not changed since last sync.
+	 */
+	void RebuildTree(bool bResync);
+
+	void ResetStats();
+	void UpdateStats(uint32 InGameInstanceIndex, uint32 InConnectionIndex, TraceServices::ENetProfilerConnectionMode InConnectionMode, uint32 InStatsPacketStartIndex, uint32 InStatsPacketEndIndex);
+
+	FNetStatsCounterNodePtr GetNetStatsCounterNode(uint32 TypeIndex) const;
+	void SelectNetStatsCounterNode(uint32 TypeIndex);
+
+protected:
+	void UpdateTree();
+
+	void UpdateStatsInternal();
+
+	/** Called when the analysis session has changed. */
+	void InsightsManager_OnSessionChanged();
+
+	/**
+	 * Populates OutSearchStrings with the strings that should be used in searching.
+	 *
+	 * @param GroupOrStatNodePtr - the group and stat node to get a text description from.
+	 * @param OutSearchStrings   - an array of strings to use in searching.
+	 *
+	 */
+	void HandleItemToStringArray(const FNetStatsCounterNodePtr& GroupOrStatNodePtr, TArray<FString>& OutSearchStrings) const;
+
+	////////////////////////////////////////////////////////////////////////////////////////////////////
+	// Tree View - Context Menu
+
+	TSharedPtr<SWidget> TreeView_GetMenuContent();
+	void TreeView_BuildSortByMenu(FMenuBuilder& MenuBuilder);
+	void TreeView_BuildViewColumnMenu(FMenuBuilder& MenuBuilder);
+
+	////////////////////////////////////////////////////////////////////////////////////////////////////
+	// Tree View - Columns' Header
+
+	void InitializeAndShowHeaderColumns();
+
+	FText GetColumnHeaderText(const FName ColumnId) const;
+
+	TSharedRef<SWidget> TreeViewHeaderRow_GenerateColumnMenu(const Insights::FTableColumn& Column);
+
+	////////////////////////////////////////////////////////////////////////////////////////////////////
+	// Tree View - Misc
+
+	void TreeView_Refresh();
+
+	/**
+	 * Called by STreeView to retrieves the children for the specified parent item.
+	 * @param InParent    - The parent node to retrieve the children from.
+	 * @param OutChildren - List of children for the parent node.
+	 */
+	void TreeView_OnGetChildren(FNetStatsCounterNodePtr InParent, TArray<FNetStatsCounterNodePtr>& OutChildren);
+
+	/** Called by STreeView when selection has changed. */
+	void TreeView_OnSelectionChanged(FNetStatsCounterNodePtr SelectedItem, ESelectInfo::Type SelectInfo);
+
+	/** Called by STreeView when a tree item is double clicked. */
+	void TreeView_OnMouseButtonDoubleClick(FNetStatsCounterNodePtr TreeNode);
+
+	////////////////////////////////////////////////////////////////////////////////////////////////////
+	// Tree View - Table Row
+
+	/** Called by STreeView to generate a table row for the specified item. */
+	TSharedRef<ITableRow> TreeView_OnGenerateRow(FNetStatsCounterNodePtr TreeNode, const TSharedRef<STableViewBase>& OwnerTable);
+
+	bool TableRow_ShouldBeEnabled(FNetStatsCounterNodePtr NodePtr) const;
+
+	void TableRow_SetHoveredCell(TSharedPtr<Insights::FTable> TablePtr, TSharedPtr<Insights::FTableColumn> ColumnPtr, FNetStatsCounterNodePtr NodePtr);
+	EHorizontalAlignment TableRow_GetColumnOutlineHAlignment(const FName ColumnId) const;
+
+	FText TableRow_GetHighlightText() const;
+	FName TableRow_GetHighlightedNodeName() const;
+
+	////////////////////////////////////////////////////////////////////////////////////////////////////
+	// Filtering
+
+	/** Populates the group and stat tree with items based on the current data. */
+	void ApplyFiltering();
+
+	void FilterOutZeroCountStatsCounters_OnCheckStateChanged(ECheckBoxState NewRadioState);
+	ECheckBoxState FilterOutZeroCountStatsCounters_IsChecked() const;
+
+	TSharedRef<SWidget> GetToggleButtonForNetStatsCounterType(const ENetStatsCounterNodeType InNetStatsCounterType);
+	void FilterByNetStatsCounterType_OnCheckStateChanged(ECheckBoxState NewRadioState, const ENetStatsCounterNodeType InNetStatsCounterType);
+	ECheckBoxState FilterByNetStatsCounterType_IsChecked(const ENetStatsCounterNodeType InNetStatsCounterType) const;
+
+	bool SearchBox_IsEnabled() const;
+	void SearchBox_OnTextChanged(const FText& InFilterText);
+
+	////////////////////////////////////////////////////////////////////////////////////////////////////
+	// Grouping
+
+	void CreateGroups();
+	void CreateGroupByOptionsSources();
+
+	void GroupBy_OnSelectionChanged(TSharedPtr<ENetStatsCounterGroupingMode> NewGroupingMode, ESelectInfo::Type SelectInfo);
+
+	TSharedRef<SWidget> GroupBy_OnGenerateWidget(TSharedPtr<ENetStatsCounterGroupingMode> InGroupingMode) const;
+
+	FText GroupBy_GetSelectedText() const;
+
+	FText GroupBy_GetSelectedTooltipText() const;
+
+	////////////////////////////////////////////////////////////////////////////////////////////////////
+	// Sorting
+
+	static const FName GetDefaultColumnBeingSorted();
+	static const EColumnSortMode::Type GetDefaultColumnSortMode();
+
+	void CreateSortings();
+
+	void UpdateCurrentSortingByColumn();
+	void SortTreeNodes();
+	void SortTreeNodesRec(FNetStatsCounterNode& Node, const Insights::ITableCellValueSorter& Sorter);
+
+	EColumnSortMode::Type GetSortModeForColumn(const FName ColumnId) const;
+	void SetSortModeForColumn(const FName& ColumnId, EColumnSortMode::Type SortMode);
+	void OnSortModeChanged(const EColumnSortPriority::Type SortPriority, const FName& ColumnId, const EColumnSortMode::Type SortMode);
+
+	////////////////////////////////////////////////////////////////////////////////////////////////////
+	// Sorting actions
+
+	// SortMode (HeaderMenu)
+	bool HeaderMenu_SortMode_IsChecked(const FName ColumnId, const EColumnSortMode::Type InSortMode);
+	bool HeaderMenu_SortMode_CanExecute(const FName ColumnId, const EColumnSortMode::Type InSortMode) const;
+	void HeaderMenu_SortMode_Execute(const FName ColumnId, const EColumnSortMode::Type InSortMode);
+
+	// SortMode (ContextMenu)
+	bool ContextMenu_SortMode_IsChecked(const EColumnSortMode::Type InSortMode);
+	bool ContextMenu_SortMode_CanExecute(const EColumnSortMode::Type InSortMode) const;
+	void ContextMenu_SortMode_Execute(const EColumnSortMode::Type InSortMode);
+
+	// SortByColumn (ContextMenu)
+	bool ContextMenu_SortByColumn_IsChecked(const FName ColumnId);
+	bool ContextMenu_SortByColumn_CanExecute(const FName ColumnId) const;
+	void ContextMenu_SortByColumn_Execute(const FName ColumnId);
+
+	////////////////////////////////////////////////////////////////////////////////////////////////////
+	// Column visibility actions
+
+	// ShowColumn
+	bool CanShowColumn(const FName ColumnId) const;
+	void ShowColumn(const FName ColumnId);
+
+	// HideColumn
+	bool CanHideColumn(const FName ColumnId) const;
+	void HideColumn(const FName ColumnId);
+
+	// ToggleColumnVisibility
+	bool IsColumnVisible(const FName ColumnId) const;
+	bool CanToggleColumnVisibility(const FName ColumnId) const;
+	void ToggleColumnVisibility(const FName ColumnId);
+
+	// ShowAllColumns (ContextMenu)
+	bool ContextMenu_ShowAllColumns_CanExecute() const;
+	void ContextMenu_ShowAllColumns_Execute();
+
+	// MinMaxMedColumns (ContextMenu)
+	bool ContextMenu_ShowMinMaxMedColumns_CanExecute() const;
+	void ContextMenu_ShowMinMaxMedColumns_Execute();
+
+	// ResetColumns (ContextMenu)
+	bool ContextMenu_ResetColumns_CanExecute() const;
+	void ContextMenu_ResetColumns_Execute();
+
+	////////////////////////////////////////////////////////////////////////////////////////////////////
+
+protected:
+	TSharedPtr<SNetworkingProfilerWindow> ProfilerWindow;
+
+	/** Table view model. */
+	TSharedPtr<Insights::FTable> Table;
+
+	/** A weak pointer to the profiler session used to populate this widget. */
+	TSharedPtr<const TraceServices::IAnalysisSession>/*Weak*/ Session;
+
+	//////////////////////////////////////////////////
+	// Tree View, Columns
+
+	/** The tree widget which holds the NetStatsCounters tree nodes. */
+	TSharedPtr<STreeView<FNetStatsCounterNodePtr>> TreeView;
+
+	/** Holds the tree view header row widget which display all columns in the tree view. */
+	TSharedPtr<SHeaderRow> TreeViewHeaderRow;
+
+	/** External scrollbar used to synchronize tree view position. */
+	TSharedPtr<SScrollBar> ExternalScrollbar;
+
+	//////////////////////////////////////////////////
+	// Hovered Column, Hovered Node
+
+	/** Name of the column currently being hovered by the mouse. */
+	FName HoveredColumnId;
+
+	/** A shared pointer to the node currently being hovered by the mouse. */
+	FNetStatsCounterNodePtr HoveredNodePtr;
+
+	/** Name of the node that should be drawn as highlighted. */
+	FName HighlightedNodeName;
+
+	//////////////////////////////////////////////////
+	// NetStatsCounters Node
+
+	/** An array of group nodes. */
+	TArray<FNetStatsCounterNodePtr> GroupNodes;
+
+	/** A filtered array of group nodes to be displayed in the tree widget. */
+	TArray<FNetStatsCounterNodePtr> FilteredGroupNodes;
+
+	/** All net event nodes. Index in this array is EventTypeIndex. */
+	TArray<FNetStatsCounterNodePtr> NetStatsCounterNodes;
+
+	/** Currently expanded group nodes. */
+	TSet<FNetStatsCounterNodePtr> ExpandedNodes;
+
+	/** If true, the expanded nodes have been saved before applying a text filter. */
+	bool bExpansionSaved;
+
+	//////////////////////////////////////////////////
+	// Search box and filters
+
+	//bool bUseFiltering;
+
+	/** The search box widget used to filter items displayed in the tree. */
+	TSharedPtr<SSearchBox> SearchBox;
+
+	/** The text based filter. */
+	TSharedPtr<FNetStatsCounterNodeTextFilter> TextFilter;
+
+	/** The filter collection. */
+	TSharedPtr<FNetStatsCounterNodeFilterCollection> Filters;
+
+	/** Holds the visibility of each net event type. */
+	bool bStatsCounterTypeIsVisible[static_cast<int>(ENetStatsCounterNodeType::InvalidOrMax)];
+
+	/** Filter out the net event types having zero total instance count (aggregated stats). */
+	bool bFilterOutZeroCountStatsCounters;
+
+	//////////////////////////////////////////////////
+	// Grouping
+
+	//bool bUseGrouping;
+
+	TArray<TSharedPtr<ENetStatsCounterGroupingMode>> GroupByOptionsSource;
+
+	TSharedPtr<SComboBox<TSharedPtr<ENetStatsCounterGroupingMode>>> GroupByComboBox;
+
+	/** How we group the net event nodes? */
+	ENetStatsCounterGroupingMode GroupingMode;
+
+	//////////////////////////////////////////////////
+	// Sorting
+
+	//bool bUseSorting;
+
+	/** All available sorters. */
+	TArray<TSharedPtr<Insights::ITableCellValueSorter>> AvailableSorters;
+
+	/** Current sorter. It is nullptr if sorting is disabled. */
+	TSharedPtr<Insights::ITableCellValueSorter> CurrentSorter;
+
+	/** Name of the column currently being sorted. Can be NAME_None if sorting is disabled (CurrentSorting == nullptr) or if a complex sorting is used (CurrentSorting != nullptr). */
+	FName ColumnBeingSorted;
+
+	/** How we sort the nodes? Ascending or Descending. */
+	EColumnSortMode::Type ColumnSortMode;
+
+	//////////////////////////////////////////////////
+
+	uint64 NextTimestamp;
+	uint32 ObjectsChangeCount;
+
+	uint32 GameInstanceIndex;
+	uint32 ConnectionIndex;
+	TraceServices::ENetProfilerConnectionMode ConnectionMode;
+	uint32 StatsPacketStartIndex;
+	uint32 StatsPacketEndIndex;
+};
+
+////////////////////////////////////////////////////////////////////////////////////////////////////
diff --git a/Engine/Source/Developer/TraceInsights/Private/Insights/NetworkingProfiler/Widgets/SNetStatsCountersViewTooltip.cpp b/Engine/Source/Developer/TraceInsights/Private/Insights/NetworkingProfiler/Widgets/SNetStatsCountersViewTooltip.cpp
new file mode 100644
--- /dev/null
+++ b/Engine/Source/Developer/TraceInsights/Private/Insights/NetworkingProfiler/Widgets/SNetStatsCountersViewTooltip.cpp
@@ -0,0 +1,273 @@
+// Copyright Epic Games, Inc. All Rights Reserved.
+
+#include "SNetStatsCountersViewTooltip.h"
+
+#include "SlateOptMacros.h"
+#include "TraceServices/Model/NetProfiler.h"
+#include "Widgets/Layout/SGridPanel.h"
+#include "Widgets/Layout/SSeparator.h"
+#include "Widgets/SBoxPanel.h"
+#include "Widgets/SToolTip.h"
+#include "Widgets/Text/STextBlock.h"
+
+// Insights
+#include "Insights/InsightsStyle.h"
+#include "Insights/Table/ViewModels/Table.h"
+#include "Insights/Table/ViewModels/TableColumn.h"
+#include "Insights/NetworkingProfiler/ViewModels/NetStatsCounterNode.h"
+#include "Insights/NetworkingProfiler/ViewModels/NetStatsCounterNodeHelper.h"
+
+////////////////////////////////////////////////////////////////////////////////////////////////////
+
+#define LOCTEXT_NAMESPACE "SNetStatsCountersView"
+
+////////////////////////////////////////////////////////////////////////////////////////////////////
+
+BEGIN_SLATE_FUNCTION_BUILD_OPTIMIZATION
+
+TSharedPtr<SToolTip> SNetStatsCountersViewTooltip::GetTableTooltip(const Insights::FTable& Table)
+{
+	TSharedPtr<SToolTip> ColumnTooltip =
+		SNew(SToolTip)
+		[
+			SNew(SVerticalBox)
+
+			+ SVerticalBox::Slot()
+			.AutoHeight()
+			.Padding(2.0f)
+			[
+				SNew(STextBlock)
+				.Text(Table.GetDisplayName())
+				.TextStyle(FInsightsStyle::Get(), TEXT("TreeTable.TooltipBold"))
+			]
+
+			+ SVerticalBox::Slot()
+			.AutoHeight()
+			.Padding(2.0f)
+			[
+				SNew(STextBlock)
+				.Text(Table.GetDescription())
+				.TextStyle(FInsightsStyle::Get(), TEXT("TreeTable.Tooltip"))
+			]
+		];
+
+	return ColumnTooltip;
+}
+
+////////////////////////////////////////////////////////////////////////////////////////////////////
+
+TSharedPtr<SToolTip> SNetStatsCountersViewTooltip::GetColumnTooltip(const Insights::FTableColumn& Column)
+{
+	TSharedPtr<SToolTip> ColumnTooltip =
+		SNew(SToolTip)
+		[
+			SNew(SVerticalBox)
+
+			+ SVerticalBox::Slot()
+			.AutoHeight()
+			.Padding(2.0f)
+			[
+				SNew(STextBlock)
+				.Text(Column.GetTitleName())
+				.TextStyle(FInsightsStyle::Get(), TEXT("TreeTable.TooltipBold"))
+			]
+
+			+ SVerticalBox::Slot()
+			.AutoHeight()
+			.Padding(2.0f)
+			[
+				SNew(STextBlock)
+				.Text(Column.GetDescription())
+				.TextStyle(FInsightsStyle::Get(), TEXT("TreeTable.Tooltip"))
+			]
+		];
+
+	return ColumnTooltip;
+}
+
+////////////////////////////////////////////////////////////////////////////////////////////////////
+
+TSharedPtr<SToolTip> SNetStatsCountersViewTooltip::GetRowTooltip(const TSharedPtr<FNetStatsCounterNode> NetStatsCounterNodePtr)
+{
+	const TraceServices::FNetProfilerAggregatedStatsCounterStats& Stats = NetStatsCounterNodePtr->GetAggregatedStats();
+	const FText InstanceCountText = FText::AsNumber(Stats.Count);
+
+	FText SumText = NetStatsCounterNodePtr->GetTextForAggregatedStatsSum(true);
+	FText MinText = NetStatsCounterNodePtr->GetTextForAggregatedStatsMin(true);
+	FText MaxText = NetStatsCounterNodePtr->GetTextForAggregatedStatsMax(true);
+	FText AvgText = NetStatsCounterNodePtr->GetTextForAggregatedStatsAverage(true);
+
+	TSharedPtr<SGridPanel> GridPanel;
+	TSharedPtr<SHorizontalBox> HBox;
+
+	TSharedPtr<SToolTip> TableCellTooltip =
+		SNew(SToolTip)
+		[
+			SAssignNew(HBox, SHorizontalBox)
+
+			+ SHorizontalBox::Slot()
+		.AutoWidth()
+		[
+			SNew(SVerticalBox)
+
+			+ SVerticalBox::Slot()
+		.AutoHeight()
+		.Padding(2.0f)
+		[
+			SNew(SSeparator)
+			.Orientation(Orient_Horizontal)
+		]
+
+	+ SVerticalBox::Slot()
+		.AutoHeight()
+		.Padding(2.0f)
+		[
+			SNew(SGridPanel)
+
+			// Id: [Id]
+		+ SGridPanel::Slot(0, 0)
+		.Padding(2.0f)
+		[
+			SNew(STextBlock)
+			.Text(LOCTEXT("TT_Id", "Id:"))
+		.TextStyle(FInsightsStyle::Get(), TEXT("TreeTable.TooltipBold"))
+		]
+	+ SGridPanel::Slot(1, 0)
+		.Padding(2.0f)
+		[
+			SNew(STextBlock)
+			.Text(FText::AsNumber(NetStatsCounterNodePtr->GetCounterTypeIndex()))
+		.TextStyle(FInsightsStyle::Get(), TEXT("TreeTable.Tooltip"))
+		]
+
+	// Name: [Name]
+	+ SGridPanel::Slot(0, 1)
+		.Padding(2.0f)
+		[
+			SNew(STextBlock)
+			.Text(LOCTEXT("TT_Name", "Name:"))
+		.TextStyle(FInsightsStyle::Get(), TEXT("TreeTable.TooltipBold"))
+		]
+	+ SGridPanel::Slot(1, 1)
+		.Padding(2.0f)
+		[
+			SNew(STextBlock)
+			.WrapTextAt(512.0f)
+		.WrappingPolicy(ETextWrappingPolicy::AllowPerCharacterWrapping)
+		.Text(FText::FromName(NetStatsCounterNodePtr->GetName()))
+		.TextStyle(FInsightsStyle::Get(), TEXT("TreeTable.Tooltip"))
+		]
+
+	// Counter Type: [Type]
+	+ SGridPanel::Slot(0, 3)
+		.Padding(2.0f)
+		[
+			SNew(STextBlock)
+			.Text(LOCTEXT("TT_Type", "Node Type:"))
+		.TextStyle(FInsightsStyle::Get(), TEXT("TreeTable.TooltipBold"))
+		]
+	+ SGridPanel::Slot(1, 3)
+		.Padding(2.0f)
+		[
+			SNew(STextBlock)
+			.Text(NetStatsCounterNodeTypeHelper::ToText(NetStatsCounterNodePtr->GetType()))
+		.TextStyle(FInsightsStyle::Get(), TEXT("TreeTable.Tooltip"))
+		]
+		]
+
+	+ SVerticalBox::Slot()
+		.AutoHeight()
+		.Padding(2.0f)
+		[
+			SNew(SSeparator)
+			.Orientation(Orient_Horizontal)
+		]
+
+	+ SVerticalBox::Slot()
+		.AutoHeight()
+		.Padding(2.0f)
+		[
+			SNew(SGridPanel)
+
+			+ SGridPanel::Slot(0, 0)
+		.Padding(2.0f)
+		[
+			SNew(STextBlock)
+			.Text(LOCTEXT("TT_NumInstances", "Num Instances:"))
+		.TextStyle(FInsightsStyle::Get(), TEXT("TreeTable.TooltipBold"))
+		]
+	+ SGridPanel::Slot(1, 0)
+		.Padding(2.0f)
+		[
+			SNew(STextBlock)
+			.Text(InstanceCountText)
+		.TextStyle(FInsightsStyle::Get(), TEXT("TreeTable.Tooltip"))
+		]
+		]
+
+	+ SVerticalBox::Slot()
+		.AutoHeight()
+		.Padding(2.0f)
+		[
+			SNew(SSeparator)
+			.Orientation(Orient_Horizontal)
+		]
+
+	+ SVerticalBox::Slot()
+		.AutoHeight()
+		.Padding(2.0f)
+		[
+			SAssignNew(GridPanel, SGridPanel)
+
+			// Aggregated stats are added here.
+		]
+
+	+ SVerticalBox::Slot()
+		.AutoHeight()
+		.Padding(2.0f)
+		[
+			SNew(SSeparator)
+			.Orientation(Orient_Horizontal)
+		]
+		]
+		];
+
+	int32 Row = 1;
+	AddAggregatedStatsRow(GridPanel, Row, LOCTEXT("TT_Sum", "Sum:"), SumText);
+	AddAggregatedStatsRow(GridPanel, Row, LOCTEXT("TT_Max", "Max:"), MaxText);
+	AddAggregatedStatsRow(GridPanel, Row, LOCTEXT("TT_Average", "Average:"), AvgText);
+	AddAggregatedStatsRow(GridPanel, Row, LOCTEXT("TT_Min", "Min:"), MinText);
+
+	return TableCellTooltip;
+}
+
+////////////////////////////////////////////////////////////////////////////////////////////////////
+
+void SNetStatsCountersViewTooltip::AddAggregatedStatsRow(TSharedPtr<SGridPanel> Grid, int32& Row, const FText& Name, const FText& Value)
+{
+	Grid->AddSlot(0, Row)
+		.Padding(2.0f)
+		[
+			SNew(STextBlock)
+			.Text(Name)
+		.TextStyle(FInsightsStyle::Get(), TEXT("TreeTable.TooltipBold"))
+		];
+
+	Grid->AddSlot(1, Row)
+		.Padding(2.0f)
+		.HAlign(HAlign_Left)
+		[
+			SNew(STextBlock)
+			.Text(Value)
+		.TextStyle(FInsightsStyle::Get(), TEXT("TreeTable.Tooltip"))
+		.ColorAndOpacity(Row == 1 ? FLinearColor::Gray : FLinearColor::White)
+		];
+
+	Row++;
+}
+
+END_SLATE_FUNCTION_BUILD_OPTIMIZATION
+
+////////////////////////////////////////////////////////////////////////////////////////////////////
+
+#undef LOCTEXT_NAMESPACE
diff --git a/Engine/Source/Developer/TraceInsights/Private/Insights/NetworkingProfiler/Widgets/SNetStatsCountersViewTooltip.h b/Engine/Source/Developer/TraceInsights/Private/Insights/NetworkingProfiler/Widgets/SNetStatsCountersViewTooltip.h
new file mode 100644
--- /dev/null
+++ b/Engine/Source/Developer/TraceInsights/Private/Insights/NetworkingProfiler/Widgets/SNetStatsCountersViewTooltip.h
@@ -0,0 +1,85 @@
+// Copyright Epic Games, Inc. All Rights Reserved.
+
+#pragma once
+
+#include "CoreMinimal.h"
+#include "Widgets/IToolTip.h"
+#include "Widgets/SToolTip.h"
+
+class SGridPanel;
+
+namespace Insights
+{
+	class FTable;
+	class FTableColumn;
+}
+
+class FNetStatsCounterNode;
+
+////////////////////////////////////////////////////////////////////////////////////////////////////
+
+/** NetStats View Tooltip */
+class SNetStatsCountersViewTooltip
+{
+public:
+	SNetStatsCountersViewTooltip() = delete;
+
+	static TSharedPtr<SToolTip> GetTableTooltip(const Insights::FTable& Table);
+	static TSharedPtr<SToolTip> GetColumnTooltip(const Insights::FTableColumn& Column);
+	static TSharedPtr<SToolTip> GetRowTooltip(const TSharedPtr<FNetStatsCounterNode> TreeNodePtr);
+
+private:
+	static void AddAggregatedStatsRow(TSharedPtr<SGridPanel> Grid, int32& Row, const FText& Name, const FText& Value);
+};
+
+////////////////////////////////////////////////////////////////////////////////////////////////////
+
+class SNetStatsCounterTableRowToolTip : public IToolTip
+{
+public:
+	SNetStatsCounterTableRowToolTip(const TSharedPtr<FNetStatsCounterNode> InTreeNodePtr) : TreeNodePtr(InTreeNodePtr) {}
+	virtual ~SNetStatsCounterTableRowToolTip() { }
+
+	virtual TSharedRef<class SWidget> AsWidget()
+	{
+		CreateToolTipWidget();
+		return ToolTipWidget.ToSharedRef();
+	}
+
+	virtual TSharedRef<SWidget> GetContentWidget()
+	{
+		CreateToolTipWidget();
+		return ToolTipWidget->GetContentWidget();
+	}
+
+	virtual void SetContentWidget(const TSharedRef<SWidget>& InContentWidget)
+	{
+		CreateToolTipWidget();
+		ToolTipWidget->SetContentWidget(InContentWidget);
+	}
+
+	void InvalidateWidget()
+	{
+		ToolTipWidget.Reset();
+	}
+
+	virtual bool IsEmpty() const { return false; }
+	virtual bool IsInteractive() const { return false; }
+	virtual void OnOpening() {}
+	virtual void OnClosed() {}
+
+private:
+	void CreateToolTipWidget()
+	{
+		if (!ToolTipWidget.IsValid())
+		{
+			ToolTipWidget = SNetStatsCountersViewTooltip::GetRowTooltip(TreeNodePtr);
+		}
+	}
+
+private:
+	TSharedPtr<SToolTip> ToolTipWidget;
+	const TSharedPtr<FNetStatsCounterNode> TreeNodePtr;
+};
+
+////////////////////////////////////////////////////////////////////////////////////////////////////
diff --git a/Engine/Source/Developer/TraceInsights/Private/Insights/NetworkingProfiler/Widgets/SNetStatsView.cpp b/Engine/Source/Developer/TraceInsights/Private/Insights/NetworkingProfiler/Widgets/SNetStatsView.cpp
--- a/Engine/Source/Developer/TraceInsights/Private/Insights/NetworkingProfiler/Widgets/SNetStatsView.cpp
+++ b/Engine/Source/Developer/TraceInsights/Private/Insights/NetworkingProfiler/Widgets/SNetStatsView.cpp
@@ -510,21 +510,6 @@ TSharedRef<SWidget> SNetStatsView::TreeViewHeaderRow_GenerateColumnMenu(const In
 			EUserInterfaceActionType::Button
 		);
 
-		FUIAction Action_ShowMinMaxMedColumns
-		(
-			FExecuteAction::CreateSP(this, &SNetStatsView::ContextMenu_ShowMinMaxMedColumns_Execute),
-			FCanExecuteAction::CreateSP(this, &SNetStatsView::ContextMenu_ShowMinMaxMedColumns_CanExecute)
-		);
-		MenuBuilder.AddMenuEntry
-		(
-			LOCTEXT("ContextMenu_ShowMinMaxMedColumns", "Reset Columns to Min/Max/Median Preset"),
-			LOCTEXT("ContextMenu_ShowMinMaxMedColumns_Desc", "Resets columns to Min/Max/Median preset."),
-			FSlateIcon(FInsightsStyle::GetStyleSetName(), "Icons.ResetColumn"),
-			Action_ShowMinMaxMedColumns,
-			NAME_None,
-			EUserInterfaceActionType::Button
-		);
-
 		FUIAction Action_ResetColumns
 		(
 			FExecuteAction::CreateSP(this, &SNetStatsView::ContextMenu_ResetColumns_Execute),
@@ -1365,56 +1350,6 @@ void SNetStatsView::ContextMenu_ShowAllColumns_Execute()
 	}
 }
 
-////////////////////////////////////////////////////////////////////////////////////////////////////
-// "Show Min/Max/Median Columns" action (ContextMenu)
-////////////////////////////////////////////////////////////////////////////////////////////////////
-
-bool SNetStatsView::ContextMenu_ShowMinMaxMedColumns_CanExecute() const
-{
-	return true;
-}
-
-////////////////////////////////////////////////////////////////////////////////////////////////////
-
-void SNetStatsView::ContextMenu_ShowMinMaxMedColumns_Execute()
-{
-	TSet<FName> Preset =
-	{
-		FNetStatsViewColumns::NameColumnID,
-		//FNetStatsViewColumns::MetaGroupNameColumnID,
-		//FNetStatsViewColumns::TypeColumnID,
-		FNetStatsViewColumns::InstanceCountColumnID,
-
-		FNetStatsViewColumns::TotalInclusiveSizeColumnID,
-		FNetStatsViewColumns::MaxInclusiveSizeColumnID,
-		FNetStatsViewColumns::AverageInclusiveSizeColumnID,
-
-		FNetStatsViewColumns::TotalExclusiveSizeColumnID,
-		FNetStatsViewColumns::MaxExclusiveSizeColumnID,
-		//FNetStatsViewColumns::AverageExclusiveSizeColumnID,
-	};
-
-	ColumnBeingSorted = FNetStatsViewColumns::TotalInclusiveSizeColumnID;
-	ColumnSortMode = EColumnSortMode::Descending;
-	UpdateCurrentSortingByColumn();
-
-	for (const TSharedRef<Insights::FTableColumn>& ColumnRef : Table->GetColumns())
-	{
-		const Insights::FTableColumn& Column = *ColumnRef;
-
-		const bool bShouldBeVisible = Preset.Contains(Column.GetId());
-
-		if (bShouldBeVisible && !Column.IsVisible())
-		{
-			ShowColumn(Column.GetId());
-		}
-		else if (!bShouldBeVisible && Column.IsVisible())
-		{
-			HideColumn(Column.GetId());
-		}
-	}
-}
-
 ////////////////////////////////////////////////////////////////////////////////////////////////////
 // ResetColumns action (ContextMenu)
 ////////////////////////////////////////////////////////////////////////////////////////////////////
diff --git a/Engine/Source/Developer/TraceInsights/Private/Insights/NetworkingProfiler/Widgets/SNetworkingProfilerToolbar.cpp b/Engine/Source/Developer/TraceInsights/Private/Insights/NetworkingProfiler/Widgets/SNetworkingProfilerToolbar.cpp
--- a/Engine/Source/Developer/TraceInsights/Private/Insights/NetworkingProfiler/Widgets/SNetworkingProfilerToolbar.cpp
+++ b/Engine/Source/Developer/TraceInsights/Private/Insights/NetworkingProfiler/Widgets/SNetworkingProfilerToolbar.cpp
@@ -48,6 +48,9 @@ void SNetworkingProfilerToolbar::Construct(const FArguments& InArgs, TSharedRef<
 				ToolbarBuilder.AddToolBarButton(FNetworkingProfilerCommands::Get().ToggleNetStatsViewVisibility,
 					NAME_None, TAttribute<FText>(), TAttribute<FText>(),
 					FSlateIcon(FInsightsStyle::GetStyleSetName(), "Icons.NetStatsView.ToolBar"));
+				ToolbarBuilder.AddToolBarButton(FNetworkingProfilerCommands::Get().ToggleNetStatsCountersViewVisibility,
+					NAME_None, TAttribute<FText>(), TAttribute<FText>(),
+					FSlateIcon(FInsightsStyle::GetStyleSetName(), "Icons.NetStatsView.ToolBar"));
 			}
 			ToolbarBuilder.EndSection();
 			ToolbarBuilder.BeginSection("Connection");
diff --git a/Engine/Source/Developer/TraceInsights/Private/Insights/NetworkingProfiler/Widgets/SNetworkingProfilerWindow.cpp b/Engine/Source/Developer/TraceInsights/Private/Insights/NetworkingProfiler/Widgets/SNetworkingProfilerWindow.cpp
--- a/Engine/Source/Developer/TraceInsights/Private/Insights/NetworkingProfiler/Widgets/SNetworkingProfilerWindow.cpp
+++ b/Engine/Source/Developer/TraceInsights/Private/Insights/NetworkingProfiler/Widgets/SNetworkingProfilerWindow.cpp
@@ -12,6 +12,7 @@
 #include "Insights/InsightsStyle.h"
 #include "Insights/NetworkingProfiler/NetworkingProfilerManager.h"
 #include "Insights/NetworkingProfiler/Widgets/SNetStatsView.h"
+#include "Insights/NetworkingProfiler/Widgets/SNetStatsCountersView.h"
 #include "Insights/NetworkingProfiler/Widgets/SNetworkingProfilerToolbar.h"
 #include "Insights/NetworkingProfiler/Widgets/SPacketContentView.h"
 #include "Insights/NetworkingProfiler/Widgets/SPacketView.h"
@@ -25,6 +26,7 @@
 const FName FNetworkingProfilerTabs::PacketViewID(TEXT("PacketView"));
 const FName FNetworkingProfilerTabs::PacketContentViewID(TEXT("PacketContent"));
 const FName FNetworkingProfilerTabs::NetStatsViewID(TEXT("NetStats"));
+const FName FNetworkingProfilerTabs::NetStatsCountersViewID(TEXT("NetStatsCounters"));
 
 ////////////////////////////////////////////////////////////////////////////////////////////////////
 
@@ -35,6 +37,8 @@ SNetworkingProfilerWindow::SNetworkingProfilerWindow()
 	, SelectionStartPosition(0)
 	, SelectionEndPosition(0)
 	, SelectedEventTypeIndex(InvalidEventTypeIndex)
+	, OldGameInstanceChangeCount(0U)
+	, OldConnectionChangeCount(0U)
 {
 }
 
@@ -45,10 +49,12 @@ SNetworkingProfilerWindow::~SNetworkingProfilerWindow()
 	CloseAllOpenTabs();
 
 	check(NetStatsView == nullptr);
+	check(NetStatsCountersView = nullptr);
 	check(PacketContentView == nullptr);
 	check(PacketView == nullptr);
 }
 
+
 ////////////////////////////////////////////////////////////////////////////////////////////////////
 
 const TCHAR* SNetworkingProfilerWindow::GetAnalyticsEventName() const
@@ -75,6 +81,10 @@ void SNetworkingProfilerWindow::Reset()
 		NetStatsView->Reset();
 	}
 
+	if (NetStatsCountersView)
+	{
+		NetStatsCountersView->Reset();
+	}
 	AvailableGameInstances.Reset();
 	SelectedGameInstance.Reset();
 	AvailableConnections.Reset();
@@ -176,6 +186,34 @@ void SNetworkingProfilerWindow::OnNetStatsViewTabClosed(TSharedRef<SDockTab> Tab
 
 ////////////////////////////////////////////////////////////////////////////////////////////////////
 
+BEGIN_SLATE_FUNCTION_BUILD_OPTIMIZATION
+TSharedRef<SDockTab> SNetworkingProfilerWindow::SpawnTab_NetStatsCountersView(const FSpawnTabArgs& Args)
+{
+	const TSharedRef<SDockTab> DockTab = SNew(SDockTab)
+		.ShouldAutosize(false)
+		.TabRole(ETabRole::PanelTab)
+		[
+			SAssignNew(NetStatsCountersView, SNetStatsCountersView, SharedThis(this))
+		];
+
+	DockTab->SetOnTabClosed(SDockTab::FOnTabClosedCallback::CreateRaw(this, &SNetworkingProfilerWindow::OnNetStatsCountersViewTabClosed));
+	AddOpenTab(DockTab);
+
+	return DockTab;
+}
+END_SLATE_FUNCTION_BUILD_OPTIMIZATION
+
+////////////////////////////////////////////////////////////////////////////////////////////////////
+
+void SNetworkingProfilerWindow::OnNetStatsCountersViewTabClosed(TSharedRef<SDockTab> TabBeingClosed)
+{
+	NetStatsCountersView = nullptr;
+
+	RemoveOpenTab(TabBeingClosed);
+}
+
+////////////////////////////////////////////////////////////////////////////////////////////////////
+
 void SNetworkingProfilerWindow::Construct(const FArguments& InArgs, const TSharedRef<SDockTab>& ConstructUnderMajorTab, const TSharedPtr<SWindow>& ConstructUnderWindow)
 {
 	TSharedPtr<FNetworkingProfilerManager> NetworkingProfilerManager = FNetworkingProfilerManager::Get();
@@ -218,14 +256,20 @@ void SNetworkingProfilerWindow::RegisterTabSpawners()
 		.SetDisplayName(LOCTEXT("NetStatsViewTabTitle", "Net Stats"))
 		.SetIcon(FSlateIcon(FInsightsStyle::GetStyleSetName(), "Icons.NetStatsView"))
 		.SetGroup(Group);
+
+	TabManagerPtr->RegisterTabSpawner(FNetworkingProfilerTabs::NetStatsCountersViewID, FOnSpawnTab::CreateRaw(this, &SNetworkingProfilerWindow::SpawnTab_NetStatsCountersView))
+		.SetDisplayName(LOCTEXT("NetworkingProfiler.NetStatsCountersViewTabTitle", "NetStatsCounters"))
+		.SetIcon(FSlateIcon(FInsightsStyle::GetStyleSetName(), "Icons.NetStatsView"))
+		.SetGroup(Group);
 }
 
 ////////////////////////////////////////////////////////////////////////////////////////////////////
 
+
 BEGIN_SLATE_FUNCTION_BUILD_OPTIMIZATION
 TSharedRef<FTabManager::FLayout> SNetworkingProfilerWindow::CreateDefaultTabLayout() const
 {
-	return FTabManager::NewLayout("InsightsNetworkingProfilerLayout_v1.2")
+	return FTabManager::NewLayout("InsightsNetworkingProfilerLayout_v1.4")
 		->AddArea
 		(
 			FTabManager::NewPrimaryArea()
@@ -238,16 +282,16 @@ TSharedRef<FTabManager::FLayout> SNetworkingProfilerWindow::CreateDefaultTabLayo
 				->Split
 				(
 					FTabManager::NewStack()
-					->SetSizeCoefficient(0.35f)
-					->SetHideTabWell(true)
+					->SetSizeCoefficient(0.35f)					
 					->AddTab(FNetworkingProfilerTabs::PacketViewID, ETabState::OpenedTab)
+					->SetHideTabWell(true)
 				)
 				->Split
 				(
 					FTabManager::NewStack()
 					->SetSizeCoefficient(0.65f)
-					->SetHideTabWell(true)
 					->AddTab(FNetworkingProfilerTabs::PacketContentViewID, ETabState::OpenedTab)
+					->SetHideTabWell(true)
 				)
 			)
 			->Split
@@ -255,6 +299,7 @@ TSharedRef<FTabManager::FLayout> SNetworkingProfilerWindow::CreateDefaultTabLayo
 				FTabManager::NewStack()
 				->SetSizeCoefficient(0.35f)
 				->AddTab(FNetworkingProfilerTabs::NetStatsViewID, ETabState::OpenedTab)
+				->AddTab(FNetworkingProfilerTabs::NetStatsCountersViewID, ETabState::OpenedTab)
 				->SetForegroundTab(FNetworkingProfilerTabs::NetStatsViewID)
 			)
 		);
@@ -278,6 +323,7 @@ void SNetworkingProfilerWindow::BindCommands()
 	Map_TogglePacketViewVisibility();
 	Map_TogglePacketContentViewVisibility();
 	Map_ToggleNetStatsViewVisibility();
+	Map_ToggleNetStatsCountersViewVisibility();
 
 	GetCommandList()->MapAction(FInsightsCommands::Get().ToggleDebugInfo, FInsightsManager::GetActionManager().ToggleDebugInfo_Custom());
 }
@@ -320,6 +366,7 @@ void SNetworkingProfilerWindow::BindCommands()
 IMPLEMENT_TOGGLE_COMMAND(TogglePacketViewVisibility, IsPacketViewVisible, ShowOrHidePacketView)
 IMPLEMENT_TOGGLE_COMMAND(TogglePacketContentViewVisibility, IsPacketContentViewVisible, ShowOrHidePacketContentView)
 IMPLEMENT_TOGGLE_COMMAND(ToggleNetStatsViewVisibility, IsNetStatsViewVisible, ShowOrHideNetStatsView)
+IMPLEMENT_TOGGLE_COMMAND(ToggleNetStatsCountersViewVisibility, IsNetStatsCountersViewVisible, ShowOrHideNetStatsCountersView)
 
 #undef IMPLEMENT_TOGGLE_COMMAND
 
@@ -329,19 +376,23 @@ void SNetworkingProfilerWindow::Tick(const FGeometry& AllottedGeometry, const do
 {
 	TSharedPtr<const TraceServices::IAnalysisSession> Session = FInsightsManager::Get()->GetSession();
 
-	uint32 GameInstanceCount = 0;
+	uint32 GameInstanceChangeCount = 0;
+	uint32 ConnectionChangeCount = 0;
 	if (Session.IsValid())
 	{
 		TraceServices::FAnalysisSessionReadScope SessionReadScope(*Session.Get());
 		const TraceServices::INetProfilerProvider* NetProfilerProvider = TraceServices::ReadNetProfilerProvider(*Session.Get());
 		if (NetProfilerProvider)
 		{
-			GameInstanceCount = NetProfilerProvider->GetGameInstanceCount();
+			GameInstanceChangeCount = NetProfilerProvider->GetGameInstanceChangeCount();
+			ConnectionChangeCount = NetProfilerProvider->GetConnectionChangeCount();
 		}
 	}
 
-	if (GameInstanceCount != AvailableGameInstances.Num())
+	if (GameInstanceChangeCount != OldGameInstanceChangeCount || ConnectionChangeCount != OldConnectionChangeCount)
 	{
+		OldGameInstanceChangeCount = GameInstanceChangeCount;
+		OldConnectionChangeCount = ConnectionChangeCount;
 		UpdateAvailableGameInstances();
 	}
 
@@ -387,9 +438,12 @@ void SNetworkingProfilerWindow::UpdateAvailableGameInstances()
 		const TraceServices::INetProfilerProvider* NetProfilerProvider = TraceServices::ReadNetProfilerProvider(*Session.Get());
 		if (NetProfilerProvider)
 		{
-			NetProfilerProvider->ReadGameInstances([this](const TraceServices::FNetProfilerGameInstance& GameInstance)
+			NetProfilerProvider->ReadGameInstances([this, NetProfilerProvider](const TraceServices::FNetProfilerGameInstance& GameInstance)
 			{
-				AvailableGameInstances.Add(MakeShared<FGameInstanceItem>(GameInstance));
+				if (NetProfilerProvider->GetConnectionCount(GameInstance.GameInstanceIndex) > 0)
+				{
+					AvailableGameInstances.Add(MakeShared<FGameInstanceItem>(GameInstance));
+				}
 			});
 		}
 	}
@@ -515,7 +569,7 @@ FText SNetworkingProfilerWindow::FGameInstanceItem::GetText() const
 {
 	if (IsInstanceNameSet())
 	{
-		return FText::Format(LOCTEXT("GameInstanceItemFmt0", "{0} Game Instance {1} [{2}]"), FText::FromString(GetInstanceName()), FText::AsNumber(GetIndex()),
+		return FText::Format(LOCTEXT("GameInstanceItemFmt0", "Game Instance {0} [{1}]"), FText::AsNumber(GetIndex()),
 			IsServer() ? FText::FromString("Server") : FText::FromString("Client"));
 	}
 	else
@@ -528,12 +582,35 @@ FText SNetworkingProfilerWindow::FGameInstanceItem::GetText() const
 
 FText SNetworkingProfilerWindow::FGameInstanceItem::GetTooltipText() const
 {
-	return GetText();
+	if (IsInstanceNameSet())
+	{
+		return FText::Format(LOCTEXT("GameInstanceItemFmt0", "{0} Game Instance {1} [{2}]"), FText::FromString(GetInstanceName()), FText::AsNumber(GetIndex()),
+			IsServer() ? FText::FromString("Server") : FText::FromString("Client"));
+	}
+	else
+	{
+		return FText::Format(LOCTEXT("GameInstanceItemFmt1", "Game Instance {0}"), FText::AsNumber(GetIndex()));
+	}
 }
 
 ////////////////////////////////////////////////////////////////////////////////////////////////////
 
 FText SNetworkingProfilerWindow::FConnectionItem::GetText() const
+{
+	if (IsConnectionAddressSet())
+	{
+		return FText::Format(LOCTEXT("ConnectionItemFmt1", "Connection {0} to {1}"), FText::AsNumber(GetIndex()),
+			FText::FromString(GetConnectionAddress()));
+	}
+	else
+	{
+		return FText::Format(LOCTEXT("ConnectionItemFmt2", "Connection {0}"), FText::AsNumber(GetIndex()));
+	}
+}
+
+////////////////////////////////////////////////////////////////////////////////////////////////////
+
+FText SNetworkingProfilerWindow::FConnectionItem::GetTooltipText() const
 {
 	if (IsConnectionNameSet())
 	{
@@ -559,13 +636,7 @@ FText SNetworkingProfilerWindow::FConnectionItem::GetText() const
 			return FText::Format(LOCTEXT("ConnectionItemFmt2", "Connection {0}"), FText::AsNumber(GetIndex()));
 		}
 	}
-}
 
-////////////////////////////////////////////////////////////////////////////////////////////////////
-
-FText SNetworkingProfilerWindow::FConnectionItem::GetTooltipText() const
-{
-	return GetText();
 }
 
 ////////////////////////////////////////////////////////////////////////////////////////////////////
@@ -602,6 +673,7 @@ TSharedRef<SWidget> SNetworkingProfilerWindow::CreateGameInstanceComboBox()
 		.OptionsSource(&AvailableGameInstances)
 		.OnSelectionChanged(this, &SNetworkingProfilerWindow::GameInstance_OnSelectionChanged)
 		.OnGenerateWidget(this, &SNetworkingProfilerWindow::GameInstance_OnGenerateWidget)
+		.OnComboBoxOpening(this, &SNetworkingProfilerWindow::UpdateAvailableGameInstances)
 		[
 			SNew(STextBlock)
 			.Text(this, &SNetworkingProfilerWindow::GameInstance_GetSelectedText)
@@ -863,7 +935,7 @@ void SNetworkingProfilerWindow::SetSelectedEventTypeIndex(const uint32 InEventTy
 
 void SNetworkingProfilerWindow::UpdateAggregatedNetStats()
 {
-	if (NetStatsView)
+	if (NetStatsView || NetStatsCountersView)
 	{
 		if (SelectedGameInstance.IsValid() &&
 			SelectedConnection.IsValid() &&
@@ -871,17 +943,35 @@ void SNetworkingProfilerWindow::UpdateAggregatedNetStats()
 			(SelectedPacketStartIndex < SelectedPacketEndIndex) &&
 			(SelectedPacketStartIndex + 1 != SelectedPacketEndIndex || SelectionStartPosition < SelectionEndPosition))
 		{
-			NetStatsView->UpdateStats(SelectedGameInstance->GetIndex(),
-									  SelectedConnection->GetIndex(),
-									  SelectedConnectionMode->Mode,
-									  SelectedPacketStartIndex,
-									  SelectedPacketEndIndex,
-									  SelectionStartPosition,
-									  SelectionEndPosition);
+			if (NetStatsView)
+			{
+				NetStatsView->UpdateStats(SelectedGameInstance->GetIndex(),
+									SelectedConnection->GetIndex(),
+									SelectedConnectionMode->Mode,
+									SelectedPacketStartIndex,
+									SelectedPacketEndIndex,
+									SelectionStartPosition,
+									SelectionEndPosition);
+			}
+			if (NetStatsCountersView)
+			{
+				NetStatsCountersView->UpdateStats(SelectedGameInstance->GetIndex(),
+									SelectedConnection->GetIndex(),
+									SelectedConnectionMode->Mode,
+									SelectedPacketStartIndex,
+									SelectedPacketEndIndex);
+			}
 		}
 		else
 		{
-			NetStatsView->ResetStats();
+			if (NetStatsView)
+			{
+				NetStatsView->ResetStats();
+			}
+			if (NetStatsCountersView)
+			{
+				NetStatsCountersView->ResetStats();
+			}
 		}
 	}
 }
diff --git a/Engine/Source/Developer/TraceInsights/Private/Insights/NetworkingProfiler/Widgets/SNetworkingProfilerWindow.h b/Engine/Source/Developer/TraceInsights/Private/Insights/NetworkingProfiler/Widgets/SNetworkingProfilerWindow.h
--- a/Engine/Source/Developer/TraceInsights/Private/Insights/NetworkingProfiler/Widgets/SNetworkingProfilerWindow.h
+++ b/Engine/Source/Developer/TraceInsights/Private/Insights/NetworkingProfiler/Widgets/SNetworkingProfilerWindow.h
@@ -14,6 +14,7 @@
 class SPacketView;
 class SPacketContentView;
 class SNetStatsView;
+class SNetStatsCountersView;
 
 ////////////////////////////////////////////////////////////////////////////////////////////////////
 
@@ -23,6 +24,7 @@ struct FNetworkingProfilerTabs
 	static const FName PacketViewID;
 	static const FName PacketContentViewID;
 	static const FName NetStatsViewID;
+	static const FName NetStatsCountersViewID;
 };
 
 ////////////////////////////////////////////////////////////////////////////////////////////////////
@@ -106,6 +108,10 @@ public:
 	const bool IsNetStatsViewVisible() const { return NetStatsView.IsValid(); }
 	void ShowOrHideNetStatsView(const bool bVisibleState) { ShowHideTab(FNetworkingProfilerTabs::NetStatsViewID, bVisibleState); }
 
+	TSharedPtr<SNetStatsCountersView> GetNetStatsCountersView() const { return NetStatsCountersView; }
+	const bool IsNetStatsCountersViewVisible() const { return NetStatsCountersView.IsValid(); }
+	void ShowOrHideNetStatsCountersView(const bool bVisibleState) { ShowHideTab(FNetworkingProfilerTabs::NetStatsCountersViewID, bVisibleState); }
+
 	const TraceServices::FNetProfilerGameInstance* GetSelectedGameInstance() const { return SelectedGameInstance ? &SelectedGameInstance->GameInstance : nullptr; }
 	uint32 GetSelectedGameInstanceIndex() const { return SelectedGameInstance ? SelectedGameInstance->GetIndex() : 0; }
 	const TraceServices::FNetProfilerConnection* GetSelectedConnection() const { return SelectedConnection ? &SelectedConnection->Connection : nullptr; }
@@ -137,6 +143,9 @@ private:
 	TSharedRef<SDockTab> SpawnTab_NetStatsView(const FSpawnTabArgs& Args);
 	void OnNetStatsViewTabClosed(TSharedRef<SDockTab> TabBeingClosed);
 
+	TSharedRef<SDockTab> SpawnTab_NetStatsCountersView(const FSpawnTabArgs& Args);
+	void OnNetStatsCountersViewTabClosed(TSharedRef<SDockTab> TabBeingClosed);
+
 	//////////////////////////////////////////////////
 
 	void BindCommands();
@@ -156,6 +165,7 @@ private:\
 	DECLARE_TOGGLE_COMMAND(TogglePacketViewVisibility)
 	DECLARE_TOGGLE_COMMAND(TogglePacketContentViewVisibility)
 	DECLARE_TOGGLE_COMMAND(ToggleNetStatsViewVisibility)
+	DECLARE_TOGGLE_COMMAND(ToggleNetStatsCountersViewVisibility)
 #undef DECLARE_TOGGLE_COMMAND
 
 	//////////////////////////////////////////////////
@@ -206,6 +216,9 @@ private:
 	/** The Net Stats widget */
 	TSharedPtr<SNetStatsView> NetStatsView;
 
+	/** The NetStatsCounters widget */
+	TSharedPtr<SNetStatsCountersView> NetStatsCountersView;
+	
 	TSharedPtr<SComboBox<TSharedPtr<FGameInstanceItem>>> GameInstanceComboBox;
 	TArray<TSharedPtr<FGameInstanceItem>> AvailableGameInstances;
 	TSharedPtr<FGameInstanceItem> SelectedGameInstance;
@@ -231,4 +244,7 @@ private:
 
 	static const uint32 InvalidEventTypeIndex = uint32(-1);
 	uint32 SelectedEventTypeIndex;
+
+	uint32 OldGameInstanceChangeCount;
+	uint32 OldConnectionChangeCount;
 };
diff --git a/Engine/Source/Developer/TraceInsights/Private/Insights/NetworkingProfiler/Widgets/SPacketContentView.cpp b/Engine/Source/Developer/TraceInsights/Private/Insights/NetworkingProfiler/Widgets/SPacketContentView.cpp
--- a/Engine/Source/Developer/TraceInsights/Private/Insights/NetworkingProfiler/Widgets/SPacketContentView.cpp
+++ b/Engine/Source/Developer/TraceInsights/Private/Insights/NetworkingProfiler/Widgets/SPacketContentView.cpp
@@ -16,6 +16,8 @@
 #include "Widgets/Input/SButton.h"
 #include "Widgets/Input/SEditableTextBox.h"
 #include "Widgets/Layout/SScrollBar.h"
+#include "SlateOptMacros.h"
+#include "Templates/SharedPointer.h"
 
 // Insights
 #include "Insights/Common/PaintUtils.h"
@@ -37,6 +39,8 @@ SPacketContentView::SPacketContentView()
 	: ProfilerWindowWeakPtr()
 	, DrawState(MakeShared<FPacketContentViewDrawState>())
 	, FilteredDrawState(MakeShared<FPacketContentViewDrawState>())
+	, AvailableAggregationModes()
+	, SelectedAggregationMode(nullptr)
 {
 	Reset();
 }
@@ -107,10 +111,108 @@ void SPacketContentView::Reset()
 
 ////////////////////////////////////////////////////////////////////////////////////////////////////
 
+BEGIN_SLATE_FUNCTION_BUILD_OPTIMIZATION
+TSharedRef<SWidget> SPacketContentView::AggregationMode_OnGenerateWidget(TSharedPtr<FAggregationModeItem> InAggregationMode) const
+{
+	return SNew(STextBlock)
+		.Text(InAggregationMode->GetText())
+		.ToolTipText(InAggregationMode->GetTooltipText());
+}
+END_SLATE_FUNCTION_BUILD_OPTIMIZATION
+
+////////////////////////////////////////////////////////////////////////////////////////////////////
+
+void SPacketContentView::AggregationMode_OnSelectionChanged(TSharedPtr<FAggregationModeItem> NewAggregationMode, ESelectInfo::Type SelectInfo)
+{
+	const bool bSameValue = (!SelectedAggregationMode.IsValid() && !NewAggregationMode.IsValid()) ||
+		(SelectedAggregationMode.IsValid() && NewAggregationMode.IsValid() &&
+			SelectedAggregationMode->Mode == NewAggregationMode->Mode);
+
+	SelectedAggregationMode = NewAggregationMode;
+
+	// Need to refresh selection
+	if (!bSameValue)
+	{
+		TSharedPtr<SNetworkingProfilerWindow> ProfilerWindow = ProfilerWindowWeakPtr.Pin();
+		if (ProfilerWindow.IsValid())
+		{
+			const TSharedPtr<SPacketView> PacketView = ProfilerWindow->GetPacketView();
+			if (PacketView.IsValid())
+			{
+				PacketView->InvalidateState();			
+			}
+		}
+	}
+}
+
+////////////////////////////////////////////////////////////////////////////////////////////////////
+
+FText SPacketContentView::AggregationMode_GetSelectedText() const
+{
+	return SelectedAggregationMode.IsValid() ? SelectedAggregationMode->GetText() : LOCTEXT("NoAggregationModeText", "None");
+}
+
+////////////////////////////////////////////////////////////////////////////////////////////////////
+
+FText SPacketContentView::AggregationMode_GetSelectedTooltipText() const
+{
+	return SelectedAggregationMode.IsValid() ? SelectedAggregationMode->GetTooltipText() : FText();
+}
+
+////////////////////////////////////////////////////////////////////////////////////////////////////
+
+FText SPacketContentView::FAggregationModeItem::GetText() const
+{
+	switch (Mode)
+	{
+	case TraceServices::ENetProfilerAggregationMode::Aggregate:
+		return LOCTEXT("AggregationMode_Aggregate", "Aggregate");
+
+	case TraceServices::ENetProfilerAggregationMode::InstanceMax:
+		return LOCTEXT("AggregationMode_InstanceMax", "InstanceMax");
+
+	default:
+		return LOCTEXT("AggregationMode_None", "None");
+	}
+}
+
+////////////////////////////////////////////////////////////////////////////////////////////////////
+
+FText SPacketContentView::FAggregationModeItem::GetTooltipText() const
+{
+	return GetText();
+}
+
+////////////////////////////////////////////////////////////////////////////////////////////////////
+
+BEGIN_SLATE_FUNCTION_BUILD_OPTIMIZATION
+TSharedRef<SWidget> SPacketContentView::CreateAggregationModeComboBox()
+{
+	AggregationModeComboBox = SNew(SComboBox<TSharedPtr<FAggregationModeItem>>)
+		.ToolTipText(this, &SPacketContentView::AggregationMode_GetSelectedTooltipText)
+		.OptionsSource(&AvailableAggregationModes)
+		.OnSelectionChanged(this, &SPacketContentView::AggregationMode_OnSelectionChanged)
+		.OnGenerateWidget(this, &SPacketContentView::AggregationMode_OnGenerateWidget)
+		[
+			SNew(STextBlock)
+			.Text(this, &SPacketContentView::AggregationMode_GetSelectedText)
+		];
+
+	return AggregationModeComboBox.ToSharedRef();
+}
+END_SLATE_FUNCTION_BUILD_OPTIMIZATION
+
 void SPacketContentView::Construct(const FArguments& InArgs, TSharedRef<SNetworkingProfilerWindow> InProfilerWindow)
 {
 	ProfilerWindowWeakPtr = InProfilerWindow;
 
+	AvailableAggregationModes.Add(MakeShared<FAggregationModeItem>(TraceServices::ENetProfilerAggregationMode::None));
+	AvailableAggregationModes.Add(MakeShared<FAggregationModeItem>(TraceServices::ENetProfilerAggregationMode::Aggregate));
+	AvailableAggregationModes.Add(MakeShared<FAggregationModeItem>(TraceServices::ENetProfilerAggregationMode::InstanceMax));
+	SelectedAggregationMode = AvailableAggregationModes[1];
+
+	TSharedRef<SWidget> AggregationModeWidget = CreateAggregationModeComboBox();
+
 	FSlimHorizontalToolBarBuilder ToolbarBuilder(TSharedPtr<const FUICommandList>(), FMultiBoxCustomization::None);
 	ToolbarBuilder.SetStyle(&FInsightsStyle::Get(), "SecondaryToolbar2");
 
@@ -307,6 +409,17 @@ void SPacketContentView::Construct(const FArguments& InArgs, TSharedRef<SNetwork
 		);
 	}
 	ToolbarBuilder.EndSection();
+	ToolbarBuilder.BeginSection("AggregationType");
+	{
+		ToolbarBuilder.AddWidget(
+			SNew(SBox)
+			.Padding(FMargin(12.0f, 0.0f, 0.0f, 0.0f))
+			[
+				AggregationModeWidget
+			]
+		);
+	}
+	ToolbarBuilder.EndSection();
 
 	ChildSlot
 	[
@@ -1027,7 +1140,7 @@ void SPacketContentView::UpdateHoveredEvent()
 		if (Event.ObjectInstanceIndex != 0)
 		{
 			Tooltip.AddNameValueTextLine(TEXT("Net Id:"), FText::AsNumber(ObjectInstance.NetId).ToString());
-			Tooltip.AddNameValueTextLine(TEXT("Type Id:"), FString::Printf(TEXT("0x%016X"), ObjectInstance.TypeId));
+			Tooltip.AddNameValueTextLine(TEXT("Type Id:"), FString::Printf(TEXT("0x%016" UINT64_x_FMT), ObjectInstance.TypeId));
 			Tooltip.AddNameValueTextLine(TEXT("Obj. LifeTime:"), FString::Format(TEXT("from {0} to {1}"),
 				{ TimeUtils::FormatTimeAuto(ObjectInstance.LifeTime.Begin), TimeUtils::FormatTimeAuto(ObjectInstance.LifeTime.End) }));
 		}
diff --git a/Engine/Source/Developer/TraceInsights/Private/Insights/NetworkingProfiler/Widgets/SPacketContentView.h b/Engine/Source/Developer/TraceInsights/Private/Insights/NetworkingProfiler/Widgets/SPacketContentView.h
--- a/Engine/Source/Developer/TraceInsights/Private/Insights/NetworkingProfiler/Widgets/SPacketContentView.h
+++ b/Engine/Source/Developer/TraceInsights/Private/Insights/NetworkingProfiler/Widgets/SPacketContentView.h
@@ -11,6 +11,7 @@
 #include "TraceServices/Model/NetProfiler.h"
 #include "Widgets/DeclarativeSyntaxSupport.h"
 #include "Widgets/SCompoundWidget.h"
+#include "Widgets/Input/SComboBox.h"
 
 // Insights
 #include "Insights/Common/FixedCircularBuffer.h"
@@ -86,6 +87,20 @@ struct FNetworkPacketEventRef
  */
 class SPacketContentView : public SCompoundWidget
 {
+private:
+	struct FAggregationModeItem
+	{
+		/** Conversion constructor. */
+		FAggregationModeItem(const TraceServices::ENetProfilerAggregationMode& InMode)
+			: Mode(InMode)
+		{}
+
+		FText GetText() const;
+		FText GetTooltipText() const;
+
+		TraceServices::ENetProfilerAggregationMode Mode;
+	};
+
 public:
 	/** Number of pixels. */
 	static constexpr float MOUSE_SNAP_DISTANCE = 2.0f;
@@ -152,6 +167,8 @@ public:
 	void EnableFilterEventType(const uint32 InEventTypeIndex);
 	void DisableFilterEventType();
 
+	TraceServices::ENetProfilerAggregationMode GetSelectedFilterEventAggregationMode() const { return SelectedAggregationMode ? SelectedAggregationMode->Mode : TraceServices::ENetProfilerAggregationMode::Aggregate; }
+
 	void FindFirstEvent();
 	void FindPreviousEvent(EEventNavigationType NavigationType);
 	void FindNextEvent(EEventNavigationType NavigationType);
@@ -210,6 +227,12 @@ private:
 
 	void AdjustForSplitContent();
 
+	TSharedRef<SWidget> CreateAggregationModeComboBox();
+	TSharedRef<SWidget> AggregationMode_OnGenerateWidget(TSharedPtr<FAggregationModeItem> InAggregationMode) const;
+	void AggregationMode_OnSelectionChanged(TSharedPtr<FAggregationModeItem> NewAggregationMode, ESelectInfo::Type SelectInfo);
+	FText AggregationMode_GetSelectedText() const;
+	FText AggregationMode_GetSelectedTooltipText() const;
+
 private:
 	TWeakPtr<SNetworkingProfilerWindow> ProfilerWindowWeakPtr;
 
@@ -242,6 +265,10 @@ private:
 
 	TSharedPtr<SScrollBar> HorizontalScrollBar;
 
+	TSharedPtr<SComboBox<TSharedPtr<FAggregationModeItem>>> AggregationModeComboBox;
+	TArray<TSharedPtr<FAggregationModeItem>> AvailableAggregationModes;
+	TSharedPtr<FAggregationModeItem> SelectedAggregationMode;
+
 	//////////////////////////////////////////////////
 	// Panning and Zooming behaviors
 
diff --git a/Engine/Source/Developer/TraceInsights/Private/Insights/NetworkingProfiler/Widgets/SPacketView.cpp b/Engine/Source/Developer/TraceInsights/Private/Insights/NetworkingProfiler/Widgets/SPacketView.cpp
--- a/Engine/Source/Developer/TraceInsights/Private/Insights/NetworkingProfiler/Widgets/SPacketView.cpp
+++ b/Engine/Source/Developer/TraceInsights/Private/Insights/NetworkingProfiler/Widgets/SPacketView.cpp
@@ -281,6 +281,7 @@ void SPacketView::UpdateState()
 		uint32 NetId = 0;
 		bool bByEventType = false;
 		uint32 EventTypeIndex = 0;
+		TraceServices::ENetProfilerAggregationMode AggregationMode = TraceServices::ENetProfilerAggregationMode::None;
 	};
 
 	FPacketFilter Filter;
@@ -296,6 +297,7 @@ void SPacketView::UpdateState()
 
 			Filter.bByEventType = PacketContentView->IsFilterByEventTypeEnabled();
 			Filter.EventTypeIndex = PacketContentView->GetFilterEventTypeIndex();
+			Filter.AggregationMode = PacketContentView->GetSelectedFilterEventAggregationMode();
 		}
 	}
 
@@ -334,7 +336,8 @@ void SPacketView::UpdateState()
 				if (PacketStartIndex <= PacketEndIndex)
 				{
 					int32 PacketIndex = PacketStartIndex;
-					NetProfilerProvider->EnumeratePackets(ConnectionIndex, ConnectionMode, PacketStartIndex, PacketEndIndex, [this, &Builder, &PacketIndex, &Filter, NetProfilerProvider](const TraceServices::FNetProfilerPacket& Packet)
+					int32 FilterMatchEventTypeIndex = -1;
+					NetProfilerProvider->EnumeratePackets(ConnectionIndex, ConnectionMode, PacketStartIndex, PacketEndIndex, [this, &Builder, &PacketIndex, &Filter, NetProfilerProvider, &FilterMatchEventTypeIndex](const TraceServices::FNetProfilerPacket& Packet)
 					{
 						FNetworkPacketAggregatedSample* SamplePtr = Builder.AddPacket(PacketIndex++, Packet);
 
@@ -345,17 +348,19 @@ void SPacketView::UpdateState()
 								SamplePtr->bAtLeastOnePacketMatchesFilter = !Filter.bByNetId && !Filter.bByEventType;
 							}
 
-							if (!SamplePtr->bAtLeastOnePacketMatchesFilter && (Filter.bByNetId || Filter.bByEventType))
+							if ((!SamplePtr->bAtLeastOnePacketMatchesFilter || Filter.AggregationMode != TraceServices::ENetProfilerAggregationMode::None) && (Filter.bByNetId || Filter.bByEventType))
 							{
 								bool bFilterMatch = false;
+								uint32 FilterMatchAggregatedEventSizeInBits = 0U;
+								uint32 FilterMatchMaxEventSizeBits = 0U;
 
 								const uint32 StartPos = 0;
 								const uint32 EndPos = Packet.ContentSizeInBits;
-								NetProfilerProvider->EnumeratePacketContentEventsByPosition(ConnectionIndex, ConnectionMode, PacketIndex - 1, StartPos, EndPos, [this, &bFilterMatch, &Filter, NetProfilerProvider](const TraceServices::FNetProfilerContentEvent& Event)
+								NetProfilerProvider->EnumeratePacketContentEventsByPosition(ConnectionIndex, ConnectionMode, PacketIndex - 1, StartPos, EndPos, [this, &bFilterMatch, &Filter, NetProfilerProvider, &FilterMatchAggregatedEventSizeInBits,&FilterMatchMaxEventSizeBits,  &FilterMatchEventTypeIndex](const TraceServices::FNetProfilerContentEvent& Event)
 								{
-									if (!bFilterMatch)
+									bool bEventMatchesFilter = true;
+									if (!bFilterMatch || (Filter.AggregationMode != TraceServices::ENetProfilerAggregationMode::None))
 									{
-										bool bEventMatchesFilter = true;
 										if (Filter.bByEventType && Filter.EventTypeIndex != Event.EventTypeIndex)
 										{
 											bEventMatchesFilter = false;
@@ -375,20 +380,41 @@ void SPacketView::UpdateState()
 												bEventMatchesFilter = false;
 											}
 										}
+
 										if (bEventMatchesFilter)
 										{
-											bFilterMatch = true;
+											const uint32 EventSize = Event.EndPos - Event.StartPos;
+											FilterMatchAggregatedEventSizeInBits += EventSize;
+											FilterMatchMaxEventSizeBits = FMath::Max(FilterMatchMaxEventSizeBits, EventSize);
+
+											if (!bFilterMatch)
+											{
+												FilterMatchEventTypeIndex = Event.NameIndex;
+												bFilterMatch = true;
+											}
 										}
 									}
+
 								});
 
 								if (bFilterMatch)
 								{
+									if (Filter.AggregationMode == TraceServices::ENetProfilerAggregationMode::Aggregate)
+									{
+										SamplePtr->FilterMatchHighlightSizeInBits = FMath::Max(SamplePtr->FilterMatchHighlightSizeInBits, FilterMatchAggregatedEventSizeInBits);
+									}
+									else if (Filter.AggregationMode == TraceServices::ENetProfilerAggregationMode::InstanceMax)
+									{
+										SamplePtr->FilterMatchHighlightSizeInBits = FMath::Max(SamplePtr->FilterMatchHighlightSizeInBits, FilterMatchMaxEventSizeBits);
+									}
+									
 									SamplePtr->bAtLeastOnePacketMatchesFilter = true;
 								}
 							}
 						}
 					});
+
+					Builder.SetHighlightEventTypeIndex(FilterMatchEventTypeIndex);
 				}
 			}
 		}
@@ -725,17 +751,15 @@ int32 SPacketView::OnPaint(const FPaintArgs& Args, const FGeometry& AllottedGeom
 			}
 			if (HoveredSample.Sample->NumPackets == 1)
 			{
-				NumLines = 8;
-				Text = FString::Format(TEXT("Packet Index: {0}\n"
-											"Sequence Number: {1}\n"
-											"Content Size: {2} bits\n"
-											"Total Size: {3} bytes ({4} unused bits)\n"
-											"Timestamp: {5}\n"
-											"Status: {6}\n"
-											"Connection State: {7}\n"
-											"Engine Frame Number: {8}"),
-					{
-						FText::AsNumber(HoveredSample.Sample->LargestPacket.Index).ToString(),
+				NumLines = 7;
+				Text = FString::Format(TEXT("Sequence Number: {0}\n"
+											"Content Size: {1} bits\n"
+											"Total Size: {2} bytes ({3} unused bits)\n"
+											"Timestamp: {4}\n"
+											"Status: {5}\n"
+											"Connection State: {6}\n"
+											"Engine Frame Number: {7}"),
+					{						
 						FText::AsNumber(HoveredSample.Sample->LargestPacket.SequenceNumber).ToString(),
 						FText::AsNumber(HoveredSample.Sample->LargestPacket.ContentSizeInBits).ToString(),
 						FText::AsNumber(HoveredSample.Sample->LargestPacket.TotalSizeInBytes).ToString(),
@@ -750,22 +774,20 @@ int32 SPacketView::OnPaint(const FPaintArgs& Args, const FGeometry& AllottedGeom
 			}
 			else
 			{
-				NumLines = 11;
+				NumLines = 10;
 				Text = FString::Format(TEXT("{0} network packets\n"
 											"({1})\n"
 											"Largest Packet\n"
-											"    Index: {2}\n"
-											"    Sequence Number: {3}\n"
-											"    Content Size: {4} bits\n"
-											"    Total Size: {5} bytes ({6} unused bits)\n"
-											"    Timestamp: {7}\n"
-											"    Status: {8}\n"
-											"    Connection State: {9}\n"
-											"    Engine Frame Number: {10}"),
+											"    Sequance Number: {2}\n"
+											"    Content Size: {3} bits\n"
+											"    Total Size: {4} bytes ({5} unused bits)\n"
+											"    Timestamp: {6}\n"
+											"    Status: {7}\n"
+											"    Connection State: {8}\n"
+											"    Engine Frame Number: {9}"),
 					{
 						HoveredSample.Sample->NumPackets,
 						::AggregatedStatusToString(HoveredSample.Sample->AggregatedStatus),
-						FText::AsNumber(HoveredSample.Sample->LargestPacket.Index).ToString(),
 						FText::AsNumber(HoveredSample.Sample->LargestPacket.SequenceNumber).ToString(),
 						FText::AsNumber(HoveredSample.Sample->LargestPacket.ContentSizeInBits).ToString(),
 						FText::AsNumber(HoveredSample.Sample->LargestPacket.TotalSizeInBytes).ToString(),
diff --git a/Engine/Source/Developer/TraceServices/Private/Analyzers/NetTraceAnalyzer.cpp b/Engine/Source/Developer/TraceServices/Private/Analyzers/NetTraceAnalyzer.cpp
--- a/Engine/Source/Developer/TraceServices/Private/Analyzers/NetTraceAnalyzer.cpp
+++ b/Engine/Source/Developer/TraceServices/Private/Analyzers/NetTraceAnalyzer.cpp
@@ -4,10 +4,12 @@
 
 #include "AnalysisServicePrivate.h"
 #include "Common/Utils.h"
+#include "TraceServices/Model/Frames.h"
 #include "HAL/LowLevelMemTracker.h"
 #include "Logging/LogMacros.h"
 #include "TraceServices/Model/Threads.h"
 
+
 DEFINE_LOG_CATEGORY_STATIC(LogNetTrace, Log, All);
 
 namespace TraceServices
@@ -24,6 +26,7 @@ enum ENetTraceAnalyzerVersion
 FNetTraceAnalyzer::FNetTraceAnalyzer(IAnalysisSession& InSession, FNetProfilerProvider& InNetProfilerProvider)
 	: Session(InSession)
 	, NetProfilerProvider(InNetProfilerProvider)
+	, FrameProvider(ReadFrameProvider(InSession))
 	, NetTraceVersion(0)
 	, NetTraceReporterVersion(0)
 {
@@ -44,6 +47,8 @@ void FNetTraceAnalyzer::OnAnalysisBegin(const FOnAnalysisContext& Context)
 	// Add some default event types that we use for generic type events to make it easier to extend
 	// ConnectionAdded/Removed connections state /name etc?
 	Builder.RouteEvent(RouteId_ConnectionClosedEvent, "NetTrace", "ConnectionClosedEvent");
+	Builder.RouteEvent(RouteId_PacketStatsCounterEvent, "NetTrace", "PacketStatsCounterEvent");
+	Builder.RouteEvent(RouteId_FrameStatsCounterEvent, "NetTrace", "FrameStatsCounterEvent");
 	Builder.RouteEvent(RouteId_ObjectCreatedEvent, "NetTrace", "ObjectCreatedEvent");
 	Builder.RouteEvent(RouteId_ObjectDestroyedEvent, "NetTrace", "ObjectDestroyedEvent");
 	Builder.RouteEvent(RouteId_ConnectionStateUpdatedEvent, "NetTrace", "ConnectionStateUpdatedEvent");
@@ -153,6 +158,18 @@ bool FNetTraceAnalyzer::OnEvent(uint16 RouteId, EStyle Style, const FOnEventCont
 		}
 		break;
 
+		case RouteId_PacketStatsCounterEvent:
+		{
+			HandlePacketStatsCounterEvent(Context, EventData);
+		}
+		break;
+
+		case RouteId_FrameStatsCounterEvent:
+		{
+			HandleFrameStatsCounterEvent(Context, EventData);
+		}
+		break;
+
 		case RouteId_ConnectionStateUpdatedEvent:
 		{
 			HandleConnectionStateUpdatedEvent(Context, EventData);
@@ -209,8 +226,6 @@ void FNetTraceAnalyzer::HandlePacketContentEvent(const FOnEventContext& Context,
 	}
 
 	const ENetProfilerConnectionMode ConnectionMode = ENetProfilerConnectionMode(PacketType);
-	FNetProfilerConnectionData& ConnectionData = NetProfilerProvider.EditConnectionData(ConnectionState->ConnectionIndex, ConnectionMode);
-	++ConnectionData.ContentEventChangeCount;
 
 	TArray<FNetProfilerContentEvent>& Events = (ConnectionState->BunchEvents)[ConnectionMode];
 	TArray<FBunchInfo>& BunchInfos = (ConnectionState->BunchInfos)[ConnectionMode];
@@ -407,6 +422,7 @@ void FNetTraceAnalyzer::FlushPacketEvents(FNetTraceConnectionState& ConnectionSt
 	
 		NextBunchOffset = FMath::Max<uint64>(BunchEvent.EndPos, NextBunchOffset);
 		++CurrentBunchEventIndex;
+		++ConnectionData.ContentEventChangeCount;
 	}
 
 	uint32 EventsToAdd = 0U;
@@ -462,6 +478,12 @@ void FNetTraceAnalyzer::HandlePacketEvent(const FOnEventContext& Context, const
 	// Update LastTimestamp, later on we will be able to get timestamps piggybacked from other analyzers
 	LastTimeStamp = Context.EventTime.AsSeconds(TimestampCycles);
 
+	// Get the NetProfilerFrameIndex for the current engine frame/timestamp
+	const uint32 NetProfilerFrameIndex = GetCurrentNetProfilerFrameIndexAndFlushFrameStatsCountersIfNeeded(
+		GameInstanceId, 
+		FrameProvider.GetFrameNumberForTimestamp(ETraceFrameType::TraceFrameType_Game, LastTimeStamp)
+	);
+
 	FNetTraceConnectionState* ConnectionState = GetActiveConnectionState(GameInstanceId, ConnectionId);
 	if (!ConnectionState)
 	{
@@ -476,6 +498,8 @@ void FNetTraceAnalyzer::HandlePacketEvent(const FOnEventContext& Context, const
 	// Flush packet events
 	FlushPacketEvents(*ConnectionState, ConnectionData, ConnectionMode);
 
+	Packet.NetProfilerFrameIndex = NetProfilerFrameIndex;
+
 	// Fill in packet data a packet must have at least 1 event?
 	Packet.StartEventIndex = ConnectionState->CurrentPacketStartIndex[ConnectionMode];
 	Packet.EventCount = ConnectionData.ContentEvents.Num() - Packet.StartEventIndex;
@@ -488,13 +512,72 @@ void FNetTraceAnalyzer::HandlePacketEvent(const FOnEventContext& Context, const
 	Packet.TotalPacketSizeInBytes = (Packet.ContentSizeInBits + 7u) >> 3u;
 	Packet.DeliveryStatus = ENetProfilerDeliveryStatus::Delivered;
 
+	// Flush PacketStats
+	Packet.StartStatsIndex = ConnectionData.PacketStats.Num();
+	Packet.StatsCount = ConnectionState->PacketStats.Num();
+	for (const FNetProfilerStats& Stat : ConnectionState->PacketStats)
+	{
+		ConnectionData.PacketStats.EmplaceBack(Stat);
+	}
+
+	ConnectionState->PacketStats.Reset();
+	++ConnectionData.PacketStatsChangeCount;
+
+	
 	// Mark the beginning of a new packet
 	ConnectionState->CurrentPacketStartIndex[ConnectionMode] = ConnectionData.ContentEvents.Num();
 	ConnectionState->CurrentPacketBitOffset[ConnectionMode] = 0U;
+	ConnectionState->CurrentPacketStatsStartIndex[ConnectionMode] = ConnectionData.PacketStats.Num();
 
 	//UE_LOG(LogNetTrace, Log, TEXT("PacketEvent GameInstanceId: %u, ConnectionId: %u, %s, Seq: %u PacketBits: %u"), (uint32)GameInstanceId, (uint32)ConnectionId, ConnectionMode ? TEXT("Incoming") : TEXT("Outgoing"), SequenceNumber, Packet.ContentSizeInBits);
 }
 
+void FNetTraceAnalyzer::FlushFrameStatsCounters(FNetTraceAnalyzer::FNetTraceGameInstanceState& GameInstanceState)
+{
+	if (FNetProfilerGameInstanceInternal* GameInstance = NetProfilerProvider.EditGameInstance(GameInstanceState.GameInstanceIndex))
+	{
+		const bool bIsNewFrame = (GameInstance->Frames->Num() == 0) || (GameInstance->Frames->Last().EngineFrameNumber != GameInstanceState.CurrentEngineFrameIndex);
+
+		if (bIsNewFrame)
+		{
+			FNetProfilerFrame& Frame = GameInstance->Frames->EmplaceBack();
+			Frame.EngineFrameNumber = GameInstanceState.CurrentEngineFrameIndex;
+			Frame.StartStatsIndex = GameInstance->FrameStats->Num();
+			Frame.StatsCount = GameInstanceState.FrameStatsCounters.Num();
+			Frame.TimeStamp = LastTimeStamp;
+		}
+		else
+		{
+			FNetProfilerFrame& Frame = GameInstance->Frames->Last();
+			Frame.StatsCount += GameInstanceState.FrameStatsCounters.Num();
+		}
+
+		for (const FNetProfilerStats& Stat : GameInstanceState.FrameStatsCounters)
+		{
+			GameInstance->FrameStats->EmplaceBack(Stat);
+		}
+
+		GameInstanceState.FrameStatsCounters.Reset();
+		GameInstanceState.CurrentNetProfilerFrameIndex = GameInstance->Frames->Num();
+
+		// Mark frames dirty
+		++GameInstance->FramesChangeCount;
+	}
+}
+
+uint32 FNetTraceAnalyzer::GetCurrentNetProfilerFrameIndexAndFlushFrameStatsCountersIfNeeded(uint32 GameInstanceId, uint32 EngineFrameIndex)
+{
+	TSharedRef<FNetTraceAnalyzer::FNetTraceGameInstanceState> GameInstanceStateRef = GetOrCreateActiveGameInstanceState(GameInstanceId);
+	FNetTraceAnalyzer::FNetTraceGameInstanceState& GameInstanceState = GameInstanceStateRef.Get();
+	if (EngineFrameIndex > GameInstanceState.CurrentEngineFrameIndex)
+	{		
+		FlushFrameStatsCounters(GameInstanceState);
+		GameInstanceState.CurrentEngineFrameIndex = EngineFrameIndex;
+	}
+
+	return GameInstanceState.CurrentNetProfilerFrameIndex;
+}
+
 void FNetTraceAnalyzer::HandlePacketDroppedEvent(const FOnEventContext& Context, const FEventData& EventData)
 {
 	const uint64 TimestampCycles = EventData.GetValue<uint64>("Timestamp");
@@ -540,6 +623,72 @@ void FNetTraceAnalyzer::HandleConnectionCreatedEvent(const FOnEventContext& Cont
 
 	ConnectionState->CurrentPacketBitOffset[ENetProfilerConnectionMode::Outgoing] = 0U;
 	ConnectionState->CurrentPacketBitOffset[ENetProfilerConnectionMode::Incoming] = 0U;
+
+	ConnectionState->CurrentPacketStatsStartIndex[ENetProfilerConnectionMode::Outgoing] = 0U;
+	ConnectionState->CurrentPacketStatsStartIndex[ENetProfilerConnectionMode::Incoming] = 0U;
+}
+
+uint32 FNetTraceAnalyzer::GetOrCreateNetProfilerStatsCounterTypeIndex(uint32 NameId, ENetProfilerStatsCounterType StatsType)
+{
+	if (const uint32* ExistingNetProfilerStatsCounterTypeIndex = TraceNetStatsCounterIdToNetProfilerStatsCounterTypeIndexMap.Find(NameId))
+	{
+		return *ExistingNetProfilerStatsCounterTypeIndex;
+	}
+	else
+	{
+		// Add new counter type
+		const uint32* NetProfilerNameIndex = TracedNameIdToNetProfilerNameIdMap.Find(NameId);
+		const uint32 NameIndex = NetProfilerNameIndex ? *NetProfilerNameIndex : 0u;
+
+		uint32 NetProfilerStatsCounterTypeIndex = NetProfilerProvider.AddNetProfilerStatsCounterType(NameIndex, StatsType);
+		TraceNetStatsCounterIdToNetProfilerStatsCounterTypeIndexMap.Add(NameId, NetProfilerStatsCounterTypeIndex);
+
+		return NetProfilerStatsCounterTypeIndex;
+	}
+}
+
+void FNetTraceAnalyzer::HandlePacketStatsCounterEvent(const FOnEventContext& Context, const FEventData& EventData)
+{
+	const uint32 StatsCounterValue = EventData.GetValue<uint32>("StatsValue");
+	const uint8 GameInstanceId = EventData.GetValue<uint8>("GameInstanceId");
+	const uint16 ConnectionId = EventData.GetValue<uint16>("ConnectionId");
+	const uint16 NameId = EventData.GetValue<uint16>("NameId");
+	
+	FNetProfilerStats Stats;
+	Stats.StatsCounterTypeIndex = GetOrCreateNetProfilerStatsCounterTypeIndex(NameId, ENetProfilerStatsCounterType::Packet);
+	Stats.StatsValue = StatsCounterValue;
+
+	// Accumulate stats for current packet
+	TSharedRef<FNetTraceGameInstanceState> GameInstanceState = GetOrCreateActiveGameInstanceState(GameInstanceId);
+	if (TSharedRef<FNetTraceConnectionState>* ConnectionState = GameInstanceState->ActiveConnections.Find(ConnectionId))
+	{
+		(*ConnectionState)->PacketStats.Emplace(Stats);			
+	}
+}
+
+void FNetTraceAnalyzer::HandleFrameStatsCounterEvent(const FOnEventContext& Context, const FEventData& EventData)
+{
+	const uint64 TimestampCycles = EventData.GetValue<uint64>("Timestamp");
+	const uint32 StatsCounterValue = EventData.GetValue<uint32>("StatsValue");
+	const uint8 GameInstanceId = EventData.GetValue<uint8>("GameInstanceId");
+	const uint16 NameId = EventData.GetValue<uint16>("NameId");
+
+	// Update LastTimestamp, later on we will be able to get timestamps piggybacked from other analyzers
+	LastTimeStamp = Context.EventTime.AsSeconds(TimestampCycles);
+
+	// Get the NetProfilerFrameIndex for the current engine frame/timestamp
+	const uint32 NetProfilerFrameIndex = GetCurrentNetProfilerFrameIndexAndFlushFrameStatsCountersIfNeeded(
+		GameInstanceId,
+		FrameProvider.GetFrameNumberForTimestamp(ETraceFrameType::TraceFrameType_Game, LastTimeStamp)
+	);
+
+	FNetProfilerStats Stats;
+	Stats.StatsCounterTypeIndex = GetOrCreateNetProfilerStatsCounterTypeIndex(NameId, ENetProfilerStatsCounterType::Frame);
+	Stats.StatsValue = StatsCounterValue;
+
+	// Accumulate stats for current frame
+	TSharedRef<FNetTraceGameInstanceState> GameInstanceState = GetOrCreateActiveGameInstanceState(GameInstanceId);
+	GameInstanceState->FrameStatsCounters.Emplace(Stats);
 }
 
 void FNetTraceAnalyzer::HandleConnectionStateUpdatedEvent(const FOnEventContext& Context, const FEventData& EventData)
@@ -688,6 +837,8 @@ TSharedRef<FNetTraceAnalyzer::FNetTraceGameInstanceState> FNetTraceAnalyzer::Get
 		ActiveGameInstances.Add(GameInstanceId, GameInstanceState);
 		GameInstanceState->GameInstanceIndex = GameInstance.Instance.GameInstanceIndex;
 
+		GameInstanceState->CurrentEngineFrameIndex = FMath::Min<uint32>(FrameProvider.GetFrameCount(TraceFrameType_Game) - 1U, 0U);
+
 		return GameInstanceState;
 	}
 }
@@ -700,6 +851,7 @@ void FNetTraceAnalyzer::DestroyActiveGameInstanceState(uint32 GameInstanceId)
 		if (FNetProfilerGameInstanceInternal* GameInstance = NetProfilerProvider.EditGameInstance((*FoundState)->GameInstanceIndex))
 		{
 			GameInstance->Instance.LifeTime.End = GetLastTimestamp();
+			NetProfilerProvider.MarkGameInstancesDirty();
 		}
 		ActiveGameInstances.Remove(GameInstanceId);
 	}
@@ -731,6 +883,7 @@ void FNetTraceAnalyzer::HandleGameInstanceUpdatedEvent(const FOnEventContext& Co
 	{
 		InternalGameInstance->Instance.bIsServer = bIsServer;
 		InternalGameInstance->Instance.InstanceName = Session.StoreString(InstanceName);
+		NetProfilerProvider.MarkGameInstancesDirty();
 	}
 }
 
diff --git a/Engine/Source/Developer/TraceServices/Private/Analyzers/NetTraceAnalyzer.h b/Engine/Source/Developer/TraceServices/Private/Analyzers/NetTraceAnalyzer.h
--- a/Engine/Source/Developer/TraceServices/Private/Analyzers/NetTraceAnalyzer.h
+++ b/Engine/Source/Developer/TraceServices/Private/Analyzers/NetTraceAnalyzer.h
@@ -10,6 +10,7 @@ namespace TraceServices
 {
 
 class IAnalysisSession;
+class IFrameProvider;
 
 class FNetTraceAnalyzer
 	: public UE::Trace::IAnalyzer
@@ -33,6 +34,8 @@ private:
 		RouteId_ConnectionCreatedEvent,
 		RouteId_ConnectionUpdatedEvent,
 		RouteId_ConnectionClosedEvent,
+		RouteId_PacketStatsCounterEvent,
+		RouteId_FrameStatsCounterEvent,
 		RouteId_ObjectCreatedEvent,
 		RouteId_ObjectDestroyedEvent,
 		RouteId_ConnectionStateUpdatedEvent,
@@ -50,8 +53,11 @@ private:
 
 	IAnalysisSession& Session;
 	FNetProfilerProvider& NetProfilerProvider;
+	const IFrameProvider& FrameProvider;
 	uint32 NetTraceVersion;
 	uint32 NetTraceReporterVersion;
+
+	// Default names
 	uint32 BunchHeaderNameIndex;
 
 	// Shared for trace
@@ -59,6 +65,8 @@ private:
 
 	TMap<uint32, uint32> TraceEventTypeToNetProfilerEventTypeIndexMap;
 
+	TMap<uint16, uint32> TraceNetStatsCounterIdToNetProfilerStatsCounterTypeIndexMap;
+
 	struct FBunchInfo
 	{
 		FNetProfilerBunchInfo BunchInfo;
@@ -77,6 +85,10 @@ private:
 		// Current packet data
 		uint32 CurrentPacketStartIndex[ENetProfilerConnectionMode::Count];
 		uint32 CurrentPacketBitOffset[ENetProfilerConnectionMode::Count];
+		uint32 CurrentPacketStatsStartIndex[ENetProfilerConnectionMode::Count];
+
+		// StatsCounters reported during the current packet
+		TArray<FNetProfilerStats> PacketStats;
 
 		// Current bunch data
 		TArray<FNetProfilerContentEvent> BunchEvents[ENetProfilerConnectionMode::Count];
@@ -99,12 +111,18 @@ private:
 		TMap<uint64, FNetTraceActiveObjectState> ActiveObjects;
 		TMap<int32, uint32> ChannelNames;
 
+		// StatsCounters reported during the current fram
+		TArray<FNetProfilerStats> FrameStatsCounters;
+		uint32 CurrentEngineFrameIndex = 0;
+		uint32 CurrentNetProfilerFrameIndex = 0;
+
 		uint32 GameInstanceIndex;
 	};
 
 private:
 
 	uint32 GetTracedEventTypeIndex(uint16 NameIndex, uint8 Level);
+	uint32 GetOrCreateNetProfilerStatsCounterTypeIndex(uint32 NameId, ENetProfilerStatsCounterType StatsType);
 
 	TSharedRef<FNetTraceGameInstanceState> GetOrCreateActiveGameInstanceState(uint32 GameInstanceId);
 	void DestroyActiveGameInstanceState(uint32 GameInstanceId);
@@ -114,6 +132,9 @@ private:
 
 	void FlushPacketEvents(FNetTraceConnectionState& ConnectionState, FNetProfilerConnectionData& ConnectionData, const ENetProfilerConnectionMode ConnectionMode);
 
+	void FlushFrameStatsCounters(FNetTraceAnalyzer::FNetTraceGameInstanceState& GameInstanceState);
+	uint32 GetCurrentNetProfilerFrameIndexAndFlushFrameStatsCountersIfNeeded(uint32 GameInstanceId, uint32 EngineFrameIndex);
+
 	void HandlePacketEvent(const FOnEventContext& Context, const FEventData& EventData);	
 	void HandlePacketContentEvent(const FOnEventContext& Context, const FEventData& EventData);
 	void HandlePacketDroppedEvent(const FOnEventContext& Context, const FEventData& EventData);
@@ -121,6 +142,8 @@ private:
 	void HandleConnectionCreatedEvent(const FOnEventContext& Context, const FEventData& EventData);
 	void HandleConnectionUpdatedEvent(const FOnEventContext& Context, const FEventData& EventData);
 	void HandleConnectionClosedEvent(const FOnEventContext& Context, const FEventData& EventData);
+	void HandleFrameStatsCounterEvent(const FOnEventContext& Context, const FEventData& EventData);
+	void HandlePacketStatsCounterEvent(const FOnEventContext& Context, const FEventData& EventData);
 	void HandleObjectCreatedEvent(const FOnEventContext& Context, const FEventData& EventData);
 	void HandleObjectDestroyedEvent(const FOnEventContext& Context, const FEventData& EventData);
 
diff --git a/Engine/Source/Developer/TraceServices/Private/Model/NetProfilerProvider.cpp b/Engine/Source/Developer/TraceServices/Private/Model/NetProfilerProvider.cpp
--- a/Engine/Source/Developer/TraceServices/Private/Model/NetProfilerProvider.cpp
+++ b/Engine/Source/Developer/TraceServices/Private/Model/NetProfilerProvider.cpp
@@ -51,9 +51,10 @@ const TCHAR* LexToString(const ENetProfilerConnectionState Value)
 
 FNetProfilerProvider::FNetProfilerProvider(IAnalysisSession& InSession)
 	: Session(InSession)
-	, NetTraceVersion(0U)
 	, Connections(InSession.GetLinearAllocator(), 4096)
+	, NetTraceVersion(0U)
 	, ConnectionChangeCount(0u)
+	, GameInstanceChangeCount(0u)
 {
 	// Use name index 0 to indicate that we do not know the name
 	AddNetProfilerName(TEXT("N/A"));
@@ -118,6 +119,23 @@ const FNetProfilerEventType* FNetProfilerProvider::GetNetProfilerEventType(uint3
 	return EventTypeIndex < (uint32)EventTypes.Num() ? &EventTypes[EventTypeIndex] : nullptr;
 }
 
+uint32 FNetProfilerProvider::AddNetProfilerStatsCounterType(uint32 NameIndex, ENetProfilerStatsCounterType Type)
+{
+	Session.WriteAccessCheck();
+
+	FNetProfilerStatsCounterType& NewStatsCounterType = StatsCounterTypes.AddDefaulted_GetRef();
+	NewStatsCounterType.StatsCounterTypeIndex = StatsCounterTypes.Num() - 1;
+	NewStatsCounterType.NameIndex = NameIndex;
+	NewStatsCounterType.Type = Type;
+
+	return NewStatsCounterType.StatsCounterTypeIndex;
+}
+
+const FNetProfilerStatsCounterType* FNetProfilerProvider::GetNetProfilerStatsCounterType(uint32 StatsCounterTypeIndex) const
+{
+	return StatsCounterTypeIndex < (uint32)StatsCounterTypes.Num() ? &StatsCounterTypes[StatsCounterTypeIndex] : nullptr;
+}
+
 FNetProfilerGameInstanceInternal& FNetProfilerProvider::CreateGameInstance()
 {
 	Session.WriteAccessCheck();
@@ -129,9 +147,18 @@ FNetProfilerGameInstanceInternal& FNetProfilerProvider::CreateGameInstance()
 	new (GameInstance.Objects) TPagedArray<FNetProfilerObjectInstance>(Session.GetLinearAllocator(), 4096);
 	GameInstance.ObjectsChangeCount = 0u;
 
+	GameInstance.Frames = (TPagedArray<FNetProfilerFrame>*)Session.GetLinearAllocator().Allocate(sizeof(TPagedArray<FNetProfilerFrame>));
+	new (GameInstance.Frames) TPagedArray<FNetProfilerFrame>(Session.GetLinearAllocator(), 4096);
+	GameInstance.FramesChangeCount = 0u;
+
+	GameInstance.FrameStats = (TPagedArray<FNetProfilerStats>*)Session.GetLinearAllocator().Allocate(sizeof(TPagedArray<FNetProfilerStats>));
+	new (GameInstance.FrameStats) TPagedArray<FNetProfilerStats>(Session.GetLinearAllocator(), 4096);
+
 	// We reserve object index 0 as an invalid object
 	CreateObject(GameInstance.Instance.GameInstanceIndex);
 
+	MarkGameInstancesDirty();
+
 	return GameInstance;
 }
 
@@ -149,6 +176,12 @@ FNetProfilerGameInstanceInternal* FNetProfilerProvider::EditGameInstance(uint32
 	}
 }
 
+void FNetProfilerProvider::MarkGameInstancesDirty()
+{
+	Session.WriteAccessCheck();
+	++GameInstanceChangeCount;
+}
+
 FNetProfilerConnectionInternal& FNetProfilerProvider::CreateConnection(uint32 GameInstanceIndex)
 {
 	Session.WriteAccessCheck();
@@ -297,6 +330,21 @@ void FNetProfilerProvider::ReadEventType(uint32 EventTypeIndex, TFunctionRef<voi
 	Callback(*GetNetProfilerEventType(EventTypeIndex));
 }
 
+void FNetProfilerProvider::ReadNetStatsCounterTypes(TFunctionRef<void(const FNetProfilerStatsCounterType*, uint64)> Callback) const
+{
+	Session.ReadAccessCheck();
+
+	Callback(StatsCounterTypes.GetData(), StatsCounterTypes.Num());
+}
+
+void FNetProfilerProvider::ReadNetStatsCounterType(uint32 TypeIndex, TFunctionRef<void(const FNetProfilerStatsCounterType&)> Callback) const
+{
+	Session.ReadAccessCheck();
+	check(TypeIndex < (uint32)StatsCounterTypes.Num());
+
+	Callback(*GetNetProfilerStatsCounterType(TypeIndex));
+}
+
 void FNetProfilerProvider::ReadGameInstances(TFunctionRef<void(const FNetProfilerGameInstance&)> Callback) const
 {
 	Session.ReadAccessCheck();
@@ -691,6 +739,90 @@ ITable<FNetProfilerAggregatedStats>* FNetProfilerProvider::CreateAggregation(uin
 	return Table;
 }
 
+ITable<FNetProfilerAggregatedStatsCounterStats>* FNetProfilerProvider::CreateStatsCountersAggregation(uint32 ConnectionIndex, ENetProfilerConnectionMode Mode, uint32 PacketIndexIntervalStart, uint32 PacketIndexIntervalEnd) const
+{
+	Session.ReadAccessCheck();
+
+	if (!ensure(ConnectionIndex < Connections.Num()))
+	{
+		return nullptr;
+	}
+
+	// [PacketIndexIntervalStart, PacketIndexIntervalEnd] is an inclusive interval.
+	if (!ensure(PacketIndexIntervalStart <= PacketIndexIntervalEnd))
+	{
+		return nullptr;
+	}
+
+	const FNetProfilerConnectionData* ConnectionData = Connections[ConnectionIndex].Data[Mode];
+	const auto& Packets = ConnectionData->Packets;
+
+	const uint32 PacketCount = Packets.Num();
+
+	if (!ensure(PacketCount > 0))
+	{
+		return nullptr;
+	}
+
+	TMap<uint32, FNetProfilerAggregatedStatsCounterStats> AggregatedStatsMap;
+	AggregatedStatsMap.Reserve(StatsCounterTypes.Num());
+
+	auto AccumulateStatsFunction = [&AggregatedStatsMap, this](const FNetProfilerStats& StatsCounter)
+	{
+		FNetProfilerAggregatedStatsCounterStats* StatsEntry = AggregatedStatsMap.Find(StatsCounter.StatsCounterTypeIndex);
+		if (!StatsEntry)
+		{
+			StatsEntry = &AggregatedStatsMap.Add(StatsCounter.StatsCounterTypeIndex);
+			StatsEntry->StatsCounterTypeIndex = StatsCounter.StatsCounterTypeIndex;
+		}
+		const uint32 StatsValue = StatsCounter.StatsValue;
+		StatsEntry->Sum += StatsValue;
+		StatsEntry->Max = FMath::Max(StatsValue, StatsEntry->Max);
+		++StatsEntry->Count;
+	};
+
+	const FNetProfilerGameInstanceInternal& GameInstance = GameInstances[Connections[ConnectionIndex].Connection.GameInstanceIndex];
+	const auto& PacketStatsCounters = ConnectionData->PacketStats;
+	const auto& FrameStatsCounters = *GameInstance.FrameStats;
+	const auto& Frames = *GameInstance.Frames;
+
+	// Iterate over packets
+	for (uint32 PacketIt = PacketIndexIntervalStart, PacketEndIt = FMath::Min(PacketIndexIntervalEnd, PacketCount - 1u); PacketIt <= PacketEndIt; ++PacketIt)
+	{
+		const FNetProfilerPacket& Packet = Packets[PacketIt];
+		const uint32 StatsCounterCount = PacketStatsCounters.Num();
+
+		// Iterate over all StatsCounters stored for the Packet
+		for (uint32 StatsCounterIt = 0; StatsCounterIt < Packet.StatsCount; ++StatsCounterIt)
+		{
+			AccumulateStatsFunction(PacketStatsCounters[Packet.StartStatsIndex + StatsCounterIt]);
+		}
+
+		// Include frame stats as well
+		const uint32 NetProfilerFrameIndex = Packet.NetProfilerFrameIndex;
+		if (NetProfilerFrameIndex < Frames.Num())
+		{
+			const FNetProfilerFrame& Frame = Frames[Packet.NetProfilerFrameIndex];
+			for (uint32 StatsCounterIt = 0; StatsCounterIt < Frame.StatsCount; ++StatsCounterIt)
+			{
+				AccumulateStatsFunction(FrameStatsCounters[Frame.StartStatsIndex + StatsCounterIt]);
+			}
+		}
+	}
+
+	// Calculate averages and populate table
+	TTable<FNetProfilerAggregatedStatsCounterStats>* Table = new TTable<FNetProfilerAggregatedStatsCounterStats>(AggregatedStatsCounterStatsTableLayout);
+	for (const auto& KV : AggregatedStatsMap)
+	{
+		FNetProfilerAggregatedStatsCounterStats& Row = Table->AddRow();
+		Row = KV.Value;
+
+		// Finalize StatsEntry
+		Row.Average = (uint64)((double)KV.Value.Sum / KV.Value.Count);
+	}
+	return Table;
+}
+
 FName GetNetProfilerProviderName()
 {
 	static FName Name(TEXT("NetProfilerProvider"));
diff --git a/Engine/Source/Developer/TraceServices/Private/Model/NetProfilerProvider.h b/Engine/Source/Developer/TraceServices/Private/Model/NetProfilerProvider.h
--- a/Engine/Source/Developer/TraceServices/Private/Model/NetProfilerProvider.h
+++ b/Engine/Source/Developer/TraceServices/Private/Model/NetProfilerProvider.h
@@ -18,22 +18,31 @@ struct FNetProfilerGameInstanceInternal
 
 	TPagedArray<FNetProfilerObjectInstance>* Objects;
 	TArray<uint32, TInlineAllocator<128>> Connections;
+
+	TPagedArray<FNetProfilerFrame>* Frames;
+	TPagedArray<FNetProfilerStats>* FrameStats;
+
 	uint32 ObjectsChangeCount;
+	uint32 FramesChangeCount;	
 };
 
 struct FNetProfilerConnectionData
 {
 	FNetProfilerConnectionData(ILinearAllocator& Allocator)
 		: Packets(Allocator, 1024)
+		, PacketStats(Allocator, 4096)
 		, ContentEvents(Allocator, 8192)
 		, PacketChangeCount(0u)
+		, PacketStatsChangeCount(0u)
 		, ContentEventChangeCount(0u)
 	{}
 
 	TPagedArray<FNetProfilerPacket> Packets;
+	TPagedArray<FNetProfilerStats> PacketStats;
 	TPagedArray<FNetProfilerContentEvent> ContentEvents;
 
 	uint32 PacketChangeCount;
+	uint32 PacketStatsChangeCount;
 	uint32 ContentEventChangeCount;
 };
 
@@ -57,8 +66,11 @@ public:
 
 	uint32 AddNetProfilerEventType(uint32 NameIndex, uint32 Level);
 
+	uint32 AddNetProfilerStatsCounterType(uint32 NameIndex, ENetProfilerStatsCounterType Type);
+
 	FNetProfilerGameInstanceInternal& CreateGameInstance();
 	FNetProfilerGameInstanceInternal* EditGameInstance(uint32 GameInstanceIndex);
+	void MarkGameInstancesDirty();
 
 	FNetProfilerConnectionInternal& CreateConnection(uint32 GameInstanceIndex);
 	FNetProfilerConnectionInternal* EditConnection(uint32 ConnectionIndex);
@@ -86,6 +98,7 @@ public:
 	// Access GameInstances
 	virtual uint32 GetGameInstanceCount() const override { return GameInstances.Num(); }
 	virtual void ReadGameInstances(TFunctionRef<void(const FNetProfilerGameInstance&)> Callback) const override;
+	virtual uint32 GetGameInstanceChangeCount() const override { return GameInstanceChangeCount; }
 
 	// Access Connections
 	virtual uint32 GetConnectionCount(uint32 GameInstanceIndex) const override;
@@ -112,30 +125,42 @@ public:
 	virtual void EnumeratePacketContentEventsByPosition(uint32 ConnectionIndex, ENetProfilerConnectionMode Mode, uint32 PacketIndex, uint32 StartPos, uint32 EndPos, TFunctionRef<void(const FNetProfilerContentEvent&)> Callback) const override;
 	virtual uint32 GetPacketContentEventChangeCount(uint32 ConnectionIndex, ENetProfilerConnectionMode Mode) const override;
 
+	// Access StatsCounterTypes
+	virtual uint32 GetNetStatsCounterTypesCount() const override { return StatsCounterTypes.Num(); }
+	virtual void ReadNetStatsCounterTypes(TFunctionRef<void(const FNetProfilerStatsCounterType*, uint64)> Callback) const override;
+	virtual void ReadNetStatsCounterType(uint32 TypeIndex, TFunctionRef<void(const FNetProfilerStatsCounterType&)> Callback) const override;
+
 	// Stats queries
 	virtual ITable<FNetProfilerAggregatedStats>* CreateAggregation(uint32 ConnectionIndex, ENetProfilerConnectionMode Mode, uint32 PacketIndexIntervalStart, uint32 PacketIndexIntervalEnd, uint32 StartPosition, uint32 EndPosition) const override;
+	virtual ITable<FNetProfilerAggregatedStatsCounterStats>* CreateStatsCountersAggregation(uint32 ConnectionIndex, ENetProfilerConnectionMode Mode, uint32 PacketIndexIntervalStart, uint32 PacketIndexIntervalEnd) const override;
 
 	const TCHAR* InternalGetNetProfilerName(uint32 NameIndex) const { return Names[NameIndex].Name; }
 private:
 	
 	const FNetProfilerName* GetNetProfilerName(uint32 ProfilerNameId) const;
 	const FNetProfilerEventType* GetNetProfilerEventType(uint32 ProfilerEventTypeId) const;
+	const FNetProfilerStatsCounterType* GetNetProfilerStatsCounterType(uint32 ProfilerStatsCounterTypeId) const;
 
 	IAnalysisSession& Session;
 
-	uint32 NetTraceVersion;
-
 	TArray<FNetProfilerName> Names;
 
 	TArray<FNetProfilerEventType> EventTypes;
 
+	TArray<FNetProfilerStatsCounterType> StatsCounterTypes;
+
 	// All GameInstances seen throughout the session
 	TArray<FNetProfilerGameInstanceInternal, TInlineAllocator<4>> GameInstances;
 
 	// All connections we have seen throughout the session
 	TPagedArray<FNetProfilerConnectionInternal> Connections;
-	uint32 ConnectionChangeCount;
+
 	TTableLayout<FNetProfilerAggregatedStats> AggregatedStatsTableLayout;
+	TTableLayout<FNetProfilerAggregatedStatsCounterStats> AggregatedStatsCounterStatsTableLayout;
+
+	uint32 NetTraceVersion;
+	uint32 ConnectionChangeCount;
+	uint32 GameInstanceChangeCount;
 };
 
 } // namespace TraceServices
diff --git a/Engine/Source/Developer/TraceServices/Public/TraceServices/Model/NetProfiler.h b/Engine/Source/Developer/TraceServices/Public/TraceServices/Model/NetProfiler.h
--- a/Engine/Source/Developer/TraceServices/Public/TraceServices/Model/NetProfiler.h
+++ b/Engine/Source/Developer/TraceServices/Public/TraceServices/Model/NetProfiler.h
@@ -11,6 +11,13 @@
 namespace TraceServices
 {
 
+enum class ENetProfilerAggregationMode : uint32
+{
+	None,
+	Aggregate,
+	InstanceMax,
+};
+
 enum class ENetProfilerDeliveryStatus : uint8
 {
 	Unknown,
@@ -117,8 +124,8 @@ struct FNetProfilerContentEvent
 {
 	uint64 StartPos : 24;		// Inclusive start position in the packet
 	uint64 EndPos : 24;			// Exclusive end position in the packet; BitSize = EndPos - StartPos
-	uint64 Level : 4;			// Level
-	uint64 Padding : 12;		// Padding
+	uint64 Level : 8;			// Level
+	uint64 Padding : 8;		// Padding
 
 	FNetProfilerBunchInfo BunchInfo;	
 
@@ -127,6 +134,26 @@ struct FNetProfilerContentEvent
 	uint32 ObjectInstanceIndex;	// Object instance, Non zero if this is a NetObject, we can then look up data by indexing into ObjectInstances
 };
 
+// This must be kept in sync with ENetTraceStatsCounterType in NetTrace.h
+enum class ENetProfilerStatsCounterType : uint8
+{
+	Packet = 0,
+	Frame = 1,
+};
+
+struct FNetProfilerStatsCounterType
+{
+	uint32 StatsCounterTypeIndex;
+	uint32 NameIndex;
+	ENetProfilerStatsCounterType Type;
+};
+
+struct FNetProfilerStats
+{
+	uint32 StatsCounterTypeIndex;
+	uint32 StatsValue;
+};
+
 struct FNetProfilerPacket
 {
 	FNetProfilerTimeStamp TimeStamp;
@@ -138,10 +165,24 @@ struct FNetProfilerPacket
 	uint32 StartEventIndex;
 	uint32 EventCount;
 
+	uint32 StartStatsIndex;
+	uint32 StatsCount;
+
+	// Index into network profiler frames
+	uint32 NetProfilerFrameIndex;
+
 	ENetProfilerDeliveryStatus DeliveryStatus;		// Indicates if the packet was delivered or not, updated as soon as we know
 	ENetProfilerConnectionState ConnectionState;
 };
 
+struct FNetProfilerFrame
+{
+	FNetProfilerTimeStamp TimeStamp;
+	uint64 EngineFrameNumber = 0;
+	uint32 StartStatsIndex = 0;
+	uint32 StatsCount = 0;
+};
+
 struct FNetProfilerConnection
 {
 	const TCHAR* Name = nullptr;
@@ -160,10 +201,10 @@ struct FNetProfilerGameInstance
 	uint32 GameInstanceIndex;
 	uint32 GameInstanceId;
 	const TCHAR* InstanceName = nullptr;
-	bool bIsServer;
+	bool bIsUsingIrisReplication = false;
+	bool bIsServer = false;
 };
 
-// What do we need?
 struct FNetProfilerAggregatedStats
 {
 	uint32 EventTypeIndex;
@@ -175,6 +216,25 @@ struct FNetProfilerAggregatedStats
 	uint32 MaxExclusive = 0U;
 };
 
+struct FNetProfilerAggregatedStatsCounterStats
+{
+	uint32 StatsCounterTypeIndex;
+	uint32 Sum;
+	uint32 Min;
+	uint32 Max;
+	uint32 Average;
+	uint32 Count;
+
+	void Reset()
+	{
+		Sum = 0;
+		Min = 0;
+		Max = 0;
+		Average = 0;
+		Count = 0;
+	}
+};
+
 // What queries do we need?
 class INetProfilerProvider : public IProvider
 {
@@ -198,6 +258,7 @@ public:
 	// Access GameInstances
 	virtual uint32 GetGameInstanceCount() const = 0;
 	virtual void ReadGameInstances(TFunctionRef<void(const FNetProfilerGameInstance&)> Callback) const = 0;
+	virtual uint32 GetGameInstanceChangeCount() const = 0;
 
 	// Access Connections
 	virtual uint32 GetConnectionCount(uint32 GameInstanceIndex) const = 0;
@@ -229,10 +290,20 @@ public:
 	// Returns a change number incremented each time a change occurs in the packet content events for the specified connection and connection mode. */
 	virtual uint32 GetPacketContentEventChangeCount(uint32 ConnectionIndex, ENetProfilerConnectionMode Mode) const = 0;
 
+	// Access StatsCounterTypes
+	virtual uint32 GetNetStatsCounterTypesCount() const = 0;
+	virtual void ReadNetStatsCounterTypes(TFunctionRef<void(const FNetProfilerStatsCounterType*, uint64)> Callback) const = 0;
+	virtual void ReadNetStatsCounterType(uint32 TypeIndex, TFunctionRef<void(const FNetProfilerStatsCounterType&)> Callback) const = 0;
+
+
 	// Computes aggregated stats for a packet interval or for a range of content events in a single packet.
 	// [PacketIndexIntervalStart, PacketIndexIntervalEnd] is the inclusive packet interval.
 	// [StartPosition, EndPosition) is the exclusive bit range interval; only used when PacketIndexIntervalStart == PacketIndexIntervalEnd.
 	virtual ITable<FNetProfilerAggregatedStats>* CreateAggregation(uint32 ConnectionIndex, ENetProfilerConnectionMode Mode, uint32 PacketIndexIntervalStart, uint32 PacketIndexIntervalEnd, uint32 StartPosition, uint32 EndPosition) const = 0;
+	
+	// Computes aggregated statscounters for a packet interval
+	// [PacketIndexIntervalStart, PacketIndexIntervalEnd] is the inclusive packet interval.
+	virtual ITable<FNetProfilerAggregatedStatsCounterStats>* CreateStatsCountersAggregation(uint32 ConnectionIndex, ENetProfilerConnectionMode Mode, uint32 PacketIndexIntervalStart, uint32 PacketIndexIntervalEnd) const = 0;
 };
 
 TRACESERVICES_API FName GetNetProfilerProviderName();
diff --git a/Engine/Source/Runtime/Net/Core/Private/Net/Core/Trace/NetTraceInternal.cpp b/Engine/Source/Runtime/Net/Core/Private/Net/Core/Trace/NetTraceInternal.cpp
--- a/Engine/Source/Runtime/Net/Core/Private/Net/Core/Trace/NetTraceInternal.cpp
+++ b/Engine/Source/Runtime/Net/Core/Private/Net/Core/Trace/NetTraceInternal.cpp
@@ -490,6 +490,22 @@ void FNetTrace::TraceConnectionClosed(uint32 GameInstanceId, uint32 ConnectionId
 	}
 }
 
+void FNetTrace::TracePacketStatsCounter(uint32 GameInstanceId, uint32 ConnectionId, FNetDebugNameId CounterNameId, uint32 StatValue)
+{
+	if (GNetTraceRuntimeVerbosity)
+	{
+		FNetTraceInternal::Reporter::ReportPacketStatsCounter(GameInstanceId, ConnectionId, CounterNameId, StatValue);
+	}
+}
+
+void FNetTrace::TraceFrameStatsCounter(uint32 GameInstanceId, FNetDebugNameId CounterNameId, uint32 StatValue)
+{
+	if (GNetTraceRuntimeVerbosity)
+	{
+		FNetTraceInternal::Reporter::ReportFrameStatsCounter(GameInstanceId, CounterNameId, StatValue);
+	}
+}
+
 FNetDebugNameId FNetTrace::TraceName(const TCHAR* Name)
 {
 	if ((GNetTraceRuntimeVerbosity == 0U) | (Name == nullptr))
diff --git a/Engine/Source/Runtime/Net/Core/Private/Net/Core/Trace/Reporters/NetTraceReporter.cpp b/Engine/Source/Runtime/Net/Core/Private/Net/Core/Trace/Reporters/NetTraceReporter.cpp
--- a/Engine/Source/Runtime/Net/Core/Private/Net/Core/Trace/Reporters/NetTraceReporter.cpp
+++ b/Engine/Source/Runtime/Net/Core/Private/Net/Core/Trace/Reporters/NetTraceReporter.cpp
@@ -67,6 +67,21 @@ UE_TRACE_EVENT_BEGIN(NetTrace, ConnectionClosedEvent)
 	UE_TRACE_EVENT_FIELD(uint8, GameInstanceId)
 UE_TRACE_EVENT_END()
 
+// StatsCounterEvent
+UE_TRACE_EVENT_BEGIN(NetTrace, PacketStatsCounterEvent)
+	UE_TRACE_EVENT_FIELD(uint32, StatsValue)
+	UE_TRACE_EVENT_FIELD(uint16, NameId)
+	UE_TRACE_EVENT_FIELD(uint16, ConnectionId)
+	UE_TRACE_EVENT_FIELD(uint8, GameInstanceId)
+UE_TRACE_EVENT_END()
+
+UE_TRACE_EVENT_BEGIN(NetTrace, FrameStatsCounterEvent)
+	UE_TRACE_EVENT_FIELD(uint64, Timestamp)
+	UE_TRACE_EVENT_FIELD(uint32, StatsValue)
+	UE_TRACE_EVENT_FIELD(uint16, NameId)
+	UE_TRACE_EVENT_FIELD(uint8, GameInstanceId)
+UE_TRACE_EVENT_END()
+
 // Provides additional information about game instance
 UE_TRACE_EVENT_BEGIN(NetTrace, InstanceUpdatedEvent)
 	UE_TRACE_EVENT_FIELD(uint8, GameInstanceId)
@@ -287,6 +302,24 @@ void FNetTraceReporter::ReportConnectionClosed(uint32 GameInstanceId, uint32 Con
 		<< ConnectionClosedEvent.GameInstanceId(GameInstanceId);
 }
 
+void FNetTraceReporter::ReportPacketStatsCounter(uint32 GameInstanceId, uint32 ConnectionId, FNetDebugNameId CounterNameId, uint32 StatValue)
+{
+	UE_TRACE_LOG(NetTrace, PacketStatsCounterEvent, NetChannel)
+		<< PacketStatsCounterEvent.StatsValue(StatValue)
+		<< PacketStatsCounterEvent.NameId(CounterNameId)
+		<< PacketStatsCounterEvent.ConnectionId(ConnectionId)
+		<< PacketStatsCounterEvent.GameInstanceId(GameInstanceId);
+}
+
+void FNetTraceReporter::ReportFrameStatsCounter(uint32 GameInstanceId, FNetDebugNameId CounterNameId, uint32 StatValue)
+{
+	UE_TRACE_LOG(NetTrace, FrameStatsCounterEvent, NetChannel)
+		<< FrameStatsCounterEvent.Timestamp(FPlatformTime::Cycles64())
+		<< FrameStatsCounterEvent.StatsValue(StatValue)
+		<< FrameStatsCounterEvent.NameId(CounterNameId)
+		<< FrameStatsCounterEvent.GameInstanceId(GameInstanceId);
+}
+
 void FNetTraceReporter::ReportObjectCreated(uint32 GameInstanceId, uint32 NetObjectId, FNetDebugNameId NameId, uint64 TypeIdentifier, uint32 OwnerId)
 {
 	UE_TRACE_LOG(NetTrace, ObjectCreatedEvent, NetChannel)
diff --git a/Engine/Source/Runtime/Net/Core/Private/Net/Core/Trace/Reporters/NetTraceReporter.h b/Engine/Source/Runtime/Net/Core/Private/Net/Core/Trace/Reporters/NetTraceReporter.h
--- a/Engine/Source/Runtime/Net/Core/Private/Net/Core/Trace/Reporters/NetTraceReporter.h
+++ b/Engine/Source/Runtime/Net/Core/Private/Net/Core/Trace/Reporters/NetTraceReporter.h
@@ -21,6 +21,8 @@ struct FNetTraceReporter
 	static void ReportConnectionStateUpdated(uint32 GameInstanceId, uint32 ConnectionId, uint8 ConnectionStateValue);
 	static void ReportConnectionUpdated(uint32 GameInstanceId, uint32 ConnectionId, const TCHAR* AddressString, const TCHAR* OwningActor);
 	static void ReportConnectionClosed(uint32 GameInstanceId, uint32 ConnectionId);
+	static void ReportPacketStatsCounter(uint32 GameInstanceId, uint32 ConnectionId, FNetDebugNameId CounterNameId, uint32 StatValue);
+	static void ReportFrameStatsCounter(uint32 GameInstanceId, FNetDebugNameId CounterNameId, uint32 StatValue);
 	static void ReportInstanceUpdated(uint32 GameInstanceId, bool bIsServer, const TCHAR* Name);
 	static void ReportInstanceDestroyed(uint32 GameInstanceId);
 };
diff --git a/Engine/Source/Runtime/Net/Core/Public/Net/Core/Trace/NetTrace.h b/Engine/Source/Runtime/Net/Core/Public/Net/Core/Trace/NetTrace.h
--- a/Engine/Source/Runtime/Net/Core/Public/Net/Core/Trace/NetTrace.h
+++ b/Engine/Source/Runtime/Net/Core/Public/Net/Core/Trace/NetTrace.h
@@ -99,6 +99,12 @@
 /** Trace additional information about the given connection */
 #define UE_NET_TRACE_CONNECTION_UPDATED(...) UE_NET_TRACE_INTERNAL_CONNECTION_UPDATED(__VA_ARGS__)
 
+/** Trace StatsCounter associated with next reported packet */
+#define UE_NET_TRACE_PACKET_STATSCOUNTER(GameInstanceId, ConnectionId, Name, StatValue, Verbosity) UE_NET_TRACE_INTERNAL_PACKET_STATSCOUNTER(GameInstanceId, ConnectionId, Name, StatValue, Verbosity)
+
+/** Trace StatsCounter associated with current frame */
+#define UE_NET_TRACE_FRAME_STATSCOUNTER(GameInstanceId, Name, StatValue, Verbosity) UE_NET_TRACE_INTERNAL_FRAME_STATSCOUNTER(GameInstanceId, Name, StatValue, Verbosity)
+
 /** That that a connection has been closed */
 #define UE_NET_TRACE_CONNECTION_CLOSED(...) UE_NET_TRACE_INTERNAL_CONNECTION_CLOSED(__VA_ARGS__)
 
@@ -151,6 +157,9 @@
 #define UE_NET_TRACE_NAMED_OBJECT_SCOPE(...)
 #define UE_NET_TRACE_NAMED_DYNAMIC_NAME_SCOPE(...)
 
+#define UE_NET_TRACE_PACKET_STATSCOUNTER(...)
+#define UE_NET_TRACE_FRAME_STATSCOUNTER(...)
+
 #define UE_NET_TRACE_PACKET_DROPPED(...)
 #define UE_NET_TRACE_PACKET_SEND(...)
 #define UE_NET_TRACE_PACKET_RECV(...)
diff --git a/Engine/Source/Runtime/Net/Core/Public/Net/Core/Trace/Private/NetTraceInternal.h b/Engine/Source/Runtime/Net/Core/Public/Net/Core/Trace/Private/NetTraceInternal.h
--- a/Engine/Source/Runtime/Net/Core/Public/Net/Core/Trace/Private/NetTraceInternal.h
+++ b/Engine/Source/Runtime/Net/Core/Public/Net/Core/Trace/Private/NetTraceInternal.h
@@ -12,13 +12,8 @@
 #include "Templates/ChooseClass.h"
 
 // $TODO: Implement pooling of TraceCollectors
-// $TODO: Add name to game instance
-// $TODO: Add name/ip to connections
-// $TODO: Add state to connection
 // $TODO: Add disconnect reason to connection
-// $TODO: Add frame concept to understand grouping of packets
 // $TODO: Add reporting of final packet sizes after packethandlers
-// $TODO: Add explicit events for Bunch/BunchHeader including
 
 extern NETCORE_API uint32 GNetTraceRuntimeVerbosity;
 
@@ -36,6 +31,12 @@ enum class ENetTracePacketType : uint8
 	Incoming,
 };
 
+enum class ENetTraceStatsCounterType : uint8
+{
+	Packet = 0,
+	Frame = 1,
+};
+
 struct FNetTracePacketInfo
 {
 	uint32 GameInstanceId;
@@ -230,6 +231,10 @@ struct FNetTrace
 	/** Trace that we have removed a connection for the given GameInstanceId */
 	NETCORE_API static void TraceConnectionClosed(uint32 GameInstanceId, uint32 ConnectionId);
 
+	/** Trace stats counters for networking */
+	NETCORE_API static void TracePacketStatsCounter(uint32 GameInstanceId, uint32 ConnectionId, FNetDebugNameId CounterNameId, uint32 StatValue);
+	NETCORE_API static void TraceFrameStatsCounter(uint32 GameInstanceId, FNetDebugNameId CounterNameId,  uint32 StatValue);
+
 	/** Trace the name */
 	NETCORE_API static FNetDebugNameId TraceName(const TCHAR* Name);
 
@@ -562,6 +567,18 @@ void FNetTraceCollector::Reset()
 #define UE_NET_TRACE_INTERNAL_PACKET_SEND(...) UE_NET_TRACE_DO_IF(GNetTraceRuntimeVerbosity, FNetTrace::TracePacket(__VA_ARGS__, ENetTracePacketType::Outgoing))
 #define UE_NET_TRACE_INTERNAL_PACKET_RECV(...) UE_NET_TRACE_DO_IF(GNetTraceRuntimeVerbosity, FNetTrace::TracePacket(__VA_ARGS__, ENetTracePacketType::Incoming))
 
+#define UE_NET_TRACE_INTERNAL_PACKET_STATSCOUNTER(GameInstanceId, ConnectionId, Name, StatValue, Verbosity) \
+	do { if (FNetTrace::GetNetTraceVerbosityEnabled(Verbosity)) { \
+			static uint16 PREPROCESSOR_JOIN(NetTraceNameId_, __LINE__) = FNetTrace::TraceName(TEXT(#Name)); \
+			FNetTrace::TracePacketStatsCounter(GameInstanceId, ConnectionId, PREPROCESSOR_JOIN(NetTraceNameId_, __LINE__), StatValue); \
+	} } while (0)
+
+#define UE_NET_TRACE_INTERNAL_FRAME_STATSCOUNTER(GameInstanceId, Name, StatValue, Verbosity) \
+	do { if (FNetTrace::GetNetTraceVerbosityEnabled(Verbosity)) { \
+			static uint16 PREPROCESSOR_JOIN(NetTraceNameId_, __LINE__) = FNetTrace::TraceName(TEXT(#Name)); \
+			FNetTrace::TraceFrameStatsCounter(GameInstanceId, PREPROCESSOR_JOIN(NetTraceNameId_, __LINE__), StatValue); \
+	} } while (0)
+
 #define UE_NET_TRACE_INTERNAL_END_SESSION(GameInstanceId) FNetTrace::TraceEndSession(GameInstanceId);
 #define UE_NET_TRACE_INTERNAL_UPDATE_INSTANCE(...) FNetTrace::TraceInstanceUpdated(__VA_ARGS__)
 					
