Fixed issue found by horde https://horde.devtools.epicgames.com/job/6299c4613f24decea4e97dcc?step=726b&issue=177288

diff --git a/Engine/Source/Developer/TraceInsights/Private/Insights/NetworkingProfiler/Widgets/SNetworkingProfilerWindow.cpp b/Engine/Source/Developer/TraceInsights/Private/Insights/NetworkingProfiler/Widgets/SNetworkingProfilerWindow.cpp
--- a/Engine/Source/Developer/TraceInsights/Private/Insights/NetworkingProfiler/Widgets/SNetworkingProfilerWindow.cpp
+++ b/Engine/Source/Developer/TraceInsights/Private/Insights/NetworkingProfiler/Widgets/SNetworkingProfilerWindow.cpp
@@ -584,12 +584,12 @@ FText SNetworkingProfilerWindow::FGameInstanceItem::GetTooltipText() const
 {
 	if (IsInstanceNameSet())
 	{
-		return FText::Format(LOCTEXT("GameInstanceItemFmt0", "{0} Game Instance {1} [{2}]"), FText::FromString(GetInstanceName()), FText::AsNumber(GetIndex()),
+		return FText::Format(LOCTEXT("GameInstanceItemTooltipFmt0", "{0} Game Instance {1} [{2}]"), FText::FromString(GetInstanceName()), FText::AsNumber(GetIndex()),
 			IsServer() ? FText::FromString("Server") : FText::FromString("Client"));
 	}
 	else
 	{
-		return FText::Format(LOCTEXT("GameInstanceItemFmt1", "Game Instance {0}"), FText::AsNumber(GetIndex()));
+		return FText::Format(LOCTEXT("GameInstanceItemTooltipFmt1", "Game Instance {0}"), FText::AsNumber(GetIndex()));
 	}
 }
 
