Updating style files now that developer programs have their own styles

diff --git a/Engine/Source/Developer/TraceInsights/Private/Insights/InsightsStyle.cpp b/Engine/Source/Developer/TraceInsights/Private/Insights/InsightsStyle.cpp
--- a/Engine/Source/Developer/TraceInsights/Private/Insights/InsightsStyle.cpp
+++ b/Engine/Source/Developer/TraceInsights/Private/Insights/InsightsStyle.cpp
@@ -410,108 +410,6 @@ void FInsightsStyle::FStyle::Initialize()
 		Set("SecondaryToolbar2.MaxUniformToolbarSize", 32.0f);
 	}
 
-	//////////////////////////////////////////////////
-	// Workaround for Automation styles...
-#if 1
-
-	const FVector2D Icon24x24(24.0f, 24.0f);
-	const FVector2D Icon40x40(40.0f, 40.0f);
-
-	Set("Automation.Header", FTextBlockStyle(NormalText)
-		.SetFont(DEFAULT_FONT("Mono", 12))
-		.SetColorAndOpacity(FLinearColor(FColor(0xffffffff))));
-
-	Set("Automation.Normal", FTextBlockStyle(NormalText)
-		.SetFont(DEFAULT_FONT("Mono", 9))
-		.SetColorAndOpacity(FLinearColor(FColor(0xffaaaaaa))));
-
-	Set("Automation.Warning", FTextBlockStyle(NormalText)
-		.SetFont(DEFAULT_FONT("Mono", 9))
-		.SetColorAndOpacity(FLinearColor(FColor(0xffbbbb44))));
-
-	Set("Automation.Error", FTextBlockStyle(NormalText)
-		.SetFont(DEFAULT_FONT("Mono", 9))
-		.SetColorAndOpacity(FLinearColor(FColor(0xffff0000))));
-
-	Set("Automation.ReportHeader", FTextBlockStyle(NormalText)
-		.SetFont(DEFAULT_FONT("Mono", 10))
-		.SetColorAndOpacity(FLinearColor(FColor(0xffffffff))));
-
-	Set("Automation.Success", new EDITOR_IMAGE_BRUSH("Automation/Success", Icon16x16));
-	Set("Automation.Warning", new EDITOR_IMAGE_BRUSH("Automation/Warning", Icon16x16));
-	Set("Automation.Fail", new EDITOR_IMAGE_BRUSH("Automation/Fail", Icon16x16));
-	Set("Automation.InProcess", new EDITOR_IMAGE_BRUSH("Automation/InProcess", Icon16x16));
-	Set("Automation.NotRun", new EDITOR_IMAGE_BRUSH("Automation/NotRun", Icon16x16, FLinearColor(0.0f, 0.0f, 0.0f, 0.4f)));
-	Set("Automation.Skipped", new EDITOR_IMAGE_BRUSH("Automation/NoSessionWarning", Icon16x16));
-	Set("Automation.ParticipantsWarning", new EDITOR_IMAGE_BRUSH("Automation/ParticipantsWarning", Icon16x16));
-	Set("Automation.Participant", new EDITOR_IMAGE_BRUSH("Automation/Participant", Icon16x16));
-
-	Set("Automation.SmokeTest", new EDITOR_IMAGE_BRUSH("Automation/SmokeTest", Icon16x16));
-	Set("Automation.SmokeTestParent", new EDITOR_IMAGE_BRUSH("Automation/SmokeTestParent", Icon16x16));
-
-	Set("AutomationWindow.RunTests", new EDITOR_IMAGE_BRUSH("Automation/RunTests", Icon40x40));
-	Set("AutomationWindow.RefreshTests", new EDITOR_IMAGE_BRUSH("Automation/RefreshTests", Icon40x40));
-	Set("AutomationWindow.FindWorkers", new EDITOR_IMAGE_BRUSH("Automation/RefreshWorkers", Icon40x40));
-	Set("AutomationWindow.StopTests", new EDITOR_IMAGE_BRUSH("Automation/StopTests", Icon40x40));
-	Set("AutomationWindow.RunTests.Small", new EDITOR_IMAGE_BRUSH("Automation/RunTests", Icon20x20));
-	Set("AutomationWindow.RefreshTests.Small", new EDITOR_IMAGE_BRUSH("Automation/RefreshTests", Icon20x20));
-	Set("AutomationWindow.FindWorkers.Small", new EDITOR_IMAGE_BRUSH("Automation/RefreshWorkers", Icon20x20));
-	Set("AutomationWindow.StopTests.Small", new EDITOR_IMAGE_BRUSH("Automation/StopTests", Icon20x20));
-
-	//filter icons
-	Set("AutomationWindow.ErrorFilter", new EDITOR_IMAGE_BRUSH("Automation/ErrorFilter", Icon40x40));
-	Set("AutomationWindow.WarningFilter", new EDITOR_IMAGE_BRUSH("Automation/WarningFilter", Icon40x40));
-	Set("AutomationWindow.SmokeTestFilter", new EDITOR_IMAGE_BRUSH("Automation/SmokeTestFilter", Icon40x40));
-	Set("AutomationWindow.DeveloperDirectoryContent", new EDITOR_IMAGE_BRUSH("Automation/DeveloperDirectoryContent", Icon40x40));
-	Set("AutomationWindow.ExcludedTestsFilter", new EDITOR_IMAGE_BRUSH("Automation/ExcludedTestsFilter", Icon40x40));
-	Set("AutomationWindow.ErrorFilter.Small", new EDITOR_IMAGE_BRUSH("Automation/ErrorFilter", Icon20x20));
-	Set("AutomationWindow.WarningFilter.Small", new EDITOR_IMAGE_BRUSH("Automation/WarningFilter", Icon20x20));
-	Set("AutomationWindow.SmokeTestFilter.Small", new EDITOR_IMAGE_BRUSH("Automation/SmokeTestFilter", Icon20x20));
-	Set("AutomationWindow.DeveloperDirectoryContent.Small", new EDITOR_IMAGE_BRUSH("Automation/DeveloperDirectoryContent", Icon20x20));
-	Set("AutomationWindow.TrackHistory", new EDITOR_IMAGE_BRUSH("Automation/TrackTestHistory", Icon40x40));
-
-	//device group settings
-	Set("AutomationWindow.GroupSettings", new EDITOR_IMAGE_BRUSH("Automation/Groups", Icon40x40));
-	Set("AutomationWindow.GroupSettings.Small", new EDITOR_IMAGE_BRUSH("Automation/Groups", Icon20x20));
-
-	//test preset icons
-	Set("AutomationWindow.PresetNew", new EDITOR_IMAGE_BRUSH("Icons/icon_add_40x", Icon16x16));
-	Set("AutomationWindow.PresetSave", new EDITOR_IMAGE_BRUSH_SVG("Starship/Common/SaveCurrent", Icon16x16));
-	Set("AutomationWindow.PresetRemove", new EDITOR_IMAGE_BRUSH("Icons/icon_Cascade_DeleteLOD_40x", Icon16x16));
-
-	//test backgrounds
-	Set("AutomationWindow.GameGroupBorder", new EDITOR_BOX_BRUSH("Automation/GameGroupBorder", FMargin(4.0f / 16.0f)));
-	Set("AutomationWindow.EditorGroupBorder", new EDITOR_BOX_BRUSH("Automation/EditorGroupBorder", FMargin(4.0f / 16.0f)));
-
-	Set("MessageLog.ListBorder", new EDITOR_BOX_BRUSH("/Docking/AppTabContentArea", FMargin(4 / 16.0f)));
-
-	Set("Launcher.Platform_Windows", new EDITOR_IMAGE_BRUSH("Launcher/All_Platforms_24x", Icon24x24));
-
-	// Output Log Window
-	{
-		const FSlateColor LogColor_SelectionBackground(EStyleColor::User2);
-		const FSlateColor LogColor_Normal(EStyleColor::User3);
-		const FSlateColor LogColor_Command(EStyleColor::User4);
-
-		const int32 LogFontSize = 9;
-		const FTextBlockStyle NormalLogText = FTextBlockStyle(NormalText)
-			.SetFont(DEFAULT_FONT("Mono", LogFontSize))
-			.SetColorAndOpacity(LogColor_Normal)
-			.SetSelectedBackgroundColor(LogColor_SelectionBackground)
-			.SetHighlightColor(FStyleColors::Black);
-
-		Set("Log.Normal", NormalLogText);
-
-		Set("Log.Command", FTextBlockStyle(NormalLogText)
-			.SetColorAndOpacity(LogColor_Command));
-
-		Set("Log.Warning", FTextBlockStyle(NormalLogText)
-			.SetColorAndOpacity(FStyleColors::Warning));
-
-		Set("Log.Error", FTextBlockStyle(NormalLogText)
-			.SetColorAndOpacity(FStyleColors::Error));
-	}
-
 	{
 		Set("ToggleButton", FButtonStyle(Button)
 			.SetNormal(FSlateNoResource())
@@ -538,7 +436,7 @@ void FInsightsStyle::FStyle::Initialize()
 		Set("Common.GotoNativeCodeHyperlink", EditNativeHyperlinkStyle);
 	}
 
-#endif
+
 	//////////////////////////////////////////////////
 }
 
diff --git a/Engine/Source/Editor/EditorStyle/Private/StarshipStyle.cpp b/Engine/Source/Editor/EditorStyle/Private/StarshipStyle.cpp
--- a/Engine/Source/Editor/EditorStyle/Private/StarshipStyle.cpp
+++ b/Engine/Source/Editor/EditorStyle/Private/StarshipStyle.cpp
@@ -4447,7 +4447,7 @@ void FStarshipEditorStyle::FStyle::SetupLevelEditorStyle()
 		Set( "LevelViewport.BlackBackground", new FSlateColorBrush( FLinearColor::Black ) ); 
 		Set( "LevelViewport.StartingPlayInEditorBorder", new BOX_BRUSH( "Old/Window/ViewportDebugBorder", 0.8f, FLinearColor(0.1f,1.0f,0.1f,1.0f) ) );
 		Set( "LevelViewport.StartingSimulateBorder", new BOX_BRUSH( "Old/Window/ViewportDebugBorder", 0.8f, FLinearColor(1.0f,1.0f,0.1f,1.0f) ) );
-		Set( "LevelViewport.NonMaximizedBorder", new BORDER_BRUSH("Common/PlainBorder", 2.f / 8.f, FStyleColors::Black));
+		Set( "LevelViewport.NonMaximizedBorder", new CORE_BORDER_BRUSH("Common/PlainBorder", 2.f / 8.f, FStyleColors::Black));
 		Set( "LevelViewport.ReturningToEditorBorder", new BOX_BRUSH( "Old/Window/ViewportDebugBorder", 0.8f, FLinearColor(0.1f,0.1f,1.0f,1.0f) ) );
 		Set( "LevelViewport.ActorLockIcon", new IMAGE_BRUSH( "Icons/ActorLockedViewport", Icon32x32 ) );
 		Set( "LevelViewport.Icon", new IMAGE_BRUSH( "Icons/icon_tab_viewport_16px", Icon16x16 ) );
