[Insights] Memory Insights: Added "MemoryTrace_" function name prefix to ignore callstack frames when detecting top function of a callstack.

diff --git a/Engine/Source/Developer/TraceInsights/Private/Insights/MemoryProfiler/ViewModels/MemAllocTable.cpp b/Engine/Source/Developer/TraceInsights/Private/Insights/MemoryProfiler/ViewModels/MemAllocTable.cpp
--- a/Engine/Source/Developer/TraceInsights/Private/Insights/MemoryProfiler/ViewModels/MemAllocTable.cpp
+++ b/Engine/Source/Developer/TraceInsights/Private/Insights/MemoryProfiler/ViewModels/MemAllocTable.cpp
@@ -810,6 +810,7 @@ void FMemAllocTable::AddDefaultColumns()
 								TEXTVIEW("FMallocPoisonProxy::"),
 								TEXTVIEW("Malloc"),
 								TEXTVIEW("Realloc"),
+								TEXTVIEW("MemoryTrace_"),
 							};
 							for (uint32 StringIndex = 0; StringIndex < UE_ARRAY_COUNT(IgnoreSymbolsByFunctionName); ++StringIndex)
 							{
