[Insights] Memory Insights / Mem Allocs: Added "Top Source File" and "Callstack Size" columns. Also improved tooltip for "Top Function" column.

diff --git a/Engine/Source/Developer/TraceInsights/Private/Insights/MemoryProfiler/ViewModels/CallstackFormatting.h b/Engine/Source/Developer/TraceInsights/Private/Insights/MemoryProfiler/ViewModels/CallstackFormatting.h
--- a/Engine/Source/Developer/TraceInsights/Private/Insights/MemoryProfiler/ViewModels/CallstackFormatting.h
+++ b/Engine/Source/Developer/TraceInsights/Private/Insights/MemoryProfiler/ViewModels/CallstackFormatting.h
@@ -2,16 +2,24 @@
 
 #pragma once
 
+#include "Misc/EnumClassFlags.h"
 #include "TraceServices/Model/Callstack.h"
 #include "TraceServices/Model/Modules.h"
 
 ////////////////////////////////////////////////////////////////////////////////////////////////////
 
-enum EStackFrameFormatFlags : uint8
+enum class EStackFrameFormatFlags : uint8
 {
-	Module			= 1 << 0,
-	FileAndLine		= 1 << 1
+	Module  = 1 << 0,
+	Symbol  = 1 << 1,
+	File    = 1 << 2,
+	Line    = 1 << 3,
+
+	ModuleAndSymbol         = Module + Symbol,
+	ModuleSymbolFileAndLine = Module + Symbol + File + Line,
+	FileAndLine             = File + Line,
 };
+ENUM_CLASS_FLAGS(EStackFrameFormatFlags);
 
 ////////////////////////////////////////////////////////////////////////////////////////////////////
 
@@ -29,35 +37,59 @@ inline const TCHAR* GetEmptyCallstackString()
 
 ////////////////////////////////////////////////////////////////////////////////////////////////////
 
-inline void FormatStackFrame(const TraceServices::FStackFrame& Frame, FStringBuilderBase& OutString, uint8 FormatFlags)
+inline void FormatStackFrame(const TraceServices::FStackFrame& Frame, FStringBuilderBase& OutString, EStackFrameFormatFlags FormatFlags)
 {
 	using namespace TraceServices;
 	const ESymbolQueryResult Result = Frame.Symbol->GetResult();
 	switch (Result)
 	{
 		case ESymbolQueryResult::OK:
-			if (FormatFlags & (uint8)EStackFrameFormatFlags::Module)
+			if (EnumHasAnyFlags(FormatFlags, EStackFrameFormatFlags::Module))
+			{
+				OutString.Append(Frame.Symbol->Module);
+				if (FormatFlags != EStackFrameFormatFlags::Module)
+				{
+					OutString.AppendChar(TEXT('!'));
+				}
+			}
+			if (EnumHasAnyFlags(FormatFlags, EStackFrameFormatFlags::Symbol))
 			{
-				OutString.Appendf(TEXT("%s!"), Frame.Symbol->Module);
+				OutString.Append(Frame.Symbol->Name);
 			}
-			OutString.Append(Frame.Symbol->Name);
-			if (FormatFlags & (uint8)EStackFrameFormatFlags::FileAndLine)
+			if (EnumHasAnyFlags(FormatFlags, EStackFrameFormatFlags::File))
 			{
-				OutString.Appendf(TEXT(" %s(%d)"), Frame.Symbol->File, Frame.Symbol->Line);
+				if (EnumHasAnyFlags(FormatFlags, EStackFrameFormatFlags::Symbol))
+				{
+					OutString.AppendChar(TEXT(' '));
+				}
+				OutString.Append(Frame.Symbol->File);
+			}
+			if (EnumHasAnyFlags(FormatFlags, EStackFrameFormatFlags::Line))
+			{
+				OutString.Appendf(TEXT("(%d)"), Frame.Symbol->Line);
 			}
 			break;
 
 		case ESymbolQueryResult::Mismatch:
 		case ESymbolQueryResult::NotFound:
 		case ESymbolQueryResult::NotLoaded:
-			if (FormatFlags & (uint8)EStackFrameFormatFlags::Module)
+			if (EnumHasAnyFlags(FormatFlags, EStackFrameFormatFlags::Module))
+			{
+				OutString.Append(Frame.Symbol->Module);
+				if (FormatFlags != EStackFrameFormatFlags::Module)
+				{
+					OutString.AppendChar(TEXT('!'));
+				}
+			}
+			if (EnumHasAnyFlags(FormatFlags, EStackFrameFormatFlags::Symbol))
 			{
-				OutString.Appendf(TEXT("%s!0x%08x (%s)"), Frame.Symbol->Module, Frame.Addr, QueryResultToString((Result)));
+				OutString.Appendf(TEXT("0x%08x"), Frame.Addr);
 			}
-			else
+			if (EnumHasAnyFlags(FormatFlags, EStackFrameFormatFlags::ModuleAndSymbol))
 			{
-				OutString.Appendf(TEXT("0x%08x (%s)"), Frame.Symbol->Name, Frame.Addr, QueryResultToString((Result)));
+				OutString.AppendChar(TEXT(' '));
 			}
+			OutString.Appendf(TEXT("(%s)"), QueryResultToString((Result)));
 			break;
 
 		case ESymbolQueryResult::Pending:
diff --git a/Engine/Source/Developer/TraceInsights/Private/Insights/MemoryProfiler/ViewModels/MemAllocGroupingByCallstack.cpp b/Engine/Source/Developer/TraceInsights/Private/Insights/MemoryProfiler/ViewModels/MemAllocGroupingByCallstack.cpp
--- a/Engine/Source/Developer/TraceInsights/Private/Insights/MemoryProfiler/ViewModels/MemAllocGroupingByCallstack.cpp
+++ b/Engine/Source/Developer/TraceInsights/Private/Insights/MemoryProfiler/ViewModels/MemAllocGroupingByCallstack.cpp
@@ -198,7 +198,7 @@ FName FMemAllocGroupingByCallstack::GetGroupName(const TraceServices::FStackFram
 FText FMemAllocGroupingByCallstack::GetGroupTooltip(const TraceServices::FStackFrame* Frame) const
 {
 	TStringBuilder<1024> String;
-	FormatStackFrame(*Frame, String, EStackFrameFormatFlags::Module | EStackFrameFormatFlags::FileAndLine);
+	FormatStackFrame(*Frame, String, EStackFrameFormatFlags::ModuleSymbolFileAndLine);
 	return FText::FromString(FString(String));
 }
 
diff --git a/Engine/Source/Developer/TraceInsights/Private/Insights/MemoryProfiler/ViewModels/MemAllocNode.cpp b/Engine/Source/Developer/TraceInsights/Private/Insights/MemoryProfiler/ViewModels/MemAllocNode.cpp
--- a/Engine/Source/Developer/TraceInsights/Private/Insights/MemoryProfiler/ViewModels/MemAllocNode.cpp
+++ b/Engine/Source/Developer/TraceInsights/Private/Insights/MemoryProfiler/ViewModels/MemAllocNode.cpp
@@ -1,6 +1,7 @@
 // Copyright Epic Games, Inc. All Rights Reserved.
 
 #include "MemAllocNode.h"
+#include "Insights/MemoryProfiler/ViewModels/CallstackFormatting.h"
 
 #define LOCTEXT_NAMESPACE "MemAllocNode"
 
@@ -11,14 +12,92 @@ const FName FMemAllocNode::TypeName(TEXT("FMemAllocNode"));
 
 ////////////////////////////////////////////////////////////////////////////////////////////////////
 
+uint64 FMemAllocNode::GetCallstackId() const
+{
+	return IsValidMemAlloc() ? uint64(GetMemAllocChecked().GetCallstack()) : 0ull;
+}
+
+////////////////////////////////////////////////////////////////////////////////////////////////////
+
 FText FMemAllocNode::GetFullCallstack() const
 {
-	if (IsValidMemAlloc())
+	return IsValidMemAlloc() ? GetMemAllocChecked().GetFullCallstack() : FText();
+}
+
+////////////////////////////////////////////////////////////////////////////////////////////////////
+
+FText FMemAllocNode::GetFullCallstackSourceFiles() const
+{
+	return IsValidMemAlloc() ? GetMemAllocChecked().GetFullCallstackSourceFiles() : FText();
+}
+
+////////////////////////////////////////////////////////////////////////////////////////////////////
+
+FText FMemAllocNode::GetTopFunction() const
+{
+	return GetTopFunctionOrSourceFile((uint8)EStackFrameFormatFlags::ModuleAndSymbol);
+}
+
+////////////////////////////////////////////////////////////////////////////////////////////////////
+
+FText FMemAllocNode::GetTopFunctionEx() const
+{
+	return GetTopFunctionOrSourceFile((uint8)EStackFrameFormatFlags::ModuleSymbolFileAndLine);
+}
+
+////////////////////////////////////////////////////////////////////////////////////////////////////
+
+FText FMemAllocNode::GetTopSourceFile() const
+{
+	return GetTopFunctionOrSourceFile((uint8)EStackFrameFormatFlags::File);
+}
+
+////////////////////////////////////////////////////////////////////////////////////////////////////
+
+FText FMemAllocNode::GetTopSourceFileEx() const
+{
+	return GetTopFunctionOrSourceFile((uint8)EStackFrameFormatFlags::FileAndLine);
+}
+
+////////////////////////////////////////////////////////////////////////////////////////////////////
+
+FText FMemAllocNode::GetTopFunctionOrSourceFile(uint8 Flags) const
+{
+	if (!IsValidMemAlloc())
 	{
-		return GetMemAllocChecked().GetFullCallstack();
+		return FText();
+	}
+
+	const Insights::FMemoryAlloc& Alloc = GetMemAllocChecked();
+	const TraceServices::FCallstack* Callstack = Alloc.GetCallstack();
+
+	if (!Callstack)
+	{
+		return FText::FromString(GetCallstackNotAvailableString());
+	}
+
+	if (Callstack->Num() == 0)
+	{
+		return FText::FromString(GetEmptyCallstackString());
+	}
+
+	const TraceServices::FStackFrame* Frame = nullptr;
+	for (uint32 FrameIndex = 0; FrameIndex < Callstack->Num(); ++FrameIndex)
+	{
+		Frame = Callstack->Frame(FrameIndex);
+		check(Frame != nullptr);
+
+		if (Frame->Symbol &&
+			Frame->Symbol->Name &&
+			Frame->Symbol->FilterStatus != TraceServices::EResolvedSymbolFilterStatus::Filtered)
+		{
+			break;
+		}
 	}
 
-	return FText();
+	TStringBuilder<1024> Str;
+	FormatStackFrame(*Frame, Str, (EStackFrameFormatFlags)Flags);
+	return FText::FromString(FString(Str));
 }
 
 ////////////////////////////////////////////////////////////////////////////////////////////////////
diff --git a/Engine/Source/Developer/TraceInsights/Private/Insights/MemoryProfiler/ViewModels/MemAllocNode.h b/Engine/Source/Developer/TraceInsights/Private/Insights/MemoryProfiler/ViewModels/MemAllocNode.h
--- a/Engine/Source/Developer/TraceInsights/Private/Insights/MemoryProfiler/ViewModels/MemAllocNode.h
+++ b/Engine/Source/Developer/TraceInsights/Private/Insights/MemoryProfiler/ViewModels/MemAllocNode.h
@@ -63,7 +63,17 @@ public:
 	bool IsValidMemAlloc() const { return GetMemTableChecked().IsValidRowIndex(GetRowIndex()); }
 	const FMemoryAlloc* GetMemAlloc() const { return GetMemTableChecked().GetMemAlloc(GetRowIndex()); }
 	const FMemoryAlloc& GetMemAllocChecked() const { return GetMemTableChecked().GetMemAllocChecked(GetRowIndex()); }
+
+	uint64 GetCallstackId() const;
 	FText GetFullCallstack() const;
+	FText GetFullCallstackSourceFiles() const;
+	FText GetTopFunction() const;
+	FText GetTopFunctionEx() const;
+	FText GetTopSourceFile() const;
+	FText GetTopSourceFileEx() const;
+
+private:
+	FText GetTopFunctionOrSourceFile(uint8 Flags) const;
 
 private:
 	FMemAllocTable* MemAllocTable;
diff --git a/Engine/Source/Developer/TraceInsights/Private/Insights/MemoryProfiler/ViewModels/MemAllocTable.cpp b/Engine/Source/Developer/TraceInsights/Private/Insights/MemoryProfiler/ViewModels/MemAllocTable.cpp
--- a/Engine/Source/Developer/TraceInsights/Private/Insights/MemoryProfiler/ViewModels/MemAllocTable.cpp
+++ b/Engine/Source/Developer/TraceInsights/Private/Insights/MemoryProfiler/ViewModels/MemAllocTable.cpp
@@ -1,6 +1,7 @@
 // Copyright Epic Games, Inc. All Rights Reserved.
 
 #include "MemAllocTable.h"
+#include "Styling/StyleColors.h"
 #include "Widgets/SBoxPanel.h"
 #include "Widgets/SToolTip.h"
 #include "Widgets/Text/STextBlock.h"
@@ -31,11 +32,13 @@ const FName FMemAllocTableColumns::StartTimeColumnId(TEXT("StartTime"));
 const FName FMemAllocTableColumns::EndTimeColumnId(TEXT("EndTime"));
 const FName FMemAllocTableColumns::DurationColumnId(TEXT("Duration"));
 const FName FMemAllocTableColumns::AddressColumnId(TEXT("Address"));
-const FName FMemAllocTableColumns::MemoryPageColumnId(TEXT("Memory Page"));
+const FName FMemAllocTableColumns::MemoryPageColumnId(TEXT("MemoryPage"));
 const FName FMemAllocTableColumns::CountColumnId(TEXT("Count"));
 const FName FMemAllocTableColumns::SizeColumnId(TEXT("Size"));
 const FName FMemAllocTableColumns::TagColumnId(TEXT("Tag"));
 const FName FMemAllocTableColumns::FunctionColumnId(TEXT("Function"));
+const FName FMemAllocTableColumns::SourceFileColumnId(TEXT("SourceFile"));
+const FName FMemAllocTableColumns::CallstackSizeColumnId(TEXT("CallstackSize"));
 
 ////////////////////////////////////////////////////////////////////////////////////////////////////
 // FMemAllocTable
@@ -682,7 +685,7 @@ void FMemAllocTable::AddDefaultColumns()
 		Column.SetTitleName(LOCTEXT("TagColumnTitle", "LLM Tag"));
 		Column.SetDescription(LOCTEXT("TagColumnDesc", "LLM tag of allocation"));
 
-		// This column is filtered with a custom filter with suggestions 
+		// This column is filtered with a custom filter with suggestions
 		// so we do not mark it here as CanBeFiltered to prevent it from having a default string filter set.
 		Column.SetFlags(ETableColumnFlags::ShouldBeVisible | ETableColumnFlags::CanBeHidden);
 
@@ -731,7 +734,7 @@ void FMemAllocTable::AddDefaultColumns()
 		AddColumn(ColumnRef);
 	}
 	//////////////////////////////////////////////////
-	// Function Column
+	// Top Function Column
 	{
 		TSharedRef<FTableColumn> ColumnRef = MakeShared<FTableColumn>(FMemAllocTableColumns::FunctionColumnId);
 		FTableColumn& Column = *ColumnRef;
@@ -739,7 +742,7 @@ void FMemAllocTable::AddDefaultColumns()
 		Column.SetIndex(ColumnIndex++);
 
 		Column.SetShortName(LOCTEXT("FunctionColumnName", "Function"));
-		Column.SetTitleName(LOCTEXT("FunctionColumnTitle", "Function (Callstack)"));
+		Column.SetTitleName(LOCTEXT("FunctionColumnTitle", "Top Function"));
 		Column.SetDescription(LOCTEXT("FunctionColumnDesc", "Resolved top function from the callstack of allocation"));
 
 		Column.SetFlags(ETableColumnFlags::ShouldBeVisible | ETableColumnFlags::CanBeHidden | ETableColumnFlags::CanBeFiltered);
@@ -767,40 +770,113 @@ void FMemAllocTable::AddDefaultColumns()
 				else //if (Node->Is<FMemAllocNode>())
 				{
 					const FMemAllocNode& MemAllocNode = static_cast<const FMemAllocNode&>(Node);
-					const FMemoryAlloc* Alloc = MemAllocNode.GetMemAlloc();
-					if (Alloc)
-					{
-						const TraceServices::FCallstack* Callstack = Alloc->GetCallstack();
+					const uint64 CallstackId = MemAllocNode.GetCallstackId();
+					const FText CallstackText = MemAllocNode.GetTopFunction();
+					return FTableCellValue(CallstackText, CallstackId);
+				}
+
+				return TOptional<FTableCellValue>();
+			}
+
+			virtual uint64 GetValueId(const FTableColumn& Column, const FBaseTreeNode& Node) const override
+			{
+				if (!Node.IsGroup()) //if (Node->Is<FMemAllocNode>())
+				{
+					const FMemAllocNode& MemAllocNode = static_cast<const FMemAllocNode&>(Node);
+					return MemAllocNode.GetCallstackId();
+				}
+
+				return 0;
+			}
+		};
+		TSharedRef<ITableCellValueGetter> Getter = MakeShared<FFunctionValueGetter>();
+		Column.SetValueGetter(Getter);
+
+		class FFunctionValueFormatter : public FTextValueFormatter
+		{
+		public:
+			virtual TSharedPtr<IToolTip> GetCustomTooltip(const FTableColumn& Column, const FBaseTreeNode& Node) const override
+			{
+				const FMemAllocNode& MemAllocNode = static_cast<const FMemAllocNode&>(Node);
+
+				return SNew(SToolTip)
+					.Visibility(TAttribute<EVisibility>::Create(TAttribute<EVisibility>::FGetter::CreateStatic(&FTableCellValueFormatter::GetTooltipVisibility)))
+					[
+						SNew(SVerticalBox)
+						+ SVerticalBox::Slot()
+						.AutoHeight()
+						.Padding(2.0f)
+						[
+							SNew(STextBlock)
+							.Text(&MemAllocNode, &FMemAllocNode::GetTopFunction)
+							.ColorAndOpacity(FSlateColor(EStyleColor::AccentBlue))
+						]
+						+ SVerticalBox::Slot()
+						.AutoHeight()
+						.Padding(2.0f, 6.0f, 2.0f, 2.0f)
+						[
+							SNew(STextBlock)
+							.Text(&MemAllocNode, &FMemAllocNode::GetFullCallstack)
+						]
+					];
+			}
+
+			virtual FText FormatValueForTooltip(const FTableColumn& Column, const FBaseTreeNode& Node) const override
+			{
+				const FMemAllocNode& MemAllocNode = static_cast<const FMemAllocNode&>(Node);
+				return MemAllocNode.GetTopFunction();
+			}
+		};
 
-						if (!Callstack)
-						{
-							return FTableCellValue(FText::FromString(GetCallstackNotAvailableString()), 0);
-						}
-
-						if (Callstack->Num() == 0)
-						{
-							return FTableCellValue(FText::FromString(GetEmptyCallstackString()), uint64(Callstack));
-						}
-
-						const TraceServices::FStackFrame* Frame = nullptr;
-						for (uint32 FrameIndex = 0; FrameIndex < Callstack->Num(); ++FrameIndex)
-						{
-							Frame = Callstack->Frame(FrameIndex);
-							check(Frame != nullptr);
-
-							if (Frame->Symbol &&
-								Frame->Symbol->Name &&
-								Frame->Symbol->FilterStatus != TraceServices::EResolvedSymbolFilterStatus::Filtered)
-							{
-								break;
-							}
-						}
-
-						TStringBuilder<1024> Str;
-						FormatStackFrame(*Frame, Str, EStackFrameFormatFlags::Module);
-						return FTableCellValue(FText::FromString(FString(Str)), uint64(Callstack));
+		TSharedRef<ITableCellValueFormatter> Formatter = MakeShared<FFunctionValueFormatter>();
+		Column.SetValueFormatter(Formatter);
+
+		TSharedRef<ITableCellValueSorter> Sorter = MakeShared<FSorterByTextValueWithId>(ColumnRef);
+		Column.SetValueSorter(Sorter);
+
+		AddColumn(ColumnRef);
+	}
+	//////////////////////////////////////////////////
+	// Source File Column
+	{
+		TSharedRef<FTableColumn> ColumnRef = MakeShared<FTableColumn>(FMemAllocTableColumns::SourceFileColumnId);
+		FTableColumn& Column = *ColumnRef;
+
+		Column.SetIndex(ColumnIndex++);
+
+		Column.SetShortName(LOCTEXT("SourceFileColumnName", "Source"));
+		Column.SetTitleName(LOCTEXT("SourceFileColumnTitle", "Top Source File"));
+		Column.SetDescription(LOCTEXT("SourceFileColumnDesc", "Source file of the top function from the callstack of allocation"));
+
+		Column.SetFlags(ETableColumnFlags::CanBeHidden | ETableColumnFlags::CanBeFiltered);
+
+		Column.SetHorizontalAlignment(HAlign_Left);
+		Column.SetInitialWidth(550.0f);
+
+		Column.SetDataType(ETableCellDataType::Text);
+
+		Column.SetIsDynamic(true);
+
+		class FSourceValueGetter : public FTableCellValueGetter
+		{
+		public:
+			virtual const TOptional<FTableCellValue> GetValue(const FTableColumn& Column, const FBaseTreeNode& Node) const override
+			{
+				if (Node.IsGroup())
+				{
+					const FTableTreeNode& NodePtr = static_cast<const FTableTreeNode&>(Node);
+					if (NodePtr.HasAggregatedValue(Column.GetId()))
+					{
+						return NodePtr.GetAggregatedValue(Column.GetId());
 					}
 				}
+				else //if (Node->Is<FMemAllocNode>())
+				{
+					const FMemAllocNode& MemAllocNode = static_cast<const FMemAllocNode&>(Node);
+					const uint64 CallstackId = MemAllocNode.GetCallstackId();
+					const FText CallstackText = MemAllocNode.GetTopSourceFile();
+					return FTableCellValue(CallstackText, CallstackId);
+				}
 
 				return TOptional<FTableCellValue>();
 			}
@@ -808,41 +884,144 @@ void FMemAllocTable::AddDefaultColumns()
 			virtual uint64 GetValueId(const FTableColumn& Column, const FBaseTreeNode& Node) const override
 			{
 				if (!Node.IsGroup()) //if (Node->Is<FMemAllocNode>())
+				{
+					const FMemAllocNode& MemAllocNode = static_cast<const FMemAllocNode&>(Node);
+					return MemAllocNode.GetCallstackId();
+				}
+				return 0;
+			}
+		};
+		TSharedRef<ITableCellValueGetter> Getter = MakeShared<FSourceValueGetter>();
+		Column.SetValueGetter(Getter);
+
+		class FSourceValueFormatter : public FTextValueFormatter
+		{
+		public:
+			virtual TSharedPtr<IToolTip> GetCustomTooltip(const FTableColumn& Column, const FBaseTreeNode& Node) const override
+			{
+				const FMemAllocNode& MemAllocNode = static_cast<const FMemAllocNode&>(Node);
+
+				return SNew(SToolTip)
+					.Visibility(TAttribute<EVisibility>::Create(TAttribute<EVisibility>::FGetter::CreateStatic(&FTableCellValueFormatter::GetTooltipVisibility)))
+					[
+						SNew(SVerticalBox)
+						+ SVerticalBox::Slot()
+						.AutoHeight()
+						.Padding(2.0f)
+						[
+							SNew(STextBlock)
+							.Text(&MemAllocNode, &FMemAllocNode::GetTopSourceFileEx)
+							.ColorAndOpacity(FSlateColor(EStyleColor::AccentBlue))
+						]
+						+ SVerticalBox::Slot()
+						.AutoHeight()
+						.Padding(2.0f, 6.0f, 2.0f, 2.0f)
+						[
+							SNew(STextBlock)
+							.Text(&MemAllocNode, &FMemAllocNode::GetFullCallstackSourceFiles)
+						]
+					];
+			}
+
+			virtual FText FormatValueForTooltip(const FTableColumn& Column, const FBaseTreeNode& Node) const override
+			{
+				const FMemAllocNode& MemAllocNode = static_cast<const FMemAllocNode&>(Node);
+				return MemAllocNode.GetTopSourceFileEx();
+			}
+		};
+
+		TSharedRef<ITableCellValueFormatter> Formatter = MakeShared<FSourceValueFormatter>();
+		Column.SetValueFormatter(Formatter);
+
+		TSharedRef<ITableCellValueSorter> Sorter = MakeShared<FSorterByTextValueWithId>(ColumnRef);
+		Column.SetValueSorter(Sorter);
+
+		AddColumn(ColumnRef);
+	}
+	//////////////////////////////////////////////////
+	// Callstack Size Column
+	{
+		TSharedRef<FTableColumn> ColumnRef = MakeShared<FTableColumn>(FMemAllocTableColumns::CallstackSizeColumnId);
+		FTableColumn& Column = *ColumnRef;
+
+		Column.SetIndex(ColumnIndex++);
+
+		Column.SetShortName(LOCTEXT("CallstackSizeColumnName", "Callstack Size"));
+		Column.SetTitleName(LOCTEXT("CallstackSizeColumnTitle", "Callstack Size"));
+		Column.SetDescription(LOCTEXT("CallstackSizeColumnDesc", "Number of callstack frames.\nTooltip shows entire callstack."));
+
+		Column.SetFlags(ETableColumnFlags::ShouldBeVisible | ETableColumnFlags::CanBeHidden | ETableColumnFlags::CanBeFiltered);
+
+		Column.SetHorizontalAlignment(HAlign_Left);
+		Column.SetInitialWidth(100.0f);
+
+		Column.SetDataType(ETableCellDataType::Int64);
+
+		class FCallstackSizeValueGetter : public FTableCellValueGetter
+		{
+		public:
+			virtual const TOptional<FTableCellValue> GetValue(const FTableColumn& Column, const FBaseTreeNode& Node) const override
+			{
+				if (Node.IsGroup())
+				{
+					const FTableTreeNode& NodePtr = static_cast<const FTableTreeNode&>(Node);
+					if (NodePtr.HasAggregatedValue(Column.GetId()))
+					{
+						return NodePtr.GetAggregatedValue(Column.GetId());
+					}
+				}
+				else //if (Node->Is<FMemAllocNode>())
 				{
 					const FMemAllocNode& MemAllocNode = static_cast<const FMemAllocNode&>(Node);
 					const FMemoryAlloc* Alloc = MemAllocNode.GetMemAlloc();
 					if (Alloc)
 					{
 						const TraceServices::FCallstack* Callstack = Alloc->GetCallstack();
-						return uint64(Callstack);
+						const int64 CallstackSize = (Callstack && (Callstack->Num() != 1 || Callstack->Addr(0) != 0)) ? Callstack->Num() : 0;
+						return FTableCellValue(CallstackSize);
 					}
 				}
 
+				return TOptional<FTableCellValue>();
+			}
+
+			virtual uint64 GetValueId(const FTableColumn& Column, const FBaseTreeNode& Node) const override
+			{
+				if (!Node.IsGroup()) //if (Node->Is<FMemAllocNode>())
+				{
+					const FMemAllocNode& MemAllocNode = static_cast<const FMemAllocNode&>(Node);
+					return MemAllocNode.GetCallstackId();
+				}
 				return 0;
 			}
 		};
-		TSharedRef<ITableCellValueGetter> Getter = MakeShared<FFunctionValueGetter>();
+		TSharedRef<ITableCellValueGetter> Getter = MakeShared<FCallstackSizeValueGetter>();
 		Column.SetValueGetter(Getter);
 
-		class FunctionValueFormatter : public FTextValueFormatter
+		class FCallstackValueFormatter : public FTableCellValueFormatter
 		{
 		public:
+			virtual FText FormatValue(const TOptional<FTableCellValue>& InValue) const override
+			{
+				return InValue.IsSet() ? FText::AsNumber(InValue.GetValue().Int64) : FText::GetEmpty();
+			}
+
 			virtual TSharedPtr<IToolTip> GetCustomTooltip(const FTableColumn& Column, const FBaseTreeNode& Node) const override
 			{
 				const FMemAllocNode& MemAllocNode = static_cast<const FMemAllocNode&>(Node);
 
 				return SNew(SToolTip)
-						.Visibility(TAttribute<EVisibility>::Create(TAttribute<EVisibility>::FGetter::CreateStatic(&FTableCellValueFormatter::GetTooltipVisibility)))
+					.Visibility(TAttribute<EVisibility>::Create(TAttribute<EVisibility>::FGetter::CreateStatic(&FTableCellValueFormatter::GetTooltipVisibility)))
+					[
+						SNew(SVerticalBox)
+						+ SVerticalBox::Slot()
+						.AutoHeight()
+						.Padding(2.0f)
 						[
-							SNew(SVerticalBox)
-							+ SVerticalBox::Slot()
-							.AutoHeight()
-							.Padding(2.0f)
-							[
-								SNew(STextBlock)
-								.Text(&MemAllocNode, &FMemAllocNode::GetFullCallstack)
-							]
-						];
+							SNew(STextBlock)
+							.Text(&MemAllocNode, &FMemAllocNode::GetFullCallstack)
+						]
+					];
 			}
 
 			virtual FText FormatValueForTooltip(const FTableColumn& Column, const FBaseTreeNode& Node) const override
@@ -850,12 +1029,17 @@ void FMemAllocTable::AddDefaultColumns()
 				const FMemAllocNode& MemAllocNode = static_cast<const FMemAllocNode&>(Node);
 				return MemAllocNode.GetFullCallstack();
 			}
+
+			virtual FText FormatValueForGrouping(const FTableColumn& Column, const FBaseTreeNode& Node) const override
+			{
+				return FormatValue(Column.GetValue(Node));
+			}
 		};
 
-		TSharedRef<ITableCellValueFormatter> Formatter = MakeShared<FunctionValueFormatter>();
+		TSharedRef<ITableCellValueFormatter> Formatter = MakeShared<FCallstackValueFormatter>();
 		Column.SetValueFormatter(Formatter);
 
-		TSharedRef<ITableCellValueSorter> Sorter = MakeShared<FSorterByTextValueWithId>(ColumnRef);
+		TSharedRef<ITableCellValueSorter> Sorter = MakeShared<FSorterByInt64Value>(ColumnRef);
 		Column.SetValueSorter(Sorter);
 
 		AddColumn(ColumnRef);
diff --git a/Engine/Source/Developer/TraceInsights/Private/Insights/MemoryProfiler/ViewModels/MemAllocTable.h b/Engine/Source/Developer/TraceInsights/Private/Insights/MemoryProfiler/ViewModels/MemAllocTable.h
--- a/Engine/Source/Developer/TraceInsights/Private/Insights/MemoryProfiler/ViewModels/MemAllocTable.h
+++ b/Engine/Source/Developer/TraceInsights/Private/Insights/MemoryProfiler/ViewModels/MemAllocTable.h
@@ -25,11 +25,13 @@ struct FMemAllocTableColumns
 	static const FName EndTimeColumnId;
 	static const FName DurationColumnId;
 	static const FName AddressColumnId;
+	static const FName MemoryPageColumnId;
 	static const FName CountColumnId;
 	static const FName SizeColumnId;
 	static const FName TagColumnId;
 	static const FName FunctionColumnId;
-	static const FName MemoryPageColumnId;
+	static const FName SourceFileColumnId;
+	static const FName CallstackSizeColumnId;
 };
 
 ////////////////////////////////////////////////////////////////////////////////////////////////////
diff --git a/Engine/Source/Developer/TraceInsights/Private/Insights/MemoryProfiler/ViewModels/MemoryAlloc.cpp b/Engine/Source/Developer/TraceInsights/Private/Insights/MemoryProfiler/ViewModels/MemoryAlloc.cpp
--- a/Engine/Source/Developer/TraceInsights/Private/Insights/MemoryProfiler/ViewModels/MemoryAlloc.cpp
+++ b/Engine/Source/Developer/TraceInsights/Private/Insights/MemoryProfiler/ViewModels/MemoryAlloc.cpp
@@ -10,16 +10,16 @@ namespace Insights
 ////////////////////////////////////////////////////////////////////////////////////////////////////
 
 FMemoryAlloc::FMemoryAlloc()
-: StartEventIndex(0)
-, EndEventIndex(0)
-, StartTime(0.0)
-, EndTime(0.0)
-, Address(0)
-, Size(0)
-, Tag(nullptr)
-, Callstack(nullptr)
-, RootHeap(0)
-, bIsBlock(false)
+	: StartEventIndex(0)
+	, EndEventIndex(0)
+	, StartTime(0.0)
+	, EndTime(0.0)
+	, Address(0)
+	, Size(0)
+	, Tag(nullptr)
+	, Callstack(nullptr)
+	, RootHeap(0)
+	, bIsBlock(false)
 {
 }
 
@@ -47,15 +47,46 @@ FText FMemoryAlloc::GetFullCallstack() const
 	const uint32 FramesNum = Callstack->Num();
 	for (uint32 Index = 0; Index < FramesNum; ++Index)
 	{
+		if (Index != 0)
+		{
+			Tooltip << TEXT("\n");
+		}
 		const TraceServices::FStackFrame* Frame = Callstack->Frame(Index);
 		check(Frame != nullptr);
-		FormatStackFrame(*Frame, Tooltip, EStackFrameFormatFlags::Module | EStackFrameFormatFlags::FileAndLine);
-		if (Index != (FramesNum - 1))
+		FormatStackFrame(*Frame, Tooltip, EStackFrameFormatFlags::ModuleSymbolFileAndLine);
+	}
+	return FText::FromString(FString(Tooltip));
+}
+
+////////////////////////////////////////////////////////////////////////////////////////////////////
+
+FText FMemoryAlloc::GetFullCallstackSourceFiles() const
+{
+	if (!Callstack)
+	{
+		return FText::FromString(GetCallstackNotAvailableString());
+	}
+
+	if (Callstack->Num() == 0)
+	{
+		return FText::FromString(GetEmptyCallstackString());
+	}
+
+	TStringBuilder<1024> Tooltip;
+	const uint32 FramesNum = Callstack->Num();
+	for (uint32 Index = 0; Index < FramesNum; ++Index)
+	{
+		if (Index != 0)
 		{
 			Tooltip << TEXT("\n");
 		}
+		const TraceServices::FStackFrame* Frame = Callstack->Frame(Index);
+		check(Frame != nullptr);
+		FormatStackFrame(*Frame, Tooltip, EStackFrameFormatFlags::FileAndLine);
 	}
 	return FText::FromString(FString(Tooltip));
 }
 
+////////////////////////////////////////////////////////////////////////////////////////////////////
+
 } // namespace Insights
diff --git a/Engine/Source/Developer/TraceInsights/Private/Insights/MemoryProfiler/ViewModels/MemoryAlloc.h b/Engine/Source/Developer/TraceInsights/Private/Insights/MemoryProfiler/ViewModels/MemoryAlloc.h
--- a/Engine/Source/Developer/TraceInsights/Private/Insights/MemoryProfiler/ViewModels/MemoryAlloc.h
+++ b/Engine/Source/Developer/TraceInsights/Private/Insights/MemoryProfiler/ViewModels/MemoryAlloc.h
@@ -33,6 +33,7 @@ public:
 	TraceServices::TagIdType GetTagId() const { return TagId; }
 	const TraceServices::FCallstack* GetCallstack() const { return Callstack; }
 	FText GetFullCallstack() const;
+	FText GetFullCallstackSourceFiles() const;
 	HeapId GetRootHeap() const { return RootHeap; };
 
 private:
diff --git a/Engine/Source/Developer/TraceInsights/Private/Insights/MemoryProfiler/Widgets/SMemAllocTableTreeView.cpp b/Engine/Source/Developer/TraceInsights/Private/Insights/MemoryProfiler/Widgets/SMemAllocTableTreeView.cpp
--- a/Engine/Source/Developer/TraceInsights/Private/Insights/MemoryProfiler/Widgets/SMemAllocTableTreeView.cpp
+++ b/Engine/Source/Developer/TraceInsights/Private/Insights/MemoryProfiler/Widgets/SMemAllocTableTreeView.cpp
@@ -559,11 +559,12 @@ void SMemAllocTableTreeView::InitAvailableViewPresets()
 		}
 		virtual void GetColumnConfigSet(TArray<FColumnConfig>& InOutConfigSet) const override
 		{
-			InOutConfigSet.Add({ FTable::GetHierarchyColumnId(),          true, 200.0f });
-			InOutConfigSet.Add({ FMemAllocTableColumns::CountColumnId,    true, 100.0f });
-			InOutConfigSet.Add({ FMemAllocTableColumns::SizeColumnId,     true, 100.0f });
-			InOutConfigSet.Add({ FMemAllocTableColumns::TagColumnId,      true, 120.0f });
-			InOutConfigSet.Add({ FMemAllocTableColumns::FunctionColumnId, true, 550.0f });
+			InOutConfigSet.Add({ FTable::GetHierarchyColumnId(),               true, 200.0f });
+			InOutConfigSet.Add({ FMemAllocTableColumns::CountColumnId,         true, 100.0f });
+			InOutConfigSet.Add({ FMemAllocTableColumns::SizeColumnId,          true, 100.0f });
+			InOutConfigSet.Add({ FMemAllocTableColumns::TagColumnId,           true, 120.0f });
+			InOutConfigSet.Add({ FMemAllocTableColumns::FunctionColumnId,      true, 550.0f });
+			InOutConfigSet.Add({ FMemAllocTableColumns::CallstackSizeColumnId, true, 100.0f });
 		}
 	};
 	AvailableViewPresets.Add(MakeShared<FDefaultViewPreset>());
@@ -612,6 +613,8 @@ void SMemAllocTableTreeView::InitAvailableViewPresets()
 			InOutConfigSet.Add({ FMemAllocTableColumns::SizeColumnId,            true, 100.0f });
 			InOutConfigSet.Add({ FMemAllocTableColumns::TagColumnId,             true, 120.0f });
 			InOutConfigSet.Add({ FMemAllocTableColumns::FunctionColumnId,        true, 550.0f });
+			InOutConfigSet.Add({ FMemAllocTableColumns::SourceFileColumnId,      true, 550.0f });
+			InOutConfigSet.Add({ FMemAllocTableColumns::CallstackSizeColumnId,   true, 100.0f });
 		}
 	};
 	AvailableViewPresets.Add(MakeShared<FDetailedViewPreset>());
diff --git a/Engine/Source/Developer/TraceInsights/Private/Insights/Table/ViewModels/TableCellValueFormatter.cpp b/Engine/Source/Developer/TraceInsights/Private/Insights/Table/ViewModels/TableCellValueFormatter.cpp
--- a/Engine/Source/Developer/TraceInsights/Private/Insights/Table/ViewModels/TableCellValueFormatter.cpp
+++ b/Engine/Source/Developer/TraceInsights/Private/Insights/Table/ViewModels/TableCellValueFormatter.cpp
@@ -32,6 +32,13 @@ FText FTableCellValueFormatter::FormatValueForTooltip(const FTableColumn& Column
 
 ////////////////////////////////////////////////////////////////////////////////////////////////////
 
+FText FTableCellValueFormatter::FormatValueForGrouping(const FTableColumn& Column, const FBaseTreeNode& Node) const
+{
+	return FormatValueForTooltip(Column.GetValue(Node));
+}
+
+////////////////////////////////////////////////////////////////////////////////////////////////////
+
 TSharedPtr<IToolTip> FTableCellValueFormatter::GetCustomTooltip(const FTableColumn& Column, const FBaseTreeNode& Node) const
 {
 	return SNew(SToolTip)
diff --git a/Engine/Source/Developer/TraceInsights/Private/Insights/Table/ViewModels/TableCellValueFormatter.h b/Engine/Source/Developer/TraceInsights/Private/Insights/Table/ViewModels/TableCellValueFormatter.h
--- a/Engine/Source/Developer/TraceInsights/Private/Insights/Table/ViewModels/TableCellValueFormatter.h
+++ b/Engine/Source/Developer/TraceInsights/Private/Insights/Table/ViewModels/TableCellValueFormatter.h
@@ -26,6 +26,7 @@ public:
 
 	virtual FText FormatValue(const FTableColumn& Column, const FBaseTreeNode& Node) const = 0;
 	virtual FText FormatValueForTooltip(const FTableColumn& Column, const FBaseTreeNode& Node) const = 0;
+	virtual FText FormatValueForGrouping(const FTableColumn& Column, const FBaseTreeNode& Node) const = 0;
 
 	virtual TSharedPtr<IToolTip> GetCustomTooltip(const FTableColumn& Column, const FBaseTreeNode& Node) const = 0;
 };
@@ -43,6 +44,7 @@ public:
 
 	virtual FText FormatValue(const FTableColumn& Column, const FBaseTreeNode& Node) const override; // { return FormatValue(Column.GetValue(Node)); }
 	virtual FText FormatValueForTooltip(const FTableColumn& Column, const FBaseTreeNode& Node) const override; // { return FormatValueForTooltip(Column.GetValue(Node)); }
+	virtual FText FormatValueForGrouping(const FTableColumn& Column, const FBaseTreeNode& Node) const override; // { return FormatValueForTooltip(Column.GetValue(Node)); }
 
 	virtual TSharedPtr<IToolTip> GetCustomTooltip(const FTableColumn& Column, const FBaseTreeNode& Node) const override;
 
diff --git a/Engine/Source/Developer/TraceInsights/Private/Insights/Table/ViewModels/TableColumn.cpp b/Engine/Source/Developer/TraceInsights/Private/Insights/Table/ViewModels/TableColumn.cpp
--- a/Engine/Source/Developer/TraceInsights/Private/Insights/Table/ViewModels/TableColumn.cpp
+++ b/Engine/Source/Developer/TraceInsights/Private/Insights/Table/ViewModels/TableColumn.cpp
@@ -55,6 +55,13 @@ FText FTableColumn::GetValueAsTooltipText(const FBaseTreeNode& InNode) const
 
 ////////////////////////////////////////////////////////////////////////////////////////////////////
 
+FText FTableColumn::GetValueAsGroupingText(const FBaseTreeNode& InNode) const
+{
+	return ValueFormatter->FormatValueForGrouping(*this, InNode);
+}
+
+////////////////////////////////////////////////////////////////////////////////////////////////////
+
 inline void TrimEndingZero(FString& Str)
 {
 	const TCHAR* Start = *Str;
diff --git a/Engine/Source/Developer/TraceInsights/Private/Insights/Table/ViewModels/TableColumn.h b/Engine/Source/Developer/TraceInsights/Private/Insights/Table/ViewModels/TableColumn.h
--- a/Engine/Source/Developer/TraceInsights/Private/Insights/Table/ViewModels/TableColumn.h
+++ b/Engine/Source/Developer/TraceInsights/Private/Insights/Table/ViewModels/TableColumn.h
@@ -177,6 +177,7 @@ public:
 
 	FText GetValueAsText(const FBaseTreeNode& InNode) const;
 	FText GetValueAsTooltipText(const FBaseTreeNode& InNode) const;
+	FText GetValueAsGroupingText(const FBaseTreeNode& InNode) const;
 
 	FString GetValueAsSerializableString(const FBaseTreeNode& InNode) const;
 
@@ -261,7 +262,7 @@ private:
 
 	/** Initial sorting mode. */
 	EColumnSortMode::Type InitialSortMode;
-	
+
 	/** Used to convert in a custom way from string to column data type in FilterConfigurator */
 	TSharedPtr<IFilterValueConverter> ValueConverter;
 
diff --git a/Engine/Source/Developer/TraceInsights/Private/Insights/Table/ViewModels/TreeNodeGrouping.h b/Engine/Source/Developer/TraceInsights/Private/Insights/Table/ViewModels/TreeNodeGrouping.h
--- a/Engine/Source/Developer/TraceInsights/Private/Insights/Table/ViewModels/TreeNodeGrouping.h
+++ b/Engine/Source/Developer/TraceInsights/Private/Insights/Table/ViewModels/TreeNodeGrouping.h
@@ -126,7 +126,7 @@ private:
 	static Type GetValue(const FTableCellValue& CellValue);
 	static FText GetValueAsText(const FTableColumn& Column, const FTableTreeNode& Node)
 	{
-		return Column.GetValueAsTooltipText(Node);
+		return Column.GetValueAsGroupingText(Node);
 	}
 };
 
