TasksInsights: added DebugName column to Tasks table. Hidden by default, as not many tasks have debug names yet

diff --git a/Engine/Source/Developer/TraceInsights/Private/Insights/TaskGraphProfiler/ViewModels/TaskEntry.h b/Engine/Source/Developer/TraceInsights/Private/Insights/TaskGraphProfiler/ViewModels/TaskEntry.h
--- a/Engine/Source/Developer/TraceInsights/Private/Insights/TaskGraphProfiler/ViewModels/TaskEntry.h
+++ b/Engine/Source/Developer/TraceInsights/Private/Insights/TaskGraphProfiler/ViewModels/TaskEntry.h
@@ -20,6 +20,8 @@ public:
 	FTaskEntry(const TraceServices::FTaskInfo &TaskInfo);
 	~FTaskEntry() {}
 
+	const TCHAR* GetDebugName() const { return DebugName; }
+
 	double GetCreatedTimestamp() const { return CreatedTimestamp; }
 	uint32 GetCreatedThreadId() const { return CreatedThreadId; }
 
diff --git a/Engine/Source/Developer/TraceInsights/Private/Insights/TaskGraphProfiler/ViewModels/TaskTable.cpp b/Engine/Source/Developer/TraceInsights/Private/Insights/TaskGraphProfiler/ViewModels/TaskTable.cpp
--- a/Engine/Source/Developer/TraceInsights/Private/Insights/TaskGraphProfiler/ViewModels/TaskTable.cpp
+++ b/Engine/Source/Developer/TraceInsights/Private/Insights/TaskGraphProfiler/ViewModels/TaskTable.cpp
@@ -18,6 +18,7 @@ namespace Insights
 ////////////////////////////////////////////////////////////////////////////////////////////////////
 // Column identifiers
 
+const FName FTaskTableColumns::DebugNameColumnId(TEXT("DebugName"));
 const FName FTaskTableColumns::CreatedTimestampColumnId(TEXT("CreatedTimestamp"));
 const FName FTaskTableColumns::CreatedThreadIdColumnId(TEXT("CreatedThreadId"));
 const FName FTaskTableColumns::LaunchedTimestampColumnId(TEXT("LaunchedTimestamp"));
@@ -69,6 +70,7 @@ public:
 
 struct DefaultTaskFieldGetterFuncts
 {
+	static FTableCellValue GetDebugName(const FTableColumn& Column, const FTaskEntry& Task) { return FTableCellValue(Task.GetDebugName()); }
 	static FTableCellValue GetCreatedTimestamp(const FTableColumn& Column, const FTaskEntry& Task) { return FTableCellValue(Task.GetCreatedTimestamp());	}
 	static FTableCellValue GetCreatedThreadId(const FTableColumn& Column, const FTaskEntry& Task) { return FTableCellValue((int64)Task.GetCreatedThreadId());	}
 
@@ -172,7 +174,7 @@ void FTaskTable::AddDefaultColumns()
 		AddHierarchyColumn(HierarchyColumnIndex, HierarchyColumnName);
 
 		const TSharedRef<FTableColumn>& ColumnRef = GetColumns()[0];
-		ColumnRef->SetInitialWidth(200.0f);
+		ColumnRef->SetInitialWidth(100.0f);
 		ColumnRef->SetShortName(LOCTEXT("TaskColumnName", "Hierarchy"));
 		ColumnRef->SetTitleName(LOCTEXT("TaskColumnTitle", "Task Hierarchy"));
 		ColumnRef->SetDescription(LOCTEXT("TaskColumnDesc", "Hierarchy of the task's tree"));
@@ -180,6 +182,39 @@ void FTaskTable::AddDefaultColumns()
 
 	int32 ColumnIndex = 0;
 
+	//////////////////////////////////////////////////
+	// DebugName Column
+	{
+		TSharedRef<FTableColumn> ColumnRef = MakeShared<FTableColumn>(FTaskTableColumns::DebugNameColumnId);
+		FTableColumn& Column = *ColumnRef;
+
+		Column.SetIndex(ColumnIndex++);
+
+		Column.SetShortName(LOCTEXT("DebugNameColumnName", "DebugName"));
+		Column.SetTitleName(LOCTEXT("DebugNameColumnTitle", "DebugName"));
+		Column.SetDescription(LOCTEXT("DebugNameColumnDesc", "DebugName"));
+
+		Column.SetFlags(ETableColumnFlags::CanBeHidden | ETableColumnFlags::CanBeFiltered);
+
+		Column.SetHorizontalAlignment(HAlign_Left);
+		Column.SetInitialWidth(200.0f);
+
+		Column.SetDataType(ETableCellDataType::CString);
+
+		TSharedRef<ITableCellValueGetter> Getter = MakeShared<FTaskColumnValueGetter<DefaultTaskFieldGetterFuncts::GetDebugName>>();
+		Column.SetValueGetter(Getter);
+
+		TSharedRef<ITableCellValueFormatter> Formatter = MakeShared<FCStringValueFormatterAsText>();
+		Column.SetValueFormatter(Formatter);
+
+		TSharedRef<ITableCellValueSorter> Sorter = MakeShared<FSorterByCStringValue>(ColumnRef);
+		Column.SetValueSorter(Sorter);
+
+		Column.SetAggregation(ETableColumnAggregation::SameValue);
+
+		AddColumn(ColumnRef);
+	}
+
 	//////////////////////////////////////////////////
 	// Created Timestamp Column
 	{
diff --git a/Engine/Source/Developer/TraceInsights/Private/Insights/TaskGraphProfiler/ViewModels/TaskTable.h b/Engine/Source/Developer/TraceInsights/Private/Insights/TaskGraphProfiler/ViewModels/TaskTable.h
--- a/Engine/Source/Developer/TraceInsights/Private/Insights/TaskGraphProfiler/ViewModels/TaskTable.h
+++ b/Engine/Source/Developer/TraceInsights/Private/Insights/TaskGraphProfiler/ViewModels/TaskTable.h
@@ -18,6 +18,7 @@ class FTableColumn;
 // Column identifiers
 struct FTaskTableColumns
 {
+	static const FName DebugNameColumnId;
 	static const FName CreatedTimestampColumnId;
 	static const FName CreatedThreadIdColumnId;
 	static const FName LaunchedTimestampColumnId;
