[Insights] Fixed static analysis warning.

diff --git a/Engine/Source/Developer/TraceInsights/Private/Insights/MemoryProfiler/ViewModels/MemAllocNode.cpp b/Engine/Source/Developer/TraceInsights/Private/Insights/MemoryProfiler/ViewModels/MemAllocNode.cpp
--- a/Engine/Source/Developer/TraceInsights/Private/Insights/MemoryProfiler/ViewModels/MemAllocNode.cpp
+++ b/Engine/Source/Developer/TraceInsights/Private/Insights/MemoryProfiler/ViewModels/MemAllocNode.cpp
@@ -76,13 +76,14 @@ FText FMemAllocNode::GetTopFunctionOrSourceFile(uint8 Flags) const
 		return FText::FromString(GetCallstackNotAvailableString());
 	}
 
-	if (Callstack->Num() == 0)
+	const uint32 FrameCount = Callstack->Num();
+	if (FrameCount == 0)
 	{
 		return FText::FromString(GetEmptyCallstackString());
 	}
 
 	const TraceServices::FStackFrame* Frame = nullptr;
-	for (uint32 FrameIndex = 0; FrameIndex < Callstack->Num(); ++FrameIndex)
+	for (uint32 FrameIndex = 0; FrameIndex < FrameCount; ++FrameIndex)
 	{
 		Frame = Callstack->Frame(FrameIndex);
 		check(Frame != nullptr);
