TasksInsights: added tracing task destruction and a context menu item to jump to task's DebugName in IDE if it's a source location

diff --git a/Engine/Source/Developer/TraceInsights/Private/Insights/TaskGraphProfiler/TaskGraphProfilerManager.cpp b/Engine/Source/Developer/TraceInsights/Private/Insights/TaskGraphProfiler/TaskGraphProfilerManager.cpp
--- a/Engine/Source/Developer/TraceInsights/Private/Insights/TaskGraphProfiler/TaskGraphProfilerManager.cpp
+++ b/Engine/Source/Developer/TraceInsights/Private/Insights/TaskGraphProfiler/TaskGraphProfilerManager.cpp
@@ -326,6 +326,11 @@ void FTaskGraphProfilerManager::GetSingleTaskTransitions(const TraceServices::FT
 	{
 		AddRelation(InSelectedEvent, Task->FinishedTimestamp, Task->StartedThreadId, ExecutionStartedDepth, Task->CompletedTimestamp, Task->CompletedThreadId, -1,  ETaskEventType::Finished);
 	}
+
+	if (Task->CompletedTimestamp != Task->DestroyedTimestamp || Task->CompletedThreadId != Task->DestroyedThreadId)
+	{
+		AddRelation(InSelectedEvent, Task->CompletedTimestamp, Task->CompletedThreadId, Task->DestroyedTimestamp, Task->DestroyedThreadId, ETaskEventType::Completed);
+	}
 }
 
 ////////////////////////////////////////////////////////////////////////////////////////////////////
@@ -493,14 +498,15 @@ void FTaskGraphProfilerManager::InitializeColorCode()
 		return FLinearColor(R / 255, G / 255, B / 255);
 	};
 
-	ColorCode[static_cast<uint32>(ETaskEventType::Created)] = ToLiniarColorNoAlpha(0xFFDC1AFF); // Yellow;
-	ColorCode[static_cast<uint32>(ETaskEventType::Launched)] = ToLiniarColorNoAlpha(0x8BC24AFF); // Green;
-	ColorCode[static_cast<uint32>(ETaskEventType::Scheduled)] = ToLiniarColorNoAlpha(0x26BBFFFF); // Blue;
-	ColorCode[static_cast<uint32>(ETaskEventType::Started)] = ToLiniarColorNoAlpha(0xFF0000FF); // Red;
+	ColorCode[static_cast<uint32>(ETaskEventType::Created)] = ToLiniarColorNoAlpha(0xFFDC1AFF); // Yellow
+	ColorCode[static_cast<uint32>(ETaskEventType::Launched)] = ToLiniarColorNoAlpha(0x8BC24AFF); // Green
+	ColorCode[static_cast<uint32>(ETaskEventType::Scheduled)] = ToLiniarColorNoAlpha(0x26BBFFFF); // Blue
+	ColorCode[static_cast<uint32>(ETaskEventType::Started)] = ToLiniarColorNoAlpha(0xFF0000FF); // Red
 	ColorCode[static_cast<uint32>(ETaskEventType::Finished)] = ToLiniarColorNoAlpha(0xFFDC1AFF); // Yellow
-	ColorCode[static_cast<uint32>(ETaskEventType::PrerequisiteStarted)] = ToLiniarColorNoAlpha(0xFF729CFF); //Pink;
+	ColorCode[static_cast<uint32>(ETaskEventType::Completed)] = ToLiniarColorNoAlpha(0xFE9B07FF); // Orange
+	ColorCode[static_cast<uint32>(ETaskEventType::PrerequisiteStarted)] = ToLiniarColorNoAlpha(0xFF729CFF); // Pink
 	ColorCode[static_cast<uint32>(ETaskEventType::NestedStarted)] = ToLiniarColorNoAlpha(0xB68F55FF); // Folder
-	ColorCode[static_cast<uint32>(ETaskEventType::SubsequentStarted)] = ToLiniarColorNoAlpha(0xFE9B07FF); // Orange;;
+	ColorCode[static_cast<uint32>(ETaskEventType::SubsequentStarted)] = ToLiniarColorNoAlpha(0xFE9B07FF); // Orange
 	ColorCode[static_cast<uint32>(ETaskEventType::NestedCompleted)] = ToLiniarColorNoAlpha(0x804D39FF); // Brown
 	ColorCode[static_cast<uint32>(ETaskEventType::CriticalPath)] = ToLiniarColorNoAlpha(0xA139BFFF); // Purple
 }
diff --git a/Engine/Source/Developer/TraceInsights/Private/Insights/TaskGraphProfiler/TaskGraphProfilerManager.h b/Engine/Source/Developer/TraceInsights/Private/Insights/TaskGraphProfiler/TaskGraphProfilerManager.h
--- a/Engine/Source/Developer/TraceInsights/Private/Insights/TaskGraphProfiler/TaskGraphProfilerManager.h
+++ b/Engine/Source/Developer/TraceInsights/Private/Insights/TaskGraphProfiler/TaskGraphProfilerManager.h
@@ -42,6 +42,7 @@ enum class ETaskEventType : uint32
 	Scheduled,
 	Started,
 	Finished,
+	Completed,
 	PrerequisiteStarted,
 	NestedStarted,
 	NestedCompleted,
diff --git a/Engine/Source/Developer/TraceInsights/Private/Insights/TaskGraphProfiler/ViewModels/TaskEntry.cpp b/Engine/Source/Developer/TraceInsights/Private/Insights/TaskGraphProfiler/ViewModels/TaskEntry.cpp
--- a/Engine/Source/Developer/TraceInsights/Private/Insights/TaskGraphProfiler/ViewModels/TaskEntry.cpp
+++ b/Engine/Source/Developer/TraceInsights/Private/Insights/TaskGraphProfiler/ViewModels/TaskEntry.cpp
@@ -23,6 +23,8 @@ FTaskEntry::FTaskEntry(const TraceServices::FTaskInfo& TaskInfo)
 	, FinishedTimestamp(TaskInfo.FinishedTimestamp)
 	, CompletedTimestamp(TaskInfo.CompletedTimestamp)
 	, CompletedThreadId(TaskInfo.CompletedThreadId)
+	, DestroyedTimestamp(TaskInfo.DestroyedTimestamp)
+	, DestroyedThreadId(TaskInfo.DestroyedThreadId)
 {
 	NumNested = TaskInfo.NestedTasks.Num();
 	NumSubsequents = TaskInfo.Subsequents.Num();
diff --git a/Engine/Source/Developer/TraceInsights/Private/Insights/TaskGraphProfiler/ViewModels/TaskEntry.h b/Engine/Source/Developer/TraceInsights/Private/Insights/TaskGraphProfiler/ViewModels/TaskEntry.h
--- a/Engine/Source/Developer/TraceInsights/Private/Insights/TaskGraphProfiler/ViewModels/TaskEntry.h
+++ b/Engine/Source/Developer/TraceInsights/Private/Insights/TaskGraphProfiler/ViewModels/TaskEntry.h
@@ -39,6 +39,9 @@ public:
 	double GetCompletedTimestamp() const { return CompletedTimestamp; }
 	uint32 GetCompletedThreadId() const { return CompletedThreadId; }
 
+	double GetDestroyedTimestamp() const { return DestroyedTimestamp; }
+	uint32 GetDestroyedThreadId() const { return DestroyedThreadId; }
+
 	uint32 GetNumNested() const { return NumNested; }
 	uint32 GetNumPrerequisites() const { return NumPrerequisites; }
 	uint32 GetNumSubsequents() const { return NumSubsequents; }
@@ -69,6 +72,9 @@ private:
 	double CompletedTimestamp;
 	uint32 CompletedThreadId;
 
+	double DestroyedTimestamp;
+	uint32 DestroyedThreadId;
+
 	uint32 NumNested;
 	uint32 NumSubsequents;
 	uint32 NumPrerequisites;
diff --git a/Engine/Source/Developer/TraceInsights/Private/Insights/TaskGraphProfiler/ViewModels/TaskTable.cpp b/Engine/Source/Developer/TraceInsights/Private/Insights/TaskGraphProfiler/ViewModels/TaskTable.cpp
--- a/Engine/Source/Developer/TraceInsights/Private/Insights/TaskGraphProfiler/ViewModels/TaskTable.cpp
+++ b/Engine/Source/Developer/TraceInsights/Private/Insights/TaskGraphProfiler/ViewModels/TaskTable.cpp
@@ -27,9 +27,11 @@ const FName FTaskTableColumns::ScheduledTimestampColumnId(TEXT("ScheduledTimesta
 const FName FTaskTableColumns::ScheduledThreadIdColumnId(TEXT("ScheduledThreadId"));
 const FName FTaskTableColumns::StartedTimestampColumnId(TEXT("StartedTimestamp"));
 const FName FTaskTableColumns::StartedThreadIdColumnId(TEXT("StartedThreadId"));
+const FName FTaskTableColumns::FinishedTimestampColumnId(TEXT("FinishedTimestamp"));
 const FName FTaskTableColumns::CompletedTimestampColumnId(TEXT("CompletedTimestamp"));
 const FName FTaskTableColumns::CompletedThreadIdColumnId(TEXT("CompletedThreadId"));
-const FName FTaskTableColumns::FinishedTimestampColumnId(TEXT("FinishedTimestamp"));
+const FName FTaskTableColumns::DestroyedTimestampColumnId(TEXT("DestroyedTimestamp"));
+const FName FTaskTableColumns::DestroyedThreadIdColumnId(TEXT("DestroyedThreadId"));
 const FName FTaskTableColumns::NumNestedColumnId(TEXT("NumNested"));
 const FName FTaskTableColumns::NumSubsequentsColumnId(TEXT("NumSubsequents"));
 const FName FTaskTableColumns::NumPrerequisitesColumnId(TEXT("NumPrerequisites"));
@@ -83,10 +85,13 @@ struct DefaultTaskFieldGetterFuncts
 	static FTableCellValue GetStartedTimestamp(const FTableColumn& Column, const FTaskEntry& Task) { return FTableCellValue(Task.GetStartedTimestamp()); }
 	static FTableCellValue GetStartedThreadId(const FTableColumn& Column, const FTaskEntry& Task) { return FTableCellValue((int64)Task.GetStartedThreadId()); }
 
+	static FTableCellValue GetFinishedTimestamp(const FTableColumn& Column, const FTaskEntry& Task) { return FTableCellValue(Task.GetFinishedTimestamp()); }
+
 	static FTableCellValue GetCompletedTimestamp(const FTableColumn& Column, const FTaskEntry& Task) { return FTableCellValue(Task.GetCompletedTimestamp()); }
 	static FTableCellValue GetCompletedThreadId(const FTableColumn& Column, const FTaskEntry& Task) { return FTableCellValue((int64)Task.GetCompletedThreadId()); }
 
-	static FTableCellValue GetFinishedTimestamp(const FTableColumn& Column, const FTaskEntry& Task) { return FTableCellValue(Task.GetFinishedTimestamp()); }
+	static FTableCellValue GetDestroyedTimestamp(const FTableColumn& Column, const FTaskEntry& Task) { return FTableCellValue(Task.GetDestroyedTimestamp()); }
+	static FTableCellValue GetDestroyedThreadId(const FTableColumn& Column, const FTaskEntry& Task) { return FTableCellValue((int64)Task.GetDestroyedThreadId()); }
 
 	static FTableCellValue GetNumNested(const FTableColumn& Column, const FTaskEntry& Task) { return FTableCellValue((int64)Task.GetNumNested());	}
 	static FTableCellValue GetNumSubsequents(const FTableColumn& Column, const FTaskEntry& Task) { return FTableCellValue((int64)Task.GetNumSubsequents()); }
@@ -114,6 +119,7 @@ struct RelativeToPreviousTaskFieldGetterFuncts
 			return FTableCellValue(Task.GetCompletedTimestamp());
 		}
 	}
+	static FTableCellValue GetDestroyedTimestamp(const FTableColumn& Column, const FTaskEntry& Task) { return FTableCellValue(Task.GetDestroyedTimestamp() - Task.GetCompletedTimestamp()); }
 };
 
 ////////////////////////////////////////////////////////////////////////////////////////////////////
@@ -137,6 +143,7 @@ struct RelativeToCreatedTaskFieldGetterFuncts
 			return FTableCellValue(Task.GetCompletedTimestamp());
 		}
 	}
+	static FTableCellValue GetDestroyedTimestamp(const FTableColumn& Column, const FTaskEntry& Task) { return FTableCellValue(Task.GetDestroyedTimestamp() - Task.GetCreatedTimestamp()); }
 };
 
 ////////////////////////////////////////////////////////////////////////////////////////////////////
@@ -560,6 +567,68 @@ void FTaskTable::AddDefaultColumns()
 		AddColumn(ColumnRef);
 	}
 	//////////////////////////////////////////////////
+	// Destroyed Timestamp Column
+	{
+		TSharedRef<FTableColumn> ColumnRef = MakeShared<FTableColumn>(FTaskTableColumns::DestroyedTimestampColumnId);
+		FTableColumn& Column = *ColumnRef;
+
+		Column.SetIndex(ColumnIndex++);
+
+		Column.SetShortName(LOCTEXT("DestroyedTimestampColumnName", "Destroyed"));
+		Column.SetTitleName(LOCTEXT("DestroyedTimestampColumnTitle", "Destroyed"));
+		Column.SetDescription(LOCTEXT("DestroyedTimestampColumnDesc", "The time when the task was Destroyed."));
+
+		Column.SetFlags(ETableColumnFlags::ShouldBeVisible | ETableColumnFlags::CanBeHidden | ETableColumnFlags::CanBeFiltered);
+
+		Column.SetHorizontalAlignment(HAlign_Left);
+		Column.SetInitialWidth(100.0f);
+
+		Column.SetDataType(ETableCellDataType::Double);
+
+		TSharedRef<ITableCellValueGetter> Getter = MakeShared<FTaskColumnValueGetter<DefaultTaskFieldGetterFuncts::GetDestroyedTimestamp>>();
+		Column.SetValueGetter(Getter);
+
+		TSharedRef<ITableCellValueFormatter> Formatter = MakeShared<FDoubleValueFormatterAsTimeAuto>();
+		Column.SetValueFormatter(Formatter);
+
+		TSharedRef<ITableCellValueSorter> Sorter = MakeShared<FSorterByDoubleValue>(ColumnRef);
+		Column.SetValueSorter(Sorter);
+
+		Column.SetAggregation(ETableColumnAggregation::Min);
+
+		AddColumn(ColumnRef);
+	}
+	//////////////////////////////////////////////////
+	// Destroyed Thread Id Column
+	{
+		TSharedRef<FTableColumn> ColumnRef = MakeShared<FTableColumn>(FTaskTableColumns::DestroyedThreadIdColumnId);
+		FTableColumn& Column = *ColumnRef;
+
+		Column.SetIndex(ColumnIndex++);
+
+		Column.SetShortName(LOCTEXT("DestroyedThreadIdColumnName", "Destroyed Thread Id"));
+		Column.SetTitleName(LOCTEXT("DestroyedThreadIdColumnTitle", "Completed Thread Id"));
+		Column.SetDescription(LOCTEXT("DestroyedThreadIdColumnDesc", "The thread the task was Destroyed on."));
+
+		Column.SetFlags(ETableColumnFlags::CanBeHidden | ETableColumnFlags::CanBeFiltered);
+
+		Column.SetHorizontalAlignment(HAlign_Left);
+		Column.SetInitialWidth(100.0f);
+
+		Column.SetDataType(ETableCellDataType::Int64);
+
+		TSharedRef<ITableCellValueGetter> Getter = MakeShared<FTaskColumnValueGetter<DefaultTaskFieldGetterFuncts::GetDestroyedThreadId>>();
+		Column.SetValueGetter(Getter);
+
+		TSharedRef<ITableCellValueFormatter> Formatter = MakeShared<FInt64ValueFormatterAsNumber>();
+		Column.SetValueFormatter(Formatter);
+
+		TSharedRef<ITableCellValueSorter> Sorter = MakeShared<FSorterByInt64Value>(ColumnRef);
+		Column.SetValueSorter(Sorter);
+
+		AddColumn(ColumnRef);
+	}
+	//////////////////////////////////////////////////
 	// Num Nested Tasks Column
 	{
 		TSharedRef<FTableColumn> ColumnRef = MakeShared<FTableColumn>(FTaskTableColumns::NumNestedColumnId);
diff --git a/Engine/Source/Developer/TraceInsights/Private/Insights/TaskGraphProfiler/ViewModels/TaskTable.h b/Engine/Source/Developer/TraceInsights/Private/Insights/TaskGraphProfiler/ViewModels/TaskTable.h
--- a/Engine/Source/Developer/TraceInsights/Private/Insights/TaskGraphProfiler/ViewModels/TaskTable.h
+++ b/Engine/Source/Developer/TraceInsights/Private/Insights/TaskGraphProfiler/ViewModels/TaskTable.h
@@ -27,9 +27,11 @@ struct FTaskTableColumns
 	static const FName ScheduledThreadIdColumnId;
 	static const FName StartedTimestampColumnId;
 	static const FName StartedThreadIdColumnId;
+	static const FName FinishedTimestampColumnId;
 	static const FName CompletedTimestampColumnId;
 	static const FName CompletedThreadIdColumnId;
-	static const FName FinishedTimestampColumnId;
+	static const FName DestroyedTimestampColumnId;
+	static const FName DestroyedThreadIdColumnId;
 	static const FName NumNestedColumnId;
 	static const FName NumSubsequentsColumnId;
 	static const FName NumPrerequisitesColumnId;
diff --git a/Engine/Source/Developer/TraceInsights/Private/Insights/TaskGraphProfiler/ViewModels/TaskTimingTrack.cpp b/Engine/Source/Developer/TraceInsights/Private/Insights/TaskGraphProfiler/ViewModels/TaskTimingTrack.cpp
--- a/Engine/Source/Developer/TraceInsights/Private/Insights/TaskGraphProfiler/ViewModels/TaskTimingTrack.cpp
+++ b/Engine/Source/Developer/TraceInsights/Private/Insights/TaskGraphProfiler/ViewModels/TaskTimingTrack.cpp
@@ -670,10 +670,8 @@ void FTaskTimingTrack::BuildDrawState(ITimingEventsTrackDrawStateBuilder& Builde
 	Builder.AddEvent(Task->LaunchedTimestamp, Task->ScheduledTimestamp, 0, TEXT("Launched"), 0, FTaskGraphProfilerManager::Get()->GetColorForTaskEventAsPackedARGB(ETaskEventType::Launched));
 	Builder.AddEvent(Task->ScheduledTimestamp, Task->StartedTimestamp, 0, TEXT("Scheduled"), 0,  FTaskGraphProfilerManager::Get()->GetColorForTaskEventAsPackedARGB(ETaskEventType::Scheduled));
 	Builder.AddEvent(Task->StartedTimestamp, Task->FinishedTimestamp, 0, TEXT("Executing"), 0, FTaskGraphProfilerManager::Get()->GetColorForTaskEventAsPackedARGB(ETaskEventType::Started));
-	if (Task->CompletedTimestamp > Task->FinishedTimestamp)
-	{
-		Builder.AddEvent(Task->FinishedTimestamp, Task->CompletedTimestamp, 0, TEXT("Finished"), 0, FTaskGraphProfilerManager::Get()->GetColorForTaskEventAsPackedARGB(ETaskEventType::Finished));
-	}
+	Builder.AddEvent(Task->FinishedTimestamp, Task->CompletedTimestamp, 0, TEXT("Finished"), 0, FTaskGraphProfilerManager::Get()->GetColorForTaskEventAsPackedARGB(ETaskEventType::Finished));
+	Builder.AddEvent(Task->CompletedTimestamp, Task->DestroyedTimestamp, 0, TEXT("Completed"), 0, FTaskGraphProfilerManager::Get()->GetColorForTaskEventAsPackedARGB(ETaskEventType::Completed));
 
 	if (bShowDetailInfoOnTaskTrack)
 	{
@@ -858,10 +856,15 @@ const TSharedPtr<const ITimingEvent> FTaskTimingTrack::GetEvent(float InPosX, fl
 			{
 				TimingEvent = MakeShared<FTaskTrackEvent>(SharedThis(this), Task->ScheduledTimestamp, Task->StartedTimestamp, 0, ETaskEventType::Scheduled);
 			}
+			// the unusual order of `if`s here is intended to give priority to smaller events
 			else if (IsEventTimeBetween(Task->FinishedTimestamp, Task->CompletedTimestamp))
 			{
 				TimingEvent = MakeShared<FTaskTrackEvent>(SharedThis(this), Task->FinishedTimestamp, Task->CompletedTimestamp, 0, ETaskEventType::Finished);
 			}
+			else if (IsEventTimeBetween(Task->CompletedTimestamp, Task->DestroyedTimestamp))
+			{
+				TimingEvent = MakeShared<FTaskTrackEvent>(SharedThis(this), Task->CompletedTimestamp >= Task->FinishedTimestamp ? Task->CompletedTimestamp : Task->FinishedTimestamp, Task->DestroyedTimestamp, 0, ETaskEventType::Completed);
+			}
 			else if (IsEventTimeBetween(Task->StartedTimestamp, Task->FinishedTimestamp))
 			{
 				TimingEvent = MakeShared<FTaskTrackEvent>(SharedThis(this), Task->StartedTimestamp, Task->FinishedTimestamp, 0, ETaskEventType::Started);
@@ -977,6 +980,8 @@ void FTaskTimingTrack::InitTooltip(FTooltipDrawState& InOutTooltip, const ITimin
 	case ETaskEventType::Finished:
 		InOutTooltip.AddNameValueTextLine(TEXT("Subsequent tasks:"), FString::Printf(TEXT("%d"), Task->Subsequents.Num()));
 		break;
+	case ETaskEventType::Completed:
+		break;
 	default:
 		checkf(false, TEXT("Unknown task event type"));
 		break;
diff --git a/Engine/Source/Developer/TraceInsights/Private/Insights/TaskGraphProfiler/ViewModels/TaskTrackEvent.cpp b/Engine/Source/Developer/TraceInsights/Private/Insights/TaskGraphProfiler/ViewModels/TaskTrackEvent.cpp
--- a/Engine/Source/Developer/TraceInsights/Private/Insights/TaskGraphProfiler/ViewModels/TaskTrackEvent.cpp
+++ b/Engine/Source/Developer/TraceInsights/Private/Insights/TaskGraphProfiler/ViewModels/TaskTrackEvent.cpp
@@ -35,6 +35,8 @@ FString FTaskTrackEvent::GetStartLabel() const
 		return TEXT("Started Time:");
 	case ETaskEventType::Finished:
 		return TEXT("Finished Time:");
+	case ETaskEventType::Completed:
+		return TEXT("Completed Time:");
 	default:
 		checkf(false, TEXT("Unknown task event type"));
 		break;
@@ -62,6 +64,8 @@ FString FTaskTrackEvent::GetEndLabel() const
 		return TEXT("Finished Time:");
 	case ETaskEventType::Finished:
 		return TEXT("Completed Time:");
+	case ETaskEventType::Completed:
+		return TEXT("Destroyed Time:");
 	default:
 		checkf(false, TEXT("Unknown task event type"));
 		break;
@@ -86,6 +90,8 @@ FString FTaskTrackEvent::GetEventName() const
 		return TEXT("Executing");
 	case ETaskEventType::Finished:
 		return TEXT("Finished");
+	case ETaskEventType::Completed:
+		return TEXT("Completed");
 	case ETaskEventType::PrerequisiteStarted:
 		return FString::Printf(TEXT("Prerequisite Task %d Executing"), GetTaskId());	
 	case ETaskEventType::NestedStarted:
diff --git a/Engine/Source/Developer/TraceInsights/Private/Insights/TaskGraphProfiler/Widgets/STaskTableTreeView.cpp b/Engine/Source/Developer/TraceInsights/Private/Insights/TaskGraphProfiler/Widgets/STaskTableTreeView.cpp
--- a/Engine/Source/Developer/TraceInsights/Private/Insights/TaskGraphProfiler/Widgets/STaskTableTreeView.cpp
+++ b/Engine/Source/Developer/TraceInsights/Private/Insights/TaskGraphProfiler/Widgets/STaskTableTreeView.cpp
@@ -3,6 +3,11 @@
 #include "STaskTableTreeView.h"
 
 #include "Framework/MultiBox/MultiBoxBuilder.h"
+#include "Internationalization/Regex.h"
+#include "ISourceCodeAccessModule.h"
+#include "ISourceCodeAccessor.h"
+#include "Misc/Paths.h"
+#include "Modules/ModuleManager.h"
 #include "SlateOptMacros.h"
 #include "TraceServices/AnalysisService.h"
 #include "TraceServices/Model/TasksProfiler.h"
@@ -56,10 +61,12 @@ public:
 	virtual void RegisterCommands() override
 	{
 		UI_COMMAND(Command_GoToTask, "Go To Task", "Pan and zoom to the task in Timing View.", EUserInterfaceActionType::Button, FInputChord());
+		UI_COMMAND(Command_OpenInIDE, "Open in IDE", "Open the source location where the selected task was launched in IDE.", EUserInterfaceActionType::Button, FInputChord());
 	}
 	PRAGMA_ENABLE_OPTIMIZATION
 
 	TSharedPtr<FUICommandInfo> Command_GoToTask;
+	TSharedPtr<FUICommandInfo> Command_OpenInIDE;
 };
 
 ////////////////////////////////////////////////////////////////////////////////////////////////////
@@ -91,6 +98,9 @@ void STaskTableTreeView::Construct(const FArguments& InArgs, TSharedPtr<Insights
 
 void STaskTableTreeView::ExtendMenu(FMenuBuilder& MenuBuilder)
 {
+	ISourceCodeAccessModule& SourceCodeAccessModule = FModuleManager::LoadModuleChecked<ISourceCodeAccessModule>("SourceCodeAccess");
+	ISourceCodeAccessor& SourceCodeAccessor = SourceCodeAccessModule.GetAccessor();
+
 	MenuBuilder.BeginSection("Node", LOCTEXT("ContextMenu_Section_Task", "Task"));
 
 	{
@@ -103,6 +113,16 @@ void STaskTableTreeView::ExtendMenu(FMenuBuilder& MenuBuilder)
 			FSlateIcon(FInsightsStyle::GetStyleSetName(), "Icons.GoToTask")
 		);
 	}
+	{
+		MenuBuilder.AddMenuEntry
+		(
+			FTaskTableTreeViewCommands::Get().Command_OpenInIDE,
+			NAME_None,
+			TAttribute<FText>(),
+			TAttribute<FText>(),
+			FSlateIcon(FAppStyle::GetAppStyleSetName(), SourceCodeAccessor.GetOpenIconName())
+		);
+	}
 
 	MenuBuilder.EndSection();
 }
@@ -112,6 +132,7 @@ void STaskTableTreeView::AddCommmands()
 	FTaskTableTreeViewCommands::Register();
 
 	CommandList->MapAction(FTaskTableTreeViewCommands::Get().Command_GoToTask, FExecuteAction::CreateSP(this, &STaskTableTreeView::ContextMenu_GoToTask_Execute), FCanExecuteAction::CreateSP(this, &STaskTableTreeView::ContextMenu_GoToTask_CanExecute));
+	CommandList->MapAction(FTaskTableTreeViewCommands::Get().Command_OpenInIDE, FExecuteAction::CreateSP(this, &STaskTableTreeView::ContextMenu_OpenInIDE_Execute), FCanExecuteAction::CreateSP(this, &STaskTableTreeView::ContextMenu_OpenInIDE_CanExecute));
 }
 
 ////////////////////////////////////////////////////////////////////////////////////////////////////
@@ -503,6 +524,68 @@ void STaskTableTreeView::ContextMenu_GoToTask_Execute()
 
 ////////////////////////////////////////////////////////////////////////////////////////////////////
 
+bool STaskTableTreeView::ContextMenu_OpenInIDE_CanExecute() const
+{
+	TArray<FTableTreeNodePtr> SelectedItems;
+	TreeView->GetSelectedItems(SelectedItems);
+
+	if (SelectedItems.Num() != 1)
+	{
+		return false;
+	}
+
+	FTaskNodePtr SelectedTask = StaticCastSharedPtr<FTaskNode>(SelectedItems[0]);
+	const FTaskEntry* TaskEntry = SelectedTask.IsValid() ? SelectedTask->GetTask() : nullptr;
+
+	if (!SelectedTask.IsValid() || SelectedTask->IsGroup() || TaskEntry == nullptr || TaskEntry->DebugName == nullptr || TaskEntry->DebugName[FCString::Strlen(TaskEntry->DebugName) - 1] != TEXT(')'))
+	{
+		return false;
+	}
+
+	return true;
+}
+
+////////////////////////////////////////////////////////////////////////////////////////////////////
+
+void STaskTableTreeView::ContextMenu_OpenInIDE_Execute()
+{
+	TArray<FTableTreeNodePtr> SelectedItems;
+	TreeView->GetSelectedItems(SelectedItems);
+
+	if (SelectedItems.Num() != 1)
+	{
+		return;
+	}
+
+	FTaskNodePtr SelectedTask = StaticCastSharedPtr<FTaskNode>(SelectedItems[0]);
+	const FTaskEntry* TaskEntry = SelectedTask.IsValid() ? SelectedTask->GetTask() : nullptr;
+
+	if (TaskEntry == nullptr || TaskEntry->DebugName == nullptr)
+	{
+		return;
+	}
+
+	const FString DebugName = TaskEntry->DebugName;
+	int32 OpenBracketPos;
+	if (DebugName[DebugName.Len() - 1] != TEXT(')') || !DebugName.FindChar(TEXT('('), OpenBracketPos) || OpenBracketPos > DebugName.Len() - 3)
+	{
+		return;
+	}
+	FString Filename = DebugName.Left(OpenBracketPos);
+	if (!FPaths::FileExists(Filename))
+	{
+		return;
+	}
+	FString LineStr = DebugName.Mid(OpenBracketPos + 1, DebugName.Len() - OpenBracketPos - 2);
+	int32 Line = FCString::Atoi(*LineStr);
+
+	ISourceCodeAccessModule& SourceCodeAccessModule = FModuleManager::LoadModuleChecked<ISourceCodeAccessModule>("SourceCodeAccess");
+	ISourceCodeAccessor& SourceCodeAccessor = SourceCodeAccessModule.GetAccessor();
+	SourceCodeAccessor.OpenFileAtLine(Filename, Line);
+}
+
+////////////////////////////////////////////////////////////////////////////////////////////////////
+
 void STaskTableTreeView::TreeView_OnMouseButtonDoubleClick(FTableTreeNodePtr TreeNode)
 {
 	if (!TreeNode->IsGroup())
diff --git a/Engine/Source/Developer/TraceInsights/Private/Insights/TaskGraphProfiler/Widgets/STaskTableTreeView.h b/Engine/Source/Developer/TraceInsights/Private/Insights/TaskGraphProfiler/Widgets/STaskTableTreeView.h
--- a/Engine/Source/Developer/TraceInsights/Private/Insights/TaskGraphProfiler/Widgets/STaskTableTreeView.h
+++ b/Engine/Source/Developer/TraceInsights/Private/Insights/TaskGraphProfiler/Widgets/STaskTableTreeView.h
@@ -105,6 +105,9 @@ private:
 	bool ContextMenu_GoToTask_CanExecute() const;
 	void ContextMenu_GoToTask_Execute();
 
+	bool ContextMenu_OpenInIDE_CanExecute() const;
+	void ContextMenu_OpenInIDE_Execute();
+
 private:
 	double QueryStartTime = 0.0f;
 	double QueryEndTime = 0.0f;
diff --git a/Engine/Source/Developer/TraceInsights/Private/Insights/ViewModels/ThreadTimingTrack.cpp b/Engine/Source/Developer/TraceInsights/Private/Insights/ViewModels/ThreadTimingTrack.cpp
--- a/Engine/Source/Developer/TraceInsights/Private/Insights/ViewModels/ThreadTimingTrack.cpp
+++ b/Engine/Source/Developer/TraceInsights/Private/Insights/ViewModels/ThreadTimingTrack.cpp
@@ -1210,9 +1210,16 @@
 
 					if (Task.CompletedTimestamp != TraceServices::FTaskInfo::InvalidTimestamp)
 					{
-						TSharedPtr<FCpuTimingTrack> Track = SharedState.GetCpuTrack(Task.CompletedThreadId);
-						FString TrackName = Track.IsValid() ? Track->GetName() : TEXT("Unknown");
-						InOutTooltip.AddNameValueTextLine(TEXT("Completed:"), FString::Printf(TEXT("%f (+%s) on %s"), Task.FinishedTimestamp, *TimeUtils::FormatTimeAuto(Task.CompletedTimestamp - Task.FinishedTimestamp), *TrackName));
+						TSharedPtr<FCpuTimingTrack> CompletedTrack = SharedState.GetCpuTrack(Task.CompletedThreadId);
+						FString CompletedTrackName = CompletedTrack.IsValid() ? CompletedTrack->GetName() : TEXT("Unknown");
+						InOutTooltip.AddNameValueTextLine(TEXT("Completed:"), FString::Printf(TEXT("%f (+%s) on %s"), Task.CompletedTimestamp, *TimeUtils::FormatTimeAuto(Task.CompletedTimestamp - Task.FinishedTimestamp), *CompletedTrackName));
+
+						if (Task.DestroyedTimestamp != TraceServices::FTaskInfo::InvalidTimestamp)
+						{
+							TSharedPtr<FCpuTimingTrack> DestroyedTrack = SharedState.GetCpuTrack(Task.DestroyedThreadId);
+							FString DestroyedTrackName = DestroyedTrack.IsValid() ? DestroyedTrack->GetName() : TEXT("Unknown");
+							InOutTooltip.AddNameValueTextLine(TEXT("Destroyed:"), FString::Printf(TEXT("%f (+%s) on %s"), Task.DestroyedTimestamp, *TimeUtils::FormatTimeAuto(Task.DestroyedTimestamp - Task.CompletedTimestamp), *DestroyedTrackName));
+						}
 					}
 				}
 				InOutTooltip.AddNameValueTextLine(TEXT("Prerequisite tasks:"), FString::Printf(TEXT("%d"), Task.Prerequisites.Num()));
diff --git a/Engine/Source/Developer/TraceServices/Private/Analyzers/TasksAnalysis.cpp b/Engine/Source/Developer/TraceServices/Private/Analyzers/TasksAnalysis.cpp
--- a/Engine/Source/Developer/TraceServices/Private/Analyzers/TasksAnalysis.cpp
+++ b/Engine/Source/Developer/TraceServices/Private/Analyzers/TasksAnalysis.cpp
@@ -29,6 +29,7 @@ void FTasksAnalyzer::OnAnalysisBegin(const FOnAnalysisContext& Context)
 	Builder.RouteEvent(RouteId_NestedAdded, "TaskTrace", "NestedAdded");
 	Builder.RouteEvent(RouteId_Finished, "TaskTrace", "Finished");
 	Builder.RouteEvent(RouteId_Completed, "TaskTrace", "Completed");
+	Builder.RouteEvent(RouteId_Destroyed, "TaskTrace", "Destroyed");
 
 	Builder.RouteEvent(RouteId_WaitingStarted, "TaskTrace", "WaitingStarted");
 	Builder.RouteEvent(RouteId_WaitingFinished, "TaskTrace", "WaitingFinished");
@@ -122,6 +123,13 @@ bool FTasksAnalyzer::OnEvent(uint16 RouteId, EStyle Style, const FOnEventContext
 			TasksProvider.TaskCompleted(TaskId, Timestamp, ThreadId);
 			break;
 		}
+		case RouteId_Destroyed:
+		{
+			double Timestamp = Context.EventTime.AsSeconds(EventData.GetValue<uint64>("Timestamp"));
+			TaskTrace::FId TaskId = EventData.GetValue<TaskTrace::FId>("TaskId");
+			TasksProvider.TaskDestroyed(TaskId, Timestamp, ThreadId);
+			break;
+		}
 		case RouteId_WaitingStarted:
 		{
 			double Timestamp = Context.EventTime.AsSeconds(EventData.GetValue<uint64>("Timestamp"));
diff --git a/Engine/Source/Developer/TraceServices/Private/Analyzers/TasksAnalysis.h b/Engine/Source/Developer/TraceServices/Private/Analyzers/TasksAnalysis.h
--- a/Engine/Source/Developer/TraceServices/Private/Analyzers/TasksAnalysis.h
+++ b/Engine/Source/Developer/TraceServices/Private/Analyzers/TasksAnalysis.h
@@ -29,6 +29,7 @@ namespace TraceServices
 			RouteId_NestedAdded,
 			RouteId_Finished,
 			RouteId_Completed,
+			RouteId_Destroyed,
 
 			RouteId_WaitingStarted,
 			RouteId_WaitingFinished,
diff --git a/Engine/Source/Developer/TraceServices/Private/Model/TasksProfiler.cpp b/Engine/Source/Developer/TraceServices/Private/Model/TasksProfiler.cpp
--- a/Engine/Source/Developer/TraceServices/Private/Model/TasksProfiler.cpp
+++ b/Engine/Source/Developer/TraceServices/Private/Model/TasksProfiler.cpp
@@ -218,6 +218,12 @@ namespace TraceServices
 		FTaskInfo* Task = TryGetTask(TaskId);
 		check(Task != nullptr);
 
+		if (Task->ScheduledTimestamp == FTaskInfo::InvalidTimestamp) // "inline" tasks or task events are never scheduled
+		{
+			Task->ScheduledTimestamp = Timestamp;
+			Task->ScheduledThreadId = ThreadId;
+		}
+
 		if (!bCountersCreated)
 		{
 			CreateCounters();
@@ -285,9 +291,25 @@ namespace TraceServices
 
 		Task->Id = TaskId;
 
+		if (Task->ScheduledTimestamp == FTaskInfo::InvalidTimestamp) // task events are never scheduled or executed
+		{
+			// pretend scheduling and execution happened at the moment of completion
+			Task->ScheduledTimestamp = Timestamp;
+			Task->ScheduledThreadId = ThreadId;
+			Task->StartedTimestamp = Timestamp;
+			Task->StartedThreadId = ThreadId;
+			Task->FinishedTimestamp = Timestamp;
+		}
+
 		TryRegisterEvent(TEXT("TaskCompleted"), TaskId, &FTaskInfo::CompletedTimestamp, Timestamp, &FTaskInfo::CompletedThreadId, ThreadId);
 	}
 
+	void FTasksProvider::TaskDestroyed(TaskTrace::FId TaskId, double Timestamp, uint32 ThreadId)
+	{
+		InitTaskIdToIndexConversion(TaskId);
+		TryRegisterEvent(TEXT("TaskDestroyed"), TaskId, &FTaskInfo::DestroyedTimestamp, Timestamp, &FTaskInfo::DestroyedThreadId, ThreadId);
+	}
+
 	void FTasksProvider::WaitingStarted(TArray<TaskTrace::FId> InTasks, double Timestamp, uint32 ThreadId)
 	{
 		FWaitingForTasks Waiting;
diff --git a/Engine/Source/Developer/TraceServices/Private/Model/TasksProfilerPrivate.h b/Engine/Source/Developer/TraceServices/Private/Model/TasksProfilerPrivate.h
--- a/Engine/Source/Developer/TraceServices/Private/Model/TasksProfilerPrivate.h
+++ b/Engine/Source/Developer/TraceServices/Private/Model/TasksProfilerPrivate.h
@@ -30,6 +30,7 @@ namespace TraceServices
 		void NestedAdded(TaskTrace::FId TaskId, TaskTrace::FId NestedId, double Timestamp, uint32 ThreadId);
 		void TaskFinished(TaskTrace::FId TaskId, double Timestamp);
 		void TaskCompleted(TaskTrace::FId TaskId, double Timestamp, uint32 ThreadId);
+		void TaskDestroyed(TaskTrace::FId TaskId, double Timestamp, uint32 ThreadId);
 
 		void WaitingStarted(TArray<TaskTrace::FId> Tasks, double Timestamp, uint32 ThreadId);
 		void WaitingFinished(double Timestamp, uint32 ThreadId);
diff --git a/Engine/Source/Developer/TraceServices/Public/TraceServices/Model/TasksProfiler.h b/Engine/Source/Developer/TraceServices/Public/TraceServices/Model/TasksProfiler.h
--- a/Engine/Source/Developer/TraceServices/Public/TraceServices/Model/TasksProfiler.h
+++ b/Engine/Source/Developer/TraceServices/Public/TraceServices/Model/TasksProfiler.h
@@ -42,6 +42,9 @@ namespace TraceServices
 		// execution done and all nested tasks are completed
 		double CompletedTimestamp = InvalidTimestamp;
 		uint32 CompletedThreadId;
+		// the last reference is released and the task is destroyed
+		double DestroyedTimestamp = InvalidTimestamp;
+		uint32 DestroyedThreadId;
 
 		//  relation with another task, when and where it was established
 		struct FRelationInfo
diff --git a/Engine/Source/Runtime/Core/Private/Async/TaskTrace.cpp b/Engine/Source/Runtime/Core/Private/Async/TaskTrace.cpp
--- a/Engine/Source/Runtime/Core/Private/Async/TaskTrace.cpp
+++ b/Engine/Source/Runtime/Core/Private/Async/TaskTrace.cpp
@@ -62,6 +62,11 @@ namespace TaskTrace
 		UE_TRACE_EVENT_FIELD(uint64, TaskId)
 	UE_TRACE_EVENT_END()
 
+	UE_TRACE_EVENT_BEGIN(TaskTrace, Destroyed)
+		UE_TRACE_EVENT_FIELD(uint64, Timestamp)
+		UE_TRACE_EVENT_FIELD(uint64, TaskId)
+	UE_TRACE_EVENT_END()
+
 	UE_TRACE_EVENT_BEGIN(TaskTrace, WaitingStarted)
 		UE_TRACE_EVENT_FIELD(uint64, Timestamp)
 		UE_TRACE_EVENT_FIELD(uint64[], Tasks)
@@ -97,7 +102,6 @@ namespace TaskTrace
 	{
 		check(bGTaskTraceInitialized);
 		check(TaskId != InvalidId);
-		check(TaskId != 0);
 
 		UE_TRACE_LOG(TaskTrace, Created, TaskChannel)
 			<< Created.Timestamp(FPlatformTime::Cycles64())
@@ -191,6 +195,15 @@ namespace TaskTrace
 			<< Completed.TaskId(TaskId);
 	}
 
+	void Destroyed(FId TaskId)
+	{
+		check(bGTaskTraceInitialized);
+
+		UE_TRACE_LOG(TaskTrace, Destroyed, TaskChannel)
+			<< Destroyed.Timestamp(FPlatformTime::Cycles64())
+			<< Destroyed.TaskId(TaskId);
+	}
+
 	FWaitingScope::FWaitingScope(const TArray<FId>& Tasks)
 	{
 		check(bGTaskTraceInitialized);
diff --git a/Engine/Source/Runtime/Core/Public/Async/TaskTrace.h b/Engine/Source/Runtime/Core/Public/Async/TaskTrace.h
--- a/Engine/Source/Runtime/Core/Public/Async/TaskTrace.h
+++ b/Engine/Source/Runtime/Core/Public/Async/TaskTrace.h
@@ -40,6 +40,7 @@ namespace TaskTrace
 	void CORE_API NestedAdded(FId TaskId, FId NestedId);
 	void CORE_API Finished(FId TaskId);
 	void CORE_API Completed(FId TaskId);
+	void CORE_API Destroyed(FId TaskId);
 
 	struct CORE_API FWaitingScope
 	{
@@ -70,6 +71,7 @@ namespace TaskTrace
 	inline void NestedAdded(FId TaskId, FId NestedId) {}
 	inline void Finished(FId TaskId) {}
 	inline void Completed(FId TaskId) {}
+	inline void Destroyed(FId TaskId) {}
 	inline FWaitingScope::FWaitingScope(const TArray<FId>& Tasks) {}
 	inline FWaitingScope::FWaitingScope(FId TaskId) {}
 	inline FWaitingScope::~FWaitingScope() {}
