TasksInsights: fixed switching between different task timestamp representations that didn't handle newly added "Destroyed" timestamp

diff --git a/Engine/Source/Developer/TraceInsights/Private/Insights/TaskGraphProfiler/ViewModels/TaskTable.cpp b/Engine/Source/Developer/TraceInsights/Private/Insights/TaskGraphProfiler/ViewModels/TaskTable.cpp
--- a/Engine/Source/Developer/TraceInsights/Private/Insights/TaskGraphProfiler/ViewModels/TaskTable.cpp
+++ b/Engine/Source/Developer/TraceInsights/Private/Insights/TaskGraphProfiler/ViewModels/TaskTable.cpp
@@ -740,6 +740,7 @@ void FTaskTable::SwitchToRelativeToPreviousTimestamps()
 	FindColumnChecked(FTaskTableColumns::StartedTimestampColumnId)->SetValueGetter(MakeShared<FTaskColumnValueGetter<RelativeToPreviousTaskFieldGetterFuncts::GetStartedTimestamp>>());
 	FindColumnChecked(FTaskTableColumns::FinishedTimestampColumnId)->SetValueGetter(MakeShared<FTaskColumnValueGetter<RelativeToPreviousTaskFieldGetterFuncts::GetFinishedTimestamp>>());
 	FindColumnChecked(FTaskTableColumns::CompletedTimestampColumnId)->SetValueGetter(MakeShared<FTaskColumnValueGetter<RelativeToPreviousTaskFieldGetterFuncts::GetCompletedTimestamp>>());
+	FindColumnChecked(FTaskTableColumns::DestroyedTimestampColumnId)->SetValueGetter(MakeShared<FTaskColumnValueGetter<RelativeToPreviousTaskFieldGetterFuncts::GetDestroyedTimestamp>>());
 }
 
 ////////////////////////////////////////////////////////////////////////////////////////////////////
@@ -751,6 +752,7 @@ void FTaskTable::SwitchToRelativeToCreatedTimestamps()
 	FindColumnChecked(FTaskTableColumns::StartedTimestampColumnId)->SetValueGetter(MakeShared<FTaskColumnValueGetter<RelativeToCreatedTaskFieldGetterFuncts::GetStartedTimestamp>>());
 	FindColumnChecked(FTaskTableColumns::FinishedTimestampColumnId)->SetValueGetter(MakeShared<FTaskColumnValueGetter<RelativeToCreatedTaskFieldGetterFuncts::GetFinishedTimestamp>>());
 	FindColumnChecked(FTaskTableColumns::CompletedTimestampColumnId)->SetValueGetter(MakeShared<FTaskColumnValueGetter<RelativeToCreatedTaskFieldGetterFuncts::GetCompletedTimestamp>>());
+	FindColumnChecked(FTaskTableColumns::DestroyedTimestampColumnId)->SetValueGetter(MakeShared<FTaskColumnValueGetter<RelativeToCreatedTaskFieldGetterFuncts::GetDestroyedTimestamp>>());
 }
 
 ////////////////////////////////////////////////////////////////////////////////////////////////////
