[Insights] Added simple RTTI to the tree nodes class hierarchy (Insights::FBaseTreeNode).

diff --git a/Engine/Source/Developer/TraceInsights/Private/Insights/CookProfiler/ViewModels/PackageNode.cpp b/Engine/Source/Developer/TraceInsights/Private/Insights/CookProfiler/ViewModels/PackageNode.cpp
--- a/Engine/Source/Developer/TraceInsights/Private/Insights/CookProfiler/ViewModels/PackageNode.cpp
+++ b/Engine/Source/Developer/TraceInsights/Private/Insights/CookProfiler/ViewModels/PackageNode.cpp
@@ -2,12 +2,12 @@
 
 #include "PackageNode.h"
 
-#define LOCTEXT_NAMESPACE "PackageNode"
+#define LOCTEXT_NAMESPACE "Insights::FPackageNode"
 
 namespace Insights
 {
 
-const FName FPackageNode::TypeName(TEXT("FPackageNode"));
+INSIGHTS_IMPLEMENT_RTTI(FPackageNode)
 
 } // namespace Insights
 
diff --git a/Engine/Source/Developer/TraceInsights/Private/Insights/CookProfiler/ViewModels/PackageNode.h b/Engine/Source/Developer/TraceInsights/Private/Insights/CookProfiler/ViewModels/PackageNode.h
--- a/Engine/Source/Developer/TraceInsights/Private/Insights/CookProfiler/ViewModels/PackageNode.h
+++ b/Engine/Source/Developer/TraceInsights/Private/Insights/CookProfiler/ViewModels/PackageNode.h
@@ -48,8 +48,7 @@ typedef TWeakPtr<class FTaskNode> FPackageNodeWeak;
  */
 class FPackageNode : public FTableTreeNode
 {
-public:
-	static const FName TypeName;
+	INSIGHTS_DECLARE_RTTI(FPackageNode, FTableTreeNode)
 
 public:
 	/** Initialization constructor for the Task node. */
@@ -66,8 +65,6 @@ public:
 	{
 	}
 
-	virtual const FName& GetTypeName() const override { return TypeName; }
-
 	/**
 	 * @return a type of this Task node or ETaskNodeType::Group for group nodes.
 	 */
diff --git a/Engine/Source/Developer/TraceInsights/Private/Insights/MemoryProfiler/ViewModels/MemAllocNode.cpp b/Engine/Source/Developer/TraceInsights/Private/Insights/MemoryProfiler/ViewModels/MemAllocNode.cpp
--- a/Engine/Source/Developer/TraceInsights/Private/Insights/MemoryProfiler/ViewModels/MemAllocNode.cpp
+++ b/Engine/Source/Developer/TraceInsights/Private/Insights/MemoryProfiler/ViewModels/MemAllocNode.cpp
@@ -1,14 +1,16 @@
 // Copyright Epic Games, Inc. All Rights Reserved.
 
 #include "MemAllocNode.h"
+
+// Insights
 #include "Insights/MemoryProfiler/ViewModels/CallstackFormatting.h"
 
-#define LOCTEXT_NAMESPACE "MemAllocNode"
+#define LOCTEXT_NAMESPACE "Insights::FMemAllocNode"
 
 namespace Insights
 {
 
-const FName FMemAllocNode::TypeName(TEXT("FMemAllocNode"));
+INSIGHTS_IMPLEMENT_RTTI(FMemAllocNode)
 
 ////////////////////////////////////////////////////////////////////////////////////////////////////
 
diff --git a/Engine/Source/Developer/TraceInsights/Private/Insights/MemoryProfiler/ViewModels/MemAllocNode.h b/Engine/Source/Developer/TraceInsights/Private/Insights/MemoryProfiler/ViewModels/MemAllocNode.h
--- a/Engine/Source/Developer/TraceInsights/Private/Insights/MemoryProfiler/ViewModels/MemAllocNode.h
+++ b/Engine/Source/Developer/TraceInsights/Private/Insights/MemoryProfiler/ViewModels/MemAllocNode.h
@@ -38,8 +38,7 @@ typedef TWeakPtr<class FMemAllocNode> FMemAllocNodeWeak;
  */
 class FMemAllocNode : public FTableTreeNode
 {
-public:
-	static const FName TypeName;
+	INSIGHTS_DECLARE_RTTI(FMemAllocNode, FTableTreeNode)
 
 public:
 	/** Initialization constructor for the MemAlloc node. */
@@ -56,8 +55,6 @@ public:
 	{
 	}
 
-	virtual const FName& GetTypeName() const override { return TypeName; }
-
 	FMemAllocTable& GetMemTableChecked() const { return *MemAllocTable; }
 
 	bool IsValidMemAlloc() const { return GetMemTableChecked().IsValidRowIndex(GetRowIndex()); }
diff --git a/Engine/Source/Developer/TraceInsights/Private/Insights/MemoryProfiler/ViewModels/MemTagNode.cpp b/Engine/Source/Developer/TraceInsights/Private/Insights/MemoryProfiler/ViewModels/MemTagNode.cpp
--- a/Engine/Source/Developer/TraceInsights/Private/Insights/MemoryProfiler/ViewModels/MemTagNode.cpp
+++ b/Engine/Source/Developer/TraceInsights/Private/Insights/MemoryProfiler/ViewModels/MemTagNode.cpp
@@ -5,9 +5,9 @@
 // Insights
 #include "Insights/MemoryProfiler/MemoryProfilerManager.h"
 
-#define LOCTEXT_NAMESPACE "MemTagNode"
+#define LOCTEXT_NAMESPACE "FMemTagNode"
 
-const FName FMemTagNode::TypeName(TEXT("FMemTagNode"));
+INSIGHTS_IMPLEMENT_RTTI(FMemTagNode)
 
 ////////////////////////////////////////////////////////////////////////////////////////////////////
 
diff --git a/Engine/Source/Developer/TraceInsights/Private/Insights/MemoryProfiler/ViewModels/MemTagNode.h b/Engine/Source/Developer/TraceInsights/Private/Insights/MemoryProfiler/ViewModels/MemTagNode.h
--- a/Engine/Source/Developer/TraceInsights/Private/Insights/MemoryProfiler/ViewModels/MemTagNode.h
+++ b/Engine/Source/Developer/TraceInsights/Private/Insights/MemoryProfiler/ViewModels/MemTagNode.h
@@ -60,8 +60,7 @@ typedef TWeakPtr<class FMemTagNode> FMemTagNodeWeak;
  */
 class FMemTagNode : public Insights::FBaseTreeNode
 {
-public:
-	static const FName TypeName;
+	INSIGHTS_DECLARE_RTTI(FMemTagNode, FBaseTreeNode)
 
 public:
 	/** Initialization constructor for the MemTag node. */
@@ -82,8 +81,6 @@ public:
 		ResetAggregatedStats();
 	}
 
-	virtual const FName& GetTypeName() const override { return TypeName; }
-
 	/**
 	 * @return a type of this MemTag node or EMemTagNodeType::Group for group nodes.
 	 */
diff --git a/Engine/Source/Developer/TraceInsights/Private/Insights/MemoryProfiler/ViewModels/MemTagNodeGroupingAndSorting.cpp b/Engine/Source/Developer/TraceInsights/Private/Insights/MemoryProfiler/ViewModels/MemTagNodeGroupingAndSorting.cpp
--- a/Engine/Source/Developer/TraceInsights/Private/Insights/MemoryProfiler/ViewModels/MemTagNodeGroupingAndSorting.cpp
+++ b/Engine/Source/Developer/TraceInsights/Private/Insights/MemoryProfiler/ViewModels/MemTagNodeGroupingAndSorting.cpp
@@ -51,10 +51,10 @@ void FMemTagNodeSortingByType::Sort(TArray<Insights::FBaseTreeNodePtr>& NodesToS
 		{
 			INSIGHTS_DEFAULT_PRESORTING_NODES(A, B);
 
-			INSIGHTS_ENSURE(A.IsValid() && A->GetTypeName() == FMemTagNode::TypeName);
+			INSIGHTS_ENSURE(A.IsValid() && A->Is<FMemTagNode>());
 			const FMemTagNodePtr MemTagNodeA = StaticCastSharedPtr<FMemTagNode, Insights::FBaseTreeNode>(A);
 
-			INSIGHTS_ENSURE(B.IsValid() && B->GetTypeName() == FMemTagNode::TypeName);
+			INSIGHTS_ENSURE(B.IsValid() && B->Is<FMemTagNode>());
 			const FMemTagNodePtr MemTagNodeB = StaticCastSharedPtr<FMemTagNode, Insights::FBaseTreeNode>(B);
 
 			if (MemTagNodeA->GetType() == MemTagNodeB->GetType())
@@ -74,10 +74,10 @@ void FMemTagNodeSortingByType::Sort(TArray<Insights::FBaseTreeNodePtr>& NodesToS
 		{
 			INSIGHTS_DEFAULT_PRESORTING_NODES(A, B);
 
-			INSIGHTS_ENSURE(A.IsValid() && A->GetTypeName() == FMemTagNode::TypeName);
+			INSIGHTS_ENSURE(A.IsValid() && A->Is<FMemTagNode>());
 			const FMemTagNodePtr MemTagNodeA = StaticCastSharedPtr<FMemTagNode, Insights::FBaseTreeNode>(A);
 
-			INSIGHTS_ENSURE(B.IsValid() && B->GetTypeName() == FMemTagNode::TypeName);
+			INSIGHTS_ENSURE(B.IsValid() && B->Is<FMemTagNode>());
 			const FMemTagNodePtr MemTagNodeB = StaticCastSharedPtr<FMemTagNode, Insights::FBaseTreeNode>(B);
 
 			if (MemTagNodeA->GetType() == MemTagNodeB->GetType())
@@ -117,10 +117,10 @@ void FMemTagNodeSortingByTracker::Sort(TArray<Insights::FBaseTreeNodePtr>& Nodes
 		{
 			INSIGHTS_DEFAULT_PRESORTING_NODES(A, B);
 
-			INSIGHTS_ENSURE(A.IsValid() && A->GetTypeName() == FMemTagNode::TypeName);
+			INSIGHTS_ENSURE(A.IsValid() && A->Is<FMemTagNode>());
 			const FMemTagNodePtr MemTagNodeA = StaticCastSharedPtr<FMemTagNode, Insights::FBaseTreeNode>(A);
 
-			INSIGHTS_ENSURE(B.IsValid() && B->GetTypeName() == FMemTagNode::TypeName);
+			INSIGHTS_ENSURE(B.IsValid() && B->Is<FMemTagNode>());
 			const FMemTagNodePtr MemTagNodeB = StaticCastSharedPtr<FMemTagNode, Insights::FBaseTreeNode>(B);
 
 			if (MemTagNodeA->GetMemTrackerId() == MemTagNodeB->GetMemTrackerId())
@@ -140,10 +140,10 @@ void FMemTagNodeSortingByTracker::Sort(TArray<Insights::FBaseTreeNodePtr>& Nodes
 		{
 			INSIGHTS_DEFAULT_PRESORTING_NODES(A, B);
 
-			INSIGHTS_ENSURE(A.IsValid() && A->GetTypeName() == FMemTagNode::TypeName);
+			INSIGHTS_ENSURE(A.IsValid() && A->Is<FMemTagNode>());
 			const FMemTagNodePtr MemTagNodeA = StaticCastSharedPtr<FMemTagNode, Insights::FBaseTreeNode>(A);
 
-			INSIGHTS_ENSURE(B.IsValid() && B->GetTypeName() == FMemTagNode::TypeName);
+			INSIGHTS_ENSURE(B.IsValid() && B->Is<FMemTagNode>());
 			const FMemTagNodePtr MemTagNodeB = StaticCastSharedPtr<FMemTagNode, Insights::FBaseTreeNode>(B);
 
 			if (MemTagNodeA->GetMemTrackerId() == MemTagNodeB->GetMemTrackerId())
@@ -183,11 +183,11 @@ void FMemTagNodeSortingByInstanceCount::Sort(TArray<Insights::FBaseTreeNodePtr>&
 		{
 			INSIGHTS_DEFAULT_PRESORTING_NODES(A, B);
 
-			INSIGHTS_ENSURE(A.IsValid() && A->GetTypeName() == FMemTagNode::TypeName);
+			INSIGHTS_ENSURE(A.IsValid() && A->Is<FMemTagNode>());
 			const FMemTagNodePtr MemTagNodeA = StaticCastSharedPtr<FMemTagNode, Insights::FBaseTreeNode>(A);
 			const uint64 ValueA = MemTagNodeA->GetAggregatedStats().InstanceCount;
 
-			INSIGHTS_ENSURE(B.IsValid() && B->GetTypeName() == FMemTagNode::TypeName);
+			INSIGHTS_ENSURE(B.IsValid() && B->Is<FMemTagNode>());
 			const FMemTagNodePtr MemTagNodeB = StaticCastSharedPtr<FMemTagNode, Insights::FBaseTreeNode>(B);
 			const uint64 ValueB = MemTagNodeB->GetAggregatedStats().InstanceCount;
 
@@ -208,11 +208,11 @@ void FMemTagNodeSortingByInstanceCount::Sort(TArray<Insights::FBaseTreeNodePtr>&
 		{
 			INSIGHTS_DEFAULT_PRESORTING_NODES(A, B);
 
-			INSIGHTS_ENSURE(A.IsValid() && A->GetTypeName() == FMemTagNode::TypeName);
+			INSIGHTS_ENSURE(A.IsValid() && A->Is<FMemTagNode>());
 			const FMemTagNodePtr MemTagNodeA = StaticCastSharedPtr<FMemTagNode, Insights::FBaseTreeNode>(A);
 			const uint64 ValueA = MemTagNodeA->GetAggregatedStats().InstanceCount;
 
-			INSIGHTS_ENSURE(B.IsValid() && B->GetTypeName() == FMemTagNode::TypeName);
+			INSIGHTS_ENSURE(B.IsValid() && B->Is<FMemTagNode>());
 			const FMemTagNodePtr MemTagNodeB = StaticCastSharedPtr<FMemTagNode, Insights::FBaseTreeNode>(B);
 			const uint64 ValueB = MemTagNodeB->GetAggregatedStats().InstanceCount;
 
diff --git a/Engine/Source/Developer/TraceInsights/Private/Insights/NetworkingProfiler/ViewModels/NetEventGroupingAndSorting.cpp b/Engine/Source/Developer/TraceInsights/Private/Insights/NetworkingProfiler/ViewModels/NetEventGroupingAndSorting.cpp
--- a/Engine/Source/Developer/TraceInsights/Private/Insights/NetworkingProfiler/ViewModels/NetEventGroupingAndSorting.cpp
+++ b/Engine/Source/Developer/TraceInsights/Private/Insights/NetworkingProfiler/ViewModels/NetEventGroupingAndSorting.cpp
@@ -36,10 +36,10 @@ void FNetEventNodeSortingByEventType::Sort(TArray<Insights::FBaseTreeNodePtr>& N
 	{
 		NodesToSort.Sort([](const Insights::FBaseTreeNodePtr& A, const Insights::FBaseTreeNodePtr& B) -> bool
 		{
-			INSIGHTS_ENSURE(A.IsValid() && A->GetTypeName() == FNetEventNode::TypeName);
+			INSIGHTS_ENSURE(A.IsValid() && A->Is<FNetEventNode>());
 			const FNetEventNodePtr NetEventNodeA = StaticCastSharedPtr<FNetEventNode, Insights::FBaseTreeNode>(A);
 
-			INSIGHTS_ENSURE(B.IsValid() && B->GetTypeName() == FNetEventNode::TypeName);
+			INSIGHTS_ENSURE(B.IsValid() && B->Is<FNetEventNode>());
 			const FNetEventNodePtr NetEventNodeB = StaticCastSharedPtr<FNetEventNode, Insights::FBaseTreeNode>(B);
 
 			if (NetEventNodeA->GetType() == NetEventNodeB->GetType())
@@ -57,10 +57,10 @@ void FNetEventNodeSortingByEventType::Sort(TArray<Insights::FBaseTreeNodePtr>& N
 	{
 		NodesToSort.Sort([](const Insights::FBaseTreeNodePtr& A, const Insights::FBaseTreeNodePtr& B) -> bool
 		{
-			INSIGHTS_ENSURE(A.IsValid() && A->GetTypeName() == FNetEventNode::TypeName);
+			INSIGHTS_ENSURE(A.IsValid() && A->Is<FNetEventNode>());
 			const FNetEventNodePtr NetEventNodeA = StaticCastSharedPtr<FNetEventNode, Insights::FBaseTreeNode>(A);
 
-			INSIGHTS_ENSURE(B.IsValid() && B->GetTypeName() == FNetEventNode::TypeName);
+			INSIGHTS_ENSURE(B.IsValid() && B->Is<FNetEventNode>());
 			const FNetEventNodePtr NetEventNodeB = StaticCastSharedPtr<FNetEventNode, Insights::FBaseTreeNode>(B);
 
 			if (NetEventNodeA->GetType() == NetEventNodeB->GetType())
@@ -98,11 +98,11 @@ void FNetEventNodeSortingByInstanceCount::Sort(TArray<Insights::FBaseTreeNodePtr
 	{
 		NodesToSort.Sort([](const Insights::FBaseTreeNodePtr& A, const Insights::FBaseTreeNodePtr& B) -> bool
 		{
-			INSIGHTS_ENSURE(A.IsValid() && A->GetTypeName() == FNetEventNode::TypeName);
+			INSIGHTS_ENSURE(A.IsValid() && A->Is<FNetEventNode>());
 			const FNetEventNodePtr NetEventNodeA = StaticCastSharedPtr<FNetEventNode, Insights::FBaseTreeNode>(A);
 			const uint64 ValueA = NetEventNodeA->GetAggregatedStats().InstanceCount;
 
-			INSIGHTS_ENSURE(B.IsValid() && B->GetTypeName() == FNetEventNode::TypeName);
+			INSIGHTS_ENSURE(B.IsValid() && B->Is<FNetEventNode>());
 			const FNetEventNodePtr NetEventNodeB = StaticCastSharedPtr<FNetEventNode, Insights::FBaseTreeNode>(B);
 			const uint64 ValueB = NetEventNodeB->GetAggregatedStats().InstanceCount;
 
@@ -121,11 +121,11 @@ void FNetEventNodeSortingByInstanceCount::Sort(TArray<Insights::FBaseTreeNodePtr
 	{
 		NodesToSort.Sort([](const Insights::FBaseTreeNodePtr& A, const Insights::FBaseTreeNodePtr& B) -> bool
 		{
-			INSIGHTS_ENSURE(A.IsValid() && A->GetTypeName() == FNetEventNode::TypeName);
+			INSIGHTS_ENSURE(A.IsValid() && A->Is<FNetEventNode>());
 			const FNetEventNodePtr NetEventNodeA = StaticCastSharedPtr<FNetEventNode, Insights::FBaseTreeNode>(A);
 			const uint64 ValueA = NetEventNodeA->GetAggregatedStats().InstanceCount;
 
-			INSIGHTS_ENSURE(B.IsValid() && B->GetTypeName() == FNetEventNode::TypeName);
+			INSIGHTS_ENSURE(B.IsValid() && B->Is<FNetEventNode>());
 			const FNetEventNodePtr NetEventNodeB = StaticCastSharedPtr<FNetEventNode, Insights::FBaseTreeNode>(B);
 			const uint64 ValueB = NetEventNodeB->GetAggregatedStats().InstanceCount;
 
@@ -164,11 +164,11 @@ void FNetEventNodeSortingByTotalInclusiveSize::Sort(TArray<Insights::FBaseTreeNo
 	{
 		NodesToSort.Sort([](const Insights::FBaseTreeNodePtr& A, const Insights::FBaseTreeNodePtr& B) -> bool
 		{
-			INSIGHTS_ENSURE(A.IsValid() && A->GetTypeName() == FNetEventNode::TypeName);
+			INSIGHTS_ENSURE(A.IsValid() && A->Is<FNetEventNode>());
 			const FNetEventNodePtr NetEventNodeA = StaticCastSharedPtr<FNetEventNode, Insights::FBaseTreeNode>(A);
 			const uint32 ValueA = NetEventNodeA->GetAggregatedStats().TotalInclusive;
 
-			INSIGHTS_ENSURE(B.IsValid() && B->GetTypeName() == FNetEventNode::TypeName);
+			INSIGHTS_ENSURE(B.IsValid() && B->Is<FNetEventNode>());
 			const FNetEventNodePtr NetEventNodeB = StaticCastSharedPtr<FNetEventNode, Insights::FBaseTreeNode>(B);
 			const uint32 ValueB = NetEventNodeB->GetAggregatedStats().TotalInclusive;
 
@@ -187,11 +187,11 @@ void FNetEventNodeSortingByTotalInclusiveSize::Sort(TArray<Insights::FBaseTreeNo
 	{
 		NodesToSort.Sort([](const Insights::FBaseTreeNodePtr& A, const Insights::FBaseTreeNodePtr& B) -> bool
 		{
-			INSIGHTS_ENSURE(A.IsValid() && A->GetTypeName() == FNetEventNode::TypeName);
+			INSIGHTS_ENSURE(A.IsValid() && A->Is<FNetEventNode>());
 			const FNetEventNodePtr NetEventNodeA = StaticCastSharedPtr<FNetEventNode, Insights::FBaseTreeNode>(A);
 			const uint32 ValueA = NetEventNodeA->GetAggregatedStats().TotalInclusive;
 
-			INSIGHTS_ENSURE(B.IsValid() && B->GetTypeName() == FNetEventNode::TypeName);
+			INSIGHTS_ENSURE(B.IsValid() && B->Is<FNetEventNode>());
 			const FNetEventNodePtr NetEventNodeB = StaticCastSharedPtr<FNetEventNode, Insights::FBaseTreeNode>(B);
 			const uint32 ValueB = NetEventNodeB->GetAggregatedStats().TotalInclusive;
 
@@ -230,11 +230,11 @@ void FNetEventNodeSortingByTotalExclusiveSize::Sort(TArray<Insights::FBaseTreeNo
 	{
 		NodesToSort.Sort([](const Insights::FBaseTreeNodePtr& A, const Insights::FBaseTreeNodePtr& B) -> bool
 		{
-			INSIGHTS_ENSURE(A.IsValid() && A->GetTypeName() == FNetEventNode::TypeName);
+			INSIGHTS_ENSURE(A.IsValid() && A->Is<FNetEventNode>());
 			const FNetEventNodePtr NetEventNodeA = StaticCastSharedPtr<FNetEventNode, Insights::FBaseTreeNode>(A);
 			const uint32 ValueA = NetEventNodeA->GetAggregatedStats().TotalExclusive;
 
-			INSIGHTS_ENSURE(B.IsValid() && B->GetTypeName() == FNetEventNode::TypeName);
+			INSIGHTS_ENSURE(B.IsValid() && B->Is<FNetEventNode>());
 			const FNetEventNodePtr NetEventNodeB = StaticCastSharedPtr<FNetEventNode, Insights::FBaseTreeNode>(B);
 			const uint32 ValueB = NetEventNodeB->GetAggregatedStats().TotalExclusive;
 
@@ -253,11 +253,11 @@ void FNetEventNodeSortingByTotalExclusiveSize::Sort(TArray<Insights::FBaseTreeNo
 	{
 		NodesToSort.Sort([](const Insights::FBaseTreeNodePtr& A, const Insights::FBaseTreeNodePtr& B) -> bool
 		{
-			INSIGHTS_ENSURE(A.IsValid() && A->GetTypeName() == FNetEventNode::TypeName);
+			INSIGHTS_ENSURE(A.IsValid() && A->Is<FNetEventNode>());
 			const FNetEventNodePtr NetEventNodeA = StaticCastSharedPtr<FNetEventNode, Insights::FBaseTreeNode>(A);
 			const uint32 ValueA = NetEventNodeA->GetAggregatedStats().TotalExclusive;
 
-			INSIGHTS_ENSURE(B.IsValid() && B->GetTypeName() == FNetEventNode::TypeName);
+			INSIGHTS_ENSURE(B.IsValid() && B->Is<FNetEventNode>());
 			const FNetEventNodePtr NetEventNodeB = StaticCastSharedPtr<FNetEventNode, Insights::FBaseTreeNode>(B);
 			const uint32 ValueB = NetEventNodeB->GetAggregatedStats().TotalExclusive;
 
diff --git a/Engine/Source/Developer/TraceInsights/Private/Insights/NetworkingProfiler/ViewModels/NetEventNode.cpp b/Engine/Source/Developer/TraceInsights/Private/Insights/NetworkingProfiler/ViewModels/NetEventNode.cpp
--- a/Engine/Source/Developer/TraceInsights/Private/Insights/NetworkingProfiler/ViewModels/NetEventNode.cpp
+++ b/Engine/Source/Developer/TraceInsights/Private/Insights/NetworkingProfiler/ViewModels/NetEventNode.cpp
@@ -2,9 +2,9 @@
 
 #include "NetEventNode.h"
 
-#define LOCTEXT_NAMESPACE "NetEventNode"
+#define LOCTEXT_NAMESPACE "FNetEventNode"
 
-const FName FNetEventNode::TypeName(TEXT("FNetEventNode"));
+INSIGHTS_IMPLEMENT_RTTI(FNetEventNode)
 
 ////////////////////////////////////////////////////////////////////////////////////////////////////
 
diff --git a/Engine/Source/Developer/TraceInsights/Private/Insights/NetworkingProfiler/ViewModels/NetEventNode.h b/Engine/Source/Developer/TraceInsights/Private/Insights/NetworkingProfiler/ViewModels/NetEventNode.h
--- a/Engine/Source/Developer/TraceInsights/Private/Insights/NetworkingProfiler/ViewModels/NetEventNode.h
+++ b/Engine/Source/Developer/TraceInsights/Private/Insights/NetworkingProfiler/ViewModels/NetEventNode.h
@@ -47,8 +47,9 @@ typedef TWeakPtr<class FNetEventNode> FNetEventNodeWeak;
  */
 class FNetEventNode : public Insights::FBaseTreeNode
 {
+	INSIGHTS_DECLARE_RTTI(FNetEventNode, FBaseTreeNode)
+
 public:
-	static const FName TypeName;
 	static constexpr uint32 InvalidEventTypeIndex = uint32(-1);
 
 public:
@@ -72,8 +73,6 @@ public:
 		ResetAggregatedStats();
 	}
 
-	virtual const FName& GetTypeName() const override { return TypeName; }
-
 	uint32 GetEventTypeIndex() const { return EventTypeIndex; }
 
 	/**
diff --git a/Engine/Source/Developer/TraceInsights/Private/Insights/NetworkingProfiler/ViewModels/NetStatsCounterGroupingAndSorting.cpp b/Engine/Source/Developer/TraceInsights/Private/Insights/NetworkingProfiler/ViewModels/NetStatsCounterGroupingAndSorting.cpp
--- a/Engine/Source/Developer/TraceInsights/Private/Insights/NetworkingProfiler/ViewModels/NetStatsCounterGroupingAndSorting.cpp
+++ b/Engine/Source/Developer/TraceInsights/Private/Insights/NetworkingProfiler/ViewModels/NetStatsCounterGroupingAndSorting.cpp
@@ -32,10 +32,10 @@ void FNetStatsCounterNodeSortingByEventType::Sort(TArray<Insights::FBaseTreeNode
 	{
 		NodesToSort.Sort([](const Insights::FBaseTreeNodePtr& A, const Insights::FBaseTreeNodePtr& B) -> bool
 		{
-			ensure(A.IsValid() && A->GetTypeName() == FNetStatsCounterNode::TypeName);
+			ensure(A.IsValid() && A->Is<FNetStatsCounterNode>());
 			const FNetStatsCounterNodePtr NetStatsCounterNodeA = StaticCastSharedPtr<FNetStatsCounterNode, Insights::FBaseTreeNode>(A);
 
-			ensure(B.IsValid() && B->GetTypeName() == FNetStatsCounterNode::TypeName);
+			ensure(B.IsValid() && B->Is<FNetStatsCounterNode>());
 			const FNetStatsCounterNodePtr NetStatsCounterNodeB = StaticCastSharedPtr<FNetStatsCounterNode, Insights::FBaseTreeNode>(B);
 
 			if (NetStatsCounterNodeA->GetType() == NetStatsCounterNodeB->GetType())
@@ -53,10 +53,10 @@ void FNetStatsCounterNodeSortingByEventType::Sort(TArray<Insights::FBaseTreeNode
 	{
 		NodesToSort.Sort([](const Insights::FBaseTreeNodePtr& A, const Insights::FBaseTreeNodePtr& B) -> bool
 		{
-			ensure(A.IsValid() && A->GetTypeName() == FNetStatsCounterNode::TypeName);
+			ensure(A.IsValid() && A->Is<FNetStatsCounterNode>());
 			const FNetStatsCounterNodePtr NetStatsCounterNodeA = StaticCastSharedPtr<FNetStatsCounterNode, Insights::FBaseTreeNode>(A);
 
-			ensure(B.IsValid() && B->GetTypeName() == FNetStatsCounterNode::TypeName);
+			ensure(B.IsValid() && B->Is<FNetStatsCounterNode>());
 			const FNetStatsCounterNodePtr NetStatsCounterNodeB = StaticCastSharedPtr<FNetStatsCounterNode, Insights::FBaseTreeNode>(B);
 
 			if (NetStatsCounterNodeA->GetType() == NetStatsCounterNodeB->GetType())
@@ -94,11 +94,11 @@ void FNetStatsCounterNodeSortingBySum::Sort(TArray<Insights::FBaseTreeNodePtr>&
 	{
 		NodesToSort.Sort([](const Insights::FBaseTreeNodePtr& A, const Insights::FBaseTreeNodePtr& B) -> bool
 		{
-			ensure(A.IsValid() && A->GetTypeName() == FNetStatsCounterNode::TypeName);
+			ensure(A.IsValid() && A->Is<FNetStatsCounterNode>());
 			const FNetStatsCounterNodePtr NetStatsCounterNodeA = StaticCastSharedPtr<FNetStatsCounterNode, Insights::FBaseTreeNode>(A);
 			const uint32 ValueA = NetStatsCounterNodeA->GetAggregatedStats().Sum;
 
-			ensure(B.IsValid() && B->GetTypeName() == FNetStatsCounterNode::TypeName);
+			ensure(B.IsValid() && B->Is<FNetStatsCounterNode>());
 			const FNetStatsCounterNodePtr NetStatsCounterNodeB = StaticCastSharedPtr<FNetStatsCounterNode, Insights::FBaseTreeNode>(B);
 			const uint32 ValueB = NetStatsCounterNodeB->GetAggregatedStats().Sum;
 
@@ -117,11 +117,11 @@ void FNetStatsCounterNodeSortingBySum::Sort(TArray<Insights::FBaseTreeNodePtr>&
 	{
 		NodesToSort.Sort([](const Insights::FBaseTreeNodePtr& A, const Insights::FBaseTreeNodePtr& B) -> bool
 		{
-			ensure(A.IsValid() && A->GetTypeName() == FNetStatsCounterNode::TypeName);
+			ensure(A.IsValid() && A->Is<FNetStatsCounterNode>());
 			const FNetStatsCounterNodePtr NetStatsCounterNodeA = StaticCastSharedPtr<FNetStatsCounterNode, Insights::FBaseTreeNode>(A);
 			const uint32 ValueA = NetStatsCounterNodeA->GetAggregatedStats().Sum;
 
-			ensure(B.IsValid() && B->GetTypeName() == FNetStatsCounterNode::TypeName);
+			ensure(B.IsValid() && B->Is<FNetStatsCounterNode>());
 			const FNetStatsCounterNodePtr NetStatsCounterNodeB = StaticCastSharedPtr<FNetStatsCounterNode, Insights::FBaseTreeNode>(B);
 			const uint32 ValueB = NetStatsCounterNodeB->GetAggregatedStats().Sum;
 
diff --git a/Engine/Source/Developer/TraceInsights/Private/Insights/NetworkingProfiler/ViewModels/NetStatsCounterNode.cpp b/Engine/Source/Developer/TraceInsights/Private/Insights/NetworkingProfiler/ViewModels/NetStatsCounterNode.cpp
--- a/Engine/Source/Developer/TraceInsights/Private/Insights/NetworkingProfiler/ViewModels/NetStatsCounterNode.cpp
+++ b/Engine/Source/Developer/TraceInsights/Private/Insights/NetworkingProfiler/ViewModels/NetStatsCounterNode.cpp
@@ -5,9 +5,9 @@
 // Insights
 #include "Insights/Common/TimeUtils.h"
 
-#define LOCTEXT_NAMESPACE "NetStatsNode"
+#define LOCTEXT_NAMESPACE "FNetStatsCounterNode"
 
-const FName FNetStatsCounterNode::TypeName(TEXT("FNetStatsCounterNode"));
+INSIGHTS_IMPLEMENT_RTTI(FNetStatsCounterNode)
 
 ////////////////////////////////////////////////////////////////////////////////////////////////////
 
diff --git a/Engine/Source/Developer/TraceInsights/Private/Insights/NetworkingProfiler/ViewModels/NetStatsCounterNode.h b/Engine/Source/Developer/TraceInsights/Private/Insights/NetworkingProfiler/ViewModels/NetStatsCounterNode.h
--- a/Engine/Source/Developer/TraceInsights/Private/Insights/NetworkingProfiler/ViewModels/NetStatsCounterNode.h
+++ b/Engine/Source/Developer/TraceInsights/Private/Insights/NetworkingProfiler/ViewModels/NetStatsCounterNode.h
@@ -48,11 +48,12 @@ typedef TWeakPtr<class FNetStatsCounterNode> FNetStatsCounterNodeWeak;
  */
 class FNetStatsCounterNode : public Insights::FBaseTreeNode
 {
+	INSIGHTS_DECLARE_RTTI(FNetStatsCounterNode, FBaseTreeNode)
+
 public:
 	using FNetProfilerAggregatedStatsCounterStats = TraceServices::FNetProfilerAggregatedStatsCounterStats;
 	using ENetProfilerStatsCounterType = TraceServices::ENetProfilerStatsCounterType;
 
-	static const FName TypeName;
 	static constexpr uint32 InvalidCounterIndex = uint32(-1);
 
 public:
@@ -74,8 +75,6 @@ public:
 		ResetAggregatedStats();
 	}
 
-	virtual const FName& GetTypeName() const override { return TypeName; }
-
 	/**
 	 * @return the counter id as provided by analyzer.
 	 */
diff --git a/Engine/Source/Developer/TraceInsights/Private/Insights/Table/ViewModels/BaseTreeNode.cpp b/Engine/Source/Developer/TraceInsights/Private/Insights/Table/ViewModels/BaseTreeNode.cpp
--- a/Engine/Source/Developer/TraceInsights/Private/Insights/Table/ViewModels/BaseTreeNode.cpp
+++ b/Engine/Source/Developer/TraceInsights/Private/Insights/Table/ViewModels/BaseTreeNode.cpp
@@ -5,11 +5,13 @@
 // Insights
 #include "Insights/Table/ViewModels/TableCellValueSorter.h"
 
-#define LOCTEXT_NAMESPACE "Insights_TreeNode"
+#define LOCTEXT_NAMESPACE "Insights::FBaseTreeNode"
 
 namespace Insights
 {
 
+INSIGHTS_IMPLEMENT_RTTI(FBaseTreeNode)
+
 FBaseTreeNode::FGroupNodeData FBaseTreeNode::DefaultGroupData;
 
 ////////////////////////////////////////////////////////////////////////////////////////////////////
diff --git a/Engine/Source/Developer/TraceInsights/Private/Insights/Table/ViewModels/BaseTreeNode.h b/Engine/Source/Developer/TraceInsights/Private/Insights/Table/ViewModels/BaseTreeNode.h
--- a/Engine/Source/Developer/TraceInsights/Private/Insights/Table/ViewModels/BaseTreeNode.h
+++ b/Engine/Source/Developer/TraceInsights/Private/Insights/Table/ViewModels/BaseTreeNode.h
@@ -4,6 +4,8 @@
 
 #include "CoreMinimal.h"
 
+#include "Insights/Common/SimpleRtti.h"
+
 namespace Insights
 {
 
@@ -29,6 +31,8 @@ typedef TWeakPtr<class FBaseTreeNode> FBaseTreeNodeWeak;
 
 class FBaseTreeNode : public TSharedFromThis<FBaseTreeNode>
 {
+	INSIGHTS_DECLARE_RTTI_BASE(FBaseTreeNode)
+
 protected:
 	struct FGroupNodeData
 	{
@@ -66,10 +70,14 @@ public:
 		}
 	}
 
-	virtual const FName& GetTypeName() const = 0;
-
-	uint32 GetDefaultSortOrder() const { return DefaultSortOrder; }
-	void SetDefaultSortOrder(uint32 Order) { DefaultSortOrder = Order; }
+	uint32 GetDefaultSortOrder() const
+	{
+		return DefaultSortOrder;
+	}
+	void SetDefaultSortOrder(uint32 Order)
+	{
+		DefaultSortOrder = Order;
+	}
 
 	/**
 	 * @return a name of this node.
@@ -94,22 +102,40 @@ public:
 	/**
 	 * @return a descriptive text for this node to display in a tooltip.
 	 */
-	virtual const FText GetTooltip() const { return IsGroup() ? GroupData->Tooltip : FText::GetEmpty(); }
+	virtual const FText GetTooltip() const
+	{
+		return GroupData->Tooltip;
+	}
 
 	/**
 	 * Sets a descriptive text for this node to display in a tooltip.
 	 */
-	virtual void SetTooltip(const FText& InTooltip) { if (IsGroup()) { GroupData->Tooltip = InTooltip; } }
+	virtual void SetTooltip(const FText& InTooltip)
+	{
+		if (IsGroup())
+		{
+			GroupData->Tooltip = InTooltip;
+		}
+	}
 
 	/**
 	 * @return a pointer to a data context for this node.
 	 */
-	virtual void* GetContext() const { return IsGroup() ? GroupData->Context : nullptr; }
+	virtual void* GetContext() const
+	{
+		return GroupData->Context;
+	}
 
 	/**
 	 * Sets a pointer to a data context for this node.
 	 */
-	virtual void SetContext(void* InContext) { if (IsGroup()) { GroupData->Context = InContext; } }
+	virtual void SetContext(void* InContext)
+	{
+		if (IsGroup())
+		{
+			GroupData->Context = InContext;
+		}
+	}
 
 	/**
 	 * @return true, if this node is a group node.
@@ -119,6 +145,17 @@ public:
 		return GroupData != &DefaultGroupData;
 	}
 
+	/**
+	 * Initializes the group data, allowing this node to accept children.
+	 */
+	void InitGroupData()
+	{
+		if (GroupData == &DefaultGroupData)
+		{
+			GroupData = new FGroupNodeData();
+		}
+	}
+
 	/**
 	 * @return a const reference to the child nodes of this group.
 	 */
@@ -220,8 +257,17 @@ public:
 		GroupData->FilteredChildren.Reset(NewSize);
 	}
 
-	bool IsExpanded() const { return GroupData->bIsExpanded; }
-	void SetExpansion(bool bOnOff) { GroupData->bIsExpanded = bOnOff; }
+	bool IsExpanded() const
+	{
+		return GroupData->bIsExpanded;
+	}
+	void SetExpansion(bool bOnOff)
+	{
+		if (IsGroup())
+		{
+			GroupData->bIsExpanded = bOnOff;
+		}
+	}
 
 protected:
 	/**
diff --git a/Engine/Source/Developer/TraceInsights/Private/Insights/Table/ViewModels/TableTreeNode.cpp b/Engine/Source/Developer/TraceInsights/Private/Insights/Table/ViewModels/TableTreeNode.cpp
--- a/Engine/Source/Developer/TraceInsights/Private/Insights/Table/ViewModels/TableTreeNode.cpp
+++ b/Engine/Source/Developer/TraceInsights/Private/Insights/Table/ViewModels/TableTreeNode.cpp
@@ -2,12 +2,12 @@
 
 #include "TableTreeNode.h"
 
-#define LOCTEXT_NAMESPACE "Insights_TableTreeNode"
+#define LOCTEXT_NAMESPACE "Insights::FTableTreeNode"
 
 namespace Insights
 {
 
-const FName FTableTreeNode::TypeName(TEXT("FTableTreeNode"));
+INSIGHTS_IMPLEMENT_RTTI(FTableTreeNode)
 
 } // namespace Insights
 
diff --git a/Engine/Source/Developer/TraceInsights/Private/Insights/Table/ViewModels/TableTreeNode.h b/Engine/Source/Developer/TraceInsights/Private/Insights/Table/ViewModels/TableTreeNode.h
--- a/Engine/Source/Developer/TraceInsights/Private/Insights/Table/ViewModels/TableTreeNode.h
+++ b/Engine/Source/Developer/TraceInsights/Private/Insights/Table/ViewModels/TableTreeNode.h
@@ -47,8 +47,7 @@ typedef TWeakPtr<class FTableTreeNode> FTableTreeNodeWeak;
  */
 class FTableTreeNode : public FBaseTreeNode
 {
-public:
-	static const FName TypeName;
+	INSIGHTS_DECLARE_RTTI(FTableTreeNode, FBaseTreeNode)
 
 public:
 	/** Initialization constructor for a table record node. */
@@ -74,8 +73,6 @@ public:
 		CleanupAggregatedValues();
 	}
 
-	virtual const FName& GetTypeName() const override { return TypeName; }
-
 	const TWeakPtr<FTable>& GetParentTable() const { return ParentTable; }
 	FTableRowId GetRowId() const { return RowId; }
 	int32 GetRowIndex() const { return RowId.RowIndex; }
diff --git a/Engine/Source/Developer/TraceInsights/Private/Insights/Table/ViewModels/UntypedTable.cpp b/Engine/Source/Developer/TraceInsights/Private/Insights/Table/ViewModels/UntypedTable.cpp
--- a/Engine/Source/Developer/TraceInsights/Private/Insights/Table/ViewModels/UntypedTable.cpp
+++ b/Engine/Source/Developer/TraceInsights/Private/Insights/Table/ViewModels/UntypedTable.cpp
@@ -26,7 +26,7 @@ public:
 
 	virtual const TOptional<FTableCellValue> GetValue(const FTableColumn& Column, const FBaseTreeNode& Node) const
 	{
-		ensure(Node.GetTypeName() == FTableTreeNode::TypeName);
+		ensure(Node.Is<FTableTreeNode>());
 		const FTableTreeNode& TableTreeNode = static_cast<const FTableTreeNode&>(Node);
 
 		if (!Node.IsGroup()) // Table Row Node
diff --git a/Engine/Source/Developer/TraceInsights/Private/Insights/TaskGraphProfiler/ViewModels/TaskNode.cpp b/Engine/Source/Developer/TraceInsights/Private/Insights/TaskGraphProfiler/ViewModels/TaskNode.cpp
--- a/Engine/Source/Developer/TraceInsights/Private/Insights/TaskGraphProfiler/ViewModels/TaskNode.cpp
+++ b/Engine/Source/Developer/TraceInsights/Private/Insights/TaskGraphProfiler/ViewModels/TaskNode.cpp
@@ -2,12 +2,12 @@
 
 #include "TaskNode.h"
 
-#define LOCTEXT_NAMESPACE "TaskNodeNode"
+#define LOCTEXT_NAMESPACE "Insights::FTaskNode"
 
 namespace Insights
 {
 
-const FName FTaskNode::TypeName(TEXT("FTaskNode"));
+INSIGHTS_IMPLEMENT_RTTI(FTaskNode)
 
 } // namespace Insights
 
diff --git a/Engine/Source/Developer/TraceInsights/Private/Insights/TaskGraphProfiler/ViewModels/TaskNode.h b/Engine/Source/Developer/TraceInsights/Private/Insights/TaskGraphProfiler/ViewModels/TaskNode.h
--- a/Engine/Source/Developer/TraceInsights/Private/Insights/TaskGraphProfiler/ViewModels/TaskNode.h
+++ b/Engine/Source/Developer/TraceInsights/Private/Insights/TaskGraphProfiler/ViewModels/TaskNode.h
@@ -48,8 +48,7 @@ typedef TWeakPtr<class FTaskNode> FTaskNodeWeak;
  */
 class FTaskNode : public FTableTreeNode
 {
-public:
-	static const FName TypeName;
+	INSIGHTS_DECLARE_RTTI(FTaskNode, FTableTreeNode)
 
 public:
 	/** Initialization constructor for the Task node. */
@@ -66,8 +65,6 @@ public:
 	{
 	}
 
-	virtual const FName& GetTypeName() const override { return TypeName; }
-
 	/**
 	 * @return a type of this Task node or ETaskNodeType::Group for group nodes.
 	 */
diff --git a/Engine/Source/Developer/TraceInsights/Private/Insights/ViewModels/FilterConfiguratorNode.cpp b/Engine/Source/Developer/TraceInsights/Private/Insights/ViewModels/FilterConfiguratorNode.cpp
--- a/Engine/Source/Developer/TraceInsights/Private/Insights/ViewModels/FilterConfiguratorNode.cpp
+++ b/Engine/Source/Developer/TraceInsights/Private/Insights/ViewModels/FilterConfiguratorNode.cpp
@@ -7,14 +7,12 @@
 // Insights
 #include "Insights/ViewModels/Filters.h"
 
-#define LOCTEXT_NAMESPACE "FilterConfiguratorNode"
+#define LOCTEXT_NAMESPACE "FFilterConfiguratorNode"
 
 namespace Insights
 {
 
-////////////////////////////////////////////////////////////////////////////////////////////////////
-
-const FName FFilterConfiguratorNode::TypeName(TEXT("FFilterConfiguratorNode"));
+INSIGHTS_IMPLEMENT_RTTI(FFilterConfiguratorNode)
 
 ////////////////////////////////////////////////////////////////////////////////////////////////////
 
diff --git a/Engine/Source/Developer/TraceInsights/Private/Insights/ViewModels/FilterConfiguratorNode.h b/Engine/Source/Developer/TraceInsights/Private/Insights/ViewModels/FilterConfiguratorNode.h
--- a/Engine/Source/Developer/TraceInsights/Private/Insights/ViewModels/FilterConfiguratorNode.h
+++ b/Engine/Source/Developer/TraceInsights/Private/Insights/ViewModels/FilterConfiguratorNode.h
@@ -45,7 +45,7 @@ typedef TWeakPtr<class FFilterConfiguratorNode> FFilterConfiguratorNodeWeak;
  */
 class FFilterConfiguratorNode : public Insights::FBaseTreeNode
 {
-	static const FName TypeName;
+	INSIGHTS_DECLARE_RTTI(FFilterConfiguratorNode, FBaseTreeNode)
 
 public:
 	/** Initialization constructor for the filter configurator node. */
@@ -57,8 +57,6 @@ public:
 
 	virtual ~FFilterConfiguratorNode() {}
 
-	virtual const FName& GetTypeName() const override { return TypeName; }
-
 	void SetAvailableFilters(TSharedPtr<TArray<TSharedPtr<struct FFilter>>> InAvailableFilters);
 	TSharedPtr<TArray<TSharedPtr<struct FFilter>>> GetAvailableFilters() { return AvailableFilters; }
 
diff --git a/Engine/Source/Developer/TraceInsights/Private/Insights/ViewModels/StatsNode.cpp b/Engine/Source/Developer/TraceInsights/Private/Insights/ViewModels/StatsNode.cpp
--- a/Engine/Source/Developer/TraceInsights/Private/Insights/ViewModels/StatsNode.cpp
+++ b/Engine/Source/Developer/TraceInsights/Private/Insights/ViewModels/StatsNode.cpp
@@ -5,9 +5,9 @@
 // Insights
 #include "Insights/Common/TimeUtils.h"
 
-#define LOCTEXT_NAMESPACE "StatsNode"
+#define LOCTEXT_NAMESPACE "FStatsNode"
 
-const FName FStatsNode::TypeName(TEXT("FStatsNode"));
+INSIGHTS_IMPLEMENT_RTTI(FStatsNode)
 
 ////////////////////////////////////////////////////////////////////////////////////////////////////
 
diff --git a/Engine/Source/Developer/TraceInsights/Private/Insights/ViewModels/StatsNode.h b/Engine/Source/Developer/TraceInsights/Private/Insights/ViewModels/StatsNode.h
--- a/Engine/Source/Developer/TraceInsights/Private/Insights/ViewModels/StatsNode.h
+++ b/Engine/Source/Developer/TraceInsights/Private/Insights/ViewModels/StatsNode.h
@@ -131,8 +131,9 @@ typedef TWeakPtr<class FStatsNode> FStatsNodeWeak;
  */
 class FStatsNode : public Insights::FBaseTreeNode
 {
+	INSIGHTS_DECLARE_RTTI(FStatsNode, FBaseTreeNode)
+
 public:
-	static const FName TypeName;
 	static constexpr uint32 InvalidCounterId = uint32(-1);
 
 public:
@@ -166,8 +167,6 @@ public:
 		ResetAggregatedStats();
 	}
 
-	virtual const FName& GetTypeName() const override { return TypeName; }
-
 	/**
 	 * @return the counter id as provided by analyzer.
 	 */
diff --git a/Engine/Source/Developer/TraceInsights/Private/Insights/ViewModels/TimerNode.cpp b/Engine/Source/Developer/TraceInsights/Private/Insights/ViewModels/TimerNode.cpp
--- a/Engine/Source/Developer/TraceInsights/Private/Insights/ViewModels/TimerNode.cpp
+++ b/Engine/Source/Developer/TraceInsights/Private/Insights/ViewModels/TimerNode.cpp
@@ -8,9 +8,9 @@
 #include "Insights/InsightsManager.h"
 #include "Insights/ViewModels/TimingEvent.h"
 
-#define LOCTEXT_NAMESPACE "TimerNode"
+#define LOCTEXT_NAMESPACE "FTimerNode"
 
-const FName FTimerNode::TypeName(TEXT("FTimerNode"));
+INSIGHTS_IMPLEMENT_RTTI(FTimerNode)
 
 const FName FTimerNode::GpuGroup(TEXT("GPU"));
 const FName FTimerNode::CpuGroup(TEXT("CPU"));
diff --git a/Engine/Source/Developer/TraceInsights/Private/Insights/ViewModels/TimerNode.h b/Engine/Source/Developer/TraceInsights/Private/Insights/ViewModels/TimerNode.h
--- a/Engine/Source/Developer/TraceInsights/Private/Insights/ViewModels/TimerNode.h
+++ b/Engine/Source/Developer/TraceInsights/Private/Insights/ViewModels/TimerNode.h
@@ -51,8 +51,9 @@ typedef TWeakPtr<class FTimerNode> FTimerNodeWeak;
  */
 class FTimerNode : public Insights::FBaseTreeNode
 {
+	INSIGHTS_DECLARE_RTTI(FTimerNode, FBaseTreeNode)
+
 public:
-	static const FName TypeName;
 	static constexpr uint32 InvalidTimerId = uint32(-1);
 
 	static const FName GpuGroup;
@@ -67,8 +68,6 @@ public:
 
 	virtual ~FTimerNode();
 
-	virtual const FName& GetTypeName() const override { return TypeName; }
-
 	/**
 	 * @return the timer id as provided by analyzer. It can be used as an index.
 	 */
