[Insights]
  - Tree Table views: Added "Path Breakdown" grouping for all string columns. See this gropuping for "Top Source File" column in Memory Insights' Mem Allocs table view.
  - Memory Insights: Added "Asset" column and "Asset" view preset to Mem Allocs table view.
  - Exposed read API in IMetadataProvider.
  - Added ThreadId for an allocation (it is needed to resolve metadata).

diff --git a/Engine/Source/Developer/TraceInsights/Private/Insights/MemoryProfiler/ViewModels/MemAllocTable.cpp b/Engine/Source/Developer/TraceInsights/Private/Insights/MemoryProfiler/ViewModels/MemAllocTable.cpp
--- a/Engine/Source/Developer/TraceInsights/Private/Insights/MemoryProfiler/ViewModels/MemAllocTable.cpp
+++ b/Engine/Source/Developer/TraceInsights/Private/Insights/MemoryProfiler/ViewModels/MemAllocTable.cpp
@@ -36,6 +36,7 @@ const FName FMemAllocTableColumns::MemoryPageColumnId(TEXT("MemoryPage"));
 const FName FMemAllocTableColumns::CountColumnId(TEXT("Count"));
 const FName FMemAllocTableColumns::SizeColumnId(TEXT("Size"));
 const FName FMemAllocTableColumns::TagColumnId(TEXT("Tag"));
+const FName FMemAllocTableColumns::AssetColumnId(TEXT("Asset"));
 const FName FMemAllocTableColumns::FunctionColumnId(TEXT("Function"));
 const FName FMemAllocTableColumns::SourceFileColumnId(TEXT("SourceFile"));
 const FName FMemAllocTableColumns::CallstackSizeColumnId(TEXT("CallstackSize"));
@@ -694,7 +695,7 @@ void FMemAllocTable::AddDefaultColumns()
 
 		Column.SetDataType(ETableCellDataType::CString);
 
-		class FMemAllocLlmTagValueGetter : public FTableCellValueGetter
+		class FMemTagValueGetter : public FTableCellValueGetter
 		{
 		public:
 			virtual const TOptional<FTableCellValue> GetValue(const FTableColumn& Column, const FBaseTreeNode& Node) const
@@ -720,7 +721,65 @@ void FMemAllocTable::AddDefaultColumns()
 				return TOptional<FTableCellValue>();
 			}
 		};
-		TSharedRef<ITableCellValueGetter> Getter = MakeShared<FMemAllocLlmTagValueGetter>();
+		TSharedRef<ITableCellValueGetter> Getter = MakeShared<FMemTagValueGetter>();
+		Column.SetValueGetter(Getter);
+
+		TSharedRef<ITableCellValueFormatter> Formatter = MakeShared<FCStringValueFormatterAsText>();
+		Column.SetValueFormatter(Formatter);
+
+		TSharedRef<ITableCellValueSorter> Sorter = MakeShared<FSorterByCStringValue>(ColumnRef);
+		Column.SetValueSorter(Sorter);
+
+		Column.SetAggregation(ETableColumnAggregation::SameValue);
+
+		AddColumn(ColumnRef);
+	}
+	//////////////////////////////////////////////////
+	// Asset Column
+	{
+		TSharedRef<FTableColumn> ColumnRef = MakeShared<FTableColumn>(FMemAllocTableColumns::AssetColumnId);
+		FTableColumn& Column = *ColumnRef;
+
+		Column.SetIndex(ColumnIndex++);
+
+		Column.SetShortName(LOCTEXT("AssetColumnName", "Asset"));
+		Column.SetTitleName(LOCTEXT("AssetColumnTitle", "Asset"));
+		Column.SetDescription(LOCTEXT("AssetColumnDesc", "Asset associated with allocation"));
+
+		Column.SetFlags(ETableColumnFlags::CanBeHidden | ETableColumnFlags::CanBeFiltered);
+
+		Column.SetHorizontalAlignment(HAlign_Left);
+		Column.SetInitialWidth(320.0f);
+
+		Column.SetDataType(ETableCellDataType::CString);
+
+		class FAssetMetadataValueGetter : public FTableCellValueGetter
+		{
+		public:
+			virtual const TOptional<FTableCellValue> GetValue(const FTableColumn& Column, const FBaseTreeNode& Node) const
+			{
+				if (Node.IsGroup())
+				{
+					const FTableTreeNode& NodePtr = static_cast<const FTableTreeNode&>(Node);
+					if (NodePtr.HasAggregatedValue(Column.GetId()))
+					{
+						return NodePtr.GetAggregatedValue(Column.GetId());
+					}
+				}
+				else //if (Node->Is<FMemAllocNode>())
+				{
+					const FMemAllocNode& MemAllocNode = static_cast<const FMemAllocNode&>(Node);
+					const FMemoryAlloc* Alloc = MemAllocNode.GetMemAlloc();
+					if (Alloc)
+					{
+						return FTableCellValue(Alloc->GetAsset());
+					}
+				}
+
+				return TOptional<FTableCellValue>();
+			}
+		};
+		TSharedRef<ITableCellValueGetter> Getter = MakeShared<FAssetMetadataValueGetter>();
 		Column.SetValueGetter(Getter);
 
 		TSharedRef<ITableCellValueFormatter> Formatter = MakeShared<FCStringValueFormatterAsText>();
diff --git a/Engine/Source/Developer/TraceInsights/Private/Insights/MemoryProfiler/ViewModels/MemAllocTable.h b/Engine/Source/Developer/TraceInsights/Private/Insights/MemoryProfiler/ViewModels/MemAllocTable.h
--- a/Engine/Source/Developer/TraceInsights/Private/Insights/MemoryProfiler/ViewModels/MemAllocTable.h
+++ b/Engine/Source/Developer/TraceInsights/Private/Insights/MemoryProfiler/ViewModels/MemAllocTable.h
@@ -29,6 +29,7 @@ struct FMemAllocTableColumns
 	static const FName CountColumnId;
 	static const FName SizeColumnId;
 	static const FName TagColumnId;
+	static const FName AssetColumnId;
 	static const FName FunctionColumnId;
 	static const FName SourceFileColumnId;
 	static const FName CallstackSizeColumnId;
diff --git a/Engine/Source/Developer/TraceInsights/Private/Insights/MemoryProfiler/ViewModels/MemoryAlloc.cpp b/Engine/Source/Developer/TraceInsights/Private/Insights/MemoryProfiler/ViewModels/MemoryAlloc.cpp
--- a/Engine/Source/Developer/TraceInsights/Private/Insights/MemoryProfiler/ViewModels/MemoryAlloc.cpp
+++ b/Engine/Source/Developer/TraceInsights/Private/Insights/MemoryProfiler/ViewModels/MemoryAlloc.cpp
@@ -16,10 +16,14 @@ FMemoryAlloc::FMemoryAlloc()
 	, EndTime(0.0)
 	, Address(0)
 	, Size(0)
+	, TagId(0)
 	, Tag(nullptr)
+	, Asset(nullptr)
 	, Callstack(nullptr)
+	, FreeCallstack(nullptr)
 	, RootHeap(0)
 	, bIsBlock(false)
+	, bIsDecline(false)
 {
 }
 
diff --git a/Engine/Source/Developer/TraceInsights/Private/Insights/MemoryProfiler/ViewModels/MemoryAlloc.h b/Engine/Source/Developer/TraceInsights/Private/Insights/MemoryProfiler/ViewModels/MemoryAlloc.h
--- a/Engine/Source/Developer/TraceInsights/Private/Insights/MemoryProfiler/ViewModels/MemoryAlloc.h
+++ b/Engine/Source/Developer/TraceInsights/Private/Insights/MemoryProfiler/ViewModels/MemoryAlloc.h
@@ -29,12 +29,14 @@ public:
 	uint64 GetAddress() const { return Address; }
 	uint64 GetPage() const { return Address & ~(4llu*1024-1); }
 	int64 GetSize() const { return Size; }
-	const TCHAR* GetTag() const { return Tag; }
 	TraceServices::TagIdType GetTagId() const { return TagId; }
+	const TCHAR* GetTag() const { return Tag; }
+	const TCHAR* GetAsset() const { return Asset; }
 	const TraceServices::FCallstack* GetCallstack() const { return Callstack; }
+	const TraceServices::FCallstack* GetFreeCallstack() const { return FreeCallstack; }
 	FText GetFullCallstack() const;
 	FText GetFullCallstackSourceFiles() const;
-	HeapId GetRootHeap() const { return RootHeap; };
+	HeapId GetRootHeap() const { return RootHeap; }
 
 private:
 	uint32 StartEventIndex;
@@ -43,8 +45,9 @@ private:
 	double EndTime;
 	uint64 Address;
 	int64 Size;
-	const TCHAR* Tag;
 	TraceServices::TagIdType TagId;
+	const TCHAR* Tag;
+	const TCHAR* Asset;
 	const TraceServices::FCallstack* Callstack;
 	const TraceServices::FCallstack* FreeCallstack;
 	HeapId RootHeap;
diff --git a/Engine/Source/Developer/TraceInsights/Private/Insights/MemoryProfiler/Widgets/SMemAllocTableTreeView.cpp b/Engine/Source/Developer/TraceInsights/Private/Insights/MemoryProfiler/Widgets/SMemAllocTableTreeView.cpp
--- a/Engine/Source/Developer/TraceInsights/Private/Insights/MemoryProfiler/Widgets/SMemAllocTableTreeView.cpp
+++ b/Engine/Source/Developer/TraceInsights/Private/Insights/MemoryProfiler/Widgets/SMemAllocTableTreeView.cpp
@@ -11,7 +11,10 @@
 #include "Styling/AppStyle.h"
 #include "TraceServices/Model/AllocationsProvider.h"
 #include "TraceServices/Model/Callstack.h"
+#include "TraceServices/Model/Definitions.h"
+#include "TraceServices/Model/MetadataProvider.h"
 #include "TraceServices/Model/Modules.h"
+#include "TraceServices/Model/Strings.h"
 #include "Widgets/Images/SImage.h"
 #include "Widgets/Input/SCheckBox.h"
 #include "Widgets/SToolTip.h"
@@ -301,6 +304,22 @@ void SMemAllocTableTreeView::UpdateQuery(TraceServices::IAllocationsProvider::EQ
 
 	const TraceServices::ICallstacksProvider* CallstacksProvider = TraceServices::ReadCallstacksProvider(*Session.Get());
 
+	const TraceServices::IMetadataProvider* MetadataProvider = TraceServices::ReadMetadataProvider(*Session.Get());
+
+	uint16 AssetMetadataType = 0;
+	if (MetadataProvider)
+	{
+		TraceServices::IMetadataProvider::FReadScopeLock MetadataProviderReadLock(*MetadataProvider);
+		AssetMetadataType = MetadataProvider->GetRegisteredMetadataType(TEXT("Asset"));
+		if (MetadataProvider->GetRegisteredMetadataName(AssetMetadataType) == nullptr)
+		{
+			// If AssetMetadataType is not valid then we do not need to furhter check the Asset metadata for each allocation.
+			MetadataProvider = nullptr;
+		}
+	}
+
+	const TraceServices::IDefinitionProvider* DefinitionProvider = TraceServices::ReadDefinitionProvider(*Session.Get());
+
 	constexpr double MaxPollTime = 0.03; // Stop getting results after 30 ms so we don't tank the frame rate too much.
 	FStopwatch TotalStopwatch;
 	TotalStopwatch.Start();
@@ -363,8 +382,31 @@ void SMemAllocTableTreeView::UpdateQuery(TraceServices::IAllocationsProvider::EQ
 					Alloc.EndTime = Allocation->GetEndTime();
 					Alloc.Address = Allocation->GetAddress();
 					Alloc.Size = int64(Allocation->GetSize());
-					Alloc.Tag = Provider.GetTagName(Allocation->GetTag());
 					Alloc.TagId = Allocation->GetTag();
+					Alloc.Tag = Provider.GetTagName(Allocation->GetTag());
+					if (MetadataProvider && DefinitionProvider)
+					{
+						TraceServices::IMetadataProvider::FReadScopeLock MetadataProviderReadLock(*MetadataProvider);
+						MetadataProvider->EnumerateMetadata(Allocation->GetThreadId(), Allocation->GetMetadataId(),
+							[AssetMetadataType, &Alloc, DefinitionProvider](uint32 StackDepth, uint16 Type, const void* Data, uint32 Size) -> bool
+							{
+								if (Type == AssetMetadataType)
+								{
+									////TODO: TraceServices::IDefinitionProvider::FReadScopeLock DefinitionProviderReadLock(*DefinitionProvider);
+									//DefinitionProvider->BeginRead();
+									//const UE::Trace::TEventRef<UE::Trace::FEventRef32>& AssetNameReference = ((const UE::Trace::TEventRef<UE::Trace::FEventRef32>*)Data)[0];
+									////const UE::Trace::TEventRef<UE::Trace::FEventRef32>& ClassNameReference = ((const UE::Trace::TEventRef<UE::Trace::FEventRef32>*)Data)[1];
+									//const TraceServices::FStringDefinition* AssetNameDef = DefinitionProvider->Get<TraceServices::FStringDefinition>(AssetNameReference);
+									//if (AssetNameDef)
+									//{
+									//	Alloc.Asset = AssetNameDef->Display;
+									//}
+									//DefinitionProvider->EndRead();
+									return false;
+								}
+								return true;
+							});
+					}
 					if (CallstacksProvider)
 					{
 						Alloc.Callstack = CallstacksProvider->GetCallstack(Allocation->GetCallstackId());
@@ -768,6 +810,57 @@ void SMemAllocTableTreeView::InitAvailableViewPresets()
 	};
 	AvailableViewPresets.Add(MakeShared<FTagViewPreset>());
 
+	//////////////////////////////////////////////////
+	// Asset Breakdown View
+
+	class FAssetViewPreset : public IViewPreset
+	{
+	public:
+		virtual FText GetName() const override
+		{
+			return LOCTEXT("Asset_PresetName", "Asset");
+		}
+		virtual FText GetToolTip() const override
+		{
+			return LOCTEXT("Asset_PresetToolTip", "Asset Breakdown View\nConfigure the tree view to show a breakdown of allocations by their Asset tag.");
+		}
+		virtual FName GetSortColumn() const override
+		{
+			return FTable::GetHierarchyColumnId();
+		}
+		virtual EColumnSortMode::Type GetSortMode() const override
+		{
+			return EColumnSortMode::Type::Ascending;
+		}
+		virtual void SetCurrentGroupings(const TArray<TSharedPtr<FTreeNodeGrouping>>& InAvailableGroupings, TArray<TSharedPtr<FTreeNodeGrouping>>& InOutCurrentGroupings) const override
+		{
+			InOutCurrentGroupings.Reset();
+
+			check(InAvailableGroupings[0]->Is<FTreeNodeGroupingFlat>());
+			InOutCurrentGroupings.Add(InAvailableGroupings[0]);
+
+			const TSharedPtr<FTreeNodeGrouping>* AssetGrouping = InAvailableGroupings.FindByPredicate(
+				[](TSharedPtr<FTreeNodeGrouping>& Grouping)
+				{
+					return Grouping->Is<FTreeNodeGroupingByPathBreakdown>() &&
+						   Grouping->As<FTreeNodeGroupingByPathBreakdown>().GetColumnId() == FMemAllocTableColumns::AssetColumnId;
+				});
+			if (AssetGrouping)
+			{
+				InOutCurrentGroupings.Add(*AssetGrouping);
+			}
+		}
+		virtual void GetColumnConfigSet(TArray<FColumnConfig>& InOutConfigSet) const override
+		{
+			InOutConfigSet.Add({ FTable::GetHierarchyColumnId(),          true, 200.0f });
+			InOutConfigSet.Add({ FMemAllocTableColumns::CountColumnId,    true, 100.0f });
+			InOutConfigSet.Add({ FMemAllocTableColumns::SizeColumnId,     true, 100.0f });
+			InOutConfigSet.Add({ FMemAllocTableColumns::TagColumnId,      true, 120.0f });
+			InOutConfigSet.Add({ FMemAllocTableColumns::FunctionColumnId, true, 400.0f });
+		}
+	};
+	AvailableViewPresets.Add(MakeShared<FAssetViewPreset>());
+
 	//////////////////////////////////////////////////
 	// (Inverted) Callstack Breakdown View
 
diff --git a/Engine/Source/Developer/TraceInsights/Private/Insights/Table/ViewModels/TreeNodeGrouping.cpp b/Engine/Source/Developer/TraceInsights/Private/Insights/Table/ViewModels/TreeNodeGrouping.cpp
--- a/Engine/Source/Developer/TraceInsights/Private/Insights/Table/ViewModels/TreeNodeGrouping.cpp
+++ b/Engine/Source/Developer/TraceInsights/Private/Insights/Table/ViewModels/TreeNodeGrouping.cpp
@@ -16,6 +16,7 @@ INSIGHTS_IMPLEMENT_RTTI(FTreeNodeGroupingFlat)
 INSIGHTS_IMPLEMENT_RTTI(FTreeNodeGroupingByUniqueValue)
 INSIGHTS_IMPLEMENT_RTTI(FTreeNodeGroupingByNameFirstLetter)
 INSIGHTS_IMPLEMENT_RTTI(FTreeNodeGroupingByType)
+INSIGHTS_IMPLEMENT_RTTI(FTreeNodeGroupingByPathBreakdown)
 
 ////////////////////////////////////////////////////////////////////////////////////////////////////
 // FTreeNodeGrouping
@@ -113,7 +114,7 @@ void FTreeNodeGroupingFlat::GroupNodes(const TArray<FTableTreeNodePtr>& Nodes, F
 
 FTreeNodeGroupingByUniqueValue::FTreeNodeGroupingByUniqueValue(TSharedRef<FTableColumn> InColumnRef)
 	: FTreeNodeGrouping(
-		InColumnRef->GetTitleName(),
+		FText::Format(LOCTEXT("Grouping_ByUniqueValue_ShortNameFmt", "{0} (Unique)"), InColumnRef->GetTitleName()),
 		FText::Format(LOCTEXT("Grouping_ByUniqueValue_TitleNameFmt", "Unique Values - {0}"), InColumnRef->GetTitleName()),
 		LOCTEXT("Grouping_ByUniqueValue_Desc", "Creates a group for each unique value."),
 		TEXT("Icons.Group.TreeItem"),
@@ -180,7 +181,7 @@ FTreeNodeGroupInfo FTreeNodeGroupingByNameFirstLetter::GetGroupForNode(const FBa
 FTreeNodeGroupingByType::FTreeNodeGroupingByType()
 	: FTreeNodeGrouping(
 		LOCTEXT("Grouping_ByTypeName_ShortName", "TypeName"),
-		LOCTEXT("Grouping_ByTypeName_TitleName", "TypeName"),
+		LOCTEXT("Grouping_ByTypeName_TitleName", "Type Name"),
 		LOCTEXT("Grouping_ByTypeName_Desc", "Creates a group for each node type."),
 		TEXT("Icons.Group.TreeItem"),
 		nullptr)
@@ -194,6 +195,148 @@ FTreeNodeGroupInfo FTreeNodeGroupingByType::GetGroupForNode(const FBaseTreeNodeP
 	return { InNode->GetTypeName(), true };
 }
 
+////////////////////////////////////////////////////////////////////////////////////////////////////
+// FTreeNodeGroupingByPathBreakdown
+////////////////////////////////////////////////////////////////////////////////////////////////////
+
+FTreeNodeGroupingByPathBreakdown::FTreeNodeGroupingByPathBreakdown(TSharedRef<FTableColumn> InColumnRef)
+	: FTreeNodeGrouping(
+		FText::Format(LOCTEXT("Grouping_ByPathBreakdown_ShortNameFmt", "{0} (Path)"), InColumnRef->GetTitleName()),
+		FText::Format(LOCTEXT("Grouping_ByPathBreakdown_TitleNameFmt", "Path Breakdown - {0}"), InColumnRef->GetTitleName()),
+		LOCTEXT("Grouping_ByPathBreakdown_Desc", "Creates a tree hierarchy out of the path structure of string values."),
+		TEXT("Icons.Group.TreeItem"),
+		nullptr)
+	, ColumnRef(InColumnRef)
+{
+}
+
+////////////////////////////////////////////////////////////////////////////////////////////////////
+
+void FTreeNodeGroupingByPathBreakdown::GroupNodes(
+	const TArray<FTableTreeNodePtr>& Nodes,
+	FTableTreeNode& ParentGroup,
+	TWeakPtr<FTable> InParentTable,
+	IAsyncOperationProgress& InAsyncOperationProgress) const
+{
+	ParentGroup.ClearChildren();
+
+	struct FDirectory
+	{
+		FDirectory() = delete;
+		FDirectory(FName InName, FDirectory* InParent, FTableTreeNode* InNode)
+			: Name(InName)
+			, Parent(InParent)
+			, Node(InNode)
+		{
+		}
+
+		FName Name;
+		FDirectory* Parent;
+		FTableTreeNode* Node;
+		TMap<FName, FDirectory*> ChildrenMap; // Directory Name -> FDirectory*
+	};
+
+	TArray<FDirectory*> Directories;
+
+	static FName RootName(TEXT("~"));
+	FDirectory* Root = new FDirectory(RootName, nullptr, &ParentGroup);
+	Directories.Add(Root);
+
+	FTableTreeNode* UnsetGroupPtr = nullptr;
+	TMap<FString, FDirectory*> FullNameGroupMap; // FullName -> FDirectory*
+
+	for (FTableTreeNodePtr NodePtr : Nodes)
+	{
+		if (InAsyncOperationProgress.ShouldCancelAsyncOp())
+		{
+			return;
+		}
+
+		if (NodePtr->IsGroup())
+		{
+			ParentGroup.AddChildAndSetGroupPtr(NodePtr);
+			continue;
+		}
+
+		FString FullName;
+		const TOptional<FTableCellValue> OptionalValue = ColumnRef->GetValue(*NodePtr);
+		if (OptionalValue.IsSet())
+		{
+			FullName = OptionalValue.GetValue().AsString();
+		}
+
+		FDirectory* Directory = Root;
+
+		if (!FullName.IsEmpty())
+		{
+			FDirectory** FoundDir = FullNameGroupMap.Find(FullName);
+			if (FoundDir)
+			{
+				Directory = *FoundDir;
+			}
+			else
+			{
+				const TCHAR* Start = *FullName;
+				const TCHAR* End = Start;
+				while (true)
+				{
+					if (*End == TEXT('/') || *End == TEXT('\\') || *End == TEXT('\0'))
+					{
+						const FName DirName(int32(End - Start), Start, 0);
+
+						FoundDir = Directory->ChildrenMap.Find(DirName);
+						if (FoundDir)
+						{
+							Directory = *FoundDir;
+						}
+						else
+						{
+							FName GroupName = DirName.IsNone() ? FName(TEXT("/")) : DirName;
+							FTableTreeNodePtr GroupPtr = MakeShared<FTableTreeNode>(GroupName, InParentTable);
+							GroupPtr->SetExpansion(false);
+							Directory->Node->AddChildAndSetGroupPtr(GroupPtr);
+							Directory = new FDirectory(DirName, Directory, GroupPtr.Get());
+							Directory->Parent->ChildrenMap.Add(GroupName, Directory);
+							Directories.Add(Directory);
+						}
+
+						if (*End == TEXT('\0'))
+						{
+							FullNameGroupMap.Add(FullName, Directory);
+							break;
+						}
+
+						Start = End + 1;
+					}
+					++End;
+				}
+			}
+		}
+
+		if (Directory != Root)
+		{
+			Directory->Node->AddChildAndSetGroupPtr(NodePtr);
+		}
+		else
+		{
+			if (!UnsetGroupPtr)
+			{
+				static FName NotAvailableName(TEXT("N/A"));
+				FTableTreeNodePtr NotAvailableNode = MakeShared<FTableTreeNode>(NotAvailableName, InParentTable);
+				NotAvailableNode->SetExpansion(false);
+				ParentGroup.AddChildAndSetGroupPtr(NotAvailableNode);
+				UnsetGroupPtr = NotAvailableNode.Get();
+			}
+			UnsetGroupPtr->AddChildAndSetGroupPtr(NodePtr);
+		}
+	}
+
+	for (FDirectory* Dir : Directories)
+	{
+		delete Dir;
+	}
+}
+
 ////////////////////////////////////////////////////////////////////////////////////////////////////
 
 } // namespace Insights
diff --git a/Engine/Source/Developer/TraceInsights/Private/Insights/Table/ViewModels/TreeNodeGrouping.h b/Engine/Source/Developer/TraceInsights/Private/Insights/Table/ViewModels/TreeNodeGrouping.h
--- a/Engine/Source/Developer/TraceInsights/Private/Insights/Table/ViewModels/TreeNodeGrouping.h
+++ b/Engine/Source/Developer/TraceInsights/Private/Insights/Table/ViewModels/TreeNodeGrouping.h
@@ -229,4 +229,24 @@ public:
 
 ////////////////////////////////////////////////////////////////////////////////////////////////////
 
+/** Creates a tree hierarchy out of the path structure of string values. */
+class FTreeNodeGroupingByPathBreakdown : public FTreeNodeGrouping
+{
+	INSIGHTS_DECLARE_RTTI(FTreeNodeGroupingByPathBreakdown, FTreeNodeGrouping);
+
+public:
+	FTreeNodeGroupingByPathBreakdown(TSharedRef<FTableColumn> InColumnRef);
+	virtual ~FTreeNodeGroupingByPathBreakdown() {}
+
+	virtual void GroupNodes(const TArray<FTableTreeNodePtr>& Nodes, FTableTreeNode& ParentGroup, TWeakPtr<FTable> InParentTable, IAsyncOperationProgress& InAsyncOperationProgress) const override;
+
+	virtual FName GetColumnId() const override { return ColumnRef->GetId(); }
+	TSharedRef<FTableColumn> GetColumn() const { return ColumnRef; }
+
+private:
+	TSharedRef<FTableColumn> ColumnRef;
+};
+
+////////////////////////////////////////////////////////////////////////////////////////////////////
+
 } // namespace Insights
diff --git a/Engine/Source/Developer/TraceInsights/Private/Insights/Table/Widgets/STableTreeView.cpp b/Engine/Source/Developer/TraceInsights/Private/Insights/Table/Widgets/STableTreeView.cpp
--- a/Engine/Source/Developer/TraceInsights/Private/Insights/Table/Widgets/STableTreeView.cpp
+++ b/Engine/Source/Developer/TraceInsights/Private/Insights/Table/Widgets/STableTreeView.cpp
@@ -1725,6 +1725,7 @@ void STableTreeView::InternalCreateGroupings()
 	//AvailableGroupings.Add(MakeShared<FTreeNodeGroupingByNameFirstLetter>());
 	//AvailableGroupings.Add(MakeShared<FTreeNodeGroupingByType>());
 
+	// By Unique Value
 	for (const TSharedRef<FTableColumn>& ColumnRef : Table->GetColumns())
 	{
 		const FTableColumn& Column = ColumnRef.Get();
@@ -1757,6 +1758,30 @@ void STableTreeView::InternalCreateGroupings()
 			}
 		}
 	}
+
+	// By Path Breakdown
+	for (const TSharedRef<FTableColumn>& ColumnRef : Table->GetColumns())
+	{
+		const FTableColumn& Column = ColumnRef.Get();
+		if (!Column.IsHierarchy())
+		{
+			switch (Column.GetDataType())
+			{
+			case ETableCellDataType::Bool:
+			case ETableCellDataType::Int64:
+			case ETableCellDataType::Float:
+			case ETableCellDataType::Double:
+				break;
+
+			case ETableCellDataType::CString:
+				AvailableGroupings.Add(MakeShared<FTreeNodeGroupingByPathBreakdown>(ColumnRef));
+				break;
+
+			default:
+				AvailableGroupings.Add(MakeShared<FTreeNodeGroupingByPathBreakdown>(ColumnRef));
+			}
+		}
+	}
 }
 
 ////////////////////////////////////////////////////////////////////////////////////////////////////
diff --git a/Engine/Source/Developer/TraceServices/Private/Model/AllocationsProvider.cpp b/Engine/Source/Developer/TraceServices/Private/Model/AllocationsProvider.cpp
--- a/Engine/Source/Developer/TraceServices/Private/Model/AllocationsProvider.cpp
+++ b/Engine/Source/Developer/TraceServices/Private/Model/AllocationsProvider.cpp
@@ -379,6 +379,13 @@ uint32 IAllocationsProvider::FAllocation::GetAlignment() const
 	return Inner->GetAlignment();
 }
 
+////////////////////////////////////////////////////////////////////////////////////////////////////
+uint32 IAllocationsProvider::FAllocation::GetThreadId() const
+{
+	const auto* Inner = (const FAllocationItem*)this;
+	return Inner->ThreadId;
+}
+
 ////////////////////////////////////////////////////////////////////////////////////////////////////
 uint32 IAllocationsProvider::FAllocation::GetCallstackId() const
 {
@@ -1306,10 +1313,10 @@ void FAllocationsProvider::EditAlloc(double Time, uint32 CallstackId, uint64 Add
 		Allocation.EndEventIndex = (uint32)-1;
 		Allocation.StartTime = Time;
 		Allocation.EndTime = std::numeric_limits<double>::infinity();
+		Allocation.ThreadId = CurrentSystemThreadId;
 		Allocation.CallstackId = CallstackId;
 		Allocation.FreeCallstackId = 0; // no callstack yet
 		Allocation.MetadataId = MetadataId;
-		Allocation.Reserved = 0;
 		Allocation.Tag = Tag;
 		Allocation.RootHeap = RootHeap;
 		Allocation.Flags = EMemoryTraceHeapAllocationFlags::None;
diff --git a/Engine/Source/Developer/TraceServices/Private/Model/AllocationsProvider.h b/Engine/Source/Developer/TraceServices/Private/Model/AllocationsProvider.h
--- a/Engine/Source/Developer/TraceServices/Private/Model/AllocationsProvider.h
+++ b/Engine/Source/Developer/TraceServices/Private/Model/AllocationsProvider.h
@@ -107,10 +107,10 @@ struct FAllocationItem
 	uint32 EndEventIndex;
 	double StartTime;
 	double EndTime;
+	uint32 ThreadId;
 	uint32 CallstackId;
 	uint32 FreeCallstackId;
 	uint32 MetadataId;
-	uint32 Reserved;
 	TagIdType Tag;
 	uint8 RootHeap;
 	EMemoryTraceHeapAllocationFlags Flags;
diff --git a/Engine/Source/Developer/TraceServices/Private/Model/MetadataProvider.h b/Engine/Source/Developer/TraceServices/Private/Model/MetadataProvider.h
--- a/Engine/Source/Developer/TraceServices/Private/Model/MetadataProvider.h
+++ b/Engine/Source/Developer/TraceServices/Private/Model/MetadataProvider.h
@@ -117,12 +117,12 @@ public:
 	//////////////////////////////////////////////////
 	// Read operations
 
-	uint16 GetRegisteredMetadataType(const TCHAR* Name) const;
-	const TCHAR* GetRegisteredMetadataName(uint16 Type) const;
+	virtual uint16 GetRegisteredMetadataType(const TCHAR* Name) const override;
+	virtual const TCHAR* GetRegisteredMetadataName(uint16 Type) const override;
 
-	uint32 GetMetadataStackSize(uint32 InThreadId, uint32 InMetadataId) const;
-	bool GetMetadata(uint32 InThreadId, uint32 InMetadataId, uint32 InStackDepth, uint16& OutType, const void*& OutData, uint32& OutSize) const;
-	void EnumerateMetadata(uint32 InThreadId, uint32 InMetadataId, TFunctionRef<bool(uint32 StackDepth, uint16 Type, const void* Data, uint32 Size)> Callback) const;
+	virtual uint32 GetMetadataStackSize(uint32 InThreadId, uint32 InMetadataId) const override;
+	virtual bool GetMetadata(uint32 InThreadId, uint32 InMetadataId, uint32 InStackDepth, uint16& OutType, const void*& OutData, uint32& OutSize) const override;
+	virtual void EnumerateMetadata(uint32 InThreadId, uint32 InMetadataId, TFunctionRef<bool(uint32 StackDepth, uint16 Type, const void* Data, uint32 Size)> Callback) const override;
 
 private:
 	FMetadataThread& GetOrAddThread(uint32 ThreadId);
diff --git a/Engine/Source/Developer/TraceServices/Public/TraceServices/Model/AllocationsProvider.h b/Engine/Source/Developer/TraceServices/Public/TraceServices/Model/AllocationsProvider.h
--- a/Engine/Source/Developer/TraceServices/Public/TraceServices/Model/AllocationsProvider.h
+++ b/Engine/Source/Developer/TraceServices/Public/TraceServices/Model/AllocationsProvider.h
@@ -99,6 +99,7 @@ public:
 		uint64 GetAddress() const;
 		uint64 GetSize() const;
 		uint32 GetAlignment() const;
+		uint32 GetThreadId() const;
 		uint32 GetCallstackId() const;
 		uint32 GetFreeCallstackId() const;
 		uint32 GetMetadataId() const;
diff --git a/Engine/Source/Developer/TraceServices/Public/TraceServices/Model/MetadataProvider.h b/Engine/Source/Developer/TraceServices/Public/TraceServices/Model/MetadataProvider.h
--- a/Engine/Source/Developer/TraceServices/Public/TraceServices/Model/MetadataProvider.h
+++ b/Engine/Source/Developer/TraceServices/Public/TraceServices/Model/MetadataProvider.h
@@ -57,7 +57,12 @@ public:
 	virtual void BeginRead() const = 0;
 	virtual void EndRead() const = 0;
 
-	//...
+	virtual uint16 GetRegisteredMetadataType(const TCHAR* Name) const = 0;
+	virtual const TCHAR* GetRegisteredMetadataName(uint16 Type) const = 0;
+
+	virtual uint32 GetMetadataStackSize(uint32 InThreadId, uint32 InMetadataId) const = 0;
+	virtual bool GetMetadata(uint32 InThreadId, uint32 InMetadataId, uint32 InStackDepth, uint16& OutType, const void*& OutData, uint32& OutSize) const = 0;
+	virtual void EnumerateMetadata(uint32 InThreadId, uint32 InMetadataId, TFunctionRef<bool(uint32 StackDepth, uint16 Type, const void* Data, uint32 Size)> Callback) const = 0;
 };
 
 TRACESERVICES_API FName GetMetadataProviderName();
