[Insights] Fix crash when closing the app whilst an async operations is in progress in the Task Table.

diff --git a/Engine/Source/Developer/TraceInsights/Private/Insights/TaskGraphProfiler/TaskGraphProfilerManager.cpp b/Engine/Source/Developer/TraceInsights/Private/Insights/TaskGraphProfiler/TaskGraphProfilerManager.cpp
--- a/Engine/Source/Developer/TraceInsights/Private/Insights/TaskGraphProfiler/TaskGraphProfilerManager.cpp
+++ b/Engine/Source/Developer/TraceInsights/Private/Insights/TaskGraphProfiler/TaskGraphProfilerManager.cpp
@@ -472,6 +472,11 @@ void FTaskGraphProfilerManager::ShowTaskRelations(TaskTrace::FId TaskId)
 
 void FTaskGraphProfilerManager::OnWindowClosedEvent()
 {
+	if (TaskTableTreeView.IsValid())
+	{
+		TaskTableTreeView->OnClose();
+	}
+
 	TSharedPtr<FTabManager> TimingTabManagerSharedPtr = TimingTabManager.Pin();
 
 	if (TimingTabManagerSharedPtr.IsValid())
