Updating icon locations
#rnx

diff --git a/Engine/Source/Developer/TraceInsights/Private/Insights/InsightsStyle.cpp b/Engine/Source/Developer/TraceInsights/Private/Insights/InsightsStyle.cpp
--- a/Engine/Source/Developer/TraceInsights/Private/Insights/InsightsStyle.cpp
+++ b/Engine/Source/Developer/TraceInsights/Private/Insights/InsightsStyle.cpp
@@ -190,7 +190,7 @@ void FInsightsStyle::FStyle::Initialize()
 	Set("Icons.PluginTracksMenu.ToolBar", new IMAGE_BRUSH_SVG("PluginTracks_20", Icon20x20));
 	Set("Icons.ViewModeMenu.ToolBar", new IMAGE_BRUSH_SVG("ViewMode_20", Icon20x20));
 
-	Set("Icons.HighlightEvents.ToolBar", new EDITOR_IMAGE_BRUSH_SVG("Starship/Common/Visualizer", Icon20x20));
+	Set("Icons.HighlightEvents.ToolBar", new CORE_IMAGE_BRUSH_SVG("Starship/Common/Visualizer", Icon20x20));
 	Set("Icons.ResetHighlight.ToolBar", new CORE_IMAGE_BRUSH_SVG("Starship/Common/Reject", Icon20x20));
 
 	//////////////////////////////////////////////////
diff --git a/Engine/Source/Editor/EditorStyle/Private/StarshipStyle.cpp b/Engine/Source/Editor/EditorStyle/Private/StarshipStyle.cpp
--- a/Engine/Source/Editor/EditorStyle/Private/StarshipStyle.cpp
+++ b/Engine/Source/Editor/EditorStyle/Private/StarshipStyle.cpp
@@ -4023,7 +4023,7 @@ void FStarshipEditorStyle::FStyle::SetupLevelEditorStyle()
 
 		Set( "EditorViewport.VisualizeGBufferMode",   new IMAGE_BRUSH_SVG("Starship/Common/BufferVisualization", Icon16x16) );
 
-		Set( "EditorViewport.Visualizers", 			  new IMAGE_BRUSH_SVG("Starship/Common/Visualizer", Icon16x16) );
+		Set( "EditorViewport.Visualizers", 			  new CORE_IMAGE_BRUSH_SVG("Starship/Common/Visualizer", Icon16x16) );
 		Set( "EditorViewport.LOD", 			  		  new IMAGE_BRUSH_SVG("Starship/Common/LOD", Icon16x16) );
 
 		Set( "EditorViewport.ReflectionOverrideMode", new IMAGE_BRUSH_SVG("Starship/Common/Reflections", Icon16x16 ) );
