[Insights] Fixed uninitialized FMemoryAlloc class members when created by a mem query.

diff --git a/Engine/Source/Developer/TraceInsights/Private/Insights/MemoryProfiler/Widgets/SMemAllocTableTreeView.cpp b/Engine/Source/Developer/TraceInsights/Private/Insights/MemoryProfiler/Widgets/SMemAllocTableTreeView.cpp
--- a/Engine/Source/Developer/TraceInsights/Private/Insights/MemoryProfiler/Widgets/SMemAllocTableTreeView.cpp
+++ b/Engine/Source/Developer/TraceInsights/Private/Insights/MemoryProfiler/Widgets/SMemAllocTableTreeView.cpp
@@ -372,18 +372,25 @@ void SMemAllocTableTreeView::UpdateQuery(TraceServices::IAllocationsProvider::EQ
 
 				uint64 AllocsDestIndex = Allocs.Num();
 				Allocs.AddUninitialized(AllocCount);
+
 				for (uint32 AllocIndex = 0; AllocIndex < AllocCount; ++AllocIndex, ++AllocsDestIndex)
 				{
 					const TraceServices::IAllocationsProvider::FAllocation* Allocation = Result->Get(AllocIndex);
 					FMemoryAlloc& Alloc = Allocs[AllocsDestIndex];
+
 					Alloc.StartEventIndex = Allocation->GetStartEventIndex();
 					Alloc.EndEventIndex = Allocation->GetEndEventIndex();
+
 					Alloc.StartTime = Allocation->GetStartTime();
 					Alloc.EndTime = Allocation->GetEndTime();
+
 					Alloc.Address = Allocation->GetAddress();
 					Alloc.Size = int64(Allocation->GetSize());
+
 					Alloc.TagId = Allocation->GetTag();
 					Alloc.Tag = Provider.GetTagName(Allocation->GetTag());
+
+					Alloc.Asset = nullptr;
 					if (MetadataProvider && DefinitionProvider)
 					{
 						TraceServices::IMetadataProvider::FReadScopeLock MetadataProviderReadLock(*MetadataProvider);
@@ -407,15 +414,25 @@ void SMemAllocTableTreeView::UpdateQuery(TraceServices::IAllocationsProvider::EQ
 								return true;
 							});
 					}
+
 					if (CallstacksProvider)
 					{
 						Alloc.Callstack = CallstacksProvider->GetCallstack(Allocation->GetCallstackId());
+						check(Alloc.Callstack != nullptr);
+
 						Alloc.FreeCallstack = CallstacksProvider->GetCallstack(Allocation->GetFreeCallstackId());
+						check(Alloc.FreeCallstack != nullptr);
 					}
+					else
+					{
+						Alloc.Callstack = nullptr;
+						Alloc.FreeCallstack = nullptr;
+					}
+
 					Alloc.RootHeap = Allocation->GetRootHeap();
 					Alloc.bIsBlock = Allocation->IsHeap();
-					check(Alloc.Callstack != nullptr);
 
+					Alloc.bIsDecline = false;
 					if (Rule->GetValue() == TraceServices::IAllocationsProvider::EQueryRule::aAfaBf)
 					{
 						if (Alloc.StartTime <= TimeMarkers[0] && Alloc.EndTime <= TimeMarkers[1]) // decline
