Fix comparison of nullptr values.

diff --git a/Engine/Source/Developer/TraceInsights/Private/Insights/Table/ViewModels/TableCellValueSorter.cpp b/Engine/Source/Developer/TraceInsights/Private/Insights/Table/ViewModels/TableCellValueSorter.cpp
--- a/Engine/Source/Developer/TraceInsights/Private/Insights/Table/ViewModels/TableCellValueSorter.cpp
+++ b/Engine/Source/Developer/TraceInsights/Private/Insights/Table/ViewModels/TableCellValueSorter.cpp
@@ -495,6 +495,16 @@ FSorterByCStringValue::FSorterByCStringValue(TSharedRef<FTableColumn> InColumnRe
 		const TCHAR* ValueA = OptionalValueA.IsSet() ? OptionalValueA.GetValue().CString : TEXT("");
 		const TCHAR* ValueB = OptionalValueB.IsSet() ? OptionalValueB.GetValue().CString : TEXT("");
 
+		// If any value is nullptr
+		if (!ValueA || !ValueB)
+		{
+			if(!ValueA && !ValueB)
+			{
+				INSIGHTS_DEFAULT_SORTING_NODES(A, B)
+			}
+			return ValueA < ValueB;
+		}
+			
 		const int32 Compare = FCString::Strcmp(ValueA, ValueB);
 		if (Compare == 0)
 		{
@@ -516,6 +526,16 @@ FSorterByCStringValue::FSorterByCStringValue(TSharedRef<FTableColumn> InColumnRe
 
 		const TCHAR* ValueA = OptionalValueA.IsSet() ? OptionalValueA.GetValue().CString : TEXT("");
 		const TCHAR* ValueB = OptionalValueB.IsSet() ? OptionalValueB.GetValue().CString : TEXT("");
+		
+		// If any value is nullptr
+		if (!ValueA || !ValueB)
+		{
+			if(!ValueA && !ValueB)
+			{
+				INSIGHTS_DEFAULT_SORTING_NODES(A, B)
+			}
+			return ValueA > ValueB;
+		}
 
 		const int32 Compare = FCString::Strcmp(ValueA, ValueB);
 		if (Compare == 0)
