[Insights]
  - Memory Insights / Mem Allocs: Added "Tag" column to the "4K Page Breakdown" view preset.
  - Memory Insights / Mem Allocs: Fixed bug where changing a preset results in hidding the "Hierarchy" column.
  - Memory Insights / Mem Allocs: Removed "Unique Values" grouping for StartEventIndex, EndEventIndex and Count columns.
  - Memory Insights / Mem Allocs: Removed "Path Breakdown" grouping for Top Function column.
  - Tree Views: Fixed the "Unique Values" grouping for CString columns crashing when encountering nullptr values (FTreeNodeGroupingByUniqueValueCString).
  - Tree Views: Fixed the "Path Breakdown" grouping adding incorrectly "/" group nodes (FTreeNodeGroupingByPathBreakdown).

diff --git a/Engine/Source/Developer/TraceInsights/Private/Insights/MemoryProfiler/Widgets/SMemAllocTableTreeView.cpp b/Engine/Source/Developer/TraceInsights/Private/Insights/MemoryProfiler/Widgets/SMemAllocTableTreeView.cpp
--- a/Engine/Source/Developer/TraceInsights/Private/Insights/MemoryProfiler/Widgets/SMemAllocTableTreeView.cpp
+++ b/Engine/Source/Developer/TraceInsights/Private/Insights/MemoryProfiler/Widgets/SMemAllocTableTreeView.cpp
@@ -944,7 +944,7 @@ void SMemAllocTableTreeView::InitAvailableViewPresets()
 	AvailableViewPresets.Add(MakeShared<FCallstackViewPreset>(true));
 
 	//////////////////////////////////////////////////
-	// Memory Page Breakdown View
+	// Address (4K Page) Breakdown View
 
 	class FPageViewPreset : public IViewPreset
 	{
@@ -989,6 +989,7 @@ void SMemAllocTableTreeView::InitAvailableViewPresets()
 			InOutConfigSet.Add({ FMemAllocTableColumns::AddressColumnId,  true, 120.0f });
 			InOutConfigSet.Add({ FMemAllocTableColumns::CountColumnId,    true, 100.0f });
 			InOutConfigSet.Add({ FMemAllocTableColumns::SizeColumnId,     true, 100.0f });
+			InOutConfigSet.Add({ FMemAllocTableColumns::TagColumnId,      true, 120.0f });
 			InOutConfigSet.Add({ FMemAllocTableColumns::FunctionColumnId, true, 400.0f });
 		}
 	};
@@ -1177,36 +1178,48 @@ void SMemAllocTableTreeView::InternalCreateGroupings()
 {
 	Insights::STableTreeView::InternalCreateGroupings();
 
+	AvailableGroupings.RemoveAll(
+		[](TSharedPtr<FTreeNodeGrouping>& Grouping)
+		{
+			if (Grouping->Is<FTreeNodeGroupingByUniqueValue>())
+			{
+				const FName ColumnId = Grouping->As<FTreeNodeGroupingByUniqueValue>().GetColumnId();
+				if (ColumnId == FMemAllocTableColumns::StartEventIndexColumnId ||
+					ColumnId == FMemAllocTableColumns::EndEventIndexColumnId ||
+					ColumnId == FMemAllocTableColumns::CountColumnId)
+				{
+					return true;
+				}
+			}
+			else if (Grouping->Is<FTreeNodeGroupingByPathBreakdown>())
+			{
+				const FName ColumnId = Grouping->As<FTreeNodeGroupingByPathBreakdown>().GetColumnId();
+				if (ColumnId == FMemAllocTableColumns::FunctionColumnId)
+				{
+					return true;
+				}
+
+			}
+			return false;
+		});
+
 	int32 Index = 1; // after the Flat ("All") grouping
 
 	AvailableGroupings.Insert(MakeShared<FMemAllocGroupingBySize>(), Index++);
 
-	TSharedPtr<FTreeNodeGrouping>* TagGroupingPtr = AvailableGroupings.FindByPredicate(
-		[](TSharedPtr<FTreeNodeGrouping>& Grouping)
-		{
-			return Grouping->Is<FTreeNodeGroupingByUniqueValue>() &&
-				   Grouping->As<FTreeNodeGroupingByUniqueValue>().GetColumnId() == FMemAllocTableColumns::TagColumnId;
-		});
-	if (TagGroupingPtr)
+	const TraceServices::IAllocationsProvider* AllocationsProvider = Session.IsValid() ? TraceServices::ReadAllocationsProvider(*Session.Get()) : nullptr;
+
+	if (AllocationsProvider)
 	{
-		TSharedPtr<FTreeNodeGroupingByUniqueValue> TagGrouping = StaticCastSharedPtr<FTreeNodeGroupingByUniqueValue>(*TagGroupingPtr);
-		AvailableGroupings.Remove(TagGrouping);
-		//TODO: TagGrouping->SetShortName(LOCTEXT("Grouping_ByTag_ShortName", "LLM Tag"));
-		//TODO: TagGrouping->SetTitleName(LOCTEXT("Grouping_ByTag_TitleName", "By LLM Tag"));
-		AvailableGroupings.Insert(TagGrouping, Index++);
+		AvailableGroupings.Insert(MakeShared<FMemAllocGroupingByTag>(*AllocationsProvider), Index++);
 	}
 
 	AvailableGroupings.Insert(MakeShared<FMemAllocGroupingByCallstack>(false, bIsCallstackGroupingByFunction), Index++);
 	AvailableGroupings.Insert(MakeShared<FMemAllocGroupingByCallstack>(true, bIsCallstackGroupingByFunction), Index++);
 
-	if (Session.IsValid())
+	if (AllocationsProvider)
 	{
-		const TraceServices::IAllocationsProvider* AllocationsProvider = TraceServices::ReadAllocationsProvider(*Session.Get());
-		if (AllocationsProvider)
-		{
-			AvailableGroupings.Insert(MakeShared<FMemAllocGroupingByHeap>(*AllocationsProvider), Index++);
-			AvailableGroupings.Insert(MakeShared<FMemAllocGroupingByTag>(*AllocationsProvider), Index++);
-		}
+		AvailableGroupings.Insert(MakeShared<FMemAllocGroupingByHeap>(*AllocationsProvider), Index++);
 	}
 }
 
diff --git a/Engine/Source/Developer/TraceInsights/Private/Insights/MemoryProfiler/Widgets/SMemAllocTableTreeView.h b/Engine/Source/Developer/TraceInsights/Private/Insights/MemoryProfiler/Widgets/SMemAllocTableTreeView.h
--- a/Engine/Source/Developer/TraceInsights/Private/Insights/MemoryProfiler/Widgets/SMemAllocTableTreeView.h
+++ b/Engine/Source/Developer/TraceInsights/Private/Insights/MemoryProfiler/Widgets/SMemAllocTableTreeView.h
@@ -24,7 +24,7 @@ class SMemAllocTableTreeView : public STableTreeView
 private:
 	struct FColumnConfig
 	{
-		const FName& ColumnId;
+		FName ColumnId;
 		bool bIsVisible;
 		float Width;
 	};
diff --git a/Engine/Source/Developer/TraceInsights/Private/Insights/Table/ViewModels/TreeNodeGrouping.cpp b/Engine/Source/Developer/TraceInsights/Private/Insights/Table/ViewModels/TreeNodeGrouping.cpp
--- a/Engine/Source/Developer/TraceInsights/Private/Insights/Table/ViewModels/TreeNodeGrouping.cpp
+++ b/Engine/Source/Developer/TraceInsights/Private/Insights/Table/ViewModels/TreeNodeGrouping.cpp
@@ -145,12 +145,84 @@ template<> bool TTreeNodeGroupingByUniqueValue<bool>::GetValue(const FTableCellV
 template<> int64 TTreeNodeGroupingByUniqueValue<int64>::GetValue(const FTableCellValue& CellValue) { return CellValue.Int64; }
 template<> float TTreeNodeGroupingByUniqueValue<float>::GetValue(const FTableCellValue& CellValue) { return CellValue.Float; }
 template<> double TTreeNodeGroupingByUniqueValue<double>::GetValue(const FTableCellValue& CellValue) { return CellValue.Double; }
-template<> const TCHAR* TTreeNodeGroupingByUniqueValue<const TCHAR*>::GetValue(const FTableCellValue& CellValue) { return CellValue.CString; }
 
-template<>
-FText TTreeNodeGroupingByUniqueValue<const TCHAR*>::GetValueAsText(const FTableColumn& Column, const FTableTreeNode& Node)
+////////////////////////////////////////////////////////////////////////////////////////////////////
+// FTreeNodeGroupingByUniqueValueCString
+////////////////////////////////////////////////////////////////////////////////////////////////////
+
+FName FTreeNodeGroupingByUniqueValueCString::GetGroupName(const TCHAR* Value)
+{
+	if (Value == nullptr || *Value == TEXT('\0'))
+	{
+		static FName EmptyGroupName(TEXT("N/A"));
+		return EmptyGroupName;
+	}
+
+	FStringView StringView(Value);
+	if (StringView.Len() >= NAME_SIZE)
+	{
+		StringView = FStringView(StringView.GetData(), NAME_SIZE - 1);
+	}
+	return FName(StringView, 0);
+}
+
+////////////////////////////////////////////////////////////////////////////////////////////////////
+
+void FTreeNodeGroupingByUniqueValueCString::GroupNodes(const TArray<FTableTreeNodePtr>& Nodes, FTableTreeNode& ParentGroup, TWeakPtr<FTable> InParentTable, IAsyncOperationProgress& InAsyncOperationProgress) const
 {
-	return Column.GetValueAsText(Node);
+	TMap<const void*, FTableTreeNodePtr> GroupMap; // using void* as TMap<const TCHAR*, T> doesn't allow nullptr to be added
+	FTableTreeNodePtr UnsetGroupPtr = nullptr;
+
+	ParentGroup.ClearChildren();
+
+	for (FTableTreeNodePtr NodePtr : Nodes)
+	{
+		if (InAsyncOperationProgress.ShouldCancelAsyncOp())
+		{
+			return;
+		}
+
+		if (NodePtr->IsGroup())
+		{
+			ParentGroup.AddChildAndSetGroupPtr(NodePtr);
+			continue;
+		}
+
+		FTableTreeNodePtr GroupPtr = nullptr;
+
+		const FTableColumn& Column = *GetColumn();
+		const TOptional<FTableCellValue> CellValue = Column.GetValue(*NodePtr);
+		if (CellValue.IsSet())
+		{
+			const TCHAR* Value = CellValue.GetValue().CString;
+
+			FTableTreeNodePtr* GroupPtrPtr = GroupMap.Find(Value);
+			if (!GroupPtrPtr)
+			{
+				const FName GroupName = GetGroupName(Value);
+				GroupPtr = MakeShared<FTableTreeNode>(GroupName, InParentTable);
+				GroupPtr->SetExpansion(false);
+				ParentGroup.AddChildAndSetGroupPtr(GroupPtr);
+				GroupMap.Add(Value, GroupPtr);
+			}
+			else
+			{
+				GroupPtr = *GroupPtrPtr;
+			}
+		}
+		else
+		{
+			if (!UnsetGroupPtr)
+			{
+				UnsetGroupPtr = MakeShared<FTableTreeNode>(FName(TEXT("<unset>")), InParentTable);
+				UnsetGroupPtr->SetExpansion(false);
+				ParentGroup.AddChildAndSetGroupPtr(UnsetGroupPtr);
+			}
+			GroupPtr = UnsetGroupPtr;
+		}
+
+		GroupPtr->AddChildAndSetGroupPtr(NodePtr);
+	}
 }
 
 ////////////////////////////////////////////////////////////////////////////////////////////////////
@@ -289,14 +361,13 @@ void FTreeNodeGroupingByPathBreakdown::GroupNodes(
 						{
 							Directory = *FoundDir;
 						}
-						else
+						else if (!DirName.IsNone())
 						{
-							FName GroupName = DirName.IsNone() ? FName(TEXT("/")) : DirName;
-							FTableTreeNodePtr GroupPtr = MakeShared<FTableTreeNode>(GroupName, InParentTable);
+							FTableTreeNodePtr GroupPtr = MakeShared<FTableTreeNode>(DirName, InParentTable);
 							GroupPtr->SetExpansion(false);
 							Directory->Node->AddChildAndSetGroupPtr(GroupPtr);
 							Directory = new FDirectory(DirName, Directory, GroupPtr.Get());
-							Directory->Parent->ChildrenMap.Add(GroupName, Directory);
+							Directory->Parent->ChildrenMap.Add(DirName, Directory);
 							Directories.Add(Directory);
 						}
 
diff --git a/Engine/Source/Developer/TraceInsights/Private/Insights/Table/ViewModels/TreeNodeGrouping.h b/Engine/Source/Developer/TraceInsights/Private/Insights/Table/ViewModels/TreeNodeGrouping.h
--- a/Engine/Source/Developer/TraceInsights/Private/Insights/Table/ViewModels/TreeNodeGrouping.h
+++ b/Engine/Source/Developer/TraceInsights/Private/Insights/Table/ViewModels/TreeNodeGrouping.h
@@ -124,11 +124,27 @@ public:
 
 private:
 	static Type GetValue(const FTableCellValue& CellValue);
-	static FText GetValueAsText(const FTableColumn& Column, const FTableTreeNode& Node)
+	static FName GetGroupName(const FTableColumn& Column, const FTableTreeNode& Node);
+};
+
+template<typename Type>
+FName TTreeNodeGroupingByUniqueValue<Type>::GetGroupName(const FTableColumn& Column, const FTableTreeNode& Node)
+{
+	FText ValueAsText = Column.GetValueAsGroupingText(Node);
+
+	if (ValueAsText.IsEmpty())
 	{
-		return Column.GetValueAsGroupingText(Node);
+		static FName EmptyGroupName(TEXT("N/A"));
+		return EmptyGroupName;
 	}
-};
+
+	FStringView StringView(ValueAsText.ToString());
+	if (StringView.Len() >= NAME_SIZE)
+	{
+		StringView = FStringView(StringView.GetData(), NAME_SIZE - 1);
+	}
+	return FName(StringView, 0);
+}
 
 template<typename Type>
 void TTreeNodeGroupingByUniqueValue<Type>::GroupNodes(const TArray<FTableTreeNodePtr>& Nodes, FTableTreeNode& ParentGroup, TWeakPtr<FTable> InParentTable, IAsyncOperationProgress& InAsyncOperationProgress) const
@@ -162,13 +178,8 @@ void TTreeNodeGroupingByUniqueValue<Type>::GroupNodes(const TArray<FTableTreeNod
 			FTableTreeNodePtr* GroupPtrPtr = GroupMap.Find(Value);
 			if (!GroupPtrPtr)
 			{
-				FText ValueAsText = GetValueAsText(Column, *NodePtr);
-				FStringView GroupName(ValueAsText.ToString());
-				if (GroupName.Len() >= NAME_SIZE)
-				{
-					GroupName = FStringView(GroupName.GetData(), NAME_SIZE - 1);
-				}
-				GroupPtr = MakeShared<FTableTreeNode>(FName(GroupName, 0), InParentTable);
+				const FName GroupName = GetGroupName(Column, *NodePtr);
+				GroupPtr = MakeShared<FTableTreeNode>(GroupName, InParentTable);
 				GroupPtr->SetExpansion(false);
 				ParentGroup.AddChildAndSetGroupPtr(GroupPtr);
 				GroupMap.Add(Value, GroupPtr);
@@ -197,7 +208,21 @@ typedef TTreeNodeGroupingByUniqueValue<bool> FTreeNodeGroupingByUniqueValueBool;
 typedef TTreeNodeGroupingByUniqueValue<int64> FTreeNodeGroupingByUniqueValueInt64;
 typedef TTreeNodeGroupingByUniqueValue<float> FTreeNodeGroupingByUniqueValueFloat;
 typedef TTreeNodeGroupingByUniqueValue<double> FTreeNodeGroupingByUniqueValueDouble;
-typedef TTreeNodeGroupingByUniqueValue<const TCHAR*> FTreeNodeGroupingByUniqueValueCString;
+
+////////////////////////////////////////////////////////////////////////////////////////////////////
+
+/** Creates a group for each unique value (assumes data type of cell values is const TCHAR*). */
+class FTreeNodeGroupingByUniqueValueCString : public FTreeNodeGroupingByUniqueValue
+{
+public:
+	FTreeNodeGroupingByUniqueValueCString(TSharedRef<FTableColumn> InColumnRef) : FTreeNodeGroupingByUniqueValue(InColumnRef) {}
+	virtual ~FTreeNodeGroupingByUniqueValueCString() {}
+
+	virtual void GroupNodes(const TArray<FTableTreeNodePtr>& Nodes, FTableTreeNode& ParentGroup, TWeakPtr<FTable> InParentTable, IAsyncOperationProgress& InAsyncOperationProgress) const override;
+
+private:
+	static FName GetGroupName(const TCHAR* Value);
+};
 
 ////////////////////////////////////////////////////////////////////////////////////////////////////
 
