[Insights] Mem Insights / LLM Tags: Fixed full name of llm tags to include parent names also when there are multiple levels of inheritance.

diff --git a/Engine/Source/Developer/TraceInsights/Private/Insights/MemoryProfiler/ViewModels/MemoryTag.cpp b/Engine/Source/Developer/TraceInsights/Private/Insights/MemoryProfiler/ViewModels/MemoryTag.cpp
--- a/Engine/Source/Developer/TraceInsights/Private/Insights/MemoryProfiler/ViewModels/MemoryTag.cpp
+++ b/Engine/Source/Developer/TraceInsights/Private/Insights/MemoryProfiler/ViewModels/MemoryTag.cpp
@@ -200,23 +200,35 @@ void FMemoryTagList::UpdateInternal()
 		// Update Parent and StatFullName for each tag.
 		for (FMemoryTag* TagPtr : Tags)
 		{
-			FMemoryTag& Tag = *TagPtr;
-			if (Tag.ParentId != FMemoryTag::InvalidTagId && Tag.Parent == nullptr)
-			{
-				for (FMemoryTag* ParentTagPtr : Tags)
-				{
-					FMemoryTag& ParentTag = *ParentTagPtr;
-					if (ParentTag.Id == Tag.ParentId)
-					{
-						ParentTag.Children.Add(TagPtr);
+			UpdateParentAndStatFullName(*TagPtr);
+		}
+	}
+}
 
-						Tag.Parent = ParentTagPtr;
-						Tag.StatFullName = ParentTag.StatName + TEXT("/") + Tag.StatName;
+////////////////////////////////////////////////////////////////////////////////////////////////////
 
-						break;
-					}
-				}
-			}
+void FMemoryTagList::UpdateParentAndStatFullName(FMemoryTag& Tag)
+{
+	if (Tag.ParentId == FMemoryTag::InvalidTagId || Tag.Parent != nullptr)
+	{
+		// Tag has no Parent or the Parent and the StatFullName are already set.
+		return;
+	}
+
+	for (FMemoryTag* ParentTagPtr : Tags)
+	{
+		FMemoryTag& ParentTag = *ParentTagPtr;
+		if (ParentTag.Id == Tag.ParentId)
+		{
+			// Ensure StatFullName for parent tag is updated first.
+			UpdateParentAndStatFullName(ParentTag);
+
+			ParentTag.Children.Add(&Tag);
+
+			Tag.Parent = ParentTagPtr;
+			Tag.StatFullName = ParentTag.StatFullName + TEXT("/") + Tag.StatName;
+
+			break;
 		}
 	}
 }
diff --git a/Engine/Source/Developer/TraceInsights/Private/Insights/MemoryProfiler/ViewModels/MemoryTag.h b/Engine/Source/Developer/TraceInsights/Private/Insights/MemoryProfiler/ViewModels/MemoryTag.h
--- a/Engine/Source/Developer/TraceInsights/Private/Insights/MemoryProfiler/ViewModels/MemoryTag.h
+++ b/Engine/Source/Developer/TraceInsights/Private/Insights/MemoryProfiler/ViewModels/MemoryTag.h
@@ -89,6 +89,7 @@ public:
 
 private:
 	void UpdateInternal();
+	void UpdateParentAndStatFullName(FMemoryTag& Tag);
 
 private:
 	TArray<FMemoryTag*> Tags; // the list of memory tags; owns the allocated memory
