TasksInsights: visualisation for parent tasks

diff --git a/Engine/Source/Developer/TraceInsights/Private/Insights/InsightsStyle.cpp b/Engine/Source/Developer/TraceInsights/Private/Insights/InsightsStyle.cpp
--- a/Engine/Source/Developer/TraceInsights/Private/Insights/InsightsStyle.cpp
+++ b/Engine/Source/Developer/TraceInsights/Private/Insights/InsightsStyle.cpp
@@ -240,6 +240,7 @@ void FInsightsStyle::FStyle::Initialize()
 	Set("Icons.ShowTaskConnections", new CORE_IMAGE_BRUSH("Icons/Profiler/profiler_Calls_32x", Icon16x16));
 	Set("Icons.ShowTaskPrerequisites", new CORE_IMAGE_BRUSH("Icons/Profiler/profiler_Calls_32x", Icon16x16));
 	Set("Icons.ShowTaskSubsequents", new CORE_IMAGE_BRUSH("Icons/Profiler/profiler_Calls_32x", Icon16x16));
+	Set("Icons.ShowParentTasks", new CORE_IMAGE_BRUSH("Icons/Profiler/profiler_Calls_32x", Icon16x16));
 	Set("Icons.ShowNestedTasks", new CORE_IMAGE_BRUSH("Icons/Profiler/profiler_Calls_32x", Icon16x16));
 	Set("Icons.ShowTaskTrack", new CORE_IMAGE_BRUSH("Icons/Profiler/profiler_Calls_32x", Icon16x16));
 	Set("Icons.ShowDetailedTaskTrackInfo", new CORE_IMAGE_BRUSH("Icons/Profiler/profiler_Calls_32x", Icon16x16));
diff --git a/Engine/Source/Developer/TraceInsights/Private/Insights/TaskGraphProfiler/TaskGraphProfilerManager.cpp b/Engine/Source/Developer/TraceInsights/Private/Insights/TaskGraphProfiler/TaskGraphProfilerManager.cpp
--- a/Engine/Source/Developer/TraceInsights/Private/Insights/TaskGraphProfiler/TaskGraphProfilerManager.cpp
+++ b/Engine/Source/Developer/TraceInsights/Private/Insights/TaskGraphProfiler/TaskGraphProfilerManager.cpp
@@ -280,9 +280,9 @@ void FTaskGraphProfilerManager::ShowTaskRelations(const TraceServices::FTaskInfo
 		}
 	};
 
-	if (bShowPrerequisites)
+	if (bShowCriticalPath)
 	{
-		GetSingleTaskRelationsForAll(Task->Prerequisites);
+		GetRelationsOnCriticalPath(Task);
 	}
 	if (bShowTransitions)
 	{
@@ -292,17 +292,21 @@ void FTaskGraphProfilerManager::ShowTaskRelations(const TraceServices::FTaskInfo
 	{
 		GetSingleTaskConnections(Task, TasksProvider, InSelectedEvent);
 	}
-	if (bShowNestedTasks)
+	if (bShowPrerequisites)
 	{
-		GetSingleTaskRelationsForAll(Task->NestedTasks);
+		GetSingleTaskRelationsForAll(Task->Prerequisites);
 	}
 	if (bShowSubsequents)
 	{
 		GetSingleTaskRelationsForAll(Task->Subsequents);
 	}
-	if (bShowCriticalPath)
+	if (bShowParentTasks)
 	{
-		GetRelationsOnCriticalPath(Task);
+		GetSingleTaskRelationsForAll(Task->ParentTasks);
+	}
+	if (bShowNestedTasks)
+	{
+		GetSingleTaskRelationsForAll(Task->NestedTasks);
 	}
 }
 
@@ -347,6 +351,24 @@ void FTaskGraphProfilerManager::GetSingleTaskConnections(const TraceServices::FT
 		AddRelation(InSelectedEvent, Prerequisite->FinishedTimestamp, Prerequisite->StartedThreadId, Task->StartedTimestamp, Task->StartedThreadId, ETaskEventType::PrerequisiteStarted);
 	}
 
+	int32 NumSubsequentsToShow = FMath::Min(Task->Subsequents.Num(), MaxTasksToShow);
+	for (int32 i = 0; i != NumSubsequentsToShow; ++i)
+	{
+		const TraceServices::FTaskInfo* Subsequent = TasksProvider->TryGetTask(Task->Subsequents[i].RelativeId);
+		check(Subsequent != nullptr);
+
+		AddRelation(InSelectedEvent, Task->CompletedTimestamp, Task->CompletedThreadId, Subsequent->StartedTimestamp, Subsequent->StartedThreadId, ETaskEventType::SubsequentStarted);
+	}
+
+	int32 NumParentsToShow = FMath::Min(Task->ParentTasks.Num(), MaxTasksToShow);
+	for (int32 i = 0; i != NumParentsToShow; ++i)
+	{
+		const TraceServices::FTaskInfo* ParentTask = TasksProvider->TryGetTask(Task->ParentTasks[i].RelativeId);
+		check(ParentTask != nullptr);
+
+		AddRelation(InSelectedEvent, Task->CompletedTimestamp, Task->CompletedThreadId, ParentTask->CompletedTimestamp, ParentTask->CompletedThreadId, ETaskEventType::ParentStarted);
+	}
+
 	int32 NumNestedToShow = FMath::Min(Task->NestedTasks.Num(), MaxTasksToShow);
 	for (int32 i = 0; i != NumNestedToShow; ++i)
 	{
@@ -354,18 +376,10 @@ void FTaskGraphProfilerManager::GetSingleTaskConnections(const TraceServices::FT
 		const TraceServices::FTaskInfo* NestedTask = TasksProvider->TryGetTask(RelationInfo.RelativeId);
 		check(NestedTask != nullptr);
 
-		int32 NestedExecutionStartedDepth = GetDepthOfTaskExecution(Task->StartedTimestamp, Task->FinishedTimestamp, Task->StartedThreadId);
+		int32 NestedExecutionStartedDepth = GetDepthOfTaskExecution(NestedTask->StartedTimestamp, NestedTask->FinishedTimestamp, NestedTask->StartedThreadId);
 
 		AddRelation(InSelectedEvent, RelationInfo.Timestamp, Task->StartedThreadId, -1, NestedTask->StartedTimestamp, NestedTask->StartedThreadId, NestedExecutionStartedDepth, ETaskEventType::NestedStarted);
-	}
-
-	int32 NumSubsequentsToShow = FMath::Min(Task->Subsequents.Num(), MaxTasksToShow);
-	for (int32 i = 0; i != NumSubsequentsToShow; ++i)
-	{
-		const TraceServices::FTaskInfo* Subsequent = TasksProvider->TryGetTask(Task->Subsequents[i].RelativeId);
-		check(Subsequent != nullptr);
-
-		AddRelation(InSelectedEvent, Task->FinishedTimestamp, Task->StartedThreadId, Subsequent->StartedTimestamp, Subsequent->StartedThreadId, ETaskEventType::SubsequentStarted);
+		AddRelation(InSelectedEvent, NestedTask->FinishedTimestamp, NestedTask->StartedThreadId, NestedExecutionStartedDepth, Task->CompletedTimestamp, Task->CompletedThreadId, -1, ETaskEventType::NestedCompleted);
 	}
 }
 
@@ -509,6 +523,7 @@ void FTaskGraphProfilerManager::InitializeColorCode()
 	ColorCode[static_cast<uint32>(ETaskEventType::Finished)] = ToLiniarColorNoAlpha(0xFFDC1AFF); // Yellow
 	ColorCode[static_cast<uint32>(ETaskEventType::Completed)] = ToLiniarColorNoAlpha(0xFE9B07FF); // Orange
 	ColorCode[static_cast<uint32>(ETaskEventType::PrerequisiteStarted)] = ToLiniarColorNoAlpha(0xFF729CFF); // Pink
+	ColorCode[static_cast<uint32>(ETaskEventType::ParentStarted)] = ToLiniarColorNoAlpha(0xB68F55FF); // Folder
 	ColorCode[static_cast<uint32>(ETaskEventType::NestedStarted)] = ToLiniarColorNoAlpha(0xB68F55FF); // Folder
 	ColorCode[static_cast<uint32>(ETaskEventType::SubsequentStarted)] = ToLiniarColorNoAlpha(0xFE9B07FF); // Orange
 	ColorCode[static_cast<uint32>(ETaskEventType::NestedCompleted)] = ToLiniarColorNoAlpha(0x804D39FF); // Brown
@@ -775,33 +790,6 @@ double FTaskGraphProfilerManager::GetRelationsOnCriticalPathDescendingRec(const
 		}
 	}
 
-	if (Task->NestedTasks.Num() > 0)
-	{
-		TArray<FTaskGraphRelation> DescendingRelations;
-
-		for (const TraceServices::FTaskInfo::FRelationInfo& NestedTaskRelation : Task->NestedTasks)
-		{
-			const TraceServices::FTaskInfo* NestedTaskInfo = TasksProvider->TryGetTask(NestedTaskRelation.RelativeId);
-			if (NestedTaskInfo != nullptr)
-			{
-				double ChainDuration = GetRelationsOnCriticalPathDescendingRec(NestedTaskInfo, TasksProvider, DescendingRelations);
-				if (ChainDuration > MaxChainDuration)
-				{
-					MaxChainDuration = ChainDuration;
-					NextTaskInChain = NestedTaskInfo;
-					// Remove the relations from the shorter branch.
-					if (Relations.Num() > InitialRelationNum)
-					{
-						Relations.RemoveAt(InitialRelationNum, Relations.Num() - InitialRelationNum, false);
-					}
-
-					Relations.Append(DescendingRelations);
-				}
-				DescendingRelations.Empty(DescendingRelations.Max());
-			}
-		}
-	}
-
 	if (NextTaskInChain)
 	{
 		int32 SourceDepth = GetDepthOfTaskExecution(Task->StartedTimestamp, Task->FinishedTimestamp, Task->StartedThreadId);
diff --git a/Engine/Source/Developer/TraceInsights/Private/Insights/TaskGraphProfiler/TaskGraphProfilerManager.h b/Engine/Source/Developer/TraceInsights/Private/Insights/TaskGraphProfiler/TaskGraphProfilerManager.h
--- a/Engine/Source/Developer/TraceInsights/Private/Insights/TaskGraphProfiler/TaskGraphProfilerManager.h
+++ b/Engine/Source/Developer/TraceInsights/Private/Insights/TaskGraphProfiler/TaskGraphProfilerManager.h
@@ -44,6 +44,7 @@ enum class ETaskEventType : uint32
 	Finished,
 	Completed,
 	PrerequisiteStarted,
+	ParentStarted,
 	NestedStarted,
 	NestedCompleted,
 	SubsequentStarted,
@@ -107,6 +108,9 @@ public:
 
 	TSharedPtr<Insights::FTaskTimingSharedState> GetTaskTimingSharedState() { return TaskTimingSharedState;	}
 
+	bool GetShowCriticalPath() const { return bShowCriticalPath; }
+	void SetShowCriticalPath(bool bInValue) { bShowCriticalPath = bInValue; }
+
 	bool GetShowTransitions() const { return bShowTransitions; }
 	void SetShowTransitions(bool bInValue) { bShowTransitions = bInValue; }
 
@@ -119,13 +123,13 @@ public:
 	bool GetShowSubsequents() const { return bShowSubsequents; }
 	void SetShowSubsequents(bool bInValue) { bShowSubsequents = bInValue; }
 
+	bool GetShowParentTasks() const { return bShowParentTasks; }
+	void SetShowParentTasks(bool bInValue) { bShowParentTasks = bInValue; }
+
 	bool GetShowNestedTasks() const { return bShowNestedTasks; }
 	void SetShowNestedTasks(bool bInValue) { bShowNestedTasks = bInValue; }
 	
-	bool GetShowCriticalPath() const { return bShowCriticalPath; }
-	void SetShowCriticalPath(bool bInValue) { bShowCriticalPath = bInValue; }
-
-	bool GetShowAnyRelations() const { return GetShowTransitions() || GetShowConnections() || GetShowPrerequisites() || GetShowSubsequents() || GetShowNestedTasks() || GetShowCriticalPath(); }
+	bool GetShowAnyRelations() const { return GetShowCriticalPath() || GetShowTransitions() || GetShowConnections() || GetShowPrerequisites() || GetShowSubsequents() || GetShowParentTasks() || GetShowNestedTasks(); }
 
 	int32 GetDepthOfTaskExecution(double TaskStartedTime, double TaskFinishedTime, uint32 ThreadId);
 
@@ -167,12 +171,13 @@ private:
 
 	TSharedPtr<Insights::STaskTableTreeView> TaskTableTreeView;
 	FLinearColor ColorCode[static_cast<uint32>(ETaskEventType::NumTaskEventTypes)];
+	bool bShowCriticalPath = false;
 	bool bShowTransitions = true;
 	bool bShowConnections = true;
 	bool bShowPrerequisites = false;
 	bool bShowSubsequents = false;
+	bool bShowParentTasks = false;
 	bool bShowNestedTasks = false;
-	bool bShowCriticalPath = false;
 };
 
 } // namespace Insights
diff --git a/Engine/Source/Developer/TraceInsights/Private/Insights/TaskGraphProfiler/ViewModels/TaskEntry.cpp b/Engine/Source/Developer/TraceInsights/Private/Insights/TaskGraphProfiler/ViewModels/TaskEntry.cpp
--- a/Engine/Source/Developer/TraceInsights/Private/Insights/TaskGraphProfiler/ViewModels/TaskEntry.cpp
+++ b/Engine/Source/Developer/TraceInsights/Private/Insights/TaskGraphProfiler/ViewModels/TaskEntry.cpp
@@ -26,9 +26,10 @@ FTaskEntry::FTaskEntry(const TraceServices::FTaskInfo& TaskInfo)
 	, DestroyedTimestamp(TaskInfo.DestroyedTimestamp)
 	, DestroyedThreadId(TaskInfo.DestroyedThreadId)
 {
-	NumNested = TaskInfo.NestedTasks.Num();
 	NumSubsequents = TaskInfo.Subsequents.Num();
 	NumPrerequisites = TaskInfo.Prerequisites.Num();
+	NumParents = TaskInfo.ParentTasks.Num();
+	NumNested = TaskInfo.NestedTasks.Num();
 }
 
 ////////////////////////////////////////////////////////////////////////////////////////////////////
diff --git a/Engine/Source/Developer/TraceInsights/Private/Insights/TaskGraphProfiler/ViewModels/TaskEntry.h b/Engine/Source/Developer/TraceInsights/Private/Insights/TaskGraphProfiler/ViewModels/TaskEntry.h
--- a/Engine/Source/Developer/TraceInsights/Private/Insights/TaskGraphProfiler/ViewModels/TaskEntry.h
+++ b/Engine/Source/Developer/TraceInsights/Private/Insights/TaskGraphProfiler/ViewModels/TaskEntry.h
@@ -42,6 +42,7 @@ public:
 	double GetDestroyedTimestamp() const { return DestroyedTimestamp; }
 	uint32 GetDestroyedThreadId() const { return DestroyedThreadId; }
 
+	uint32 GetNumParents() const { return NumParents; }
 	uint32 GetNumNested() const { return NumNested; }
 	uint32 GetNumPrerequisites() const { return NumPrerequisites; }
 	uint32 GetNumSubsequents() const { return NumSubsequents; }
@@ -75,9 +76,10 @@ private:
 	double DestroyedTimestamp;
 	uint32 DestroyedThreadId;
 
-	uint32 NumNested;
 	uint32 NumSubsequents;
 	uint32 NumPrerequisites;
+	uint32 NumParents;
+	uint32 NumNested;
 };
 
 ////////////////////////////////////////////////////////////////////////////////////////////////////
diff --git a/Engine/Source/Developer/TraceInsights/Private/Insights/TaskGraphProfiler/ViewModels/TaskTable.cpp b/Engine/Source/Developer/TraceInsights/Private/Insights/TaskGraphProfiler/ViewModels/TaskTable.cpp
--- a/Engine/Source/Developer/TraceInsights/Private/Insights/TaskGraphProfiler/ViewModels/TaskTable.cpp
+++ b/Engine/Source/Developer/TraceInsights/Private/Insights/TaskGraphProfiler/ViewModels/TaskTable.cpp
@@ -32,6 +32,7 @@ const FName FTaskTableColumns::CompletedTimestampColumnId(TEXT("CompletedTimesta
 const FName FTaskTableColumns::CompletedThreadIdColumnId(TEXT("CompletedThreadId"));
 const FName FTaskTableColumns::DestroyedTimestampColumnId(TEXT("DestroyedTimestamp"));
 const FName FTaskTableColumns::DestroyedThreadIdColumnId(TEXT("DestroyedThreadId"));
+const FName FTaskTableColumns::NumParentColumnId(TEXT("NumParents"));
 const FName FTaskTableColumns::NumNestedColumnId(TEXT("NumNested"));
 const FName FTaskTableColumns::NumSubsequentsColumnId(TEXT("NumSubsequents"));
 const FName FTaskTableColumns::NumPrerequisitesColumnId(TEXT("NumPrerequisites"));
@@ -93,6 +94,7 @@ struct DefaultTaskFieldGetterFuncts
 	static FTableCellValue GetDestroyedTimestamp(const FTableColumn& Column, const FTaskEntry& Task) { return FTableCellValue(Task.GetDestroyedTimestamp()); }
 	static FTableCellValue GetDestroyedThreadId(const FTableColumn& Column, const FTaskEntry& Task) { return FTableCellValue((int64)Task.GetDestroyedThreadId()); }
 
+	static FTableCellValue GetNumParents(const FTableColumn& Column, const FTaskEntry& Task) { return FTableCellValue((int64)Task.GetNumParents()); }
 	static FTableCellValue GetNumNested(const FTableColumn& Column, const FTaskEntry& Task) { return FTableCellValue((int64)Task.GetNumNested());	}
 	static FTableCellValue GetNumSubsequents(const FTableColumn& Column, const FTaskEntry& Task) { return FTableCellValue((int64)Task.GetNumSubsequents()); }
 	static FTableCellValue GetNumPrerequisites(const FTableColumn& Column, const FTaskEntry& Task) { return FTableCellValue((int64)Task.GetNumPrerequisites());	}
@@ -201,7 +203,7 @@ void FTaskTable::AddDefaultColumns()
 		Column.SetTitleName(LOCTEXT("DebugNameColumnTitle", "DebugName"));
 		Column.SetDescription(LOCTEXT("DebugNameColumnDesc", "DebugName"));
 
-		Column.SetFlags(ETableColumnFlags::CanBeHidden | ETableColumnFlags::CanBeFiltered);
+		Column.SetFlags(ETableColumnFlags::ShouldBeVisible | ETableColumnFlags::CanBeHidden | ETableColumnFlags::CanBeFiltered);
 
 		Column.SetHorizontalAlignment(HAlign_Left);
 		Column.SetInitialWidth(200.0f);
@@ -629,6 +631,36 @@ void FTaskTable::AddDefaultColumns()
 		AddColumn(ColumnRef);
 	}
 	//////////////////////////////////////////////////
+	// Num Parent Tasks Column
+	{
+		TSharedRef<FTableColumn> ColumnRef = MakeShared<FTableColumn>(FTaskTableColumns::NumParentColumnId);
+		FTableColumn& Column = *ColumnRef;
+
+		Column.SetIndex(ColumnIndex++);
+
+		Column.SetShortName(LOCTEXT("ParentTasksColumnName", "Parent Tasks"));
+		Column.SetTitleName(LOCTEXT("ParentTasksColumnTitle", "Parent Tasks"));
+		Column.SetDescription(LOCTEXT("ParentTasksColumnDesc", "The number of parent tasks."));
+
+		Column.SetFlags(ETableColumnFlags::ShouldBeVisible | ETableColumnFlags::CanBeHidden | ETableColumnFlags::CanBeFiltered);
+
+		Column.SetHorizontalAlignment(HAlign_Left);
+		Column.SetInitialWidth(100.0f);
+
+		Column.SetDataType(ETableCellDataType::Int64);
+
+		TSharedRef<ITableCellValueGetter> Getter = MakeShared<FTaskColumnValueGetter<DefaultTaskFieldGetterFuncts::GetNumParents>>();
+		Column.SetValueGetter(Getter);
+
+		TSharedRef<ITableCellValueFormatter> Formatter = MakeShared<FInt64ValueFormatterAsNumber>();
+		Column.SetValueFormatter(Formatter);
+
+		TSharedRef<ITableCellValueSorter> Sorter = MakeShared<FSorterByInt64Value>(ColumnRef);
+		Column.SetValueSorter(Sorter);
+
+		AddColumn(ColumnRef);
+	}
+	//////////////////////////////////////////////////
 	// Num Nested Tasks Column
 	{
 		TSharedRef<FTableColumn> ColumnRef = MakeShared<FTableColumn>(FTaskTableColumns::NumNestedColumnId);
diff --git a/Engine/Source/Developer/TraceInsights/Private/Insights/TaskGraphProfiler/ViewModels/TaskTable.h b/Engine/Source/Developer/TraceInsights/Private/Insights/TaskGraphProfiler/ViewModels/TaskTable.h
--- a/Engine/Source/Developer/TraceInsights/Private/Insights/TaskGraphProfiler/ViewModels/TaskTable.h
+++ b/Engine/Source/Developer/TraceInsights/Private/Insights/TaskGraphProfiler/ViewModels/TaskTable.h
@@ -32,6 +32,7 @@ struct FTaskTableColumns
 	static const FName CompletedThreadIdColumnId;
 	static const FName DestroyedTimestampColumnId;
 	static const FName DestroyedThreadIdColumnId;
+	static const FName NumParentColumnId;
 	static const FName NumNestedColumnId;
 	static const FName NumSubsequentsColumnId;
 	static const FName NumPrerequisitesColumnId;
diff --git a/Engine/Source/Developer/TraceInsights/Private/Insights/TaskGraphProfiler/ViewModels/TaskTimingTrack.cpp b/Engine/Source/Developer/TraceInsights/Private/Insights/TaskGraphProfiler/ViewModels/TaskTimingTrack.cpp
--- a/Engine/Source/Developer/TraceInsights/Private/Insights/TaskGraphProfiler/ViewModels/TaskTimingTrack.cpp
+++ b/Engine/Source/Developer/TraceInsights/Private/Insights/TaskGraphProfiler/ViewModels/TaskTimingTrack.cpp
@@ -74,6 +74,12 @@ public:
 				   EUserInterfaceActionType::ToggleButton,
 				   FInputChord(EKeys::S));
 
+		UI_COMMAND(Command_ShowParentTasks,
+				   "Show Transitions of Parent Tasks",
+				   "Show/hide stage transitions for the current task's parent tasks.",
+				   EUserInterfaceActionType::ToggleButton,
+				   FInputChord(EKeys::R));
+
 		UI_COMMAND(Command_ShowNestedTasks,
 				   "Show Transitions of Nested Tasks",
 				   "Show/hide stage transitions for the current task's nested tasks.",
@@ -104,6 +110,7 @@ public:
 	TSharedPtr<FUICommandInfo> Command_ShowTaskConnections;
 	TSharedPtr<FUICommandInfo> Command_ShowTaskPrerequisites;
 	TSharedPtr<FUICommandInfo> Command_ShowTaskSubsequents;
+	TSharedPtr<FUICommandInfo> Command_ShowParentTasks;
 	TSharedPtr<FUICommandInfo> Command_ShowNestedTasks;
 	TSharedPtr<FUICommandInfo> Command_ShowCriticalPath;
 	TSharedPtr<FUICommandInfo> Command_ShowTaskTrack;
@@ -287,6 +294,15 @@ void FTaskTimingSharedState::BuildTasksSubMenu(FMenuBuilder& MenuBuilder)
 		FSlateIcon(FInsightsStyle::GetStyleSetName(), "Icons.ShowTaskSubsequents")
 	);
 
+	MenuBuilder.AddMenuEntry
+	(
+		FTaskTimingStateCommands::Get().Command_ShowParentTasks,
+		NAME_None,
+		TAttribute<FText>(),
+		TAttribute<FText>(),
+		FSlateIcon(FInsightsStyle::GetStyleSetName(), "Icons.ShowParentTasks")
+	);
+
 	MenuBuilder.AddMenuEntry
 	(
 		FTaskTimingStateCommands::Get().Command_ShowNestedTasks,
@@ -350,6 +366,12 @@ void FTaskTimingSharedState::InitCommandList()
 		FCanExecuteAction::CreateSP(this, &FTaskTimingSharedState::ContextMenu_ShowTaskSubsequents_CanExecute),
 		FIsActionChecked::CreateSP(this, &FTaskTimingSharedState::ContextMenu_ShowTaskSubsequents_IsChecked));
 
+	CommandList->MapAction(
+		FTaskTimingStateCommands::Get().Command_ShowParentTasks,
+		FExecuteAction::CreateSP(this, &FTaskTimingSharedState::ContextMenu_ShowParentTasks_Execute),
+		FCanExecuteAction::CreateSP(this, &FTaskTimingSharedState::ContextMenu_ShowParentTasks_CanExecute),
+		FIsActionChecked::CreateSP(this, &FTaskTimingSharedState::ContextMenu_ShowParentTasks_IsChecked));
+
 	CommandList->MapAction(
 		FTaskTimingStateCommands::Get().Command_ShowNestedTasks,
 		FExecuteAction::CreateSP(this, &FTaskTimingSharedState::ContextMenu_ShowNestedTasks_Execute),
@@ -487,6 +509,34 @@ void FTaskTimingSharedState::ContextMenu_ShowTaskSubsequents_Execute()
 	}
 }
 
+////////////////////////////////////////////////////////////////////////////////////////////////////
+// ParentTasks
+
+bool FTaskTimingSharedState::ContextMenu_ShowParentTasks_CanExecute()
+{
+	return Insights::FTaskGraphProfilerManager::Get().IsValid() && Insights::FTaskGraphProfilerManager::Get()->GetIsAvailable();
+}
+
+////////////////////////////////////////////////////////////////////////////////////////////////////
+
+bool FTaskTimingSharedState::ContextMenu_ShowParentTasks_IsChecked()
+{
+	TSharedPtr<Insights::FTaskGraphProfilerManager> TaskGraphManager = Insights::FTaskGraphProfilerManager::Get();
+	return TaskGraphManager.IsValid() && TaskGraphManager->GetIsAvailable() && TaskGraphManager->GetShowParentTasks();
+}
+
+////////////////////////////////////////////////////////////////////////////////////////////////////
+
+void FTaskTimingSharedState::ContextMenu_ShowParentTasks_Execute()
+{
+	TSharedPtr<Insights::FTaskGraphProfilerManager> TaskGraphManager = Insights::FTaskGraphProfilerManager::Get();
+	if (TaskGraphManager.IsValid() && TaskGraphManager->GetIsAvailable())
+	{
+		TaskGraphManager->SetShowParentTasks(!TaskGraphManager->GetShowParentTasks());
+		OnTaskSettingsChanged();
+	}
+}
+
 ////////////////////////////////////////////////////////////////////////////////////////////////////
 // NestedTasks
 
@@ -689,6 +739,20 @@ void FTaskTimingTrack::BuildDrawState(ITimingEventsTrackDrawStateBuilder& Builde
 			}
 		}
 
+		// parent tasks can overlap with prerequisites, hence they are shown under them (`Depth` is not reset)
+		for (TraceServices::FTaskInfo::FRelationInfo Relation : Task->ParentTasks)
+		{
+			if (Depth >= FTimingProfilerManager::Get()->GetEventDepthLimit() - 1)
+			{
+				break;
+			}
+			const TraceServices::FTaskInfo* TaskInfo = TasksProvider->TryGetTask(Relation.RelativeId);
+			if (TaskInfo)
+			{
+				Builder.AddEvent(TaskInfo->StartedTimestamp, TaskInfo->FinishedTimestamp, ++Depth, *FString::Printf(TEXT("Parent Task %d Executing"), TaskInfo->Id), 0, FTaskGraphProfilerManager::Get()->GetColorForTaskEventAsPackedARGB(ETaskEventType::ParentStarted));
+			}
+		}
+
 		Depth = 0;
 		for (TraceServices::FTaskInfo::FRelationInfo Relation : Task->NestedTasks)
 		{
@@ -877,13 +941,13 @@ const TSharedPtr<const ITimingEvent> FTaskTimingTrack::GetEvent(float InPosX, fl
 		}
 		else if(bShowDetailInfoOnTaskTrack)
 		{
-			auto GetEventFromRelations = [&TasksProvider, this, EventTime, SecondsPerPixel](const TArray<TraceServices::FTaskInfo::FRelationInfo>& Relations, int32 Depth, ETaskEventType EventType) -> TSharedPtr<FTaskTrackEvent>
+			auto GetEventFromRelations = [&TasksProvider, this, EventTime, Depth, SecondsPerPixel](const TArray<TraceServices::FTaskInfo::FRelationInfo>& Relations, int32 RelationIndex, ETaskEventType EventType) -> TSharedPtr<FTaskTrackEvent>
 			{
-				if (Relations.Num() < Depth)
+				if (RelationIndex < 0 || Relations.Num() <= RelationIndex)
 				{
 					return nullptr;
 				}
-				const TraceServices::FTaskInfo::FRelationInfo& Relation = Relations[Depth - 1];
+				const TraceServices::FTaskInfo::FRelationInfo& Relation = Relations[RelationIndex];
 				{
 					const TraceServices::FTaskInfo* TaskInfo = TasksProvider->TryGetTask(Relation.RelativeId);
 					if (TaskInfo && EventTime >= TaskInfo->StartedTimestamp - 2 * SecondsPerPixel && EventTime <= TaskInfo->FinishedTimestamp + 2 * SecondsPerPixel)
@@ -897,16 +961,21 @@ const TSharedPtr<const ITimingEvent> FTaskTimingTrack::GetEvent(float InPosX, fl
 				return nullptr;
 			};
 
-			TimingEvent = GetEventFromRelations(Task->Prerequisites, Depth, ETaskEventType::PrerequisiteStarted);
+			TimingEvent = GetEventFromRelations(Task->Prerequisites, Depth - 1, ETaskEventType::PrerequisiteStarted);
 
 			if (!TimingEvent.IsValid())
 			{
-				TimingEvent = GetEventFromRelations(Task->NestedTasks, Depth, ETaskEventType::NestedStarted);
+				TimingEvent = GetEventFromRelations(Task->ParentTasks, Depth - 1 - Task->Prerequisites.Num(), ETaskEventType::ParentStarted);
 			}
 
 			if (!TimingEvent.IsValid())
 			{
-				TimingEvent = GetEventFromRelations(Task->Subsequents, Depth, ETaskEventType::SubsequentStarted);
+				TimingEvent = GetEventFromRelations(Task->NestedTasks, Depth - 1, ETaskEventType::NestedStarted);
+			}
+
+			if (!TimingEvent.IsValid())
+			{
+				TimingEvent = GetEventFromRelations(Task->Subsequents, Depth - 1, ETaskEventType::SubsequentStarted);
 			}
 		}
 	}
@@ -961,6 +1030,7 @@ void FTaskTimingTrack::InitTooltip(FTooltipDrawState& InOutTooltip, const ITimin
 	case ETaskEventType::Created:
 		break;
 	case ETaskEventType::PrerequisiteStarted:
+	case ETaskEventType::ParentStarted:
 	case ETaskEventType::NestedStarted:
 	case ETaskEventType::SubsequentStarted:
 	{
@@ -978,9 +1048,10 @@ void FTaskTimingTrack::InitTooltip(FTooltipDrawState& InOutTooltip, const ITimin
 		InOutTooltip.AddNameValueTextLine(TEXT("Nested tasks:"), FString::Printf(TEXT("%d"), Task->NestedTasks.Num()));
 		break;
 	case ETaskEventType::Finished:
-		InOutTooltip.AddNameValueTextLine(TEXT("Subsequent tasks:"), FString::Printf(TEXT("%d"), Task->Subsequents.Num()));
 		break;
 	case ETaskEventType::Completed:
+		InOutTooltip.AddNameValueTextLine(TEXT("Parent tasks:"), FString::Printf(TEXT("%d"), Task->ParentTasks.Num()));
+		InOutTooltip.AddNameValueTextLine(TEXT("Subsequent tasks:"), FString::Printf(TEXT("%d"), Task->Subsequents.Num()));
 		break;
 	default:
 		checkf(false, TEXT("Unknown task event type"));
diff --git a/Engine/Source/Developer/TraceInsights/Private/Insights/TaskGraphProfiler/ViewModels/TaskTimingTrack.h b/Engine/Source/Developer/TraceInsights/Private/Insights/TaskGraphProfiler/ViewModels/TaskTimingTrack.h
--- a/Engine/Source/Developer/TraceInsights/Private/Insights/TaskGraphProfiler/ViewModels/TaskTimingTrack.h
+++ b/Engine/Source/Developer/TraceInsights/Private/Insights/TaskGraphProfiler/ViewModels/TaskTimingTrack.h
@@ -74,6 +74,10 @@ private:
 	bool ContextMenu_ShowTaskSubsequents_CanExecute();
 	bool ContextMenu_ShowTaskSubsequents_IsChecked();
 
+	void ContextMenu_ShowParentTasks_Execute();
+	bool ContextMenu_ShowParentTasks_CanExecute();
+	bool ContextMenu_ShowParentTasks_IsChecked();
+
 	void ContextMenu_ShowNestedTasks_Execute();
 	bool ContextMenu_ShowNestedTasks_CanExecute();
 	bool ContextMenu_ShowNestedTasks_IsChecked();
diff --git a/Engine/Source/Developer/TraceInsights/Private/Insights/TaskGraphProfiler/ViewModels/TaskTrackEvent.cpp b/Engine/Source/Developer/TraceInsights/Private/Insights/TaskGraphProfiler/ViewModels/TaskTrackEvent.cpp
--- a/Engine/Source/Developer/TraceInsights/Private/Insights/TaskGraphProfiler/ViewModels/TaskTrackEvent.cpp
+++ b/Engine/Source/Developer/TraceInsights/Private/Insights/TaskGraphProfiler/ViewModels/TaskTrackEvent.cpp
@@ -30,6 +30,7 @@ FString FTaskTrackEvent::GetStartLabel() const
 		return TEXT("Scheduled Time:");
 	case ETaskEventType::Started:
 	case ETaskEventType::PrerequisiteStarted:
+	case ETaskEventType::ParentStarted:
 	case ETaskEventType::NestedStarted:
 	case ETaskEventType::SubsequentStarted:
 		return TEXT("Started Time:");
@@ -59,6 +60,7 @@ FString FTaskTrackEvent::GetEndLabel() const
 		return TEXT("Started Time:");
 	case ETaskEventType::Started:
 	case ETaskEventType::PrerequisiteStarted:
+	case ETaskEventType::ParentStarted:
 	case ETaskEventType::NestedStarted:
 	case ETaskEventType::SubsequentStarted:
 		return TEXT("Finished Time:");
@@ -94,6 +96,8 @@ FString FTaskTrackEvent::GetEventName() const
 		return TEXT("Completed");
 	case ETaskEventType::PrerequisiteStarted:
 		return FString::Printf(TEXT("Prerequisite Task %d Executing"), GetTaskId());	
+	case ETaskEventType::ParentStarted:
+		return FString::Printf(TEXT("Parent Task %d Executing"), GetTaskId());
 	case ETaskEventType::NestedStarted:
 		return FString::Printf(TEXT("Nested Task %d Executing"), GetTaskId());
 	case ETaskEventType::SubsequentStarted:
diff --git a/Engine/Source/Developer/TraceInsights/Private/Insights/ViewModels/ThreadTimingTrack.cpp b/Engine/Source/Developer/TraceInsights/Private/Insights/ViewModels/ThreadTimingTrack.cpp
--- a/Engine/Source/Developer/TraceInsights/Private/Insights/ViewModels/ThreadTimingTrack.cpp
+++ b/Engine/Source/Developer/TraceInsights/Private/Insights/ViewModels/ThreadTimingTrack.cpp
@@ -1226,6 +1226,7 @@ void FThreadTimingTrack::InitTooltip(FTooltipDrawState& InOutTooltip, const ITim
 				}
 				InOutTooltip.AddNameValueTextLine(TEXT("Prerequisite tasks:"), FString::Printf(TEXT("%d"), Task.Prerequisites.Num()));
 				InOutTooltip.AddNameValueTextLine(TEXT("Subsequent tasks:"), FString::Printf(TEXT("%d"), Task.Subsequents.Num()));
+				InOutTooltip.AddNameValueTextLine(TEXT("Parent tasks:"), FString::Printf(TEXT("%d"), Task.ParentTasks.Num()));
 				InOutTooltip.AddNameValueTextLine(TEXT("Nested tasks:"), FString::Printf(TEXT("%d"), Task.NestedTasks.Num()));
 			};
 
