[Insights] Trivial refactoring in preparation for a larger change re TPagedArrayIterator.

diff --git a/Engine/Source/Developer/TraceInsights/Private/Insights/Tests/InsightsTestRunner.h b/Engine/Source/Developer/TraceInsights/Private/Insights/Tests/InsightsTestRunner.h
--- a/Engine/Source/Developer/TraceInsights/Private/Insights/Tests/InsightsTestRunner.h
+++ b/Engine/Source/Developer/TraceInsights/Private/Insights/Tests/InsightsTestRunner.h
@@ -15,21 +15,17 @@ DECLARE_LOG_CATEGORY_EXTERN(InsightsTestRunner, Log, All);
 class TRACEINSIGHTS_API FInsightsTestRunner : public TSharedFromThis<FInsightsTestRunner>, public IInsightsComponent
 {
 public:
-	void ScheduleCommand(const FString& InCmd);
 	virtual ~FInsightsTestRunner();
+
+	void ScheduleCommand(const FString& InCmd);
+
 	virtual void Initialize(IUnrealInsightsModule& InsightsModule) override;
 	virtual void Shutdown() override;
 	virtual void RegisterMajorTabs(IUnrealInsightsModule& InsightsModule) override;
 	virtual void UnregisterMajorTabs() override;
 
 	static TSharedPtr<FInsightsTestRunner> CreateInstance();
-
 	static TSharedPtr<FInsightsTestRunner> Get();
-	/** The delegate to be invoked when this manager ticks. */
-	FTickerDelegate OnTick;
-
-	/** Handle to the registered OnTick. */
-	FTSTicker::FDelegateHandle OnTickHandle;
 
 	bool Tick(float DeltaTime);
 
@@ -45,6 +41,13 @@ private:
 	void OnSessionAnalysisCompleted();
 	TSharedRef<SDockTab> SpawnAutomationWindowTab(const FSpawnTabArgs& Args);
 
+private:
+	/** The delegate to be invoked when this manager ticks. */
+	FTickerDelegate OnTick;
+
+	/** Handle to the registered OnTick. */
+	FTSTicker::FDelegateHandle OnTickHandle;
+
 	FDelegateHandle SessionAnalysisCompletedHandle;
 
 	FString CommandToExecute;
diff --git a/Engine/Source/Developer/TraceInsights/Private/Insights/ViewModels/Filters.h b/Engine/Source/Developer/TraceInsights/Private/Insights/ViewModels/Filters.h
--- a/Engine/Source/Developer/TraceInsights/Private/Insights/ViewModels/Filters.h
+++ b/Engine/Source/Developer/TraceInsights/Private/Insights/ViewModels/Filters.h
@@ -37,7 +37,7 @@ enum class EFilterOperator : uint8
 	Lt = 2, // Less Than
 	Lte = 3, // Less than or equal to
 	Gt = 4, // Greater than
-	Gte = 5, // Greater than or equal to 
+	Gte = 5, // Greater than or equal to
 	Contains = 6,
 	NotContains = 7,
 };
@@ -58,20 +58,23 @@ template<typename T>
 class FFilterOperator : public IFilterOperator
 {
 public:
-
 	typedef TFunction<bool(T, T)> OperatorFunc;
 
 	FFilterOperator(EFilterOperator InKey, FString InName, OperatorFunc InFunc)
 		: Func(InFunc)
 		, Key(InKey)
 		, Name(InName)
-	{}
-	virtual ~FFilterOperator() {};
+	{
+	}
+	virtual ~FFilterOperator()
+	{
+	}
 
 	virtual EFilterOperator GetKey() override { return Key; }
 	virtual FString GetName() override { return Name; };
 
 	OperatorFunc Func;
+
 private:
 	EFilterOperator Key;
 	FString Name;
@@ -93,7 +96,8 @@ struct FFilterGroupOperator
 		: Type(InType)
 		, Name(InName)
 		, Desc(InDesc)
-	{}
+	{
+	}
 
 	EFilterGroupOperator Type;
 	FText Name;
@@ -126,10 +130,13 @@ public:
 		, Desc(InDesc)
 		, DataType(InDataType)
 		, SupportedOperators(InSupportedOperators)
-	{}
-	virtual ~FFilter() {}
+	{
+	}
+	virtual ~FFilter()
+	{
+	}
 
-	SupportedOperatorsArrayPtr GetSupportedOperators() { return SupportedOperators;	}
+	SupportedOperatorsArrayPtr GetSupportedOperators() const { return SupportedOperators; }
 
 	int32 Key;
 	FText Name;
@@ -150,9 +157,12 @@ struct FFilterWithSuggestions : FFilter
 public:
 	FFilterWithSuggestions(int32 InKey, FText InName, FText InDesc, EFilterDataType InDataType, SupportedOperatorsArrayPtr InSupportedOperators)
 		: FFilter(InKey, InName, InDesc, InDataType, InSupportedOperators)
-	{ }
+	{
+	}
 
-	virtual ~FFilterWithSuggestions() {}
+	virtual ~FFilterWithSuggestions()
+	{
+	}
 
 	GetSuggestionsCallback Callback;
 };
@@ -188,11 +198,11 @@ public:
 		FilterGroupOperators.Add(MakeShared<FFilterGroupOperator>(EFilterGroupOperator::Or, LOCTEXT("AnyOf", "Any Of (OR)"), LOCTEXT("AnyOfDesc", "Any of the children must be true for the group to return true. Equivalent to an OR operation.")));
 	}
 
-	const TArray<TSharedPtr<FFilterGroupOperator>>& GetFilterGroupOperators() { return FilterGroupOperators; }
+	const TArray<TSharedPtr<FFilterGroupOperator>>& GetFilterGroupOperators() const { return FilterGroupOperators; }
 
-	TSharedPtr<TArray<TSharedPtr<IFilterOperator>>> GetDoubleOperators() { return DoubleOperators; }
-	TSharedPtr<TArray<TSharedPtr<IFilterOperator>>> GetIntegerOperators() { return IntegerOperators; }
-	TSharedPtr<TArray<TSharedPtr<IFilterOperator>>> GetStringOperators() { return StringOperators; }
+	TSharedPtr<TArray<TSharedPtr<IFilterOperator>>> GetDoubleOperators() const { return DoubleOperators; }
+	TSharedPtr<TArray<TSharedPtr<IFilterOperator>>> GetIntegerOperators() const { return IntegerOperators; }
+	TSharedPtr<TArray<TSharedPtr<IFilterOperator>>> GetStringOperators() const { return StringOperators; }
 
 private:
 	TArray<TSharedPtr<FFilterGroupOperator>> FilterGroupOperators;
@@ -216,11 +226,11 @@ public:
 	void RegisterTabSpawner();
 	TSharedRef<SDockTab> SpawnTab(const FSpawnTabArgs& Args);
 
-	const TArray<TSharedPtr<FFilterGroupOperator>>& GetFilterGroupOperators() { return FilterStorage.GetFilterGroupOperators(); }
+	const TArray<TSharedPtr<FFilterGroupOperator>>& GetFilterGroupOperators() const { return FilterStorage.GetFilterGroupOperators(); }
 
-	TSharedPtr<TArray<TSharedPtr<IFilterOperator>>> GetDoubleOperators() { return FilterStorage.GetDoubleOperators(); }
-	TSharedPtr<TArray<TSharedPtr<IFilterOperator>>> GetIntegerOperators() { return FilterStorage.GetIntegerOperators(); }
-	TSharedPtr<TArray<TSharedPtr<IFilterOperator>>> GetStringOperators() { return FilterStorage.GetStringOperators(); }
+	TSharedPtr<TArray<TSharedPtr<IFilterOperator>>> GetDoubleOperators() const { return FilterStorage.GetDoubleOperators(); }
+	TSharedPtr<TArray<TSharedPtr<IFilterOperator>>> GetIntegerOperators() const { return FilterStorage.GetIntegerOperators(); }
+	TSharedPtr<TArray<TSharedPtr<IFilterOperator>>> GetStringOperators() const { return FilterStorage.GetStringOperators(); }
 
 	TSharedPtr<SWidget> CreateFilterConfiguratorWidget(TSharedPtr<class FFilterConfigurator> FilterConfiguratorViewModel);
 
@@ -269,7 +279,7 @@ public:
 		return DataMap.Contains(Key);
 	}
 
-	bool GetReturnValueForUnsetFilters() const { return bReturnValueForUnsetFilters;	}
+	bool GetReturnValueForUnsetFilters() const { return bReturnValueForUnsetFilters; }
 	void SetReturnValueForUnsetFilters(bool InValue) { bReturnValueForUnsetFilters = InValue; }
 
 private:
diff --git a/Engine/Source/Developer/TraceServices/Public/Common/PagedArray.h b/Engine/Source/Developer/TraceServices/Public/Common/PagedArray.h
--- a/Engine/Source/Developer/TraceServices/Public/Common/PagedArray.h
+++ b/Engine/Source/Developer/TraceServices/Public/Common/PagedArray.h
@@ -54,7 +54,6 @@
 public:
 	TPagedArrayIterator()
 	{
-
 	}
 
 	TPagedArrayIterator(const TPagedArray<ItemType, PageType>& InOuter, uint64 InitialPageIndex, uint64 InitialItemIndex)
@@ -71,36 +70,19 @@
 		}
 	}
 
-	uint64 GetCurrentPageIndex()
+	//////////////////////////////////////////////////
+	// Page Iterator
+
+	uint64 GetCurrentPageIndex() const
 	{
 		return CurrentPageIndex;
 	}
 
-	const PageType* GetCurrentPage()
+	const PageType* GetCurrentPage() const
 	{
 		return Outer->FirstPage + CurrentPageIndex;
 	}
 
-	const ItemType* GetCurrentItem()
-	{
-		return CurrentItem;
-	}
-
-	const ItemType& operator*() const
-	{
-		return *CurrentItem;
-	}
-
-	const ItemType* operator->() const
-	{
-		return CurrentItem;
-	}
-
-	explicit operator bool() const
-	{
-		return CurrentItem != nullptr;
-	}
-
 	const PageType* PrevPage()
 	{
 		if (CurrentPageIndex == 0)
@@ -131,6 +113,47 @@
 		return GetCurrentPage();
 	}
 
+	//////////////////////////////////////////////////
+	// Item Iterator
+
+	const ItemType* GetCurrentItem() const
+	{
+		return CurrentItem;
+	}
+
+	const ItemType* SetPosition(uint64 Index)
+	{
+		uint64 PageIndex = Index / Outer->PageSize;
+		uint64 ItemIndexInPage = Index % Outer->PageSize;
+		if (PageIndex != CurrentPageIndex)
+		{
+			check(PageIndex < Outer->PagesArray.Num());
+			CurrentPageIndex = PageIndex;
+			OnCurrentPageChanged();
+		}
+		PageType* CurrentPage = Outer->FirstPage + CurrentPageIndex;
+		check(ItemIndexInPage < CurrentPage->Count);
+		CurrentItem = CurrentPage->Items + ItemIndexInPage;
+		return CurrentItem;
+	}
+
+	const ItemType* PrevItem()
+	{
+		if (CurrentItem == CurrentPageFirstItem)
+		{
+			if (!PrevPage())
+			{
+				return nullptr;
+			}
+			else
+			{
+				return CurrentItem;
+			}
+		}
+		--CurrentItem;
+		return CurrentItem;
+	}
+
 	const ItemType* NextItem()
 	{
 		if (CurrentItem == CurrentPageLastItem)
@@ -148,6 +171,24 @@
 		return CurrentItem;
 	}
 
+	//////////////////////////////////////////////////
+	// operators
+
+	const ItemType& operator*() const
+	{
+		return *CurrentItem;
+	}
+
+	const ItemType* operator->() const
+	{
+		return CurrentItem;
+	}
+
+	explicit operator bool() const
+	{
+		return CurrentItem != nullptr;
+	}
+
 	TPagedArrayIterator& operator++()
 	{
 		NextItem();
@@ -161,23 +202,6 @@
 		return Tmp;
 	}
 
-	const ItemType* PrevItem()
-	{
-		if (CurrentItem == CurrentPageFirstItem)
-		{
-			if (!PrevPage())
-			{
-				return nullptr;
-			}
-			else
-			{
-				return CurrentItem;
-			}
-		}
-		--CurrentItem;
-		return CurrentItem;
-	}
-
 	TPagedArrayIterator& operator--()
 	{
 		PrevItem();
@@ -191,22 +215,6 @@
 		return Tmp;
 	}
 
-	const ItemType* SetPosition(uint64 Index)
-	{
-		uint64 PageIndex = Index / Outer->PageSize;
-		uint64 ItemIndexInPage = Index % Outer->PageSize;
-		if (PageIndex != CurrentPageIndex)
-		{
-			check(PageIndex < Outer->PagesArray.Num());
-			CurrentPageIndex = PageIndex;
-			OnCurrentPageChanged();
-		}
-		PageType* CurrentPage = Outer->FirstPage + CurrentPageIndex;
-		check(ItemIndexInPage < CurrentPage->Count);
-		CurrentItem = CurrentPage->Items + ItemIndexInPage;
-		return CurrentItem;
-	}
-
 private:
 	void OnCurrentPageChanged()
 	{
@@ -227,7 +235,7 @@
 		checkSlow(Lhs.Outer == Rhs.Outer); // Needs to be iterators of the same array
 		return Lhs.CurrentItem != Rhs.CurrentItem;
 	}
-	
+
 	const TPagedArray<ItemType, PageType>* Outer = nullptr;
 	const ItemType* CurrentItem = nullptr;
 	const ItemType* CurrentPageFirstItem = nullptr;
@@ -334,17 +342,33 @@
 		return LastPage;
 	}
 
+	const PageType* GetLastPage() const
+	{
+		return LastPage;
+	}
+
 	PageType* GetPage(uint64 PageIndex)
 	{
 		return FirstPage + PageIndex;
 	}
 
+	const PageType* GetPage(uint64 PageIndex) const
+	{
+		return FirstPage + PageIndex;
+	}
+
 	PageType* GetItemPage(uint64 ItemIndex)
 	{
 		uint64 PageIndex = ItemIndex / PageSize;
 		return FirstPage + PageIndex;
 	}
 
+	const PageType* GetItemPage(uint64 ItemIndex) const
+	{
+		uint64 PageIndex = ItemIndex / PageSize;
+		return FirstPage + PageIndex;
+	}
+
 	TIterator GetIterator() const
 	{
 		return TIterator(*this, 0, 0);
@@ -404,7 +428,7 @@
 		const ItemType* Item = LastPage->Items + LastPage->Count - 1;
 		return *Item;
 	}
-	
+
 	FORCEINLINE TIterator begin() { return TIterator(*this, 0, 0); }
 	FORCEINLINE TIterator begin() const { return TIterator(*this, 0, 0); }
 	FORCEINLINE TIterator end() { return TIterator(*this, NumPages()-1, LastPage->Count - 1); }
