[Insights] Allow screenshots to be saved to a file

diff --git a/Engine/Source/Developer/TraceInsights/Private/Insights/ViewModels/MarkersTimingTrack.cpp b/Engine/Source/Developer/TraceInsights/Private/Insights/ViewModels/MarkersTimingTrack.cpp
--- a/Engine/Source/Developer/TraceInsights/Private/Insights/ViewModels/MarkersTimingTrack.cpp
+++ b/Engine/Source/Developer/TraceInsights/Private/Insights/ViewModels/MarkersTimingTrack.cpp
@@ -2,11 +2,14 @@
 
 #include "MarkersTimingTrack.h"
 
+#include "DesktopPlatformModule.h"
 #include "Fonts/FontMeasure.h"
 #include "Framework/Application/SlateApplication.h"
 #include "Framework/MultiBox/MultiBoxBuilder.h"
 #include "IImageWrapper.h"
 #include "IImageWrapperModule.h"
+#include "Misc/FileHelper.h"
+#include "Misc/Paths.h"
 #include "Modules/ModuleManager.h"
 #include "Styling/AppStyle.h"
 #include "TraceServices/Model/Log.h"
@@ -20,6 +23,8 @@
 #include "Insights/ViewModels/TimingEvent.h"
 #include "Insights/ViewModels/TimingTrackViewport.h"
 #include "Insights/ViewModels/TooltipDrawState.h"
+#include "Insights/Widgets/STimingProfilerWindow.h"
+#include "Insights/Widgets/STimingView.h"
 
 #include <limits>
 
@@ -311,6 +316,22 @@ void FMarkersTimingTrack::BuildContextMenu(FMenuBuilder& MenuBuilder)
 		);
 	}
 	MenuBuilder.EndSection();
+
+	MenuBuilder.BeginSection(TEXT("Screenshot"), LOCTEXT("ContextMenu_Section_Screenshot", "Screenshot"));
+	{
+		MenuBuilder.AddMenuEntry(
+			LOCTEXT("ContextMenu_SaveScreenshot", "Save Screenshot"),
+			LOCTEXT("ContextMenu_SaveScreenshot_Desc", "Save the hovered screenshot to a file."),
+			FSlateIcon(),
+			FUIAction(FExecuteAction::CreateSP(this, &FMarkersTimingTrack::SaveScreenshot_Execute),
+				FCanExecuteAction::CreateSP(this, &FMarkersTimingTrack::SaveScreenshot_CanExecute)),
+			NAME_None,
+			EUserInterfaceActionType::Button
+		);
+
+		TryGetHoveredEventScreenshotId(LastScreenshotId);
+	}
+	MenuBuilder.EndSection();
 }
 
 ////////////////////////////////////////////////////////////////////////////////////////////////////
@@ -484,7 +505,7 @@ void FMarkersTimingTrack::InitTooltip(FTooltipDrawState& InOutTooltip, const ITi
 	const FTimingEvent& Event = StaticCast<const FTimingEvent&>(InTooltipEvent);
 
 	const TraceServices::IScreenshotProvider& ScreenshotProvider = TraceServices::ReadScreenshotProvider(*Session.Get());
-	TSharedPtr<const TraceServices::FScreenshot> Screenshot = ScreenshotProvider.GetScreenshot((uint16)Event.GetType());
+	TSharedPtr<const TraceServices::FScreenshot> Screenshot = ScreenshotProvider.GetScreenshot((uint32)Event.GetType());
 
 	if (!Screenshot.IsValid())
 	{
@@ -530,6 +551,97 @@ void FMarkersTimingTrack::InitTooltip(FTooltipDrawState& InOutTooltip, const ITi
 	InOutTooltip.UpdateLayout();
 }
 
+////////////////////////////////////////////////////////////////////////////////////////////////////
+
+bool FMarkersTimingTrack::SaveScreenshot_CanExecute()
+{
+	return LastScreenshotId != TraceServices::FScreenshot::InvalidScreenshotId;
+}
+
+////////////////////////////////////////////////////////////////////////////////////////////////////
+
+void FMarkersTimingTrack::SaveScreenshot_Execute()
+{
+	if (LastScreenshotId == TraceServices::FScreenshot::InvalidScreenshotId)
+	{
+		return;
+	}
+
+	TSharedPtr<const TraceServices::IAnalysisSession> Session = FInsightsManager::Get()->GetSession();
+	const TraceServices::IScreenshotProvider& ScreenshotProvider = TraceServices::ReadScreenshotProvider(*Session.Get());
+	TSharedPtr<const TraceServices::FScreenshot> Screenshot = ScreenshotProvider.GetScreenshot(LastScreenshotId);
+
+	if (!Screenshot.IsValid())
+	{
+		return;
+	}
+
+	TArray<FString> SaveFilenames;
+	bool bDialogResult = false;
+
+	FString DefaultFile = Screenshot->Name;
+	DefaultFile.Append(TEXT(".png"));
+
+	IDesktopPlatform* DesktopPlatform = FDesktopPlatformModule::Get();
+	if (DesktopPlatform)
+	{
+		const FString DefaultPath = FPaths::ProjectSavedDir();
+		bDialogResult = DesktopPlatform->SaveFileDialog(
+			FSlateApplication::Get().FindBestParentWindowHandleForDialogs(nullptr),
+			LOCTEXT("SaveScreenshotTitle", "Save Screenshot").ToString(),
+			DefaultPath,
+			DefaultFile,
+			TEXT("Portable Network Graphics File (*.png)|*.png"),
+			EFileDialogFlags::None,
+			SaveFilenames
+		);
+	}
+
+	if (!bDialogResult || SaveFilenames.Num() == 0)
+	{
+		return;
+	}
+
+	FString& Path = SaveFilenames[0];
+
+	FFileHelper::SaveArrayToFile(Screenshot->Data, *Path);
+}
+
+////////////////////////////////////////////////////////////////////////////////////////////////////
+
+bool FMarkersTimingTrack::TryGetHoveredEventScreenshotId(uint32& OutScreenshotId)
+{
+	OutScreenshotId = TraceServices::FScreenshot::InvalidScreenshotId;
+
+	TSharedPtr<STimingProfilerWindow> Window = FTimingProfilerManager::Get()->GetProfilerWindow();
+	if (!Window.IsValid())
+	{
+		return false;
+	}
+
+	TSharedPtr<STimingView> TimingView = Window->GetTimingView();
+	if (!TimingView.IsValid())
+	{
+		return false;
+	}
+
+	const TSharedPtr<const ITimingEvent> HoveredTimingEvent = TimingView->GetHoveredEvent();
+	if (!HoveredTimingEvent.IsValid())
+	{
+		return false;
+	}
+
+	const FBaseTimingTrack& Track = HoveredTimingEvent->GetTrack().Get();
+	if (!HoveredTimingEvent->CheckTrack(this))
+	{
+		return false;
+	}
+
+	const FTimingEvent& Event = StaticCast<const FTimingEvent&>(*HoveredTimingEvent);
+	OutScreenshotId = (uint32) Event.GetType();
+	return true;
+}
+
 ////////////////////////////////////////////////////////////////////////////////////////////////////
 // FTimeMarkerTrackBuilder
 ////////////////////////////////////////////////////////////////////////////////////////////////////
diff --git a/Engine/Source/Developer/TraceInsights/Private/Insights/ViewModels/MarkersTimingTrack.h b/Engine/Source/Developer/TraceInsights/Private/Insights/ViewModels/MarkersTimingTrack.h
--- a/Engine/Source/Developer/TraceInsights/Private/Insights/ViewModels/MarkersTimingTrack.h
+++ b/Engine/Source/Developer/TraceInsights/Private/Insights/ViewModels/MarkersTimingTrack.h
@@ -69,6 +69,9 @@ public:
 	void SetBookmarksTrack() { SetBookmarksTrackFlag(true); SetDirtyFlag(); }
 	void SetLogsTrack() { SetBookmarksTrackFlag(false); SetDirtyFlag(); }
 
+	bool SaveScreenshot_CanExecute();
+	void SaveScreenshot_Execute();
+
 	// Stats
 	int32 GetNumLogMessages() const { return NumLogMessages; }
 	int32 GetNumBoxes() const { return TimeMarkerBoxes.Num(); }
@@ -100,6 +103,8 @@ private:
 
 	void UpdateCategory(const TraceServices::FLogCategoryInfo*& InOutCategory, const TCHAR* CategoryName);
 
+	bool TryGetHoveredEventScreenshotId(uint32& OutScreenshotId);
+
 private:
 	TArray<FTimeMarkerBoxInfo> TimeMarkerBoxes;
 	TArray<FTimeMarkerTextInfo> TimeMarkerTexts;
@@ -118,6 +123,7 @@ private:
 	// Slate resources
 	const FSlateBrush* WhiteBrush;
 	const FSlateFontInfo Font;
+	uint32 LastScreenshotId;
 };
 
 ////////////////////////////////////////////////////////////////////////////////////////////////////
diff --git a/Engine/Source/Developer/TraceServices/Public/TraceServices/Model/Screenshot.h b/Engine/Source/Developer/TraceServices/Public/TraceServices/Model/Screenshot.h
--- a/Engine/Source/Developer/TraceServices/Public/TraceServices/Model/Screenshot.h
+++ b/Engine/Source/Developer/TraceServices/Public/TraceServices/Model/Screenshot.h
@@ -18,6 +18,8 @@ struct FScreenshot
 	uint32 ChunkNum;
 	uint32 Size;
 	TArray<uint8> Data;
+
+	static const uint32 InvalidScreenshotId = (uint32)-1;
 };
 
 class IScreenshotProvider
