[Insights] Mem Insights: Added "Class Name" column to Allocs Table togather with a custom view preset.

diff --git a/Engine/Source/Developer/TraceInsights/Private/Insights/MemoryProfiler/ViewModels/MemAllocTable.cpp b/Engine/Source/Developer/TraceInsights/Private/Insights/MemoryProfiler/ViewModels/MemAllocTable.cpp
--- a/Engine/Source/Developer/TraceInsights/Private/Insights/MemoryProfiler/ViewModels/MemAllocTable.cpp
+++ b/Engine/Source/Developer/TraceInsights/Private/Insights/MemoryProfiler/ViewModels/MemAllocTable.cpp
@@ -37,6 +37,7 @@ const FName FMemAllocTableColumns::CountColumnId(TEXT("Count"));
 const FName FMemAllocTableColumns::SizeColumnId(TEXT("Size"));
 const FName FMemAllocTableColumns::TagColumnId(TEXT("Tag"));
 const FName FMemAllocTableColumns::AssetColumnId(TEXT("Asset"));
+const FName FMemAllocTableColumns::ClassNameColumnId(TEXT("ClassName"));
 const FName FMemAllocTableColumns::FunctionColumnId(TEXT("Function"));
 const FName FMemAllocTableColumns::SourceFileColumnId(TEXT("SourceFile"));
 const FName FMemAllocTableColumns::CallstackSizeColumnId(TEXT("CallstackSize"));
@@ -793,6 +794,64 @@ void FMemAllocTable::AddDefaultColumns()
 		AddColumn(ColumnRef);
 	}
 	//////////////////////////////////////////////////
+	// Class Name Column
+	{
+		TSharedRef<FTableColumn> ColumnRef = MakeShared<FTableColumn>(FMemAllocTableColumns::ClassNameColumnId);
+		FTableColumn& Column = *ColumnRef;
+
+		Column.SetIndex(ColumnIndex++);
+
+		Column.SetShortName(LOCTEXT("ClassNameColumnName", "Class Name"));
+		Column.SetTitleName(LOCTEXT("ClassNameColumnTitle", "Class Name"));
+		Column.SetDescription(LOCTEXT("ClassNameColumnDesc", "Class of asset associated with allocation"));
+
+		Column.SetFlags(ETableColumnFlags::CanBeHidden | ETableColumnFlags::CanBeFiltered);
+
+		Column.SetHorizontalAlignment(HAlign_Left);
+		Column.SetInitialWidth(120.0f);
+
+		Column.SetDataType(ETableCellDataType::CString);
+
+		class FClassNameMetadataValueGetter : public FTableCellValueGetter
+		{
+		public:
+			virtual const TOptional<FTableCellValue> GetValue(const FTableColumn& Column, const FBaseTreeNode& Node) const
+			{
+				if (Node.IsGroup())
+				{
+					const FTableTreeNode& NodePtr = static_cast<const FTableTreeNode&>(Node);
+					if (NodePtr.HasAggregatedValue(Column.GetId()))
+					{
+						return NodePtr.GetAggregatedValue(Column.GetId());
+					}
+				}
+				else //if (Node->Is<FMemAllocNode>())
+				{
+					const FMemAllocNode& MemAllocNode = static_cast<const FMemAllocNode&>(Node);
+					const FMemoryAlloc* Alloc = MemAllocNode.GetMemAlloc();
+					if (Alloc)
+					{
+						return FTableCellValue(Alloc->GetClassName());
+					}
+				}
+
+				return TOptional<FTableCellValue>();
+			}
+		};
+		TSharedRef<ITableCellValueGetter> Getter = MakeShared<FClassNameMetadataValueGetter>();
+		Column.SetValueGetter(Getter);
+
+		TSharedRef<ITableCellValueFormatter> Formatter = MakeShared<FCStringValueFormatterAsText>();
+		Column.SetValueFormatter(Formatter);
+
+		TSharedRef<ITableCellValueSorter> Sorter = MakeShared<FSorterByCStringValue>(ColumnRef);
+		Column.SetValueSorter(Sorter);
+
+		Column.SetAggregation(ETableColumnAggregation::SameValue);
+
+		AddColumn(ColumnRef);
+	}
+	//////////////////////////////////////////////////
 	// Top Function Column
 	{
 		TSharedRef<FTableColumn> ColumnRef = MakeShared<FTableColumn>(FMemAllocTableColumns::FunctionColumnId);
diff --git a/Engine/Source/Developer/TraceInsights/Private/Insights/MemoryProfiler/ViewModels/MemAllocTable.h b/Engine/Source/Developer/TraceInsights/Private/Insights/MemoryProfiler/ViewModels/MemAllocTable.h
--- a/Engine/Source/Developer/TraceInsights/Private/Insights/MemoryProfiler/ViewModels/MemAllocTable.h
+++ b/Engine/Source/Developer/TraceInsights/Private/Insights/MemoryProfiler/ViewModels/MemAllocTable.h
@@ -30,6 +30,7 @@ struct FMemAllocTableColumns
 	static const FName SizeColumnId;
 	static const FName TagColumnId;
 	static const FName AssetColumnId;
+	static const FName ClassNameColumnId;
 	static const FName FunctionColumnId;
 	static const FName SourceFileColumnId;
 	static const FName CallstackSizeColumnId;
diff --git a/Engine/Source/Developer/TraceInsights/Private/Insights/MemoryProfiler/ViewModels/MemoryAlloc.h b/Engine/Source/Developer/TraceInsights/Private/Insights/MemoryProfiler/ViewModels/MemoryAlloc.h
--- a/Engine/Source/Developer/TraceInsights/Private/Insights/MemoryProfiler/ViewModels/MemoryAlloc.h
+++ b/Engine/Source/Developer/TraceInsights/Private/Insights/MemoryProfiler/ViewModels/MemoryAlloc.h
@@ -32,6 +32,7 @@ public:
 	TraceServices::TagIdType GetTagId() const { return TagId; }
 	const TCHAR* GetTag() const { return Tag; }
 	const TCHAR* GetAsset() const { return Asset; }
+	const TCHAR* GetClassName() const { return ClassName; }
 	const TraceServices::FCallstack* GetCallstack() const { return Callstack; }
 	const TraceServices::FCallstack* GetFreeCallstack() const { return FreeCallstack; }
 	FText GetFullCallstack() const;
@@ -48,6 +49,7 @@ private:
 	TraceServices::TagIdType TagId;
 	const TCHAR* Tag;
 	const TCHAR* Asset;
+	const TCHAR* ClassName;
 	const TraceServices::FCallstack* Callstack;
 	const TraceServices::FCallstack* FreeCallstack;
 	HeapId RootHeap;
diff --git a/Engine/Source/Developer/TraceInsights/Private/Insights/MemoryProfiler/Widgets/SMemAllocTableTreeView.cpp b/Engine/Source/Developer/TraceInsights/Private/Insights/MemoryProfiler/Widgets/SMemAllocTableTreeView.cpp
--- a/Engine/Source/Developer/TraceInsights/Private/Insights/MemoryProfiler/Widgets/SMemAllocTableTreeView.cpp
+++ b/Engine/Source/Developer/TraceInsights/Private/Insights/MemoryProfiler/Widgets/SMemAllocTableTreeView.cpp
@@ -399,6 +399,7 @@ void SMemAllocTableTreeView::UpdateQuery(TraceServices::IAllocationsProvider::EQ
 					Alloc.Tag = Provider.GetTagName(Allocation->GetTag());
 
 					Alloc.Asset = nullptr;
+					Alloc.ClassName = nullptr;
 					const uint32 MetadataId = Allocation->GetMetadataId();
 					if (MetadataId != TraceServices::IMetadataProvider::InvalidMetadataId && MetadataProvider && DefinitionProvider)
 					{
@@ -416,6 +417,12 @@ void SMemAllocTableTreeView::UpdateQuery(TraceServices::IAllocationsProvider::EQ
 									{
 										Alloc.Asset = AssetName->Display;
 									}
+									const auto ClassNameRef = Reader.GetValueAs<UE::Trace::FEventRef32>((uint8*)Data, 1);
+									const auto ClassName = DefinitionProvider->Get<TraceServices::FStringDefinition>(*ClassNameRef);
+									if (ClassName)
+									{
+										Alloc.ClassName = ClassName->Display;
+									}
 									return false;
 								}
 								return true;
@@ -885,6 +892,68 @@ void SMemAllocTableTreeView::InitAvailableViewPresets()
 	};
 	AvailableViewPresets.Add(MakeShared<FAssetViewPreset>());
 
+	//////////////////////////////////////////////////
+	// Class Name Breakdown View
+
+	class FClassNameViewPreset : public IViewPreset
+	{
+	public:
+		virtual FText GetName() const override
+		{
+			return LOCTEXT("ClassName_PresetName", "Class Name");
+		}
+		virtual FText GetToolTip() const override
+		{
+			return LOCTEXT("ClassName_PresetToolTip", "Class Breakdown View\nConfigure the tree view to show a breakdown of allocations by their Class name.");
+		}
+		virtual FName GetSortColumn() const override
+		{
+			return FTable::GetHierarchyColumnId();
+		}
+		virtual EColumnSortMode::Type GetSortMode() const override
+		{
+			return EColumnSortMode::Type::Ascending;
+		}
+		virtual void SetCurrentGroupings(const TArray<TSharedPtr<FTreeNodeGrouping>>& InAvailableGroupings, TArray<TSharedPtr<FTreeNodeGrouping>>& InOutCurrentGroupings) const override
+		{
+			InOutCurrentGroupings.Reset();
+
+			check(InAvailableGroupings[0]->Is<FTreeNodeGroupingFlat>());
+			InOutCurrentGroupings.Add(InAvailableGroupings[0]);
+
+			const TSharedPtr<FTreeNodeGrouping>* ClassNameGrouping = InAvailableGroupings.FindByPredicate(
+				[](TSharedPtr<FTreeNodeGrouping>& Grouping)
+				{
+					return Grouping->Is<FTreeNodeGroupingByUniqueValueCString>() &&
+						   Grouping->As<FTreeNodeGroupingByUniqueValueCString>().GetColumnId() == FMemAllocTableColumns::ClassNameColumnId;
+				});
+			if (ClassNameGrouping)
+			{
+				InOutCurrentGroupings.Add(*ClassNameGrouping);
+			}
+
+			const TSharedPtr<FTreeNodeGrouping>* AssetGrouping = InAvailableGroupings.FindByPredicate(
+				[](TSharedPtr<FTreeNodeGrouping>& Grouping)
+				{
+					return Grouping->Is<FTreeNodeGroupingByUniqueValueCString>() &&
+						   Grouping->As<FTreeNodeGroupingByUniqueValueCString>().GetColumnId() == FMemAllocTableColumns::AssetColumnId;
+				});
+			if (AssetGrouping)
+			{
+				InOutCurrentGroupings.Add(*AssetGrouping);
+			}
+		}
+		virtual void GetColumnConfigSet(TArray<FColumnConfig>& InOutConfigSet) const override
+		{
+			InOutConfigSet.Add({ FTable::GetHierarchyColumnId(),          true, 200.0f });
+			InOutConfigSet.Add({ FMemAllocTableColumns::CountColumnId,    true, 100.0f });
+			InOutConfigSet.Add({ FMemAllocTableColumns::SizeColumnId,     true, 100.0f });
+			InOutConfigSet.Add({ FMemAllocTableColumns::TagColumnId,      true, 120.0f });
+			InOutConfigSet.Add({ FMemAllocTableColumns::FunctionColumnId, true, 400.0f });
+		}
+	};
+	AvailableViewPresets.Add(MakeShared<FClassNameViewPreset>());
+
 	//////////////////////////////////////////////////
 	// (Inverted) Callstack Breakdown View
 
@@ -1201,11 +1270,11 @@ void SMemAllocTableTreeView::InternalCreateGroupings()
 			else if (Grouping->Is<FTreeNodeGroupingByPathBreakdown>())
 			{
 				const FName ColumnId = Grouping->As<FTreeNodeGroupingByPathBreakdown>().GetColumnId();
-				if (ColumnId == FMemAllocTableColumns::FunctionColumnId)
+				if (ColumnId == FMemAllocTableColumns::FunctionColumnId ||
+					ColumnId == FMemAllocTableColumns::ClassNameColumnId)
 				{
 					return true;
 				}
-
 			}
 			return false;
 		});
