fixed the column name that wasn't changed by mistake after copy-paste

diff --git a/Engine/Source/Developer/TraceInsights/Private/Insights/TaskGraphProfiler/ViewModels/TaskTable.cpp b/Engine/Source/Developer/TraceInsights/Private/Insights/TaskGraphProfiler/ViewModels/TaskTable.cpp
--- a/Engine/Source/Developer/TraceInsights/Private/Insights/TaskGraphProfiler/ViewModels/TaskTable.cpp
+++ b/Engine/Source/Developer/TraceInsights/Private/Insights/TaskGraphProfiler/ViewModels/TaskTable.cpp
@@ -609,7 +609,7 @@ void FTaskTable::AddDefaultColumns()
 		Column.SetIndex(ColumnIndex++);
 
 		Column.SetShortName(LOCTEXT("DestroyedThreadIdColumnName", "Destroyed Thread Id"));
-		Column.SetTitleName(LOCTEXT("DestroyedThreadIdColumnTitle", "Completed Thread Id"));
+		Column.SetTitleName(LOCTEXT("DestroyedThreadIdColumnTitle", "Destroyed Thread Id"));
 		Column.SetDescription(LOCTEXT("DestroyedThreadIdColumnDesc", "The thread the task was Destroyed on."));
 
 		Column.SetFlags(ETableColumnFlags::CanBeHidden | ETableColumnFlags::CanBeFiltered);
