Track asset class for packages in CookProfiler

diff --git a/Engine/Source/Developer/TraceInsights/Private/Insights/CookProfiler/ViewModels/PackageEntry.cpp b/Engine/Source/Developer/TraceInsights/Private/Insights/CookProfiler/ViewModels/PackageEntry.cpp
--- a/Engine/Source/Developer/TraceInsights/Private/Insights/CookProfiler/ViewModels/PackageEntry.cpp
+++ b/Engine/Source/Developer/TraceInsights/Private/Insights/CookProfiler/ViewModels/PackageEntry.cpp
@@ -14,6 +14,7 @@ FPackageEntry::FPackageEntry(const TraceServices::FPackageData& PackageData)
 	, SaveTime(PackageData.SaveTime)
 	, BeginCacheForCookedPlatformData(PackageData.BeginCacheForCookedPlatformData)
 	, IsCachedCookedPlatformDataLoaded(PackageData.IsCachedCookedPlatformDataLoaded)
+	, AssetClass(PackageData.AssetClass)
 {
 }
 
diff --git a/Engine/Source/Developer/TraceInsights/Private/Insights/CookProfiler/ViewModels/PackageEntry.h b/Engine/Source/Developer/TraceInsights/Private/Insights/CookProfiler/ViewModels/PackageEntry.h
--- a/Engine/Source/Developer/TraceInsights/Private/Insights/CookProfiler/ViewModels/PackageEntry.h
+++ b/Engine/Source/Developer/TraceInsights/Private/Insights/CookProfiler/ViewModels/PackageEntry.h
@@ -25,6 +25,7 @@ public:
 	const double GetSaveTime() const { return SaveTime; }
 	const double GetBeginCacheForCookedPlatformData() const { return BeginCacheForCookedPlatformData; }
 	const double GetIsCachedCookedPlatformDataLoaded() const { return IsCachedCookedPlatformDataLoaded; }
+	const TCHAR* GetAssetClass() const { return AssetClass; }
 
 private:
 	uint64 Id;
@@ -33,6 +34,7 @@ private:
 	double SaveTime;
 	double BeginCacheForCookedPlatformData;
 	double IsCachedCookedPlatformDataLoaded;
+	const TCHAR* AssetClass;
 };
 
 ////////////////////////////////////////////////////////////////////////////////////////////////////
diff --git a/Engine/Source/Developer/TraceInsights/Private/Insights/CookProfiler/ViewModels/PackageTable.cpp b/Engine/Source/Developer/TraceInsights/Private/Insights/CookProfiler/ViewModels/PackageTable.cpp
--- a/Engine/Source/Developer/TraceInsights/Private/Insights/CookProfiler/ViewModels/PackageTable.cpp
+++ b/Engine/Source/Developer/TraceInsights/Private/Insights/CookProfiler/ViewModels/PackageTable.cpp
@@ -23,7 +23,8 @@ const FName FPackageTableColumns::NameColumnId(TEXT("Name"));
 const FName FPackageTableColumns::LoadTimeColumnId(TEXT("LoadTime"));
 const FName FPackageTableColumns::SaveTimeColumnId(TEXT("SaveTime"));
 const FName FPackageTableColumns::BeginCacheForCookedPlatformDataTimeColumnId(TEXT("BeginCacheForCookedPlatformDataTimeColumnId"));
-const FName FPackageTableColumns::GetIsCachedCookedPlatformDataLoaded(TEXT("GetIsCachedCookedPlatformDataLoaded"));
+const FName FPackageTableColumns::GetIsCachedCookedPlatformDataLoadedColumnId(TEXT("GetIsCachedCookedPlatformDataLoaded"));
+const FName FPackageTableColumns::PackageAssetClassColumnId(TEXT("AssetClass"));
 
 ////////////////////////////////////////////////////////////////////////////////////////////////////
 typedef FTableCellValue(*PackageFieldGetter)(const FTableColumn&, const FPackageEntry&);
@@ -61,11 +62,12 @@ public:
 struct DefaultPackageFieldGetterFuncts
 {
 	static FTableCellValue GetId(const FTableColumn& Column, const FPackageEntry& Package) { return FTableCellValue((int64)Package.GetId());	}
-	static FTableCellValue GetName(const FTableColumn& Column, const FPackageEntry& Package) { return FTableCellValue((int64)Package.GetName());	}
+	static FTableCellValue GetName(const FTableColumn& Column, const FPackageEntry& Package) { return FTableCellValue((const TCHAR*)Package.GetName());	}
 	static FTableCellValue GetLoadTime(const FTableColumn& Column, const FPackageEntry& Package) { return FTableCellValue((double)Package.GetLoadTime());	}
 	static FTableCellValue GetSaveTime(const FTableColumn& Column, const FPackageEntry& Package) { return FTableCellValue((double)Package.GetSaveTime());	}
 	static FTableCellValue GetBeginCacheForCookedPlatformData(const FTableColumn& Column, const FPackageEntry& Package) { return FTableCellValue((double)Package.GetBeginCacheForCookedPlatformData());	}
 	static FTableCellValue GetIsCachedCookedPlatformDataLoaded(const FTableColumn& Column, const FPackageEntry& Package) { return FTableCellValue((double)Package.GetIsCachedCookedPlatformDataLoaded());	}
+	static FTableCellValue GetAssetClass(const FTableColumn& Column, const FPackageEntry& Package) { return FTableCellValue((const TCHAR*)Package.GetAssetClass());	}
 };
 
 ////////////////////////////////////////////////////////////////////////////////////////////////////
@@ -236,7 +238,7 @@ void FPackageTable::AddDefaultColumns()
 	//////////////////////////////////////////////////
 	// BeginCacheForCookedPlatformData Time Column
 	{
-		TSharedRef<FTableColumn> ColumnRef = MakeShared<FTableColumn>(FPackageTableColumns::GetIsCachedCookedPlatformDataLoaded);
+		TSharedRef<FTableColumn> ColumnRef = MakeShared<FTableColumn>(FPackageTableColumns::GetIsCachedCookedPlatformDataLoadedColumnId);
 		FTableColumn& Column = *ColumnRef;
 
 		Column.SetIndex(ColumnIndex++);
@@ -261,6 +263,36 @@ void FPackageTable::AddDefaultColumns()
 		TSharedRef<ITableCellValueSorter> Sorter = MakeShared<FSorterByDoubleValue>(ColumnRef);
 		Column.SetValueSorter(Sorter);
 
+		AddColumn(ColumnRef);
+	}
+	//////////////////////////////////////////////////
+	// Asset Class Column
+	{
+		TSharedRef<FTableColumn> ColumnRef = MakeShared<FTableColumn>(FPackageTableColumns::PackageAssetClassColumnId);
+		FTableColumn& Column = *ColumnRef;
+
+		Column.SetIndex(ColumnIndex++);
+
+		Column.SetShortName(LOCTEXT("AssetClassColumnName", "Asset Class"));
+		Column.SetTitleName(LOCTEXT("AssetClassTitle", "Asset Class"));
+		Column.SetDescription(LOCTEXT("AssetClassColumnDesc", "The class of the most significant asset in the package."));
+
+		Column.SetFlags(ETableColumnFlags::ShouldBeVisible | ETableColumnFlags::CanBeHidden | ETableColumnFlags::CanBeFiltered);
+
+		Column.SetHorizontalAlignment(HAlign_Left);
+		Column.SetInitialWidth(100.0f);
+
+		Column.SetDataType(ETableCellDataType::CString);
+
+		TSharedRef<ITableCellValueGetter> Getter = MakeShared<FPackageColumnValueGetter<DefaultPackageFieldGetterFuncts::GetAssetClass>>();
+		Column.SetValueGetter(Getter);
+
+		TSharedRef<ITableCellValueFormatter> Formatter = MakeShared<FCStringValueFormatterAsText>();
+		Column.SetValueFormatter(Formatter);
+
+		TSharedRef<ITableCellValueSorter> Sorter = MakeShared<FSorterByCStringValue>(ColumnRef);
+		Column.SetValueSorter(Sorter);
+
 		AddColumn(ColumnRef);
 	}
 }
diff --git a/Engine/Source/Developer/TraceInsights/Private/Insights/CookProfiler/ViewModels/PackageTable.h b/Engine/Source/Developer/TraceInsights/Private/Insights/CookProfiler/ViewModels/PackageTable.h
--- a/Engine/Source/Developer/TraceInsights/Private/Insights/CookProfiler/ViewModels/PackageTable.h
+++ b/Engine/Source/Developer/TraceInsights/Private/Insights/CookProfiler/ViewModels/PackageTable.h
@@ -23,7 +23,8 @@ struct FPackageTableColumns
 	static const FName LoadTimeColumnId;
 	static const FName SaveTimeColumnId;
 	static const FName BeginCacheForCookedPlatformDataTimeColumnId;
-	static const FName GetIsCachedCookedPlatformDataLoaded;
+	static const FName GetIsCachedCookedPlatformDataLoadedColumnId;
+	static const FName PackageAssetClassColumnId;
 };
 
 ////////////////////////////////////////////////////////////////////////////////////////////////////
diff --git a/Engine/Source/Developer/TraceServices/Private/Analyzers/CookAnalysis.cpp b/Engine/Source/Developer/TraceServices/Private/Analyzers/CookAnalysis.cpp
--- a/Engine/Source/Developer/TraceServices/Private/Analyzers/CookAnalysis.cpp
+++ b/Engine/Source/Developer/TraceServices/Private/Analyzers/CookAnalysis.cpp
@@ -28,6 +28,7 @@ void FCookAnalyzer::OnAnalysisBegin(const FOnAnalysisContext& Context)
 
 	Builder.RouteEvent(RouteId_Package, "CookTrace", "Package");
 	Builder.RouteEvent(RouteId_PackageStat, "CookTrace", "PackageStat");
+	Builder.RouteEvent(RouteId_PackageAssetClass, "CookTrace", "PackageAssetClass");
 }
 
 bool FCookAnalyzer::OnEvent(uint16 RouteId, EStyle Style, const FOnEventContext& Context)
@@ -89,6 +90,17 @@ bool FCookAnalyzer::OnEvent(uint16 RouteId, EStyle Style, const FOnEventContext&
 		}
 		break;
 	}
+	case RouteId_PackageAssetClass:
+	{
+		uint64 Id = EventData.GetValue<uint64>("Id");
+		FString ClassName;
+		EventData.GetString("ClassName", ClassName);
+
+		TraceServices::FProviderEditScopeLock ProviderEditScope(CookProfilerProvider);
+		FPackageData* Package = CookProfilerProvider.EditPackage(Id);
+		Package->AssetClass = Session.StoreString(ClassName);
+		break;
+	}
 	}
 
 	return true;
diff --git a/Engine/Source/Developer/TraceServices/Private/Analyzers/CookAnalysis.h b/Engine/Source/Developer/TraceServices/Private/Analyzers/CookAnalysis.h
--- a/Engine/Source/Developer/TraceServices/Private/Analyzers/CookAnalysis.h
+++ b/Engine/Source/Developer/TraceServices/Private/Analyzers/CookAnalysis.h
@@ -23,6 +23,7 @@ private:
 	{
 		RouteId_Package,
 		RouteId_PackageStat,
+		RouteId_PackageAssetClass,
 	};
 
 	IAnalysisSession& Session;
diff --git a/Engine/Source/Developer/TraceServices/Private/Model/CookProfilerProvider.cpp b/Engine/Source/Developer/TraceServices/Private/Model/CookProfilerProvider.cpp
--- a/Engine/Source/Developer/TraceServices/Private/Model/CookProfilerProvider.cpp
+++ b/Engine/Source/Developer/TraceServices/Private/Model/CookProfilerProvider.cpp
@@ -11,10 +11,12 @@ thread_local int32 GThreadCurrentReadCookProviderLockCount;
 thread_local int32 GThreadCurrentWriteCookProviderLockCount;
 
 const TCHAR* GUnknownPackage = TEXT("Unknown Package");
+const TCHAR* GUnknownClass = TEXT("Unknown Class");
 
 FPackageData::FPackageData(uint64 InId)
 	: Id(InId)
 	, Name(GUnknownPackage)
+	, AssetClass(GUnknownClass)
 {}
 
 void FCookProfilerProvider::BeginEdit() const
diff --git a/Engine/Source/Developer/TraceServices/Public/TraceServices/Model/CookProfilerProvider.h b/Engine/Source/Developer/TraceServices/Public/TraceServices/Model/CookProfilerProvider.h
--- a/Engine/Source/Developer/TraceServices/Public/TraceServices/Model/CookProfilerProvider.h
+++ b/Engine/Source/Developer/TraceServices/Public/TraceServices/Model/CookProfilerProvider.h
@@ -18,6 +18,7 @@ struct FPackageData
 	double SaveTime = 0;
 	double BeginCacheForCookedPlatformData = 0; // BeginCacheForCookedPlatformData is the name of the function from the cooker.
 	double IsCachedCookedPlatformDataLoaded = 0; // IsCachedCookedPlatformDataLoaded is the name of the function from the cooker.
+	const TCHAR* AssetClass;
 };
 
 class ICookProfilerProvider
diff --git a/Engine/Source/Runtime/Core/Private/ProfilingDebugging/CookStats.cpp b/Engine/Source/Runtime/Core/Private/ProfilingDebugging/CookStats.cpp
--- a/Engine/Source/Runtime/Core/Private/ProfilingDebugging/CookStats.cpp
+++ b/Engine/Source/Runtime/Core/Private/ProfilingDebugging/CookStats.cpp
@@ -20,6 +20,11 @@ UE_TRACE_EVENT_BEGIN(CookTrace, PackageStat)
 	UE_TRACE_EVENT_FIELD(uint8, StatType)
 UE_TRACE_EVENT_END()
 
+UE_TRACE_EVENT_BEGIN(CookTrace, PackageAssetClass)
+	UE_TRACE_EVENT_FIELD(uint64, Id)
+	UE_TRACE_EVENT_FIELD(UE::Trace::WideString, ClassName)
+UE_TRACE_EVENT_END()
+
 CORE_API FCookStatsManager::FGatherCookStatsDelegate FCookStatsManager::CookStatsCallbacks;
 
 CORE_API void FCookStatsManager::LogCookStats(AddStatFuncRef AddStat)
@@ -27,11 +32,11 @@ CORE_API void FCookStatsManager::LogCookStats(AddStatFuncRef AddStat)
 	CookStatsCallbacks.Broadcast(AddStat);
 }
 
-void TracePackage(uint64 InId, FString InName)
+void TracePackage(uint64 InId, const FStringView InName)
 {
 	UE_TRACE_LOG(CookTrace, Package, CookChannel)
 		<< Package.Id(InId)
-		<< Package.Name(*InName, uint16(InName.Len()))
+		<< Package.Name(InName.GetData(), uint16(InName.Len()))
 		<< Package.Cycle(FPlatformTime::Cycles64());
 }
 
@@ -43,4 +48,16 @@ void TracePackageStat(uint64 InId, uint64 Duration, EPackageEventStatType StatTy
 		<< PackageStat.StatType((uint8)StatType);
 }
 
+void TracePackageAssetClass(uint64 InId, const FStringView InName)
+{
+	UE_TRACE_LOG(CookTrace, PackageAssetClass, CookChannel)
+		<< PackageAssetClass.Id(InId)
+		<< PackageAssetClass.ClassName(InName.GetData(), uint16(InName.Len()));
+}
+
+bool ShouldTracePackageInfo()
+{
+	return CookChannel.IsEnabled();
+}
+
 #endif
diff --git a/Engine/Source/Runtime/Core/Public/ProfilingDebugging/CookStats.h b/Engine/Source/Runtime/Core/Public/ProfilingDebugging/CookStats.h
--- a/Engine/Source/Runtime/Core/Public/ProfilingDebugging/CookStats.h
+++ b/Engine/Source/Runtime/Core/Public/ProfilingDebugging/CookStats.h
@@ -433,8 +433,10 @@ struct FDerivedDataUsageStats
 
 CORE_API UE_TRACE_CHANNEL_EXTERN(CookChannel);
 
-CORE_API void TracePackage(uint64 InId, FString InName);
+CORE_API void TracePackage(uint64 InId, const FStringView InName);
 CORE_API void TracePackageStat(uint64 InId, uint64 Duration, EPackageEventStatType StatType);
+CORE_API void TracePackageAssetClass(uint64 InId, const FStringView InName);
+CORE_API bool ShouldTracePackageInfo();
 
 struct FScopedCookStat
 {
@@ -485,5 +487,7 @@ private:
 
 #define TracePackage(...)
 #define TracePackageStat(...)
+#define TracePackageAssetClass(...)
+#define ShouldTracePackageInfo(...) false
 
 #endif
diff --git a/Engine/Source/Runtime/CoreUObject/Private/UObject/LinkerLoad.cpp b/Engine/Source/Runtime/CoreUObject/Private/UObject/LinkerLoad.cpp
--- a/Engine/Source/Runtime/CoreUObject/Private/UObject/LinkerLoad.cpp
+++ b/Engine/Source/Runtime/CoreUObject/Private/UObject/LinkerLoad.cpp
@@ -19,6 +19,7 @@
 #include "UObject/UObjectHash.h"
 #include "Misc/PackageName.h"
 #include "Blueprint/BlueprintSupport.h"
+#include "Misc/PathViews.h"
 #include "Misc/PreloadableFile.h"
 #include "Misc/SecureHash.h"
 #include "Misc/StringBuilder.h"
@@ -72,6 +73,43 @@ DECLARE_FLOAT_ACCUMULATOR_STAT(TEXT("Fixup editor-only flags time"), STAT_Editor
 
 FName FLinkerLoad::NAME_LoadErrors("LoadErrors");
 
+/**
+* Helper function to determine and trace the most important asset class.
+*/
+void TrackPackageAssetClass(UPackage* Package, FLinkerLoad& LinkerLoad, const TArray<FObjectExport>& Exports)
+{
+#if ENABLE_COOK_STATS
+	if (!ShouldTracePackageInfo() || Exports.Num() == 0 || !Package)
+	{
+		return;
+	}
+
+	FName PackageName = Package->GetFName();
+	WriteToString<256> PackageNameStr(PackageName);
+	FStringView PackageLeafName = FPathViews::GetCleanFilename(PackageNameStr.ToView());
+	const FObjectExport* MostImportant = nullptr;
+	for (const FObjectExport& Export : Exports)
+	{
+		if (Export.bIsAsset && Export.ClassIndex.IsImport())
+		{
+			WriteToString<256> ObjectName(Export.ObjectName);
+			if (ObjectName.ToView() == PackageLeafName)
+			{
+				MostImportant = &Export;
+				break;
+			}
+			if (!MostImportant)
+			{
+				MostImportant = &Export;
+			}
+		}
+	}
+	if (MostImportant)
+	{
+		TracePackageAssetClass(PackageName.ToUnstableInt(), LinkerLoad.Imp(MostImportant->ClassIndex).ObjectName.ToString());
+	}
+#endif
+}
 
 /*----------------------------------------------------------------------------
 Helpers
@@ -1956,7 +1994,15 @@ FLinkerLoad::ELinkerStatus FLinkerLoad::SerializeExportMap()
 	}
 
 	// Return whether we finished this step and it's safe to start with the next.
-	return ((ExportMapIndex == Summary.ExportCount) && !IsTimeLimitExceeded( TEXT("serializing export map") )) ? LINKER_Loaded : LINKER_TimedOut;
+	if ((ExportMapIndex == Summary.ExportCount) && !IsTimeLimitExceeded(TEXT("serializing export map")))
+	{
+		TrackPackageAssetClass(LinkerRoot, *this, ExportMap);
+		return LINKER_Loaded;
+	}
+	else
+	{
+		return LINKER_TimedOut;
+	}
 }
 
 #if WITH_TEXT_ARCHIVE_SUPPORT
