[Insights] Fixed order of GPU/GPU2 tracks (first GPU to always be on top).

diff --git a/Engine/Source/Developer/TraceInsights/Private/Insights/ViewModels/ThreadTimingTrack.cpp b/Engine/Source/Developer/TraceInsights/Private/Insights/ViewModels/ThreadTimingTrack.cpp
--- a/Engine/Source/Developer/TraceInsights/Private/Insights/ViewModels/ThreadTimingTrack.cpp
+++ b/Engine/Source/Developer/TraceInsights/Private/Insights/ViewModels/ThreadTimingTrack.cpp
@@ -490,7 +490,7 @@ void FThreadTimingSharedState::Tick(Insights::ITimingViewSession& InSession, con
 				if (TimingProfilerProvider->GetGpu2TimelineIndex(GpuTimelineIndex))
 				{
 					Gpu2Track = MakeShared<FGpuTimingTrack>(*this, TEXT("GPU2"), nullptr, GpuTimelineIndex, FGpuTimingTrack::Gpu2ThreadId);
-					Gpu2Track->SetOrder(FTimingTrackOrder::Gpu);
+					Gpu2Track->SetOrder(FTimingTrackOrder::Gpu + 1);
 					Gpu2Track->SetVisibilityFlag(bShowHideAllGpuTracks);
 					InSession.AddScrollableTrack(Gpu2Track);
 				}
