[Insights]
  - Mem Insights: Merged tags with same name (but different id) when grouping by Tags in Mem Allocs table.
  - Mem Insights: Fixed grouping by Tags in Mem Allocs table when all tags are unknown/untagged.

diff --git a/Engine/Source/Developer/TraceInsights/Private/Insights/MemoryProfiler/ViewModels/MemAllocGroupingByTag.cpp b/Engine/Source/Developer/TraceInsights/Private/Insights/MemoryProfiler/ViewModels/MemAllocGroupingByTag.cpp
--- a/Engine/Source/Developer/TraceInsights/Private/Insights/MemoryProfiler/ViewModels/MemAllocGroupingByTag.cpp
+++ b/Engine/Source/Developer/TraceInsights/Private/Insights/MemoryProfiler/ViewModels/MemAllocGroupingByTag.cpp
@@ -6,6 +6,8 @@
 
 #define LOCTEXT_NAMESPACE "Insights::FMemAllocGroupingByTag"
 
+#define INSIGHTS_MERGE_MEM_TAGS_BY_NAME 1
+
 namespace Insights
 {
 
@@ -36,28 +38,148 @@ void FMemAllocGroupingByTag::GroupNodes(const TArray<FTableTreeNodePtr>& Nodes,
 	
 	ParentGroup.ClearChildren();
 
-	struct FTagParentRel
+	constexpr TagIdType InvalidTagId = ~0u;
+	struct FTagNode
 	{
+		FTagNode(FName InName, TagIdType InId, TagIdType InParentId, FTableTreeNodePtr InTreeNode)
+			: Name(InName), Id(InId), ParentId(InParentId), TreeNode(InTreeNode), Parent(nullptr), AllocCount(0)
+		{}
+
+		void MergeTagNodesByName(TMap<TagIdType, FTagNode*>& IdToNodeMap)
+		{
+			const int32 ChildrenCount = Children.Num();
+			for (int32 Index = 0; Index < ChildrenCount; ++Index)
+			{
+				FTagNode* CurrentNode = Children[Index];
+				for (int32 TargetIndex = 0; TargetIndex < Index; ++TargetIndex)
+				{
+					FTagNode* TargetNode = Children[TargetIndex];
+					if (TargetNode->Name == CurrentNode->Name)
+					{
+						// Move children from CurrentNode to TargetNode.
+						for (FTagNode* Child : CurrentNode->Children)
+						{
+							TargetNode->Children.Add(Child);
+							Child->Parent = TargetNode;
+						}
+						CurrentNode->Children.Reset();
+
+						// Allocs with CurrentNode's tag id will go to TargetNode.
+						IdToNodeMap.Emplace(CurrentNode->Id, TargetNode);
+
+						break;
+					}
+				}
+			}
+			for (FTagNode* Child : Children)
+			{
+				Child->MergeTagNodesByName(IdToNodeMap);
+			}
+		};
+
+		FName Name;
+
 		TagIdType Id;
-		TagIdType Parent;
+		TagIdType ParentId;
+		FTableTreeNodePtr TreeNode;
+
+		FTagNode* Parent;
+		TArray<FTagNode*> Children;
+		uint64 AllocCount;
 	};
-	TMap<TagIdType,FTableTreeNodePtr> TagIdToNode;
-	TArray<FTagParentRel> TagParentFixup;
-	TSet<TagIdType> NonEmptyIds;
+	TArray<FTagNode*> TagNodes;
+	TMap<TagIdType, FTagNode*> IdToNodeMap;
 
-	// Create tag nodes
+	class FScopedMemory
+	{
+	public:
+		FScopedMemory(TArray<FTagNode*>& InTagNodes)
+			: TagNodes(InTagNodes)
+		{
+		}
+		~FScopedMemory()
+		{
+			for (FTagNode* TagNode : TagNodes)
+			{
+				delete TagNode;
+			}
+		}
+
+	private:
+		TArray<FTagNode*>& TagNodes;
+	};
+	FScopedMemory _(TagNodes);
+
+	// Create tag nodes.
 	TagProvider.EnumerateTags([&](const TCHAR* Name, const TCHAR* FullName, TagIdType Id, TagIdType ParentId)
 	{
-		const auto Node = MakeShared<FTableTreeNode>(FName(Name), InParentTable);
-		TagIdToNode.Emplace(Id, Node);
-		TagParentFixup.Emplace(FTagParentRel {Id, ParentId});
+		const FName NodeName(Name);
+#if INSIGHTS_MERGE_MEM_TAGS_BY_NAME
+		FTableTreeNodePtr TreeNode = MakeShared<FTableTreeNode>(NodeName, InParentTable);
+#else
+		const FName NodeNameEx(Name, int32(Id));
+		FTableTreeNodePtr TreeNode = MakeShared<FTableTreeNode>(NodeNameEx, InParentTable);
+#endif
+		FTagNode* TagNode = new FTagNode(NodeName, Id, ParentId, TreeNode);
+		TagNodes.Add(TagNode);
+		IdToNodeMap.Add(Id, TagNode);
 	});
 
-	// Untagged node is used if we cant find the id.
-	const auto UntaggedNode = TagIdToNode.Find(0);
-	NonEmptyIds.Add(0);
+	// Create the Untagged node (if not already present).
+	constexpr TagIdType UntaggedNodeId = 0;
+	FTagNode* UntaggedNode = IdToNodeMap.FindRef(UntaggedNodeId);
+	if (!UntaggedNode)
+	{
+		const FName NodeName(TEXT("Untagged"));
+		FTableTreeNodePtr UntaggedTreeNode = MakeShared<FTableTreeNode>(NodeName, InParentTable);
+		UntaggedNode = new FTagNode(NodeName, UntaggedNodeId, InvalidTagId, UntaggedTreeNode);
+		TagNodes.Add(UntaggedNode);
+		IdToNodeMap.Add(UntaggedNodeId, UntaggedNode);
+	}
+
+	// Create the Unknown tag node.
+	constexpr TagIdType UnknownTagNodeId = ~0u - 1;
+	FTagNode* UnknownTagNode = IdToNodeMap.FindRef(UnknownTagNodeId);
+	if (!UnknownTagNode)
+	{
+		const FName NodeName(TEXT("Unknown"));
+		FTableTreeNodePtr UnknownTagTreeNode = MakeShared<FTableTreeNode>(NodeName, InParentTable);
+		UnknownTagNode = new FTagNode(NodeName, UnknownTagNodeId, InvalidTagId, UnknownTagTreeNode);
+		TagNodes.Add(UnknownTagNode);
+		IdToNodeMap.Add(UnknownTagNodeId, UnknownTagNode);
+	}
+
+	// Create the virtual tree of tag group nodes.
+	FTagNode Root(FName(TEXT("Root")), InvalidTagId, InvalidTagId, nullptr);
+	for (FTagNode* TagNode : TagNodes)
+	{
+		if (TagNode->ParentId == InvalidTagId)
+		{
+			TagNode->Parent = &Root;
+		}
+		else
+		{
+			TagNode->Parent = IdToNodeMap.FindRef(TagNode->ParentId);
+			if (!TagNode->Parent)
+			{
+				TagNode->Parent = UnknownTagNode;
+			}
+		}
+
+		TagNode->Parent->Children.Add(TagNode);
+	}
 
-	// Add nodes to correct tag group
+	if (InAsyncOperationProgress.ShouldCancelAsyncOp())
+	{
+		return;
+	}
+
+#if INSIGHTS_MERGE_MEM_TAGS_BY_NAME
+	// Merge tag group nodes (with same parent) by name.
+	Root.MergeTagNodesByName(IdToNodeMap);
+#endif
+
+	// Add nodes to the correct tag group.
 	for (FTableTreeNodePtr NodePtr : Nodes)
 	{
 		if (InAsyncOperationProgress.ShouldCancelAsyncOp())
@@ -76,57 +198,62 @@ void FMemAllocGroupingByTag::GroupNodes(const TArray<FTableTreeNodePtr>& Nodes,
 		if (Alloc)
 		{
 			const TagIdType TagId = Alloc->GetTagId();
-			const auto TagNode = TagIdToNode.Find(TagId);
-			if (TagNode)
-			{
-				(*TagNode)->AddChildAndSetGroupPtr(NodePtr);
-				NonEmptyIds.Add(TagId);
-			}
-			else if (UntaggedNode)
+			FTagNode* TagNode = IdToNodeMap.FindRef(TagId);
+			if (!TagNode)
 			{
-				(*UntaggedNode)->AddChildAndSetGroupPtr(NodePtr);
+				TagNode = UnknownTagNode;
 			}
+			TagNode->TreeNode->AddChildAndSetGroupPtr(NodePtr);
+			TagNode->AllocCount++;
+		}
+		else
+		{
+			ParentGroup.AddChildAndSetGroupPtr(NodePtr);
 		}
 	}
 
-	// Sort by parent id
-	Algo::SortBy(TagParentFixup, &FTagParentRel::Parent);
-	
-	// Fixup parent relationships and filter out empty nodes
-	for (auto TagParentRel : TagParentFixup)
+	if (InAsyncOperationProgress.ShouldCancelAsyncOp())
 	{
-		if (InAsyncOperationProgress.ShouldCancelAsyncOp())
-		{
-			return;
-		}
-		
-		if (!NonEmptyIds.Contains(TagParentRel.Id))
-		{
-			continue;
-		}
-		
-		if (TagParentRel.Parent == ~0)
+		return;
+	}
+
+	for (FTagNode* TagNode : TagNodes)
+	{
+		FTagNode* Parent = TagNode->Parent;
+		while (Parent)
 		{
-			// Add to parent directly
-			ParentGroup.AddChildAndSetGroupPtr(TagIdToNode[TagParentRel.Id]);
+			Parent->AllocCount += TagNode->AllocCount;
+			Parent = Parent->Parent;
 		}
-		else
+	}
+
+	if (InAsyncOperationProgress.ShouldCancelAsyncOp())
+	{
+		return;
+	}
+
+	// Fixup parent relationships and filter out empty nodes.
+	for (FTagNode* TagNode : TagNodes)
+	{
+		if (TagNode->AllocCount > 0)
 		{
-			// Try to find the parent node and add it as a child
-			if (const auto ParentNode = TagIdToNode.Find(TagParentRel.Parent))
+			check(TagNode->TreeNode);
+			if (TagNode->Parent->TreeNode)
 			{
-				(*ParentNode)->AddChildAndSetGroupPtr(TagIdToNode[TagParentRel.Id]);
+				TagNode->Parent->TreeNode->AddChildAndSetGroupPtr(TagNode->TreeNode);
+			}
+			else
+			{
+				check(TagNode->Parent == &Root);
+				ParentGroup.AddChildAndSetGroupPtr(TagNode->TreeNode);
 			}
-
-			// Indicate that parent also is non-empty. Since the fixup array is sorted by parent id the root nodes
-			// be processed last.
-			NonEmptyIds.Add(TagParentRel.Parent);
 		}
 	}
 }
 	
 ////////////////////////////////////////////////////////////////////////////////////////////////////
 
+#undef INSIGHTS_MERGE_MEM_TAGS_BY_NAME
 #undef LOCTEXT_NAMESPACE
 
-}
+} // namespace Insights
