TasksInsights: selecting ParallelFor timer draws arrows to and from its  ParallelFor tasks

diff --git a/Engine/Source/Developer/TraceInsights/Private/Insights/TaskGraphProfiler/ViewModels/TaskTimingTrack.cpp b/Engine/Source/Developer/TraceInsights/Private/Insights/TaskGraphProfiler/ViewModels/TaskTimingTrack.cpp
--- a/Engine/Source/Developer/TraceInsights/Private/Insights/TaskGraphProfiler/ViewModels/TaskTimingTrack.cpp
+++ b/Engine/Source/Developer/TraceInsights/Private/Insights/TaskGraphProfiler/ViewModels/TaskTimingTrack.cpp
@@ -1067,6 +1067,7 @@ void FTaskTimingTrack::GetEventRelations(const FThreadTrackEvent& InSelectedEven
 {
 	const int32 MaxTasksToShow = 30;
 	double StartTime = InSelectedEvent.GetStartTime();
+	double EndTime = InSelectedEvent.GetEndTime();
 
 	TSharedPtr<const TraceServices::IAnalysisSession> Session = FInsightsManager::Get()->GetSession();
 	const TraceServices::ITasksProvider* TasksProvider = TraceServices::ReadTasksProvider(*Session.Get());
@@ -1104,6 +1105,17 @@ void FTaskTimingTrack::GetEventRelations(const FThreadTrackEvent& InSelectedEven
 				}
 			}
 		}
+
+		TArray<TaskTrace::FId> ParallelForTasks = TasksProvider->TryGetParallelForTasks(Timer->Name, ThreadId, StartTime, EndTime);
+		int32 NumTasksToShow = FMath::Min(ParallelForTasks.Num(), MaxTasksToShow);
+		for (int32 TaskIndex = 0; TaskIndex != NumTasksToShow; ++TaskIndex)
+		{
+			const TraceServices::FTaskInfo* Task = TasksProvider->TryGetTask(ParallelForTasks[TaskIndex]);
+			int32 TaskExecutionDepth = FTaskGraphProfilerManager::Get()->GetDepthOfTaskExecution(Task->StartedTimestamp, Task->FinishedTimestamp, Task->StartedThreadId);
+
+			FTaskGraphProfilerManager::Get()->AddRelation(&InSelectedEvent, Task->StartedTimestamp, ThreadId, -1, Task->StartedTimestamp, Task->StartedThreadId, TaskExecutionDepth, ETaskEventType::NestedStarted);
+			FTaskGraphProfilerManager::Get()->AddRelation(&InSelectedEvent, Task->CompletedTimestamp, Task->CompletedThreadId, Task->CompletedTimestamp, ThreadId, ETaskEventType::NestedCompleted);
+		}
 	}
 }
 
diff --git a/Engine/Source/Developer/TraceServices/Private/Model/TasksProfiler.cpp b/Engine/Source/Developer/TraceServices/Private/Model/TasksProfiler.cpp
--- a/Engine/Source/Developer/TraceServices/Private/Model/TasksProfiler.cpp
+++ b/Engine/Source/Developer/TraceServices/Private/Model/TasksProfiler.cpp
@@ -518,6 +518,29 @@ namespace TraceServices
 		return Waiting.FinishedTimestamp > Timestamp || Waiting.FinishedTimestamp == FTaskInfo::InvalidTimestamp ? &Waiting : nullptr;
 	}
 
+	TArray<TaskTrace::FId> FTasksProvider::TryGetParallelForTasks(const TCHAR* TimerName, uint32 ThreadId, double StartTime, double EndTime) const
+	{
+		TArray<TaskTrace::FId> Result;
+
+		if (FCString::Strcmp(TimerName, TEXT("ParallelFor")) != 0)
+		{
+			return Result;
+		}
+
+		EnumerateTasks(StartTime, EndTime, 
+			[ThreadId, StartTime, EndTime, &Result] (const FTaskInfo& Task)
+			{
+				if (Task.LaunchedThreadId == ThreadId && Task.LaunchedTimestamp > StartTime && Task.FinishedTimestamp < EndTime)
+				{
+					Result.Add(Task.Id);
+				}
+				return ETaskEnumerationResult::Continue;
+			}
+		);
+
+		return Result;
+	}
+
 	int64 FTasksProvider::GetNumTasks() const
 	{
 		return Tasks.Num();
diff --git a/Engine/Source/Developer/TraceServices/Private/Model/TasksProfilerPrivate.h b/Engine/Source/Developer/TraceServices/Private/Model/TasksProfilerPrivate.h
--- a/Engine/Source/Developer/TraceServices/Private/Model/TasksProfilerPrivate.h
+++ b/Engine/Source/Developer/TraceServices/Private/Model/TasksProfilerPrivate.h
@@ -40,6 +40,7 @@ namespace TraceServices
 		virtual const FTaskInfo* TryGetTask(uint32 ThreadId, double Timestamp) const override;
 		virtual const FTaskInfo* TryGetTask(TaskTrace::FId TaskId) const;
 		virtual const FWaitingForTasks* TryGetWaiting(const TCHAR* TimerName, uint32 ThreadId, double Timestamp) const override;
+		virtual TArray<TaskTrace::FId> TryGetParallelForTasks(const TCHAR* TimerName, uint32 ThreadId, double StartTime, double EndTime) const override;
 		virtual int64 GetNumTasks() const override;
 		virtual void EnumerateTasks(double StartTime, double EndTime, TaskCallback Callback) const override;
 
diff --git a/Engine/Source/Developer/TraceServices/Public/TraceServices/Model/TasksProfiler.h b/Engine/Source/Developer/TraceServices/Public/TraceServices/Model/TasksProfiler.h
--- a/Engine/Source/Developer/TraceServices/Public/TraceServices/Model/TasksProfiler.h
+++ b/Engine/Source/Developer/TraceServices/Public/TraceServices/Model/TasksProfiler.h
@@ -100,6 +100,9 @@ namespace TraceServices
 		// returns an info about waiting for tasks completion for given thread and timestamp, if any, otherwise `nullptr`
 		virtual const FWaitingForTasks* TryGetWaiting(const TCHAR* TimerName, uint32 ThreadId, double Timestamp) const = 0;
 
+		// returns an info about ParallelFor tasks for given timer
+		virtual TArray<TaskTrace::FId> TryGetParallelForTasks(const TCHAR* TimerName, uint32 ThreadId, double StartTime, double EndTime) const = 0;
+
 		// returns the number of tasks stored in the provider
 		virtual int64 GetNumTasks() const = 0;
 
