[Insights] Table Import Tool initial implementation

diff --git a/Engine/Source/Developer/TraceInsights/Private/Insights/Common/InsightsMenuBuilder.cpp b/Engine/Source/Developer/TraceInsights/Private/Insights/Common/InsightsMenuBuilder.cpp
--- a/Engine/Source/Developer/TraceInsights/Private/Insights/Common/InsightsMenuBuilder.cpp
+++ b/Engine/Source/Developer/TraceInsights/Private/Insights/Common/InsightsMenuBuilder.cpp
@@ -8,6 +8,7 @@
 #include "WorkspaceMenuStructure.h"
 #include "WorkspaceMenuStructureModule.h"
 
+#include "Insights/ImportTool/TableImportTool.h"
 #include "Insights/InsightsStyle.h"
 #include "Insights/InsightsManager.h"
 
@@ -65,6 +66,11 @@ void FInsightsMenuBuilder::PopulateMenu(FMenuBuilder& MenuBuilder)
 		false,
 		FSlateIcon(FAppStyle::Get().GetStyleSetName(), "Icons.FolderOpen")
 	);
+	MenuBuilder.AddMenuEntry(
+		LOCTEXT("ImportTable", "Import Table"),
+		LOCTEXT("ImportTable_ToolTip", "Import CSV or TSV data from a file to an Insights Table."),
+		FSlateIcon(FInsightsStyle::GetStyleSetName(), "Icons.Import.Table"),
+		FUIAction(FExecuteAction::CreateLambda([] { Insights::FTableImportTool::Get()->StartImportProcess(); })));
 	MenuBuilder.AddMenuEntry(
 		LOCTEXT("AutoOpenLiveTrace", "Auto Open Live Trace"),
 		LOCTEXT("AutoOpenLiveTrace_ToolTip", "If enabled, the analysis starts automatically for each new live trace session, replacing the current analysis session."),
diff --git a/Engine/Source/Developer/TraceInsights/Private/Insights/ImportTool/TableImportTool.cpp b/Engine/Source/Developer/TraceInsights/Private/Insights/ImportTool/TableImportTool.cpp
new file mode 100644
--- /dev/null
+++ b/Engine/Source/Developer/TraceInsights/Private/Insights/ImportTool/TableImportTool.cpp
@@ -0,0 +1,223 @@
+// Copyright Epic Games, Inc. All Rights Reserved.
+
+#include "TableImportTool.h"
+
+#include "DesktopPlatformModule.h"
+#include "Framework/Application/SlateApplication.h"
+#include "Framework/Docking/TabManager.h"
+#include "Framework/Docking/WorkspaceItem.h"
+#include "Logging/MessageLog.h"
+#include "MessageLog/Public/MessageLogModule.h"
+#include "Misc/FileHelper.h"
+#include "SlateOptMacros.h"
+#include "TraceServices/Model/Threads.h"
+#include "Widgets/Docking/SDockTab.h"
+
+// Insights
+#include "Insights/Common/Stopwatch.h"
+#include "Insights/InsightsManager.h"
+#include "Insights/InsightsStyle.h"
+#include "Insights/Log.h"
+#include "Insights/Table/ViewModels/UntypedTable.h"
+#include "Insights/Table/Widgets/SUntypedTableTreeView.h"
+
+#define LOCTEXT_NAMESPACE "ImportTableTool"
+
+namespace Insights
+{
+
+////////////////////////////////////////////////////////////////////////////////////////////////////
+
+TSharedPtr<FTableImportTool> FTableImportTool::Instance;
+
+////////////////////////////////////////////////////////////////////////////////////////////////////
+
+FTableImportTool::FTableImportTool()
+{
+}
+
+////////////////////////////////////////////////////////////////////////////////////////////////////
+
+FTableImportTool::~FTableImportTool()
+{
+}
+
+////////////////////////////////////////////////////////////////////////////////////////////////////
+
+TSharedPtr<FTableImportTool> FTableImportTool::CreateInstance()
+{
+	ensure(Instance == nullptr);
+	Instance = MakeShared<FTableImportTool>();
+
+	return Instance;
+}
+
+////////////////////////////////////////////////////////////////////////////////////////////////////
+
+TSharedPtr<FTableImportTool> FTableImportTool::Get()
+{
+	return Instance;
+}
+
+////////////////////////////////////////////////////////////////////////////////////////////////////
+
+void FTableImportTool::StartImportProcess()
+{
+	TArray<FString> Filenames;
+	bool bDialogResult = false;
+
+	IDesktopPlatform* DesktopPlatform = FDesktopPlatformModule::Get();
+	if (DesktopPlatform)
+	{
+		const FString DefaultPath = FPaths::ProjectSavedDir();
+		const FString DefaultFile = TEXT("");
+		bDialogResult = DesktopPlatform->OpenFileDialog(
+			FSlateApplication::Get().FindBestParentWindowHandleForDialogs(nullptr),
+			LOCTEXT("ImportTableTitle", "Import Table").ToString(),
+			DefaultPath,
+			DefaultFile,
+			TEXT("Tab-Separated Values (*.tsv)|*.tsv|Text Files (*.txt)|*.txt|Comma-Separated Values (*.csv)|*.csv|All Files (*.*)|*.*"),
+			EFileDialogFlags::None,
+			Filenames
+		);
+	}
+
+	if (!bDialogResult)
+	{
+		return;
+	}
+
+	if (Filenames.Num() != 1)
+	{
+		FMessageLog ReportMessageLog(FInsightsManager::Get()->GetLogListingName());
+		ReportMessageLog.Error(LOCTEXT("NotSingleFile", "A single file can be imported at a time!"));
+		ReportMessageLog.Notify();
+		return;
+	}
+
+	using namespace TraceServices;
+
+	FName TableId = GetTableID(Filenames[0]);
+	DisplayTable(TableId);
+
+	FTableImportService::ImportTable(Filenames[0], TableId, [this](FName TableViewID, TSharedPtr<TraceServices::ITable<TraceServices::FImportTableRow>> Table) {
+		this->TableImportServiceCallback(TableViewID, Table);
+	});
+}
+
+////////////////////////////////////////////////////////////////////////////////////////////////////
+
+void FTableImportTool::DisplayTable(FName TableViewID)
+{
+	FText ImportTableTreeViewTabDisplayName;
+	if (TableViewID.GetNumber())
+	{
+		ImportTableTreeViewTabDisplayName = FText::Format(LOCTEXT("TableImportTool.ImportTableTreeViewTabDisplayName", "{0}({1})"),
+														  FText::FromString(TableViewID.GetPlainNameString()),
+														  FText::AsNumber(TableViewID.GetNumber()));
+	}
+	else
+	{
+		ImportTableTreeViewTabDisplayName = FText::FromString(TableViewID.GetPlainNameString());
+	}
+
+	FGlobalTabmanager::Get()->RegisterNomadTabSpawner(TableViewID, FOnSpawnTab::CreateSP(this, &FTableImportTool::SpawnTab_TableImportTreeView, TableViewID, ImportTableTreeViewTabDisplayName))
+		.SetDisplayName(ImportTableTreeViewTabDisplayName)
+		.SetIcon(FSlateIcon(FInsightsStyle::GetStyleSetName(), "Icons.TableTreeView"));
+
+	FGlobalTabmanager::Get()->TryInvokeTab(TableViewID);
+}
+
+////////////////////////////////////////////////////////////////////////////////////////////////////
+
+void FTableImportTool::TableImportServiceCallback(FName TableViewID, TSharedPtr<TraceServices::ITable<TraceServices::FImportTableRow>> Table)
+{
+	const FOpenImportedTableTabData* TableViewPtr = OpenTablesMap.Find(TableViewID);
+	if (TableViewPtr)
+	{
+		TSharedPtr<SUntypedTableTreeView> TableView = TableViewPtr->TableTreeView;
+		TableView->UpdateSourceTable(Table);
+	}
+}
+
+////////////////////////////////////////////////////////////////////////////////////////////////////
+
+BEGIN_SLATE_FUNCTION_BUILD_OPTIMIZATION
+TSharedRef<SDockTab> FTableImportTool::SpawnTab_TableImportTreeView(const FSpawnTabArgs& Args, FName TableViewID, FText InDisplayName)
+{
+	TSharedRef<FUntypedTable> UntypedTable = MakeShared<Insights::FUntypedTable>();
+	UntypedTable->Reset();
+	UntypedTable->SetDisplayName(InDisplayName);
+
+	TSharedPtr<Insights::SUntypedTableTreeView> TableTreeView;
+
+	const TSharedRef<SDockTab> DockTab = SNew(SDockTab)
+		.ShouldAutosize(false)
+		.TabRole(ETabRole::PanelTab)
+		[
+			SAssignNew(TableTreeView, SUntypedTableTreeView, UntypedTable)
+		];
+
+	FOpenImportedTableTabData OpenTabData;
+	OpenTabData.TableTreeView = TableTreeView;
+	OpenTabData.Tab = DockTab;
+
+	OpenTablesMap.Add(TableViewID, OpenTabData);
+	DockTab->SetOnTabClosed(SDockTab::FOnTabClosedCallback::CreateSP(this, &FTableImportTool::OnTableImportTreeViewTabClosed));
+
+	return DockTab;
+}
+END_SLATE_FUNCTION_BUILD_OPTIMIZATION
+
+////////////////////////////////////////////////////////////////////////////////////////////////////
+
+void FTableImportTool::OnTableImportTreeViewTabClosed(TSharedRef<SDockTab> TabBeingClosed)
+{
+	TSharedRef<SUntypedTableTreeView> TableTreeView = StaticCastSharedRef<SUntypedTableTreeView>(TabBeingClosed->GetContent());
+	FOpenImportedTableTabData OpenTabData;
+	OpenTabData.TableTreeView = TableTreeView;
+	OpenTabData.Tab = TabBeingClosed;
+	const FName* TableIdPtr = OpenTablesMap.FindKey(OpenTabData);
+	check(TableIdPtr);
+
+	FGlobalTabmanager::Get()->UnregisterNomadTabSpawner(*TableIdPtr);
+	TableTreeView->OnClose();
+	OpenTablesMap.Remove(*TableIdPtr);
+}
+
+////////////////////////////////////////////////////////////////////////////////////////////////////
+
+FName FTableImportTool::GetTableID(const FString& Path)
+{
+	FName TableId(Path);
+	while (OpenTablesMap.Contains(TableId))
+	{
+		TableId.SetNumber(TableId.GetNumber() + 1);
+	}
+
+	return TableId;
+}
+
+////////////////////////////////////////////////////////////////////////////////////////////////////
+
+void FTableImportTool::OnWindowClosedEvent()
+{
+	CloseAllOpenTabs();
+}
+
+////////////////////////////////////////////////////////////////////////////////////////////////////
+
+void FTableImportTool::CloseAllOpenTabs()
+{
+	while (!OpenTablesMap.IsEmpty())
+	{
+		TSharedPtr<SDockTab> Tab = OpenTablesMap.CreateConstIterator().Value().Tab;
+		Tab->RequestCloseTab();
+	}
+}
+
+////////////////////////////////////////////////////////////////////////////////////////////////////
+
+} // namespace Insights
+
+#undef LOCTEXT_NAMESPACE
diff --git a/Engine/Source/Developer/TraceInsights/Private/Insights/ImportTool/TableImportTool.h b/Engine/Source/Developer/TraceInsights/Private/Insights/ImportTool/TableImportTool.h
new file mode 100644
--- /dev/null
+++ b/Engine/Source/Developer/TraceInsights/Private/Insights/ImportTool/TableImportTool.h
@@ -0,0 +1,74 @@
+// Copyright Epic Games, Inc. All Rights Reserved.
+
+#pragma once
+
+#include "CoreMinimal.h"
+
+#include "TraceServices/Model/TableImport.h"
+#include "TraceServices/Containers/Tables.h"
+
+#include "Insights/IUnrealInsightsModule.h"
+
+class SDockTab;
+class FSpawnTabArgs;
+
+namespace Insights
+{
+
+class SUntypedTableTreeView;
+
+////////////////////////////////////////////////////////////////////////////////////////////////////
+
+struct FOpenImportedTableTabData
+{
+	TSharedPtr<SDockTab> Tab;
+	TSharedPtr<SUntypedTableTreeView> TableTreeView;
+
+	bool operator ==(const FOpenImportedTableTabData& Other) const
+	{
+		return Tab == Other.Tab && TableTreeView == Other.TableTreeView;
+	}
+};
+
+////////////////////////////////////////////////////////////////////////////////////////////////////
+
+class FTableImportTool : public TSharedFromThis<FTableImportTool>, public IInsightsComponent
+{
+public:
+	/** Default constructor. */
+	FTableImportTool();
+
+	/** Destructor. */
+	virtual ~FTableImportTool();
+
+	/** Creates an instance of the Table Import Tool. */
+	static TSharedPtr<FTableImportTool> CreateInstance();
+
+	static TSharedPtr<FTableImportTool> Get();
+
+	// IInsightsComponent
+	virtual void Initialize(IUnrealInsightsModule& InsightsModule) override {}
+	virtual void Shutdown() override {}
+	virtual void RegisterMajorTabs(IUnrealInsightsModule& InsightsModule) override {}
+	virtual void UnregisterMajorTabs() override {}
+	virtual void OnWindowClosedEvent() override;
+
+	void StartImportProcess();
+
+	void DisplayTable(FName TableViewID);
+	void TableImportServiceCallback(FName TableViewID, TSharedPtr<TraceServices::ITable<TraceServices::FImportTableRow>> Table);
+
+	TSharedRef<SDockTab> SpawnTab_TableImportTreeView(const FSpawnTabArgs& Args, FName TableViewID, FText InDisplayName);
+	void OnTableImportTreeViewTabClosed(TSharedRef<SDockTab> TabBeingClosed);
+
+	void CloseAllOpenTabs();
+private:
+
+	FName GetTableID(const FString& Path);
+	static TSharedPtr<FTableImportTool> Instance;
+	TMap<FName, FOpenImportedTableTabData> OpenTablesMap;
+};
+
+////////////////////////////////////////////////////////////////////////////////////////////////////
+
+} // namespace Insights
diff --git a/Engine/Source/Developer/TraceInsights/Private/Insights/InsightsStyle.cpp b/Engine/Source/Developer/TraceInsights/Private/Insights/InsightsStyle.cpp
--- a/Engine/Source/Developer/TraceInsights/Private/Insights/InsightsStyle.cpp
+++ b/Engine/Source/Developer/TraceInsights/Private/Insights/InsightsStyle.cpp
@@ -333,6 +333,8 @@ void FInsightsStyle::FStyle::Initialize()
 	Set("Icons.NetEvent.TreeItem", new TODO_IMAGE_BRUSH(Icon12x12));
 	Set("Icons.MemTag.TreeItem", new TODO_IMAGE_BRUSH(Icon12x12));
 
+	Set("Icons.Import.Table", new TODO_IMAGE_BRUSH(Icon12x12));
+
 	Set("TreeTable.TooltipBold", FTextBlockStyle(NormalText)
 		.SetFont(DEFAULT_FONT("Bold", 8))
 		.SetColorAndOpacity(FLinearColor(0.5f, 0.5f, 0.5f, 1.0f))
diff --git a/Engine/Source/Developer/TraceInsights/Private/Insights/TraceInsightsModule.cpp b/Engine/Source/Developer/TraceInsights/Private/Insights/TraceInsightsModule.cpp
--- a/Engine/Source/Developer/TraceInsights/Private/Insights/TraceInsightsModule.cpp
+++ b/Engine/Source/Developer/TraceInsights/Private/Insights/TraceInsightsModule.cpp
@@ -19,6 +19,7 @@
 // Insights
 #include "Insights/ContextSwitches/ContextSwitchesProfilerManager.h"
 #include "Insights/CookProfiler/CookProfilerManager.h"
+#include "Insights/ImportTool/TableImportTool.h"
 #include "Insights/InsightsManager.h"
 #include "Insights/InsightsStyle.h"
 #include "Insights/IUnrealInsightsModule.h"
@@ -69,6 +70,7 @@
 	RegisterComponent(Insights::FTaskGraphProfilerManager::CreateInstance());
 	RegisterComponent(Insights::FContextSwitchesProfilerManager::CreateInstance());
 	RegisterComponent(Insights::FCookProfilerManager::CreateInstance());
+	RegisterComponent(Insights::FTableImportTool::CreateInstance());
 
 #if !WITH_EDITOR
 	ISourceCodeAccessModule& SourceCodeAccessModule = FModuleManager::LoadModuleChecked<ISourceCodeAccessModule>("SourceCodeAccess");
diff --git a/Engine/Source/Developer/TraceServices/Private/Model/TableImportTask.cpp b/Engine/Source/Developer/TraceServices/Private/Model/TableImportTask.cpp
new file mode 100644
--- /dev/null
+++ b/Engine/Source/Developer/TraceServices/Private/Model/TableImportTask.cpp
@@ -0,0 +1,198 @@
+// Copyright Epic Games, Inc. All Rights Reserved.
+
+#include "TableImportTask.h"
+
+#include "Async/TaskGraphInterfaces.h"
+#include "Misc/FileHelper.h"
+#include "Tasks/Task.h"
+
+namespace TraceServices
+{
+
+FTableImportTask::FTableImportTask(const FString& InFilePath, FName InTableId, FTableImportService::TableImportCallback InCallback)
+	: Callback(InCallback)
+	, FilePath(InFilePath)
+	, TableId(InTableId)
+{
+}
+
+FTableImportTask::~FTableImportTask()
+{
+}
+
+void FTableImportTask::operator()()
+{
+	TArray<FString> Lines;
+	FFileHelper::LoadFileToStringArray(Lines, *FilePath);
+
+	if (FilePath.EndsWith(TEXT(".csv")))
+	{
+		Separator = TEXT(",");
+	}
+	else if (FilePath.EndsWith(TEXT(".tsv")))
+	{
+		Separator = TEXT("\t");
+	}
+
+	Table = MakeShared<TImportTable<FImportTableRow>>();
+
+	if (!ParseHeader(Lines[0]))
+	{
+		return;
+	}
+
+	if (!CreateLayout(Lines[1]))
+	{
+		return;
+	}
+
+	if (!ParseData(Lines))
+	{
+		return;
+	}
+
+	FFunctionGraphTask::CreateAndDispatchWhenReady(
+		[Callback = this->Callback, Table = this->Table, TableId = this->TableId]()
+		{
+			Callback(TableId, Table);
+		},
+		TStatId(),
+		nullptr,
+		ENamedThreads::GameThread);
+}
+
+bool FTableImportTask::ParseHeader(const FString& HeaderLine)
+{
+	HeaderLine.ParseIntoArray(ColumnNames, *Separator);
+
+	return ColumnNames.Num() > 0;
+}
+
+bool FTableImportTask::CreateLayout(const FString& Line)
+{
+	TArray<FString> Values;
+	SplitLineIntoValues(Line, Values);
+
+	TTableLayout<FImportTableRow>& Layout = Table->EditLayout();
+	for (int32 Index = 0; Index < Values.Num(); ++Index)
+	{
+		auto ProjectorFunc = [Index](const FImportTableRow& Row)
+		{
+			return Row.GetValue(Index);
+		};
+
+		if (Values[Index].IsNumeric())
+		{
+			if (Values[Index].Contains(TEXT(".")))
+			{
+				Layout.AddColumn<double>(*ColumnNames[Index], ProjectorFunc);
+			}
+			else
+			{
+				Layout.AddColumn<uint32>(*ColumnNames[Index], ProjectorFunc);
+			}
+		}
+		else
+		{
+			Layout.AddColumn<const TCHAR*>(*ColumnNames[Index], ProjectorFunc);
+		}
+	}
+
+	return true;
+}
+
+bool FTableImportTask::ParseData(TArray<FString>& Lines)
+{
+	TTableLayout<FImportTableRow>& Layout = Table->EditLayout();
+	bool Restart = false;
+
+	for (int32 LineIndex = 1; LineIndex < Lines.Num(); ++LineIndex)
+	{
+		TArray<FString> Values;
+		SplitLineIntoValues(Lines[LineIndex], Values);
+
+		FImportTableRow& NewRow = Table->AddRow();
+		NewRow.SetNumValues(Values.Num());
+		for (int32 ValueIndex = 0; ValueIndex < Values.Num(); ++ValueIndex)
+		{
+			ETableColumnType ColumnType = Layout.GetColumnType(ValueIndex);
+			const TCHAR* Value = *Values[ValueIndex];
+			if (ColumnType == ETableColumnType::TableColumnType_CString)
+			{
+				const TCHAR* StoredValue = Table->GetStringStore().Store(Value);
+				NewRow.SetValue(ValueIndex, StoredValue);
+			}
+			else if (ColumnType == ETableColumnType::TableColumnType_Double)
+			{
+				if (!Values[ValueIndex].IsNumeric())
+				{
+					Layout.SetColumnType(ValueIndex, ETableColumnType::TableColumnType_CString);
+					Restart = true;
+					break;
+				}
+				
+				NewRow.SetValue(ValueIndex, FCString::Atod(Value));
+			}
+			else if (ColumnType == ETableColumnType::TableColumnType_Int)
+			{
+				if (!Values[ValueIndex].IsNumeric())
+				{
+					Layout.SetColumnType(ValueIndex, ETableColumnType::TableColumnType_CString);
+					Restart = true;
+					break;
+				}
+				else if (Values[ValueIndex].Contains(TEXT(".")))
+				{
+					Layout.SetColumnType(ValueIndex, ETableColumnType::TableColumnType_Double);
+					Restart = true;
+					break;
+				}
+
+				NewRow.SetValue(ValueIndex, FCString::Atoi(Value));
+			}
+		}
+
+		if (Restart)
+		{
+			TSharedPtr<TImportTable<FImportTableRow >> NewTable = MakeShared<TImportTable<FImportTableRow>>();
+			NewTable->EditLayout() = Layout;
+			Table = NewTable;
+			ParseData(Lines);
+
+			break;
+		}
+	}
+
+	return true;
+}
+
+void FTableImportTask::SplitLineIntoValues(const FString& InLine, TArray<FString>& OutValues)
+{
+	//Parse the line into values. Separators inside quotes are ignored.
+	int Start = 0;
+	bool IsInQuotes = false;
+	int Index;
+	for (Index = 0; Index < InLine.Len(); ++Index)
+	{
+		if (InLine[Index] == TEXT('"') && (Index == 0 || InLine[Index - 1] != TEXT('\\')))
+		{
+			IsInQuotes = !IsInQuotes;
+		}
+		else if(!IsInQuotes && InLine[Index] == Separator[0])
+		{
+			FString Value = InLine.Mid(Start, Index - Start);
+			OutValues.Add(Value.TrimQuotes());
+			Start = Index + 1;
+		}
+	}
+
+	FString Value = InLine.Mid(Start, Index - Start);
+	OutValues.Add(Value.TrimQuotes());
+}
+
+void FTableImportService::ImportTable(const FString& InPath, FName TableId, FTableImportService::TableImportCallback InCallback)
+{
+	UE::Tasks::Launch(UE_SOURCE_LOCATION, FTableImportTask(InPath, TableId, InCallback));
+}
+
+} // namespace TraceServices
diff --git a/Engine/Source/Developer/TraceServices/Private/Model/TableImportTask.h b/Engine/Source/Developer/TraceServices/Private/Model/TableImportTask.h
new file mode 100644
--- /dev/null
+++ b/Engine/Source/Developer/TraceServices/Private/Model/TableImportTask.h
@@ -0,0 +1,73 @@
+// Copyright Epic Games, Inc. All Rights Reserved.
+
+#pragma once
+
+#include "AllocationsProvider.h"
+#include "Common/StringStore.h"
+#include "CoreMinimal.h"
+#include "Model/Tables.h"
+#include "TraceServices/Model/TableImport.h"
+
+namespace TraceServices
+{
+
+struct FColumnValueContainer;
+
+class FImportTableRow
+{
+public:
+	FColumnValueContainer GetValue(uint32 Index) const { return Values[Index]; }
+
+	template<typename T>
+	void SetValue(uint32 Index, T Value) { Values[Index] = Value; }
+	void SetNumValues(uint32 Size) { Values.AddUninitialized(Size); }
+
+private:
+	TArray<FColumnValueContainer> Values;
+};
+
+constexpr int GImportTableAllocatorSlabSize = 32 << 20;
+
+template<typename RowType>
+class TImportTable
+	: public TTable<RowType, GImportTableAllocatorSlabSize>
+{
+public:
+	TImportTable()
+		: TTable<RowType, GImportTableAllocatorSlabSize>()
+		, StringStore(TTable<RowType, GImportTableAllocatorSlabSize>::Allocator)
+	{
+	}
+	virtual ~TImportTable() {}
+
+	FStringStore& GetStringStore() { return StringStore; }
+
+private:
+	FStringStore StringStore;
+};
+
+class FTableImportTask
+{
+public:
+	FTableImportTask(const FString& InFilePath, FName InTableId, FTableImportService::TableImportCallback InCallback);
+	~FTableImportTask();
+
+	void operator()();
+
+	bool ParseHeader(const FString& HeaderLine);
+	bool CreateLayout(const FString& Line);
+	bool ParseData(TArray<FString>& Lines);
+
+private:
+	void SplitLineIntoValues(const FString& InLine, TArray<FString>& OutValues);
+
+	FTableImportService::TableImportCallback Callback;
+	TSharedPtr<TImportTable<FImportTableRow>> Table;
+
+	TArray<FString> ColumnNames;
+	FString FilePath;
+	FName TableId;
+	FString Separator;
+};
+
+} // namespace TraceServices
diff --git a/Engine/Source/Developer/TraceServices/Private/Model/Tables.h b/Engine/Source/Developer/TraceServices/Private/Model/Tables.h
--- a/Engine/Source/Developer/TraceServices/Private/Model/Tables.h
+++ b/Engine/Source/Developer/TraceServices/Private/Model/Tables.h
@@ -197,6 +197,17 @@ public:
 		return *this;
 	}
 
+	template<typename ColumnNativeType>
+	TTableLayout<RowType>& AddColumn(const TCHAR* ColumnName, TFunction<FColumnValueContainer(const RowType&)> Projector, uint32 DisplayHintFlags = 0)
+	{
+		Columns.Add({
+			ColumnName,
+			GetColumnTypeFromNativeType<ColumnNativeType>(),
+			DisplayHintFlags,
+			Projector });
+		return *this;
+	}
+
 	uint64 GetColumnCount() const override
 	{
 		return Columns.Num();
@@ -212,6 +223,11 @@ public:
 		return Columns[ColumnIndex].Type;
 	}
 
+	void SetColumnType(uint64 ColumnIndex, ETableColumnType ColumnType)
+	{
+		Columns[ColumnIndex].Type = ColumnType;
+	}
+
 	uint32 GetColumnDisplayHintFlags(uint64 ColumnIndex) const override
 	{
 		return Columns[ColumnIndex].DisplayHintFlags;
@@ -431,13 +447,13 @@ private:
 	const TPagedArray<RowType>& Rows;
 };
 
-template<typename RowType>
+template<typename RowType, int AllocatorSlabSize = 2 << 20>
 class TTable
 	: public TTableBase<RowType>
 {
 public:
 	TTable()
-		: Allocator(2 << 20)
+		: Allocator(AllocatorSlabSize)
 		, Rows(Allocator, 1024)
 	{
 
@@ -445,7 +461,7 @@ public:
 
 	TTable(TTableLayout<RowType> Layout)
 		: TTableBase<RowType>(Layout)
-		, Allocator(2 << 20)
+		, Allocator(AllocatorSlabSize)
 		, Rows(Allocator, 1024)
 	{
 
@@ -456,13 +472,15 @@ public:
 		return Rows.PushBack();
 	}
 
+protected:
+	FSlabAllocator Allocator;
+
 private:
 	virtual const TPagedArray<RowType>& GetRows() const override
 	{
 		return Rows;
 	}
 
-	FSlabAllocator Allocator;
 	TPagedArray<RowType> Rows;
 };
 
diff --git a/Engine/Source/Developer/TraceServices/Public/TraceServices/Model/TableImport.h b/Engine/Source/Developer/TraceServices/Public/TraceServices/Model/TableImport.h
new file mode 100644
--- /dev/null
+++ b/Engine/Source/Developer/TraceServices/Public/TraceServices/Model/TableImport.h
@@ -0,0 +1,26 @@
+// Copyright Epic Games, Inc. All Rights Reserved.
+
+#pragma once
+
+#include "AllocationsProvider.h"
+#include "CoreMinimal.h"
+#include "TraceServices/Containers/Tables.h"
+
+namespace TraceServices
+{
+class FImportTableRow;
+
+class ITableImportData
+{
+	TSharedPtr<ITable<FImportTableRow>> GetTable();
+};
+
+class FTableImportService
+{
+public:
+	typedef TFunction<void(FName TableId, TSharedPtr<ITable<FImportTableRow>> Data)> TableImportCallback;
+
+	TRACESERVICES_API static void ImportTable(const FString& InPath, FName TableId, TableImportCallback InCallback);
+};
+
+} // namespace TraceServices
