[Insights] Made imported table views to run in async mode.

diff --git a/Engine/Source/Developer/TraceInsights/Private/Insights/ImportTool/TableImportTool.cpp b/Engine/Source/Developer/TraceInsights/Private/Insights/ImportTool/TableImportTool.cpp
--- a/Engine/Source/Developer/TraceInsights/Private/Insights/ImportTool/TableImportTool.cpp
+++ b/Engine/Source/Developer/TraceInsights/Private/Insights/ImportTool/TableImportTool.cpp
@@ -156,6 +156,7 @@ TSharedRef<SDockTab> FTableImportTool::SpawnTab_TableImportTreeView(const FSpawn
 		.TabRole(ETabRole::PanelTab)
 		[
 			SAssignNew(TableTreeView, SUntypedTableTreeView, UntypedTable)
+			.RunInAsyncMode(true)
 		];
 
 	FOpenImportedTableTabData OpenTabData;
diff --git a/Engine/Source/Developer/TraceInsights/Private/Insights/Table/Widgets/SUntypedTableTreeView.cpp b/Engine/Source/Developer/TraceInsights/Private/Insights/Table/Widgets/SUntypedTableTreeView.cpp
--- a/Engine/Source/Developer/TraceInsights/Private/Insights/Table/Widgets/SUntypedTableTreeView.cpp
+++ b/Engine/Source/Developer/TraceInsights/Private/Insights/Table/Widgets/SUntypedTableTreeView.cpp
@@ -31,6 +31,7 @@ SUntypedTableTreeView::~SUntypedTableTreeView()
 
 void SUntypedTableTreeView::Construct(const FArguments& InArgs, TSharedPtr<Insights::FUntypedTable> InTablePtr)
 {
+	bRunInAsyncMode = InArgs._RunInAsyncMode;
 	ConstructWidget(InTablePtr);
 }
 
diff --git a/Engine/Source/Developer/TraceInsights/Private/Insights/Table/Widgets/SUntypedTableTreeView.h b/Engine/Source/Developer/TraceInsights/Private/Insights/Table/Widgets/SUntypedTableTreeView.h
--- a/Engine/Source/Developer/TraceInsights/Private/Insights/Table/Widgets/SUntypedTableTreeView.h
+++ b/Engine/Source/Developer/TraceInsights/Private/Insights/Table/Widgets/SUntypedTableTreeView.h
@@ -22,7 +22,10 @@ public:
 	/** Virtual destructor. */
 	virtual ~SUntypedTableTreeView();
 
-	SLATE_BEGIN_ARGS(SUntypedTableTreeView) {}
+	SLATE_BEGIN_ARGS(SUntypedTableTreeView)
+		: _RunInAsyncMode(false)
+		{}
+		SLATE_ARGUMENT(bool, RunInAsyncMode)
 	SLATE_END_ARGS()
 
 	/**
