[Insights] Clean up imported table tab display name and tooltip text.

diff --git a/Engine/Source/Developer/TraceInsights/Private/Insights/ImportTool/TableImportTool.cpp b/Engine/Source/Developer/TraceInsights/Private/Insights/ImportTool/TableImportTool.cpp
--- a/Engine/Source/Developer/TraceInsights/Private/Insights/ImportTool/TableImportTool.cpp
+++ b/Engine/Source/Developer/TraceInsights/Private/Insights/ImportTool/TableImportTool.cpp
@@ -110,19 +110,22 @@ void FTableImportTool::StartImportProcess()
 void FTableImportTool::DisplayTable(FName TableViewID)
 {
 	FText ImportTableTreeViewTabDisplayName;
+	FString FileName = FPaths::GetCleanFilename(TableViewID.GetPlainNameString());
 	if (TableViewID.GetNumber())
 	{
 		ImportTableTreeViewTabDisplayName = FText::Format(LOCTEXT("TableImportTool.ImportTableTreeViewTabDisplayName", "{0}({1})"),
-														  FText::FromString(TableViewID.GetPlainNameString()),
+														  FText::FromString(FileName),
 														  FText::AsNumber(TableViewID.GetNumber()));
 	}
 	else
 	{
-		ImportTableTreeViewTabDisplayName = FText::FromString(TableViewID.GetPlainNameString());
+		ImportTableTreeViewTabDisplayName = FText::FromString(FileName);
 	}
 
+	FString TooltipText = FPaths::ConvertRelativePathToFull(TableViewID.GetPlainNameString());
 	FGlobalTabmanager::Get()->RegisterNomadTabSpawner(TableViewID, FOnSpawnTab::CreateSP(this, &FTableImportTool::SpawnTab_TableImportTreeView, TableViewID, ImportTableTreeViewTabDisplayName))
 		.SetDisplayName(ImportTableTreeViewTabDisplayName)
+		.SetTooltipText(FText::FromString(TooltipText))
 		.SetIcon(FSlateIcon(FInsightsStyle::GetStyleSetName(), "Icons.TableTreeView"));
 
 	FGlobalTabmanager::Get()->TryInvokeTab(TableViewID);
