[Insights] Add filter for branch to the session list

diff --git a/Engine/Source/Developer/TraceInsights/Private/Insights/Widgets/SStartPageWindow.cpp b/Engine/Source/Developer/TraceInsights/Private/Insights/Widgets/SStartPageWindow.cpp
--- a/Engine/Source/Developer/TraceInsights/Private/Insights/Widgets/SStartPageWindow.cpp
+++ b/Engine/Source/Developer/TraceInsights/Private/Insights/Widgets/SStartPageWindow.cpp
@@ -791,6 +791,7 @@ STraceStoreWindow::STraceStoreWindow()
 	, FilterByAppName()
 	, FilterByBuildConfig()
 	, FilterByBuildTarget()
+	, FilterByBranch()
 	, bFilterStatsTextIsDirty(true)
 	, FilterStatsText()
 	, SortColumn(TraceStoreColumns::Date)
@@ -1019,7 +1020,7 @@ TSharedRef<SWidget> STraceStoreWindow::ConstructFiltersToolbar()
 		ToolbarBuilder.AddComboButton(
 			FUIAction(),
 			FOnGetContent::CreateSP(this, &STraceStoreWindow::MakeBuildConfigFilterMenu),
-			LOCTEXT("FilterByBuildConfigText", "Build Config"),
+			LOCTEXT("FilterByBuildConfigText", "Config"),
 			LOCTEXT("FilterByBuildConfigToolTip", "Filters the list of trace sessions by build configuration."),
 			FSlateIcon(FInsightsStyle::GetStyleSetName(), "Icons.Filter.ToolBar"),
 			false
@@ -1029,11 +1030,21 @@ TSharedRef<SWidget> STraceStoreWindow::ConstructFiltersToolbar()
 		ToolbarBuilder.AddComboButton(
 			FUIAction(),
 			FOnGetContent::CreateSP(this, &STraceStoreWindow::MakeBuildTargetFilterMenu),
-			LOCTEXT("FilterByBuildTargetText", "Build Target"),
+			LOCTEXT("FilterByBuildTargetText", "Target"),
 			LOCTEXT("FilterByBuildTargetToolTip", "Filters the list of trace sessions by build target."),
 			FSlateIcon(FInsightsStyle::GetStyleSetName(), "Icons.Filter.ToolBar"),
 			false
 		);
+
+		// Filter by Branch
+		ToolbarBuilder.AddComboButton(
+			FUIAction(),
+			FOnGetContent::CreateSP(this, &STraceStoreWindow::MakeBranchFilterMenu),
+			LOCTEXT("FilterByBranchText", "Branch"),
+			LOCTEXT("FilterByBranchToolTip", "Filters the list of trace sessions by branch."),
+			FSlateIcon(FInsightsStyle::GetStyleSetName(), "Icons.Filter.ToolBar"),
+			false
+		);
 	}
 	ToolbarBuilder.EndSection();
 
@@ -2271,6 +2282,28 @@ void STraceStoreWindow::BuildBuildTargetFilterSubMenu(FMenuBuilder& InMenuBuilde
 
 ////////////////////////////////////////////////////////////////////////////////////////////////////
 
+TSharedRef<SWidget> STraceStoreWindow::MakeBranchFilterMenu()
+{
+	FSlateApplication::Get().CloseToolTip();
+
+	FMenuBuilder MenuBuilder(/*bInShouldCloseWindowAfterMenuSelection=*/true, nullptr);
+
+	MenuBuilder.BeginSection("Filter", LOCTEXT("MenuSection_BranchFilter", "Branch Filter"));
+	BuildBranchFilterSubMenu(MenuBuilder);
+	MenuBuilder.EndSection();
+
+	return MenuBuilder.MakeWidget();
+}
+
+////////////////////////////////////////////////////////////////////////////////////////////////////
+
+void STraceStoreWindow::BuildBranchFilterSubMenu(FMenuBuilder& InMenuBuilder)
+{
+	FilterByBranch->BuildMenu(InMenuBuilder, *this);
+}
+
+////////////////////////////////////////////////////////////////////////////////////////////////////
+
 FText STraceStoreWindow::GetTraceStoreDirectory() const
 {
 	return FText::FromString(FPaths::ConvertRelativePathToFull(FInsightsManager::Get()->GetStoreDir()));
@@ -2449,6 +2482,11 @@ FTraceFilterByBuildTarget::FTraceFilterByBuildTarget()
 	ToggleAllActionTooltip = LOCTEXT("FilterByBuildTarget_ToggleAll_Tooltip", "Shows or hides traces for all build targets.");
 }
 
+FTraceFilterByBranch::FTraceFilterByBranch()
+{
+	ToggleAllActionTooltip = LOCTEXT("FilterByBranch_ToggleAll_Tooltip", "Shows or hides traces for all branches.");
+}
+
 ////////////////////////////////////////////////////////////////////////////////////////////////////
 
 void STraceStoreWindow::FilterByNameSearchBox_OnTextChanged(const FText& InFilterText)
@@ -2494,6 +2532,9 @@ void STraceStoreWindow::CreateFilters()
 
 	FilterByBuildTarget = MakeShared<FTraceFilterByBuildTarget>();
 	Filters->Add(FilterByBuildTarget);
+
+	FilterByBranch = MakeShared<FTraceFilterByBranch>();
+	Filters->Add(FilterByBranch);
 }
 
 ////////////////////////////////////////////////////////////////////////////////////////////////////
@@ -2520,7 +2561,8 @@ void STraceStoreWindow::UpdateFiltering()
 		FilterByPlatform->IsEmpty() &&
 		FilterByAppName->IsEmpty() &&
 		FilterByBuildConfig->IsEmpty() &&
-		FilterByBuildTarget->IsEmpty())
+		FilterByBuildTarget->IsEmpty() && 
+		FilterByBranch->IsEmpty())
 	{
 		// No filtering.
 		FilteredTraceViewModels = TraceViewModels;
diff --git a/Engine/Source/Developer/TraceInsights/Private/Insights/Widgets/SStartPageWindow.h b/Engine/Source/Developer/TraceInsights/Private/Insights/Widgets/SStartPageWindow.h
--- a/Engine/Source/Developer/TraceInsights/Private/Insights/Widgets/SStartPageWindow.h
+++ b/Engine/Source/Developer/TraceInsights/Private/Insights/Widgets/SStartPageWindow.h
@@ -215,6 +215,18 @@ protected:
 	}
 };
 
+class FTraceFilterByBranch : public FTraceFilterByStringSet
+{
+public:
+	FTraceFilterByBranch();
+
+protected:
+	virtual FString GetFilterValueForTrace(const FTraceViewModel& InTrace) const override
+	{
+		return InTrace.Branch.ToString();
+	}
+};
+
 ////////////////////////////////////////////////////////////////////////////////////////////////////
 
 /** Implements the Trace Store window. */
@@ -311,6 +323,9 @@ private:
 	TSharedRef<SWidget> MakeBuildTargetFilterMenu();
 	void BuildBuildTargetFilterSubMenu(FMenuBuilder& InMenuBuilder);
 
+	TSharedRef<SWidget> MakeBranchFilterMenu();
+	void BuildBranchFilterSubMenu(FMenuBuilder& InMenuBuilder);
+
 	FReply RefreshTraces_OnClicked();
 
 	void RefreshTraceList();
@@ -469,6 +484,7 @@ private:
 	TSharedPtr<FTraceFilterByAppName> FilterByAppName;
 	TSharedPtr<FTraceFilterByBuildConfig> FilterByBuildConfig;
 	TSharedPtr<FTraceFilterByBuildTarget> FilterByBuildTarget;
+	TSharedPtr<FTraceFilterByBranch> FilterByBranch;
 
 	bool bFilterStatsTextIsDirty;
 	FText FilterStatsText;
