[Insights] Fix Linux crash when double clicking on a Task in the Task Table.

diff --git a/Engine/Source/Developer/TraceInsights/Private/Insights/TaskGraphProfiler/Widgets/STaskTableTreeView.cpp b/Engine/Source/Developer/TraceInsights/Private/Insights/TaskGraphProfiler/Widgets/STaskTableTreeView.cpp
--- a/Engine/Source/Developer/TraceInsights/Private/Insights/TaskGraphProfiler/Widgets/STaskTableTreeView.cpp
+++ b/Engine/Source/Developer/TraceInsights/Private/Insights/TaskGraphProfiler/Widgets/STaskTableTreeView.cpp
@@ -506,16 +506,18 @@ void STaskTableTreeView::ContextMenu_GoToTask_Execute()
 
 	TimingView->SelectTimingTrack(Track, true);
 
-	FTimingEventSearchParameters SearchParams(TaskEntry->StartedTimestamp, TaskEntry->FinishedTimestamp, ETimingEventSearchFlags::StopAtFirstMatch, 
-		[TaskEntry](double StartTime, double EndTime, uint32 Depth)
+	auto SearchFilter = [TaskEntry](double StartTime, double EndTime, uint32 Depth)
+	{
+		if (StartTime >= TaskEntry->GetStartedTimestamp() && EndTime <= TaskEntry->GetFinishedTimestamp())
 		{
-			if (StartTime >= TaskEntry->GetStartedTimestamp() && EndTime <= TaskEntry->GetFinishedTimestamp())
-			{
-				return true;
-			}
+			return true;
+		}
+
+		return false;
+	};
+
+	FTimingEventSearchParameters SearchParams(TaskEntry->StartedTimestamp, TaskEntry->FinishedTimestamp, ETimingEventSearchFlags::StopAtFirstMatch, SearchFilter);
 
-			return false;
-		});
 
 	const TSharedPtr<const ITimingEvent> FoundEvent = Track->SearchEvent(SearchParams);
 	TimingView->SelectTimingEvent(FoundEvent, true);
