[Insights]
  - Moved the Import Table as first item in main Menu. Also changed its icon.
  - Renamed "Package" view (Cooking profiler) to "Packages" and set the generic table icon.

diff --git a/Engine/Source/Developer/TraceInsights/Private/Insights/Common/InsightsMenuBuilder.cpp b/Engine/Source/Developer/TraceInsights/Private/Insights/Common/InsightsMenuBuilder.cpp
--- a/Engine/Source/Developer/TraceInsights/Private/Insights/Common/InsightsMenuBuilder.cpp
+++ b/Engine/Source/Developer/TraceInsights/Private/Insights/Common/InsightsMenuBuilder.cpp
@@ -54,6 +54,12 @@ void FInsightsMenuBuilder::PopulateMenu(FMenuBuilder& MenuBuilder)
 {
 #if !WITH_EDITOR
 	MenuBuilder.BeginSection("Insights");
+	MenuBuilder.AddMenuEntry(
+		LOCTEXT("ImportTable", "Import Table..."),
+		LOCTEXT("ImportTable_ToolTip", "Import CSV or TSV data from a file to an Insights Table."),
+		FSlateIcon(FInsightsStyle::GetStyleSetName(), "Icons.ImportTable"),
+		FUIAction(FExecuteAction::CreateLambda([] { Insights::FTableImportTool::Get()->StartImportProcess(); })));
+	MenuBuilder.AddSeparator();
 	MenuBuilder.AddMenuEntry(
 		LOCTEXT("OpenSessionBrowser", "Session Browser"),
 		LOCTEXT("OpenSessionBrowser_ToolTip", "Opens the Unreal Insights Session Browser window."),
@@ -66,11 +72,6 @@ void FInsightsMenuBuilder::PopulateMenu(FMenuBuilder& MenuBuilder)
 		false,
 		FSlateIcon(FAppStyle::Get().GetStyleSetName(), "Icons.FolderOpen")
 	);
-	MenuBuilder.AddMenuEntry(
-		LOCTEXT("ImportTable", "Import Table"),
-		LOCTEXT("ImportTable_ToolTip", "Import CSV or TSV data from a file to an Insights Table."),
-		FSlateIcon(FInsightsStyle::GetStyleSetName(), "Icons.Import.Table"),
-		FUIAction(FExecuteAction::CreateLambda([] { Insights::FTableImportTool::Get()->StartImportProcess(); })));
 	MenuBuilder.AddMenuEntry(
 		LOCTEXT("AutoOpenLiveTrace", "Auto Open Live Trace"),
 		LOCTEXT("AutoOpenLiveTrace_ToolTip", "If enabled, the analysis starts automatically for each new live trace session, replacing the current analysis session."),
diff --git a/Engine/Source/Developer/TraceInsights/Private/Insights/CookProfiler/CookProfilerManager.cpp b/Engine/Source/Developer/TraceInsights/Private/Insights/CookProfiler/CookProfilerManager.cpp
--- a/Engine/Source/Developer/TraceInsights/Private/Insights/CookProfiler/CookProfilerManager.cpp
+++ b/Engine/Source/Developer/TraceInsights/Private/Insights/CookProfiler/CookProfilerManager.cpp
@@ -185,9 +185,9 @@ void FCookProfilerManager::RegisterTimingProfilerLayoutExtensions(FInsightsMajor
 
 	FMinorTabConfig& MinorTabConfig = InOutExtender.AddMinorTabConfig();
 	MinorTabConfig.TabId = FCookProfilerTabs::PackageTableTreeViewTabID;
-	MinorTabConfig.TabLabel = LOCTEXT("PackageTableTreeViewTabTitle", "Package");
-	MinorTabConfig.TabTooltip = LOCTEXT("PackageTableTreeViewTabTitleTooltip", "Opens the Package Tree View tab, that allows cook profilling.");
-	MinorTabConfig.TabIcon = FSlateIcon(FInsightsStyle::GetStyleSetName(), "Icons.TasksView");
+	MinorTabConfig.TabLabel = LOCTEXT("PackageTableTreeViewTabTitle", "Packages");
+	MinorTabConfig.TabTooltip = LOCTEXT("PackageTableTreeViewTabTitleTooltip", "Opens the Packages Tree View tab, that allows cook profilling.");
+	MinorTabConfig.TabIcon = FSlateIcon(FInsightsStyle::GetStyleSetName(), "Icons.PackagesView");
 	MinorTabConfig.OnSpawnTab = FOnSpawnTab::CreateRaw(this, &FCookProfilerManager::SpawnTab_PackageTableTreeView);
 	MinorTabConfig.CanSpawnTab = FCanSpawnTab::CreateRaw(this, &FCookProfilerManager::CanSpawnTab_PackageTableTreeView);
 	MinorTabConfig.WorkspaceGroup = InOutExtender.GetWorkspaceGroup();
diff --git a/Engine/Source/Developer/TraceInsights/Private/Insights/InsightsStyle.cpp b/Engine/Source/Developer/TraceInsights/Private/Insights/InsightsStyle.cpp
--- a/Engine/Source/Developer/TraceInsights/Private/Insights/InsightsStyle.cpp
+++ b/Engine/Source/Developer/TraceInsights/Private/Insights/InsightsStyle.cpp
@@ -183,6 +183,7 @@ void FInsightsStyle::FStyle::Initialize()
 	Set("Icons.TableTreeView.ToolBar", new EDITOR_IMAGE_BRUSH_SVG("Starship/Common/Spreadsheet", Icon20x20));
 
 	Set("Icons.TasksView", new IMAGE_BRUSH_SVG("Tasks", Icon16x16));
+	Set("Icons.PackagesView", new EDITOR_IMAGE_BRUSH_SVG("Starship/Common/Spreadsheet", Icon16x16));
 
 	Set("Icons.AllTracksMenu.ToolBar", new IMAGE_BRUSH_SVG("AllTracks_20", Icon20x20));
 	Set("Icons.CpuGpuTracksMenu.ToolBar", new IMAGE_BRUSH_SVG("CpuGpuTracks_20", Icon20x20));
@@ -316,6 +317,8 @@ void FInsightsStyle::FStyle::Initialize()
 	Set("Icons.Rename", new CORE_IMAGE_BRUSH_SVG("Starship/Common/Rename", Icon16x16));
 	Set("Icons.Delete", new CORE_IMAGE_BRUSH_SVG("Starship/Common/Delete", Icon16x16));
 
+	Set("Icons.ImportTable", new CORE_IMAGE_BRUSH_SVG("Starship/Common/Import", Icon16x16));
+
 	//////////////////////////////////////////////////
 
 	Set("TreeTable.RowBackground", new EDITOR_IMAGE_BRUSH("Old/White", Icon16x16, FLinearColor(1.0f, 1.0f, 1.0f, 0.25f)));
@@ -333,8 +336,6 @@ void FInsightsStyle::FStyle::Initialize()
 	Set("Icons.NetEvent.TreeItem", new TODO_IMAGE_BRUSH(Icon12x12));
 	Set("Icons.MemTag.TreeItem", new TODO_IMAGE_BRUSH(Icon12x12));
 
-	Set("Icons.Import.Table", new TODO_IMAGE_BRUSH(Icon12x12));
-
 	Set("TreeTable.TooltipBold", FTextBlockStyle(NormalText)
 		.SetFont(DEFAULT_FONT("Bold", 8))
 		.SetColorAndOpacity(FLinearColor(0.5f, 0.5f, 0.5f, 1.0f))
