[Insights] Updated description of some memory rules to avoid some confusions (i.e. "freed after time T" --> "not freed until at least time T").

diff --git a/Engine/Source/Developer/TraceInsights/Private/Insights/MemoryProfiler/ViewModels/MemorySharedState.cpp b/Engine/Source/Developer/TraceInsights/Private/Insights/MemoryProfiler/ViewModels/MemorySharedState.cpp
--- a/Engine/Source/Developer/TraceInsights/Private/Insights/MemoryProfiler/ViewModels/MemorySharedState.cpp
+++ b/Engine/Source/Developer/TraceInsights/Private/Insights/MemoryProfiler/ViewModels/MemorySharedState.cpp
@@ -957,7 +957,7 @@ void FMemorySharedState::InitMemoryRules()
 		ERule::Aaf, 1,
 		LOCTEXT("MemRule_Aaf_Short", "A**"),
 		LOCTEXT("MemRule_Aaf_Verbose", "After"),
-		LOCTEXT("MemRule_Aaf_Desc", "Identifies allocations allocated and freed after time A.\n(A ≤ a ≤ f)")));
+		LOCTEXT("MemRule_Aaf_Desc", "Identifies allocations allocated after time A.\n(A ≤ a ≤ f)")));
 
 	MemoryRules.Add(MakeShared<Insights::FMemoryRuleSpec>(
 		ERule::aAfB, 2,
@@ -969,13 +969,13 @@ void FMemorySharedState::InitMemoryRules()
 		ERule::AaBf, 2,
 		LOCTEXT("MemRule_AaBf_Short", "A*B*"),
 		LOCTEXT("MemRule_AaBf_Verbose", "Growth"),
-		LOCTEXT("MemRule_AaBf_Desc", "Identifies allocations allocated between time A and time B and freed after time B.\n(A ≤ a ≤ B ≤ f)")));
+		LOCTEXT("MemRule_AaBf_Desc", "Identifies allocations allocated between time A and time B and not freed until at least time B.\n(A ≤ a ≤ B ≤ f)")));
 
 	MemoryRules.Add(MakeShared<Insights::FMemoryRuleSpec>(
 		ERule::aAfaBf, 2,
 		LOCTEXT("MemRule_aAfaBf_Short", "A*B*/*A*B"),
 		LOCTEXT("MemRule_aAfaBf_Verbose", "Growth vs. Decline"),
-		LOCTEXT("MemRule_aAfaBf_Desc", "Identifies \"growth\" allocations, allocated between time A and time B and freed after time B (A ≤ a ≤ B ≤ f)\nand \"decline\" allocations, allocated before time A and freed between time A and time B (a ≤ A ≤ f ≤ B).\nThe \"decline\" allocations are changed to have negative size, so the size aggregation shows variation between A and B.")));
+		LOCTEXT("MemRule_aAfaBf_Desc", "Identifies \"growth\" allocations, allocated between time A and time B and not freed until at least time B (A ≤ a ≤ B ≤ f)\nand \"decline\" allocations, allocated before time A and freed between time A and time B (a ≤ A ≤ f ≤ B).\nThe \"decline\" allocations are changed to have negative size, so the size aggregation shows variation between A and B.")));
 
 	MemoryRules.Add(MakeShared<Insights::FMemoryRuleSpec>(
 		ERule::AfB, 2,
@@ -999,13 +999,13 @@ void FMemorySharedState::InitMemoryRules()
 		ERule::aABf, 2,
 		LOCTEXT("MemRule_aABf_Short", "*A B*"),
 		LOCTEXT("MemRule_aABf_Verbose", "Long Living Allocs"),
-		LOCTEXT("MemRule_aABf_Desc", "Identifies allocations allocated before time A and freed after time B.\n(a ≤ A ≤ B ≤ f)")));
+		LOCTEXT("MemRule_aABf_Desc", "Identifies allocations allocated before time A and not freed until at least time B.\n(a ≤ A ≤ B ≤ f)")));
 
 	MemoryRules.Add(MakeShared<Insights::FMemoryRuleSpec>(
 		ERule::AaBCf, 3,
 		LOCTEXT("MemRule_AaBCf_Short", "A*B C*"),
 		LOCTEXT("MemRule_AaBCf_Verbose", "Memory Leaks"),
-		LOCTEXT("MemRule_AaBCf_Desc", "Identifies allocations allocated between time A and time B and freed after time C.\n(A ≤ a ≤ B ≤ C ≤ f)")));
+		LOCTEXT("MemRule_AaBCf_Desc", "Identifies allocations allocated between time A and time B and not freed until at least time C.\n(A ≤ a ≤ B ≤ C ≤ f)")));
 
 	MemoryRules.Add(MakeShared<Insights::FMemoryRuleSpec>(
 		ERule::AaBfC, 3,
