[Insights]
  - Session Browser: Added user notifications for Rename and Delete trace actions.
  - Session Browser: Fixed rename to not replace an existing file anymore, but to warn user instead.

diff --git a/Engine/Source/Developer/TraceInsights/Private/Insights/Widgets/SStartPageWindow.cpp b/Engine/Source/Developer/TraceInsights/Private/Insights/Widgets/SStartPageWindow.cpp
--- a/Engine/Source/Developer/TraceInsights/Private/Insights/Widgets/SStartPageWindow.cpp
+++ b/Engine/Source/Developer/TraceInsights/Private/Insights/Widgets/SStartPageWindow.cpp
@@ -224,7 +224,14 @@ public:
 				if (FPaths::FileExists(TraceFile))
 				{
 					FString NewTraceFile = FPaths::Combine(FPaths::GetPath(TraceFile), NewTraceName + TEXT(".utrace"));
-					if (IFileManager::Get().Move(*NewTraceFile, *TraceFile, true))
+					if (FPaths::FileExists(NewTraceFile))
+					{
+						UE_LOG(TraceInsights, Warning, TEXT("[TraceStore] Failed to rename \"%s\" to \"%s\"! File already exists."), *TraceName, *NewTraceName);
+
+						FText Message = FText::Format(LOCTEXT("RenameFailFmt1", "Failed to rename \"{0}\" to \"{1}\"!\nFile already exists."), FText::FromString(TraceName), FText::FromString(NewTraceName));
+						WeakParentWidget.Pin()->ShowFailMessage(Message);
+					}
+					else if (IFileManager::Get().Move(*NewTraceFile, *TraceFile, false))
 					{
 						UE_LOG(TraceInsights, Verbose, TEXT("[TraceStore] Renamed utrace file (\"%s\")."), *NewTraceFile);
 						TracePin->Name = InText;
@@ -239,6 +246,16 @@ public:
 								UE_LOG(TraceInsights, Verbose, TEXT("[TraceStore] Renamed ucache file (\"%s\")."), *NewCacheFile);
 							}
 						}
+
+						FText Message = FText::Format(LOCTEXT("RenameSuccessFmt", "Renamed \"{0}\" to \"{1}\"."), FText::FromString(TraceName), FText::FromString(NewTraceName));
+						WeakParentWidget.Pin()->ShowSuccessMessage(Message);
+					}
+					else
+					{
+						UE_LOG(TraceInsights, Warning, TEXT("[TraceStore] Failed to rename \"%s\" to \"%s\"!"), *TraceName, *NewTraceName);
+
+						FText Message = FText::Format(LOCTEXT("RenameFailFmt2", "Failed to rename \"{0}\" to \"{1}\"!"), FText::FromString(TraceName), FText::FromString(NewTraceName));
+						WeakParentWidget.Pin()->ShowFailMessage(Message);
 					}
 				}
 			}
@@ -770,7 +787,6 @@ private:
 
 STraceStoreWindow::STraceStoreWindow()
 	: NotificationList()
-	, ActiveNotifications()
 	, OverlaySettingsSlot(nullptr)
 	, DurationActive(0.0f)
 	, ActiveTimerHandle()
@@ -1451,25 +1467,32 @@ void STraceStoreWindow::DeleteTraceFile()
 	FString TraceName = SelectedTrace->Name.ToString();
 	UE_LOG(TraceInsights, Log, TEXT("[TraceStore] Deleting \"%s\"..."), *TraceName);
 
-	if (FPaths::FileExists(TraceFile))
+	if (FPaths::FileExists(TraceFile) && IFileManager::Get().Delete(*TraceFile))
 	{
-		if (IFileManager::Get().Delete(*TraceFile))
-		{
-			UE_LOG(TraceInsights, Verbose, TEXT("[TraceStore] Deleted utrace file (\"%s\")."), *TraceFile);
+		UE_LOG(TraceInsights, Verbose, TEXT("[TraceStore] Deleted utrace file (\"%s\")."), *TraceFile);
 
-			FString CacheFile = FPaths::ChangeExtension(TraceFile, TEXT("ucache"));
-			if (FPaths::FileExists(CacheFile))
+		FString CacheFile = FPaths::ChangeExtension(TraceFile, TEXT("ucache"));
+		if (FPaths::FileExists(CacheFile))
+		{
+			if (IFileManager::Get().Delete(*CacheFile))
 			{
-				if (IFileManager::Get().Delete(*CacheFile))
-				{
-					UE_LOG(TraceInsights, Verbose, TEXT("[TraceStore] Deleted ucache file (\"%s\")."), *CacheFile);
-				}
+				UE_LOG(TraceInsights, Verbose, TEXT("[TraceStore] Deleted ucache file (\"%s\")."), *CacheFile);
 			}
-
-			TraceViewModels.Remove(SelectedTrace);
-			TraceViewModelMap.Remove(SelectedTrace->TraceId);
-			OnTraceListChanged();
 		}
+
+		TraceViewModels.Remove(SelectedTrace);
+		TraceViewModelMap.Remove(SelectedTrace->TraceId);
+		OnTraceListChanged();
+
+		FText Message = FText::Format(LOCTEXT("DeleteSuccessFmt", "Deleted \"{0}\"."), FText::FromString(TraceName));
+		ShowSuccessMessage(Message);
+	}
+	else
+	{
+		UE_LOG(TraceInsights, Warning, TEXT("[TraceStore] Failed to delete utrace file (\"%s\")!"), *TraceFile);
+
+		FText Message = FText::Format(LOCTEXT("DeleteFailFmt", "Failed to delete \"{0}\"!"), FText::FromString(TraceName));
+		ShowFailMessage(Message);
 	}
 }
 
@@ -2804,6 +2827,34 @@ void STraceStoreWindow::UpdateSorting()
 
 ////////////////////////////////////////////////////////////////////////////////////////////////////
 
+void STraceStoreWindow::ShowSuccessMessage(FText& InMessage)
+{
+	FNotificationInfo NotificationInfo(InMessage);
+	NotificationInfo.bFireAndForget = false;
+	NotificationInfo.bUseLargeFont = false;
+	NotificationInfo.bUseSuccessFailIcons = true;
+	NotificationInfo.ExpireDuration = 3.0f;
+	SNotificationItemRef NotificationItem = NotificationList->AddNotification(NotificationInfo);
+	NotificationItem->SetCompletionState(SNotificationItem::CS_Success);
+	NotificationItem->ExpireAndFadeout();
+}
+
+////////////////////////////////////////////////////////////////////////////////////////////////////
+
+void STraceStoreWindow::ShowFailMessage(FText& InMessage)
+{
+	FNotificationInfo NotificationInfo(InMessage);
+	NotificationInfo.bFireAndForget = false;
+	NotificationInfo.bUseLargeFont = false;
+	NotificationInfo.bUseSuccessFailIcons = true;
+	NotificationInfo.ExpireDuration = 3.0f;
+	SNotificationItemRef NotificationItem = NotificationList->AddNotification(NotificationInfo);
+	NotificationItem->SetCompletionState(SNotificationItem::CS_Fail);
+	NotificationItem->ExpireAndFadeout();
+}
+
+////////////////////////////////////////////////////////////////////////////////////////////////////
+
 #undef LOCTEXT_NAMESPACE
 #define LOCTEXT_NAMESPACE "SConnectionWindow"
 
@@ -2813,7 +2864,6 @@ void STraceStoreWindow::UpdateSorting()
 
 SConnectionWindow::SConnectionWindow()
 : NotificationList()
-, ActiveNotifications()
 , ConnectTask()
 , bIsConnecting(false)
 , bIsConnectedSuccessfully(false)
@@ -3119,10 +3169,9 @@ FReply SConnectionWindow::Connect_OnClicked()
 				NotificationInfo.bUseLargeFont = false;
 				NotificationInfo.bUseSuccessFailIcons = true;
 				NotificationInfo.ExpireDuration = 10.0f;
-				SNotificationItemWeak NotificationItem = NotificationList->AddNotification(NotificationInfo);
-				NotificationItem.Pin()->SetCompletionState(SNotificationItem::CS_Success);
-				NotificationItem.Pin()->ExpireAndFadeout();
-				ActiveNotifications.Add(TEXT("ConnectSuccess"), NotificationItem);
+				SNotificationItemRef NotificationItem = NotificationList->AddNotification(NotificationInfo);
+				NotificationItem->SetCompletionState(SNotificationItem::CS_Success);
+				NotificationItem->ExpireAndFadeout();
 			}
 			else
 			{
@@ -3133,10 +3182,9 @@ FReply SConnectionWindow::Connect_OnClicked()
 				NotificationInfo.bUseLargeFont = false;
 				NotificationInfo.bUseSuccessFailIcons = true;
 				NotificationInfo.ExpireDuration = 10.0f;
-				SNotificationItemWeak NotificationItem = NotificationList->AddNotification(NotificationInfo);
-				NotificationItem.Pin()->SetCompletionState(SNotificationItem::CS_Fail);
-				NotificationItem.Pin()->ExpireAndFadeout();
-				ActiveNotifications.Add(TEXT("ConnectFailed"), NotificationItem);
+				SNotificationItemRef NotificationItem = NotificationList->AddNotification(NotificationInfo);
+				NotificationItem->SetCompletionState(SNotificationItem::CS_Fail);
+				NotificationItem->ExpireAndFadeout();
 			}
 
 			bIsConnecting = false;
diff --git a/Engine/Source/Developer/TraceInsights/Private/Insights/Widgets/SStartPageWindow.h b/Engine/Source/Developer/TraceInsights/Private/Insights/Widgets/SStartPageWindow.h
--- a/Engine/Source/Developer/TraceInsights/Private/Insights/Widgets/SStartPageWindow.h
+++ b/Engine/Source/Developer/TraceInsights/Private/Insights/Widgets/SStartPageWindow.h
@@ -440,13 +440,13 @@ private:
 
 	void UpdateTraceListView();
 
+	void ShowSuccessMessage(FText& InMessage);
+	void ShowFailMessage(FText& InMessage);
+
 private:
 	/** Widget for the non-intrusive notifications. */
 	TSharedPtr<SNotificationList> NotificationList;
 
-	/** Holds all active and visible notifications, stored as FGuid -> SNotificationItemWeak. */
-	TMap<FString, SNotificationItemWeak> ActiveNotifications;
-
 	/** Overlay slot which contains the profiler settings widget. */
 	SOverlay::FOverlaySlot* OverlaySettingsSlot;
 
@@ -547,9 +547,6 @@ private:
 	/** Widget for the non-intrusive notifications. */
 	TSharedPtr<SNotificationList> NotificationList;
 
-	/** Holds all active and visible notifications, stored as FGuid -> SNotificationItemWeak. */
-	TMap<FString, SNotificationItemWeak> ActiveNotifications;
-
 	FGraphEventRef ConnectTask;
 
 	std::atomic<bool> bIsConnecting;
