Insights: Added TimingInsights.ExportTimerStatistics command

diff --git a/Engine/Source/Developer/TraceInsights/Private/Insights/TimingProfilerManager.cpp b/Engine/Source/Developer/TraceInsights/Private/Insights/TimingProfilerManager.cpp
--- a/Engine/Source/Developer/TraceInsights/Private/Insights/TimingProfilerManager.cpp
+++ b/Engine/Source/Developer/TraceInsights/Private/Insights/TimingProfilerManager.cpp
@@ -737,6 +737,106 @@ bool FTimingProfilerManager::Exec(const TCHAR* Cmd, FOutputDevice& Ar)
 		return true;
 	}
 
+	if (FParse::Command(&Cmd, TEXT("TimingInsights.ExportTimerStatistics")))
+	{
+		Ar.Logf(TEXT("TimingInsights.ExportTimerStatistics %s"), Cmd);
+
+		check(FInsightsManager::Get().IsValid() && FInsightsManager::Get()->GetSession().IsValid());
+		Insights::FTimingExporter Exporter(*FInsightsManager::Get()->GetSession().Get());
+		Insights::FTimingExporter::FExportTimerStatisticsParams Params; // default (all timing events)
+
+		// These variables needs to be in the same scope with the call to Exporter.ExportTimingEventsAsText().
+		TArray<FName> Columns; // referenced by Params.Columns
+		TSet<uint32> IncludedThreads; // referenced in Params.ThreadFilter lambda function
+		TSet<uint32> IncludedTimers; // referenced in Params.TimingEventFilter lambda function
+
+		//////////////////////////////////////////////////
+
+		const bool bUseEscape = true;
+		FString Filename = FParse::Token(Cmd, bUseEscape);
+		Ar.Logf(TEXT("  Filename: %s"), *Filename);
+
+		while (Cmd && Cmd[0] != TEXT('\0'))
+		{
+			FString Token;
+			if (FParse::Token(Cmd, Token, bUseEscape))
+			{
+				Ar.Logf(TEXT("  Token: %s"), *Token);
+
+				static constexpr TCHAR ColumnsToken[] = TEXT("-columns=");
+				static constexpr TCHAR ThreadsToken[] = TEXT("-threads=");
+				static constexpr TCHAR TimersToken[] = TEXT("-timers=");
+				static constexpr TCHAR RegionToken[] = TEXT("-region=");
+				static constexpr TCHAR StartTimeToken[] = TEXT("-startTime=");
+				static constexpr TCHAR EndTimeToken[] = TEXT("-endTime=");
+
+				if (Token.StartsWith(ColumnsToken))
+				{
+					// Comma-delimited list of column names. Supports *?-type wildcard.
+					// Default: -columns="ThreadId,TimerId,StartTime,EndTime,Depth"
+					// Example: -columns="*" -columns="TimerName,Duration"
+					Token.RightChopInline(UE_ARRAY_COUNT(ColumnsToken) - 1);
+					Token.TrimQuotesInline();
+					Exporter.MakeExportTimingEventsColumnList(Token, Columns);
+					Params.Columns = &Columns;
+				}
+				else if (Token.StartsWith(ThreadsToken))
+				{
+					// Comma-delimited list of thread names. Supports *?-type wildcard.
+					// Default: -threads="*" (all threads; no filter)
+					// Example: -threads="GameThread" -threads="GPU1,GPU2,GameThread,Render*"
+					Token.RightChopInline(UE_ARRAY_COUNT(ThreadsToken) - 1);
+					Token.TrimQuotesInline();
+					Params.ThreadFilter = Exporter.MakeThreadFilterInclusive(Token, IncludedThreads);
+				}
+				else if (Token.StartsWith(TimersToken))
+				{
+					// Comma-delimited list of timer names. Supports *?-type wildcard.
+					// Default: -timers="*" (all timers; no filter)
+					// Example: -timers="A,B,*z"
+					Token.RightChopInline(UE_ARRAY_COUNT(TimersToken) - 1);
+					Token.TrimQuotesInline();
+					Params.TimingEventFilter = Exporter.MakeTimingEventFilterByTimersInclusive(Token, IncludedTimers);
+				}
+				else if (Token.StartsWith(StartTimeToken))
+				{
+					// Default: -startTime=-infinte
+					// Example: -startTime=10.0
+					Token.RightChopInline(UE_ARRAY_COUNT(StartTimeToken) - 1);
+					Params.IntervalStartTime = atof(TCHAR_TO_ANSI(*Token));
+				}
+				else if (Token.StartsWith(EndTimeToken))
+				{
+					// Default: -endTime=+infinte
+					// Example: -endTime=20.0
+					Token.RightChopInline(UE_ARRAY_COUNT(EndTimeToken) - 1);
+					Params.IntervalEndTime = atof(TCHAR_TO_ANSI(*Token));
+				}
+				else if (Token.StartsWith(RegionToken))
+				{
+					// Default: -Region=
+					// Example: -Region=RegionName
+					Token.RightChopInline(UE_ARRAY_COUNT(RegionToken) - 1);
+					Params.Region = TCHAR_TO_ANSI(*Token);
+				}
+				else
+				{
+					Ar.Logf(ELogVerbosity::Warning, TEXT("Unknown Cmd Param: %s"), *Token);
+				}
+			}
+		}
+
+		//////////////////////////////////////////////////
+
+		if (Params.ThreadFilter == nullptr)
+		{
+			Params.ThreadFilter = [](unsigned int){ return true; };
+		}
+		
+		int32 Result = Exporter.ExportTimerStatisticsAsText(Filename, Params);
+		return Result > 0;
+	}
+	
 	return false;
 }
 
diff --git a/Engine/Source/Developer/TraceInsights/Private/Insights/ViewModels/TimingExporter.cpp b/Engine/Source/Developer/TraceInsights/Private/Insights/ViewModels/TimingExporter.cpp
--- a/Engine/Source/Developer/TraceInsights/Private/Insights/ViewModels/TimingExporter.cpp
+++ b/Engine/Source/Developer/TraceInsights/Private/Insights/ViewModels/TimingExporter.cpp
@@ -10,6 +10,7 @@
 #include "Insights/Log.h"
 #include "Insights/ViewModels/ThreadTimingTrack.h"
 #include "Insights/TimingProfilerManager.h"
+#include "TraceServices/Model/Bookmarks.h"
 
 namespace Insights
 {
@@ -545,6 +546,74 @@ int32 FTimingExporter::ExportTimingEventsAsText(const FString& Filename, FExport
 	return TimingEventCount;
 }
 
+////////////////////////////////////////////////////////////////////////////////////////////////////
+	
+int32 FTimingExporter::ExportTimerStatisticsAsText(const FString& Filename, FExportTimerStatisticsParams& Params) const
+{
+	TraceServices::ITable<TraceServices::FTimingProfilerAggregatedStats>* StatsTable;
+
+	double RegionStartTime = 0;
+	double RegionEndTime = std::numeric_limits<double>::max();
+	
+	{
+		TraceServices::FAnalysisSessionReadScope SessionReadScope(Session);
+		
+		const TraceServices::ITimingProfilerProvider* TimingProfilerProvider = TraceServices::ReadTimingProfilerProvider(Session);
+		if (!TimingProfilerProvider)
+		{
+			UE_LOG(TraceInsights, Error, TEXT("Unable to access TimingProfilerProvider for ExportTimerStatisticsAsText"));
+			return -1;
+		}
+		
+		if (!Params.Region.IsEmpty())
+		{
+			
+			const TraceServices::IBookmarkProvider& BookmarkProvider = TraceServices::ReadBookmarkProvider(Session);
+			const FString RegionStartName = TEXT("RegionStart:") + Params.Region;
+			const FString RegionEndName = TEXT("RegionEnd:") + Params.Region;
+			UE_LOG(TraceInsights, Log, TEXT("Looking for the following bookmarks: '%s' and '%s'"), *RegionStartName, *RegionEndName);
+			
+			BookmarkProvider.EnumerateBookmarks(0, std::numeric_limits<double>::max(),
+				[&RegionStartTime, &RegionEndTime, &RegionStartName, &RegionEndName](const TraceServices::FBookmark& InBookmark){
+					if (RegionStartTime == 0 && RegionStartName.Equals(InBookmark.Text))
+					{
+						RegionStartTime = InBookmark.Time;
+					}
+					if (RegionEndTime == std::numeric_limits<double>::max() && RegionEndName.Equals(InBookmark.Text))
+					{
+						RegionEndTime = InBookmark.Time;
+					}
+			});
+
+			if (RegionStartTime == 0 || RegionEndTime == std::numeric_limits<double>::max())
+			{
+				UE_LOG(TraceInsights, Error, TEXT("Unable to find 'RegionStart:' and 'RegionEnd:' bookmarks for region '%s'"), *Params.Region);
+				return -1;
+			}
+		}
+
+		if (RegionStartTime != 0 && RegionEndTime != std::numeric_limits<double>::max())
+		{
+			Params.IntervalStartTime = RegionStartTime;
+			Params.IntervalEndTime = RegionEndTime;
+			UE_LOG(TraceInsights, Display, TEXT("Limit selection to region '%s' from %f to %f"), *Params.Region, RegionStartTime, RegionEndTime);
+		}
+
+		//@Todo: this does not yet handle the -column and -timers parameters.
+		StatsTable = TimingProfilerProvider->CreateAggregation(Params.IntervalStartTime, Params.IntervalEndTime, Params.ThreadFilter, true);
+	}
+		
+	FStopwatch Stopwatch;
+	Stopwatch.Start();
+
+	TraceServices::Table2Csv(*StatsTable, *Filename);
+
+	Stopwatch.Stop();
+	const double TotalTime = Stopwatch.GetAccumulatedTime();
+	UE_LOG(TraceInsights, Log, TEXT("Exported timing statistics to file in %.3fs (\"%s\")."), TotalTime, *Filename);
+	return 1;
+}
+
 ////////////////////////////////////////////////////////////////////////////////////////////////////
 
 FTimingExporter::FThreadFilterFunc FTimingExporter::MakeThreadFilterInclusive(const FString& InFilterString, TSet<uint32>& OutIncludedThreads) const
diff --git a/Engine/Source/Developer/TraceInsights/Private/Insights/ViewModels/TimingExporter.h b/Engine/Source/Developer/TraceInsights/Private/Insights/ViewModels/TimingExporter.h
--- a/Engine/Source/Developer/TraceInsights/Private/Insights/ViewModels/TimingExporter.h
+++ b/Engine/Source/Developer/TraceInsights/Private/Insights/ViewModels/TimingExporter.h
@@ -65,12 +65,22 @@ public:
 
 		/**
 		 * Filters the timing events by time.
-		 * Only timning events that intersects the [StartTime, EndTime] interval are exported.
+		 * Only timing events that intersects the [StartTime, EndTime] interval are exported.
 		 */
 		double IntervalStartTime = -std::numeric_limits<double>::infinity();
 		double IntervalEndTime = +std::numeric_limits<double>::infinity();
 	};
 
+	struct FExportTimerStatisticsParams : public FExportTimingEventsParams
+	{
+		/**
+		 * The time region to be exported. This is defined by having corresponding "RegionStart:Name" and
+		 * "RegionEnd:Name" bookmarks in the bookmarks channel.
+		 * If empty falls back to IntervalStartTime and IntervalEndTime.
+		 */
+		FString Region;
+	};
+	
 private:
 	typedef TUtf8StringBuilder<1024> FUtf8StringBuilder;
 
@@ -94,6 +104,12 @@ public:
 	int32 ExportThreadsAsText(const FString& Filename, FExportThreadsParams& Params) const;
 	int32 ExportTimersAsText(const FString& Filename, FExportTimersParams& Params) const;
 	int32 ExportTimingEventsAsText(const FString& Filename, FExportTimingEventsParams& Params) const;
+	/**
+	 * Exports Timer Statistics (min,max, inclusive average, exclusive average, etc.).
+	 * Supports specifying a range to export via bookmarks, but does not support timer selection via -timers
+	 * or column selection via -columns yet
+	 */
+	int32 ExportTimerStatisticsAsText(const FString& Filename, FExportTimerStatisticsParams& Params) const;
 
 	void MakeExportTimingEventsColumnList(const FString& InColumnsString, TArray<FName>& OutColumnList) const;
 
@@ -151,6 +167,8 @@ private:
 	const TraceServices::IAnalysisSession& Session;
 	mutable TSet<FName> ExportTimingEventsColumns;
 	mutable TArray<FName> ExportTimingEventsDefaultColumns;
+	mutable TArray<FName> ExportTimerStatisticsDefaultColumns;
+	mutable TArray<FName> ExportTimerStatisticsColumns;
 
 	static const FName ExportTimingEvents_ThreadIdColumn;
 	static const FName ExportTimingEvents_ThreadNameColumn;
diff --git a/Engine/Source/Developer/TraceServices/Public/TraceServices/Containers/Tables.h b/Engine/Source/Developer/TraceServices/Public/TraceServices/Containers/Tables.h
--- a/Engine/Source/Developer/TraceServices/Public/TraceServices/Containers/Tables.h
+++ b/Engine/Source/Developer/TraceServices/Public/TraceServices/Containers/Tables.h
@@ -74,6 +74,6 @@ public:
 	virtual ITableReader<RowType>* CreateReader() const = 0;
 };
 
-void Table2Csv(const IUntypedTable& Table, const TCHAR* Filename);
+TRACESERVICES_API void Table2Csv(const IUntypedTable& Table, const TCHAR* Filename);
 
 } // namespace TraceServices
