[Insights] Add the Name column to the Cooked Packages Table and add aggregations to the table columns.

diff --git a/Engine/Source/Developer/TraceInsights/Private/Insights/CookProfiler/ViewModels/PackageTable.cpp b/Engine/Source/Developer/TraceInsights/Private/Insights/CookProfiler/ViewModels/PackageTable.cpp
--- a/Engine/Source/Developer/TraceInsights/Private/Insights/CookProfiler/ViewModels/PackageTable.cpp
+++ b/Engine/Source/Developer/TraceInsights/Private/Insights/CookProfiler/ViewModels/PackageTable.cpp
@@ -141,7 +141,7 @@ void FPackageTable::AddDefaultColumns()
 		TSharedRef<ITableCellValueSorter> Sorter = MakeShared<FSorterByInt64Value>(ColumnRef);
 		Column.SetValueSorter(Sorter);
 
-		Column.SetAggregation(ETableColumnAggregation::None);
+		Column.SetAggregation(ETableColumnAggregation::SameValue);
 
 		AddColumn(ColumnRef);
 	}
@@ -173,6 +173,8 @@ void FPackageTable::AddDefaultColumns()
 		TSharedRef<ITableCellValueSorter> Sorter = MakeShared<FSorterByDoubleValue>(ColumnRef);
 		Column.SetValueSorter(Sorter);
 
+		Column.SetAggregation(ETableColumnAggregation::Sum);
+
 		AddColumn(ColumnRef);
 	}
 	//////////////////////////////////////////////////
@@ -203,6 +205,8 @@ void FPackageTable::AddDefaultColumns()
 		TSharedRef<ITableCellValueSorter> Sorter = MakeShared<FSorterByDoubleValue>(ColumnRef);
 		Column.SetValueSorter(Sorter);
 
+		Column.SetAggregation(ETableColumnAggregation::Sum);
+
 		AddColumn(ColumnRef);
 	}
 	//////////////////////////////////////////////////
@@ -233,6 +237,8 @@ void FPackageTable::AddDefaultColumns()
 		TSharedRef<ITableCellValueSorter> Sorter = MakeShared<FSorterByDoubleValue>(ColumnRef);
 		Column.SetValueSorter(Sorter);
 
+		Column.SetAggregation(ETableColumnAggregation::Sum);
+
 		AddColumn(ColumnRef);
 	}	
 	//////////////////////////////////////////////////
@@ -263,6 +269,8 @@ void FPackageTable::AddDefaultColumns()
 		TSharedRef<ITableCellValueSorter> Sorter = MakeShared<FSorterByDoubleValue>(ColumnRef);
 		Column.SetValueSorter(Sorter);
 
+		Column.SetAggregation(ETableColumnAggregation::Sum);
+
 		AddColumn(ColumnRef);
 	}
 	//////////////////////////////////////////////////
@@ -293,6 +301,40 @@ void FPackageTable::AddDefaultColumns()
 		TSharedRef<ITableCellValueSorter> Sorter = MakeShared<FSorterByCStringValue>(ColumnRef);
 		Column.SetValueSorter(Sorter);
 
+		Column.SetAggregation(ETableColumnAggregation::SameValue);
+
+		AddColumn(ColumnRef);
+	}
+	//////////////////////////////////////////////////
+	// Package name Column
+	{
+		TSharedRef<FTableColumn> ColumnRef = MakeShared<FTableColumn>(FPackageTableColumns::NameColumnId);
+		FTableColumn& Column = *ColumnRef;
+
+		Column.SetIndex(ColumnIndex++);
+
+		Column.SetShortName(LOCTEXT("PackageNameColumnName", "Package Name"));
+		Column.SetTitleName(LOCTEXT("PackageNameTitle", "PackageName"));
+		Column.SetDescription(LOCTEXT("PackageNameColumnDesc", "The name of the package."));
+
+		Column.SetFlags(ETableColumnFlags::CanBeHidden | ETableColumnFlags::CanBeFiltered);
+
+		Column.SetHorizontalAlignment(HAlign_Left);
+		Column.SetInitialWidth(100.0f);
+
+		Column.SetDataType(ETableCellDataType::CString);
+
+		TSharedRef<ITableCellValueGetter> Getter = MakeShared<FPackageColumnValueGetter<DefaultPackageFieldGetterFuncts::GetName>>();
+		Column.SetValueGetter(Getter);
+
+		TSharedRef<ITableCellValueFormatter> Formatter = MakeShared<FCStringValueFormatterAsText>();
+		Column.SetValueFormatter(Formatter);
+
+		TSharedRef<ITableCellValueSorter> Sorter = MakeShared<FSorterByCStringValue>(ColumnRef);
+		Column.SetValueSorter(Sorter);
+
+		Column.SetAggregation(ETableColumnAggregation::SameValue);
+
 		AddColumn(ColumnRef);
 	}
 }
