[Insights] Show a descriptive status during the csv import process.

diff --git a/Engine/Source/Developer/TraceInsights/Private/Insights/ImportTool/TableImportTool.cpp b/Engine/Source/Developer/TraceInsights/Private/Insights/ImportTool/TableImportTool.cpp
--- a/Engine/Source/Developer/TraceInsights/Private/Insights/ImportTool/TableImportTool.cpp
+++ b/Engine/Source/Developer/TraceInsights/Private/Insights/ImportTool/TableImportTool.cpp
@@ -139,6 +139,7 @@ void FTableImportTool::TableImportServiceCallback(FName TableViewID, TSharedPtr<
 	if (TableViewPtr)
 	{
 		TSharedPtr<SUntypedTableTreeView> TableView = TableViewPtr->TableTreeView;
+		TableView->ClearCurrentOperationNameOverride();
 		TableView->UpdateSourceTable(Table);
 	}
 }
@@ -162,6 +163,8 @@ TSharedRef<SDockTab> FTableImportTool::SpawnTab_TableImportTreeView(const FSpawn
 			.RunInAsyncMode(true)
 		];
 
+	TableTreeView->SetCurrentOperationNameOverride(LOCTEXT("TableImportMessage", "Importing Data"));
+
 	FOpenImportedTableTabData OpenTabData;
 	OpenTabData.TableTreeView = TableTreeView;
 	OpenTabData.Tab = DockTab;
diff --git a/Engine/Source/Developer/TraceInsights/Private/Insights/Table/Widgets/SUntypedTableTreeView.cpp b/Engine/Source/Developer/TraceInsights/Private/Insights/Table/Widgets/SUntypedTableTreeView.cpp
--- a/Engine/Source/Developer/TraceInsights/Private/Insights/Table/Widgets/SUntypedTableTreeView.cpp
+++ b/Engine/Source/Developer/TraceInsights/Private/Insights/Table/Widgets/SUntypedTableTreeView.cpp
@@ -141,6 +141,56 @@ void SUntypedTableTreeView::RebuildTree(bool bResync)
 
 ////////////////////////////////////////////////////////////////////////////////////////////////////
 
+bool SUntypedTableTreeView::IsRunning() const
+{
+	return !CurrentOperationNameOverride.IsEmpty() || STableTreeView::IsRunning();
+}
+
+////////////////////////////////////////////////////////////////////////////////////////////////////
+
+double SUntypedTableTreeView::GetAllOperationsDuration()
+{
+	if (!CurrentOperationNameOverride.IsEmpty())
+	{
+		CurrentOperationStopwatch.Update();
+		return CurrentOperationStopwatch.GetAccumulatedTime();
+	}
+	
+	return STableTreeView::GetAllOperationsDuration();
+}
+
+////////////////////////////////////////////////////////////////////////////////////////////////////
+
+FText SUntypedTableTreeView::GetCurrentOperationName() const
+{
+	if (!CurrentOperationNameOverride.IsEmpty())
+	{
+		return CurrentOperationNameOverride;
+	}
+
+	return STableTreeView::GetCurrentOperationName();
+}
+
+////////////////////////////////////////////////////////////////////////////////////////////////////
+
+void SUntypedTableTreeView::SetCurrentOperationNameOverride(const FText& InOperationName)
+{
+	CurrentOperationStopwatch.Start();
+	CurrentOperationNameOverride = InOperationName;
+}
+
+////////////////////////////////////////////////////////////////////////////////////////////////////
+
+void SUntypedTableTreeView::ClearCurrentOperationNameOverride()
+{
+	CurrentOperationStopwatch.Stop();
+	CurrentOperationStopwatch.Reset();
+
+	CurrentOperationNameOverride = FText::GetEmpty();
+}
+
+////////////////////////////////////////////////////////////////////////////////////////////////////
+
 } // namespace Insights
 
 #undef LOCTEXT_NAMESPACE
diff --git a/Engine/Source/Developer/TraceInsights/Private/Insights/Table/Widgets/SUntypedTableTreeView.h b/Engine/Source/Developer/TraceInsights/Private/Insights/Table/Widgets/SUntypedTableTreeView.h
--- a/Engine/Source/Developer/TraceInsights/Private/Insights/Table/Widgets/SUntypedTableTreeView.h
+++ b/Engine/Source/Developer/TraceInsights/Private/Insights/Table/Widgets/SUntypedTableTreeView.h
@@ -5,6 +5,7 @@
 #include "CoreMinimal.h"
 
 // Insights
+#include "Insights/Common/Stopwatch.h"
 #include "Insights/Table/ViewModels/UntypedTable.h"
 #include "Insights/Table/Widgets/STableTreeView.h"
 
@@ -40,6 +41,19 @@ public:
 
 	virtual void Reset();
 
+	////////////////////////////////////////////////////////////////////////////////////////////////////
+	// IAsyncOperationStatusProvider implementation
+
+	virtual bool IsRunning() const override;
+	virtual double GetAllOperationsDuration() override;
+	virtual double GetCurrentOperationDuration() override { return 0.0; }
+	virtual uint32 GetOperationCount() const override { return 1; }
+	virtual FText GetCurrentOperationName() const override;
+
+	////////////////////////////////////////////////////////////////////////////////////////////////////
+	void SetCurrentOperationNameOverride(const FText& InOperationName);
+	void ClearCurrentOperationNameOverride();
+
 	/**
 	 * Rebuilds the tree (if necessary).
 	 * @param bResync - If true, it forces a resync even if the list did not changed since last sync.
@@ -47,6 +61,8 @@ public:
 	virtual void RebuildTree(bool bResync);
 
 private:
+	FStopwatch CurrentOperationStopwatch;
+	FText CurrentOperationNameOverride;
 };
 
 ////////////////////////////////////////////////////////////////////////////////////////////////////
diff --git a/Engine/Source/Developer/TraceInsights/Private/Insights/Widgets/SAsyncOperationStatus.cpp b/Engine/Source/Developer/TraceInsights/Private/Insights/Widgets/SAsyncOperationStatus.cpp
--- a/Engine/Source/Developer/TraceInsights/Private/Insights/Widgets/SAsyncOperationStatus.cpp
+++ b/Engine/Source/Developer/TraceInsights/Private/Insights/Widgets/SAsyncOperationStatus.cpp
@@ -110,7 +110,17 @@ float SAsyncOperationStatus::ComputeOpacity() const
 
 FText SAsyncOperationStatus::GetText() const
 {
-	return LOCTEXT("DefaultText", "Computing");
+	TSharedPtr<IAsyncOperationStatusProvider> StatusProviderSharedPtr = StatusProvider.Pin();
+	FText OperationName = StatusProviderSharedPtr->GetCurrentOperationName();
+
+	if (OperationName.IsEmpty())
+	{
+		return LOCTEXT("DefaultText", "Computing");
+	}
+	else
+	{
+		return OperationName;
+	}
 }
 
 ////////////////////////////////////////////////////////////////////////////////////////////////////
