[Insights] Fixed session name when a trace name contains a dot char.

diff --git a/Engine/Source/Developer/TraceInsights/Private/Insights/InsightsManager.cpp b/Engine/Source/Developer/TraceInsights/Private/Insights/InsightsManager.cpp
--- a/Engine/Source/Developer/TraceInsights/Private/Insights/InsightsManager.cpp
+++ b/Engine/Source/Developer/TraceInsights/Private/Insights/InsightsManager.cpp
@@ -877,9 +877,14 @@ void FInsightsManager::LoadTrace(uint32 InTraceId, bool InAutoQuit)
 	const UE::Trace::FStoreClient::FTraceInfo* TraceInfo = StoreClient->GetTraceInfoById(InTraceId);
 	if (TraceInfo != nullptr)
 	{
-		FUtf8StringView Name = TraceInfo->GetName();
-		TraceName = FPaths::Combine(TraceName, FString(Name.Len(), Name.GetData()));
-		TraceName = FPaths::SetExtension(TraceName, ".utrace");
+		const FUtf8StringView Utf8NameView = TraceInfo->GetName();
+		FString Name(Utf8NameView);
+		if (!Name.EndsWith(TEXT(".utrace")))
+		{
+			Name += TEXT(".utrace");
+		}
+		TraceName = FPaths::Combine(TraceName, Name);
+		FPaths::NormalizeFilename(TraceName);
 	}
 
 	Session = AnalysisService->StartAnalysis(InTraceId, *TraceName, MoveTemp(TraceData));
