[Insights] Added more validation for renaming a utrace file from the Unreal Insights Session Browser. Trimmed start and end spaces from the new trace name when renaming; added warning when trying to rename a live session; added detection for invalid file names.

diff --git a/Engine/Source/Developer/TraceInsights/Private/Insights/Widgets/SStartPageWindow.cpp b/Engine/Source/Developer/TraceInsights/Private/Insights/Widgets/SStartPageWindow.cpp
--- a/Engine/Source/Developer/TraceInsights/Private/Insights/Widgets/SStartPageWindow.cpp
+++ b/Engine/Source/Developer/TraceInsights/Private/Insights/Widgets/SStartPageWindow.cpp
@@ -212,56 +212,89 @@ public:
 		TSharedPtr<FTraceViewModel> TracePin = WeakTrace.Pin();
 		if (TracePin.IsValid())
 		{
-			if (InCommitType != ETextCommit::OnCleared &&
-				!TracePin->Name.EqualTo(InText))
+			if (InCommitType != ETextCommit::OnCleared)
 			{
-				FString TraceName = TracePin->Name.ToString();
-				FString NewTraceName = InText.ToString();
+				Rename(*TracePin, InText);
+			}
+			TracePin->bIsRenaming = false;
+		}
+	}
 
-				UE_LOG(TraceInsights, Log, TEXT("[TraceStore] Renaming \"%s\" to \"%s\"..."), *TraceName, *NewTraceName);
+	void Rename(FTraceViewModel& Trace, const FText& InText)
+	{
+		const FString TraceName = Trace.Name.ToString();
+		const FString NewTraceName = InText.ToString().TrimStartAndEnd();
 
-				FString TraceFile = TracePin->Uri.ToString();
-				if (FPaths::FileExists(TraceFile))
-				{
-					FString NewTraceFile = FPaths::Combine(FPaths::GetPath(TraceFile), NewTraceName + TEXT(".utrace"));
-					if (FPaths::FileExists(NewTraceFile))
-					{
-						UE_LOG(TraceInsights, Warning, TEXT("[TraceStore] Failed to rename \"%s\" to \"%s\"! File already exists."), *TraceName, *NewTraceName);
+		if (NewTraceName == TraceName)
+		{
+			return;
+		}
 
-						FText Message = FText::Format(LOCTEXT("RenameFailFmt1", "Failed to rename \"{0}\" to \"{1}\"!\nFile already exists."), FText::FromString(TraceName), FText::FromString(NewTraceName));
-						WeakParentWidget.Pin()->ShowFailMessage(Message);
-					}
-					else if (IFileManager::Get().Move(*NewTraceFile, *TraceFile, false))
-					{
-						UE_LOG(TraceInsights, Verbose, TEXT("[TraceStore] Renamed utrace file (\"%s\")."), *NewTraceFile);
-						TracePin->Name = InText;
-						TracePin->Uri = FText::FromString(NewTraceFile);
-
-						FString CacheFile = FPaths::ChangeExtension(TraceFile, TEXT("ucache"));
-						if (FPaths::FileExists(CacheFile))
-						{
-							FString NewCacheFile = FPaths::Combine(FPaths::GetPath(CacheFile), InText.ToString() + TEXT(".ucache"));
-							if (IFileManager::Get().Move(*NewCacheFile, *CacheFile, true))
-							{
-								UE_LOG(TraceInsights, Verbose, TEXT("[TraceStore] Renamed ucache file (\"%s\")."), *NewCacheFile);
-							}
-						}
-
-						FText Message = FText::Format(LOCTEXT("RenameSuccessFmt", "Renamed \"{0}\" to \"{1}\"."), FText::FromString(TraceName), FText::FromString(NewTraceName));
-						WeakParentWidget.Pin()->ShowSuccessMessage(Message);
-					}
-					else
-					{
-						UE_LOG(TraceInsights, Warning, TEXT("[TraceStore] Failed to rename \"%s\" to \"%s\"!"), *TraceName, *NewTraceName);
+		if (Trace.bIsLive)
+		{
+			FText Message = LOCTEXT("RenameLive", "Cannot rename a live session!");
+			WeakParentWidget.Pin()->ShowFailMessage(Message);
+			return;
+		}
 
-						FText Message = FText::Format(LOCTEXT("RenameFailFmt2", "Failed to rename \"{0}\" to \"{1}\"!"), FText::FromString(TraceName), FText::FromString(NewTraceName));
-						WeakParentWidget.Pin()->ShowFailMessage(Message);
-					}
-				}
+		UE_LOG(TraceInsights, Log, TEXT("[TraceStore] Renaming \"%s\" to \"%s\"..."), *TraceName, *NewTraceName);
+
+		const FString TraceFile = Trace.Uri.ToString();
+		const FString NewTraceFile = FPaths::Combine(FPaths::GetPath(TraceFile), NewTraceName + TEXT(".utrace"));
+
+		FText Reason;
+		bool bHasReservedCharacters = false;
+		for (const TCHAR* P = *NewTraceName; *P; ++P)
+		{
+			if (*P == TEXT('/') || *P == TEXT('\\'))
+			{
+				bHasReservedCharacters = true;
+				Reason = LOCTEXT("RenameReservedCharacters", "Name may not contain / or \\ characters.");
+				break;
 			}
+		}
+		if (bHasReservedCharacters || !FPaths::ValidatePath(NewTraceFile, &Reason))
+		{
+			FText Message = FText::Format(LOCTEXT("RenameFailFmt3", "Failed to rename \"{0}\" to \"{1}\"!\n{2}"), Trace.Name, FText::FromString(NewTraceName), Reason);
+			WeakParentWidget.Pin()->ShowFailMessage(Message);
+			return;
+		}
 
-			TracePin->bIsRenaming = false;
+		if (FPaths::FileExists(NewTraceFile))
+		{
+			UE_LOG(TraceInsights, Warning, TEXT("[TraceStore] Failed to rename \"%s\" to \"%s\"! File already exists."), *TraceName, *NewTraceName);
+
+			FText Message = FText::Format(LOCTEXT("RenameFailFmt1", "Failed to rename \"{0}\" to \"{1}\"!\nFile already exists."), Trace.Name, FText::FromString(NewTraceName));
+			WeakParentWidget.Pin()->ShowFailMessage(Message);
+			return;
 		}
+
+		if (!FPaths::FileExists(TraceFile) ||
+			!IFileManager::Get().Move(*NewTraceFile, *TraceFile, false))
+		{
+			UE_LOG(TraceInsights, Warning, TEXT("[TraceStore] Failed to rename \"%s\" to \"%s\"!"), *TraceName, *NewTraceName);
+
+			FText Message = FText::Format(LOCTEXT("RenameFailFmt2", "Failed to rename \"{0}\" to \"{1}\"!"), Trace.Name, FText::FromString(NewTraceName));
+			WeakParentWidget.Pin()->ShowFailMessage(Message);
+			return;
+		}
+
+		UE_LOG(TraceInsights, Verbose, TEXT("[TraceStore] Renamed utrace file (\"%s\")."), *NewTraceFile);
+		Trace.Name = FText::FromString(NewTraceName);
+		Trace.Uri = FText::FromString(NewTraceFile);
+
+		const FString CacheFile = FPaths::ChangeExtension(TraceFile, TEXT("ucache"));
+		if (FPaths::FileExists(CacheFile))
+		{
+			const FString NewCacheFile = FPaths::Combine(FPaths::GetPath(CacheFile), NewTraceName + TEXT(".ucache"));
+			if (IFileManager::Get().Move(*NewCacheFile, *CacheFile, true))
+			{
+				UE_LOG(TraceInsights, Verbose, TEXT("[TraceStore] Renamed ucache file (\"%s\")."), *NewCacheFile);
+			}
+		}
+
+		FText Message = FText::Format(LOCTEXT("RenameSuccessFmt", "Renamed \"{0}\" to \"{1}\"."), Trace.Name, FText::FromString(NewTraceName));
+		WeakParentWidget.Pin()->ShowSuccessMessage(Message);
 	}
 
 	FText GetTraceIndexAndId() const
@@ -2584,7 +2617,7 @@ void STraceStoreWindow::UpdateFiltering()
 		FilterByPlatform->IsEmpty() &&
 		FilterByAppName->IsEmpty() &&
 		FilterByBuildConfig->IsEmpty() &&
-		FilterByBuildTarget->IsEmpty() && 
+		FilterByBuildTarget->IsEmpty() &&
 		FilterByBranch->IsEmpty())
 	{
 		// No filtering.
