[Insights] Add view presets to the Cooked Packages table.

diff --git a/Engine/Source/Developer/TraceInsights/Private/Insights/CookProfiler/Widgets/SPackageTableTreeView.cpp b/Engine/Source/Developer/TraceInsights/Private/Insights/CookProfiler/Widgets/SPackageTableTreeView.cpp
--- a/Engine/Source/Developer/TraceInsights/Private/Insights/CookProfiler/Widgets/SPackageTableTreeView.cpp
+++ b/Engine/Source/Developer/TraceInsights/Private/Insights/CookProfiler/Widgets/SPackageTableTreeView.cpp
@@ -4,6 +4,7 @@
 
 #include "Framework/MultiBox/MultiBoxBuilder.h"
 #include "SlateOptMacros.h"
+#include "Widgets/Input/STextComboBox.h"
 
 // Insights
 #include "Insights/InsightsStyle.h"
@@ -11,6 +12,7 @@
 #include "Insights/CookProfiler/ViewModels/PackageEntry.h"
 #include "Insights/CookProfiler/ViewModels/PackageNode.h"
 #include "Insights/Table/ViewModels/TableColumn.h"
+#include "Insights/Table/ViewModels/TreeNodeGrouping.h"
 #include "Insights/TimingProfilerManager.h"
 #include "Insights/ViewModels/FilterConfigurator.h"
 #include "Insights/ViewModels/ThreadTimingTrack.h"
@@ -74,6 +76,8 @@ void SPackageTableTreeView::Construct(const FArguments& InArgs, TSharedPtr<FPack
 	ConstructWidget(InTablePtr);
 
 	AddCommmands();
+
+	ViewPreset_OnSelectionChanged(AvailableViewPresets[0], ESelectInfo::Direct);
 }
 
 ////////////////////////////////////////////////////////////////////////////////////////////////////
@@ -177,7 +181,38 @@ FText SPackageTableTreeView::GetCurrentOperationName() const
 BEGIN_SLATE_FUNCTION_BUILD_OPTIMIZATION
 TSharedPtr<SWidget> SPackageTableTreeView::ConstructToolbar()
 {
-	return nullptr;
+	TSharedPtr<SHorizontalBox> Box = SNew(SHorizontalBox);
+
+	Box->AddSlot()
+		.AutoWidth()
+		.VAlign(VAlign_Center)
+		.Padding(0.0f, 0.0f, 4.0f, 0.0f)
+		[
+			SNew(STextBlock)
+			.Text(LOCTEXT("Preset", "Preset:"))
+		];
+
+	Box->AddSlot()
+		.AutoWidth()
+		.VAlign(VAlign_Center)
+		[
+			SNew(SBox)
+			.MinDesiredWidth(150)
+			[
+				SAssignNew(PresetComboBox, SComboBox<TSharedRef<ITableTreeViewPreset>>)
+				.ToolTipText(this, &SPackageTableTreeView::ViewPreset_GetSelectedToolTipText)
+				.OptionsSource(GetAvailableViewPresets())
+				.OnSelectionChanged(this, &SPackageTableTreeView::ViewPreset_OnSelectionChanged)
+				.InitiallySelectedItem(SelectedViewPreset)
+				.OnGenerateWidget(this, &SPackageTableTreeView::ViewPreset_OnGenerateWidget)
+				[
+					SNew(STextBlock)
+					.Text(this, &SPackageTableTreeView::ViewPreset_GetSelectedText)
+				]
+			]
+		];
+
+	return Box;
 }
 
 ////////////////////////////////////////////////////////////////////////////////////////////////////
@@ -190,44 +225,203 @@ END_SLATE_FUNCTION_BUILD_OPTIMIZATION
 
 ////////////////////////////////////////////////////////////////////////////////////////////////////
 
-void SPackageTableTreeView::ApplyColumnConfig(const TArrayView<FColumnConfig>& Preset)
+void SPackageTableTreeView::InternalCreateGroupings()
 {
-	for (const TSharedRef<FTableColumn>& ColumnRef : Table->GetColumns())
-	{
-		FTableColumn& Column = ColumnRef.Get();
-		for (const FColumnConfig& Config : Preset)
+	STableTreeView::InternalCreateGroupings();
+
+	AvailableGroupings.RemoveAll(
+		[](TSharedPtr<FTreeNodeGrouping>& Grouping)
 		{
-			if (Column.GetId() == Config.ColumnId)
+			if (Grouping->Is<FTreeNodeGroupingByUniqueValue>())
 			{
-				if (Config.bIsVisible)
+				const FName ColumnId = Grouping->As<FTreeNodeGroupingByUniqueValue>().GetColumnId();
+				if (ColumnId == FPackageTableColumns::BeginCacheForCookedPlatformDataTimeColumnId ||
+					ColumnId == FPackageTableColumns::GetIsCachedCookedPlatformDataLoadedColumnId ||
+					ColumnId == FPackageTableColumns::SaveTimeColumnId ||
+					ColumnId == FPackageTableColumns::IdColumnId ||
+					ColumnId == FPackageTableColumns::NameColumnId ||
+					ColumnId == FPackageTableColumns::LoadTimeColumnId)
 				{
-					ShowColumn(Column);
-					if (Config.Width > 0.0f)
-					{
-						TreeViewHeaderRow->SetColumnWidth(Column.GetId(), Config.Width);
-					}
+					return true;
 				}
-				else
+			}
+			else if (Grouping->Is<FTreeNodeGroupingByPathBreakdown>())
+			{
+				const FName ColumnId = Grouping->As<FTreeNodeGroupingByPathBreakdown>().GetColumnId();
+				if (ColumnId == FPackageTableColumns::PackageAssetClassColumnId)
 				{
-					HideColumn(Column);
+					return true;
 				}
 			}
-		}
-	}
+			return false;
+		});
 }
 
 ////////////////////////////////////////////////////////////////////////////////////////////////////
 
-void SPackageTableTreeView::InternalCreateGroupings()
+void SPackageTableTreeView::TreeView_OnMouseButtonDoubleClick(FTableTreeNodePtr TreeNode)
 {
-	STableTreeView::InternalCreateGroupings();
+	STableTreeView::TreeView_OnMouseButtonDoubleClick(TreeNode);
 }
 
 ////////////////////////////////////////////////////////////////////////////////////////////////////
 
-void SPackageTableTreeView::TreeView_OnMouseButtonDoubleClick(FTableTreeNodePtr TreeNode)
+void SPackageTableTreeView::InitAvailableViewPresets()
 {
-	STableTreeView::TreeView_OnMouseButtonDoubleClick(TreeNode);
+	//////////////////////////////////////////////////
+	// Default View
+
+	class FDefaultViewPreset : public ITableTreeViewPreset
+	{
+	public:
+		virtual FText GetName() const override
+		{
+			return LOCTEXT("Default_PresetName", "Default");
+		}
+		virtual FText GetToolTip() const override
+		{
+			return LOCTEXT("Default_PresetToolTip", "Default View\nConfigure the tree view to show default packages info.");
+		}
+		virtual FName GetSortColumn() const override
+		{
+			return FTable::GetHierarchyColumnId();
+		}
+		virtual EColumnSortMode::Type GetSortMode() const override
+		{
+			return EColumnSortMode::Type::Ascending;
+		}
+		virtual void SetCurrentGroupings(const TArray<TSharedPtr<FTreeNodeGrouping>>& InAvailableGroupings, TArray<TSharedPtr<FTreeNodeGrouping>>& InOutCurrentGroupings) const override
+		{
+			InOutCurrentGroupings.Reset();
+
+			check(InAvailableGroupings[0]->Is<FTreeNodeGroupingFlat>());
+			InOutCurrentGroupings.Add(InAvailableGroupings[0]);
+		}
+		virtual void GetColumnConfigSet(TArray<FTableColumnConfig>& InOutConfigSet) const override
+		{
+			InOutConfigSet.Add({ FTable::GetHierarchyColumnId(), true, 500.0f });
+			InOutConfigSet.Add({ FPackageTableColumns::IdColumnId, true, 100.0f });
+			InOutConfigSet.Add({ FPackageTableColumns::LoadTimeColumnId, true, 100.0f });
+			InOutConfigSet.Add({ FPackageTableColumns::SaveTimeColumnId, true, 100.0f });
+			InOutConfigSet.Add({ FPackageTableColumns::BeginCacheForCookedPlatformDataTimeColumnId, true, 100.0f });
+			InOutConfigSet.Add({ FPackageTableColumns::GetIsCachedCookedPlatformDataLoadedColumnId, true, 100.0f });
+			InOutConfigSet.Add({ FPackageTableColumns::PackageAssetClassColumnId, true, 100.0f });
+			InOutConfigSet.Add({ FPackageTableColumns::NameColumnId, false, 100.0f });
+		}
+	};
+	AvailableViewPresets.Add(MakeShared<FDefaultViewPreset>());
+
+	//////////////////////////////////////////////////
+	// Package Path Breakdown
+
+	class FPackagePathViewPreset : public ITableTreeViewPreset
+	{
+	public:
+		virtual FText GetName() const override
+		{
+			return LOCTEXT("PackagePath_PresetName", "Package Path");
+		}
+		virtual FText GetToolTip() const override
+		{
+			return LOCTEXT("PackagePath_PresetToolTip", "Configure the tree view to show the packages grouped by package path.");
+		}
+		virtual FName GetSortColumn() const override
+		{
+			return FTable::GetHierarchyColumnId();
+		}
+		virtual EColumnSortMode::Type GetSortMode() const override
+		{
+			return EColumnSortMode::Type::Ascending;
+		}
+		virtual void SetCurrentGroupings(const TArray<TSharedPtr<FTreeNodeGrouping>>& InAvailableGroupings, TArray<TSharedPtr<FTreeNodeGrouping>>& InOutCurrentGroupings) const override
+		{
+			InOutCurrentGroupings.Reset();
+
+			check(InAvailableGroupings[0]->Is<FTreeNodeGroupingFlat>());
+			InOutCurrentGroupings.Add(InAvailableGroupings[0]);
+
+			const TSharedPtr<FTreeNodeGrouping>* PackagePathGrouping = InAvailableGroupings.FindByPredicate(
+				[](TSharedPtr<FTreeNodeGrouping>& Grouping)
+				{
+					return Grouping->Is<FTreeNodeGroupingByPathBreakdown>() &&
+						Grouping->As<FTreeNodeGroupingByPathBreakdown>().GetColumnId() == FPackageTableColumns::NameColumnId;
+				});
+
+			if (PackagePathGrouping)
+			{
+				InOutCurrentGroupings.Add(*PackagePathGrouping);
+			}
+		}
+		virtual void GetColumnConfigSet(TArray<FTableColumnConfig>& InOutConfigSet) const override
+		{
+			InOutConfigSet.Add({ FTable::GetHierarchyColumnId(), true, 500.0f });
+			InOutConfigSet.Add({ FPackageTableColumns::IdColumnId, true, 100.0f });
+			InOutConfigSet.Add({ FPackageTableColumns::LoadTimeColumnId, true, 100.0f });
+			InOutConfigSet.Add({ FPackageTableColumns::SaveTimeColumnId, true, 100.0f });
+			InOutConfigSet.Add({ FPackageTableColumns::BeginCacheForCookedPlatformDataTimeColumnId, true, 100.0f });
+			InOutConfigSet.Add({ FPackageTableColumns::GetIsCachedCookedPlatformDataLoadedColumnId, true, 100.0f });
+			InOutConfigSet.Add({ FPackageTableColumns::PackageAssetClassColumnId, true, 100.0f });
+			InOutConfigSet.Add({ FPackageTableColumns::NameColumnId, false, 100.0f });
+		}
+	};
+	AvailableViewPresets.Add(MakeShared<FPackagePathViewPreset>());
+
+	//////////////////////////////////////////////////
+	// Asset Class Breakdown
+
+	class FAssetClassViewPreset : public ITableTreeViewPreset
+	{
+	public:
+		virtual FText GetName() const override
+		{
+			return LOCTEXT("AssetClass_PresetName", "Asset Class");
+		}
+		virtual FText GetToolTip() const override
+		{
+			return LOCTEXT("AssetClass_PresetToolTip", "Configure the tree view to show the packages grouped by their most important asset class.");
+		}
+		virtual FName GetSortColumn() const override
+		{
+			return FTable::GetHierarchyColumnId();
+		}
+		virtual EColumnSortMode::Type GetSortMode() const override
+		{
+			return EColumnSortMode::Type::Ascending;
+		}
+		virtual void SetCurrentGroupings(const TArray<TSharedPtr<FTreeNodeGrouping>>& InAvailableGroupings, TArray<TSharedPtr<FTreeNodeGrouping>>& InOutCurrentGroupings) const override
+		{
+			InOutCurrentGroupings.Reset();
+
+			check(InAvailableGroupings[0]->Is<FTreeNodeGroupingFlat>());
+			InOutCurrentGroupings.Add(InAvailableGroupings[0]);
+
+			const TSharedPtr<FTreeNodeGrouping>* PackagePathGrouping = InAvailableGroupings.FindByPredicate(
+				[](TSharedPtr<FTreeNodeGrouping>& Grouping)
+				{
+					return Grouping->Is<FTreeNodeGroupingByUniqueValue>() &&
+						Grouping->As<FTreeNodeGroupingByUniqueValue>().GetColumnId() == FPackageTableColumns::PackageAssetClassColumnId;
+				});
+
+			if (PackagePathGrouping)
+			{
+				InOutCurrentGroupings.Add(*PackagePathGrouping);
+			}
+		}
+		virtual void GetColumnConfigSet(TArray<FTableColumnConfig>& InOutConfigSet) const override
+		{
+			InOutConfigSet.Add({ FTable::GetHierarchyColumnId(), true, 500.0f });
+			InOutConfigSet.Add({ FPackageTableColumns::IdColumnId, true, 100.0f });
+			InOutConfigSet.Add({ FPackageTableColumns::LoadTimeColumnId, true, 100.0f });
+			InOutConfigSet.Add({ FPackageTableColumns::SaveTimeColumnId, true, 100.0f });
+			InOutConfigSet.Add({ FPackageTableColumns::BeginCacheForCookedPlatformDataTimeColumnId, true, 100.0f });
+			InOutConfigSet.Add({ FPackageTableColumns::GetIsCachedCookedPlatformDataLoadedColumnId, true, 100.0f });
+			InOutConfigSet.Add({ FPackageTableColumns::NameColumnId, true, 100.0f });
+			InOutConfigSet.Add({ FPackageTableColumns::PackageAssetClassColumnId, false, 100.0f });
+		}
+	};
+	AvailableViewPresets.Add(MakeShared<FAssetClassViewPreset>());
+
+	SelectedViewPreset = AvailableViewPresets[0];
 }
 
 ////////////////////////////////////////////////////////////////////////////////////////////////////
diff --git a/Engine/Source/Developer/TraceInsights/Private/Insights/CookProfiler/Widgets/SPackageTableTreeView.h b/Engine/Source/Developer/TraceInsights/Private/Insights/CookProfiler/Widgets/SPackageTableTreeView.h
--- a/Engine/Source/Developer/TraceInsights/Private/Insights/CookProfiler/Widgets/SPackageTableTreeView.h
+++ b/Engine/Source/Developer/TraceInsights/Private/Insights/CookProfiler/Widgets/SPackageTableTreeView.h
@@ -18,14 +18,6 @@ namespace Insights
 
 class SPackageTableTreeView : public STableTreeView
 {
-private:
-	struct FColumnConfig
-	{
-		const FName& ColumnId;
-		bool bIsVisible;
-		float Width;
-	};
-
 public:
 	/** Default constructor. */
 	SPackageTableTreeView();
@@ -82,9 +74,10 @@ protected:
 	virtual void ExtendMenu(FMenuBuilder& MenuBuilder) override;
 
 private:
-	void ApplyColumnConfig(const TArrayView<FColumnConfig>& Preset);
 	void AddCommmands();
 
+	virtual void InitAvailableViewPresets() override;
+
 private:
 	bool bDataLoaded = false;
 };
diff --git a/Engine/Source/Developer/TraceInsights/Private/Insights/MemoryProfiler/Widgets/SMemAllocTableTreeView.cpp b/Engine/Source/Developer/TraceInsights/Private/Insights/MemoryProfiler/Widgets/SMemAllocTableTreeView.cpp
--- a/Engine/Source/Developer/TraceInsights/Private/Insights/MemoryProfiler/Widgets/SMemAllocTableTreeView.cpp
+++ b/Engine/Source/Developer/TraceInsights/Private/Insights/MemoryProfiler/Widgets/SMemAllocTableTreeView.cpp
@@ -54,7 +54,6 @@ const int SMemAllocTableTreeView::LLMFilterIndex = 0x0000FFFFE;
 SMemAllocTableTreeView::SMemAllocTableTreeView()
 {
 	bRunInAsyncMode = true;
-	InitAvailableViewPresets();
 }
 
 ////////////////////////////////////////////////////////////////////////////////////////////////////
@@ -560,7 +559,7 @@ TSharedPtr<SWidget> SMemAllocTableTreeView::ConstructToolbar()
 			SNew(SBox)
 			.MinDesiredWidth(150)
 			[
-				SAssignNew(PresetComboBox, SComboBox<TSharedRef<IViewPreset>>)
+				SAssignNew(PresetComboBox, SComboBox<TSharedRef<ITableTreeViewPreset>>)
 				.ToolTipText(this, &SMemAllocTableTreeView::ViewPreset_GetSelectedToolTipText)
 				.OptionsSource(GetAvailableViewPresets())
 				.OnSelectionChanged(this, &SMemAllocTableTreeView::ViewPreset_OnSelectionChanged)
@@ -572,7 +571,7 @@ TSharedPtr<SWidget> SMemAllocTableTreeView::ConstructToolbar()
 			]
 		];
 
-	//for (const TSharedRef<IViewPreset>& ViewPreset : AvailableViewPresets)
+	//for (const TSharedRef<ITableTreeViewPreset>& ViewPreset : AvailableViewPresets)
 	//{
 	//	Box->AddSlot()
 	//		.AutoWidth()
@@ -581,7 +580,7 @@ TSharedPtr<SWidget> SMemAllocTableTreeView::ConstructToolbar()
 	//			SNew(SButton)
 	//			.Text(ViewPreset->GetName())
 	//			.ToolTipText(ViewPreset->GetToolTip())
-	//			.OnClicked(this, &SMemAllocTableTreeView::OnApplyViewPreset, (const IViewPreset*)&ViewPreset.Get())
+	//			.OnClicked(this, &SMemAllocTableTreeView::OnApplyViewPreset, (const ITableTreeViewPreset*)&ViewPreset.Get())
 	//		];
 	//}
 
@@ -603,7 +602,7 @@ void SMemAllocTableTreeView::InitAvailableViewPresets()
 	//////////////////////////////////////////////////
 	// Default View
 
-	class FDefaultViewPreset : public IViewPreset
+	class FDefaultViewPreset : public ITableTreeViewPreset
 	{
 	public:
 		virtual FText GetName() const override
@@ -629,7 +628,7 @@ void SMemAllocTableTreeView::InitAvailableViewPresets()
 			check(InAvailableGroupings[0]->Is<FTreeNodeGroupingFlat>());
 			InOutCurrentGroupings.Add(InAvailableGroupings[0]);
 		}
-		virtual void GetColumnConfigSet(TArray<FColumnConfig>& InOutConfigSet) const override
+		virtual void GetColumnConfigSet(TArray<FTableColumnConfig>& InOutConfigSet) const override
 		{
 			InOutConfigSet.Add({ FTable::GetHierarchyColumnId(),               true, 200.0f });
 			InOutConfigSet.Add({ FMemAllocTableColumns::CountColumnId,         true, 100.0f });
@@ -644,7 +643,7 @@ void SMemAllocTableTreeView::InitAvailableViewPresets()
 	//////////////////////////////////////////////////
 	// Detailed View
 
-	class FDetailedViewPreset : public IViewPreset
+	class FDetailedViewPreset : public ITableTreeViewPreset
 	{
 	public:
 		virtual FText GetName() const override
@@ -670,7 +669,7 @@ void SMemAllocTableTreeView::InitAvailableViewPresets()
 			check(InAvailableGroupings[0]->Is<FTreeNodeGroupingFlat>());
 			InOutCurrentGroupings.Add(InAvailableGroupings[0]);
 		}
-		virtual void GetColumnConfigSet(TArray<FColumnConfig>& InOutConfigSet) const override
+		virtual void GetColumnConfigSet(TArray<FTableColumnConfig>& InOutConfigSet) const override
 		{
 			InOutConfigSet.Add({ FTable::GetHierarchyColumnId(),                 true, 200.0f });
 			InOutConfigSet.Add({ FMemAllocTableColumns::StartEventIndexColumnId, true, 100.0f });
@@ -694,7 +693,7 @@ void SMemAllocTableTreeView::InitAvailableViewPresets()
 	//////////////////////////////////////////////////
 	// Heap Breakdown View
 
-	class FHeapViewPreset : public IViewPreset
+	class FHeapViewPreset : public ITableTreeViewPreset
 	{
 	public:
 		virtual FText GetName() const override
@@ -730,7 +729,7 @@ void SMemAllocTableTreeView::InitAvailableViewPresets()
 				InOutCurrentGroupings.Add(*HeapGrouping);
 			}
 		}
-		virtual void GetColumnConfigSet(TArray<FColumnConfig>& InOutConfigSet) const override
+		virtual void GetColumnConfigSet(TArray<FTableColumnConfig>& InOutConfigSet) const override
 		{
 			InOutConfigSet.Add({ FTable::GetHierarchyColumnId(),          true, 400.0f });
 			InOutConfigSet.Add({ FMemAllocTableColumns::CountColumnId,    true, 100.0f });
@@ -744,7 +743,7 @@ void SMemAllocTableTreeView::InitAvailableViewPresets()
 	//////////////////////////////////////////////////
 	// Size Breakdown View
 
-	class FSizeViewPreset : public IViewPreset
+	class FSizeViewPreset : public ITableTreeViewPreset
 	{
 	public:
 		virtual FText GetName() const override
@@ -780,7 +779,7 @@ void SMemAllocTableTreeView::InitAvailableViewPresets()
 				InOutCurrentGroupings.Add(*SizeGrouping);
 			}
 		}
-		virtual void GetColumnConfigSet(TArray<FColumnConfig>& InOutConfigSet) const override
+		virtual void GetColumnConfigSet(TArray<FTableColumnConfig>& InOutConfigSet) const override
 		{
 			InOutConfigSet.Add({ FTable::GetHierarchyColumnId(),          true, 200.0f });
 			InOutConfigSet.Add({ FMemAllocTableColumns::AddressColumnId,  true, 120.0f });
@@ -795,7 +794,7 @@ void SMemAllocTableTreeView::InitAvailableViewPresets()
 	//////////////////////////////////////////////////
 	// Tag Breakdown View
 
-	class FTagViewPreset : public IViewPreset
+	class FTagViewPreset : public ITableTreeViewPreset
 	{
 	public:
 		virtual FText GetName() const override
@@ -831,7 +830,7 @@ void SMemAllocTableTreeView::InitAvailableViewPresets()
 				InOutCurrentGroupings.Add(*TagGrouping);
 			}
 		}
-		virtual void GetColumnConfigSet(TArray<FColumnConfig>& InOutConfigSet) const override
+		virtual void GetColumnConfigSet(TArray<FTableColumnConfig>& InOutConfigSet) const override
 		{
 			InOutConfigSet.Add({ FTable::GetHierarchyColumnId(),          true, 200.0f });
 			InOutConfigSet.Add({ FMemAllocTableColumns::CountColumnId,    true, 100.0f });
@@ -844,7 +843,7 @@ void SMemAllocTableTreeView::InitAvailableViewPresets()
 	//////////////////////////////////////////////////
 	// Asset Breakdown View
 
-	class FAssetViewPreset : public IViewPreset
+	class FAssetViewPreset : public ITableTreeViewPreset
 	{
 	public:
 		virtual FText GetName() const override
@@ -881,7 +880,7 @@ void SMemAllocTableTreeView::InitAvailableViewPresets()
 				InOutCurrentGroupings.Add(*AssetGrouping);
 			}
 		}
-		virtual void GetColumnConfigSet(TArray<FColumnConfig>& InOutConfigSet) const override
+		virtual void GetColumnConfigSet(TArray<FTableColumnConfig>& InOutConfigSet) const override
 		{
 			InOutConfigSet.Add({ FTable::GetHierarchyColumnId(),          true, 200.0f });
 			InOutConfigSet.Add({ FMemAllocTableColumns::CountColumnId,    true, 100.0f });
@@ -895,7 +894,7 @@ void SMemAllocTableTreeView::InitAvailableViewPresets()
 	//////////////////////////////////////////////////
 	// Class Name Breakdown View
 
-	class FClassNameViewPreset : public IViewPreset
+	class FClassNameViewPreset : public ITableTreeViewPreset
 	{
 	public:
 		virtual FText GetName() const override
@@ -943,7 +942,7 @@ void SMemAllocTableTreeView::InitAvailableViewPresets()
 				InOutCurrentGroupings.Add(*AssetGrouping);
 			}
 		}
-		virtual void GetColumnConfigSet(TArray<FColumnConfig>& InOutConfigSet) const override
+		virtual void GetColumnConfigSet(TArray<FTableColumnConfig>& InOutConfigSet) const override
 		{
 			InOutConfigSet.Add({ FTable::GetHierarchyColumnId(),          true, 200.0f });
 			InOutConfigSet.Add({ FMemAllocTableColumns::CountColumnId,    true, 100.0f });
@@ -957,7 +956,7 @@ void SMemAllocTableTreeView::InitAvailableViewPresets()
 	//////////////////////////////////////////////////
 	// (Inverted) Callstack Breakdown View
 
-	class FCallstackViewPreset : public IViewPreset
+	class FCallstackViewPreset : public ITableTreeViewPreset
 	{
 	public:
 		FCallstackViewPreset(bool bIsInverted)
@@ -1004,7 +1003,7 @@ void SMemAllocTableTreeView::InitAvailableViewPresets()
 				InOutCurrentGroupings.Add(*CallstackGrouping);
 			}
 		}
-		virtual void GetColumnConfigSet(TArray<FColumnConfig>& InOutConfigSet) const override
+		virtual void GetColumnConfigSet(TArray<FTableColumnConfig>& InOutConfigSet) const override
 		{
 			InOutConfigSet.Add({ FTable::GetHierarchyColumnId(),          true, 400.0f });
 			InOutConfigSet.Add({ FMemAllocTableColumns::CountColumnId,    true, 100.0f });
@@ -1022,7 +1021,7 @@ void SMemAllocTableTreeView::InitAvailableViewPresets()
 	//////////////////////////////////////////////////
 	// Address (4K Page) Breakdown View
 
-	class FPageViewPreset : public IViewPreset
+	class FPageViewPreset : public ITableTreeViewPreset
 	{
 	public:
 		virtual FText GetName() const override
@@ -1059,7 +1058,7 @@ void SMemAllocTableTreeView::InitAvailableViewPresets()
 				InOutCurrentGroupings.Add(*MemoryPageGrouping);
 			}
 		}
-		virtual void GetColumnConfigSet(TArray<FColumnConfig>& InOutConfigSet) const override
+		virtual void GetColumnConfigSet(TArray<FTableColumnConfig>& InOutConfigSet) const override
 		{
 			InOutConfigSet.Add({ FTable::GetHierarchyColumnId(),          true, 200.0f });
 			InOutConfigSet.Add({ FMemAllocTableColumns::AddressColumnId,  true, 120.0f });
@@ -1078,94 +1077,6 @@ void SMemAllocTableTreeView::InitAvailableViewPresets()
 
 ////////////////////////////////////////////////////////////////////////////////////////////////////
 
-FReply SMemAllocTableTreeView::OnApplyViewPreset(const IViewPreset* InPreset)
-{
-	ApplyViewPreset(*InPreset);
-	return FReply::Handled();
-}
-
-////////////////////////////////////////////////////////////////////////////////////////////////////
-
-void SMemAllocTableTreeView::ApplyViewPreset(const IViewPreset& InPreset)
-{
-	ColumnBeingSorted = InPreset.GetSortColumn();
-	ColumnSortMode = InPreset.GetSortMode();
-	UpdateCurrentSortingByColumn();
-
-	PreChangeGroupings();
-	InPreset.SetCurrentGroupings(AvailableGroupings, CurrentGroupings);
-	PostChangeGroupings();
-
-	TArray<FColumnConfig> ColumnConfigSet;
-	InPreset.GetColumnConfigSet(ColumnConfigSet);
-	ApplyColumnConfig(ColumnConfigSet);
-}
-
-////////////////////////////////////////////////////////////////////////////////////////////////////
-
-void SMemAllocTableTreeView::ApplyColumnConfig(const TArrayView<FColumnConfig>& InColumnConfigSet)
-{
-	// TODO: Reorder columns as in the config set.
-	// Currenly we only apply visibility and column width.
-	for (const TSharedRef<FTableColumn>& ColumnRef : Table->GetColumns())
-	{
-		FTableColumn& Column = ColumnRef.Get();
-		const FName ColumnId = Column.GetId();
-		const FColumnConfig* ConfigPtr = InColumnConfigSet.FindByPredicate([ColumnId](const FColumnConfig& Config) { return ColumnId == Config.ColumnId; });
-		if (ConfigPtr && ConfigPtr->bIsVisible)
-		{
-			ShowColumn(Column);
-			if (ConfigPtr->Width > 0.0f)
-			{
-				TreeViewHeaderRow->SetColumnWidth(ColumnId, ConfigPtr->Width);
-			}
-		}
-		else
-		{
-			HideColumn(Column);
-		}
-	}
-}
-
-////////////////////////////////////////////////////////////////////////////////////////////////////
-
-void SMemAllocTableTreeView::ViewPreset_OnSelectionChanged(TSharedPtr<IViewPreset> InPreset, ESelectInfo::Type SelectInfo)
-{
-	SelectedViewPreset = InPreset;
-	if (InPreset.IsValid())
-	{
-		ApplyViewPreset(*InPreset);
-	}
-}
-
-////////////////////////////////////////////////////////////////////////////////////////////////////
-
-BEGIN_SLATE_FUNCTION_BUILD_OPTIMIZATION
-TSharedRef<SWidget> SMemAllocTableTreeView::ViewPreset_OnGenerateWidget(TSharedRef<IViewPreset> InPreset)
-{
-	return SNew(STextBlock)
-		.Text(InPreset->GetName())
-		.ToolTipText(InPreset->GetToolTip())
-		.Margin(2.0f);
-}
-END_SLATE_FUNCTION_BUILD_OPTIMIZATION
-
-////////////////////////////////////////////////////////////////////////////////////////////////////
-
-FText SMemAllocTableTreeView::ViewPreset_GetSelectedText() const
-{
-	return SelectedViewPreset ? SelectedViewPreset->GetName() : LOCTEXT("Custom_ToolTip", "Custom");
-}
-
-////////////////////////////////////////////////////////////////////////////////////////////////////
-
-FText SMemAllocTableTreeView::ViewPreset_GetSelectedToolTipText() const
-{
-	return SelectedViewPreset ? SelectedViewPreset->GetToolTip() : LOCTEXT("CustomPreset_ToolTip", "Custom Preset");
-}
-
-////////////////////////////////////////////////////////////////////////////////////////////////////
-
 BEGIN_SLATE_FUNCTION_BUILD_OPTIMIZATION
 TSharedPtr<SWidget> SMemAllocTableTreeView::ConstructFooter()
 {
diff --git a/Engine/Source/Developer/TraceInsights/Private/Insights/MemoryProfiler/Widgets/SMemAllocTableTreeView.h b/Engine/Source/Developer/TraceInsights/Private/Insights/MemoryProfiler/Widgets/SMemAllocTableTreeView.h
--- a/Engine/Source/Developer/TraceInsights/Private/Insights/MemoryProfiler/Widgets/SMemAllocTableTreeView.h
+++ b/Engine/Source/Developer/TraceInsights/Private/Insights/MemoryProfiler/Widgets/SMemAllocTableTreeView.h
@@ -4,7 +4,6 @@
 
 #include "CoreMinimal.h"
 #include "TraceServices/Model/AllocationsProvider.h"
-#include "Widgets/Input/SComboBox.h"
 
 // Insights
 #include "Insights/Common/Stopwatch.h"
@@ -21,25 +20,6 @@ class FMemoryRuleSpec;
 
 class SMemAllocTableTreeView : public STableTreeView
 {
-private:
-	struct FColumnConfig
-	{
-		FName ColumnId;
-		bool bIsVisible;
-		float Width;
-	};
-
-	class IViewPreset
-	{
-	public:
-		virtual FText GetName() const = 0;
-		virtual FText GetToolTip() const = 0;
-		virtual FName GetSortColumn() const = 0;
-		virtual EColumnSortMode::Type GetSortMode() const = 0;
-		virtual void SetCurrentGroupings(const TArray<TSharedPtr<FTreeNodeGrouping>>& InAvailableGroupings, TArray<TSharedPtr<FTreeNodeGrouping>>& InOutCurrentGroupings) const = 0;
-		virtual void GetColumnConfigSet(TArray<FColumnConfig>& InOutConfigSet) const = 0;
-	};
-
 public:
 	/** Default constructor. */
 	SMemAllocTableTreeView();
@@ -137,14 +117,6 @@ private:
 	ECheckBoxState CallstackGroupingByFunction_IsChecked() const;
 
 	void InitAvailableViewPresets();
-	const TArray<TSharedRef<IViewPreset>>* GetAvailableViewPresets() const { return &AvailableViewPresets; }
-	FReply OnApplyViewPreset(const IViewPreset* InPreset);
-	void ApplyViewPreset(const IViewPreset& InPreset);
-	void ApplyColumnConfig(const TArrayView<FColumnConfig>& InTableConfig);
-	void ViewPreset_OnSelectionChanged(TSharedPtr<IViewPreset> InPreset, ESelectInfo::Type SelectInfo);
-	TSharedRef<SWidget> ViewPreset_OnGenerateWidget(TSharedRef<IViewPreset> InPreset);
-	FText ViewPreset_GetSelectedText() const;
-	FText ViewPreset_GetSelectedToolTipText() const;
 	void PopulateLLMTagSuggestionList(const FString& Text, TArray<FString>& OutSuggestions);
 
 private:
@@ -159,9 +131,6 @@ private:
 	FStopwatch QueryStopwatch;
 	bool bHasPendingQueryReset = false;
 	bool bIsCallstackGroupingByFunction = true;
-	TArray<TSharedRef<IViewPreset>> AvailableViewPresets;
-	TSharedPtr<IViewPreset> SelectedViewPreset;
-	TSharedPtr<SComboBox<TSharedRef<IViewPreset>>> PresetComboBox;
 };
 
 ////////////////////////////////////////////////////////////////////////////////////////////////////
diff --git a/Engine/Source/Developer/TraceInsights/Private/Insights/Table/Widgets/STableTreeView.cpp b/Engine/Source/Developer/TraceInsights/Private/Insights/Table/Widgets/STableTreeView.cpp
--- a/Engine/Source/Developer/TraceInsights/Private/Insights/Table/Widgets/STableTreeView.cpp
+++ b/Engine/Source/Developer/TraceInsights/Private/Insights/Table/Widgets/STableTreeView.cpp
@@ -177,6 +177,8 @@ void STableTreeView::ConstructWidget(TSharedPtr<FTable> InTablePtr)
 	check(InTablePtr.IsValid());
 	Table = InTablePtr;
 
+	InitAvailableViewPresets();
+
 	SAssignNew(ExternalScrollbar, SScrollBar)
 	.AlwaysShowScrollbar(true);
 
@@ -3552,6 +3554,94 @@ void STableTreeView::ExportToFileRec(const FBaseTreeNodePtr& InGroupNode, TArray
 
 ////////////////////////////////////////////////////////////////////////////////////////////////////
 
+FReply STableTreeView::OnApplyViewPreset(const ITableTreeViewPreset* InPreset)
+{
+	ApplyViewPreset(*InPreset);
+	return FReply::Handled();
+}
+
+////////////////////////////////////////////////////////////////////////////////////////////////////
+
+void STableTreeView::ApplyViewPreset(const ITableTreeViewPreset& InPreset)
+{
+	ColumnBeingSorted = InPreset.GetSortColumn();
+	ColumnSortMode = InPreset.GetSortMode();
+	UpdateCurrentSortingByColumn();
+
+	PreChangeGroupings();
+	InPreset.SetCurrentGroupings(AvailableGroupings, CurrentGroupings);
+	PostChangeGroupings();
+
+	TArray<FTableColumnConfig > ColumnConfigSet;
+	InPreset.GetColumnConfigSet(ColumnConfigSet);
+	ApplyColumnConfig(ColumnConfigSet);
+}
+
+////////////////////////////////////////////////////////////////////////////////////////////////////
+
+void STableTreeView::ApplyColumnConfig(const TArrayView<FTableColumnConfig >& InColumnConfigSet)
+{
+	// TODO: Reorder columns as in the config set.
+	// Currenly we only apply visibility and column width.
+	for (const TSharedRef<FTableColumn>& ColumnRef : Table->GetColumns())
+	{
+		FTableColumn& Column = ColumnRef.Get();
+		const FName ColumnId = Column.GetId();
+		const FTableColumnConfig* ConfigPtr = InColumnConfigSet.FindByPredicate([ColumnId](const FTableColumnConfig& Config) { return ColumnId == Config.ColumnId; });
+		if (ConfigPtr && ConfigPtr->bIsVisible)
+		{
+			ShowColumn(Column);
+			if (ConfigPtr->Width > 0.0f)
+			{
+				TreeViewHeaderRow->SetColumnWidth(ColumnId, ConfigPtr->Width);
+			}
+		}
+		else
+		{
+			HideColumn(Column);
+		}
+	}
+}
+
+////////////////////////////////////////////////////////////////////////////////////////////////////
+
+void STableTreeView::ViewPreset_OnSelectionChanged(TSharedPtr<ITableTreeViewPreset> InPreset, ESelectInfo::Type SelectInfo)
+{
+	SelectedViewPreset = InPreset;
+	if (InPreset.IsValid())
+	{
+		ApplyViewPreset(*InPreset);
+	}
+}
+
+////////////////////////////////////////////////////////////////////////////////////////////////////
+
+BEGIN_SLATE_FUNCTION_BUILD_OPTIMIZATION
+TSharedRef<SWidget> STableTreeView::ViewPreset_OnGenerateWidget(TSharedRef<ITableTreeViewPreset> InPreset)
+{
+	return SNew(STextBlock)
+		.Text(InPreset->GetName())
+		.ToolTipText(InPreset->GetToolTip())
+		.Margin(2.0f);
+}
+END_SLATE_FUNCTION_BUILD_OPTIMIZATION
+
+////////////////////////////////////////////////////////////////////////////////////////////////////
+
+FText STableTreeView::ViewPreset_GetSelectedText() const
+{
+	return SelectedViewPreset ? SelectedViewPreset->GetName() : LOCTEXT("Custom_ToolTip", "Custom");
+}
+
+////////////////////////////////////////////////////////////////////////////////////////////////////
+
+FText STableTreeView::ViewPreset_GetSelectedToolTipText() const
+{
+	return SelectedViewPreset ? SelectedViewPreset->GetToolTip() : LOCTEXT("CustomPreset_ToolTip", "Custom Preset");
+}
+
+////////////////////////////////////////////////////////////////////////////////////////////////////
+
 } // namespace Insights
 
 #undef LOCTEXT_NAMESPACE
diff --git a/Engine/Source/Developer/TraceInsights/Private/Insights/Table/Widgets/STableTreeView.h b/Engine/Source/Developer/TraceInsights/Private/Insights/Table/Widgets/STableTreeView.h
--- a/Engine/Source/Developer/TraceInsights/Private/Insights/Table/Widgets/STableTreeView.h
+++ b/Engine/Source/Developer/TraceInsights/Private/Insights/Table/Widgets/STableTreeView.h
@@ -8,6 +8,7 @@
 #include "Misc/TextFilter.h"
 #include "SlateFwd.h"
 #include "Widgets/DeclarativeSyntaxSupport.h"
+#include "Widgets/Input/SComboBox.h"
 #include "Widgets/Navigation/SBreadcrumbTrail.h"
 #include "Widgets/SCompoundWidget.h"
 #include "Widgets/SWidget.h"
@@ -62,6 +63,24 @@ enum class EAsyncOperationType : uint32
 
 ENUM_CLASS_FLAGS(EAsyncOperationType);
 
+struct FTableColumnConfig
+{
+	FName ColumnId;
+	bool bIsVisible;
+	float Width;
+};
+
+class ITableTreeViewPreset
+{
+public:
+	virtual FText GetName() const = 0;
+	virtual FText GetToolTip() const = 0;
+	virtual FName GetSortColumn() const = 0;
+	virtual EColumnSortMode::Type GetSortMode() const = 0;
+	virtual void SetCurrentGroupings(const TArray<TSharedPtr<FTreeNodeGrouping>>& InAvailableGroupings, TArray<TSharedPtr<FTreeNodeGrouping>>& InOutCurrentGroupings) const = 0;
+	virtual void GetColumnConfigSet(TArray<FTableColumnConfig>& InOutConfigSet) const = 0;
+};
+
 ////////////////////////////////////////////////////////////////////////////////////////////////////
 /**
  * A custom widget used to display the list of tree nodes.
@@ -422,6 +441,16 @@ protected:
 
 	FText GetTreeViewBannerText() const { return TreeViewBannerText; }
 
+	virtual void InitAvailableViewPresets() {};
+	const TArray<TSharedRef<ITableTreeViewPreset>>* GetAvailableViewPresets() const { return &AvailableViewPresets; }
+	FReply OnApplyViewPreset(const ITableTreeViewPreset* InPreset);
+	void ApplyViewPreset(const ITableTreeViewPreset& InPreset);
+	void ApplyColumnConfig(const TArrayView<FTableColumnConfig>& InTableConfig);
+	void ViewPreset_OnSelectionChanged(TSharedPtr<ITableTreeViewPreset> InPreset, ESelectInfo::Type SelectInfo);
+	TSharedRef<SWidget> ViewPreset_OnGenerateWidget(TSharedRef<ITableTreeViewPreset> InPreset);
+	FText ViewPreset_GetSelectedText() const;
+	FText ViewPreset_GetSelectedToolTipText() const;
+
 protected:
 	/** Table view model. */
 	TSharedPtr<Insights::FTable> Table;
@@ -544,6 +573,10 @@ protected:
 
 	FText TreeViewBannerText;
 
+	TArray<TSharedRef<ITableTreeViewPreset>> AvailableViewPresets;
+	TSharedPtr<ITableTreeViewPreset> SelectedViewPreset;
+	TSharedPtr<SComboBox<TSharedRef<ITableTreeViewPreset>>> PresetComboBox;
+
 	static constexpr int32 MAX_NUMBER_OF_NODES_TO_EXPAND = 1000 * 1000;
 	static constexpr int32 MAX_DEPTH_TO_EXPAND = 100;
 
diff --git a/Engine/Source/Developer/TraceInsights/Private/Insights/TaskGraphProfiler/Widgets/STaskTableTreeView.cpp b/Engine/Source/Developer/TraceInsights/Private/Insights/TaskGraphProfiler/Widgets/STaskTableTreeView.cpp
--- a/Engine/Source/Developer/TraceInsights/Private/Insights/TaskGraphProfiler/Widgets/STaskTableTreeView.cpp
+++ b/Engine/Source/Developer/TraceInsights/Private/Insights/TaskGraphProfiler/Widgets/STaskTableTreeView.cpp
@@ -307,34 +307,6 @@ END_SLATE_FUNCTION_BUILD_OPTIMIZATION
 
 ////////////////////////////////////////////////////////////////////////////////////////////////////
 
-void STaskTableTreeView::ApplyColumnConfig(const TArrayView<FColumnConfig>& Preset)
-{
-	for (const TSharedRef<FTableColumn>& ColumnRef : Table->GetColumns())
-	{
-		FTableColumn& Column = ColumnRef.Get();
-		for (const FColumnConfig& Config : Preset)
-		{
-			if (Column.GetId() == Config.ColumnId)
-			{
-				if (Config.bIsVisible)
-				{
-					ShowColumn(Column);
-					if (Config.Width > 0.0f)
-					{
-						TreeViewHeaderRow->SetColumnWidth(Column.GetId(), Config.Width);
-					}
-				}
-				else
-				{
-					HideColumn(Column);
-				}
-			}
-		}
-	}
-}
-
-////////////////////////////////////////////////////////////////////////////////////////////////////
-
 void STaskTableTreeView::InternalCreateGroupings()
 {
 	STableTreeView::InternalCreateGroupings();
diff --git a/Engine/Source/Developer/TraceInsights/Private/Insights/TaskGraphProfiler/Widgets/STaskTableTreeView.h b/Engine/Source/Developer/TraceInsights/Private/Insights/TaskGraphProfiler/Widgets/STaskTableTreeView.h
--- a/Engine/Source/Developer/TraceInsights/Private/Insights/TaskGraphProfiler/Widgets/STaskTableTreeView.h
+++ b/Engine/Source/Developer/TraceInsights/Private/Insights/TaskGraphProfiler/Widgets/STaskTableTreeView.h
@@ -91,7 +91,6 @@ protected:
 	virtual void ExtendMenu(FMenuBuilder& MenuBuilder) override;
 
 private:
-	void ApplyColumnConfig(const TArrayView<FColumnConfig>& Preset);
 	void AddCommmands();
 
 	const TArray<TSharedPtr<ETimestampOptions>>* GetAvailableTimestampOptions();
