[Insights] Fixed case where SourceCodeAccessor remains in a "started" state and does not respond anymore to more requests (i.e. after trying to open a file while the tool is not already open). Ticking the SourceCodeAccessor unblocks it.

diff --git a/Engine/Source/Developer/TraceInsights/Private/Insights/InsightsManager.cpp b/Engine/Source/Developer/TraceInsights/Private/Insights/InsightsManager.cpp
--- a/Engine/Source/Developer/TraceInsights/Private/Insights/InsightsManager.cpp
+++ b/Engine/Source/Developer/TraceInsights/Private/Insights/InsightsManager.cpp
@@ -6,6 +6,8 @@
 #include "Framework/Application/SlateApplication.h"
 #include "HAL/PlatformMemory.h"
 #include "HAL/PlatformProcess.h"
+#include "ISourceCodeAccessModule.h"
+#include "ISourceCodeAccessor.h"
 #include "Logging/MessageLog.h"
 #include "MessageLog/Public/MessageLogModule.h"
 #include "Misc/CString.h"
@@ -516,6 +518,10 @@ bool FInsightsManager::Tick(float DeltaTime)
 		}
 	}
 
+	ISourceCodeAccessModule& SourceCodeAccessModule = FModuleManager::LoadModuleChecked<ISourceCodeAccessModule>("SourceCodeAccess");
+	ISourceCodeAccessor& SourceCodeAccessor = SourceCodeAccessModule.GetAccessor();
+	SourceCodeAccessor.Tick(DeltaTime);
+
 	CheckMemoryUsage();
 #endif // !WITH_EDITOR
 
