[Insights] Tasks will be selected in the TaskTable when a Timing Event that is contained in a task is selected in the Timing View.

diff --git a/Engine/Source/Developer/TraceInsights/Private/Insights/Table/Widgets/STableTreeView.cpp b/Engine/Source/Developer/TraceInsights/Private/Insights/Table/Widgets/STableTreeView.cpp
--- a/Engine/Source/Developer/TraceInsights/Private/Insights/Table/Widgets/STableTreeView.cpp
+++ b/Engine/Source/Developer/TraceInsights/Private/Insights/Table/Widgets/STableTreeView.cpp
@@ -2762,11 +2762,23 @@ void STableTreeView::SelectNodeByTableRowIndex(int32 RowIndex)
 	if (RowIndex >= 0 && RowIndex < TableTreeNodes.Num())
 	{
 		FTableTreeNodePtr NodePtr = TableTreeNodes[RowIndex];
-		if (ensure(NodePtr))
+		ensure(NodePtr);
+
+		if (!NodePtr)
 		{
-			TreeView->SetSelection(NodePtr);
-			TreeView->RequestScrollIntoView(NodePtr);
+			return;
+		}
+
+		FBaseTreeNodePtr GroupNode = NodePtr->GetGroupPtr().Pin();
+		while (GroupNode.IsValid())
+		{
+			TreeView->SetItemExpansion(StaticCastSharedPtr<FTableTreeNode>(GroupNode), true);
+			GroupNode = GroupNode->GetGroupPtr().Pin();
+
 		}
+
+		TreeView->SetSelection(NodePtr);
+		TreeView->RequestScrollIntoView(NodePtr);
 	}
 }
 
@@ -3642,6 +3654,28 @@ FText STableTreeView::ViewPreset_GetSelectedToolTipText() const
 
 ////////////////////////////////////////////////////////////////////////////////////////////////////
 
+void STableTreeView::StopAllTableDataTasks(bool bWait)
+{
+	DataTaskInfos.RemoveAllSwap([](TSharedPtr<FTableTaskInfo> Data) { return Data->Event->IsComplete(); });
+	
+	for (int32 Index = 0; Index < DataTaskInfos.Num(); ++Index)
+	{
+		DataTaskInfos[Index]->CancellationToken->Cancel();
+	}
+
+	if (bWait)
+	{
+		for (int32 Index = 0; Index < DataTaskInfos.Num(); ++Index)
+		{
+			DataTaskInfos[Index]->Event->Wait();
+		}
+
+		DataTaskInfos.Empty();
+	}
+}
+
+////////////////////////////////////////////////////////////////////////////////////////////////////
+
 } // namespace Insights
 
 #undef LOCTEXT_NAMESPACE
diff --git a/Engine/Source/Developer/TraceInsights/Private/Insights/Table/Widgets/STableTreeView.h b/Engine/Source/Developer/TraceInsights/Private/Insights/Table/Widgets/STableTreeView.h
--- a/Engine/Source/Developer/TraceInsights/Private/Insights/Table/Widgets/STableTreeView.h
+++ b/Engine/Source/Developer/TraceInsights/Private/Insights/Table/Widgets/STableTreeView.h
@@ -81,6 +81,26 @@ public:
 	virtual void GetColumnConfigSet(TArray<FTableColumnConfig>& InOutConfigSet) const = 0;
 };
 
+class FTableTaskCancellationToken
+{
+public:
+	FTableTaskCancellationToken()
+		: bCancel(false)
+	{}
+
+	bool ShouldCancel() { return bCancel.load(); }
+	void Cancel() { bCancel.store(true); }
+
+private:
+	std::atomic<bool> bCancel;
+};
+
+struct FTableTaskInfo
+{
+	FGraphEventRef Event;
+	TSharedPtr< FTableTaskCancellationToken> CancellationToken;
+};
+
 ////////////////////////////////////////////////////////////////////////////////////////////////////
 /**
  * A custom widget used to display the list of tree nodes.
@@ -91,6 +111,8 @@ class STableTreeView : public SCompoundWidget, public IAsyncOperationStatusProvi
 	friend class FTableTreeViewSortAsyncTask;
 	friend class FTableTreeViewGroupAsyncTask;
 	friend class FTableTreeViewAsyncCompleteTask;
+	friend class FSearchForItemToSelectTask;
+	friend class FSelectNodeByTableRowIndexTask;
 
 public:
 	/** Default constructor. */
@@ -451,6 +473,23 @@ protected:
 	FText ViewPreset_GetSelectedText() const;
 	FText ViewPreset_GetSelectedToolTipText() const;
 
+	virtual void SearchForItem(TSharedPtr<FTableTaskCancellationToken> CancellationToken) {};
+
+	// Table data tasks should be tasks that operate read only operations on the data from the Table
+	// They should not operate on the tree nodes because they will run concurrently with the populated table UI.
+	template<typename T, typename... TArgs>
+	TSharedPtr<FTableTaskInfo> StartTableDataTask(TArgs&&... Args)
+	{
+		TSharedPtr<FTableTaskInfo> Info = MakeShared<FTableTaskInfo>();
+		Info->CancellationToken = MakeShared<FTableTaskCancellationToken>();
+		Info->Event = TGraphTask<T>::CreateTask().ConstructAndDispatchWhenReady(Info->CancellationToken, Forward<TArgs>(Args)...);
+		DataTaskInfos.Add(Info);
+
+		return Info;
+	}
+
+	void StopAllTableDataTasks(bool bWait = true);
+
 protected:
 	/** Table view model. */
 	TSharedPtr<Insights::FTable> Table;
@@ -561,6 +600,7 @@ protected:
 	FFilterConfigurator* CurrentAsyncOpFilterConfigurator = nullptr;
 	FAsyncOperationProgress AsyncOperationProgress;
 	FGraphEventRef DispatchEvent;
+	TArray<TSharedPtr<FTableTaskInfo>> DataTaskInfos;
 	TArray<FTableTreeNodePtr> NodesToExpand;
 
 	//////////////////////////////////////////////////
@@ -700,4 +740,59 @@ private:
 	TSharedPtr<STableTreeView> TableTreeViewPtr;
 };
 
+////////////////////////////////////////////////////////////////////////////////////////////////////
+
+class FSearchForItemToSelectTask
+{
+public:
+	FSearchForItemToSelectTask(TSharedPtr<FTableTaskCancellationToken> InToken, TSharedPtr<STableTreeView> InPtr)
+		: CancellationToken(InToken)
+		, TableTreeViewPtr(InPtr)
+	{}
+
+	FORCEINLINE TStatId GetStatId() const { RETURN_QUICK_DECLARE_CYCLE_STAT(FSearchForItemToSelectTask, STATGROUP_TaskGraphTasks); }
+	ENamedThreads::Type GetDesiredThread() { return ENamedThreads::Type::AnyThread; }
+	static ESubsequentsMode::Type GetSubsequentsMode() { return ESubsequentsMode::TrackSubsequents; }
+
+	void DoTask(ENamedThreads::Type CurrentThread, const FGraphEventRef& MyCompletionGraphEvent)
+	{
+		TableTreeViewPtr->SearchForItem(CancellationToken);
+	}
+
+private:
+	TSharedPtr<FTableTaskCancellationToken> CancellationToken;
+	TSharedPtr<STableTreeView> TableTreeViewPtr;
+};
+
+////////////////////////////////////////////////////////////////////////////////////////////////////
+
+class FSelectNodeByTableRowIndexTask
+{
+public:
+	FSelectNodeByTableRowIndexTask(TSharedPtr<FTableTaskCancellationToken> InToken, TSharedPtr<STableTreeView> InPtr, uint32 InRowIndex)
+		: CancellationToken(InToken)
+		, TableTreeViewPtr(InPtr)
+		, RowIndex(InRowIndex)
+		{}
+
+	FORCEINLINE TStatId GetStatId() const { RETURN_QUICK_DECLARE_CYCLE_STAT(FSelectNodeByTableRowIndexTask, STATGROUP_TaskGraphTasks); }
+	ENamedThreads::Type GetDesiredThread() { return ENamedThreads::Type::GameThread; }
+	static ESubsequentsMode::Type GetSubsequentsMode() { return ESubsequentsMode::TrackSubsequents; }
+
+	void DoTask(ENamedThreads::Type CurrentThread, const FGraphEventRef& MyCompletionGraphEvent)
+	{
+		if (!CancellationToken->ShouldCancel())
+		{
+			TableTreeViewPtr->SelectNodeByTableRowIndex(RowIndex);
+		}
+	}
+
+private:
+	TSharedPtr< FTableTaskCancellationToken> CancellationToken;
+	TSharedPtr<STableTreeView> TableTreeViewPtr;
+	uint32 RowIndex;
+};
+
+////////////////////////////////////////////////////////////////////////////////////////////////////
+
 } // namespace Insights
diff --git a/Engine/Source/Developer/TraceInsights/Private/Insights/TaskGraphProfiler/TaskGraphProfilerManager.cpp b/Engine/Source/Developer/TraceInsights/Private/Insights/TaskGraphProfiler/TaskGraphProfilerManager.cpp
--- a/Engine/Source/Developer/TraceInsights/Private/Insights/TaskGraphProfiler/TaskGraphProfilerManager.cpp
+++ b/Engine/Source/Developer/TraceInsights/Private/Insights/TaskGraphProfiler/TaskGraphProfilerManager.cpp
@@ -805,6 +805,16 @@ double FTaskGraphProfilerManager::GetRelationsOnCriticalPathDescendingRec(const
 
 ////////////////////////////////////////////////////////////////////////////////////////////////////
 
+void FTaskGraphProfilerManager::SelectTaskInTaskTable(TaskTrace::FId InId)
+{
+	if (TaskTableTreeView.IsValid())
+	{
+		TaskTableTreeView->SelectTaskEntry(InId);
+	}
+}
+
+////////////////////////////////////////////////////////////////////////////////////////////////////
+
 } // namespace Insights
 
 #undef LOCTEXT_NAMESPACE
diff --git a/Engine/Source/Developer/TraceInsights/Private/Insights/TaskGraphProfiler/TaskGraphProfilerManager.h b/Engine/Source/Developer/TraceInsights/Private/Insights/TaskGraphProfiler/TaskGraphProfilerManager.h
--- a/Engine/Source/Developer/TraceInsights/Private/Insights/TaskGraphProfiler/TaskGraphProfilerManager.h
+++ b/Engine/Source/Developer/TraceInsights/Private/Insights/TaskGraphProfiler/TaskGraphProfilerManager.h
@@ -133,6 +133,8 @@ public:
 
 	int32 GetDepthOfTaskExecution(double TaskStartedTime, double TaskFinishedTime, uint32 ThreadId);
 
+	void SelectTaskInTaskTable(TaskTrace::FId);
+
 private:
 	/** Updates this manager, done through FCoreTicker. */
 	bool Tick(float DeltaTime);
diff --git a/Engine/Source/Developer/TraceInsights/Private/Insights/TaskGraphProfiler/ViewModels/TaskTimingTrack.cpp b/Engine/Source/Developer/TraceInsights/Private/Insights/TaskGraphProfiler/ViewModels/TaskTimingTrack.cpp
--- a/Engine/Source/Developer/TraceInsights/Private/Insights/TaskGraphProfiler/ViewModels/TaskTimingTrack.cpp
+++ b/Engine/Source/Developer/TraceInsights/Private/Insights/TaskGraphProfiler/ViewModels/TaskTimingTrack.cpp
@@ -847,6 +847,7 @@ void FTaskTimingTrack::OnTimingEventSelected(TSharedPtr<const ITimingEvent> InSe
 	if (Task != nullptr)
 	{
 		TaskId = Task->Id;
+		FTaskGraphProfilerManager::Get()->SelectTaskInTaskTable(TaskId);
 	}
 	else
 	{
diff --git a/Engine/Source/Developer/TraceInsights/Private/Insights/TaskGraphProfiler/Widgets/STaskTableTreeView.cpp b/Engine/Source/Developer/TraceInsights/Private/Insights/TaskGraphProfiler/Widgets/STaskTableTreeView.cpp
--- a/Engine/Source/Developer/TraceInsights/Private/Insights/TaskGraphProfiler/Widgets/STaskTableTreeView.cpp
+++ b/Engine/Source/Developer/TraceInsights/Private/Insights/TaskGraphProfiler/Widgets/STaskTableTreeView.cpp
@@ -167,11 +167,13 @@ void STaskTableTreeView::RebuildTree(bool bResync)
 
 	if (NewQueryStartTime != QueryStartTime || NewQueryEndTime != QueryEndTime)
 	{
+		StopAllTableDataTasks();
+
 		QueryStartTime = NewQueryStartTime;
 		QueryEndTime = NewQueryEndTime;
 		TSharedPtr<FTaskTable> TaskTable = GetTaskTable();
-		TArray<FTaskEntry>& Allocs = TaskTable->GetTaskEntries();
-		Allocs.Empty();
+		TArray<FTaskEntry>& Tasks = TaskTable->GetTaskEntries();
+		Tasks.Empty();
 		TableTreeNodes.Empty();
 
 		if (QueryStartTime < QueryEndTime && Session.IsValid())
@@ -186,10 +188,10 @@ void STaskTableTreeView::RebuildTree(bool bResync)
 
 				TArray<FTableTreeNodePtr>* Nodes = &TableTreeNodes;
 
-				TasksProvider->EnumerateTasks(QueryStartTime, QueryEndTime, [&Allocs, &TaskTable, &BaseNodeName, Nodes](const TraceServices::FTaskInfo& TaskInfo)
+				TasksProvider->EnumerateTasks(QueryStartTime, QueryEndTime, [&Tasks, &TaskTable, &BaseNodeName, Nodes](const TraceServices::FTaskInfo& TaskInfo)
 					{
-						Allocs.Emplace(TaskInfo);
-						uint32 Index = Allocs.Num() - 1;
+						Tasks.Emplace(TaskInfo);
+						uint32 Index = Tasks.Num() - 1;
 						FName NodeName(BaseNodeName, TaskInfo.Id + 1);
 						FTaskNodePtr NodePtr = MakeShared<FTaskNode>(NodeName, TaskTable, Index);
 						Nodes->Add(NodePtr);
@@ -571,6 +573,37 @@ void STaskTableTreeView::TreeView_OnMouseButtonDoubleClick(FTableTreeNodePtr Tre
 
 ////////////////////////////////////////////////////////////////////////////////////////////////////
 
+void STaskTableTreeView::SelectTaskEntry(TaskTrace::FId InId)
+{
+	TaskIdToSelect = InId;
+	StartTableDataTask<FSearchForItemToSelectTask>(SharedThis(this));
+}
+
+////////////////////////////////////////////////////////////////////////////////////////////////////
+
+void STaskTableTreeView::SearchForItem(TSharedPtr<FTableTaskCancellationToken> CancellationToken)
+{
+	TSharedPtr<FTaskTable> TaskTable = GetTaskTable();
+	TArray<FTaskEntry>& Tasks = TaskTable->GetTaskEntries();
+	
+	uint32 NumEntries = (uint32)Tasks.Num();
+	for (uint32 Index = 0; Index < NumEntries; ++Index)
+	{
+		if (CancellationToken->ShouldCancel())
+		{
+			break;
+		}
+
+		if (Tasks[Index].Id == TaskIdToSelect)
+		{
+			TGraphTask<FSelectNodeByTableRowIndexTask>::CreateTask().ConstructAndDispatchWhenReady(CancellationToken, SharedThis(this), Index);
+			break;
+		}
+	}
+}
+
+////////////////////////////////////////////////////////////////////////////////////////////////////
+
 } // namespace Insights
 
 #undef LOCTEXT_NAMESPACE
diff --git a/Engine/Source/Developer/TraceInsights/Private/Insights/TaskGraphProfiler/Widgets/STaskTableTreeView.h b/Engine/Source/Developer/TraceInsights/Private/Insights/TaskGraphProfiler/Widgets/STaskTableTreeView.h
--- a/Engine/Source/Developer/TraceInsights/Private/Insights/TaskGraphProfiler/Widgets/STaskTableTreeView.h
+++ b/Engine/Source/Developer/TraceInsights/Private/Insights/TaskGraphProfiler/Widgets/STaskTableTreeView.h
@@ -2,7 +2,7 @@
 
 #pragma once
 
-#include "CoreMinimal.h"
+#include "Async/TaskTrace.h"
 #include "TraceServices/Model/TasksProfiler.h"
 
 // Insights
@@ -85,11 +85,15 @@ public:
 
 	virtual void TreeView_OnMouseButtonDoubleClick(FTableTreeNodePtr TreeNode) override;
 
+	void SelectTaskEntry(TaskTrace::FId InId);
+
 protected:
 	virtual void InternalCreateGroupings() override;
 
 	virtual void ExtendMenu(FMenuBuilder& MenuBuilder) override;
 
+	virtual void SearchForItem(TSharedPtr<FTableTaskCancellationToken> CancellationToken) override;
+
 private:
 	void AddCommmands();
 
@@ -113,6 +117,8 @@ private:
 
 	ETimestampOptions SelectedTimestampOption = ETimestampOptions::RelativeToPrevious;
 	TArray<TSharedPtr<ETimestampOptions>> AvailableTimestampOptions;
+
+	TaskTrace::FId TaskIdToSelect;
 };
 
 ////////////////////////////////////////////////////////////////////////////////////////////////////
