[Insights] Make csv files as the default option for exporting and importing tables.

diff --git a/Engine/Source/Developer/TraceInsights/Private/Insights/ImportTool/TableImportTool.cpp b/Engine/Source/Developer/TraceInsights/Private/Insights/ImportTool/TableImportTool.cpp
--- a/Engine/Source/Developer/TraceInsights/Private/Insights/ImportTool/TableImportTool.cpp
+++ b/Engine/Source/Developer/TraceInsights/Private/Insights/ImportTool/TableImportTool.cpp
@@ -76,7 +76,7 @@ void FTableImportTool::StartImportProcess()
 			LOCTEXT("ImportTableTitle", "Import Table").ToString(),
 			DefaultPath,
 			DefaultFile,
-			TEXT("Tab-Separated Values (*.tsv)|*.tsv|Text Files (*.txt)|*.txt|Comma-Separated Values (*.csv)|*.csv|All Files (*.*)|*.*"),
+			TEXT("Comma-Separated Values (*.csv)|*.csv|Tab-Separated Values (*.tsv)|*.tsv|Text Files (*.txt)|*.txt|All Files (*.*)|*.*"),
 			EFileDialogFlags::None,
 			Filenames
 		);
diff --git a/Engine/Source/Developer/TraceInsights/Private/Insights/Table/Widgets/STableTreeView.cpp b/Engine/Source/Developer/TraceInsights/Private/Insights/Table/Widgets/STableTreeView.cpp
--- a/Engine/Source/Developer/TraceInsights/Private/Insights/Table/Widgets/STableTreeView.cpp
+++ b/Engine/Source/Developer/TraceInsights/Private/Insights/Table/Widgets/STableTreeView.cpp
@@ -3474,7 +3474,7 @@ void STableTreeView::ContextMenu_ExportToFile_Execute(bool bInExportCollapsed, b
 			LOCTEXT("ExportFileTitle", "Export Table").ToString(),
 			DefaultPath,
 			DefaultFile,
-			TEXT("Tab-Separated Values (*.tsv)|*.tsv|Text Files (*.txt)|*.txt|Comma-Separated Values (*.csv)|*.csv|All Files (*.*)|*.*"),
+			TEXT("Comma-Separated Values (*.csv)|*.csv|Tab-Separated Values (*.tsv)|*.tsv|Text Files (*.txt)|*.txt|All Files (*.*)|*.*"),
 			EFileDialogFlags::None,
 			SaveFilenames
 		);
diff --git a/Engine/Source/Developer/TraceInsights/Private/Insights/Widgets/SLogView.cpp b/Engine/Source/Developer/TraceInsights/Private/Insights/Widgets/SLogView.cpp
--- a/Engine/Source/Developer/TraceInsights/Private/Insights/Widgets/SLogView.cpp
+++ b/Engine/Source/Developer/TraceInsights/Private/Insights/Widgets/SLogView.cpp
@@ -1832,7 +1832,7 @@ void SLogView::SaveLogsToFile(bool bSaveLogsInSelectedRangeOnly) const
 			LOCTEXT("SaveLogsToFileTitle", "Save Logs").ToString(),
 			DefaultBrowsePath,
 			TEXT(""),
-			TEXT("Tab-Separated Values (*.tsv)|*.tsv|Text Files (*.txt)|*.txt|Comma-Separated Values (*.csv)|*.csv|All Files (*.*)|*.*"),
+			TEXT("Comma-Separated Values (*.csv)|*.csv|Tab-Separated Values (*.tsv)|*.tsv|Text Files (*.txt)|*.txt|All Files (*.*)|*.*"),
 			EFileDialogFlags::None,
 			SaveFilenames
 		);
diff --git a/Engine/Source/Developer/TraceInsights/Private/Insights/Widgets/STimersView.cpp b/Engine/Source/Developer/TraceInsights/Private/Insights/Widgets/STimersView.cpp
--- a/Engine/Source/Developer/TraceInsights/Private/Insights/Widgets/STimersView.cpp
+++ b/Engine/Source/Developer/TraceInsights/Private/Insights/Widgets/STimersView.cpp
@@ -2616,7 +2616,7 @@ bool STimersView::OpenSaveTextFileDialog(const FString& InDialogTitle, const FSt
 			InDialogTitle,
 			DefaultPath,
 			InDefaultFile,
-			TEXT("Tab-Separated Values (*.tsv)|*.tsv|Text Files (*.txt)|*.txt|Comma-Separated Values (*.csv)|*.csv|All Files (*.*)|*.*"),
+			TEXT("Comma-Separated Values (*.csv)|*.csv|Tab-Separated Values (*.tsv)|*.tsv|Text Files (*.txt)|*.txt|All Files (*.*)|*.*"),
 			EFileDialogFlags::None,
 			SaveFilenames
 		);
