[Insights] Session Browser: Fixed max width for the search boxes (search by name/command line and the two auto-start filters).

diff --git a/Engine/Source/Developer/TraceInsights/Private/Insights/Widgets/SStartPageWindow.cpp b/Engine/Source/Developer/TraceInsights/Private/Insights/Widgets/SStartPageWindow.cpp
--- a/Engine/Source/Developer/TraceInsights/Private/Insights/Widgets/SStartPageWindow.cpp
+++ b/Engine/Source/Developer/TraceInsights/Private/Insights/Widgets/SStartPageWindow.cpp
@@ -1026,23 +1026,27 @@ TSharedRef<SWidget> STraceStoreWindow::ConstructFiltersToolbar()
 
 		// Text Filter (Search Box)
 		ToolbarBuilder.AddWidget(
-			SAssignNew(FilterByNameSearchBox, SSearchBox)
-			.MinDesiredWidth(150.0f)
-			.HintText_Lambda([this]()
-				{
-					return bSearchByCommandLine ?
-						LOCTEXT("CmdLineFilter_Hint", "Command Line") :
-						LOCTEXT("NameFilter_Hint", "Name");
-				})
-			.ToolTipText_Lambda([this]()
-				{
-					return bSearchByCommandLine ?
-						LOCTEXT("CmdLineFilter_Tooltip", "Type here to filter the list of trace sessions by command line.") :
-						LOCTEXT("NameFilter_Tooltip", "Type here to filter the list of trace sessions by name.");
-				})
-			.IsEnabled_Lambda([this]() { return TraceViewModels.Num() > 0; })
-			.OnTextChanged(this, &STraceStoreWindow::FilterByNameSearchBox_OnTextChanged)
-			.DelayChangeNotificationsWhileTyping(true)
+			SNew(SBox)
+			.MaxDesiredWidth(400.0f)
+			[
+				SAssignNew(FilterByNameSearchBox, SSearchBox)
+				.MinDesiredWidth(150.0f)
+				.HintText_Lambda([this]()
+					{
+						return bSearchByCommandLine ?
+							LOCTEXT("CmdLineFilter_Hint", "Command Line") :
+							LOCTEXT("NameFilter_Hint", "Name");
+					})
+				.ToolTipText_Lambda([this]()
+					{
+						return bSearchByCommandLine ?
+							LOCTEXT("CmdLineFilter_Tooltip", "Type here to filter the list of trace sessions by command line.") :
+							LOCTEXT("NameFilter_Tooltip", "Type here to filter the list of trace sessions by name.");
+					})
+				.IsEnabled_Lambda([this]() { return TraceViewModels.Num() > 0; })
+				.OnTextChanged(this, &STraceStoreWindow::FilterByNameSearchBox_OnTextChanged)
+				.DelayChangeNotificationsWhileTyping(true)
+			]
 		);
 
 		// Filter by Platform
@@ -1378,9 +1382,13 @@ TSharedRef<SWidget> STraceStoreWindow::ConstructAutoStartPanel()
 	.HAlign(HAlign_Left)
 	.VAlign(VAlign_Center)
 	[
-		SAssignNew(AutoStartPlatformFilter, SSearchBox)
-		.HintText(LOCTEXT("AutoStartPlatformFilter_Hint", "Platform"))
-		.ToolTipText(LOCTEXT("AutoStartPlatformFilter_Tooltip", "Type here to specify the Platform filter.\nAuto-start analysis will be enabled only for live trace sessions with this specified Platform."))
+		SNew(SBox)
+		.MaxDesiredWidth(200.0f)
+		[
+			SAssignNew(AutoStartPlatformFilter, SSearchBox)
+			.HintText(LOCTEXT("AutoStartPlatformFilter_Hint", "Platform"))
+			.ToolTipText(LOCTEXT("AutoStartPlatformFilter_Tooltip", "Type here to specify the Platform filter.\nAuto-start analysis will be enabled only for live trace sessions with this specified Platform."))
+		]
 	]
 
 	+ SHorizontalBox::Slot()
@@ -1389,9 +1397,13 @@ TSharedRef<SWidget> STraceStoreWindow::ConstructAutoStartPanel()
 	.HAlign(HAlign_Left)
 	.VAlign(VAlign_Center)
 	[
-		SAssignNew(AutoStartAppNameFilter, SSearchBox)
-		.HintText(LOCTEXT("AutoStartAppNameFilter_Hint", "AppName"))
-		.ToolTipText(LOCTEXT("AutoStartAppNameFilter_Tooltip", "Type here to specify the AppName filter.\nAuto-start analysis will be enabled only for live trace sessions with this specified AppName."))
+		SNew(SBox)
+		.MaxDesiredWidth(200.0f)
+		[
+			SAssignNew(AutoStartAppNameFilter, SSearchBox)
+			.HintText(LOCTEXT("AutoStartAppNameFilter_Hint", "AppName"))
+			.ToolTipText(LOCTEXT("AutoStartAppNameFilter_Tooltip", "Type here to specify the AppName filter.\nAuto-start analysis will be enabled only for live trace sessions with this specified AppName."))
+		]
 	];
 
 	return Widget;
