[Insights] Show full llm tag path name is the Allocation Table and remove duplicated items from the LLM filter.

diff --git a/Engine/Source/Developer/TraceInsights/Private/Insights/MemoryProfiler/Widgets/SMemAllocTableTreeView.cpp b/Engine/Source/Developer/TraceInsights/Private/Insights/MemoryProfiler/Widgets/SMemAllocTableTreeView.cpp
--- a/Engine/Source/Developer/TraceInsights/Private/Insights/MemoryProfiler/Widgets/SMemAllocTableTreeView.cpp
+++ b/Engine/Source/Developer/TraceInsights/Private/Insights/MemoryProfiler/Widgets/SMemAllocTableTreeView.cpp
@@ -2,6 +2,7 @@
 
 #include "SMemAllocTableTreeView.h"
 
+#include "Containers/Set.h"
 #include "Framework/MultiBox/MultiBoxBuilder.h"
 #include "GenericPlatform/GenericPlatformMisc.h"
 #include "ISourceCodeAccessModule.h"
@@ -395,7 +396,7 @@ void SMemAllocTableTreeView::UpdateQuery(TraceServices::IAllocationsProvider::EQ
 					Alloc.Size = int64(Allocation->GetSize());
 
 					Alloc.TagId = Allocation->GetTag();
-					Alloc.Tag = Provider.GetTagName(Allocation->GetTag());
+					Alloc.Tag = Provider.GetTagFullPath(Allocation->GetTag());
 
 					Alloc.Asset = nullptr;
 					Alloc.ClassName = nullptr;
@@ -1294,15 +1295,20 @@ void SMemAllocTableTreeView::PopulateLLMTagSuggestionList(const FString& Text, T
 
 	TraceServices::FProviderReadScopeLock _(*AllocationsProvider);
 
-	AllocationsProvider->EnumerateTags([&OutSuggestions, &Text](const TCHAR* Display, const TCHAR* FullPath, TraceServices::TagIdType CurrentTag, TraceServices::TagIdType ParentTag)
+	// Use a Set to avoid duplicate tag names.
+	TSet<FString> Suggestions;
+	AllocationsProvider->EnumerateTags([&Suggestions, &Text](const TCHAR* Display, const TCHAR* FullPath, TraceServices::TagIdType CurrentTag, TraceServices::TagIdType ParentTag)
 	{
-		if (Text.IsEmpty() || FCString::Stristr(Display, *Text))
+		if (Text.IsEmpty() || FCString::Stristr(FullPath, *Text))
 		{
-			OutSuggestions.Push(Display);
+			Suggestions.Add(FullPath);
 		}
 
 		return true;
 	});
+
+	OutSuggestions = Suggestions.Array();
+	OutSuggestions.Sort();
 }
 
 ////////////////////////////////////////////////////////////////////////////////////////////////////
diff --git a/Engine/Source/Developer/TraceServices/Private/Model/AllocationsProvider.cpp b/Engine/Source/Developer/TraceServices/Private/Model/AllocationsProvider.cpp
--- a/Engine/Source/Developer/TraceServices/Private/Model/AllocationsProvider.cpp
+++ b/Engine/Source/Developer/TraceServices/Private/Model/AllocationsProvider.cpp
@@ -293,6 +293,14 @@ const TCHAR* FTagTracker::GetTagString(TagIdType InTag) const
 
 ////////////////////////////////////////////////////////////////////////////////////////////////////
 
+const TCHAR* FTagTracker::GetTagFullPath(TagIdType InTag) const
+{
+	const FTagEntry* Entry = TagMap.Find(InTag);
+	return Entry ? Entry->FullPath : TEXT("Unknown");
+}
+
+////////////////////////////////////////////////////////////////////////////////////////////////////
+
 void FTagTracker::EnumerateTags(TFunctionRef<void (const TCHAR*, const TCHAR*, TagIdType, TagIdType)> Callback) const
 {
 	for (const auto& EntryPair : TagMap)
diff --git a/Engine/Source/Developer/TraceServices/Private/Model/AllocationsProvider.h b/Engine/Source/Developer/TraceServices/Private/Model/AllocationsProvider.h
--- a/Engine/Source/Developer/TraceServices/Private/Model/AllocationsProvider.h
+++ b/Engine/Source/Developer/TraceServices/Private/Model/AllocationsProvider.h
@@ -77,6 +77,7 @@ public:
 	void PopTag(uint32 ThreadId, uint8 Tracker);
 	TagIdType GetCurrentTag(uint32 ThreadId, uint8 Tracker) const;
 	const TCHAR* GetTagString(TagIdType Tag) const;
+	const TCHAR* GetTagFullPath(TagIdType Tag) const;
 	void EnumerateTags(TFunctionRef<void(const TCHAR*, const TCHAR*, TagIdType, TagIdType)> Callback) const;
 
 	void PushTagFromPtr(uint32 ThreadId, uint8 Tracker, TagIdType Tag);
@@ -362,7 +363,8 @@ public:
 	void EnumerateLiveAllocs(TFunctionRef<void(const FAllocationItem& Alloc)> Callback) const;
 	uint32 GetNumLiveAllocs() const;
 
-	virtual const TCHAR* GetTagName(TagIdType Tag) const { ReadAccessCheck(); return TagTracker.GetTagString(Tag); }
+	virtual const TCHAR* GetTagName(TagIdType Tag) const override { ReadAccessCheck(); return TagTracker.GetTagString(Tag); }
+	virtual const TCHAR* GetTagFullPath(TagIdType Tag) const override { ReadAccessCheck(); return TagTracker.GetTagFullPath(Tag); }
 	bool HasTagFromPtrScope(uint32 ThreadId, uint8 Tracker) const { ReadAccessCheck(); return TagTracker.HasTagFromPtrScope(ThreadId, Tracker); }
 
 	void DebugPrint() const;
diff --git a/Engine/Source/Developer/TraceServices/Public/TraceServices/Model/AllocationsProvider.h b/Engine/Source/Developer/TraceServices/Public/TraceServices/Model/AllocationsProvider.h
--- a/Engine/Source/Developer/TraceServices/Public/TraceServices/Model/AllocationsProvider.h
+++ b/Engine/Source/Developer/TraceServices/Public/TraceServices/Model/AllocationsProvider.h
@@ -160,6 +160,7 @@ public:
 	// Returns the display name of the specified LLM tag.
 	// Lifetime of returned string matches the session lifetime.
 	virtual const TCHAR* GetTagName(TagIdType Tag) const = 0;
+	virtual const TCHAR* GetTagFullPath(TagIdType Tag) const = 0;
 };
 
 TRACESERVICES_API FName GetAllocationsProviderName();
