[Insights] Keep only .csv and .tsv extensions as valid for table import files. Fix int64 parsing during table import.

diff --git a/Engine/Source/Developer/TraceInsights/Private/Insights/ImportTool/TableImportTool.cpp b/Engine/Source/Developer/TraceInsights/Private/Insights/ImportTool/TableImportTool.cpp
--- a/Engine/Source/Developer/TraceInsights/Private/Insights/ImportTool/TableImportTool.cpp
+++ b/Engine/Source/Developer/TraceInsights/Private/Insights/ImportTool/TableImportTool.cpp
@@ -76,7 +76,7 @@ void FTableImportTool::StartImportProcess()
 			LOCTEXT("ImportTableTitle", "Import Table").ToString(),
 			DefaultPath,
 			DefaultFile,
-			TEXT("Comma-Separated Values (*.csv)|*.csv|Tab-Separated Values (*.tsv)|*.tsv|Text Files (*.txt)|*.txt|All Files (*.*)|*.*"),
+			TEXT("Comma-Separated Values (*.csv)|*.csv|Tab-Separated Values (*.tsv)|*.tsv"),
 			EFileDialogFlags::None,
 			Filenames
 		);
diff --git a/Engine/Source/Developer/TraceServices/Private/Model/TableImportTask.cpp b/Engine/Source/Developer/TraceServices/Private/Model/TableImportTask.cpp
--- a/Engine/Source/Developer/TraceServices/Private/Model/TableImportTask.cpp
+++ b/Engine/Source/Developer/TraceServices/Private/Model/TableImportTask.cpp
@@ -193,7 +193,7 @@ bool FTableImportTask::ParseData(TArray<FString>& Lines)
 					break;
 				}
 
-				NewRow.SetValue(ValueIndex, FCString::Atoi(Value));
+				NewRow.SetValue(ValueIndex, FCString::Atoi64(Value));
 			}
 		}
 
