[Insights] Fixed unsafe type cast warnings in TraceInsights module.

diff --git a/Engine/Source/Developer/TraceInsights/Private/Insights/ContextSwitches/ViewModels/CpuCoreTimingTrack.cpp b/Engine/Source/Developer/TraceInsights/Private/Insights/ContextSwitches/ViewModels/CpuCoreTimingTrack.cpp
--- a/Engine/Source/Developer/TraceInsights/Private/Insights/ContextSwitches/ViewModels/CpuCoreTimingTrack.cpp
+++ b/Engine/Source/Developer/TraceInsights/Private/Insights/ContextSwitches/ViewModels/CpuCoreTimingTrack.cpp
@@ -200,7 +200,7 @@ void FCpuCoreTimingTrack::PostDraw(const ITimingTrackDrawContext& Context) const
 		const ITimingViewDrawHelper& Helper = Context.GetHelper();
 
 		FString Str = FString::Printf(TEXT("%s (Duration.: %s)"),
-			*GetThreadName(SelectedEvent.GetType()),
+			*GetThreadName(static_cast<uint32>(SelectedEvent.GetType())),
 			*TimeUtils::FormatTimeAuto(SelectedEvent.GetDuration()));
 
 		DrawSelectedEventInfo(Str, Context.GetViewport(), Context.GetDrawContext(), Helper.GetWhiteBrush(), Helper.GetEventFont());
diff --git a/Engine/Source/Developer/TraceInsights/Private/Insights/CookProfiler/Widgets/SPackageTableTreeView.cpp b/Engine/Source/Developer/TraceInsights/Private/Insights/CookProfiler/Widgets/SPackageTableTreeView.cpp
--- a/Engine/Source/Developer/TraceInsights/Private/Insights/CookProfiler/Widgets/SPackageTableTreeView.cpp
+++ b/Engine/Source/Developer/TraceInsights/Private/Insights/CookProfiler/Widgets/SPackageTableTreeView.cpp
@@ -197,7 +197,7 @@ TSharedPtr<SWidget> SPackageTableTreeView::ConstructToolbar()
 		.VAlign(VAlign_Center)
 		[
 			SNew(SBox)
-			.MinDesiredWidth(150)
+			.MinDesiredWidth(150.0f)
 			[
 				SAssignNew(PresetComboBox, SComboBox<TSharedRef<ITableTreeViewPreset>>)
 				.ToolTipText(this, &SPackageTableTreeView::ViewPreset_GetSelectedToolTipText)
diff --git a/Engine/Source/Developer/TraceInsights/Private/Insights/MemoryProfiler/ViewModels/MemAllocFilterValueConverter.cpp b/Engine/Source/Developer/TraceInsights/Private/Insights/MemoryProfiler/ViewModels/MemAllocFilterValueConverter.cpp
--- a/Engine/Source/Developer/TraceInsights/Private/Insights/MemoryProfiler/ViewModels/MemAllocFilterValueConverter.cpp
+++ b/Engine/Source/Developer/TraceInsights/Private/Insights/MemoryProfiler/ViewModels/MemAllocFilterValueConverter.cpp
@@ -53,10 +53,9 @@ bool FMemoryFilterValueConverter::Convert(const FString& Input, int64& Output, F
 				return false;
 			}
 
-			// Equals to (PrefixPos - Prefixes + 1) * 2^10
-			Output = 1LL << ((PrefixPos - Prefixes + 1) * 10);
+			const int64 Unit = 1LL << ((PrefixPos - Prefixes + 1) * 10);
 			double Value = FCString::Atod(*TrimmedInput.Left(Index));
-			Output *= Value;
+			Output = static_cast<int64>(Value * Unit);
 			break;
 		}
 		else
diff --git a/Engine/Source/Developer/TraceInsights/Private/Insights/MemoryProfiler/ViewModels/MemAllocGroupingByCallstack.cpp b/Engine/Source/Developer/TraceInsights/Private/Insights/MemoryProfiler/ViewModels/MemAllocGroupingByCallstack.cpp
--- a/Engine/Source/Developer/TraceInsights/Private/Insights/MemoryProfiler/ViewModels/MemAllocGroupingByCallstack.cpp
+++ b/Engine/Source/Developer/TraceInsights/Private/Insights/MemoryProfiler/ViewModels/MemAllocGroupingByCallstack.cpp
@@ -87,10 +87,11 @@ void FMemAllocGroupingByCallstack::GroupNodes(const TArray<FTableTreeNodePtr>& N
 			else if (Callstack)
 			{
 				NumFrames = static_cast<int32>(Callstack->Num());
+				check(NumFrames <= 256); // see Callstack->Frame(uint8)
 
 				for (int32 FrameDepth = NumFrames - 1; FrameDepth >= 0; --FrameDepth)
 				{
-					const TraceServices::FStackFrame* Frame = Callstack->Frame(bIsInverted ? NumFrames - FrameDepth - 1 : FrameDepth);
+					const TraceServices::FStackFrame* Frame = Callstack->Frame(static_cast<uint8>(bIsInverted ? NumFrames - FrameDepth - 1 : FrameDepth));
 					check(Frame != nullptr);
 
 					if (bIsGroupingByFunction)
diff --git a/Engine/Source/Developer/TraceInsights/Private/Insights/MemoryProfiler/ViewModels/MemAllocGroupingByHeap.cpp b/Engine/Source/Developer/TraceInsights/Private/Insights/MemoryProfiler/ViewModels/MemAllocGroupingByHeap.cpp
--- a/Engine/Source/Developer/TraceInsights/Private/Insights/MemoryProfiler/ViewModels/MemAllocGroupingByHeap.cpp
+++ b/Engine/Source/Developer/TraceInsights/Private/Insights/MemoryProfiler/ViewModels/MemAllocGroupingByHeap.cpp
@@ -59,7 +59,7 @@ void FMemAllocGroupingByHeap::GroupNodes(const TArray<FTableTreeNodePtr>& Nodes,
 	TArray<FTableTreeNodePtr> HeapNodes;
 	HeapNodes.AddZeroed(256);
 
-	AllocProvider.EnumerateRootHeaps([&](uint8 Id, const TraceServices::IAllocationsProvider::FHeapSpec& Spec)
+	AllocProvider.EnumerateRootHeaps([&](HeapId Id, const TraceServices::IAllocationsProvider::FHeapSpec& Spec)
 	{
 		ParentGroup.AddChildAndSetGroupPtr(MakeGroupNodeHierarchy(Spec, InParentTable, HeapNodes));
 	});
@@ -83,9 +83,9 @@ void FMemAllocGroupingByHeap::GroupNodes(const TArray<FTableTreeNodePtr>& Nodes,
 
 		if (Alloc)
 		{
-			//todo: Calculating the real HeapId when the allocation is first added to the provider was too expensive, so deferring that operation to here could make sense.
-			//const uint8 HeapId = AllocProvider.GetParentBlock(Alloc->GetAddress());
-			const uint8 HeapId = Alloc->GetRootHeap();
+			//TODO: Calculating the real HeapId when the allocation is first added to the provider was too expensive, so deferring that operation to here could make sense.
+			//const HeapId Heap = AllocProvider.GetParentBlock(Alloc->GetAddress());
+			const uint8 HeapId = static_cast<uint8>(Alloc->GetRootHeap());
 			if (FTableTreeNodePtr GroupPtr = HeapNodes[HeapId])
 			{
 				GroupPtr->AddChildAndSetGroupPtr(NodePtr);
diff --git a/Engine/Source/Developer/TraceInsights/Private/Insights/MemoryProfiler/ViewModels/MemAllocGroupingBySize.cpp b/Engine/Source/Developer/TraceInsights/Private/Insights/MemoryProfiler/ViewModels/MemAllocGroupingBySize.cpp
--- a/Engine/Source/Developer/TraceInsights/Private/Insights/MemoryProfiler/ViewModels/MemAllocGroupingBySize.cpp
+++ b/Engine/Source/Developer/TraceInsights/Private/Insights/MemoryProfiler/ViewModels/MemAllocGroupingBySize.cpp
@@ -107,7 +107,7 @@ void FMemAllocGroupingBySize::GroupNodes(const TArray<FTableTreeNodePtr>& Nodes,
 			uint32 ThresholdIndex;
 			if (bIsPow2)
 			{
-				ThresholdIndex = 64 - FMath::CountLeadingZeros64(Size);
+				ThresholdIndex = 64u - static_cast<uint32>(FMath::CountLeadingZeros64(Size));
 			}
 			else
 			{
diff --git a/Engine/Source/Developer/TraceInsights/Private/Insights/MemoryProfiler/ViewModels/MemAllocNode.cpp b/Engine/Source/Developer/TraceInsights/Private/Insights/MemoryProfiler/ViewModels/MemAllocNode.cpp
--- a/Engine/Source/Developer/TraceInsights/Private/Insights/MemoryProfiler/ViewModels/MemAllocNode.cpp
+++ b/Engine/Source/Developer/TraceInsights/Private/Insights/MemoryProfiler/ViewModels/MemAllocNode.cpp
@@ -78,16 +78,17 @@ FText FMemAllocNode::GetTopFunctionOrSourceFile(uint8 Flags) const
 		return FText::FromString(GetCallstackNotAvailableString());
 	}
 
-	const uint32 FrameCount = Callstack->Num();
-	if (FrameCount == 0)
+	const uint32 NumCallstackFrames = Callstack->Num();
+	if (NumCallstackFrames == 0)
 	{
 		return FText::FromString(GetEmptyCallstackString());
 	}
+	check(NumCallstackFrames <= 256); // see Callstack->Frame(uint8)
 
 	const TraceServices::FStackFrame* Frame = nullptr;
-	for (uint32 FrameIndex = 0; FrameIndex < FrameCount; ++FrameIndex)
+	for (uint32 FrameIndex = 0; FrameIndex < NumCallstackFrames; ++FrameIndex)
 	{
-		Frame = Callstack->Frame(FrameIndex);
+		Frame = Callstack->Frame(static_cast<uint8>(FrameIndex));
 		check(Frame != nullptr);
 
 		if (Frame->Symbol &&
diff --git a/Engine/Source/Developer/TraceInsights/Private/Insights/MemoryProfiler/ViewModels/MemoryAlloc.cpp b/Engine/Source/Developer/TraceInsights/Private/Insights/MemoryProfiler/ViewModels/MemoryAlloc.cpp
--- a/Engine/Source/Developer/TraceInsights/Private/Insights/MemoryProfiler/ViewModels/MemoryAlloc.cpp
+++ b/Engine/Source/Developer/TraceInsights/Private/Insights/MemoryProfiler/ViewModels/MemoryAlloc.cpp
@@ -48,14 +48,15 @@ FText FMemoryAlloc::GetFullCallstack() const
 	}
 
 	TStringBuilder<1024> Tooltip;
-	const uint32 FramesNum = Callstack->Num();
-	for (uint32 Index = 0; Index < FramesNum; ++Index)
+	const uint32 NumCallstackFrames = Callstack->Num();
+	check(NumCallstackFrames <= 256); // see Callstack->Frame(uint8)
+	for (uint32 Index = 0; Index < NumCallstackFrames; ++Index)
 	{
 		if (Index != 0)
 		{
 			Tooltip << TEXT("\n");
 		}
-		const TraceServices::FStackFrame* Frame = Callstack->Frame(Index);
+		const TraceServices::FStackFrame* Frame = Callstack->Frame(static_cast<uint8>(Index));
 		check(Frame != nullptr);
 		FormatStackFrame(*Frame, Tooltip, EStackFrameFormatFlags::ModuleSymbolFileAndLine);
 	}
@@ -77,14 +78,15 @@ FText FMemoryAlloc::GetFullCallstackSourceFiles() const
 	}
 
 	TStringBuilder<1024> Tooltip;
-	const uint32 FramesNum = Callstack->Num();
-	for (uint32 Index = 0; Index < FramesNum; ++Index)
+	const uint32 NumCallstackFrames = Callstack->Num();
+	check(NumCallstackFrames <= 256); // see Callstack->Frame(uint8)
+	for (uint32 FrameIndex = 0; FrameIndex < NumCallstackFrames; ++FrameIndex)
 	{
-		if (Index != 0)
+		if (FrameIndex != 0)
 		{
 			Tooltip << TEXT("\n");
 		}
-		const TraceServices::FStackFrame* Frame = Callstack->Frame(Index);
+		const TraceServices::FStackFrame* Frame = Callstack->Frame(static_cast<uint8>(FrameIndex));
 		check(Frame != nullptr);
 		FormatStackFrame(*Frame, Tooltip, EStackFrameFormatFlags::FileAndLine);
 	}
diff --git a/Engine/Source/Developer/TraceInsights/Private/Insights/MemoryProfiler/ViewModels/MemoryGraphTrack.cpp b/Engine/Source/Developer/TraceInsights/Private/Insights/MemoryProfiler/ViewModels/MemoryGraphTrack.cpp
--- a/Engine/Source/Developer/TraceInsights/Private/Insights/MemoryProfiler/ViewModels/MemoryGraphTrack.cpp
+++ b/Engine/Source/Developer/TraceInsights/Private/Insights/MemoryProfiler/ViewModels/MemoryGraphTrack.cpp
@@ -678,7 +678,7 @@ void FMemoryGraphTrack::DrawVerticalAxisGrid(const ITimingTrackDrawContext& Cont
 	ViewportY.SetSize(GetHeight());
 	ViewportY.SetScaleLimits(std::numeric_limits<double>::min(), std::numeric_limits<double>::max());
 	ViewportY.SetScale(Series->GetScaleY());
-	ViewportY.ScrollAtPos(Series->GetBaselineY() - GetHeight());
+	ViewportY.ScrollAtPos(static_cast<float>(Series->GetBaselineY()) - GetHeight());
 
 	const float ViewWidth = Context.GetViewport().GetWidth();
 	const float RoundedViewHeight = FMath::RoundToFloat(GetHeight());
@@ -833,8 +833,8 @@ void FMemoryGraphTrack::DrawVerticalAxisGrid(const ITimingTrackDrawContext& Cont
 
 		if (MinMaxAxis == 2)
 		{
-			const float MX = Context.GetMousePosition().X;
-			const float MY = Context.GetMousePosition().Y;
+			const float MX = static_cast<float>(Context.GetMousePosition().X);
+			const float MY = static_cast<float>(Context.GetMousePosition().Y);
 
 			//constexpr float MX1 = 80.0f; // start fading out
 			constexpr float MX2 = 120.0f; // completly faded out
@@ -896,13 +896,14 @@ void FMemoryGraphTrack::DrawHorizontalAxisLabel(const TDrawHorizontalAxisLabelPa
 
 	const float FontScale = Params.DrawContext.Geometry.Scale;
 	const FVector2D TextSize = Params.FontMeasureService->Measure(LabelText, Font, FontScale) / FontScale;
+	const float TextW = static_cast<float>(TextSize.X);
 	constexpr float TextH = 14.0f;
 
 	// Draw background for value text.
-	Params.DrawContext.DrawBox(Params.DrawContext.LayerId + 1, Params.X - TextSize.X - 4.0f, Params.Y, TextSize.X + 5.0f, TextH, Params.Brush, Params.TextBgColor);
+	Params.DrawContext.DrawBox(Params.DrawContext.LayerId + 1, Params.X - TextW - 4.0f, Params.Y, TextW + 5.0f, TextH, Params.Brush, Params.TextBgColor);
 
 	// Draw value text.
-	Params.DrawContext.DrawText(Params.DrawContext.LayerId + 2, Params.X - TextSize.X - 2.0f, Params.Y + 1.0f, LabelText, Font, Params.TextColor);
+	Params.DrawContext.DrawText(Params.DrawContext.LayerId + 2, Params.X - TextW - 2.0f, Params.Y + 1.0f, LabelText, Font, Params.TextColor);
 }
 
 ////////////////////////////////////////////////////////////////////////////////////////////////////
diff --git a/Engine/Source/Developer/TraceInsights/Private/Insights/MemoryProfiler/Widgets/SMemAllocTableTreeView.cpp b/Engine/Source/Developer/TraceInsights/Private/Insights/MemoryProfiler/Widgets/SMemAllocTableTreeView.cpp
--- a/Engine/Source/Developer/TraceInsights/Private/Insights/MemoryProfiler/Widgets/SMemAllocTableTreeView.cpp
+++ b/Engine/Source/Developer/TraceInsights/Private/Insights/MemoryProfiler/Widgets/SMemAllocTableTreeView.cpp
@@ -151,7 +151,7 @@ void SMemAllocTableTreeView::RebuildTree(bool bResync)
 						continue;
 					}
 
-					FName NodeName(Alloc->bIsBlock ? BaseHeapName : BaseNodeName, Alloc->GetStartEventIndex() + 1);
+					FName NodeName(Alloc->bIsBlock ? BaseHeapName : BaseNodeName, static_cast<int32>(Alloc->GetStartEventIndex() + 1));
 					FMemAllocNodePtr NodePtr = MakeShared<FMemAllocNode>(NodeName, MemAllocTable, AllocIndex);
 					TableTreeNodes.Add(NodePtr);
 				}
@@ -378,7 +378,7 @@ void SMemAllocTableTreeView::UpdateQuery(TraceServices::IAllocationsProvider::EQ
 				const uint32 AllocCount = Result->Num();
 				TotalAllocCount += AllocCount;
 
-				uint64 AllocsDestIndex = Allocs.Num();
+				int32 AllocsDestIndex = Allocs.Num();
 				Allocs.AddUninitialized(AllocCount);
 
 				for (uint32 AllocIndex = 0; AllocIndex < AllocCount; ++AllocIndex, ++AllocsDestIndex)
@@ -558,7 +558,7 @@ TSharedPtr<SWidget> SMemAllocTableTreeView::ConstructToolbar()
 		.VAlign(VAlign_Center)
 		[
 			SNew(SBox)
-			.MinDesiredWidth(150)
+			.MinDesiredWidth(150.0f)
 			[
 				SAssignNew(PresetComboBox, SComboBox<TSharedRef<ITableTreeViewPreset>>)
 				.ToolTipText(this, &SMemAllocTableTreeView::ViewPreset_GetSelectedToolTipText)
@@ -1489,9 +1489,10 @@ uint32 SMemAllocTableTreeView::CountSourceFiles(FMemAllocNode& MemAllocNode)
 
 	uint32 NumSourceFiles = 0;
 	const uint32 NumCallstackFrames = Alloc->Callstack->Num();
+	check(NumCallstackFrames <= 256); // see Callstack->Frame(uint8)
 	for (uint32 FrameIndex = 0; FrameIndex < NumCallstackFrames; ++FrameIndex)
 	{
-		const TraceServices::FStackFrame* Frame = Alloc->Callstack->Frame(FrameIndex);
+		const TraceServices::FStackFrame* Frame = Alloc->Callstack->Frame(static_cast<uint8>(FrameIndex));
 		if (Frame && Frame->Symbol && Frame->Symbol->File)
 		{
 			++NumSourceFiles;
@@ -1518,9 +1519,10 @@ void SMemAllocTableTreeView::BuildOpenSourceSubMenu(FMenuBuilder& MenuBuilder)
 				ISourceCodeAccessor& SourceCodeAccessor = SourceCodeAccessModule.GetAccessor();
 
 				const uint32 NumCallstackFrames = Alloc->Callstack->Num();
+				check(NumCallstackFrames <= 256); // see Callstack->Frame(uint8)
 				for (uint32 FrameIndex = 0; FrameIndex < NumCallstackFrames; ++FrameIndex)
 				{
-					const TraceServices::FStackFrame* Frame = Alloc->Callstack->Frame(FrameIndex);
+					const TraceServices::FStackFrame* Frame = Alloc->Callstack->Frame(static_cast<uint8>(FrameIndex));
 					if (Frame && Frame->Symbol && Frame->Symbol->File)
 					{
 						FText ItemLabel;
diff --git a/Engine/Source/Developer/TraceInsights/Private/Insights/MemoryProfiler/Widgets/SMemoryProfilerWindow.cpp b/Engine/Source/Developer/TraceInsights/Private/Insights/MemoryProfiler/Widgets/SMemoryProfilerWindow.cpp
--- a/Engine/Source/Developer/TraceInsights/Private/Insights/MemoryProfiler/Widgets/SMemoryProfilerWindow.cpp
+++ b/Engine/Source/Developer/TraceInsights/Private/Insights/MemoryProfiler/Widgets/SMemoryProfilerWindow.cpp
@@ -509,7 +509,7 @@ void SMemoryProfilerWindow::CreateTimingViewMarkers()
 	{
 		TSharedRef<Insights::FTimeMarker> TimeMarker = MakeShared<Insights::FTimeMarker>();
 
-		TimeMarkerName[0] = TEXT('A') + Index; // "A", "B", "C", etc.
+		TimeMarkerName[0] = static_cast<TCHAR>(TEXT('A') + Index); // "A", "B", "C", etc.
 		TimeMarker->SetName(TimeMarkerName);
 
 		const uint32 HueStep = 256 / MaxNumTimeMarkers;
diff --git a/Engine/Source/Developer/TraceInsights/Private/Insights/NetworkingProfiler/ViewModels/PacketContentViewDrawHelper.cpp b/Engine/Source/Developer/TraceInsights/Private/Insights/NetworkingProfiler/ViewModels/PacketContentViewDrawHelper.cpp
--- a/Engine/Source/Developer/TraceInsights/Private/Insights/NetworkingProfiler/ViewModels/PacketContentViewDrawHelper.cpp
+++ b/Engine/Source/Developer/TraceInsights/Private/Insights/NetworkingProfiler/ViewModels/PacketContentViewDrawHelper.cpp
@@ -44,7 +44,7 @@ void FPacketContentViewDrawStateBuilder::AddEvent(const TraceServices::FNetProfi
 	PacketEvent.ObjectInstanceIndex = Event.ObjectInstanceIndex;
 	PacketEvent.NetId = NetId;
 	PacketEvent.BitOffset = Event.StartPos;
-	PacketEvent.BitSize = Event.EndPos - Event.StartPos;
+	PacketEvent.BitSize = static_cast<uint32>(Event.EndPos - Event.StartPos);
 	PacketEvent.Level = Event.Level;
 
 	const FAxisViewportDouble& ViewportX = Viewport.GetHorizontalAxisViewport();
@@ -237,7 +237,7 @@ void FPacketContentViewDrawStateBuilder::AddEvent(const TraceServices::FNetProfi
 			{
 				Builder.Append(TEXT(" ("));
 			}
-			const uint32 EventSize = Event.EndPos - Event.StartPos;
+			const uint32 EventSize = static_cast<uint32>(Event.EndPos - Event.StartPos);
 			if (EventSize == 1)
 			{
 				Builder.Append(TEXT("1 bit)"));
@@ -513,10 +513,10 @@ void FPacketContentViewDrawHelper::DrawEventHighlight(const FNetworkPacketEvent&
 	{
 		// Animate color from white (if selected and hovered) or yellow (if only selected) to black, using a squared sine function.
 		const double Time = static_cast<double>(FPlatformTime::Cycles64()) * FPlatformTime::GetSecondsPerCycle64();
-		float S = FMath::Sin(2.0 * Time);
-		S = S * S; // squared, to ensure only positive [0 - 1] values
-		const float Blue = (Mode == EHighlightMode::SelectedAndHovered) ? 0.0f : S;
-		const FLinearColor Color(S, S, Blue, 1.0f);
+		float Hue = static_cast<float>(FMath::Sin(2.0 * Time));
+		Hue = Hue * Hue; // squared, to ensure only positive [0 - 1] values
+		const float Blue = (Mode == EHighlightMode::SelectedAndHovered) ? 0.0f : Hue;
+		const FLinearColor Color(Hue, Hue, Blue, 1.0f);
 
 		// Draw border around the timing event box.
 #if INSIGHTS_USE_LEGACY_BORDER
diff --git a/Engine/Source/Developer/TraceInsights/Private/Insights/NetworkingProfiler/ViewModels/PacketViewDrawHelper.cpp b/Engine/Source/Developer/TraceInsights/Private/Insights/NetworkingProfiler/ViewModels/PacketViewDrawHelper.cpp
--- a/Engine/Source/Developer/TraceInsights/Private/Insights/NetworkingProfiler/ViewModels/PacketViewDrawHelper.cpp
+++ b/Engine/Source/Developer/TraceInsights/Private/Insights/NetworkingProfiler/ViewModels/PacketViewDrawHelper.cpp
@@ -284,10 +284,10 @@ void FPacketViewDrawHelper::DrawSampleHighlight(const FNetworkPacketAggregatedSa
 	{
 		// Animate color from white (if selected and hovered) or yellow (if only selected) to black, using a squared sine function.
 		const double Time = static_cast<double>(FPlatformTime::Cycles64()) * FPlatformTime::GetSecondsPerCycle64();
-		float S = FMath::Sin(2.0 * Time);
-		S = S * S; // squared, to ensure only positive [0 - 1] values
-		const float Blue = (Mode == EHighlightMode::SelectedAndHovered) ? 0.0f : S;
-		const FLinearColor Color(S, S, Blue, 1.0f);
+		float Hue = static_cast<float>(FMath::Sin(2.0 * Time));
+		Hue = Hue * Hue; // squared, to ensure only positive [0 - 1] values
+		const float Blue = (Mode == EHighlightMode::SelectedAndHovered) ? 0.0f : Hue;
+		const FLinearColor Color(Hue, Hue, Blue, 1.0f);
 
 		// Draw border around the selected box.
 #if INSIGHTS_USE_LEGACY_BORDER
@@ -322,9 +322,9 @@ void FPacketViewDrawHelper::DrawSelection(int32 StartPacketIndex, int32 EndPacke
 
 	//// Animate color from white (if selected and hovered) or yellow (if only selected) to black, using a squared sine function.
 	//const double Time = static_cast<double>(FPlatformTime::Cycles64()) * FPlatformTime::GetSecondsPerCycle64();
-	//float S = FMath::Sin(2.0 * Time);
-	//S = S * S; // squared, to ensure only positive [0 - 1] values
-	//const FLinearColor Color(S, S, S, 1.0f);
+	//float Hue = static_cast<float>(FMath::Sin(2.0 * Time));
+	//Hue = Hue * Hue; // squared, to ensure only positive [0 - 1] values
+	//const FLinearColor Color(Hue, Hue, Hue, 1.0f);
 	//
 	//// Draw border around the selected box.
 	//DrawContext.DrawBox(X1 - 1.0f, Y - 1.0f, X2 - X1 + 2.0f, H + 2.0f, SelectedEventBorderBrush, Color);
diff --git a/Engine/Source/Developer/TraceInsights/Private/Insights/NetworkingProfiler/Widgets/SNetworkingProfilerWindow.cpp b/Engine/Source/Developer/TraceInsights/Private/Insights/NetworkingProfiler/Widgets/SNetworkingProfilerWindow.cpp
--- a/Engine/Source/Developer/TraceInsights/Private/Insights/NetworkingProfiler/Widgets/SNetworkingProfilerWindow.cpp
+++ b/Engine/Source/Developer/TraceInsights/Private/Insights/NetworkingProfiler/Widgets/SNetworkingProfilerWindow.cpp
@@ -924,8 +924,7 @@ void SNetworkingProfilerWindow::SetSelectedEventTypeIndex(const uint32 InEventTy
 		{
 			if (NetStatsView)
 			{
-				const uint64 NodeId = static_cast<uint64>(SelectedEventTypeIndex);
-				NetStatsView->SelectNetEventNode(NodeId);
+				NetStatsView->SelectNetEventNode(SelectedEventTypeIndex);
 			}
 		}
 	}
diff --git a/Engine/Source/Developer/TraceInsights/Private/Insights/NetworkingProfiler/Widgets/SPacketContentView.cpp b/Engine/Source/Developer/TraceInsights/Private/Insights/NetworkingProfiler/Widgets/SPacketContentView.cpp
--- a/Engine/Source/Developer/TraceInsights/Private/Insights/NetworkingProfiler/Widgets/SPacketContentView.cpp
+++ b/Engine/Source/Developer/TraceInsights/Private/Insights/NetworkingProfiler/Widgets/SPacketContentView.cpp
@@ -836,8 +836,8 @@ void SPacketContentView::Tick(const FGeometry& AllottedGeometry, const double In
 	if (ThisGeometry != AllottedGeometry || bIsViewportDirty)
 	{
 		bIsViewportDirty = false;
-		const float ViewWidth = AllottedGeometry.GetLocalSize().X;
-		const float ViewHeight = AllottedGeometry.GetLocalSize().Y;
+		const float ViewWidth = static_cast<float>(AllottedGeometry.GetLocalSize().X);
+		const float ViewHeight = static_cast<float>(AllottedGeometry.GetLocalSize().Y);
 		Viewport.SetSize(ViewWidth, ViewHeight);
 		bIsStateDirty = true;
 	}
@@ -1047,7 +1047,7 @@ void SPacketContentView::UpdateState(float FontScale)
 				//const int64 StartPos = static_cast<int64>(FMath::FloorToDouble(ViewportX.GetValueAtOffset(0.0f)));
 				//const int64 EndPos = static_cast<int64>(FMath::CeilToDouble(ViewportX.GetValueAtOffset(ViewportX.GetSize())));
 				const uint32 StartPos = 0;
-				const uint32 EndPos = PacketBitSize;
+				const uint32 EndPos = static_cast<uint32>(PacketBitSize);
 				NetProfilerProvider->EnumeratePacketContentEventsByPosition(ConnectionIndex, ConnectionMode, PacketIndex, StartPos, EndPos, [this, &Builder, &FilteredDrawStateBuilder, NetProfilerProvider](const TraceServices::FNetProfilerContentEvent& Event)
 				{
 					const TCHAR* Name = nullptr;
@@ -1088,7 +1088,7 @@ void SPacketContentView::UpdateState(float FontScale)
 
 void SPacketContentView::UpdateHoveredEvent()
 {
-	HoveredEvent = GetEventAtMousePosition(MousePosition.X, MousePosition.Y);
+	HoveredEvent = GetEventAtMousePosition(static_cast<float>(MousePosition.X), static_cast<float>(MousePosition.Y));
 	//if (!HoveredEvent.IsValid())
 	//{
 	//	HoveredEvent = GetEventAtMousePosition(MousePosition.X - 1.0f, MousePosition.Y);
@@ -1226,8 +1226,8 @@ int32 SPacketContentView::OnPaint(const FPaintArgs& Args, const FGeometry& Allot
 	const ESlateDrawEffect DrawEffects = bEnabled ? ESlateDrawEffect::None : ESlateDrawEffect::DisabledEffect;
 	FDrawContext DrawContext(AllottedGeometry, MyCullingRect, InWidgetStyle, DrawEffects, OutDrawElements, LayerId);
 
-	const float ViewWidth = AllottedGeometry.Size.X;
-	const float ViewHeight = AllottedGeometry.Size.Y;
+	const float ViewWidth = static_cast<float>(AllottedGeometry.Size.X);
+	const float ViewHeight = static_cast<float>(AllottedGeometry.Size.Y);
 
 	//////////////////////////////////////////////////
 	{
@@ -1288,7 +1288,7 @@ int32 SPacketContentView::OnPaint(const FPaintArgs& Args, const FGeometry& Allot
 
 		const TSharedRef<FSlateFontMeasure> FontMeasureService = FSlateApplication::Get().GetRenderer()->GetFontMeasureService();
 		const float FontScale = DrawContext.Geometry.Scale;
-		const float MaxFontCharHeight = FontMeasureService->Measure(TEXT("!"), SummaryFont, FontScale).Y / FontScale;
+		const float MaxFontCharHeight = static_cast<float>(FontMeasureService->Measure(TEXT("!"), SummaryFont, FontScale).Y / FontScale);
 		const float DbgDY = MaxFontCharHeight;
 
 		const float DbgW = 280.0f;
@@ -1489,7 +1489,7 @@ FReply SPacketContentView::OnMouseMove(const FGeometry& MyGeometry, const FPoint
 				}
 
 				FAxisViewportDouble& ViewportX = Viewport.GetHorizontalAxisViewport();
-				const float PosX = ViewportPosXOnButtonDown + (MousePositionOnButtonDown.X - MousePosition.X);
+				const float PosX = ViewportPosXOnButtonDown + static_cast<float>(MousePositionOnButtonDown.X - MousePosition.X);
 				ViewportX.ScrollAtPos(PosX);
 				UpdateHorizontalScrollBar();
 				bIsStateDirty = true;
@@ -1546,7 +1546,7 @@ FReply SPacketContentView::OnMouseWheel(const FGeometry& MyGeometry, const FPoin
 	{
 		// Zoom in/out horizontally.
 		const float Delta = MouseEvent.GetWheelDelta();
-		ZoomHorizontally(Delta, MousePosition.X);
+		ZoomHorizontally(Delta, static_cast<float>(MousePosition.X));
 	}
 
 	return FReply::Handled();
@@ -1657,13 +1657,13 @@ FReply SPacketContentView::OnKeyDown(const FGeometry& MyGeometry, const FKeyEven
 	else if (InKeyEvent.GetKey() == EKeys::Equals ||
 			 InKeyEvent.GetKey() == EKeys::Add)
 	{
-		ZoomHorizontally(1.0f, MousePosition.X);
+		ZoomHorizontally(1.0f, static_cast<float>(MousePosition.X));
 		return FReply::Handled();
 	}
 	else if (InKeyEvent.GetKey() == EKeys::Hyphen ||
 			 InKeyEvent.GetKey() == EKeys::Subtract)
 	{
-		ZoomHorizontally(-1.0f, MousePosition.X);
+		ZoomHorizontally(-1.0f, static_cast<float>(MousePosition.X));
 		return FReply::Handled();
 	}
 
diff --git a/Engine/Source/Developer/TraceInsights/Private/Insights/NetworkingProfiler/Widgets/SPacketView.cpp b/Engine/Source/Developer/TraceInsights/Private/Insights/NetworkingProfiler/Widgets/SPacketView.cpp
--- a/Engine/Source/Developer/TraceInsights/Private/Insights/NetworkingProfiler/Widgets/SPacketView.cpp
+++ b/Engine/Source/Developer/TraceInsights/Private/Insights/NetworkingProfiler/Widgets/SPacketView.cpp
@@ -171,8 +171,8 @@ void SPacketView::Tick(const FGeometry& AllottedGeometry, const double InCurrent
 	if (ThisGeometry != AllottedGeometry || bIsViewportDirty)
 	{
 		bIsViewportDirty = false;
-		const float ViewWidth = AllottedGeometry.GetLocalSize().X;
-		const float ViewHeight = AllottedGeometry.GetLocalSize().Y;
+		const float ViewWidth = static_cast<float>(AllottedGeometry.GetLocalSize().X);
+		const float ViewHeight = static_cast<float>(AllottedGeometry.GetLocalSize().Y);
 		Viewport.SetSize(ViewWidth, ViewHeight);
 		bIsStateDirty = true;
 	}
@@ -384,7 +384,7 @@ void SPacketView::UpdateState()
 
 										if (bEventMatchesFilter)
 										{
-											const uint32 EventSize = Event.EndPos - Event.StartPos;
+											const uint32 EventSize = static_cast<uint32>(Event.EndPos - Event.StartPos);
 											FilterMatchAggregatedEventSizeInBits += EventSize;
 											FilterMatchMaxEventSizeBits = FMath::Max(FilterMatchMaxEventSizeBits, EventSize);
 
@@ -434,7 +434,7 @@ void SPacketView::UpdateState()
 		for (int32 PacketIndex = StartIndex; PacketIndex < EndIndex; ++PacketIndex)
 		{
 			FRandomStream RandomStream((PacketIndex * PacketIndex * PacketIndex) ^ 0x2c2c57ed);
-			int64 Size = RandomStream.RandRange(0, 2000);
+			const uint32 Size = static_cast<uint32>(RandomStream.RandRange(0, 2000));
 
 			TraceServices::ENetProfilerDeliveryStatus Status = TraceServices::ENetProfilerDeliveryStatus::Unknown;
 			const float Fraction = RandomStream.GetFraction();
@@ -502,12 +502,12 @@ FNetworkPacketSampleRef SPacketView::GetSample(const int32 InPacketIndex)
 
 ////////////////////////////////////////////////////////////////////////////////////////////////////
 
-FNetworkPacketSampleRef SPacketView::GetSampleAtMousePosition(float X, float Y)
+FNetworkPacketSampleRef SPacketView::GetSampleAtMousePosition(double X, double Y)
 {
 	if (!bIsStateDirty)
 	{
 		float SampleW = Viewport.GetSampleWidth();
-		int32 SampleIndex = FMath::FloorToInt(X / SampleW);
+		int32 SampleIndex = FMath::FloorToInt(static_cast<float>(X) / SampleW);
 		if (SampleIndex >= 0)
 		{
 			if (PacketSeries->NumAggregatedPackets > 0 &&
@@ -529,7 +529,9 @@ FNetworkPacketSampleRef SPacketView::GetSampleAtMousePosition(float X, float Y)
 					const float BottomY = FMath::Min(ViewHeight, ViewHeight - BaselineY + ToleranceY);
 					const float TopY = FMath::Max(0.0f, ViewHeight - ValueY - ToleranceY);
 
-					if (Y >= TopY && Y < BottomY)
+					const float MY = static_cast<float>(Y);
+
+					if (MY >= TopY && MY < BottomY)
 					{
 						return FNetworkPacketSampleRef(PacketSeries, MakeShared<FNetworkPacketAggregatedSample>(Sample));
 					}
@@ -542,16 +544,16 @@ FNetworkPacketSampleRef SPacketView::GetSampleAtMousePosition(float X, float Y)
 
 ////////////////////////////////////////////////////////////////////////////////////////////////////
 
-void SPacketView::SelectSampleAtMousePosition(float X, float Y, const FPointerEvent& MouseEvent)
+void SPacketView::SelectSampleAtMousePosition(double X, double Y, const FPointerEvent& MouseEvent)
 {
 	FNetworkPacketSampleRef SampleRef = GetSampleAtMousePosition(X, Y);
 	if (!SampleRef.IsValid())
 	{
-		SampleRef = GetSampleAtMousePosition(X - 1.0f, Y);
+		SampleRef = GetSampleAtMousePosition(X - 1.0, Y);
 	}
 	if (!SampleRef.IsValid())
 	{
-		SampleRef = GetSampleAtMousePosition(X + 1.0f, Y);
+		SampleRef = GetSampleAtMousePosition(X + 1.0, Y);
 	}
 
 	bool bRaiseSelectionChanged = false;
@@ -687,8 +689,8 @@ int32 SPacketView::OnPaint(const FPaintArgs& Args, const FGeometry& AllottedGeom
 
 	const FSlateBrush* WhiteBrush = FInsightsStyle::Get().GetBrush("WhiteBrush");
 
-	const float ViewWidth = AllottedGeometry.Size.X;
-	const float ViewHeight = AllottedGeometry.Size.Y;
+	const float ViewWidth = static_cast<float>(AllottedGeometry.Size.X);
+	const float ViewHeight = static_cast<float>(AllottedGeometry.Size.Y);
 
 	int32 NumDrawSamples = 0;
 
@@ -805,7 +807,7 @@ int32 SPacketView::OnPaint(const FPaintArgs& Args, const FGeometry& AllottedGeom
 			const FVector2D TextSize = FontMeasureService->Measure(Text, SummaryFont, FontScale) / FontScale;
 
 			const float DX = 2.0f;
-			const float W2 = TextSize.X / 2 + DX;
+			const float W2 = static_cast<float>(TextSize.X) / 2 + DX;
 
 			const FAxisViewportInt32& ViewportX = Viewport.GetHorizontalAxisViewport();
 
@@ -836,7 +838,7 @@ int32 SPacketView::OnPaint(const FPaintArgs& Args, const FGeometry& AllottedGeom
 	const bool bShouldDisplayDebugInfo = FInsightsManager::Get()->IsDebugInfoEnabled();
 	if (bShouldDisplayDebugInfo)
 	{
-		const float MaxFontCharHeight = FontMeasureService->Measure(TEXT("!"), SummaryFont, FontScale).Y / FontScale;
+		const float MaxFontCharHeight = static_cast<float>(FontMeasureService->Measure(TEXT("!"), SummaryFont, FontScale).Y / FontScale);
 		const float DbgDY = MaxFontCharHeight;
 
 		const float DbgW = 280.0f;
@@ -985,13 +987,14 @@ void SPacketView::DrawVerticalAxisGrid(FDrawContext& DrawContext, const FSlateBr
 			const int64 ValueBits = static_cast<int64>(Value);
 			const FString Text = (ValueBits == 0) ? TEXT("0") : FString::Format(TEXT("{0} bits"), { FText::AsNumber(ValueBits).ToString() });
 			const FVector2D TextSize = FontMeasureService->Measure(Text, Font, FontScale) / FontScale;
+			const float TextW = static_cast<float>(TextSize.X);
 			constexpr float TextH = 14.0f;
 
 			// Draw background for value text.
-			DrawContext.DrawBox(ViewWidth - TextSize.X - 4.0f, Y - TextH, TextSize.X + 4.0f, TextH, Brush, TextBgColor);
+			DrawContext.DrawBox(ViewWidth - TextW - 4.0f, Y - TextH, TextW + 4.0f, TextH, Brush, TextBgColor);
 
 			// Draw value text.
-			DrawContext.DrawText(ViewWidth - TextSize.X - 2.0f, Y - TextH + 1.0f, Text, Font, TextColor);
+			DrawContext.DrawText(ViewWidth - TextW - 2.0f, Y - TextH + 1.0f, Text, Font, TextColor);
 		}
 		DrawContext.LayerId++;
 	}
@@ -1034,7 +1037,7 @@ void SPacketView::DrawHorizontalAxisGrid(FDrawContext& DrawContext, const FSlate
 		const int32 Grid = ((Delta + Power10 - 1) / Power10) * Power10; // next value divisible with a multiple of 10
 
 		// Skip grid lines for negative indices.
-		double StartIndex = ((LeftIndex + Grid - 1) / Grid) * Grid;
+		int32 StartIndex = ((LeftIndex + Grid - 1) / Grid) * Grid;
 		while (StartIndex < 0)
 		{
 			StartIndex += Grid;
@@ -1121,7 +1124,7 @@ FReply SPacketView::OnMouseButtonUp(const FGeometry& MyGeometry, const FPointerE
 			}
 			else if (bIsValidForMouseClick)
 			{
-				SelectSampleAtMousePosition(MousePositionOnButtonUp.X, MousePositionOnButtonUp.Y, MouseEvent);
+				SelectSampleAtMousePosition(static_cast<float>(MousePositionOnButtonUp.X), static_cast<float>(MousePositionOnButtonUp.Y), MouseEvent);
 			}
 
 			bIsLMB_Pressed = false;
@@ -1178,7 +1181,7 @@ FReply SPacketView::OnMouseMove(const FGeometry& MyGeometry, const FPointerEvent
 				}
 
 				FAxisViewportInt32& ViewportX = Viewport.GetHorizontalAxisViewport();
-				const float PosX = ViewportPosXOnButtonDown + (MousePositionOnButtonDown.X - MousePosition.X);
+				const float PosX = ViewportPosXOnButtonDown + static_cast<float>(MousePositionOnButtonDown.X - MousePosition.X);
 				ViewportX.ScrollAtValue(ViewportX.GetValueAtPos(PosX)); // align viewport position with sample
 				UpdateHorizontalScrollBar();
 				bIsStateDirty = true;
@@ -1193,11 +1196,11 @@ FReply SPacketView::OnMouseMove(const FGeometry& MyGeometry, const FPointerEvent
 			HoveredSample = GetSampleAtMousePosition(MousePosition.X, MousePosition.Y);
 			if (!HoveredSample.IsValid())
 			{
-				HoveredSample = GetSampleAtMousePosition(MousePosition.X - 1.0f, MousePosition.Y);
+				HoveredSample = GetSampleAtMousePosition(MousePosition.X - 1.0, MousePosition.Y);
 			}
 			if (!HoveredSample.IsValid())
 			{
-				HoveredSample = GetSampleAtMousePosition(MousePosition.X + 1.0f, MousePosition.Y);
+				HoveredSample = GetSampleAtMousePosition(MousePosition.X + 1.0, MousePosition.Y);
 			}
 		}
 
@@ -1239,17 +1242,17 @@ FReply SPacketView::OnMouseWheel(const FGeometry& MyGeometry, const FPointerEven
 		FAxisViewportDouble& ViewportY = Viewport.GetVerticalAxisViewport();
 
 		// Zoom in/out vertically.
-		const float Delta = MouseEvent.GetWheelDelta();
-		constexpr float ZoomStep = 0.25f; // as percent
-		float ScaleY;
+		const double Delta = MouseEvent.GetWheelDelta();
+		constexpr double ZoomStep = 0.25; // as percent
+		double ScaleY;
 
 		if (Delta > 0)
 		{
-			ScaleY = ViewportY.GetScale() * FMath::Pow(1.0f + ZoomStep, Delta);
+			ScaleY = ViewportY.GetScale() * FMath::Pow(1.0 + ZoomStep, Delta);
 		}
 		else
 		{
-			ScaleY = ViewportY.GetScale() * FMath::Pow(1.0f / (1.0f + ZoomStep), -Delta);
+			ScaleY = ViewportY.GetScale() * FMath::Pow(1.0 / (1.0 + ZoomStep), -Delta);
 		}
 
 		ViewportY.SetScale(ScaleY);
@@ -1259,7 +1262,7 @@ FReply SPacketView::OnMouseWheel(const FGeometry& MyGeometry, const FPointerEven
 	{
 		// Zoom in/out horizontally.
 		const float Delta = MouseEvent.GetWheelDelta();
-		ZoomHorizontally(Delta, MousePosition.X);
+		ZoomHorizontally(Delta, static_cast<float>(MousePosition.X));
 	}
 
 	return FReply::Handled();
@@ -1361,12 +1364,12 @@ FReply SPacketView::OnKeyDown(const FGeometry& MyGeometry, const FKeyEvent& InKe
 		if (InKeyEvent.GetModifierKeys().IsShiftDown())
 		{
 			FAxisViewportDouble& ViewportY = Viewport.GetVerticalAxisViewport();
-			const float ScaleY = ViewportY.GetScale() * 1.25f;
+			const double ScaleY = ViewportY.GetScale() * 1.25;
 			ViewportY.SetScale(ScaleY);
 		}
 		else
 		{
-			ZoomHorizontally(1.0f, MousePosition.X);
+			ZoomHorizontally(1.0f, static_cast<float>(MousePosition.X));
 		}
 		return FReply::Handled();
 	}
@@ -1376,12 +1379,12 @@ FReply SPacketView::OnKeyDown(const FGeometry& MyGeometry, const FKeyEvent& InKe
 		if (InKeyEvent.GetModifierKeys().IsShiftDown())
 		{
 			FAxisViewportDouble& ViewportY = Viewport.GetVerticalAxisViewport();
-			const float ScaleY = ViewportY.GetScale() * 0.8f;
+			const double ScaleY = ViewportY.GetScale() * 0.8;
 			ViewportY.SetScale(ScaleY);
 		}
 		else
 		{
-			ZoomHorizontally(-1.0f, MousePosition.X);
+			ZoomHorizontally(-1.0f, static_cast<float>(MousePosition.X));
 		}
 		return FReply::Handled();
 	}
diff --git a/Engine/Source/Developer/TraceInsights/Private/Insights/NetworkingProfiler/Widgets/SPacketView.h b/Engine/Source/Developer/TraceInsights/Private/Insights/NetworkingProfiler/Widgets/SPacketView.h
--- a/Engine/Source/Developer/TraceInsights/Private/Insights/NetworkingProfiler/Widgets/SPacketView.h
+++ b/Engine/Source/Developer/TraceInsights/Private/Insights/NetworkingProfiler/Widgets/SPacketView.h
@@ -154,8 +154,8 @@ private:
 	void DrawVerticalAxisGrid(FDrawContext& DrawContext, const FSlateBrush* Brush, const FSlateFontInfo& Font) const;
 
 	FNetworkPacketSampleRef GetSample(const int32 InPacketIndex);
-	FNetworkPacketSampleRef GetSampleAtMousePosition(float X, float Y);
-	void SelectSampleAtMousePosition(float X, float Y, const FPointerEvent& MouseEvent);
+	FNetworkPacketSampleRef GetSampleAtMousePosition(double X, double Y);
+	void SelectSampleAtMousePosition(double X, double Y, const FPointerEvent& MouseEvent);
 	void OnSelectionChanged();
 
 	void ShowContextMenu(const FPointerEvent& MouseEvent);
diff --git a/Engine/Source/Developer/TraceInsights/Private/Insights/Table/ViewModels/TableColumn.cpp b/Engine/Source/Developer/TraceInsights/Private/Insights/Table/ViewModels/TableColumn.cpp
--- a/Engine/Source/Developer/TraceInsights/Private/Insights/Table/ViewModels/TableColumn.cpp
+++ b/Engine/Source/Developer/TraceInsights/Private/Insights/Table/ViewModels/TableColumn.cpp
@@ -81,7 +81,7 @@ inline void TrimEndingZero(FString& Str)
 		{
 			--Data;
 		}
-		Str.LeftInline(Data - Start + 1);
+		Str.LeftInline(static_cast<int32>(Data - Start + 1));
 	}
 }
 
diff --git a/Engine/Source/Developer/TraceInsights/Private/Insights/Table/ViewModels/UntypedTable.cpp b/Engine/Source/Developer/TraceInsights/Private/Insights/Table/ViewModels/UntypedTable.cpp
--- a/Engine/Source/Developer/TraceInsights/Private/Insights/Table/ViewModels/UntypedTable.cpp
+++ b/Engine/Source/Developer/TraceInsights/Private/Insights/Table/ViewModels/UntypedTable.cpp
@@ -107,7 +107,7 @@ bool AreTableLayoutsEqual(const TraceServices::ITableLayout& TableLayoutA, const
 		return false;
 	}
 
-	int32 ColumnCount = TableLayoutA.GetColumnCount();
+	int32 ColumnCount = static_cast<int32>(TableLayoutA.GetColumnCount());
 	for (int32 ColumnIndex = 0; ColumnIndex < ColumnCount; ++ColumnIndex)
 	{
 		if (TableLayoutA.GetColumnType(ColumnIndex) != TableLayoutB.GetColumnType(ColumnIndex))
@@ -159,7 +159,7 @@ bool FUntypedTable::UpdateSourceTable(TSharedPtr<TraceServices::IUntypedTable> I
 void FUntypedTable::CreateColumns(const TraceServices::ITableLayout& TableLayout)
 {
 	ensure(GetColumnCount() == 0);
-	const int32 ColumnCount = TableLayout.GetColumnCount();
+	const int32 ColumnCount = static_cast<int32>(TableLayout.GetColumnCount());
 
 	//////////////////////////////////////////////////
 	// Hierarchy Column
diff --git a/Engine/Source/Developer/TraceInsights/Private/Insights/Table/Widgets/STableTreeView.h b/Engine/Source/Developer/TraceInsights/Private/Insights/Table/Widgets/STableTreeView.h
--- a/Engine/Source/Developer/TraceInsights/Private/Insights/Table/Widgets/STableTreeView.h
+++ b/Engine/Source/Developer/TraceInsights/Private/Insights/Table/Widgets/STableTreeView.h
@@ -39,6 +39,7 @@ namespace TraceServices
 namespace Insights
 {
 
+class FFilterConfigurator;
 class FTable;
 class FTableColumn;
 class FTreeNodeGrouping;
diff --git a/Engine/Source/Developer/TraceInsights/Private/Insights/Table/Widgets/SUntypedTableTreeView.cpp b/Engine/Source/Developer/TraceInsights/Private/Insights/Table/Widgets/SUntypedTableTreeView.cpp
--- a/Engine/Source/Developer/TraceInsights/Private/Insights/Table/Widgets/SUntypedTableTreeView.cpp
+++ b/Engine/Source/Developer/TraceInsights/Private/Insights/Table/Widgets/SUntypedTableTreeView.cpp
@@ -83,7 +83,7 @@ void SUntypedTableTreeView::RebuildTree(bool bResync)
 
 	if (Session.IsValid() && SourceTable.IsValid() && TableReader.IsValid())
 	{
-		const int32 TotalRowCount = SourceTable->GetRowCount();
+		const int32 TotalRowCount = static_cast<int32>(SourceTable->GetRowCount());
 		if (TotalRowCount != TableTreeNodes.Num())
 		{
 			TableTreeNodes.Empty(TotalRowCount);
diff --git a/Engine/Source/Developer/TraceInsights/Private/Insights/TaskGraphProfiler/TaskGraphProfilerManager.cpp b/Engine/Source/Developer/TraceInsights/Private/Insights/TaskGraphProfiler/TaskGraphProfilerManager.cpp
--- a/Engine/Source/Developer/TraceInsights/Private/Insights/TaskGraphProfiler/TaskGraphProfilerManager.cpp
+++ b/Engine/Source/Developer/TraceInsights/Private/Insights/TaskGraphProfiler/TaskGraphProfilerManager.cpp
@@ -509,25 +509,24 @@ void FTaskGraphProfilerManager::InitializeColorCode()
 {
 	auto ToLiniarColorNoAlpha = [](uint32 Value)
 	{
-		float R = (Value & 0xFF000000) >> 24;
-		float G = (Value & 0x00FF0000) >> 16;
-		float B = (Value & 0x0000FF00) >> 8;
-
-		return FLinearColor(R / 255, G / 255, B / 255);
+		const float R = static_cast<float>((Value & 0xFF000000) >> 24) / 255.0f;
+		const float G = static_cast<float>((Value & 0x00FF0000) >> 16) / 255.0f;
+		const float B = static_cast<float>((Value & 0x0000FF00) >>  8) / 255.0f;
+		return FLinearColor(R, G, B);
 	};
 
-	ColorCode[static_cast<uint32>(ETaskEventType::Created)] = ToLiniarColorNoAlpha(0xFFDC1AFF); // Yellow
-	ColorCode[static_cast<uint32>(ETaskEventType::Launched)] = ToLiniarColorNoAlpha(0x8BC24AFF); // Green
-	ColorCode[static_cast<uint32>(ETaskEventType::Scheduled)] = ToLiniarColorNoAlpha(0x26BBFFFF); // Blue
-	ColorCode[static_cast<uint32>(ETaskEventType::Started)] = ToLiniarColorNoAlpha(0xFF0000FF); // Red
-	ColorCode[static_cast<uint32>(ETaskEventType::Finished)] = ToLiniarColorNoAlpha(0xFFDC1AFF); // Yellow
-	ColorCode[static_cast<uint32>(ETaskEventType::Completed)] = ToLiniarColorNoAlpha(0xFE9B07FF); // Orange
+	ColorCode[static_cast<uint32>(ETaskEventType::Created)]             = ToLiniarColorNoAlpha(0xFFDC1AFF); // Yellow
+	ColorCode[static_cast<uint32>(ETaskEventType::Launched)]            = ToLiniarColorNoAlpha(0x8BC24AFF); // Green
+	ColorCode[static_cast<uint32>(ETaskEventType::Scheduled)]           = ToLiniarColorNoAlpha(0x26BBFFFF); // Blue
+	ColorCode[static_cast<uint32>(ETaskEventType::Started)]             = ToLiniarColorNoAlpha(0xFF0000FF); // Red
+	ColorCode[static_cast<uint32>(ETaskEventType::Finished)]            = ToLiniarColorNoAlpha(0xFFDC1AFF); // Yellow
+	ColorCode[static_cast<uint32>(ETaskEventType::Completed)]           = ToLiniarColorNoAlpha(0xFE9B07FF); // Orange
 	ColorCode[static_cast<uint32>(ETaskEventType::PrerequisiteStarted)] = ToLiniarColorNoAlpha(0xFF729CFF); // Pink
-	ColorCode[static_cast<uint32>(ETaskEventType::ParentStarted)] = ToLiniarColorNoAlpha(0xB68F55FF); // Folder
-	ColorCode[static_cast<uint32>(ETaskEventType::NestedStarted)] = ToLiniarColorNoAlpha(0xB68F55FF); // Folder
-	ColorCode[static_cast<uint32>(ETaskEventType::SubsequentStarted)] = ToLiniarColorNoAlpha(0xFE9B07FF); // Orange
-	ColorCode[static_cast<uint32>(ETaskEventType::NestedCompleted)] = ToLiniarColorNoAlpha(0x804D39FF); // Brown
-	ColorCode[static_cast<uint32>(ETaskEventType::CriticalPath)] = ToLiniarColorNoAlpha(0xA139BFFF); // Purple
+	ColorCode[static_cast<uint32>(ETaskEventType::ParentStarted)]       = ToLiniarColorNoAlpha(0xB68F55FF); // Folder
+	ColorCode[static_cast<uint32>(ETaskEventType::NestedStarted)]       = ToLiniarColorNoAlpha(0xB68F55FF); // Folder
+	ColorCode[static_cast<uint32>(ETaskEventType::SubsequentStarted)]   = ToLiniarColorNoAlpha(0xFE9B07FF); // Orange
+	ColorCode[static_cast<uint32>(ETaskEventType::NestedCompleted)]     = ToLiniarColorNoAlpha(0x804D39FF); // Brown
+	ColorCode[static_cast<uint32>(ETaskEventType::CriticalPath)]        = ToLiniarColorNoAlpha(0xA139BFFF); // Purple
 }
 
 ////////////////////////////////////////////////////////////////////////////////////////////////////
diff --git a/Engine/Source/Developer/TraceInsights/Private/Insights/TaskGraphProfiler/ViewModels/TaskGraphRelation.cpp b/Engine/Source/Developer/TraceInsights/Private/Insights/TaskGraphProfiler/ViewModels/TaskGraphRelation.cpp
--- a/Engine/Source/Developer/TraceInsights/Private/Insights/TaskGraphProfiler/ViewModels/TaskGraphRelation.cpp
+++ b/Engine/Source/Developer/TraceInsights/Private/Insights/TaskGraphProfiler/ViewModels/TaskGraphRelation.cpp
@@ -88,32 +88,32 @@ void FTaskGraphRelation::Draw(const FDrawContext& DrawContext, const FTimingTrac
 
 	const FVector2D StartPoint = FVector2D(X1, Y1);
 	const FVector2D EndPoint = FVector2D(X2, Y2);
-	const float Distance = FVector2D::Distance(StartPoint, EndPoint);
+	const double Distance = FVector2D::Distance(StartPoint, EndPoint);
 
-	constexpr float LineHeightAtStart = 4.0f;
-	constexpr float LineLengthAtStart = 4.0f;
-	constexpr float LineLengthAtEnd = 12.0f;
+	constexpr double LineHeightAtStart = 4.0;
+	constexpr double LineLengthAtStart = 4.0;
+	constexpr double LineLengthAtEnd = 12.0;
 
-	const FVector2D StartDir(FMath::Max(X2 - X1, 4.0f * (LineLengthAtStart + LineLengthAtEnd)), 0.0f);
+	const FVector2D StartDir(FMath::Max(static_cast<double>(X2 - X1), 4.0 * (LineLengthAtStart + LineLengthAtEnd)), 0.0);
 
 	constexpr float OutlineThickness = 5.0f;
 	constexpr float LineThickness = 3.0f;
 
-	constexpr float ArrowDirectionLen = 10.0f;
-	constexpr float ArrowRotationAngle = 20.0f;
-	FVector2D ArrowDirection(-ArrowDirectionLen, 0.0f);
+	constexpr double ArrowDirectionLen = 10.0;
+	constexpr double ArrowRotationAngle = 20.0;
+	FVector2D ArrowDirection(-ArrowDirectionLen, 0.0);
 
 	const FLinearColor OutlineColor(0.0f, 0.0f, 0.0f, 1.0f);
 	const FLinearColor Color = FTaskGraphProfilerManager::Get()->GetColorForTaskEvent(Type);
 
 	TArray<FVector2D> LinePoints;
-	LinePoints.Add(StartPoint + FVector2D(0.0f, -LineHeightAtStart / 2.0f));
-	LinePoints.Add(StartPoint + FVector2D(0.0f, +LineHeightAtStart / 2.0f));
+	LinePoints.Add(StartPoint + FVector2D(0.0, -LineHeightAtStart / 2.0));
+	LinePoints.Add(StartPoint + FVector2D(0.0, +LineHeightAtStart / 2.0));
 	DrawContext.DrawLines(OutlineLayerId, 0.0f, 0.0f, LinePoints, ESlateDrawEffect::None, OutlineColor, /*bAntialias=*/ true, OutlineThickness);
 	DrawContext.DrawLines(LayerId, 0.0f, 0.0f, LinePoints, ESlateDrawEffect::None, Color, /*bAntialias=*/ true, LineThickness);
 
-	constexpr float MinDistance = 1.5f * (LineLengthAtStart + LineLengthAtEnd);
-	constexpr float MaxDistance = 10000.0f; // arbitrary limit to avoid stack overflow in recursive FLineBuilder::Subdivide when rendering splines
+	constexpr double MinDistance = 1.5 * (LineLengthAtStart + LineLengthAtEnd);
+	constexpr double MaxDistance = 10000.0; // arbitrary limit to avoid stack overflow in recursive FLineBuilder::Subdivide when rendering splines
 	if (Distance > MinDistance && Distance < MaxDistance && !FMath::IsNearlyEqual(StartPoint.Y, EndPoint.Y))
 	{
 		FVector2D SplineStart(StartPoint.X + LineLengthAtStart, StartPoint.Y);
diff --git a/Engine/Source/Developer/TraceInsights/Private/Insights/TaskGraphProfiler/ViewModels/TaskTimingTrack.cpp b/Engine/Source/Developer/TraceInsights/Private/Insights/TaskGraphProfiler/ViewModels/TaskTimingTrack.cpp
--- a/Engine/Source/Developer/TraceInsights/Private/Insights/TaskGraphProfiler/ViewModels/TaskTimingTrack.cpp
+++ b/Engine/Source/Developer/TraceInsights/Private/Insights/TaskGraphProfiler/ViewModels/TaskTimingTrack.cpp
@@ -876,7 +876,7 @@ const TSharedPtr<const ITimingEvent> FTaskTimingTrack::GetEvent(float InPosX, fl
 	if (DY >= 0 && DY < GetHeight() - 1.0f - 2 * Layout.TimelineDY)
 	{
 		const double EventTime = Viewport.SlateUnitsToTime(InPosX);
-		const int32 Depth = DY / (Layout.EventH + Layout.EventDY);
+		const int32 Depth = FMath::RoundToInt(DY / (Layout.EventH + Layout.EventDY));
 
 		TSharedPtr<const TraceServices::IAnalysisSession> Session = FInsightsManager::Get()->GetSession();
 		if (!Session.IsValid())
@@ -940,7 +940,7 @@ const TSharedPtr<const ITimingEvent> FTaskTimingTrack::GetEvent(float InPosX, fl
 				TimingEvent->SetTaskId(Task->Id);
 			}
 		}
-		else if(bShowDetailInfoOnTaskTrack)
+		else if (bShowDetailInfoOnTaskTrack)
 		{
 			auto GetEventFromRelations = [&TasksProvider, this, EventTime, Depth, SecondsPerPixel](const TArray<TraceServices::FTaskInfo::FRelationInfo>& Relations, int32 RelationIndex, ETaskEventType EventType) -> TSharedPtr<FTaskTrackEvent>
 			{
diff --git a/Engine/Source/Developer/TraceInsights/Private/Insights/TaskGraphProfiler/ViewModels/TaskTrackEvent.h b/Engine/Source/Developer/TraceInsights/Private/Insights/TaskGraphProfiler/ViewModels/TaskTrackEvent.h
--- a/Engine/Source/Developer/TraceInsights/Private/Insights/TaskGraphProfiler/ViewModels/TaskTrackEvent.h
+++ b/Engine/Source/Developer/TraceInsights/Private/Insights/TaskGraphProfiler/ViewModels/TaskTrackEvent.h
@@ -3,6 +3,7 @@
 #pragma once
 
 #include "CoreMinimal.h"
+#include "Async/TaskTrace.h"
 
 // Insights
 #include "Insights/TaskGraphProfiler/TaskGraphProfilerManager.h"
@@ -28,14 +29,14 @@ public:
 
 	ETaskEventType GetTaskEventType() const { return TaskEventType; }
 
-	uint32 GetTaskId() const { return TaskId; }
-	void SetTaskId(uint32 InTaskId) { TaskId = InTaskId; }
+	TaskTrace::FId GetTaskId() const { return TaskId; }
+	void SetTaskId(TaskTrace::FId InTaskId) { TaskId = InTaskId; }
 
 private:
 	ETaskEventType TaskEventType;
-	uint32 TaskId = ~0;
+	TaskTrace::FId TaskId = TaskTrace::InvalidId;
 };
 
 ////////////////////////////////////////////////////////////////////////////////////////////////////
 
-} // namespace Insights
\ No newline at end of file
+} // namespace Insights
diff --git a/Engine/Source/Developer/TraceInsights/Private/Insights/TaskGraphProfiler/Widgets/STaskTableTreeView.cpp b/Engine/Source/Developer/TraceInsights/Private/Insights/TaskGraphProfiler/Widgets/STaskTableTreeView.cpp
--- a/Engine/Source/Developer/TraceInsights/Private/Insights/TaskGraphProfiler/Widgets/STaskTableTreeView.cpp
+++ b/Engine/Source/Developer/TraceInsights/Private/Insights/TaskGraphProfiler/Widgets/STaskTableTreeView.cpp
@@ -192,7 +192,7 @@ void STaskTableTreeView::RebuildTree(bool bResync)
 					{
 						Tasks.Emplace(TaskInfo);
 						uint32 Index = Tasks.Num() - 1;
-						FName NodeName(BaseNodeName, TaskInfo.Id + 1);
+						FName NodeName(BaseNodeName, static_cast<int32>(TaskInfo.Id + 1));
 						FTaskNodePtr NodePtr = MakeShared<FTaskNode>(NodeName, TaskTable, Index);
 						Nodes->Add(NodePtr);
 						return TraceServices::ETaskEnumerationResult::Continue;
diff --git a/Engine/Source/Developer/TraceInsights/Private/Insights/TraceInsightsModule.cpp b/Engine/Source/Developer/TraceInsights/Private/Insights/TraceInsightsModule.cpp
--- a/Engine/Source/Developer/TraceInsights/Private/Insights/TraceInsightsModule.cpp
+++ b/Engine/Source/Developer/TraceInsights/Private/Insights/TraceInsightsModule.cpp
@@ -250,7 +250,9 @@ void FTraceInsightsModule::CreateSessionBrowser(const FCreateSessionBrowserParam
 	// Get desktop metrics. It also ensures the correct metrics will be used later in SWindow.
 	FDisplayMetrics DisplayMetrics;
 	FSlateApplication::Get().GetDisplayMetrics(DisplayMetrics);
-	const float DPIScaleFactor = FPlatformApplicationMisc::GetDPIScaleFactorAtPoint(DisplayMetrics.PrimaryDisplayWorkAreaRect.Left, DisplayMetrics.PrimaryDisplayWorkAreaRect.Top);
+	const float DPIScaleFactor = FPlatformApplicationMisc::GetDPIScaleFactorAtPoint(
+		static_cast<float>(DisplayMetrics.PrimaryDisplayWorkAreaRect.Left),
+		static_cast<float>(DisplayMetrics.PrimaryDisplayWorkAreaRect.Top));
 
 	const FVector2D ClientSize(960.0f * DPIScaleFactor, 640.0f * DPIScaleFactor);
 
@@ -342,7 +344,9 @@ void FTraceInsightsModule::CreateSessionViewer(bool bAllowDebugTools)
 	// Get desktop metrics. It also ensures the correct metrics will be used later in SWindow.
 	FDisplayMetrics DisplayMetrics;
 	FSlateApplication::Get().GetDisplayMetrics(DisplayMetrics);
-	const float DPIScaleFactor = FPlatformApplicationMisc::GetDPIScaleFactorAtPoint(DisplayMetrics.PrimaryDisplayWorkAreaRect.Left, DisplayMetrics.PrimaryDisplayWorkAreaRect.Top);
+	const float DPIScaleFactor = FPlatformApplicationMisc::GetDPIScaleFactorAtPoint(
+		static_cast<float>(DisplayMetrics.PrimaryDisplayWorkAreaRect.Left),
+		static_cast<float>(DisplayMetrics.PrimaryDisplayWorkAreaRect.Top));
 
 	const FVector2D ClientSize(1280.f * DPIScaleFactor, 720.0f * DPIScaleFactor);
 
diff --git a/Engine/Source/Developer/TraceInsights/Private/Insights/ViewModels/CounterAggregation.cpp b/Engine/Source/Developer/TraceInsights/Private/Insights/ViewModels/CounterAggregation.cpp
--- a/Engine/Source/Developer/TraceInsights/Private/Insights/ViewModels/CounterAggregation.cpp
+++ b/Engine/Source/Developer/TraceInsights/Private/Insights/ViewModels/CounterAggregation.cpp
@@ -253,7 +253,7 @@ void TCounterAggregationHelper<Type>::PostProcess(TAggregatedStatsEx<Type>& Stat
 
 		if (bComputeMedian)
 		{
-			const int32 HalfCount = StatsEx.Count / 2;
+			const int32 HalfCount = IntCastChecked<int32>(StatsEx.Count / 2);
 
 			// Compute median value.
 			int32 Count = 0;
diff --git a/Engine/Source/Developer/TraceInsights/Private/Insights/ViewModels/DrawHelpers.cpp b/Engine/Source/Developer/TraceInsights/Private/Insights/ViewModels/DrawHelpers.cpp
--- a/Engine/Source/Developer/TraceInsights/Private/Insights/ViewModels/DrawHelpers.cpp
+++ b/Engine/Source/Developer/TraceInsights/Private/Insights/ViewModels/DrawHelpers.cpp
@@ -327,14 +327,14 @@ void FDrawHelpers::DrawSelection(
 	{
 		const TSharedRef<FSlateFontMeasure> FontMeasureService = FSlateApplication::Get().GetRenderer()->GetFontMeasureService();
 		const float FontScale = DrawContext.Geometry.Scale;
-		const float TextWidth = FontMeasureService->Measure(Text, Font, FontScale).X / FontScale;
+		const float TextWidth = static_cast<float>(FontMeasureService->Measure(Text, Font, FontScale).X / FontScale);
 
 		const float CenterX = (SelectionX1 + SelectionX2) / 2.0f;
 
-		DrawContext.DrawBox(CenterX - TextWidth / 2 - 2.0, ArrowY - 6.0f, TextWidth + 4.0f, 13.0f, Brush, ArrowColor);
+		DrawContext.DrawBox(CenterX - TextWidth / 2.0f - 2.0f, ArrowY - 6.0f, TextWidth + 4.0f, 13.0f, Brush, ArrowColor);
 		DrawContext.LayerId++;
 
-		DrawContext.DrawText(CenterX - TextWidth / 2, ArrowY - 6.0f, Text, Font, FLinearColor::White);
+		DrawContext.DrawText(CenterX - TextWidth / 2.0f, ArrowY - 6.0f, Text, Font, FLinearColor::White);
 		DrawContext.LayerId++;
 	}
 }
diff --git a/Engine/Source/Developer/TraceInsights/Private/Insights/ViewModels/FrameTimingTrack.cpp b/Engine/Source/Developer/TraceInsights/Private/Insights/ViewModels/FrameTimingTrack.cpp
--- a/Engine/Source/Developer/TraceInsights/Private/Insights/ViewModels/FrameTimingTrack.cpp
+++ b/Engine/Source/Developer/TraceInsights/Private/Insights/ViewModels/FrameTimingTrack.cpp
@@ -355,13 +355,15 @@ void FFrameTimingTrack::DrawSelectedEventInfo(const FTimingEvent& SelectedEvent,
 		const TSharedRef<FSlateFontMeasure> FontMeasureService = FSlateApplication::Get().GetRenderer()->GetFontMeasureService();
 		const float FontScale = DrawContext.Geometry.Scale;
 		const FVector2D Size = FontMeasureService->Measure(Str, Font, FontScale) / FontScale;
-		const float X = Viewport.GetWidth() - Size.X - 23.0f;
-		const float Y = Viewport.GetPosY() + Viewport.GetHeight() - Size.Y - 18.0f;
+		const float W = static_cast<float>(Size.X);
+		const float H = static_cast<float>(Size.Y);
+		const float X = Viewport.GetWidth() - W - 23.0f;
+		const float Y = Viewport.GetPosY() + Viewport.GetHeight() - H - 18.0f;
 
 		const FLinearColor BackgroundColor(0.05f, 0.05f, 0.05f, 1.0f);
 		const FLinearColor TextColor(0.7f, 0.7f, 0.7f, 1.0f);
 
-		DrawContext.DrawBox(X - 8.0f, Y - 2.0f, Size.X + 16.0f, Size.Y + 4.0f, WhiteBrush, BackgroundColor);
+		DrawContext.DrawBox(X - 8.0f, Y - 2.0f, W + 16.0f, H + 4.0f, WhiteBrush, BackgroundColor);
 		DrawContext.LayerId++;
 
 		DrawContext.DrawText(X, Y, Str, Font, TextColor);
diff --git a/Engine/Source/Developer/TraceInsights/Private/Insights/ViewModels/FrameTrackHelper.cpp b/Engine/Source/Developer/TraceInsights/Private/Insights/ViewModels/FrameTrackHelper.cpp
--- a/Engine/Source/Developer/TraceInsights/Private/Insights/ViewModels/FrameTrackHelper.cpp
+++ b/Engine/Source/Developer/TraceInsights/Private/Insights/ViewModels/FrameTrackHelper.cpp
@@ -41,7 +41,7 @@ void FFrameTrackSeriesBuilder::AddFrame(const TraceServices::FFrame& Frame)
 {
 	NumAddedFrames++;
 
-	const int32 FrameIndex = Frame.Index;
+	const int32 FrameIndex = IntCastChecked<int32>(Frame.Index);
 
 	int32 SampleIndex = (FrameIndex - FirstFrameIndex) / FramesPerSample;
 	if (SampleIndex >= 0 && SampleIndex < NumSamples)
diff --git a/Engine/Source/Developer/TraceInsights/Private/Insights/ViewModels/GraphTrack.cpp b/Engine/Source/Developer/TraceInsights/Private/Insights/ViewModels/GraphTrack.cpp
--- a/Engine/Source/Developer/TraceInsights/Private/Insights/ViewModels/GraphTrack.cpp
+++ b/Engine/Source/Developer/TraceInsights/Private/Insights/ViewModels/GraphTrack.cpp
@@ -92,8 +92,8 @@ void FGraphTrack::PostUpdate(const ITimingTrackUpdateContext& Context)
 	constexpr float HeaderWidth = 100.0f;
 	constexpr float HeaderHeight = 14.0f;
 
-	const float MouseX = Context.GetMousePosition().X;
-	const float MouseY = Context.GetMousePosition().Y;
+	const float MouseX = static_cast<float>(Context.GetMousePosition().X);
+	const float MouseY = static_cast<float>(Context.GetMousePosition().Y);
 
 	if (MouseY >= GetPosY() && MouseY < GetPosY() + GetHeight())
 	{
@@ -149,13 +149,13 @@ void FGraphTrack::Draw(const ITimingTrackDrawContext& Context) const
 
 	// Set clipping area.
 	{
-		const FVector2D AbsPos = DrawContext.Geometry.GetAbsolutePosition();
+		const FVector2f AbsPos = FVector2f(DrawContext.Geometry.GetAbsolutePosition());
 		const float Scale = DrawContext.Geometry.GetAccumulatedLayoutTransform().GetScale();
 		const float L = AbsPos.X;
 		const float R = AbsPos.X + (Viewport.GetWidth() * Scale);
 		const float T = AbsPos.Y + (GetPosY() * Scale);
 		const float B = AbsPos.Y + ((GetPosY() + GetHeight()) * Scale);
-		const FSlateClippingZone ClipZone(FVector2D(L, T), FVector2D(R, T), FVector2D(L, B), FVector2D(R, B));
+		const FSlateClippingZone ClipZone(FSlateRect(L, T, R, B));
 		DrawContext.ElementList.PushClip(ClipZone);
 	}
 
@@ -209,13 +209,13 @@ void FGraphTrack::DrawSeries(const FGraphSeries& Series, FDrawContext& DrawConte
 	// of the fill and the outer lines appear to get rasterized differently, the latter missing one pixel
 	// on its upper side.
 	{
-		const FVector2D AbsPos = DrawContext.Geometry.GetAbsolutePosition();
+		const FVector2f AbsPos = FVector2f(DrawContext.Geometry.GetAbsolutePosition());
 		const float Scale = DrawContext.Geometry.GetAccumulatedLayoutTransform().GetScale();
 		const float L = AbsPos.X;
 		const float R = AbsPos.X + (Viewport.GetWidth() * Scale);
 		const float T = AbsPos.Y + ((GetPosY() + (GetBorderY() - 1.0f)) * Scale);
 		const float B = AbsPos.Y + ((GetPosY() + (GetHeight() - (GetBorderY() - 1.0f))) * Scale);
-		const FSlateClippingZone ClipZone(FVector2D(L, T), FVector2D(R, T), FVector2D(L, B), FVector2D(R, B));
+		const FSlateClippingZone ClipZone(FSlateRect(L, T, R, B));
 		DrawContext.ElementList.PushClip(ClipZone);
 	}
 
@@ -249,8 +249,8 @@ void FGraphTrack::DrawSeries(const FGraphSeries& Series, FDrawContext& DrawConte
 		FSlateResourceHandle ResourceHandle = FSlateApplication::Get().GetRenderer()->GetResourceHandle(*WhiteBrush);
 		const FSlateShaderResourceProxy* ResourceProxy = ResourceHandle.GetResourceProxy();
 
-		FVector2f AtlasOffset = ResourceProxy ? ResourceProxy->StartUV : FVector2f(0.f, 0.f);
-		FVector2f AtlasUVSize = ResourceProxy ? ResourceProxy->SizeUV : FVector2f(1.f, 1.f);
+		FVector2D AtlasOffset = ResourceProxy ? FVector2D(ResourceProxy->StartUV) : FVector2D(0.0, 0.0);
+		FVector2D AtlasUVSize = ResourceProxy ? FVector2D(ResourceProxy->SizeUV) : FVector2D(1.0, 1.0);
 
 		const FVector2D Size = DrawContext.Geometry.GetLocalSize();
 
@@ -268,10 +268,10 @@ void FGraphTrack::DrawSeries(const FGraphSeries& Series, FDrawContext& DrawConte
 			Indices.Reserve(LinePoints.Num() * 6);
 			Verts.Reserve(LinePoints.Num() * 2);
 
-			const float TopV = 0.0f;
-			const float BottomV = GetHeight() / Size.Y;
-			const float BaselineY = static_cast<float>(Series.GetBaselineY());
-			const float BaselineV = FMath::Clamp<float>(BaselineY / Size.Y, TopV, BottomV);
+			const double TopV = 0.0;
+			const double BottomV = GetHeight() / Size.Y;
+			const double BaselineY = Series.GetBaselineY();
+			const double BaselineV = FMath::Clamp<double>(BaselineY / Size.Y, TopV, BottomV);
 
 			int32 PrevSide = 0;
 
@@ -285,20 +285,20 @@ void FGraphTrack::DrawSeries(const FGraphSeries& Series, FDrawContext& DrawConte
 				{
 					// Compute intersection point.
 					const FVector2D& PrevLinePoint = LinePoints[PointIndex - 1];
-					const float X = PrevLinePoint.X + (LinePoint.X - PrevLinePoint.X) / ((BaselineY - LinePoint.Y) / (PrevLinePoint.Y - BaselineY) + 1.0f);
+					const double X = PrevLinePoint.X + (LinePoint.X - PrevLinePoint.X) / ((BaselineY - LinePoint.Y) / (PrevLinePoint.Y - BaselineY) + 1.0);
 
 					// Add an intersection point vertex.
-					FVector2f UV(X / Size.X, BaselineV);	// LWC_TODO: Precision loss
-					Verts.Add(FSlateVertex::Make<ESlateVertexRounding::Disabled>(RenderTransform, UV * FVector2f(Size), AtlasOffset + UV * AtlasUVSize, FillColor));	// LWC_TODO: Precision loss
+					FVector2D UV(X / Size.X, BaselineV);
+					Verts.Add(FSlateVertex::Make<ESlateVertexRounding::Disabled>(RenderTransform, FVector2f(UV * Size), FVector2f(AtlasOffset + UV * AtlasUVSize), FillColor));
 
 					// Add a value point vertex.
 					UV.X = LinePoint.X / Size.X;
 					UV.Y = TopV + LinePoint.Y / Size.Y;
-					Verts.Add(FSlateVertex::Make<ESlateVertexRounding::Disabled>(RenderTransform, UV * FVector2f(Size), AtlasOffset + UV * AtlasUVSize, FillColor));	// LWC_TODO: Precision loss
+					Verts.Add(FSlateVertex::Make<ESlateVertexRounding::Disabled>(RenderTransform, FVector2f(UV * Size), FVector2f(AtlasOffset + UV * AtlasUVSize), FillColor));
 
 					// Add a baseline vertex.
 					UV.Y = BaselineV;
-					Verts.Add(FSlateVertex::Make<ESlateVertexRounding::Disabled>(RenderTransform, UV * FVector2f(Size), AtlasOffset + FVector2f(UV.X, 0.5f) * AtlasUVSize, FillColor));	// LWC_TODO: Precision loss
+					Verts.Add(FSlateVertex::Make<ESlateVertexRounding::Disabled>(RenderTransform, FVector2f(UV * Size), FVector2f(AtlasOffset + FVector2D(UV.X, 0.5) * AtlasUVSize), FillColor));
 
 					int32 Index0 = Verts.Num() - 5;
 					int32 Index1 = Verts.Num() - 4;
@@ -317,12 +317,12 @@ void FGraphTrack::DrawSeries(const FGraphSeries& Series, FDrawContext& DrawConte
 				else
 				{
 					// Add a value point vertex.
-					FVector2f UV(LinePoint.X / Size.X, TopV + LinePoint.Y / Size.Y);	// LWC_TODO: Precision loss
-					Verts.Add(FSlateVertex::Make<ESlateVertexRounding::Disabled>(RenderTransform, UV * FVector2f(Size), AtlasOffset + UV * AtlasUVSize, FillColor));	// LWC_TODO: Precision loss
+					FVector2D UV(LinePoint.X / Size.X, TopV + LinePoint.Y / Size.Y);
+					Verts.Add(FSlateVertex::Make<ESlateVertexRounding::Disabled>(RenderTransform, FVector2f(UV * Size), FVector2f(AtlasOffset + UV * AtlasUVSize), FillColor));
 
 					// Add a baseline vertex.
 					UV.Y = BaselineV;
-					Verts.Add(FSlateVertex::Make<ESlateVertexRounding::Disabled>(RenderTransform, UV * FVector2f(Size), AtlasOffset + FVector2f(UV.X, 0.5f) * AtlasUVSize, FillColor));	// LWC_TODO: Precision loss
+					Verts.Add(FSlateVertex::Make<ESlateVertexRounding::Disabled>(RenderTransform, FVector2f(UV * Size), FVector2f(AtlasOffset + FVector2D(UV.X, 0.5) * AtlasUVSize), FillColor));
 
 					if (Verts.Num() >= 4)
 					{
@@ -396,8 +396,8 @@ void FGraphTrack::DrawSeries(const FGraphSeries& Series, FDrawContext& DrawConte
 			for (int32 Index = 0; Index < NumPoints; ++Index)
 			{
 				const FVector2D& Pt = Series.Points[Index];
-				const float PtX = Pt.X - PointVisualSize / 2.0f - 0.5f;
-				const float PtY = LocalPosY + Pt.Y - PointVisualSize / 2.0f - 0.5f;
+				const float PtX = static_cast<float>(Pt.X) - PointVisualSize / 2.0f - 0.5f;
+				const float PtY = LocalPosY + static_cast<float>(Pt.Y) - PointVisualSize / 2.0f - 0.5f;
 				DrawContext.DrawBox(PtX, PtY, PointVisualSize + 2.0f, PointVisualSize + 2.0f, PointBrush, Series.BorderColor);
 			}
 			DrawContext.LayerId++;
@@ -407,8 +407,8 @@ void FGraphTrack::DrawSeries(const FGraphSeries& Series, FDrawContext& DrawConte
 		for (int32 Index = 0; Index < NumPoints; ++Index)
 		{
 			const FVector2D& Pt = Series.Points[Index];
-			const float PtX = Pt.X - PointVisualSize / 2.0f + 0.5f;
-			const float PtY = LocalPosY + Pt.Y - PointVisualSize / 2.0f + 0.5f;
+			const float PtX = static_cast<float>(Pt.X) - PointVisualSize / 2.0f + 0.5f;
+			const float PtY = LocalPosY + static_cast<float>(Pt.Y) - PointVisualSize / 2.0f - 0.5f;
 			DrawContext.DrawBox(PtX, PtY, PointVisualSize, PointVisualSize, PointBrush, Series.Color);
 		}
 		DrawContext.LayerId++;
@@ -425,8 +425,8 @@ void FGraphTrack::DrawSeries(const FGraphSeries& Series, FDrawContext& DrawConte
 			for (int32 Index = 0; Index < NumPoints; ++Index)
 			{
 				const FVector2D& Pt = Series.Points[Index];
-				const float PtX = Pt.X - BorderPtSize / 2.0f + 0.5f;
-				const float PtY = LocalPosY + Pt.Y - BorderPtSize / 2.0f + 0.5f;
+				const float PtX = static_cast<float>(Pt.X) - BorderPtSize / 2.0f + 0.5f;
+				const float PtY = LocalPosY + static_cast<float>(Pt.Y) - BorderPtSize / 2.0f + 0.5f;
 				DrawContext.DrawBox(PtX, PtY, BorderPtSize, BorderPtSize, BorderBrush, Series.BorderColor);
 				//DrawContext.DrawRotatedBox(PtX, PtY, BorderPtSize, BorderPtSize, BorderBrush, Series.BorderColor, Angle, BorderRotationPoint);
 			}
@@ -439,8 +439,8 @@ void FGraphTrack::DrawSeries(const FGraphSeries& Series, FDrawContext& DrawConte
 		for (int32 Index = 0; Index < NumPoints; ++Index)
 		{
 			const FVector2D& Pt = Series.Points[Index];
-			const float PtX = Pt.X - PtSize / 2.0f + 0.5f;
-			const float PtY = LocalPosY + Pt.Y - PtSize / 2.0f + 0.5f;
+			const float PtX = static_cast<float>(Pt.X) - PtSize / 2.0f + 0.5f;
+			const float PtY = LocalPosY + static_cast<float>(Pt.Y) - PtSize / 2.0f + 0.5f;
 			DrawContext.DrawBox(PtX, PtY, PtSize, PtSize, WhiteBrush, Series.Color);
 			//DrawContext.DrawRotatedBox(PtX, PtY, PtSize, PtSize, WhiteBrush, Series.Color, Angle, RotationPoint);
 		}
@@ -461,14 +461,14 @@ void FGraphTrack::DrawSeries(const FGraphSeries& Series, FDrawContext& DrawConte
 
 		// Draw white corner at (0, 0) using MakeLines.
 		{
-			TArray<FVector2D> HLine;
-			HLine.Add(FVector2D(0.0f, 0.0f));
-			HLine.Add(FVector2D(10.0f, 0.0f));
+			TArray<FVector2f> HLine;
+			HLine.Add(FVector2f(0.0f, 0.0f));
+			HLine.Add(FVector2f(10.0f, 0.0f));
 			FSlateDrawElement::MakeLines(DrawContext.ElementList, DrawContext.LayerId, LineGeo, HLine, LineDrawEffects, FLinearColor::White, false, 1.0f);
 
-			TArray<FVector2D> VLine;
-			VLine.Add(FVector2D(0.0f, 0.0f));
-			VLine.Add(FVector2D(0.0f, 10.0f));
+			TArray<FVector2f> VLine;
+			VLine.Add(FVector2f(0.0f, 0.0f));
+			VLine.Add(FVector2f(0.0f, 10.0f));
 			FSlateDrawElement::MakeLines(DrawContext.ElementList, DrawContext.LayerId, LineGeo, VLine, LineDrawEffects, FLinearColor::White, false, 1.0f);
 		}
 		DrawContext.LayerId++;
@@ -484,17 +484,17 @@ void FGraphTrack::DrawSeries(const FGraphSeries& Series, FDrawContext& DrawConte
 		for (int32 Index = 0; Index < NumDbgPoints; ++Index)
 		{
 			const FVector2D& Pt = Series.Points[Index];
-			const float PtX = Pt.X;
-			const float PtY = Pt.Y;
+			const float PtX = static_cast<float>(Pt.X);
+			const float PtY = static_cast<float>(Pt.Y);
 
-			TArray<FVector2D> HLine;
-			HLine.Add(FVector2D(PtX - 7.0f, PtY));
-			HLine.Add(FVector2D(PtX + 8.0f, PtY));
+			TArray<FVector2f> HLine;
+			HLine.Add(FVector2f(PtX - 7.0f, PtY));
+			HLine.Add(FVector2f(PtX + 8.0f, PtY));
 			FSlateDrawElement::MakeLines(DrawContext.ElementList, DrawContext.LayerId, LineGeo, HLine, LineDrawEffects, FLinearColor::White, false, 1.0f);
 
-			TArray<FVector2D> VLine;
-			VLine.Add(FVector2D(PtX, PtY - 7.0f));
-			VLine.Add(FVector2D(PtX, PtY + 8.0f));
+			TArray<FVector2f> VLine;
+			VLine.Add(FVector2f(PtX, PtY - 7.0f));
+			VLine.Add(FVector2f(PtX, PtY + 8.0f));
 			FSlateDrawElement::MakeLines(DrawContext.ElementList, DrawContext.LayerId, LineGeo, VLine, LineDrawEffects, FLinearColor::White, false, 1.0f);
 		}
 		DrawContext.LayerId++;
@@ -503,8 +503,8 @@ void FGraphTrack::DrawSeries(const FGraphSeries& Series, FDrawContext& DrawConte
 		for (int32 Index = 0; Index < NumDbgPoints; ++Index)
 		{
 			const FVector2D& Pt = Series.Points[Index];
-			const float PtX = Pt.X;
-			const float PtY = LocalPosY + Pt.Y;
+			const float PtX = static_cast<float>(Pt.X);
+			const float PtY = LocalPosY + static_cast<float>(Pt.Y);
 			DrawContext.DrawBox(PtX - 5.0f, PtY, 11.0f, PixelUnit, WhiteBrush, Series.Color);
 			DrawContext.DrawBox(PtX, PtY - 5.0f, PixelUnit, 11.0f, WhiteBrush, Series.Color);
 		}
diff --git a/Engine/Source/Developer/TraceInsights/Private/Insights/ViewModels/GraphTrackBuilder.cpp b/Engine/Source/Developer/TraceInsights/Private/Insights/ViewModels/GraphTrackBuilder.cpp
--- a/Engine/Source/Developer/TraceInsights/Private/Insights/ViewModels/GraphTrackBuilder.cpp
+++ b/Engine/Source/Developer/TraceInsights/Private/Insights/ViewModels/GraphTrackBuilder.cpp
@@ -328,7 +328,7 @@ void FGraphTrackBuilder::AddBox(double Time, double Duration, double Value)
 		FGraphSeries::FBox& LastBox = Series.Boxes.Last();
 		if (LastBox.W == 1.0f && LastBox.X == X1)
 		{
-			const float RoundedBaselineY = FMath::RoundToFloat(Series.GetBaselineY());
+			const float RoundedBaselineY = static_cast<float>(FMath::RoundToFloat(Series.GetBaselineY()));
 
 			if (LastBox.Y < RoundedBaselineY)
 			{
diff --git a/Engine/Source/Developer/TraceInsights/Private/Insights/ViewModels/MarkersTimingTrack.cpp b/Engine/Source/Developer/TraceInsights/Private/Insights/ViewModels/MarkersTimingTrack.cpp
--- a/Engine/Source/Developer/TraceInsights/Private/Insights/ViewModels/MarkersTimingTrack.cpp
+++ b/Engine/Source/Developer/TraceInsights/Private/Insights/ViewModels/MarkersTimingTrack.cpp
@@ -140,7 +140,7 @@ void FMarkersTimingTrack::Update(const ITimingTrackUpdateContext& Context)
 
 void FMarkersTimingTrack::PostUpdate(const ITimingTrackUpdateContext& Context)
 {
-	const float MouseY = Context.GetMousePosition().Y;
+	const float MouseY = static_cast<float>(Context.GetMousePosition().Y);
 	SetHoveredState(MouseY >= GetPosY() && MouseY < GetPosY() + GetHeight());
 
 	Header.PostUpdate(Context);
@@ -430,11 +430,7 @@ const TSharedPtr<const ITimingEvent> FMarkersTimingTrack::GetEvent(float InPosX,
 	if (DY >= 0 && DY < GetHeight())
 	{
 		const int32 NumBoxes = TimeMarkerTexts.Num();
-		int32 FoundIndex = Algo::LowerBoundBy(TimeMarkerTexts, InPosX, [](const FTimeMarkerTextInfo& Text)
-		{
-			return Text.X;
-		});
-
+		int32 FoundIndex = Algo::LowerBoundBy(TimeMarkerTexts, InPosX, [](const FTimeMarkerTextInfo& Text) { return Text.X; });
 		if (FoundIndex > 0)
 		{
 			--FoundIndex;
@@ -451,7 +447,7 @@ const TSharedPtr<const ITimingEvent> FMarkersTimingTrack::GetEvent(float InPosX,
 			Width = TimeMarkerTexts[FoundIndex + 1].X;
 		}
 
-		uint16 ScreenshotId = std::numeric_limits<uint16>::max();
+		uint32 ScreenshotId = TraceServices::FScreenshot::InvalidScreenshotId;
 
 		const FTimeMarkerTextInfo& Text = TimeMarkerTexts[FoundIndex];
 
@@ -468,15 +464,17 @@ const TSharedPtr<const ITimingEvent> FMarkersTimingTrack::GetEvent(float InPosX,
 			{
 				if (Message.Category == this->ScreenshotCategory)
 				{
+					check(Message.Line >= 0);
 					ScreenshotId = Message.Line;
 				}
 			});
 
-		if (ScreenshotId == std::numeric_limits<uint16>::max())
+		if (ScreenshotId == TraceServices::FScreenshot::InvalidScreenshotId)
 		{
 			return TimingEvent;
 		}
 
+		//TODO: make a custom FScreenshotEvent
 		TimingEvent = MakeShared<FTimingEvent>(SharedThis(this), Viewport.SlateUnitsToTime(Text.X), Viewport.SlateUnitsToTime(Text.X + Width), 0, ScreenshotId);
 	}
 
diff --git a/Engine/Source/Developer/TraceInsights/Private/Insights/ViewModels/ThreadTimingTrack.cpp b/Engine/Source/Developer/TraceInsights/Private/Insights/ViewModels/ThreadTimingTrack.cpp
--- a/Engine/Source/Developer/TraceInsights/Private/Insights/ViewModels/ThreadTimingTrack.cpp
+++ b/Engine/Source/Developer/TraceInsights/Private/Insights/ViewModels/ThreadTimingTrack.cpp
@@ -99,7 +99,7 @@ static void AppendMetadataToTooltip(FTooltipDrawState& Tooltip, TArrayView<const
 			break;
 		}
 
-		FString Key(Context.AsLength(), Context.AsCString());
+		FString Key(static_cast<int32>(Context.AsLength()), Context.AsCString());
 		Key += TEXT(":");
 
 		if (!CborReader.ReadNext(Context))
@@ -135,7 +135,7 @@ static void AppendMetadataToTooltip(FTooltipDrawState& Tooltip, TArrayView<const
 
 		case ECborCode::ByteString:
 			{
-				FAnsiStringView Value(Context.AsCString(), Context.AsLength());
+				FAnsiStringView Value(Context.AsCString(), static_cast<int32>(Context.AsLength()));
 				FString ValueStr(Value);
 				Tooltip.AddNameValueTextLine(Key, ValueStr);
 				continue;
@@ -244,7 +244,7 @@ static void AppendMetadataToString(FString& Str, TArrayView<const uint8>& Metada
 
 		case ECborCode::ByteString:
 			{
-				Str.AppendChars(Context.AsCString(), Context.AsLength());
+				Str.AppendChars(Context.AsCString(), static_cast<int32>(Context.AsLength()));
 				continue;
 			}
 		}
@@ -1397,7 +1397,7 @@ const TSharedPtr<const ITimingEvent> FThreadTimingTrack::GetEvent(float InPosX,
 	// If mouse is not above first sub-track or below last sub-track...
 	if (DY >= 0 && DY < TrackLanesHeight)
 	{
-		const int32 Depth = DY / (Layout.EventH + Layout.EventDY);
+		const int32 Depth = FMath::RoundToInt(DY / (Layout.EventH + Layout.EventDY));
 
 		const double SecondsPerPixel = 1.0 / Viewport.GetScaleX();
 
diff --git a/Engine/Source/Developer/TraceInsights/Private/Insights/ViewModels/TimeRulerTrack.cpp b/Engine/Source/Developer/TraceInsights/Private/Insights/ViewModels/TimeRulerTrack.cpp
--- a/Engine/Source/Developer/TraceInsights/Private/Insights/ViewModels/TimeRulerTrack.cpp
+++ b/Engine/Source/Developer/TraceInsights/Private/Insights/ViewModels/TimeRulerTrack.cpp
@@ -110,15 +110,17 @@ TSharedPtr<Insights::FTimeMarker> FTimeRulerTrack::GetTimeMarkerAtPos(const FVec
 	TSharedPtr<Insights::FTimeMarker> ClosestTimeMarker = nullptr;
 
 	constexpr float TimeMarkerBoxHeight = 12.0f;
-	if (InPosition.Y >= GetPosY() && InPosition.Y < GetPosY() + TimeMarkerBoxHeight)
+	const float InPositionY = static_cast<float>(InPosition.Y);
+	if (InPositionY >= GetPosY() && InPositionY < GetPosY() + TimeMarkerBoxHeight)
 	{
+		const float InPositionX = static_cast<float>(InPosition.X);
 		float MinDX = 42.0f;
 		for (TSharedRef<Insights::FTimeMarker>& TimeMarker : TimeMarkers)
 		{
 			if (TimeMarker->IsVisible())
 			{
 				const float MarkerX = InViewport.TimeToSlateUnitsRounded(TimeMarker->GetTime());
-				const float DX = FMath::Abs(InPosition.X - MarkerX);
+				const float DX = FMath::Abs(InPositionX - MarkerX);
 				if (DX <= MinDX)
 				{
 					MinDX = DX;
@@ -159,11 +161,11 @@ void FTimeRulerTrack::PostUpdate(const ITimingTrackUpdateContext& Context)
 	//constexpr float HeaderWidth = 100.0f;
 	//constexpr float HeaderHeight = 14.0f;
 
-	const float MouseY = Context.GetMousePosition().Y;
+	const float MouseY = static_cast<float>(Context.GetMousePosition().Y);
 	if (MouseY >= GetPosY() && MouseY < GetPosY() + GetHeight())
 	{
 		SetHoveredState(true);
-		//const float MouseX = Context.GetMousePosition().X;
+		//const float MouseX = static_cast<float>(Context.GetMousePosition().X);
 		//SetHeaderHoveredState(MouseX < HeaderWidth && MouseY < GetPosY() + HeaderHeight);
 	}
 	else
@@ -225,7 +227,7 @@ void FTimeRulerTrack::Draw(const ITimingTrackDrawContext& Context) const
 	{
 		const double T = Viewport.SlateUnitsToTime(X);
 		FString Text = TimeUtils::FormatTime(T, Precision);
-		const float TextWidth = FontMeasureService->Measure(Text, Font, FontScale).X / FontScale;
+		const float TextWidth = static_cast<float>(FontMeasureService->Measure(Text, Font, FontScale).X / FontScale);
 		DrawContext.DrawText(X - TextWidth / 2, TextY, Text, Font,
 			(T < Viewport.GetMinValidTime() || T >= Viewport.GetMaxValidTime()) ? FLinearColor(0.7f, 0.5f, 0.5f, 1.0f) : FLinearColor(0.8f, 0.8f, 0.8f, 1.0f));
 	}
@@ -251,7 +253,7 @@ void FTimeRulerTrack::PostDraw(const ITimingTrackDrawContext& Context) const
 
 		// Time at current mouse position.
 		FString MousePosText;
-		const double MousePosTime = Viewport.SlateUnitsToTime(MousePosition.X);
+		const double MousePosTime = Viewport.SlateUnitsToTime(static_cast<float>(MousePosition.X));
 		CrtMousePosTime = MousePosTime;
 
 		const double DT = 100.0 / Viewport.GetScaleX();
@@ -270,7 +272,7 @@ void FTimeRulerTrack::PostDraw(const ITimingTrackDrawContext& Context) const
 		const TSharedRef<FSlateFontMeasure> FontMeasureService = FSlateApplication::Get().GetRenderer()->GetFontMeasureService();
 		const float FontScale = DrawContext.Geometry.Scale;
 
-		const float MousePosTextWidth = FMath::RoundToFloat(FontMeasureService->Measure(MousePosText, Font, FontScale).X / FontScale);
+		const float MousePosTextWidth = FMath::RoundToFloat(static_cast<float>(FontMeasureService->Measure(MousePosText, Font, FontScale).X / FontScale));
 
 		if (!FMath::IsNearlyEqual(CrtMousePosTextWidth, MousePosTextWidth))
 		{
@@ -280,7 +282,7 @@ void FTimeRulerTrack::PostDraw(const ITimingTrackDrawContext& Context) const
 
 		const float TextY = GetPosY() + 11.0f;
 
-		float X = MousePosition.X;
+		float X = static_cast<float>(MousePosition.X);
 		float W = CrtMousePosTextWidth + 4.0f;
 		if (bIsSelecting && SelectionStartTime < SelectionEndTime)
 		{
@@ -289,32 +291,32 @@ void FTimeRulerTrack::PostDraw(const ITimingTrackDrawContext& Context) const
 			float SelectionX2 = Viewport.TimeToSlateUnitsRounded(SelectionEndTime);
 			if (FMath::Abs(X - SelectionX1) > FMath::Abs(SelectionX2 - X))
 			{
-				X = SelectionX2 + W / 2;
+				X = SelectionX2 + W / 2.0f;
 			}
 			else
 			{
-				X = SelectionX1 - W / 2;
+				X = SelectionX1 - W / 2.0f;
 			}
 			MousePosTextForegroundColor = FLinearColor(0.01f, 0.05f, 0.2f, 1.0f);
 		}
 		else
 		{
 			// Draw horizontal line at mouse position.
-			//DrawContext.DrawBox(0.0f, MousePosition.Y, Viewport.Width, 1.0f, WhiteBrush, MousePosLineColor);
+			//DrawContext.DrawBox(0.0f, static_cast<float>(MousePosition.Y), Viewport.Width, 1.0f, WhiteBrush, MousePosLineColor);
 
 			// Draw vertical line at mouse position.
-			DrawContext.DrawBox(MousePosition.X, Viewport.GetPosY(), 1.0f, Viewport.GetHeight(), WhiteBrush, MousePosLineColor);
+			DrawContext.DrawBox(static_cast<float>(MousePosition.X), Viewport.GetPosY(), 1.0f, Viewport.GetHeight(), WhiteBrush, MousePosLineColor);
 
 			// Stroke the vertical line above current time box.
-			DrawContext.DrawBox(MousePosition.X, GetPosY(), 1.0f, TextY - GetPosY(), WhiteBrush, MousePosTextBackgroundColor);
+			DrawContext.DrawBox(static_cast<float>(MousePosition.X), GetPosY(), 1.0f, TextY - GetPosY(), WhiteBrush, MousePosTextBackgroundColor);
 		}
 
 		// Fill the current time box.
-		DrawContext.DrawBox(X - W / 2, TextY, W, 12.0f, WhiteBrush, MousePosTextBackgroundColor);
+		DrawContext.DrawBox(X - W / 2.0f, TextY, W, 12.0f, WhiteBrush, MousePosTextBackgroundColor);
 		DrawContext.LayerId++;
 
 		// Draw current time text.
-		DrawContext.DrawText(X - MousePosTextWidth / 2, TextY, MousePosText, Font, MousePosTextForegroundColor);
+		DrawContext.DrawText(X - MousePosTextWidth / 2.0f, TextY, MousePosText, Font, MousePosTextForegroundColor);
 		DrawContext.LayerId++;
 	}
 
@@ -404,7 +406,7 @@ void FTimeRulerTrack::DrawTimeMarker(const ITimingTrackDrawContext& Context, con
 	const TSharedRef<FSlateFontMeasure> FontMeasureService = FSlateApplication::Get().GetRenderer()->GetFontMeasureService();
 	const float FontScale = DrawContext.Geometry.Scale;
 
-	const float TextWidth = FMath::RoundToFloat(FontMeasureService->Measure(TimeMarkerText, Font, FontScale).X / FontScale);
+	const float TextWidth = FMath::RoundToFloat(static_cast<float>(FontMeasureService->Measure(TimeMarkerText, Font, FontScale).X / FontScale));
 
 	if (!FMath::IsNearlyEqual(TimeMarker.GetCrtTextWidth(), TextWidth))
 	{
diff --git a/Engine/Source/Developer/TraceInsights/Private/Insights/ViewModels/TimingEventsTrack.cpp b/Engine/Source/Developer/TraceInsights/Private/Insights/ViewModels/TimingEventsTrack.cpp
--- a/Engine/Source/Developer/TraceInsights/Private/Insights/ViewModels/TimingEventsTrack.cpp
+++ b/Engine/Source/Developer/TraceInsights/Private/Insights/ViewModels/TimingEventsTrack.cpp
@@ -227,12 +227,12 @@ void FTimingEventsTrack::PostUpdate(const ITimingTrackUpdateContext& Context)
 	constexpr float HeaderWidth = 100.0f;
 	constexpr float HeaderHeight = 14.0f;
 
-	const float MouseY = Context.GetMousePosition().Y;
+	const float MouseY = static_cast<float>(Context.GetMousePosition().Y);
 	if (MouseY >= GetPosY() && MouseY < GetPosY() + GetHeight())
 	{
 		SetHoveredState(true);
 
-		const float MouseX = Context.GetMousePosition().X;
+		const float MouseX = static_cast<float>(Context.GetMousePosition().X);
 		SetHeaderHoveredState(MouseX < HeaderWidth && MouseY < GetPosY() + HeaderHeight);
 	}
 	else
@@ -402,7 +402,7 @@ const TSharedPtr<const ITimingEvent> FTimingEventsTrack::GetEvent(float InPosX,
 	// If mouse is not above first sub-track or below last sub-track...
 	if (DY >= 0 && DY < TrackLanesHeight)
 	{
-		const int32 Depth = DY / (Layout.EventH + Layout.EventDY);
+		const int32 Depth = FMath::RoundToInt(DY / (Layout.EventH + Layout.EventDY));
 		const double SecondsPerPixel = 1.0 / Viewport.GetScaleX();
 		const double TimeAtPosX = Viewport.SlateUnitsToTime(InPosX);
 
@@ -465,13 +465,15 @@ void FTimingEventsTrack::DrawSelectedEventInfo(const FString& InText, const FTim
 	const TSharedRef<FSlateFontMeasure> FontMeasureService = FSlateApplication::Get().GetRenderer()->GetFontMeasureService();
 	const float FontScale = DrawContext.Geometry.Scale;
 	const FVector2D Size = FontMeasureService->Measure(InText, Font, FontScale) / FontScale;
-	const float X = Viewport.GetWidth() - Size.X - 23.0f;
-	const float Y = Viewport.GetPosY() + Viewport.GetHeight() - Size.Y - 18.0f;
+	const float W = static_cast<float>(Size.X);
+	const float H = static_cast<float>(Size.Y);
+	const float X = Viewport.GetWidth() - W - 23.0f;
+	const float Y = Viewport.GetPosY() + Viewport.GetHeight() - H - 18.0f;
 
 	const FLinearColor BackgroundColor(0.05f, 0.05f, 0.05f, 1.0f);
 	const FLinearColor TextColor(0.7f, 0.7f, 0.7f, 1.0f);
 
-	DrawContext.DrawBox(X - 8.0f, Y - 2.0f, Size.X + 16.0f, Size.Y + 4.0f, WhiteBrush, BackgroundColor);
+	DrawContext.DrawBox(X - 8.0f, Y - 2.0f, W + 16.0f, H + 4.0f, WhiteBrush, BackgroundColor);
 	DrawContext.LayerId++;
 
 	DrawContext.DrawText(X, Y, InText, Font, TextColor);
diff --git a/Engine/Source/Developer/TraceInsights/Private/Insights/ViewModels/TimingGraphTrack.cpp b/Engine/Source/Developer/TraceInsights/Private/Insights/ViewModels/TimingGraphTrack.cpp
--- a/Engine/Source/Developer/TraceInsights/Private/Insights/ViewModels/TimingGraphTrack.cpp
+++ b/Engine/Source/Developer/TraceInsights/Private/Insights/ViewModels/TimingGraphTrack.cpp
@@ -573,7 +573,7 @@ void FTimingGraphTrack::DrawVerticalAxisGrid(const ITimingTrackDrawContext& Cont
 	ViewportY.SetSize(GetHeight());
 	ViewportY.SetScaleLimits(std::numeric_limits<double>::min(), std::numeric_limits<double>::max());
 	ViewportY.SetScale(SharedValueViewport.GetScaleY());
-	ViewportY.ScrollAtPos(SharedValueViewport.GetBaselineY() - GetHeight());
+	ViewportY.ScrollAtPos(static_cast<float>(SharedValueViewport.GetBaselineY()) - GetHeight());
 
 	const float ViewWidth = Context.GetViewport().GetWidth();
 	const float RoundedViewHeight = FMath::RoundToFloat(GetHeight());
@@ -617,7 +617,7 @@ void FTimingGraphTrack::DrawVerticalAxisGrid(const ITimingTrackDrawContext& Cont
 			86400.0	// 1d
 		};
 		constexpr int32 NumThresholds = sizeof(Thresholds) / sizeof(double);
-		int32 Index = Algo::LowerBound(Thresholds, Delta);
+		int32 Index = static_cast<int32>(Algo::LowerBound(Thresholds, Delta));
 		if (Index > 0)
 		{
 			Index--;
@@ -654,11 +654,11 @@ void FTimingGraphTrack::DrawVerticalAxisGrid(const ITimingTrackDrawContext& Cont
 			DrawContext.DrawBox(0, Y, ViewWidth, 1, Brush, GridColor);
 
 			const FVector2D LabelTextSize = FontMeasureService->Measure(LabelText, Font, FontScale) / FontScale;
-			const float LabelX = X0 - LabelTextSize.X - 4.0f;
+			const float LabelX = X0 - static_cast<float>(LabelTextSize.X) - 4.0f;
 			const float LabelY = FMath::Min(Y0 + GetHeight() - TextH, FMath::Max(Y0, Y - TextH / 2));
 
 			// Draw background for value text.
-			DrawContext.DrawBox(LabelX, LabelY, LabelTextSize.X + 4.0f, TextH, Brush, TextBgColor);
+			DrawContext.DrawBox(LabelX, LabelY, static_cast<float>(LabelTextSize.X) + 4.0f, TextH, Brush, TextBgColor);
 
 			// Draw value text.
 			DrawContext.DrawText(LabelX + 2.0f, LabelY + 1.0f, LabelText, Font, TextColor);
diff --git a/Engine/Source/Developer/TraceInsights/Private/Insights/ViewModels/TimingTrackViewport.cpp b/Engine/Source/Developer/TraceInsights/Private/Insights/ViewModels/TimingTrackViewport.cpp
--- a/Engine/Source/Developer/TraceInsights/Private/Insights/ViewModels/TimingTrackViewport.cpp
+++ b/Engine/Source/Developer/TraceInsights/Private/Insights/ViewModels/TimingTrackViewport.cpp
@@ -224,7 +224,7 @@ bool FTimingTrackViewport::OnUserScrolled(TSharedPtr<SScrollBar> ScrollBar, floa
 {
 	const double S = 1.0 / (MaxValidTime - MinValidTime);
 	const double Page = EndTime - StartTime;
-	const float ThumbSizeFraction = FMath::Clamp<float>(Page * S, 0.0f, 1.0f);
+	const float ThumbSizeFraction = FMath::Clamp<float>(static_cast<float>(Page * S), 0.0f, 1.0f);
 	const float OffsetFraction = FMath::Clamp<float>(ScrollOffset, 0.0f, 1.0f - ThumbSizeFraction);
 
 	ScrollBar->SetState(OffsetFraction, ThumbSizeFraction);
@@ -245,7 +245,7 @@ void FTimingTrackViewport::UpdateScrollBar(TSharedPtr<SScrollBar> ScrollBar) con
 {
 	const double S = 1.0 / (MaxValidTime - MinValidTime);
 	const double Page = EndTime - StartTime;
-	const float ThumbSizeFraction = FMath::Clamp<float>(Page * S, 0.0f, 1.0f);
+	const float ThumbSizeFraction = FMath::Clamp<float>(static_cast<float>(Page * S), 0.0f, 1.0f);
 	const float ScrollOffset = static_cast<float>((StartTime - MinValidTime) * S);
 	const float OffsetFraction = FMath::Clamp<float>(ScrollOffset, 0.0f, 1.0f - ThumbSizeFraction);
 
@@ -258,7 +258,7 @@ bool FTimingTrackViewport::OnUserScrolledY(TSharedPtr<SScrollBar> ScrollBar, flo
 {
 	const float S = 1.0f / ScrollHeight;
 	const float Page = Height - TopOffset - BottomOffset;
-	const float ThumbSizeFraction = FMath::Clamp<float>(Page * S, 0.0f, 1.0f);
+	const float ThumbSizeFraction = FMath::Clamp<float>(static_cast<float>(Page * S), 0.0f, 1.0f);
 	const float OffsetFraction = FMath::Clamp<float>(ScrollOffset, 0.0f, 1.0f - ThumbSizeFraction);
 
 	ScrollBar->SetState(OffsetFraction, ThumbSizeFraction);
diff --git a/Engine/Source/Developer/TraceInsights/Private/Insights/ViewModels/TimingViewDrawHelper.cpp b/Engine/Source/Developer/TraceInsights/Private/Insights/ViewModels/TimingViewDrawHelper.cpp
--- a/Engine/Source/Developer/TraceInsights/Private/Insights/ViewModels/TimingViewDrawHelper.cpp
+++ b/Engine/Source/Developer/TraceInsights/Private/Insights/ViewModels/TimingViewDrawHelper.cpp
@@ -705,7 +705,7 @@ void FTimingViewDrawHelper::DrawTrackHeader(const FBaseTimingTrack& Track, const
 		{
 			const TSharedRef<FSlateFontMeasure> FontMeasureService = FSlateApplication::Get().GetRenderer()->GetFontMeasureService();
 			const float FontScale = DrawContext.Geometry.Scale;
-			float TextWidth = FontMeasureService->Measure(Track.GetName(), EventFont, FontScale).X / FontScale;
+			float TextWidth = static_cast<float>(FontMeasureService->Measure(Track.GetName(), EventFont, FontScale).X / FontScale);
 
 			constexpr float PinWidth = 8.0f;
 			if (Track.IsSelected())
@@ -858,10 +858,10 @@ void FTimingViewDrawHelper::DrawTimingEventHighlight(double StartTime, double En
 	{
 		// Animate color from white (if selected and hovered) or yellow (if only selected) to black, using a squared sine function.
 		const double Time = static_cast<double>(FPlatformTime::Cycles64()) * FPlatformTime::GetSecondsPerCycle64();
-		float S = FMath::Sin(2.0 * Time);
-		S = S * S; // squared, to ensure only positive [0 - 1] values
-		const float Blue = (Mode == EDrawEventMode::SelectedAndHovered) ? 0.0f : S;
-		const FLinearColor Color(S, S, Blue, 1.0f);
+		float Hue = static_cast<float>(FMath::Sin(2.0 * Time));
+		Hue = Hue * Hue; // squared, to ensure only positive [0 - 1] values
+		const float Blue = (Mode == EDrawEventMode::SelectedAndHovered) ? 0.0f : Hue;
+		const FLinearColor Color(Hue, Hue, Blue, 1.0f);
 
 		// Draw border around the timing event box.
 #if INSIGHTS_USE_LEGACY_BORDER
diff --git a/Engine/Source/Developer/TraceInsights/Private/Insights/ViewModels/TooltipDrawState.cpp b/Engine/Source/Developer/TraceInsights/Private/Insights/ViewModels/TooltipDrawState.cpp
--- a/Engine/Source/Developer/TraceInsights/Private/Insights/ViewModels/TooltipDrawState.cpp
+++ b/Engine/Source/Developer/TraceInsights/Private/Insights/ViewModels/TooltipDrawState.cpp
@@ -137,9 +137,10 @@ void FTooltipDrawState::UpdateLayout()
 	{
 		if (TextInfo.Type == FDrawTextType::Name)
 		{
-			if (ValueOffsetX < TextInfo.TextSize.X)
+			const float TextW = static_cast<float>(TextInfo.TextSize.X);
+			if (ValueOffsetX < TextW)
 			{
-				ValueOffsetX = TextInfo.TextSize.X;
+				ValueOffsetX = TextW;
 			}
 		}
 	}
@@ -156,11 +157,11 @@ void FTooltipDrawState::UpdateLayout()
 			break;
 
 		case FDrawTextType::Value:
-			RightX = ValueOffsetX + NameValueDX + TextInfo.TextSize.X;
+			RightX = ValueOffsetX + NameValueDX + static_cast<float>(TextInfo.TextSize.X);
 			break;
 
 		default:
-			RightX = TextInfo.X + TextInfo.TextSize.X;
+			RightX = TextInfo.X + static_cast<float>(TextInfo.TextSize.X);
 		}
 		if (DesiredSize.X < RightX)
 		{
@@ -176,7 +177,7 @@ void FTooltipDrawState::Update()
 {
 	if (Size.X != DesiredSize.X)
 	{
-		Size.X = Size.X * 0.75f + DesiredSize.X * 0.25f;
+		Size.X = Size.X * 0.75 + DesiredSize.X * 0.25;
 
 		if (FMath::IsNearlyEqual(Size.X, DesiredSize.X))
 		{
@@ -186,7 +187,7 @@ void FTooltipDrawState::Update()
 
 	if (Size.Y != DesiredSize.Y)
 	{
-		Size.Y = Size.Y * 0.5f + DesiredSize.Y * 0.5f;
+		Size.Y = Size.Y * 0.5 + DesiredSize.Y * 0.5;
 
 		if (FMath::IsNearlyEqual(Size.Y, DesiredSize.Y))
 		{
@@ -195,9 +196,9 @@ void FTooltipDrawState::Update()
 	}
 
 	float RealDesiredOpacity;
-	if (DesiredSize.X > 1.0f)
+	if (DesiredSize.X > 1.0)
 	{
-		const float DesiredOpacityByTooltipWidth = 1.0f - FMath::Abs(Size.X - DesiredSize.X) / DesiredSize.X;
+		const float DesiredOpacityByTooltipWidth = static_cast<float>(1.0 - FMath::Abs(Size.X - DesiredSize.X) / DesiredSize.X);
 
 		if (FMath::IsNearlyEqual(DesiredOpacity, DesiredOpacityByTooltipWidth, 0.001f))
 		{
@@ -237,31 +238,31 @@ void FTooltipDrawState::Update()
 
 void FTooltipDrawState::SetPosition(const FVector2D& MousePosition, const float MinX, const float MaxX, const float MinY, const float MaxY)
 {
-	Position.X = FMath::Max(MinX, FMath::Min(MousePosition.X + 12.0f, MaxX - Size.X));
-	Position.Y = FMath::Max(MinY, FMath::Min(MousePosition.Y + 15.0f, MaxY - Size.Y));
+	Position.X = FMath::Max(MinX, FMath::Min(MousePosition.X + 12.0, MaxX - Size.X));
+	Position.Y = FMath::Max(MinY, FMath::Min(MousePosition.Y + 15.0, MaxY - Size.Y));
 }
 
 ////////////////////////////////////////////////////////////////////////////////////////////////////
 
 void FTooltipDrawState::Draw(const FDrawContext& DrawContext) const
 {
-	if (Opacity > 0.0f && Size.X > 0.0f && Size.Y > 0.0f)
+	if (Opacity > 0.0f && Size.X > 0.0 && Size.Y > 0.0)
 	{
 		// Draw background.
-		DrawContext.DrawBox(Position.X, Position.Y, Size.X, Size.Y, WhiteBrush, BackgroundColor.CopyWithNewOpacity(Opacity));
+		DrawContext.DrawBox(static_cast<float>(Position.X), static_cast<float>(Position.Y), static_cast<float>(Size.X), static_cast<float>(Size.Y), WhiteBrush, BackgroundColor.CopyWithNewOpacity(Opacity));
 		if (Size.X < DesiredSize.X)
 		{
-			DrawContext.DrawBox(Position.X + Size.X, Position.Y, DesiredSize.X - Size.X, Size.Y, WhiteBrush, BackgroundColor.CopyWithNewOpacity(Opacity * 0.5f));
+			DrawContext.DrawBox(static_cast<float>(Position.X + Size.X), static_cast<float>(Position.Y), static_cast<float>(DesiredSize.X - Size.X), static_cast<float>(Size.Y), WhiteBrush, BackgroundColor.CopyWithNewOpacity(Opacity * 0.5f));
 		}
 		DrawContext.LayerId++;
 
 		// Draw border.
-		//DrawContext.DrawBox(Position.X, Position.Y, Size.X, Size.Y, BorderBrush, BorderColor.CopyWithNewOpacity(Opacity));
+		//DrawContext.DrawBox(static_cast<float>(Position.X), static_cast<float>(Position.Y), static_cast<float>(Size.X), static_cast<float>(Size.Y), BorderBrush, BorderColor.CopyWithNewOpacity(Opacity));
 		//DrawContext.LayerId++;
 
 		if (ImageBrush)
 		{
-			DrawContext.DrawBox(Position.X, Position.Y, DesiredSize.X, DesiredSize.Y, ImageBrush.Get(), FLinearColor::White);
+			DrawContext.DrawBox(static_cast<float>(Position.X), static_cast<float>(Position.Y), static_cast<float>(DesiredSize.X), static_cast<float>(DesiredSize.Y), ImageBrush.Get(), FLinearColor::White);
 			DrawContext.LayerId++;
 			return;
 		}
@@ -269,11 +270,11 @@ void FTooltipDrawState::Draw(const FDrawContext& DrawContext) const
 		// Draw cached texts.
 		for (const FDrawTextInfo& TextInfo : Texts)
 		{
-			float X = Position.X;
+			float X = static_cast<float>(Position.X);
 			switch (TextInfo.Type)
 			{
 				case FDrawTextType::Name:
-					X += ValueOffsetX - TextInfo.TextSize.X;
+					X += ValueOffsetX - static_cast<float>(TextInfo.TextSize.X);
 					break;
 
 				case FDrawTextType::Value:
@@ -283,7 +284,7 @@ void FTooltipDrawState::Draw(const FDrawContext& DrawContext) const
 				default:
 					X += TextInfo.X;
 			}
-			DrawContext.DrawText(X, Position.Y + TextInfo.Y, TextInfo.Text, Font, TextInfo.Color.CopyWithNewOpacity(Opacity));
+			DrawContext.DrawText(X, static_cast<float>(Position.Y) + TextInfo.Y, TextInfo.Text, Font, TextInfo.Color.CopyWithNewOpacity(Opacity));
 		}
 		DrawContext.LayerId++;
 	}
diff --git a/Engine/Source/Developer/TraceInsights/Private/Insights/ViewModels/TrackHeader.cpp b/Engine/Source/Developer/TraceInsights/Private/Insights/ViewModels/TrackHeader.cpp
--- a/Engine/Source/Developer/TraceInsights/Private/Insights/ViewModels/TrackHeader.cpp
+++ b/Engine/Source/Developer/TraceInsights/Private/Insights/ViewModels/TrackHeader.cpp
@@ -47,7 +47,7 @@ void FTrackHeader::Reset()
 void FTrackHeader::UpdateSize()
 {
 	const TSharedRef<FSlateFontMeasure> FontMeasureService = FSlateApplication::Get().GetRenderer()->GetFontMeasureService();
-	const float TextWidth = FontMeasureService->Measure(ParentTrack.GetName(), Font, FontScale).X / FontScale;
+	const float TextWidth = static_cast<float>(FontMeasureService->Measure(ParentTrack.GetName(), Font, FontScale).X / FontScale);
 
 	Width = TextWidth + 4.0f;
 	if (bCanBeCollapsed)
@@ -87,8 +87,8 @@ void FTrackHeader::Update(const ITimingTrackUpdateContext& InContext)
 
 void FTrackHeader::PostUpdate(const ITimingTrackUpdateContext& Context)
 {
-	const float MouseX = Context.GetMousePosition().X;
-	const float MouseY = Context.GetMousePosition().Y;
+	const float MouseX = static_cast<float>(Context.GetMousePosition().X);
+	const float MouseY = static_cast<float>(Context.GetMousePosition().Y);
 
 	ParentTrack.SetHeaderHoveredState(MouseX < Width && MouseY < ParentTrack.GetPosY() + Height);
 
@@ -197,7 +197,7 @@ void FTrackHeader::DrawInternal(const ITimingTrackDrawContext& Context, bool bDr
 			//TODO: use a brush/image instead
 			for (float A = 1.0f; A < ArrowSizeY; A += 1.0f)
 			{
-				DrawContext.DrawBox(ArrowX - 3.0 + ArrowSizeY - A, ArrowY + A - 1.0f, A, 1.0f, WhiteBrush, Color);
+				DrawContext.DrawBox(ArrowX - 3.0f + ArrowSizeY - A, ArrowY + A - 1.0f, A, 1.0f, WhiteBrush, Color);
 			}
 		}
 	}
diff --git a/Engine/Source/Developer/TraceInsights/Private/Insights/Widgets/SFrameTrack.cpp b/Engine/Source/Developer/TraceInsights/Private/Insights/Widgets/SFrameTrack.cpp
--- a/Engine/Source/Developer/TraceInsights/Private/Insights/Widgets/SFrameTrack.cpp
+++ b/Engine/Source/Developer/TraceInsights/Private/Insights/Widgets/SFrameTrack.cpp
@@ -143,8 +143,8 @@ void SFrameTrack::Tick(const FGeometry& AllottedGeometry, const double InCurrent
 	if (ThisGeometry != AllottedGeometry || bIsViewportDirty)
 	{
 		bIsViewportDirty = false;
-		const float ViewWidth = AllottedGeometry.GetLocalSize().X;
-		const float ViewHeight = AllottedGeometry.GetLocalSize().Y;
+		const float ViewWidth = static_cast<float>(AllottedGeometry.GetLocalSize().X);
+		const float ViewHeight = static_cast<float>(AllottedGeometry.GetLocalSize().Y);
 		Viewport.SetSize(ViewWidth, ViewHeight);
 		bIsStateDirty = true;
 	}
@@ -189,7 +189,7 @@ void SFrameTrack::Tick(const FGeometry& AllottedGeometry, const double InCurrent
 			{
 				TSharedPtr<FFrameTrackSeries> SeriesPtr = FindOrAddSeries(FrameType);
 
-				int32 NumFrames = FramesProvider.GetFrameCount(static_cast<ETraceFrameType>(FrameType));
+				int32 NumFrames = static_cast<int32>(FramesProvider.GetFrameCount(static_cast<ETraceFrameType>(FrameType)));
 				if (NumFrames > ViewportX.GetMaxValue())
 				{
 					ViewportX.SetMinMaxInterval(0, NumFrames);
@@ -301,14 +301,16 @@ void SFrameTrack::UpdateState()
 
 ////////////////////////////////////////////////////////////////////////////////////////////////////
 
-FFrameTrackSampleRef SFrameTrack::GetSampleAtMousePosition(float X, float Y)
+FFrameTrackSampleRef SFrameTrack::GetSampleAtMousePosition(double X, double Y)
 {
 	if (!bIsStateDirty)
 	{
 		float SampleW = Viewport.GetSampleWidth();
-		int32 SampleIndex = FMath::FloorToInt(X / SampleW);
+		int32 SampleIndex = FMath::FloorToInt(static_cast<float>(X) / SampleW);
 		if (SampleIndex >= 0)
 		{
+			const float MY = static_cast<float>(Y);
+
 			// Search in reverse paint order.
 			for (int32 SeriesIndex = SeriesOrder.Num() - 1; SeriesIndex >= 0; --SeriesIndex)
 			{
@@ -349,7 +351,7 @@ FFrameTrackSampleRef SFrameTrack::GetSampleAtMousePosition(float X, float Y)
 							const float BottomY = FMath::Min(ViewHeight, ViewHeight - BaselineY + ToleranceY);
 							const float TopY = FMath::Max(0.0f, ViewHeight - ValueY - ToleranceY);
 
-							if (Y >= TopY && Y < BottomY)
+							if (MY >= TopY && MY < BottomY)
 							{
 								LLM_SCOPE_BYTAG(Insights);
 								return FFrameTrackSampleRef(SeriesPtr, MakeShared<FFrameTrackSample>(Sample));
@@ -365,16 +367,16 @@ FFrameTrackSampleRef SFrameTrack::GetSampleAtMousePosition(float X, float Y)
 
 ////////////////////////////////////////////////////////////////////////////////////////////////////
 
-void SFrameTrack::SelectFrameAtMousePosition(float X, float Y)
+void SFrameTrack::SelectFrameAtMousePosition(double X, double Y)
 {
 	FFrameTrackSampleRef SampleRef = GetSampleAtMousePosition(X, Y);
 	if (!SampleRef.IsValid())
 	{
-		SampleRef = GetSampleAtMousePosition(X - 1.0f, Y);
+		SampleRef = GetSampleAtMousePosition(X - 1.0, Y);
 	}
 	if (!SampleRef.IsValid())
 	{
-		SampleRef = GetSampleAtMousePosition(X + 1.0f, Y);
+		SampleRef = GetSampleAtMousePosition(X + 1.0, Y);
 	}
 
 	if (SampleRef.IsValid())
@@ -419,8 +421,8 @@ int32 SFrameTrack::OnPaint(const FPaintArgs& Args, const FGeometry& AllottedGeom
 
 	const FSlateBrush* WhiteBrush = FInsightsStyle::Get().GetBrush("WhiteBrush");
 
-	const float ViewWidth = AllottedGeometry.Size.X;
-	const float ViewHeight = AllottedGeometry.Size.Y;
+	const float ViewWidth = static_cast<float>(AllottedGeometry.Size.X);
+	const float ViewHeight = static_cast<float>(AllottedGeometry.Size.Y);
 
 	int32 NumDrawSamples = 0;
 
@@ -521,17 +523,17 @@ int32 SFrameTrack::OnPaint(const FPaintArgs& Args, const FGeometry& AllottedGeom
 			const FString Text2(StringBuilder);
 
 			const float FontScale = DrawContext.Geometry.Scale;
-			const FVector2D TextSize1 = FontMeasureService->Measure(Text1, SummaryFont, FontScale) / FontScale;
-			const FVector2D TextSize2 = FontMeasureService->Measure(Text2, SummaryFont, FontScale) / FontScale;
+			const FVector2f TextSize1(FontMeasureService->Measure(Text1, SummaryFont, FontScale) / FontScale);
+			const FVector2f TextSize2(FontMeasureService->Measure(Text2, SummaryFont, FontScale) / FontScale);
 
 			const FAxisViewportInt32& ViewportX = Viewport.GetHorizontalAxisViewport();
 
 			const float FrameX = ViewportX.GetOffsetForValue(HoveredSample.Sample->LargestFrameIndex);
-			const float CX0 = FMath::RoundToFloat(FrameX + Viewport.GetSampleWidth() / 2);
+			const float CX0 = FMath::RoundToFloat(FrameX + Viewport.GetSampleWidth() / 2.0f);
 
 			constexpr float DX = 3.0f;
-			const float DX1 = FMath::RoundToFloat(TextSize1.X / 2);
-			const float DX2 = FMath::RoundToFloat(TextSize2.X / 2);
+			const float DX1 = FMath::RoundToFloat(TextSize1.X / 2.0f);
+			const float DX2 = FMath::RoundToFloat(TextSize2.X / 2.0f);
 			const float TooltipDesiredSizeX = FMath::Max(DX1, DX2) + DX;
 
 			if (TooltipSizeX != TooltipDesiredSizeX)
@@ -593,7 +595,7 @@ int32 SFrameTrack::OnPaint(const FPaintArgs& Args, const FGeometry& AllottedGeom
 	if (bShouldDisplayDebugInfo)
 	{
 		const float FontScale = DrawContext.Geometry.Scale;
-		const float MaxFontCharHeight = FontMeasureService->Measure(TEXT("!"), SummaryFont, FontScale).Y / FontScale;
+		const float MaxFontCharHeight = static_cast<float>(FontMeasureService->Measure(TEXT("!"), SummaryFont, FontScale).Y / FontScale);
 		const float DbgDY = MaxFontCharHeight;
 
 		const float DbgW = 280.0f;
@@ -798,11 +800,11 @@ void SFrameTrack::DrawVerticalAxisGrid(FDrawContext& DrawContext, const FSlateBr
 								  (Axis.Value <= 1.0) ? FString::Printf(TEXT("%s (%.0f fps)"), *TimeUtils::FormatTimeAuto(Axis.Value), 1.0 / Axis.Value) :
 														TimeUtils::FormatTimeAuto(Axis.Value);
 
-		const FVector2D LabelTextSize = FontMeasureService->Measure(LabelText, Font, FontScale) / FontScale;
-		const float LabelX = bDrawVerticalAxisLabelsOnLeftSide ? 0.0f : ViewWidth - LabelTextSize.X - 4.0f;
+		const float LabelTextWidth = static_cast<float>(FontMeasureService->Measure(LabelText, Font, FontScale).X / FontScale);
+		const float LabelX = bDrawVerticalAxisLabelsOnLeftSide ? 0.0f : ViewWidth - LabelTextWidth - 4.0f;
 
 		// Draw background for value text.
-		DrawContext.DrawBox(LabelX, Axis.LabelY, LabelTextSize.X + 4.0f, TextH, Brush, TextBgColor);
+		DrawContext.DrawBox(LabelX, Axis.LabelY, LabelTextWidth + 4.0f, TextH, Brush, TextBgColor);
 
 		// Draw value text.
 		DrawContext.DrawText(LabelX + 2.0f, Axis.LabelY + 1.0f, LabelText, Font, TextColor);
@@ -847,7 +849,7 @@ void SFrameTrack::DrawHorizontalAxisGrid(FDrawContext& DrawContext, const FSlate
 		const int32 Grid = ((Delta + Power10 - 1) / Power10) * Power10; // next value divisible with a multiple of 10
 
 		// Skip grid lines for negative indices.
-		double StartIndex = ((LeftIndex + Grid - 1) / Grid) * Grid;
+		int32 StartIndex = ((LeftIndex + Grid - 1) / Grid) * Grid;
 		while (StartIndex < 0)
 		{
 			StartIndex += Grid;
@@ -878,8 +880,8 @@ void SFrameTrack::DrawHorizontalAxisGrid(FDrawContext& DrawContext, const FSlate
 			{
 				const float X = FMath::RoundToFloat(ViewportX.GetOffsetForValue(Index));
 				const FString LabelText = FText::AsNumber(Index).ToString();
-				const FVector2D LabelTextSize = FontMeasureService->Measure(LabelText, Font, FontScale) / FontScale;
-				DrawContext.DrawBox(X, 10.0f, LabelTextSize.X + 4.0f, 12.0f, Brush, LabelBoxColor);
+				const float LabelTextWidth = static_cast<float>(FontMeasureService->Measure(LabelText, Font, FontScale).X / FontScale);
+				DrawContext.DrawBox(X, 10.0f, LabelTextWidth + 4.0f, 12.0f, Brush, LabelBoxColor);
 				DrawContext.DrawText(X + 2.0f, 10.0f, LabelText, Font, LabelTextColor);
 			}
 			DrawContext.LayerId++;
@@ -935,7 +937,7 @@ FReply SFrameTrack::OnMouseButtonUp(const FGeometry& MyGeometry, const FPointerE
 			}
 			else if (bIsValidForMouseClick)
 			{
-				SelectFrameAtMousePosition(MousePositionOnButtonUp.X, MousePositionOnButtonUp.Y);
+				SelectFrameAtMousePosition(static_cast<float>(MousePositionOnButtonUp.X), static_cast<float>(MousePositionOnButtonUp.Y));
 			}
 
 			bIsLMB_Pressed = false;
@@ -992,7 +994,7 @@ FReply SFrameTrack::OnMouseMove(const FGeometry& MyGeometry, const FPointerEvent
 				}
 
 				FAxisViewportInt32& ViewportX = Viewport.GetHorizontalAxisViewport();
-				const float PosX = ViewportPosXOnButtonDown + (MousePositionOnButtonDown.X - MousePosition.X);
+				const float PosX = ViewportPosXOnButtonDown + static_cast<float>(MousePositionOnButtonDown.X - MousePosition.X);
 				ViewportX.ScrollAtValue(ViewportX.GetValueAtPos(PosX)); // align viewport position with sample (frame index)
 				UpdateHorizontalScrollBar();
 				bIsStateDirty = true;
@@ -1007,11 +1009,11 @@ FReply SFrameTrack::OnMouseMove(const FGeometry& MyGeometry, const FPointerEvent
 			HoveredSample = GetSampleAtMousePosition(MousePosition.X, MousePosition.Y);
 			if (!HoveredSample.IsValid())
 			{
-				HoveredSample = GetSampleAtMousePosition(MousePosition.X - 1.0f, MousePosition.Y);
+				HoveredSample = GetSampleAtMousePosition(MousePosition.X - 1.0, MousePosition.Y);
 			}
 			if (!HoveredSample.IsValid())
 			{
-				HoveredSample = GetSampleAtMousePosition(MousePosition.X + 1.0f, MousePosition.Y);
+				HoveredSample = GetSampleAtMousePosition(MousePosition.X + 1.0, MousePosition.Y);
 			}
 			if (HoveredSample.IsValid())
 			{
@@ -1066,17 +1068,17 @@ FReply SFrameTrack::OnMouseWheel(const FGeometry& MyGeometry, const FPointerEven
 		FAxisViewportDouble& ViewportY = Viewport.GetVerticalAxisViewport();
 
 		// Zoom in/out vertically.
-		const float Delta = MouseEvent.GetWheelDelta();
-		constexpr float ZoomStep = 0.25f; // as percent
-		float ScaleY;
+		const double Delta = MouseEvent.GetWheelDelta();
+		constexpr double ZoomStep = 0.25; // as percent
+		double ScaleY;
 
 		if (Delta > 0)
 		{
-			ScaleY = ViewportY.GetScale() * FMath::Pow(1.0f + ZoomStep, Delta);
+			ScaleY = ViewportY.GetScale() * FMath::Pow(1.0 + ZoomStep, Delta);
 		}
 		else
 		{
-			ScaleY = ViewportY.GetScale() * FMath::Pow(1.0f / (1.0f + ZoomStep), -Delta);
+			ScaleY = ViewportY.GetScale() * FMath::Pow(1.0 / (1.0 + ZoomStep), -Delta);
 		}
 
 		ViewportY.SetScale(ScaleY);
@@ -1086,7 +1088,7 @@ FReply SFrameTrack::OnMouseWheel(const FGeometry& MyGeometry, const FPointerEven
 	{
 		// Zoom in/out horizontally.
 		const float Delta = MouseEvent.GetWheelDelta();
-		ZoomHorizontally(Delta, MousePosition.X);
+		ZoomHorizontally(Delta, static_cast<float>(MousePosition.X));
 	}
 
 	return FReply::Handled();
diff --git a/Engine/Source/Developer/TraceInsights/Private/Insights/Widgets/SFrameTrack.h b/Engine/Source/Developer/TraceInsights/Private/Insights/Widgets/SFrameTrack.h
--- a/Engine/Source/Developer/TraceInsights/Private/Insights/Widgets/SFrameTrack.h
+++ b/Engine/Source/Developer/TraceInsights/Private/Insights/Widgets/SFrameTrack.h
@@ -136,8 +136,8 @@ protected:
 	void DrawHorizontalAxisGrid(FDrawContext& DrawContext, const FSlateBrush* Brush, const FSlateFontInfo& Font, bool bDrawBackgroundLayer) const;
 	void DrawVerticalAxisGrid(FDrawContext& DrawContext, const FSlateBrush* Brush, const FSlateFontInfo& Font) const;
 
-	FFrameTrackSampleRef GetSampleAtMousePosition(float X, float Y);
-	void SelectFrameAtMousePosition(float X, float Y);
+	FFrameTrackSampleRef GetSampleAtMousePosition(double X, double Y);
+	void SelectFrameAtMousePosition(double X, double Y);
 
 	void ShowContextMenu(const FPointerEvent& MouseEvent);
 
diff --git a/Engine/Source/Developer/TraceInsights/Private/Insights/Widgets/SLogView.cpp b/Engine/Source/Developer/TraceInsights/Private/Insights/Widgets/SLogView.cpp
--- a/Engine/Source/Developer/TraceInsights/Private/Insights/Widgets/SLogView.cpp
+++ b/Engine/Source/Developer/TraceInsights/Private/Insights/Widgets/SLogView.cpp
@@ -1366,8 +1366,8 @@ TSharedRef<SWidget> SLogView::MakeCategoryFilterMenu()
 	CreateCategoriesFilterMenuSection(MenuBuilder);
 	MenuBuilder.EndSection();
 
-	const float MaxMenuHeight = FMath::Clamp(this->GetCachedGeometry().GetLocalSize().Y - 40.0f, 100.0f, 500.0f);
-	return MenuBuilder.MakeWidget(nullptr, MaxMenuHeight);
+	const float MaxMenuHeight = FMath::Clamp(static_cast<float>(this->GetCachedGeometry().GetLocalSize().Y) - 40.0f, 100.0f, 500.0f);
+	return MenuBuilder.MakeWidget(nullptr, FMath::RoundToInt(MaxMenuHeight));
 }
 
 ////////////////////////////////////////////////////////////////////////////////////////////////////
diff --git a/Engine/Source/Developer/TraceInsights/Private/Insights/Widgets/SStatsView.cpp b/Engine/Source/Developer/TraceInsights/Private/Insights/Widgets/SStatsView.cpp
--- a/Engine/Source/Developer/TraceInsights/Private/Insights/Widgets/SStatsView.cpp
+++ b/Engine/Source/Developer/TraceInsights/Private/Insights/Widgets/SStatsView.cpp
@@ -1872,7 +1872,7 @@ void SStatsView::RebuildTree(bool bResync)
 
 		const TraceServices::ICounterProvider& CountersProvider = TraceServices::ReadCounterProvider(*Session.Get());
 
-		const uint32 CounterCount = CountersProvider.GetCounterCount();
+		const uint32 CounterCount = static_cast<uint32>(CountersProvider.GetCounterCount());
 		if (CounterCount != PreviousNodeCount)
 		{
 			check(CounterCount > PreviousNodeCount);
diff --git a/Engine/Source/Developer/TraceInsights/Private/Insights/Widgets/STimingView.cpp b/Engine/Source/Developer/TraceInsights/Private/Insights/Widgets/STimingView.cpp
--- a/Engine/Source/Developer/TraceInsights/Private/Insights/Widgets/STimingView.cpp
+++ b/Engine/Source/Developer/TraceInsights/Private/Insights/Widgets/STimingView.cpp
@@ -548,8 +548,8 @@ void STimingView::Tick(const FGeometry& AllottedGeometry, const double InCurrent
 		OverscrollBottom = FMath::Max(0.0f, OverscrollBottom - InDeltaTime * OverscrollFadeSpeed);
 	}
 
-	const float ViewWidth = AllottedGeometry.GetLocalSize().X;
-	const float ViewHeight = AllottedGeometry.GetLocalSize().Y;
+	const float ViewWidth = static_cast<float>(AllottedGeometry.GetLocalSize().X);
+	const float ViewHeight = static_cast<float>(AllottedGeometry.GetLocalSize().Y);
 
 	////////////////////////////////////////////////////////////////////////////////////////////////////
 	// Update viewport.
@@ -1127,8 +1127,8 @@ int32 STimingView::OnPaint(const FPaintArgs& Args, const FGeometry& AllottedGeom
 	const TSharedRef<FSlateFontMeasure> FontMeasureService = FSlateApplication::Get().GetRenderer()->GetFontMeasureService();
 	const float FontScale = AllottedGeometry.Scale;
 
-	const float ViewWidth = AllottedGeometry.GetLocalSize().X;
-	const float ViewHeight = AllottedGeometry.GetLocalSize().Y;
+	const float ViewWidth = static_cast<float>(AllottedGeometry.GetLocalSize().X);
+	const float ViewHeight = static_cast<float>(AllottedGeometry.GetLocalSize().Y);
 
 #if 0 // Enabling this may further increase UI performance (TODO: profile if this is really needed again).
 	// Warm up Slate vertex/index buffers to avoid initial freezes due to multiple resizes of those buffers.
@@ -1226,7 +1226,7 @@ int32 STimingView::OnPaint(const FPaintArgs& Args, const FGeometry& AllottedGeom
 		PreDrawTracksDurationHistory.AddValue(PreDrawTracksStopwatch.AccumulatedTime);
 	}
 
-	const FVector2D Position = AllottedGeometry.GetAbsolutePosition();
+	const FVector2f Position = FVector2f(AllottedGeometry.GetAbsolutePosition());
 	const float Scale = AllottedGeometry.GetAccumulatedLayoutTransform().GetScale();
 
 	//////////////////////////////////////////////////
@@ -1242,12 +1242,14 @@ int32 STimingView::OnPaint(const FPaintArgs& Args, const FGeometry& AllottedGeom
 			const float TopY = Viewport.GetPosY() + Viewport.GetTopOffset();
 			const float BottomY = Viewport.GetPosY() + Viewport.GetHeight() - Viewport.GetBottomOffset();
 
-			const float L = Position.X;
-			const float R = Position.X + Viewport.GetWidth() * Scale;
-			const float T = Position.Y + TopY * Scale;
-			const float B = Position.Y + BottomY * Scale;
-			const FSlateClippingZone ClipZone(FVector2D(L, T), FVector2D(R, T), FVector2D(L, B), FVector2D(R, B));
-			DrawContext.ElementList.PushClip(ClipZone);
+			{
+				const float L = Position.X;
+				const float R = Position.X + Viewport.GetWidth() * Scale;
+				const float T = Position.Y + TopY * Scale;
+				const float B = Position.Y + BottomY * Scale;
+				const FSlateClippingZone ClipZone(FSlateRect(L, T, R, B));
+				DrawContext.ElementList.PushClip(ClipZone);
+			}
 
 			for (const TSharedPtr<FBaseTimingTrack>& TrackPtr : ScrollableTracks)
 			{
@@ -1277,12 +1279,14 @@ int32 STimingView::OnPaint(const FPaintArgs& Args, const FGeometry& AllottedGeom
 			const float TopY = Viewport.GetPosY();
 			const float BottomY = Viewport.GetPosY() + Viewport.GetTopOffset();
 
-			const float L = Position.X;
-			const float R = Position.X + Viewport.GetWidth() * Scale;
-			const float T = Position.Y + TopY * Scale;
-			const float B = Position.Y + BottomY * Scale;
-			const FSlateClippingZone ClipZone(FVector2D(L, T), FVector2D(R, T), FVector2D(L, B), FVector2D(R, B));
-			DrawContext.ElementList.PushClip(ClipZone);
+			{
+				const float L = Position.X;
+				const float R = Position.X + Viewport.GetWidth() * Scale;
+				const float T = Position.Y + TopY * Scale;
+				const float B = Position.Y + BottomY * Scale;
+				const FSlateClippingZone ClipZone(FSlateRect(L, T, R, B));
+				DrawContext.ElementList.PushClip(ClipZone);
+			}
 
 			for (const TSharedPtr<FBaseTimingTrack>& TrackPtr : TopDockedTracks)
 			{
@@ -1307,12 +1311,14 @@ int32 STimingView::OnPaint(const FPaintArgs& Args, const FGeometry& AllottedGeom
 			const float TopY = Viewport.GetPosY() + Viewport.GetHeight() - Viewport.GetBottomOffset();
 			const float BottomY = Viewport.GetPosY() + Viewport.GetHeight();
 
-			const float L = Position.X;
-			const float R = Position.X + Viewport.GetWidth() * Scale;
-			const float T = Position.Y + TopY * Scale;
-			const float B = Position.Y + BottomY * Scale;
-			const FSlateClippingZone ClipZone(FVector2D(L, T), FVector2D(R, T), FVector2D(L, B), FVector2D(R, B));
-			DrawContext.ElementList.PushClip(ClipZone);
+			{
+				const float L = Position.X;
+				const float R = Position.X + Viewport.GetWidth() * Scale;
+				const float T = Position.Y + TopY * Scale;
+				const float B = Position.Y + BottomY * Scale;
+				const FSlateClippingZone ClipZone(FSlateRect(L, T, R, B));
+				DrawContext.ElementList.PushClip(ClipZone);
+			}
 
 			for (const TSharedPtr<FBaseTimingTrack>& TrackPtr : BottomDockedTracks)
 			{
@@ -1443,7 +1449,7 @@ int32 STimingView::OnPaint(const FPaintArgs& Args, const FGeometry& AllottedGeom
 			const float R = Position.X + Viewport.GetWidth() * Scale;
 			const float T = Position.Y + TopY * Scale;
 			const float B = Position.Y + BottomY * Scale;
-			const FSlateClippingZone ClipZone(FVector2D(L, T), FVector2D(R, T), FVector2D(L, B), FVector2D(R, B));
+			const FSlateClippingZone ClipZone(FSlateRect(L, T, R, B));
 			DrawContext.ElementList.PushClip(ClipZone);
 			bIsClipZoneSet = true;
 		}
@@ -1525,7 +1531,7 @@ int32 STimingView::OnPaint(const FPaintArgs& Args, const FGeometry& AllottedGeom
 	{
 		const FSlateFontInfo& SummaryFont = MainFont;
 
-		const float MaxFontCharHeight = FontMeasureService->Measure(TEXT("!"), SummaryFont, FontScale).Y / FontScale;
+		const float MaxFontCharHeight = static_cast<float>(FontMeasureService->Measure(TEXT("!"), SummaryFont, FontScale).Y / FontScale);
 		const float DbgDY = MaxFontCharHeight;
 
 		const float DbgW = 320.0f;
@@ -1917,7 +1923,7 @@ FReply STimingView::OnMouseButtonDown(const FGeometry& MyGeometry, const FPointe
 
 	if (bAllowPanningOnScreenEdges)
 	{
-		const FVector2D ScreenSpacePosition = MouseEvent.GetScreenSpacePosition();
+		const FVector2f ScreenSpacePosition = FVector2f(MouseEvent.GetScreenSpacePosition());
 		DPIScaleFactor = FPlatformApplicationMisc::GetDPIScaleFactorAtPoint(ScreenSpacePosition.X, ScreenSpacePosition.Y);
 
 		EdgeFrameCountX = 0;
@@ -2040,7 +2046,7 @@ FReply STimingView::OnMouseButtonDown(const FGeometry& MyGeometry, const FPointe
 		bIsDragging = false;
 		TimeRulerTrack->StopScrubbing();
 
-		SelectionStartTime = Viewport.SlateUnitsToTime(MousePositionOnButtonDown.X);
+		SelectionStartTime = Viewport.SlateUnitsToTime(static_cast<float>(MousePositionOnButtonDown.X));
 		SelectionEndTime = SelectionStartTime;
 		LastSelectionType = ESelectionType::None;
 		RaiseSelectionChanging();
@@ -2142,7 +2148,7 @@ FReply STimingView::OnMouseButtonUp(const FGeometry& MyGeometry, const FPointerE
 			if (bIsValidForMouseClick)
 			{
 				// Select the hovered timing event (if any).
-				UpdateHoveredTimingEvent(MousePositionOnButtonUp.X, MousePositionOnButtonUp.Y);
+				UpdateHoveredTimingEvent(static_cast<float>(MousePositionOnButtonUp.X), static_cast<float>(MousePositionOnButtonUp.Y));
 				SelectHoveredTimingTrack();
 				SelectHoveredTimingEvent();
 
@@ -2353,7 +2359,7 @@ FReply STimingView::OnMouseMove(const FGeometry& MyGeometry, const FPointerEvent
 
 				if ((int32)PanningMode & (int32)EPanningMode::Vertical)
 				{
-					const float ScrollPosY = ViewportScrollPosYOnButtonDown + (MousePositionOnButtonDown.Y - MousePosition.Y);
+					const float ScrollPosY = ViewportScrollPosYOnButtonDown + static_cast<float>(MousePositionOnButtonDown.Y - MousePosition.Y);
 					ScrollAtPosY(ScrollPosY);
 				}
 			}
@@ -2364,8 +2370,8 @@ FReply STimingView::OnMouseMove(const FGeometry& MyGeometry, const FPointerEvent
 			{
 				bIsDragging = true;
 
-				SelectionStartTime = Viewport.SlateUnitsToTime(MousePositionOnButtonDown.X);
-				SelectionEndTime = Viewport.SlateUnitsToTime(MousePosition.X);
+				SelectionStartTime = Viewport.SlateUnitsToTime(static_cast<float>(MousePositionOnButtonDown.X));
+				SelectionEndTime = Viewport.SlateUnitsToTime(static_cast<float>(MousePosition.X));
 				if (SelectionStartTime > SelectionEndTime)
 				{
 					double Temp = SelectionStartTime;
@@ -2381,7 +2387,7 @@ FReply STimingView::OnMouseMove(const FGeometry& MyGeometry, const FPointerEvent
 			if (HasMouseCapture())
 			{
 				bIsDragging = true;
-				double Time = Viewport.SlateUnitsToTime(MousePosition.X);
+				double Time = Viewport.SlateUnitsToTime(static_cast<float>(MousePosition.X));
 
 				// Snap to markers.
 				const bool bSnapTimeMarkers = MarkersTrack->IsVisible();
@@ -2398,7 +2404,7 @@ FReply STimingView::OnMouseMove(const FGeometry& MyGeometry, const FPointerEvent
 		}
 		else
 		{
-			UpdateHoveredTimingEvent(MousePosition.X, MousePosition.Y);
+			UpdateHoveredTimingEvent(static_cast<float>(MousePosition.X), static_cast<float>(MousePosition.Y));
 		}
 
 		Reply = FReply::Handled();
@@ -2504,8 +2510,8 @@ FReply STimingView::OnMouseWheel(const FGeometry& MyGeometry, const FPointerEven
 	else
 	{
 		// Zoom in/out horizontally.
-		const double Delta = MouseEvent.GetWheelDelta();
-		if (Viewport.RelativeZoomWithFixedX(Delta, MousePosition.X))
+		const float Delta = MouseEvent.GetWheelDelta();
+		if (Viewport.RelativeZoomWithFixedX(Delta, static_cast<float>(MousePosition.X)))
 		{
 			UpdateHorizontalScrollBar();
 		}
@@ -4811,8 +4817,8 @@ TSharedRef<SDockTab> STimingView::SpawnQuickFindTab(const FSpawnTabArgs& Args)
 	if (OwnerWindow.IsValid() && OwnerWindow != FSlateApplication::Get().FindWidgetWindow(SharedThis(this)))
 	{
 		TSharedPtr<SWindow> TopmostAncestor = OwnerWindow->GetTopmostAncestor();
-		const FVector2D DPIProbePoint = TopmostAncestor->GetPositionInScreen() + FVector2D(10.0f, 10.0f);
-		const float LocalDPIScaleFactor = FPlatformApplicationMisc::GetDPIScaleFactorAtPoint(DPIProbePoint.X, DPIProbePoint.Y);
+		const FVector2D DPIProbePoint = TopmostAncestor->GetPositionInScreen() + FVector2D(10.0, 10.0);
+		const float LocalDPIScaleFactor = FPlatformApplicationMisc::GetDPIScaleFactorAtPoint(static_cast<float>(DPIProbePoint.X), static_cast<float>(DPIProbePoint.Y));
 		OwnerWindow->Resize(FVector2D(600 * LocalDPIScaleFactor, 400 * LocalDPIScaleFactor));
 	}
 
diff --git a/Engine/Source/Developer/TraceInsights/Public/Insights/Common/PaintUtils.h b/Engine/Source/Developer/TraceInsights/Public/Insights/Common/PaintUtils.h
--- a/Engine/Source/Developer/TraceInsights/Public/Insights/Common/PaintUtils.h
+++ b/Engine/Source/Developer/TraceInsights/Public/Insights/Common/PaintUtils.h
@@ -83,7 +83,7 @@ struct TRACEINSIGHTS_API FDrawContext
 		{
 			const TSharedRef<FSlateFontMeasure> FontMeasureService = FSlateApplication::Get().GetRenderer()->GetFontMeasureService();
 			const float FontScale = Geometry.Scale;
-			const float TextWidth = FontMeasureService->Measure(Text, Font, FontScale).X / FontScale;
+			const float TextWidth = static_cast<float>(FontMeasureService->Measure(Text, Font, FontScale).X / FontScale);
 			if (HAlign == HAlign_Right)
 			{
 				TextX -= TextWidth;
diff --git a/Engine/Source/Developer/TraceInsights/Public/Insights/ViewModels/AxisViewportDouble.h b/Engine/Source/Developer/TraceInsights/Public/Insights/ViewModels/AxisViewportDouble.h
--- a/Engine/Source/Developer/TraceInsights/Public/Insights/ViewModels/AxisViewportDouble.h
+++ b/Engine/Source/Developer/TraceInsights/Public/Insights/ViewModels/AxisViewportDouble.h
@@ -47,8 +47,8 @@ public:
 	float GetSize() const { return Size; }
 	bool SetSize(const float InSize);
 
-	int32 GetMinValue() const { return MinValue; }
-	int32 GetMaxValue() const { return MaxValue; }
+	double GetMinValue() const { return MinValue; }
+	double GetMaxValue() const { return MaxValue; }
 
 	float GetMinPos() const { return MinPosition; }
 	float GetMaxPos() const { return MaxPosition; }
diff --git a/Engine/Source/Programs/UnrealInsights/Private/UserInterfaceCommand.cpp b/Engine/Source/Programs/UnrealInsights/Private/UserInterfaceCommand.cpp
--- a/Engine/Source/Programs/UnrealInsights/Private/UserInterfaceCommand.cpp
+++ b/Engine/Source/Programs/UnrealInsights/Private/UserInterfaceCommand.cpp
@@ -215,11 +215,11 @@ void FUserInterfaceCommand::Run()
 
 		FSlateApplication::Get().PumpMessages();
 		FSlateApplication::Get().Tick();
-		FTSTicker::GetCoreTicker().Tick(DeltaTime);
+		FTSTicker::GetCoreTicker().Tick(static_cast<float>(DeltaTime));
 
 		// Throttle frame rate.
 		const float FrameTime = UserInterfaceCommand::IsApplicationBackground() ? BackgroundFrameTime : IdealFrameTime;
-		UserInterfaceCommand::AdaptiveSleep(FMath::Max<float>(0.0f, FrameTime - (FPlatformTime::Seconds() - LastTime)));
+		UserInterfaceCommand::AdaptiveSleep(FMath::Max<float>(0.0f, FrameTime - static_cast<float>(FPlatformTime::Seconds() - LastTime)));
 
 		double CurrentTime = FPlatformTime::Seconds();
 		DeltaTime =  CurrentTime - LastTime;
