Fix remove\trace command not working because name of a file has extra .utrace at the end

diff --git a/Engine/Source/Developer/TraceInsights/Private/Insights/Widgets/SStartPageWindow.cpp b/Engine/Source/Developer/TraceInsights/Private/Insights/Widgets/SStartPageWindow.cpp
--- a/Engine/Source/Developer/TraceInsights/Private/Insights/Widgets/SStartPageWindow.cpp
+++ b/Engine/Source/Developer/TraceInsights/Private/Insights/Widgets/SStartPageWindow.cpp
@@ -1678,7 +1678,7 @@ void STraceStoreWindow::RefreshTraceList()
 					FTraceViewModel& Trace = *TracePtr;
 					Trace.TraceId = SourceTrace.TraceId;
 					Trace.Name = FText::FromString(SourceTrace.Name);
-					FString Uri = FInsightsManager::Get()->GetStoreDir() + TEXT("/") + SourceTrace.Name + TEXT(".utrace");
+					FString Uri = FPaths::ChangeExtension(FPaths::Combine(FInsightsManager::Get()->GetStoreDir(), SourceTrace.Name), TEXT(".utrace"));
 					FPaths::NormalizeFilename(Uri);
 					Trace.Uri = FText::FromString(Uri);
 					UpdateTrace(Trace, SourceTrace);
