Added size column for Modules tab

diff --git a/Engine/Source/Developer/TraceInsights/Private/Insights/Widgets/SModulesView.cpp b/Engine/Source/Developer/TraceInsights/Private/Insights/Widgets/SModulesView.cpp
--- a/Engine/Source/Developer/TraceInsights/Private/Insights/Widgets/SModulesView.cpp
+++ b/Engine/Source/Developer/TraceInsights/Private/Insights/Widgets/SModulesView.cpp
@@ -122,6 +122,16 @@ public:
 	}
 
 	FText GetToolTipText() const { return FText(); }
+	
+	inline uint32 GetSize() const
+	{
+		return Module->Size;
+	}
+
+	FText GetSizeAsText() const
+	{
+		return FText::AsNumber(GetSize());
+	}
 
 private:
 	const TraceServices::FModule* Module;
@@ -248,6 +258,16 @@ public:
 					.Text(this, &SModuleRow::GetSymbolsFile)
 				];
 		}
+		else if (ColumnName == ModulesViewColumns::SizeColumnName)
+		{
+			return SNew(SBox)
+				.Padding(FMargin(4.0, 0.0))
+				[
+					SNew(STextBlock)
+					.ColorAndOpacity(this, &SModuleRow::GetSizeColor)
+					.Text(this, &SModuleRow::GetSizeAsText)
+				];
+		}
 		else
 		{
 			return SNew(STextBlock).Text(LOCTEXT("UnknownColumn", "Unknown Column"));
@@ -383,6 +403,19 @@ public:
 			return FText();
 		}
 	}
+	
+	FText GetSizeAsText() const
+	{
+		const TSharedPtr<FModule> ModulePin = WeakModule.Pin();
+		if (ModulePin.IsValid())
+		{
+			return ModulePin->GetSizeAsText();
+		}
+		else
+		{
+			return FText();
+		}
+	}
 
 	inline const FLinearColor GetRedColor() const       { return FLinearColor(1.0f, 0.7f, 0.7f, 1.0f); }
 	inline const FLinearColor GetYellowColor() const    { return FLinearColor(1.0f, 0.9f, 0.6f, 1.0f); }
@@ -540,6 +573,11 @@ public:
 			return GetBlackColor();
 		}
 	}
+	
+	FSlateColor GetSizeColor() const
+	{
+		return GetWhileColor();
+	}
 
 private:
 	TWeakPtr<FModule> WeakModule;
@@ -643,6 +681,13 @@ void SModulesView::Construct(const FArguments& InArgs)
 							.InitialSortMode(EColumnSortMode::Ascending)
 							.SortMode(this, &SModulesView::GetSortModeForColumn, ModulesViewColumns::AddressRangeColumnName)
 							.OnSort(this, &SModulesView::OnSortModeChanged)
+							
+							+ SHeaderRow::Column(ModulesViewColumns::SizeColumnName)
+							.ManualWidth(76.0f)
+							.DefaultLabel(LOCTEXT("SizeColumn", "Size"))
+							.InitialSortMode(EColumnSortMode::Descending)
+							.SortMode(this, &SModulesView::GetSortModeForColumn, ModulesViewColumns::SizeColumnName)
+							.OnSort(this, &SModulesView::OnSortModeChanged)
 
 							+ SHeaderRow::Column(ModulesViewColumns::ModuleNameColumnName)
 							.ManualWidth(300.0f)
@@ -1088,6 +1133,19 @@ void SModulesView::UpdateSorting()
 				});
 		}
 	}
+	else if (SortColumn == ModulesViewColumns::SizeColumnName)
+	{
+		if (SortMode == EColumnSortMode::Ascending)
+		{
+			Modules.Sort([](const TSharedPtr<FModule>& A, const TSharedPtr<FModule>& B)
+				{ return A->GetSize() < B->GetSize(); });
+		}
+		else
+		{
+			Modules.Sort([](const TSharedPtr<FModule>& A, const TSharedPtr<FModule>& B)
+				{ return A->GetSize() > B->GetSize(); });
+		}
+	}
 }
 
 ////////////////////////////////////////////////////////////////////////////////////////////////////
diff --git a/Engine/Source/Developer/TraceInsights/Private/Insights/Widgets/SModulesView.h b/Engine/Source/Developer/TraceInsights/Private/Insights/Widgets/SModulesView.h
--- a/Engine/Source/Developer/TraceInsights/Private/Insights/Widgets/SModulesView.h
+++ b/Engine/Source/Developer/TraceInsights/Private/Insights/Widgets/SModulesView.h
@@ -28,6 +28,7 @@ namespace ModulesViewColumns
 	static const FName FailedColumnName(TEXT("Failed"));
 	static const FName StatusColumnName(TEXT("Status"));
 	static const FName SymbolsFileColumnName(TEXT("SymbolsFile"));
+	static const FName SizeColumnName(TEXT("Size"));
 }
 
 ////////////////////////////////////////////////////////////////////////////////////////////////////
