Show module name in a callstack grouping when symbol is not resolved.

diff --git a/Engine/Source/Developer/TraceInsights/Private/Insights/MemoryProfiler/ViewModels/MemAllocGroupingByCallstack.cpp b/Engine/Source/Developer/TraceInsights/Private/Insights/MemoryProfiler/ViewModels/MemAllocGroupingByCallstack.cpp
--- a/Engine/Source/Developer/TraceInsights/Private/Insights/MemoryProfiler/ViewModels/MemAllocGroupingByCallstack.cpp
+++ b/Engine/Source/Developer/TraceInsights/Private/Insights/MemoryProfiler/ViewModels/MemAllocGroupingByCallstack.cpp
@@ -96,6 +96,12 @@ void FMemAllocGroupingByCallstack::GroupNodes(const TArray<FTableTreeNodePtr>& N
 
 					if (bIsGroupingByFunction)
 					{
+						// Skip noise for inverted callstack
+						if (bIsInverted && Frame->Symbol->FilterStatus.load() == TraceServices::EResolvedSymbolFilterStatus::Filtered)
+						{
+							continue;
+						}
+
 						const FName GroupName = GetGroupName(Frame);
 
 						// Merge with parent group, if it has the same name (i.e. same function).
@@ -186,11 +192,11 @@ FName FMemAllocGroupingByCallstack::GetGroupName(const TraceServices::FStackFram
 	}
 	else if (Result == TraceServices::ESymbolQueryResult::Pending)
 	{
-		return FName(FString::Printf(TEXT("0x%X [...]"), Frame->Addr), 0);
+		return FName(FString::Printf(TEXT("%s!0x%X [...]"), Frame->Symbol->Module, Frame->Addr), 0);
 	}
 	else
 	{
-		return FName(FString::Printf(TEXT("0x%X"), Frame->Addr), 0);
+		return FName(FString::Printf(TEXT("%s!0x%X"), Frame->Symbol->Module, Frame->Addr), 0);
 	}
 }
 
