[Insights] Fixed private header dependency to TraceServices module from TraceInsights module.

diff --git a/Engine/Source/Developer/TraceInsights/Private/Insights/MemoryProfiler/Common/SymbolSearchPathsHelper.h b/Engine/Source/Developer/TraceInsights/Private/Insights/MemoryProfiler/Common/SymbolSearchPathsHelper.h
--- a/Engine/Source/Developer/TraceInsights/Private/Insights/MemoryProfiler/Common/SymbolSearchPathsHelper.h
+++ b/Engine/Source/Developer/TraceInsights/Private/Insights/MemoryProfiler/Common/SymbolSearchPathsHelper.h
@@ -1,7 +1,7 @@
 ﻿// Copyright Epic Games, Inc. All Rights Reserved.
 #pragma once
 #include "Internationalization/Text.h"
-#include "Model/ModuleProvider.h"
+#include "TraceServices/Model/Modules.h"
 
 #define LOCTEXT_NAMESPACE "SymbolSearchPathsHelper"
 
diff --git a/Engine/Source/Developer/TraceInsights/Private/Insights/Widgets/SModulesView.cpp b/Engine/Source/Developer/TraceInsights/Private/Insights/Widgets/SModulesView.cpp
--- a/Engine/Source/Developer/TraceInsights/Private/Insights/Widgets/SModulesView.cpp
+++ b/Engine/Source/Developer/TraceInsights/Private/Insights/Widgets/SModulesView.cpp
@@ -6,15 +6,14 @@
 #include "Framework/MultiBox/MultiBoxBuilder.h"
 #include "SlateOptMacros.h"
 #include "Textures/SlateIcon.h"
+#include "TraceServices/Model/Modules.h"
 #include "Widgets/Layout/SScrollBox.h"
 #include "Widgets/Text/STextBlock.h"
 
 // Insights
-//#include "Insights/InsightsManager.h"
 #include "DesktopPlatformModule.h"
 #include "Insights/InsightsManager.h"
 #include "Insights/InsightsStyle.h"
-#include "Model/ModuleProvider.h"
 
 ////////////////////////////////////////////////////////////////////////////////////////////////////
 
diff --git a/Engine/Source/Developer/TraceInsights/TraceInsights.Build.cs b/Engine/Source/Developer/TraceInsights/TraceInsights.Build.cs
--- a/Engine/Source/Developer/TraceInsights/TraceInsights.Build.cs
+++ b/Engine/Source/Developer/TraceInsights/TraceInsights.Build.cs
@@ -64,11 +64,5 @@ public class TraceInsights : ModuleRules
 				"MessageLog",
 			}
 		);
-
-		PrivateIncludePaths.AddRange(
-			new string[] {
-				System.IO.Path.Combine(EngineDirectory,"Source/Developer/TraceServices/Private"),
-			}
-		);
 	}
 }
