[Insights] Memory Insights / Mem Allocs: Renamed "Callstack Size" column to just "Callstack".

diff --git a/Engine/Source/Developer/TraceInsights/Private/Insights/MemoryProfiler/ViewModels/MemAllocTable.cpp b/Engine/Source/Developer/TraceInsights/Private/Insights/MemoryProfiler/ViewModels/MemAllocTable.cpp
--- a/Engine/Source/Developer/TraceInsights/Private/Insights/MemoryProfiler/ViewModels/MemAllocTable.cpp
+++ b/Engine/Source/Developer/TraceInsights/Private/Insights/MemoryProfiler/ViewModels/MemAllocTable.cpp
@@ -1057,16 +1057,16 @@ void FMemAllocTable::AddDefaultColumns()
 		AddColumn(ColumnRef);
 	}
 	//////////////////////////////////////////////////
-	// Callstack Size Column
+	// Callstack (Size) Column
 	{
 		TSharedRef<FTableColumn> ColumnRef = MakeShared<FTableColumn>(FMemAllocTableColumns::CallstackSizeColumnId);
 		FTableColumn& Column = *ColumnRef;
 
 		Column.SetIndex(ColumnIndex++);
 
-		Column.SetShortName(LOCTEXT("CallstackSizeColumnName", "Callstack Size"));
-		Column.SetTitleName(LOCTEXT("CallstackSizeColumnTitle", "Callstack Size"));
-		Column.SetDescription(LOCTEXT("CallstackSizeColumnDesc", "Number of callstack frames.\nTooltip shows entire callstack."));
+		Column.SetShortName(LOCTEXT("CallstackSizeColumnName", "Callstack"));
+		Column.SetTitleName(LOCTEXT("CallstackSizeColumnTitle", "Callstack"));
+		Column.SetDescription(LOCTEXT("CallstackSizeColumnDesc", "Number of callstack frames.\nTooltip shows the entire callstack."));
 
 		Column.SetFlags(ETableColumnFlags::ShouldBeVisible | ETableColumnFlags::CanBeHidden | ETableColumnFlags::CanBeFiltered);
 
