[Insights] Added more logging to Insights automated tests.

diff --git a/Engine/Source/Developer/TraceInsights/Private/Insights/Tests/InsightsTestUtils.cpp b/Engine/Source/Developer/TraceInsights/Private/Insights/Tests/InsightsTestUtils.cpp
--- a/Engine/Source/Developer/TraceInsights/Private/Insights/Tests/InsightsTestUtils.cpp
+++ b/Engine/Source/Developer/TraceInsights/Private/Insights/Tests/InsightsTestUtils.cpp
@@ -9,6 +9,7 @@
 #include "TraceServices/ITraceServicesModule.h"
 #include "TraceServices/ModuleService.h"
 
+#include "Insights/Common/Stopwatch.h"
 #include "Insights/IUnrealInsightsModule.h"
 
 ////////////////////////////////////////////////////////////////////////////////////////////////////
@@ -43,19 +44,30 @@ bool FInsightsTestUtils::AnalyzeTrace(const TCHAR* Path) const
 		return false;
 	}
 
-	int32 MaxWaits = 2000;
+	FStopwatch StopWatch;
+	StopWatch.Start();
+
+	double Duration = 0.0f;
+	constexpr double MaxDuration = 75.0f;
 	while (!Session->IsAnalysisComplete())
 	{
 		FPlatformProcess::Sleep(0.033f);
 
-		if (MaxWaits < 0)
+		if (Duration > MaxDuration)
 		{
-			Test->AddError(TEXT("Session analysis takes too long to complete. Aborting test."));
+			Test->AddError(FString::Format(TEXT("Session analysis took longer than the maximum allowed time of {0} seconds. Aborting test."), { MaxDuration }));
 			return false;
 		}
-		--MaxWaits;
+
+		StopWatch.Update();
+		Duration = StopWatch.GetAccumulatedTime();
 	}
 
+	StopWatch.Stop();
+	Duration = StopWatch.GetAccumulatedTime();
+
+	Test->AddInfo(FString::Format(TEXT("Session analysis took {0} seconds."), { Duration }));
+
 	return true;
 }
 
