[Insights]
  - Added Export functionality to the Counters panel:
    * "Export..." � Exports the selection (counters and their aggregated statistics) to a text file (tab-separated values or comma-separated values).
    * "Export Values..." � Exports the values of the selected counter to a text file (tab-separated values or comma-separated values). Exports the values only in the selected time region (if any) or the entire session if no time region is selected.
    * "Export Counters..." � Exports the list of counters to a text file (tab-separated values or comma-separated values).

diff --git a/Engine/Source/Developer/TraceInsights/Private/Insights/ViewModels/TimingExporter.cpp b/Engine/Source/Developer/TraceInsights/Private/Insights/ViewModels/TimingExporter.cpp
--- a/Engine/Source/Developer/TraceInsights/Private/Insights/ViewModels/TimingExporter.cpp
+++ b/Engine/Source/Developer/TraceInsights/Private/Insights/ViewModels/TimingExporter.cpp
@@ -3,6 +3,8 @@
 #include "TimingExporter.h"
 #include "HAL/PlatformFileManager.h"
 #include "Logging/MessageLog.h"
+#include "TraceServices/Model/Bookmarks.h"
+#include "TraceServices/Model/Counters.h"
 #include "TraceServices/Model/Threads.h"
 
 // Insights
@@ -10,7 +12,6 @@
 #include "Insights/Log.h"
 #include "Insights/ViewModels/ThreadTimingTrack.h"
 #include "Insights/TimingProfilerManager.h"
-#include "TraceServices/Model/Bookmarks.h"
 
 namespace Insights
 {
@@ -828,6 +829,190 @@ FTimingExporter::FTimingEventFilterFunc FTimingExporter::MakeTimingEventFilterBy
 
 ////////////////////////////////////////////////////////////////////////////////////////////////////
 
+int32 FTimingExporter::ExportCountersAsText(const FString& Filename, FExportCountersParams& Params) const
+{
+	checkf(Params.Columns == nullptr, TEXT("Custom list of columns is not yet supported!"));
+
+	FStopwatch Stopwatch;
+	Stopwatch.Start();
+
+	IFileHandle* ExportFileHandle = OpenExportFile(*Filename);
+	if (!ExportFileHandle)
+	{
+		return -1;
+	}
+
+	UTF8CHAR Separator = UTF8CHAR('\t');
+	if (Filename.EndsWith(TEXT(".csv")))
+	{
+		Separator = UTF8CHAR(',');
+	}
+	const UTF8CHAR LineEnd = UTF8CHAR('\n');
+
+	FUtf8StringBuilder StringBuilder;
+
+	// Write header.
+	{
+		StringBuilder.Append(UTF8TEXT("Id"));
+		StringBuilder.AppendChar(Separator);
+		StringBuilder.Append(UTF8TEXT("Type"));
+		StringBuilder.AppendChar(Separator);
+		StringBuilder.Append(UTF8TEXT("Name"));
+		StringBuilder.AppendChar(LineEnd);
+
+		ExportFileHandle->Write((const uint8*)StringBuilder.ToString(), StringBuilder.Len() * sizeof(UTF8CHAR));
+	}
+
+	int32 CounterCount = 0;
+
+	// Write values.
+	if (true) // TraceServices::ReadCounterProvider(Session)
+	{
+		TraceServices::FAnalysisSessionReadScope SessionReadScope(Session);
+		const TraceServices::ICounterProvider& CounterProvider = TraceServices::ReadCounterProvider(Session);
+
+		CounterProvider.EnumerateCounters([&](uint32 CounterId, const TraceServices::ICounter& Counter)
+		{
+			StringBuilder.Reset();
+			StringBuilder.Appendf(UTF8TEXT("%u"), CounterId);
+			StringBuilder.AppendChar(Separator);
+			if (Counter.IsFloatingPoint())
+			{
+				StringBuilder.Append(UTF8TEXT("Double"));
+			}
+			else
+			{
+				StringBuilder.Append(UTF8TEXT("Int64"));
+			}
+			if (Counter.IsResetEveryFrame())
+			{
+				StringBuilder.Append(UTF8TEXT("|ResetEveryFrame"));
+			}
+			StringBuilder.AppendChar(Separator);
+			AppendString(StringBuilder, Counter.GetName(), Separator);
+			StringBuilder.AppendChar(LineEnd);
+			ExportFileHandle->Write((const uint8*)StringBuilder.ToString(), StringBuilder.Len() * sizeof(UTF8CHAR));
+			++CounterCount;
+		});
+	}
+
+	ExportFileHandle->Flush();
+	delete ExportFileHandle;
+	ExportFileHandle = nullptr;
+
+	Stopwatch.Stop();
+	const double TotalTime = Stopwatch.GetAccumulatedTime();
+	UE_LOG(TraceInsights, Log, TEXT("Exported %d counters to file in %.3fs (\"%s\")."), CounterCount, TotalTime, *Filename);
+
+	return CounterCount;
+}
+
+////////////////////////////////////////////////////////////////////////////////////////////////////
+
+int32 FTimingExporter::ExportCounterAsText(const FString& Filename, uint32 CounterId, FExportCounterParams& Params) const
+{
+	checkf(Params.Columns == nullptr, TEXT("Custom list of columns is not yet supported!"));
+
+	FStopwatch Stopwatch;
+	Stopwatch.Start();
+
+	IFileHandle* ExportFileHandle = OpenExportFile(*Filename);
+	if (!ExportFileHandle)
+	{
+		return -1;
+	}
+
+	UTF8CHAR Separator = UTF8CHAR('\t');
+	if (Filename.EndsWith(TEXT(".csv")))
+	{
+		Separator = UTF8CHAR(',');
+	}
+	const UTF8CHAR LineEnd = UTF8CHAR('\n');
+
+	FUtf8StringBuilder StringBuilder;
+
+	// Write header.
+	if (Params.bExportOps)
+	{
+		StringBuilder.Append(UTF8TEXT("Time"));
+		StringBuilder.AppendChar(Separator);
+		StringBuilder.Append(UTF8TEXT("Op"));
+		StringBuilder.AppendChar(Separator);
+		StringBuilder.Append(UTF8TEXT("Value"));
+		StringBuilder.AppendChar(LineEnd);
+	}
+	else
+	{
+		StringBuilder.Append(UTF8TEXT("Time"));
+		StringBuilder.AppendChar(Separator);
+		StringBuilder.Append(UTF8TEXT("Value"));
+		StringBuilder.AppendChar(LineEnd);
+	}
+	ExportFileHandle->Write((const uint8*)StringBuilder.ToString(), StringBuilder.Len() * sizeof(UTF8CHAR));
+
+	FString CounterName;
+	int32 ValueCount = 0;
+
+	// Write values.
+	if (true) // TraceServices::ReadCounterProvider(Session)
+	{
+		TraceServices::FAnalysisSessionReadScope SessionReadScope(Session);
+		const TraceServices::ICounterProvider& CounterProvider = TraceServices::ReadCounterProvider(Session);
+
+		CounterProvider.ReadCounter(CounterId, [&](const TraceServices::ICounter& Counter)
+		{
+			CounterName = Counter.GetName();
+
+			// Iterate the counter values.
+			if (Params.bExportOps)
+			{
+				//TODO: export counter events / operations
+			}
+			else
+			{
+				if (Counter.IsFloatingPoint())
+				{
+					Counter.EnumerateFloatValues(Params.IntervalStartTime, Params.IntervalEndTime, false, [&](double Time, double Value)
+					{
+						StringBuilder.Reset();
+						StringBuilder.Appendf(UTF8TEXT("%.9f"), Time);
+						StringBuilder.AppendChar(Separator);
+						StringBuilder.Appendf(UTF8TEXT("%.9f"), Value);
+						StringBuilder.AppendChar(LineEnd);
+						ExportFileHandle->Write((const uint8*)StringBuilder.ToString(), StringBuilder.Len() * sizeof(UTF8CHAR));
+						++ValueCount;
+					});
+				}
+				else
+				{
+					Counter.EnumerateValues(Params.IntervalStartTime, Params.IntervalEndTime, false, [&](double Time, int64 IntValue)
+					{
+						StringBuilder.Reset();
+						StringBuilder.Appendf(UTF8TEXT("%.9f"), Time);
+						StringBuilder.AppendChar(Separator);
+						StringBuilder.Appendf(UTF8TEXT("%lli"), IntValue);
+						StringBuilder.AppendChar(LineEnd);
+						ExportFileHandle->Write((const uint8*)StringBuilder.ToString(), StringBuilder.Len() * sizeof(UTF8CHAR));
+						++ValueCount;
+					});
+				}
+			}
+		});
+	}
+
+	ExportFileHandle->Flush();
+	delete ExportFileHandle;
+	ExportFileHandle = nullptr;
+
+	Stopwatch.Stop();
+	const double TotalTime = Stopwatch.GetAccumulatedTime();
+	UE_LOG(TraceInsights, Log, TEXT("Exported counter %d (\"%s\", %d values) to file in %.3fs (\"%s\")."), CounterId, *CounterName, ValueCount, TotalTime, *Filename);
+
+	return 1;
+}
+
+////////////////////////////////////////////////////////////////////////////////////////////////////
+
 #undef LOCTEXT_NAMESPACE
 
 } // namespace Insights
diff --git a/Engine/Source/Developer/TraceInsights/Private/Insights/ViewModels/TimingExporter.h b/Engine/Source/Developer/TraceInsights/Private/Insights/ViewModels/TimingExporter.h
--- a/Engine/Source/Developer/TraceInsights/Private/Insights/ViewModels/TimingExporter.h
+++ b/Engine/Source/Developer/TraceInsights/Private/Insights/ViewModels/TimingExporter.h
@@ -80,7 +80,37 @@ public:
 		 */
 		FString Region;
 	};
-	
+
+	struct FExportCountersParams
+	{
+		/**
+		 * The list of columns to be exported.
+		 * If nullptr, it uses the default list of columns.
+		 */
+		const TArray<FName>* Columns = nullptr;
+	};
+
+	struct FExportCounterParams
+	{
+		/**
+		 * The list of columns to be exported.
+		 * If nullptr, it uses the default list of columns.
+		 */
+		const TArray<FName>* Columns = nullptr;
+
+		/**
+		 * Filters the counter events by time.
+		 * Only timing events that intersects the [StartTime, EndTime] interval are exported.
+		 */
+		double IntervalStartTime = -std::numeric_limits<double>::infinity();
+		double IntervalEndTime = +std::numeric_limits<double>::infinity();
+
+		/**
+		 * If true, will export the operations and the operation values, instead of final values.
+		 */
+		bool bExportOps = false;
+	};
+
 private:
 	typedef TUtf8StringBuilder<1024> FUtf8StringBuilder;
 
@@ -101,9 +131,15 @@ public:
 	FTimingExporter(const TraceServices::IAnalysisSession& InSession);
 	virtual ~FTimingExporter();
 
+	//////////////////////////////////////////////////////////////////////
+	// Exporters
+
 	int32 ExportThreadsAsText(const FString& Filename, FExportThreadsParams& Params) const;
+
 	int32 ExportTimersAsText(const FString& Filename, FExportTimersParams& Params) const;
+
 	int32 ExportTimingEventsAsText(const FString& Filename, FExportTimingEventsParams& Params) const;
+
 	/**
 	 * Exports Timer Statistics (min,max, inclusive average, exclusive average, etc.).
 	 * Supports specifying a range to export via bookmarks, but does not support timer selection via -timers
@@ -111,6 +147,13 @@ public:
 	 */
 	int32 ExportTimerStatisticsAsText(const FString& Filename, FExportTimerStatisticsParams& Params) const;
 
+	int32 ExportCountersAsText(const FString& Filename, FExportCountersParams& Params) const;
+
+	int32 ExportCounterAsText(const FString& Filename, uint32 CounterId, FExportCounterParams& Params) const;
+
+	//////////////////////////////////////////////////////////////////////
+	// Utilities
+
 	void MakeExportTimingEventsColumnList(const FString& InColumnsString, TArray<FName>& OutColumnList) const;
 
 	//////////////////////////////////////////////////////////////////////
diff --git a/Engine/Source/Developer/TraceInsights/Private/Insights/Widgets/SStatsView.cpp b/Engine/Source/Developer/TraceInsights/Private/Insights/Widgets/SStatsView.cpp
--- a/Engine/Source/Developer/TraceInsights/Private/Insights/Widgets/SStatsView.cpp
+++ b/Engine/Source/Developer/TraceInsights/Private/Insights/Widgets/SStatsView.cpp
@@ -2,10 +2,13 @@
 
 #include "SStatsView.h"
 
+#include "DesktopPlatformModule.h"
 #include "Framework/Commands/Commands.h"
 #include "Framework/Commands/UICommandList.h"
 #include "Framework/MultiBox/MultiBoxBuilder.h"
 #include "HAL/PlatformApplicationMisc.h"
+#include "HAL/PlatformFileManager.h"
+#include "Logging/MessageLog.h"
 #include "SlateOptMacros.h"
 #include "Styling/AppStyle.h"
 #include "TraceServices/Model/Counters.h"
@@ -22,6 +25,7 @@
 #include "Insights/Common/Stopwatch.h"
 #include "Insights/Common/TimeUtils.h"
 #include "Insights/InsightsStyle.h"
+#include "Insights/Log.h"
 #include "Insights/Table/ViewModels/Table.h"
 #include "Insights/Table/ViewModels/TableColumn.h"
 #include "Insights/TimingProfilerCommon.h"
@@ -29,6 +33,7 @@
 #include "Insights/ViewModels/CounterAggregation.h"
 #include "Insights/ViewModels/StatsNodeHelper.h"
 #include "Insights/ViewModels/StatsViewColumnFactory.h"
+#include "Insights/ViewModels/TimingExporter.h"
 #include "Insights/ViewModels/TimingGraphTrack.h"
 #include "Insights/Widgets/SAsyncOperationStatus.h"
 #include "Insights/Widgets/SStatsViewTooltip.h"
@@ -69,10 +74,38 @@ public:
 			"Copies the selection (counters and their aggregated statistics) to clipboard.",
 			EUserInterfaceActionType::Button,
 			FInputChord(EModifierKey::Control, EKeys::C));
+
+		UI_COMMAND(Command_Export,
+			"Export...",
+			"Exports the selection (counters and their aggregated statistics) to a text file (tab-separated values or comma-separated values).",
+			EUserInterfaceActionType::Button,
+			FInputChord(EModifierKey::Control, EKeys::S));
+
+		UI_COMMAND(Command_ExportValues,
+			"Export Values...",
+			"Exports the values of the selected counter to a text file (tab-separated values or comma-separated values).\nExports the values only in the selected time region (if any) or the entire session if no time region is selected.",
+			EUserInterfaceActionType::Button,
+			FInputChord());
+
+		UI_COMMAND(Command_ExportOps,
+			"Export Operations...",
+			"Exports the incremental operations/values of the selected counter to a text file (tab-separated values or comma-separated values).\nExports the ops/values only in the selected time region (if any) or the entire session if no time region is selected.",
+			EUserInterfaceActionType::Button,
+			FInputChord());
+
+		UI_COMMAND(Command_ExportCounters,
+			"Export Counters...",
+			"Exports the list of counters to a text file (tab-separated values or comma-separated values).",
+			EUserInterfaceActionType::Button,
+			FInputChord());
 	}
 	PRAGMA_ENABLE_OPTIMIZATION
 
 	TSharedPtr<FUICommandInfo> Command_CopyToClipboard;
+	TSharedPtr<FUICommandInfo> Command_Export;
+	TSharedPtr<FUICommandInfo> Command_ExportValues;
+	TSharedPtr<FUICommandInfo> Command_ExportOps;
+	TSharedPtr<FUICommandInfo> Command_ExportCounters;
 };
 
 ////////////////////////////////////////////////////////////////////////////////////////////////////
@@ -117,6 +150,10 @@ void SStatsView::InitCommandList()
 	FStatsViewCommands::Register();
 	CommandList = MakeShared<FUICommandList>();
 	CommandList->MapAction(FStatsViewCommands::Get().Command_CopyToClipboard, FExecuteAction::CreateSP(this, &SStatsView::ContextMenu_CopySelectedToClipboard_Execute), FCanExecuteAction::CreateSP(this, &SStatsView::ContextMenu_CopySelectedToClipboard_CanExecute));
+	CommandList->MapAction(FStatsViewCommands::Get().Command_Export, FExecuteAction::CreateSP(this, &SStatsView::ContextMenu_Export_Execute), FCanExecuteAction::CreateSP(this, &SStatsView::ContextMenu_Export_CanExecute));
+	CommandList->MapAction(FStatsViewCommands::Get().Command_ExportValues, FExecuteAction::CreateSP(this, &SStatsView::ContextMenu_ExportValues_Execute), FCanExecuteAction::CreateSP(this, &SStatsView::ContextMenu_ExportValues_CanExecute));
+	CommandList->MapAction(FStatsViewCommands::Get().Command_ExportOps, FExecuteAction::CreateSP(this, &SStatsView::ContextMenu_ExportOps_Execute), FCanExecuteAction::CreateSP(this, &SStatsView::ContextMenu_ExportOps_CanExecute));
+	CommandList->MapAction(FStatsViewCommands::Get().Command_ExportCounters, FExecuteAction::CreateSP(this, &SStatsView::ContextMenu_ExportCounters_Execute), FCanExecuteAction::CreateSP(this, &SStatsView::ContextMenu_ExportCounters_CanExecute));
 }
 
 ////////////////////////////////////////////////////////////////////////////////////////////////////
@@ -442,6 +479,44 @@ TSharedPtr<SWidget> SStatsView::TreeView_GetMenuContent()
 			TAttribute<FText>(),
 			FSlateIcon(FAppStyle::Get().GetStyleSetName(), "GenericCommands.Copy")
 		);
+
+		MenuBuilder.AddMenuEntry
+		(
+			FStatsViewCommands::Get().Command_Export,
+			NAME_None,
+			TAttribute<FText>(),
+			TAttribute<FText>(),
+			FSlateIcon(FAppStyle::Get().GetStyleSetName(), "Icons.Save")
+		);
+
+		MenuBuilder.AddMenuEntry
+		(
+			FStatsViewCommands::Get().Command_ExportValues,
+			NAME_None,
+			TAttribute<FText>(),
+			TAttribute<FText>(),
+			FSlateIcon(FAppStyle::Get().GetStyleSetName(), "Icons.Save")
+		);
+
+#if 0 // not implemented yet
+		MenuBuilder.AddMenuEntry
+		(
+			FStatsViewCommands::Get().Command_ExportOps,
+			NAME_None,
+			TAttribute<FText>(),
+			TAttribute<FText>(),
+			FSlateIcon(FAppStyle::Get().GetStyleSetName(), "Icons.Save")
+		);
+#endif
+
+		MenuBuilder.AddMenuEntry
+		(
+			FStatsViewCommands::Get().Command_ExportCounters,
+			NAME_None,
+			TAttribute<FText>(),
+			TAttribute<FText>(),
+			FSlateIcon(FAppStyle::Get().GetStyleSetName(), "Icons.Save")
+		);
 	}
 	MenuBuilder.EndSection();
 
@@ -2158,6 +2233,305 @@ void SStatsView::ContextMenu_CopySelectedToClipboard_Execute()
 
 ////////////////////////////////////////////////////////////////////////////////////////////////////
 
+bool SStatsView::ContextMenu_Export_CanExecute() const
+{
+	const TArray<FStatsNodePtr> SelectedNodes = TreeView->GetSelectedItems();
+	return SelectedNodes.Num() > 0;
+}
+
+////////////////////////////////////////////////////////////////////////////////////////////////////
+
+void SStatsView::ContextMenu_Export_Execute()
+{
+	if (!Table->IsValid())
+	{
+		return;
+	}
+
+	TArray<Insights::FBaseTreeNodePtr> SelectedNodes;
+	for (FStatsNodePtr Item : TreeView->GetSelectedItems())
+	{
+		SelectedNodes.Add(Item);
+	}
+
+	if (SelectedNodes.Num() == 0)
+	{
+		return;
+	}
+
+	const FString DialogTitle = LOCTEXT("Export_Title", "Export Aggregated Counter Stats").ToString();
+	const FString DefaultFile = TEXT("CounterStats.tsv");
+	FString Filename;
+	if (!OpenSaveTextFileDialog(DialogTitle, DefaultFile, Filename))
+	{
+		return;
+	}
+
+	IFileHandle* ExportFileHandle = OpenExportFile(*Filename);
+	if (!ExportFileHandle)
+	{
+		return;
+	}
+
+	FStopwatch Stopwatch;
+	Stopwatch.Start();
+
+	UTF16CHAR BOM = UNICODE_BOM;
+	ExportFileHandle->Write((uint8*)&BOM, sizeof(UTF16CHAR));
+
+	TCHAR Separator = TEXT('\t');
+	if (Filename.EndsWith(TEXT(".csv")))
+	{
+		Separator = TEXT(',');
+	}
+	constexpr TCHAR LineEnd = TEXT('\n');
+	constexpr TCHAR QuotationMarkBegin = TEXT('\"');
+	constexpr TCHAR QuotationMarkEnd = TEXT('\"');
+
+	TStringBuilder<1024> StringBuilder;
+
+	TArray<TSharedRef<Insights::FTableColumn>> VisibleColumns;
+	Table->GetVisibleColumns(VisibleColumns);
+
+	// Write header.
+	{
+		bool bIsFirstColumn = true;
+		for (const TSharedRef<Insights::FTableColumn>& ColumnRef : VisibleColumns)
+		{
+			if (bIsFirstColumn)
+			{
+				bIsFirstColumn = false;
+			}
+			else
+			{
+				StringBuilder.AppendChar(Separator);
+			}
+			FString Value = ColumnRef->GetShortName().ToString().ReplaceCharWithEscapedChar();
+			int32 CharIndex;
+			if (Value.FindChar(Separator, CharIndex))
+			{
+				StringBuilder.AppendChar(QuotationMarkBegin);
+				StringBuilder.Append(Value);
+				StringBuilder.AppendChar(QuotationMarkEnd);
+			}
+			else
+			{
+				StringBuilder.Append(Value);
+			}
+		}
+		StringBuilder.AppendChar(LineEnd);
+		ExportFileHandle->Write((const uint8*)StringBuilder.ToString(), StringBuilder.Len() * sizeof(TCHAR));
+	}
+
+	if (CurrentSorter.IsValid())
+	{
+		CurrentSorter->Sort(SelectedNodes, ColumnSortMode == EColumnSortMode::Ascending ? Insights::ESortMode::Ascending : Insights::ESortMode::Descending);
+	}
+
+	const int32 NodeCount = SelectedNodes.Num();
+	for (int32 Index = 0; Index < NodeCount; Index++)
+	{
+		const Insights::FBaseTreeNodePtr& Node = SelectedNodes[Index];
+
+		StringBuilder.Reset();
+
+		bool bIsFirstColumn = true;
+		for (const TSharedRef<Insights::FTableColumn>& ColumnRef : VisibleColumns)
+		{
+			if (bIsFirstColumn)
+			{
+				bIsFirstColumn = false;
+			}
+			else
+			{
+				StringBuilder.AppendChar(Separator);
+			}
+
+			FString Value = ColumnRef->GetValueAsSerializableString(*Node).ReplaceCharWithEscapedChar();
+			int32 CharIndex;
+			if (Value.FindChar(Separator, CharIndex))
+			{
+				StringBuilder.AppendChar(QuotationMarkBegin);
+				StringBuilder.Append(Value);
+				StringBuilder.AppendChar(QuotationMarkEnd);
+			}
+			else
+			{
+				StringBuilder.Append(Value);
+			}
+		}
+		StringBuilder.AppendChar(LineEnd);
+		ExportFileHandle->Write((const uint8*)StringBuilder.ToString(), StringBuilder.Len() * sizeof(TCHAR));
+	}
+
+	ExportFileHandle->Flush();
+	delete ExportFileHandle;
+	ExportFileHandle = nullptr;
+
+	Stopwatch.Stop();
+	const double TotalTime = Stopwatch.GetAccumulatedTime();
+	UE_LOG(TraceInsights, Log, TEXT("Exported aggregated counter stats to file in %.3fs (\"%s\")."), TotalTime, *Filename);
+}
+
+////////////////////////////////////////////////////////////////////////////////////////////////////
+
+bool SStatsView::ContextMenu_ExportValues_CanExecute() const
+{
+	const TArray<FStatsNodePtr> SelectedNodes = TreeView->GetSelectedItems();
+	return (SelectedNodes.Num() == 1) && !SelectedNodes[0]->IsGroup();
+}
+
+////////////////////////////////////////////////////////////////////////////////////////////////////
+
+void SStatsView::ContextMenu_ExportValues_Execute() const
+{
+	if (!Session.IsValid())
+	{
+		return;
+	}
+
+	const TArray<FStatsNodePtr> SelectedNodes = TreeView->GetSelectedItems();
+	if (SelectedNodes.Num() != 1 || SelectedNodes[0]->IsGroup())
+	{
+		return;
+	}
+
+	const uint32 CounterId = SelectedNodes[0]->GetCounterId();
+
+	const FString DialogTitle = LOCTEXT("ExportValues_Title", "Export Counter Values").ToString();
+	const FString DefaultFile = TEXT("CounterValues.tsv");
+	FString Filename;
+	if (!OpenSaveTextFileDialog(DialogTitle, DefaultFile, Filename))
+	{
+		return;
+	}
+
+	Insights::FTimingExporter Exporter(*Session.Get());
+	Insights::FTimingExporter::FExportCounterParams Params; // default
+	Exporter.ExportCounterAsText(*Filename, CounterId, Params);
+}
+
+////////////////////////////////////////////////////////////////////////////////////////////////////
+
+bool SStatsView::ContextMenu_ExportOps_CanExecute() const
+{
+	const TArray<FStatsNodePtr> SelectedNodes = TreeView->GetSelectedItems();
+	return (SelectedNodes.Num() == 1) && !SelectedNodes[0]->IsGroup();
+}
+
+////////////////////////////////////////////////////////////////////////////////////////////////////
+
+void SStatsView::ContextMenu_ExportOps_Execute() const
+{
+	if (!Session.IsValid())
+	{
+		return;
+	}
+
+	const TArray<FStatsNodePtr> SelectedNodes = TreeView->GetSelectedItems();
+	if (SelectedNodes.Num() != 1 || SelectedNodes[0]->IsGroup())
+	{
+		return;
+	}
+
+	const uint32 CounterId = SelectedNodes[0]->GetCounterId();
+
+	const FString DialogTitle = LOCTEXT("ExportOps_Title", "Export Counter Ops").ToString();
+	const FString DefaultFile = TEXT("CounterOps.tsv");
+	FString Filename;
+	if (!OpenSaveTextFileDialog(DialogTitle, DefaultFile, Filename))
+	{
+		return;
+	}
+
+	Insights::FTimingExporter Exporter(*Session.Get());
+	Insights::FTimingExporter::FExportCounterParams Params; // default
+	Params.bExportOps = true;
+	Exporter.ExportCounterAsText(*Filename, CounterId, Params);
+}
+
+////////////////////////////////////////////////////////////////////////////////////////////////////
+
+bool SStatsView::ContextMenu_ExportCounters_CanExecute() const
+{
+	return true;
+}
+
+////////////////////////////////////////////////////////////////////////////////////////////////////
+
+void SStatsView::ContextMenu_ExportCounters_Execute() const
+{
+	if (!Session.IsValid())
+	{
+		return;
+	}
+
+	const FString DialogTitle = LOCTEXT("ExportCounters_Title", "Export Counters").ToString();
+	const FString DefaultFile = TEXT("Counters.tsv");
+	FString Filename;
+	if (!OpenSaveTextFileDialog(DialogTitle, DefaultFile, Filename))
+	{
+		return;
+	}
+
+	Insights::FTimingExporter Exporter(*Session.Get());
+	Insights::FTimingExporter::FExportCountersParams Params; // default
+	Exporter.ExportCountersAsText(*Filename, Params);
+}
+
+////////////////////////////////////////////////////////////////////////////////////////////////////
+
+////////////////////////////////////////////////////////////////////////////////////////////////////
+
+bool SStatsView::OpenSaveTextFileDialog(const FString& InDialogTitle, const FString& InDefaultFile, FString& OutFilename) const
+{
+	TArray<FString> SaveFilenames;
+	bool bDialogResult = false;
+
+	IDesktopPlatform* DesktopPlatform = FDesktopPlatformModule::Get();
+	if (DesktopPlatform)
+	{
+		const FString DefaultPath = FPaths::ProjectSavedDir();
+		bDialogResult = DesktopPlatform->SaveFileDialog(
+			FSlateApplication::Get().FindBestParentWindowHandleForDialogs(nullptr),
+			InDialogTitle,
+			DefaultPath,
+			InDefaultFile,
+			TEXT("Tab-Separated Values (*.tsv)|*.tsv|Text Files (*.txt)|*.txt|Comma-Separated Values (*.csv)|*.csv|All Files (*.*)|*.*"),
+			EFileDialogFlags::None,
+			SaveFilenames
+		);
+	}
+
+	if (!bDialogResult || SaveFilenames.Num() == 0)
+	{
+		return false;
+	}
+
+	OutFilename = SaveFilenames[0];
+	return true;
+}
+
+////////////////////////////////////////////////////////////////////////////////////////////////////
+
+IFileHandle* SStatsView::OpenExportFile(const TCHAR* InFilename) const
+{
+	IFileHandle* ExportFileHandle = FPlatformFileManager::Get().GetPlatformFile().OpenWrite(InFilename);
+
+	if (ExportFileHandle == nullptr)
+	{
+		FName LogListingName = FTimingProfilerManager::Get()->GetLogListingName();
+		FMessageLog ReportMessageLog((LogListingName != NAME_None) ? LogListingName : TEXT("Other"));
+		ReportMessageLog.Error(LOCTEXT("FailedToOpenFile", "Export failed. Failed to open file for write."));
+		ReportMessageLog.Notify();
+		return nullptr;
+	}
+
+	return ExportFileHandle;
+}
+
+////////////////////////////////////////////////////////////////////////////////////////////////////
+
 FReply SStatsView::OnKeyDown(const FGeometry& MyGeometry, const FKeyEvent& InKeyEvent)
 {
 	return CommandList->ProcessCommandBindings(InKeyEvent) == true ? FReply::Handled() : FReply::Unhandled();
diff --git a/Engine/Source/Developer/TraceInsights/Private/Insights/Widgets/SStatsView.h b/Engine/Source/Developer/TraceInsights/Private/Insights/Widgets/SStatsView.h
--- a/Engine/Source/Developer/TraceInsights/Private/Insights/Widgets/SStatsView.h
+++ b/Engine/Source/Developer/TraceInsights/Private/Insights/Widgets/SStatsView.h
@@ -122,6 +122,21 @@ private:
 	bool ContextMenu_CopySelectedToClipboard_CanExecute() const;
 	void ContextMenu_CopySelectedToClipboard_Execute();
 
+	bool ContextMenu_Export_CanExecute() const;
+	void ContextMenu_Export_Execute();
+
+	bool ContextMenu_ExportValues_CanExecute() const;
+	void ContextMenu_ExportValues_Execute() const;
+
+	bool ContextMenu_ExportOps_CanExecute() const;
+	void ContextMenu_ExportOps_Execute() const;
+
+	bool ContextMenu_ExportCounters_CanExecute() const;
+	void ContextMenu_ExportCounters_Execute() const;
+
+	bool OpenSaveTextFileDialog(const FString& InDialogTitle, const FString& InDefaultFile, FString& OutFilename) const;
+	class IFileHandle* OpenExportFile(const TCHAR* InFilename) const;
+
 	////////////////////////////////////////////////////////////////////////////////////////////////////
 	// Tree View - Columns' Header
 
diff --git a/Engine/Source/Developer/TraceInsights/Private/Insights/Widgets/STimersView.cpp b/Engine/Source/Developer/TraceInsights/Private/Insights/Widgets/STimersView.cpp
--- a/Engine/Source/Developer/TraceInsights/Private/Insights/Widgets/STimersView.cpp
+++ b/Engine/Source/Developer/TraceInsights/Private/Insights/Widgets/STimersView.cpp
@@ -74,13 +74,47 @@ public:
 	PRAGMA_DISABLE_OPTIMIZATION
 	virtual void RegisterCommands() override
 	{
-		UI_COMMAND(Command_CopyToClipboard, "Copy To Clipboard", "Copies the selection (timers and their aggregated statistics) to clipboard.", EUserInterfaceActionType::Button, FInputChord(EModifierKey::Control, EKeys::C));
-		UI_COMMAND(Command_Export, "Export...", "Exports the selection (timers and their aggregated statistics) to a text file (tab-separated values or comma-separated values).", EUserInterfaceActionType::Button, FInputChord(EModifierKey::Control, EKeys::S));
-		UI_COMMAND(Command_ExportTimingEventsSelection, "Export Timing Events (Selection)...", "Exports the timing events to a text file (tab-separated values or comma-separated values).\nOnly exports the timing events of the selected timer(s), for the visible CPU/GPU Thread tracks and in the selected time region (if any).", EUserInterfaceActionType::Button, FInputChord());
-		UI_COMMAND(Command_ExportTimingEvents, "Export Timing Events (All)...", "Exports all the timing events (for all CPU/GPU threads) to a text file (tab-separated values or comma-separated values).", EUserInterfaceActionType::Button, FInputChord());
-		UI_COMMAND(Command_ExportThreads, "Export Threads...", "Exports the list of threads to a text file (tab-separated values or comma-separated values).", EUserInterfaceActionType::Button, FInputChord());
-		UI_COMMAND(Command_ExportTimers, "Export Timers...", "Exports the list of timers to a text file (tab-separated values or comma-separated values).", EUserInterfaceActionType::Button, FInputChord());
-		UI_COMMAND(Command_OpenSource, "Open Source", "Opens the source file of the selected timer in the registered IDE.", EUserInterfaceActionType::Button, FInputChord());
+		UI_COMMAND(Command_CopyToClipboard,
+			"Copy To Clipboard",
+			"Copies the selection (timers and their aggregated statistics) to clipboard.",
+			EUserInterfaceActionType::Button,
+			FInputChord(EModifierKey::Control, EKeys::C));
+
+		UI_COMMAND(Command_Export,
+			"Export...",
+			"Exports the selection (timers and their aggregated statistics) to a text file (tab-separated values or comma-separated values).",
+			EUserInterfaceActionType::Button,
+			FInputChord(EModifierKey::Control, EKeys::S));
+
+		UI_COMMAND(Command_ExportTimingEventsSelection,
+			"Export Timing Events (Selection)...",
+			"Exports the timing events to a text file (tab-separated values or comma-separated values).\nOnly exports the timing events of the selected timer(s), for the visible CPU/GPU Thread tracks and in the selected time region (if any).",
+			EUserInterfaceActionType::Button,
+			FInputChord());
+
+		UI_COMMAND(Command_ExportTimingEvents,
+			"Export Timing Events (All)...",
+			"Exports all the timing events (for all CPU/GPU threads) to a text file (tab-separated values or comma-separated values).",
+			EUserInterfaceActionType::Button,
+			FInputChord());
+
+		UI_COMMAND(Command_ExportThreads,
+			"Export Threads...",
+			"Exports the list of threads to a text file (tab-separated values or comma-separated values).",
+			EUserInterfaceActionType::Button,
+			FInputChord());
+
+		UI_COMMAND(Command_ExportTimers,
+			"Export Timers...",
+			"Exports the list of timers to a text file (tab-separated values or comma-separated values).",
+			EUserInterfaceActionType::Button,
+			FInputChord());
+
+		UI_COMMAND(Command_OpenSource,
+			"Open Source",
+			"Opens the source file of the selected timer in the registered IDE.",
+			EUserInterfaceActionType::Button,
+			FInputChord());
 	}
 	PRAGMA_ENABLE_OPTIMIZATION
 
@@ -2235,7 +2269,8 @@ void STimersView::ContextMenu_CopyToClipboard_Execute()
 
 bool STimersView::ContextMenu_Export_CanExecute() const
 {
-	return true;
+	const TArray<FTimerNodePtr> SelectedNodes = TreeView->GetSelectedItems();
+	return SelectedNodes.Num() > 0;
 }
 
 ////////////////////////////////////////////////////////////////////////////////////////////////////
