[Insights]
  - Fixed regression in detecting timing event under mouse (offset of half of a lane height) - UE-.
  - Fixed regression related to drawing of points in graph tracks - UE-.

diff --git a/Engine/Source/Developer/TraceInsights/Private/Insights/TaskGraphProfiler/ViewModels/TaskTimingTrack.cpp b/Engine/Source/Developer/TraceInsights/Private/Insights/TaskGraphProfiler/ViewModels/TaskTimingTrack.cpp
--- a/Engine/Source/Developer/TraceInsights/Private/Insights/TaskGraphProfiler/ViewModels/TaskTimingTrack.cpp
+++ b/Engine/Source/Developer/TraceInsights/Private/Insights/TaskGraphProfiler/ViewModels/TaskTimingTrack.cpp
@@ -876,7 +876,7 @@ const TSharedPtr<const ITimingEvent> FTaskTimingTrack::GetEvent(float InPosX, fl
 	if (DY >= 0 && DY < GetHeight() - 1.0f - 2 * Layout.TimelineDY)
 	{
 		const double EventTime = Viewport.SlateUnitsToTime(InPosX);
-		const int32 Depth = FMath::RoundToInt(DY / (Layout.EventH + Layout.EventDY));
+		const int32 Depth = static_cast<int32>(DY / (Layout.EventH + Layout.EventDY));
 
 		TSharedPtr<const TraceServices::IAnalysisSession> Session = FInsightsManager::Get()->GetSession();
 		if (!Session.IsValid())
diff --git a/Engine/Source/Developer/TraceInsights/Private/Insights/ViewModels/GraphTrack.cpp b/Engine/Source/Developer/TraceInsights/Private/Insights/ViewModels/GraphTrack.cpp
--- a/Engine/Source/Developer/TraceInsights/Private/Insights/ViewModels/GraphTrack.cpp
+++ b/Engine/Source/Developer/TraceInsights/Private/Insights/ViewModels/GraphTrack.cpp
@@ -408,7 +408,7 @@ void FGraphTrack::DrawSeries(const FGraphSeries& Series, FDrawContext& DrawConte
 		{
 			const FVector2D& Pt = Series.Points[Index];
 			const float PtX = static_cast<float>(Pt.X) - PointVisualSize / 2.0f + 0.5f;
-			const float PtY = LocalPosY + static_cast<float>(Pt.Y) - PointVisualSize / 2.0f - 0.5f;
+			const float PtY = LocalPosY + static_cast<float>(Pt.Y) - PointVisualSize / 2.0f + 0.5f;
 			DrawContext.DrawBox(PtX, PtY, PointVisualSize, PointVisualSize, PointBrush, Series.Color);
 		}
 		DrawContext.LayerId++;
diff --git a/Engine/Source/Developer/TraceInsights/Private/Insights/ViewModels/ThreadTimingTrack.cpp b/Engine/Source/Developer/TraceInsights/Private/Insights/ViewModels/ThreadTimingTrack.cpp
--- a/Engine/Source/Developer/TraceInsights/Private/Insights/ViewModels/ThreadTimingTrack.cpp
+++ b/Engine/Source/Developer/TraceInsights/Private/Insights/ViewModels/ThreadTimingTrack.cpp
@@ -1397,7 +1397,7 @@ const TSharedPtr<const ITimingEvent> FThreadTimingTrack::GetEvent(float InPosX,
 	// If mouse is not above first sub-track or below last sub-track...
 	if (DY >= 0 && DY < TrackLanesHeight)
 	{
-		const int32 Depth = FMath::RoundToInt(DY / (Layout.EventH + Layout.EventDY));
+		const int32 Depth = static_cast<int32>(DY / (Layout.EventH + Layout.EventDY));
 
 		const double SecondsPerPixel = 1.0 / Viewport.GetScaleX();
 
diff --git a/Engine/Source/Developer/TraceInsights/Private/Insights/ViewModels/TimingEventsTrack.cpp b/Engine/Source/Developer/TraceInsights/Private/Insights/ViewModels/TimingEventsTrack.cpp
--- a/Engine/Source/Developer/TraceInsights/Private/Insights/ViewModels/TimingEventsTrack.cpp
+++ b/Engine/Source/Developer/TraceInsights/Private/Insights/ViewModels/TimingEventsTrack.cpp
@@ -402,7 +402,7 @@ const TSharedPtr<const ITimingEvent> FTimingEventsTrack::GetEvent(float InPosX,
 	// If mouse is not above first sub-track or below last sub-track...
 	if (DY >= 0 && DY < TrackLanesHeight)
 	{
-		const int32 Depth = FMath::RoundToInt(DY / (Layout.EventH + Layout.EventDY));
+		const int32 Depth = static_cast<int32>(DY / (Layout.EventH + Layout.EventDY));
 		const double SecondsPerPixel = 1.0 / Viewport.GetScaleX();
 		const double TimeAtPosX = Viewport.SlateUnitsToTime(InPosX);
 
