[Insights]
  - Tree/table views: Fixed "SameValue" aggregation for CString columns.
  - Tree/table views: Fixed sorting for string columns to be case insensitive and to sort undefined values first.
  - Modules view: Fixed sorting by [REDACTED_PERSON] column and by Symbols column to be case insensitive.

diff --git a/Engine/Source/Developer/TraceInsights/Private/Insights/Table/ViewModels/TableCellValueSorter.cpp b/Engine/Source/Developer/TraceInsights/Private/Insights/Table/ViewModels/TableCellValueSorter.cpp
--- a/Engine/Source/Developer/TraceInsights/Private/Insights/Table/ViewModels/TableCellValueSorter.cpp
+++ b/Engine/Source/Developer/TraceInsights/Private/Insights/Table/ViewModels/TableCellValueSorter.cpp
@@ -492,20 +492,20 @@ FSorterByCStringValue::FSorterByCStringValue(TSharedRef<FTableColumn> InColumnRe
 		const TOptional<FTableCellValue> OptionalValueA = Column.GetValue(*A);
 		const TOptional<FTableCellValue> OptionalValueB = Column.GetValue(*B);
 
-		const TCHAR* ValueA = OptionalValueA.IsSet() ? OptionalValueA.GetValue().CString : TEXT("");
-		const TCHAR* ValueB = OptionalValueB.IsSet() ? OptionalValueB.GetValue().CString : TEXT("");
+		const TCHAR* ValueA = OptionalValueA.IsSet() ? OptionalValueA.GetValue().CString : nullptr;
+		const TCHAR* ValueB = OptionalValueB.IsSet() ? OptionalValueB.GetValue().CString : nullptr;
 
 		// If any value is nullptr
 		if (!ValueA || !ValueB)
 		{
-			if(!ValueA && !ValueB)
+			if (!ValueA && !ValueB)
 			{
 				INSIGHTS_DEFAULT_SORTING_NODES(A, B)
 			}
 			return ValueA < ValueB;
 		}
-			
-		const int32 Compare = FCString::Strcmp(ValueA, ValueB);
+
+		const int32 Compare = FCString::Stricmp(ValueA, ValueB);
 		if (Compare == 0)
 		{
 			INSIGHTS_DEFAULT_SORTING_NODES(A, B)
@@ -524,20 +524,20 @@ FSorterByCStringValue::FSorterByCStringValue(TSharedRef<FTableColumn> InColumnRe
 		const TOptional<FTableCellValue> OptionalValueA = Column.GetValue(*A);
 		const TOptional<FTableCellValue> OptionalValueB = Column.GetValue(*B);
 
-		const TCHAR* ValueA = OptionalValueA.IsSet() ? OptionalValueA.GetValue().CString : TEXT("");
-		const TCHAR* ValueB = OptionalValueB.IsSet() ? OptionalValueB.GetValue().CString : TEXT("");
-		
+		const TCHAR* ValueA = OptionalValueA.IsSet() ? OptionalValueA.GetValue().CString : nullptr;
+		const TCHAR* ValueB = OptionalValueB.IsSet() ? OptionalValueB.GetValue().CString : nullptr;
+
 		// If any value is nullptr
 		if (!ValueA || !ValueB)
 		{
-			if(!ValueA && !ValueB)
+			if (!ValueA && !ValueB)
 			{
 				INSIGHTS_DEFAULT_SORTING_NODES(A, B)
 			}
 			return ValueA > ValueB;
 		}
 
-		const int32 Compare = FCString::Strcmp(ValueA, ValueB);
+		const int32 Compare = FCString::Stricmp(ValueA, ValueB);
 		if (Compare == 0)
 		{
 			INSIGHTS_DEFAULT_SORTING_NODES(A, B)
@@ -624,7 +624,7 @@ void FSorterByTextValue::Sort(TArray<FBaseTreeNodePtr>& NodesToSort, ESortMode S
 			const FString& ValueB = UniqueValues[B.ValueId];
 
 			// Sort by value (ascending).
-			return ValueA.Compare(ValueB) <= 0;
+			return ValueA.Compare(ValueB, ESearchCase::IgnoreCase) <= 0;
 		});
 	}
 	else // if (SortMode == ESortMode::Descending)
@@ -652,7 +652,7 @@ void FSorterByTextValue::Sort(TArray<FBaseTreeNodePtr>& NodesToSort, ESortMode S
 			const FString& ValueB = UniqueValues[B.ValueId];
 
 			// Sort by value (descending).
-			return ValueA.Compare(ValueB) >= 0;
+			return ValueA.Compare(ValueB, ESearchCase::IgnoreCase) >= 0;
 		});
 	}
 
@@ -748,7 +748,7 @@ void FSorterByTextValueWithId::Sort(TArray<FBaseTreeNodePtr>& NodesToSort, ESort
 			const FString& ValueB = UniqueValues[B.ValueIndex];
 
 			// Sort by value (ascending).
-			return ValueA.Compare(ValueB) <= 0;
+			return ValueA.Compare(ValueB, ESearchCase::IgnoreCase) <= 0;
 		});
 	}
 	else // if (SortMode == ESortMode::Descending)
@@ -776,7 +776,7 @@ void FSorterByTextValueWithId::Sort(TArray<FBaseTreeNodePtr>& NodesToSort, ESort
 			const FString& ValueB = UniqueValues[B.ValueIndex];
 
 			// Sort by value (descending).
-			return ValueA.Compare(ValueB) >= 0;
+			return ValueA.Compare(ValueB, ESearchCase::IgnoreCase) >= 0;
 		});
 	}
 
diff --git a/Engine/Source/Developer/TraceInsights/Private/Insights/Table/Widgets/STableTreeView.cpp b/Engine/Source/Developer/TraceInsights/Private/Insights/Table/Widgets/STableTreeView.cpp
--- a/Engine/Source/Developer/TraceInsights/Private/Insights/Table/Widgets/STableTreeView.cpp
+++ b/Engine/Source/Developer/TraceInsights/Private/Insights/Table/Widgets/STableTreeView.cpp
@@ -1667,12 +1667,11 @@ void STableTreeView::UpdateCStringSameValueAggregationRec(FTableColumn& Column,
 			if (!NodePtr->IsGroup())
 			{
 				const TOptional<FTableCellValue> NodeValue = Column.GetValue(*NodePtr);
-				if (!NodeValue.IsSet() ||
-					NodeValue.GetValue().DataType != ETableCellDataType::CString ||
-					AggregatedValue != NodeValue.GetValue().CString)
+				if (NodeValue.IsSet() &&
+					NodeValue.GetValue().DataType == ETableCellDataType::CString &&
+					AggregatedValue == NodeValue.GetValue().CString)
 				{
-					AggregatedValue = nullptr;
-					break;
+					continue;
 				}
 			}
 			else
@@ -1681,14 +1680,16 @@ void STableTreeView::UpdateCStringSameValueAggregationRec(FTableColumn& Column,
 				if (TableNode.HasAggregatedValue(Column.GetId()))
 				{
 					const FTableCellValue& ChildGroupAggregatedValue = TableNode.GetAggregatedValue(Column.GetId());
-					if (ChildGroupAggregatedValue.DataType != ETableCellDataType::CString ||
-						AggregatedValue != ChildGroupAggregatedValue.CString)
+					if (ChildGroupAggregatedValue.DataType == ETableCellDataType::CString &&
+						AggregatedValue == ChildGroupAggregatedValue.CString)
 					{
-						AggregatedValue = nullptr;
-						break;
+						continue;
 					}
 				}
 			}
+
+			AggregatedValue = nullptr;
+			break;
 		}
 	}
 
@@ -3657,7 +3658,7 @@ FText STableTreeView::ViewPreset_GetSelectedToolTipText() const
 void STableTreeView::StopAllTableDataTasks(bool bWait)
 {
 	DataTaskInfos.RemoveAllSwap([](TSharedPtr<FTableTaskInfo> Data) { return Data->Event->IsComplete(); });
-	
+
 	for (int32 Index = 0; Index < DataTaskInfos.Num(); ++Index)
 	{
 		DataTaskInfos[Index]->CancellationToken->Cancel();
diff --git a/Engine/Source/Developer/TraceInsights/Private/Insights/Widgets/SModulesView.cpp b/Engine/Source/Developer/TraceInsights/Private/Insights/Widgets/SModulesView.cpp
--- a/Engine/Source/Developer/TraceInsights/Private/Insights/Widgets/SModulesView.cpp
+++ b/Engine/Source/Developer/TraceInsights/Private/Insights/Widgets/SModulesView.cpp
@@ -1015,7 +1015,7 @@ void SModulesView::UpdateSorting()
 					{
 						return false;
 					}
-					return FCString::Strcmp(A->GetModule()->Name, B->GetModule()->Name) < 0;
+					return FCString::Stricmp(A->GetModule()->Name, B->GetModule()->Name) < 0;
 				});
 		}
 		else
@@ -1030,7 +1030,7 @@ void SModulesView::UpdateSorting()
 					{
 						return true;
 					}
-					return FCString::Strcmp(A->GetModule()->Name, B->GetModule()->Name) > 0;
+					return FCString::Stricmp(A->GetModule()->Name, B->GetModule()->Name) > 0;
 				});
 		}
 	}
@@ -1113,7 +1113,7 @@ void SModulesView::UpdateSorting()
 					{
 						return false;
 					}
-					return FCString::Strcmp(A->GetModule()->StatusMessage, B->GetModule()->StatusMessage) < 0;
+					return FCString::Stricmp(A->GetModule()->StatusMessage, B->GetModule()->StatusMessage) < 0;
 				});
 		}
 		else
@@ -1128,7 +1128,7 @@ void SModulesView::UpdateSorting()
 					{
 						return true;
 					}
-					return FCString::Strcmp(A->GetModule()->StatusMessage, B->GetModule()->StatusMessage) > 0;
+					return FCString::Stricmp(A->GetModule()->StatusMessage, B->GetModule()->StatusMessage) > 0;
 				});
 		}
 	}
