[Insights] Session Browser: Fixed trace rename message displaying the updated name twice (i.e. "Renamed X to X.").

diff --git a/Engine/Source/Developer/TraceInsights/Private/Insights/Widgets/SStartPageWindow.cpp b/Engine/Source/Developer/TraceInsights/Private/Insights/Widgets/SStartPageWindow.cpp
--- a/Engine/Source/Developer/TraceInsights/Private/Insights/Widgets/SStartPageWindow.cpp
+++ b/Engine/Source/Developer/TraceInsights/Private/Insights/Widgets/SStartPageWindow.cpp
@@ -293,7 +293,7 @@ public:
 			}
 		}
 
-		FText Message = FText::Format(LOCTEXT("RenameSuccessFmt", "Renamed \"{0}\" to \"{1}\"."), Trace.Name, FText::FromString(NewTraceName));
+		FText Message = FText::Format(LOCTEXT("RenameSuccessFmt", "Renamed \"{0}\" to \"{1}\"."), FText::FromString(TraceName), FText::FromString(NewTraceName));
 		WeakParentWidget.Pin()->ShowSuccessMessage(Message);
 	}
 
