[Insights] Improvements to layout of Packages table.

diff --git a/Engine/Source/Developer/TraceInsights/Private/Insights/CookProfiler/ViewModels/PackageTable.cpp b/Engine/Source/Developer/TraceInsights/Private/Insights/CookProfiler/ViewModels/PackageTable.cpp
--- a/Engine/Source/Developer/TraceInsights/Private/Insights/CookProfiler/ViewModels/PackageTable.cpp
+++ b/Engine/Source/Developer/TraceInsights/Private/Insights/CookProfiler/ViewModels/PackageTable.cpp
@@ -128,7 +128,7 @@ void FPackageTable::AddDefaultColumns()
 		Column.SetFlags(ETableColumnFlags::ShouldBeVisible | ETableColumnFlags::CanBeHidden | ETableColumnFlags::CanBeFiltered);
 
 		Column.SetHorizontalAlignment(HAlign_Left);
-		Column.SetInitialWidth(100.0f);
+		Column.SetInitialWidth(80.0f);
 
 		Column.SetDataType(ETableCellDataType::Int64);
 
@@ -240,7 +240,7 @@ void FPackageTable::AddDefaultColumns()
 		Column.SetAggregation(ETableColumnAggregation::Sum);
 
 		AddColumn(ColumnRef);
-	}	
+	}
 	//////////////////////////////////////////////////
 	// BeginCacheForCookedPlatformData Time Column
 	{
@@ -288,7 +288,7 @@ void FPackageTable::AddDefaultColumns()
 		Column.SetFlags(ETableColumnFlags::ShouldBeVisible | ETableColumnFlags::CanBeHidden | ETableColumnFlags::CanBeFiltered);
 
 		Column.SetHorizontalAlignment(HAlign_Left);
-		Column.SetInitialWidth(100.0f);
+		Column.SetInitialWidth(200.0f);
 
 		Column.SetDataType(ETableCellDataType::CString);
 
@@ -306,7 +306,7 @@ void FPackageTable::AddDefaultColumns()
 		AddColumn(ColumnRef);
 	}
 	//////////////////////////////////////////////////
-	// Package name Column
+	// Package Name Column
 	{
 		TSharedRef<FTableColumn> ColumnRef = MakeShared<FTableColumn>(FPackageTableColumns::NameColumnId);
 		FTableColumn& Column = *ColumnRef;
@@ -314,13 +314,13 @@ void FPackageTable::AddDefaultColumns()
 		Column.SetIndex(ColumnIndex++);
 
 		Column.SetShortName(LOCTEXT("PackageNameColumnName", "Package Name"));
-		Column.SetTitleName(LOCTEXT("PackageNameTitle", "PackageName"));
+		Column.SetTitleName(LOCTEXT("PackageNameTitle", "Package Name"));
 		Column.SetDescription(LOCTEXT("PackageNameColumnDesc", "The name of the package."));
 
 		Column.SetFlags(ETableColumnFlags::CanBeHidden | ETableColumnFlags::CanBeFiltered);
 
 		Column.SetHorizontalAlignment(HAlign_Left);
-		Column.SetInitialWidth(100.0f);
+		Column.SetInitialWidth(400.0f);
 
 		Column.SetDataType(ETableCellDataType::CString);
 
diff --git a/Engine/Source/Developer/TraceInsights/Private/Insights/CookProfiler/Widgets/SPackageTableTreeView.cpp b/Engine/Source/Developer/TraceInsights/Private/Insights/CookProfiler/Widgets/SPackageTableTreeView.cpp
--- a/Engine/Source/Developer/TraceInsights/Private/Insights/CookProfiler/Widgets/SPackageTableTreeView.cpp
+++ b/Engine/Source/Developer/TraceInsights/Private/Insights/CookProfiler/Widgets/SPackageTableTreeView.cpp
@@ -300,13 +300,13 @@ void SPackageTableTreeView::InitAvailableViewPresets()
 		virtual void GetColumnConfigSet(TArray<FTableColumnConfig>& InOutConfigSet) const override
 		{
 			InOutConfigSet.Add({ FTable::GetHierarchyColumnId(), true, 500.0f });
-			InOutConfigSet.Add({ FPackageTableColumns::IdColumnId, true, 100.0f });
+			InOutConfigSet.Add({ FPackageTableColumns::IdColumnId, true, 80.0f });
 			InOutConfigSet.Add({ FPackageTableColumns::LoadTimeColumnId, true, 100.0f });
 			InOutConfigSet.Add({ FPackageTableColumns::SaveTimeColumnId, true, 100.0f });
 			InOutConfigSet.Add({ FPackageTableColumns::BeginCacheForCookedPlatformDataTimeColumnId, true, 100.0f });
 			InOutConfigSet.Add({ FPackageTableColumns::GetIsCachedCookedPlatformDataLoadedColumnId, true, 100.0f });
-			InOutConfigSet.Add({ FPackageTableColumns::PackageAssetClassColumnId, true, 100.0f });
-			InOutConfigSet.Add({ FPackageTableColumns::NameColumnId, false, 100.0f });
+			InOutConfigSet.Add({ FPackageTableColumns::PackageAssetClassColumnId, true, 200.0f });
+			InOutConfigSet.Add({ FPackageTableColumns::NameColumnId, false, 400.0f });
 		}
 	};
 	AvailableViewPresets.Add(MakeShared<FDefaultViewPreset>());
@@ -355,13 +355,13 @@ void SPackageTableTreeView::InitAvailableViewPresets()
 		virtual void GetColumnConfigSet(TArray<FTableColumnConfig>& InOutConfigSet) const override
 		{
 			InOutConfigSet.Add({ FTable::GetHierarchyColumnId(), true, 500.0f });
-			InOutConfigSet.Add({ FPackageTableColumns::IdColumnId, true, 100.0f });
+			InOutConfigSet.Add({ FPackageTableColumns::IdColumnId, true, 80.0f });
 			InOutConfigSet.Add({ FPackageTableColumns::LoadTimeColumnId, true, 100.0f });
 			InOutConfigSet.Add({ FPackageTableColumns::SaveTimeColumnId, true, 100.0f });
 			InOutConfigSet.Add({ FPackageTableColumns::BeginCacheForCookedPlatformDataTimeColumnId, true, 100.0f });
 			InOutConfigSet.Add({ FPackageTableColumns::GetIsCachedCookedPlatformDataLoadedColumnId, true, 100.0f });
-			InOutConfigSet.Add({ FPackageTableColumns::PackageAssetClassColumnId, true, 100.0f });
-			InOutConfigSet.Add({ FPackageTableColumns::NameColumnId, false, 100.0f });
+			InOutConfigSet.Add({ FPackageTableColumns::PackageAssetClassColumnId, true, 200.0f });
+			InOutConfigSet.Add({ FPackageTableColumns::NameColumnId, false, 400.0f });
 		}
 	};
 	AvailableViewPresets.Add(MakeShared<FPackagePathViewPreset>());
@@ -409,14 +409,14 @@ void SPackageTableTreeView::InitAvailableViewPresets()
 		}
 		virtual void GetColumnConfigSet(TArray<FTableColumnConfig>& InOutConfigSet) const override
 		{
-			InOutConfigSet.Add({ FTable::GetHierarchyColumnId(), true, 500.0f });
-			InOutConfigSet.Add({ FPackageTableColumns::IdColumnId, true, 100.0f });
+			InOutConfigSet.Add({ FTable::GetHierarchyColumnId(), true, 300.0f });
+			InOutConfigSet.Add({ FPackageTableColumns::IdColumnId, true, 80.0f });
 			InOutConfigSet.Add({ FPackageTableColumns::LoadTimeColumnId, true, 100.0f });
 			InOutConfigSet.Add({ FPackageTableColumns::SaveTimeColumnId, true, 100.0f });
 			InOutConfigSet.Add({ FPackageTableColumns::BeginCacheForCookedPlatformDataTimeColumnId, true, 100.0f });
 			InOutConfigSet.Add({ FPackageTableColumns::GetIsCachedCookedPlatformDataLoadedColumnId, true, 100.0f });
-			InOutConfigSet.Add({ FPackageTableColumns::NameColumnId, true, 100.0f });
-			InOutConfigSet.Add({ FPackageTableColumns::PackageAssetClassColumnId, false, 100.0f });
+			InOutConfigSet.Add({ FPackageTableColumns::NameColumnId, true, 400.0f });
+			InOutConfigSet.Add({ FPackageTableColumns::PackageAssetClassColumnId, false, 200.0f });
 		}
 	};
 	AvailableViewPresets.Add(MakeShared<FAssetClassViewPreset>());
