changed PRAGMA_DISABLE_OPTIMIZATION to UE_DISABLE_OPTIMIZATION_SHIP

diff --git a/Engine/Plugins/MovieScene/MovieRenderPipeline/Source/MovieRenderPipelineEditor/Private/Widgets/SMoviePipelineConfigEditor.cpp b/Engine/Plugins/MovieScene/MovieRenderPipeline/Source/MovieRenderPipelineEditor/Private/Widgets/SMoviePipelineConfigEditor.cpp
--- a/Engine/Plugins/MovieScene/MovieRenderPipeline/Source/MovieRenderPipelineEditor/Private/Widgets/SMoviePipelineConfigEditor.cpp
+++ b/Engine/Plugins/MovieScene/MovieRenderPipeline/Source/MovieRenderPipelineEditor/Private/Widgets/SMoviePipelineConfigEditor.cpp
@@ -450,7 +450,7 @@ protected:
 /**
  * Widget used to edit a Movie Render Pipeline Shot Config.
  */
-PRAGMA_DISABLE_OPTIMIZATION
+UE_DISABLE_OPTIMIZATION_SHIP
 void SMoviePipelineConfigEditor::Construct(const FArguments& InArgs)
 {
 	bRequestDetailsRefresh = true;
@@ -575,7 +575,7 @@ TSharedRef<SWidget> SMoviePipelineConfigEditor::MakeAddSettingButton()
 		.Icon(FAppStyle::Get().GetBrush("Icons.Plus"))
 		.Text(LOCTEXT("AddNewSetting_Text", "Setting"));
 }
-PRAGMA_ENABLE_OPTIMIZATION
+UE_ENABLE_OPTIMIZATION_SHIP
 
 void SMoviePipelineConfigEditor::OnBlueprintReinstanced()
 {
diff --git a/Engine/Plugins/MovieScene/MovieRenderPipeline/Source/MovieRenderPipelineEditor/Private/Widgets/SMoviePipelineConfigPanel.cpp b/Engine/Plugins/MovieScene/MovieRenderPipeline/Source/MovieRenderPipelineEditor/Private/Widgets/SMoviePipelineConfigPanel.cpp
--- a/Engine/Plugins/MovieScene/MovieRenderPipeline/Source/MovieRenderPipelineEditor/Private/Widgets/SMoviePipelineConfigPanel.cpp
+++ b/Engine/Plugins/MovieScene/MovieRenderPipeline/Source/MovieRenderPipelineEditor/Private/Widgets/SMoviePipelineConfigPanel.cpp
@@ -46,7 +46,7 @@ SMoviePipelineConfigPanel::~SMoviePipelineConfigPanel()
 {
 }
 
-PRAGMA_DISABLE_OPTIMIZATION
+UE_DISABLE_OPTIMIZATION_SHIP
 void SMoviePipelineConfigPanel::Construct(const FArguments& InArgs, TSubclassOf<UMoviePipelineConfigBase> InConfigType)
 {
 	ConfigAssetType = InConfigType;
@@ -237,7 +237,7 @@ void SMoviePipelineConfigPanel::Construct(const FArguments& InArgs, TSubclassOf<
 	}
 }
 
-PRAGMA_ENABLE_OPTIMIZATION
+UE_ENABLE_OPTIMIZATION_SHIP
 
 FText SMoviePipelineConfigPanel::GetConfigTypeLabel() const
 {
diff --git a/Engine/Plugins/MovieScene/MovieRenderPipeline/Source/MovieRenderPipelineEditor/Private/Widgets/SMoviePipelineQueueEditor.cpp b/Engine/Plugins/MovieScene/MovieRenderPipeline/Source/MovieRenderPipelineEditor/Private/Widgets/SMoviePipelineQueueEditor.cpp
--- a/Engine/Plugins/MovieScene/MovieRenderPipeline/Source/MovieRenderPipelineEditor/Private/Widgets/SMoviePipelineQueueEditor.cpp
+++ b/Engine/Plugins/MovieScene/MovieRenderPipeline/Source/MovieRenderPipelineEditor/Private/Widgets/SMoviePipelineQueueEditor.cpp
@@ -1000,7 +1000,7 @@ TSharedRef<SWidget> SQueueShotListRow::GenerateWidgetForColumn(const FName& Colu
 }
 
 
-PRAGMA_DISABLE_OPTIMIZATION
+UE_DISABLE_OPTIMIZATION_SHIP
 void SMoviePipelineQueueEditor::Construct(const FArguments& InArgs)
 {
 	CachedQueueSerialNumber = uint32(-1);
@@ -1150,7 +1150,7 @@ TSharedRef<SWidget> SMoviePipelineQueueEditor::OnGenerateNewJobFromAssetMenu()
 	return MenuBuilder.MakeWidget();
 }
 
-PRAGMA_ENABLE_OPTIMIZATION
+UE_ENABLE_OPTIMIZATION_SHIP
 
 void SMoviePipelineQueueEditor::OnCreateJobFromAsset(const FAssetData& InAsset)
 {
diff --git a/Engine/Plugins/MovieScene/MovieRenderPipeline/Source/MovieRenderPipelineEditor/Private/Widgets/SMoviePipelineQueuePanel.cpp b/Engine/Plugins/MovieScene/MovieRenderPipeline/Source/MovieRenderPipelineEditor/Private/Widgets/SMoviePipelineQueuePanel.cpp
--- a/Engine/Plugins/MovieScene/MovieRenderPipeline/Source/MovieRenderPipelineEditor/Private/Widgets/SMoviePipelineQueuePanel.cpp
+++ b/Engine/Plugins/MovieScene/MovieRenderPipeline/Source/MovieRenderPipelineEditor/Private/Widgets/SMoviePipelineQueuePanel.cpp
@@ -52,7 +52,7 @@
 
 #define LOCTEXT_NAMESPACE "SMoviePipelineQueuePanel"
 
-PRAGMA_DISABLE_OPTIMIZATION
+UE_DISABLE_OPTIMIZATION_SHIP
 void SMoviePipelineQueuePanel::Construct(const FArguments& InArgs)
 {
 	FPropertyEditorModule& PropertyEditorModule = FModuleManager::Get().LoadModuleChecked<FPropertyEditorModule>("PropertyEditor");
@@ -257,7 +257,7 @@ void SMoviePipelineQueuePanel::Construct(const FArguments& InArgs)
 	];
 }
 
-PRAGMA_ENABLE_OPTIMIZATION
+UE_ENABLE_OPTIMIZATION_SHIP
 
 FReply SMoviePipelineQueuePanel::OnRenderLocalRequested()
 {
diff --git a/Engine/Plugins/Runtime/NetworkPredictionInsights/Source/NetworkPredictionInsights/Private/UI/NetworkPredictionInsightsCommands.cpp b/Engine/Plugins/Runtime/NetworkPredictionInsights/Source/NetworkPredictionInsights/Private/UI/NetworkPredictionInsightsCommands.cpp
--- a/Engine/Plugins/Runtime/NetworkPredictionInsights/Source/NetworkPredictionInsights/Private/UI/NetworkPredictionInsightsCommands.cpp
+++ b/Engine/Plugins/Runtime/NetworkPredictionInsights/Source/NetworkPredictionInsights/Private/UI/NetworkPredictionInsightsCommands.cpp
@@ -16,7 +16,8 @@ FNetworkPredictionInsightsCommands::FNetworkPredictionInsightsCommands()
 }
 
 // UI_COMMAND takes long for the compiler to optimize
-PRAGMA_DISABLE_OPTIMIZATION
+UE_DISABLE_OPTIMIZATION_SHIP
+
 void FNetworkPredictionInsightsCommands::RegisterCommands()
 {
 	UI_COMMAND(ToggleAutoScrollSimulationFrames, "AutoScroll", "Toggles auto scrolling in the simulation frame view when new data is received.", EUserInterfaceActionType::ToggleButton, FInputChord(EKeys::SpaceBar));
@@ -25,6 +26,6 @@ void FNetworkPredictionInsightsCommands::RegisterCommands()
 	UI_COMMAND(FirstEngineFrame, "<<", "Jump to the first frame in the Simulation View", EUserInterfaceActionType::Button, FInputChord(EKeys::Home));
 	UI_COMMAND(LastEngineFrame, ">>", "Jump to the last frame in the Simulation View", EUserInterfaceActionType::Button, FInputChord(EKeys::End));
 }
-PRAGMA_ENABLE_OPTIMIZATION
+UE_ENABLE_OPTIMIZATION_SHIP
 
 #undef LOCTEXT_NAMESPACE
\ No newline at end of file
diff --git a/Engine/Plugins/VirtualProduction/Takes/Source/TakeRecorder/Private/Widgets/SLevelSequenceTakeEditor.cpp b/Engine/Plugins/VirtualProduction/Takes/Source/TakeRecorder/Private/Widgets/SLevelSequenceTakeEditor.cpp
--- a/Engine/Plugins/VirtualProduction/Takes/Source/TakeRecorder/Private/Widgets/SLevelSequenceTakeEditor.cpp
+++ b/Engine/Plugins/VirtualProduction/Takes/Source/TakeRecorder/Private/Widgets/SLevelSequenceTakeEditor.cpp
@@ -104,7 +104,7 @@ TArray<UClass*> FindRecordingSourceClasses()
 	return Classes;
 }
 
-PRAGMA_DISABLE_OPTIMIZATION
+UE_DISABLE_OPTIMIZATION_SHIP
 void SLevelSequenceTakeEditor::Construct(const FArguments& InArgs)
 {
 	bRequestDetailsRefresh = true;
@@ -149,7 +149,7 @@ TSharedRef<SWidget> SLevelSequenceTakeEditor::MakeAddSourceButton()
 		.Icon(FAppStyle::Get().GetBrush("Icons.Plus"))
 		.Text(LOCTEXT("AddNewSource_Text", "Source"));
 }
-PRAGMA_ENABLE_OPTIMIZATION
+UE_ENABLE_OPTIMIZATION_SHIP
 
 void SLevelSequenceTakeEditor::AddExternalSettingsObject(UObject* InObject)
 {
diff --git a/Engine/Plugins/VirtualProduction/Takes/Source/TakeRecorder/Private/Widgets/STakePresetAssetEditor.cpp b/Engine/Plugins/VirtualProduction/Takes/Source/TakeRecorder/Private/Widgets/STakePresetAssetEditor.cpp
--- a/Engine/Plugins/VirtualProduction/Takes/Source/TakeRecorder/Private/Widgets/STakePresetAssetEditor.cpp
+++ b/Engine/Plugins/VirtualProduction/Takes/Source/TakeRecorder/Private/Widgets/STakePresetAssetEditor.cpp
@@ -30,7 +30,7 @@ STakePresetAssetEditor::~STakePresetAssetEditor()
 	Toolkit.Reset();
 }
 
-PRAGMA_DISABLE_OPTIMIZATION
+UE_DISABLE_OPTIMIZATION_SHIP
 void STakePresetAssetEditor::Construct(const FArguments& InArgs, TSharedPtr<FTakePresetToolkit> InToolkit, TWeakPtr<STakeRecorderTabContent> OuterTabContent)
 {
 	WeakTabContent = OuterTabContent;
@@ -117,7 +117,7 @@ void STakePresetAssetEditor::Construct(const FArguments& InArgs, TSharedPtr<FTak
 		]
 	];
 }
-PRAGMA_ENABLE_OPTIMIZATION
+UE_ENABLE_OPTIMIZATION_SHIP
 
 ULevelSequence* STakePresetAssetEditor::GetLevelSequence() const
 {
diff --git a/Engine/Plugins/VirtualProduction/Takes/Source/TakeRecorder/Private/Widgets/STakeRecorderPanel.cpp b/Engine/Plugins/VirtualProduction/Takes/Source/TakeRecorder/Private/Widgets/STakeRecorderPanel.cpp
--- a/Engine/Plugins/VirtualProduction/Takes/Source/TakeRecorder/Private/Widgets/STakeRecorderPanel.cpp
+++ b/Engine/Plugins/VirtualProduction/Takes/Source/TakeRecorder/Private/Widgets/STakeRecorderPanel.cpp
@@ -77,7 +77,7 @@ STakeRecorderPanel::~STakeRecorderPanel()
 	TakeRecorderModule.GetExternalObjectAddRemoveEventDelegate().Remove(OnWidgetExternalObjectChangedHandle);
 }
 
-PRAGMA_DISABLE_OPTIMIZATION
+UE_DISABLE_OPTIMIZATION_SHIP
 void STakeRecorderPanel::Construct(const FArguments& InArgs)
 {
 	// If a recording is currently underway, initialize to that now
@@ -530,7 +530,7 @@ TSharedRef<SWidget> STakeRecorderPanel::MakeToolBar()
 
 	return Border;
 }
-PRAGMA_ENABLE_OPTIMIZATION
+UE_ENABLE_OPTIMIZATION_SHIP
 
 ULevelSequence* STakeRecorderPanel::GetLevelSequence() const
 {
diff --git a/Engine/Source/Developer/Profiler/Private/ProfilerCommands.cpp b/Engine/Source/Developer/Profiler/Private/ProfilerCommands.cpp
--- a/Engine/Source/Developer/Profiler/Private/ProfilerCommands.cpp
+++ b/Engine/Source/Developer/Profiler/Private/ProfilerCommands.cpp
@@ -32,7 +32,7 @@ FProfilerCommands::FProfilerCommands()
 
 
 /** UI_COMMAND takes long for the compile to optimize */
-PRAGMA_DISABLE_OPTIMIZATION
+UE_DISABLE_OPTIMIZATION_SHIP
 void FProfilerCommands::RegisterCommands()
 {
 	/*-----------------------------------------------------------------------------
@@ -71,7 +71,7 @@ void FProfilerCommands::RegisterCommands()
 
 	UI_COMMAND( EventGraph_SelectAllFrames, "Select all frames", "Selects all frames in the data graph and displays them in the event graph", EUserInterfaceActionType::Button, FInputChord() );
 }
-PRAGMA_ENABLE_OPTIMIZATION
+UE_ENABLE_OPTIMIZATION_SHIP
 
 //static_assert(sizeof(FProfilerActionManager) == 0, "Cannot contain any variables at this moment.");
 
diff --git a/Engine/Source/Developer/ProjectLauncher/Private/Models/ProjectLauncherCommands.h b/Engine/Source/Developer/ProjectLauncher/Private/Models/ProjectLauncherCommands.h
--- a/Engine/Source/Developer/ProjectLauncher/Private/Models/ProjectLauncherCommands.h
+++ b/Engine/Source/Developer/ProjectLauncher/Private/Models/ProjectLauncherCommands.h
@@ -34,7 +34,7 @@ public:
 
 	//~ TCommands interface
 
-	PRAGMA_DISABLE_OPTIMIZATION
+	UE_DISABLE_OPTIMIZATION_SHIP
 	virtual void RegisterCommands() override
 	{
 		UI_COMMAND(QuickLaunch, "Quick Launch", "Builds, cooks, and launches a build.", EUserInterfaceActionType::ToggleButton, FInputChord(EModifierKey::Control, EKeys::L));
@@ -46,7 +46,7 @@ public:
 		UI_COMMAND(DuplicateProfile, "Duplicate", "Duplicate the launch profile.", EUserInterfaceActionType::Button, FInputChord());//, FInputChord(EModifierKey::Control, EKeys::W));
 		UI_COMMAND(DeleteProfile, "Delete", "Delete the launch profile.", EUserInterfaceActionType::Button, FInputChord());//, FInputChord(EModifierKey::Control, EKeys::Delete));
 	}
-	PRAGMA_ENABLE_OPTIMIZATION
+	UE_ENABLE_OPTIMIZATION_SHIP
 
 public:
 
diff --git a/Engine/Source/Developer/TraceInsights/Private/Insights/Common/TimeUtils.cpp b/Engine/Source/Developer/TraceInsights/Private/Insights/Common/TimeUtils.cpp
--- a/Engine/Source/Developer/TraceInsights/Private/Insights/Common/TimeUtils.cpp
+++ b/Engine/Source/Developer/TraceInsights/Private/Insights/Common/TimeUtils.cpp
@@ -837,7 +837,7 @@ int32 GetNumDigits(const double Precision)
 	return D;
 }
 
-PRAGMA_DISABLE_OPTIMIZATION
+UE_DISABLE_OPTIMIZATION_SHIP
 int32 GetNumDigitsOptDisabled(const double Precision)
 {
 	//double Log10 = -log10(Precision);
@@ -845,9 +845,9 @@ int32 GetNumDigitsOptDisabled(const double Precision)
 	int32 D = (Log10 > 0) ? FMath::CeilToInt(Log10) : 0;
 	return D;
 }
-PRAGMA_ENABLE_OPTIMIZATION
+UE_ENABLE_OPTIMIZATION_SHIP
 
-PRAGMA_DISABLE_OPTIMIZATION
+UE_DISABLE_OPTIMIZATION_SHIP
 void TestOptimizationIssue()
 {
 	constexpr double Ns = 0.000000001;
@@ -857,7 +857,7 @@ void TestOptimizationIssue()
 	ensure(D1 == 9); // 10 ?
 	ensure(D1 == D2);
 }
-PRAGMA_ENABLE_OPTIMIZATION
+UE_ENABLE_OPTIMIZATION_SHIP
 
 ////////////////////////////////////////////////////////////////////////////////////////////////////
 
diff --git a/Engine/Source/Developer/TraceInsights/Private/Insights/ContextSwitches/ViewModels/ContextSwitchesSharedState.cpp b/Engine/Source/Developer/TraceInsights/Private/Insights/ContextSwitches/ViewModels/ContextSwitchesSharedState.cpp
--- a/Engine/Source/Developer/TraceInsights/Private/Insights/ContextSwitches/ViewModels/ContextSwitchesSharedState.cpp
+++ b/Engine/Source/Developer/TraceInsights/Private/Insights/ContextSwitches/ViewModels/ContextSwitchesSharedState.cpp
@@ -50,7 +50,7 @@ FContextSwitchesStateCommands::~FContextSwitchesStateCommands()
 ////////////////////////////////////////////////////////////////////////////////////////////////////
 
 // UI_COMMAND takes long for the compiler to optimize
-PRAGMA_DISABLE_OPTIMIZATION
+UE_DISABLE_OPTIMIZATION_SHIP
 void FContextSwitchesStateCommands::RegisterCommands()
 {
 	UI_COMMAND(Command_ShowCoreTracks,
@@ -101,7 +101,7 @@ void FContextSwitchesStateCommands::RegisterCommands()
 		EUserInterfaceActionType::Button,
 		FInputChord());
 }
-PRAGMA_ENABLE_OPTIMIZATION
+UE_ENABLE_OPTIMIZATION_SHIP
 
 ////////////////////////////////////////////////////////////////////////////////////////////////////
 // FContextSwitchesSharedState
diff --git a/Engine/Source/Developer/TraceInsights/Private/Insights/InsightsCommands.cpp b/Engine/Source/Developer/TraceInsights/Private/Insights/InsightsCommands.cpp
--- a/Engine/Source/Developer/TraceInsights/Private/Insights/InsightsCommands.cpp
+++ b/Engine/Source/Developer/TraceInsights/Private/Insights/InsightsCommands.cpp
@@ -31,7 +31,7 @@ FInsightsCommands::FInsightsCommands()
 ////////////////////////////////////////////////////////////////////////////////////////////////////
 
 // UI_COMMAND takes long for the compile to optimize
-PRAGMA_DISABLE_OPTIMIZATION
+UE_DISABLE_OPTIMIZATION_SHIP
 void FInsightsCommands::RegisterCommands()
 {
 	UI_COMMAND(InsightsManager_Load,
@@ -52,7 +52,7 @@ void FInsightsCommands::RegisterCommands()
 		EUserInterfaceActionType::Button,
 		FInputChord(EModifierKey::Control, EKeys::O));
 }
-PRAGMA_ENABLE_OPTIMIZATION
+UE_ENABLE_OPTIMIZATION_SHIP
 
 ////////////////////////////////////////////////////////////////////////////////////////////////////
 // ToggleDebugInfo
diff --git a/Engine/Source/Developer/TraceInsights/Private/Insights/LoadingProfiler/LoadingProfilerCommands.cpp b/Engine/Source/Developer/TraceInsights/Private/Insights/LoadingProfiler/LoadingProfilerCommands.cpp
--- a/Engine/Source/Developer/TraceInsights/Private/Insights/LoadingProfiler/LoadingProfilerCommands.cpp
+++ b/Engine/Source/Developer/TraceInsights/Private/Insights/LoadingProfiler/LoadingProfilerCommands.cpp
@@ -48,7 +48,7 @@ FLoadingProfilerCommands::FLoadingProfilerCommands()
 ////////////////////////////////////////////////////////////////////////////////////////////////////
 
 // UI_COMMAND takes long for the compiler to optimize
-PRAGMA_DISABLE_OPTIMIZATION
+UE_DISABLE_OPTIMIZATION_SHIP
 void FLoadingProfilerCommands::RegisterCommands()
 {
 	UI_COMMAND(ToggleTimingViewVisibility,
@@ -87,7 +87,7 @@ void FLoadingProfilerCommands::RegisterCommands()
 		EUserInterfaceActionType::ToggleButton,
 		FInputChord());
 }
-PRAGMA_ENABLE_OPTIMIZATION
+UE_ENABLE_OPTIMIZATION_SHIP
 
 ////////////////////////////////////////////////////////////////////////////////////////////////////
 // Toggle Commands
diff --git a/Engine/Source/Developer/TraceInsights/Private/Insights/MemoryProfiler/MemoryProfilerCommands.cpp b/Engine/Source/Developer/TraceInsights/Private/Insights/MemoryProfiler/MemoryProfilerCommands.cpp
--- a/Engine/Source/Developer/TraceInsights/Private/Insights/MemoryProfiler/MemoryProfilerCommands.cpp
+++ b/Engine/Source/Developer/TraceInsights/Private/Insights/MemoryProfiler/MemoryProfilerCommands.cpp
@@ -54,7 +54,7 @@ FMemoryProfilerCommands::~FMemoryProfilerCommands()
 ////////////////////////////////////////////////////////////////////////////////////////////////////
 
 // UI_COMMAND takes long for the compiler to optimize
-PRAGMA_DISABLE_OPTIMIZATION
+UE_DISABLE_OPTIMIZATION_SHIP
 void FMemoryProfilerCommands::RegisterCommands()
 {
 	UI_COMMAND(ToggleTimingViewVisibility,
@@ -81,7 +81,7 @@ void FMemoryProfilerCommands::RegisterCommands()
 		EUserInterfaceActionType::ToggleButton,
 		FInputChord());
 }
-PRAGMA_ENABLE_OPTIMIZATION
+UE_ENABLE_OPTIMIZATION_SHIP
 
 ////////////////////////////////////////////////////////////////////////////////////////////////////
 // Toggle Commands
diff --git a/Engine/Source/Developer/TraceInsights/Private/Insights/MemoryProfiler/ViewModels/MemorySharedState.cpp b/Engine/Source/Developer/TraceInsights/Private/Insights/MemoryProfiler/ViewModels/MemorySharedState.cpp
--- a/Engine/Source/Developer/TraceInsights/Private/Insights/MemoryProfiler/ViewModels/MemorySharedState.cpp
+++ b/Engine/Source/Developer/TraceInsights/Private/Insights/MemoryProfiler/ViewModels/MemorySharedState.cpp
@@ -53,7 +53,7 @@ FMemoryTimingViewCommands::~FMemoryTimingViewCommands()
 ////////////////////////////////////////////////////////////////////////////////////////////////////
 
 // UI_COMMAND takes long for the compiler to optimize
-PRAGMA_DISABLE_OPTIMIZATION
+UE_DISABLE_OPTIMIZATION_SHIP
 void FMemoryTimingViewCommands::RegisterCommands()
 {
 	UI_COMMAND(ShowHideAllMemoryTracks,
@@ -62,7 +62,7 @@ void FMemoryTimingViewCommands::RegisterCommands()
 		EUserInterfaceActionType::ToggleButton,
 		FInputChord(EModifierKey::Control, EKeys::M));
 }
-PRAGMA_ENABLE_OPTIMIZATION
+UE_ENABLE_OPTIMIZATION_SHIP
 
 ////////////////////////////////////////////////////////////////////////////////////////////////////
 // FMemorySharedState
diff --git a/Engine/Source/Developer/TraceInsights/Private/Insights/NetworkingProfiler/NetworkingProfilerCommands.cpp b/Engine/Source/Developer/TraceInsights/Private/Insights/NetworkingProfiler/NetworkingProfilerCommands.cpp
--- a/Engine/Source/Developer/TraceInsights/Private/Insights/NetworkingProfiler/NetworkingProfilerCommands.cpp
+++ b/Engine/Source/Developer/TraceInsights/Private/Insights/NetworkingProfiler/NetworkingProfilerCommands.cpp
@@ -48,7 +48,7 @@ FNetworkingProfilerCommands::FNetworkingProfilerCommands()
 ////////////////////////////////////////////////////////////////////////////////////////////////////
 
 // UI_COMMAND takes long for the compiler to optimize
-PRAGMA_DISABLE_OPTIMIZATION
+UE_DISABLE_OPTIMIZATION_SHIP
 void FNetworkingProfilerCommands::RegisterCommands()
 {
 	UI_COMMAND(TogglePacketViewVisibility,
@@ -75,7 +75,7 @@ void FNetworkingProfilerCommands::RegisterCommands()
 		EUserInterfaceActionType::ToggleButton,
 		FInputChord());
 }
-PRAGMA_ENABLE_OPTIMIZATION
+UE_ENABLE_OPTIMIZATION_SHIP
 
 ////////////////////////////////////////////////////////////////////////////////////////////////////
 // Toggle Commands
diff --git a/Engine/Source/Developer/TraceInsights/Private/Insights/Table/Widgets/STableTreeView.cpp b/Engine/Source/Developer/TraceInsights/Private/Insights/Table/Widgets/STableTreeView.cpp
--- a/Engine/Source/Developer/TraceInsights/Private/Insights/Table/Widgets/STableTreeView.cpp
+++ b/Engine/Source/Developer/TraceInsights/Private/Insights/Table/Widgets/STableTreeView.cpp
@@ -62,7 +62,7 @@ public:
 	}
 
 	// UI_COMMAND takes long for the compiler to optimize
-	PRAGMA_DISABLE_OPTIMIZATION
+	UE_DISABLE_OPTIMIZATION_SHIP
 	virtual void RegisterCommands() override
 	{
 		UI_COMMAND(Command_CopyToClipboard, "Copy To Clipboard", "Copies the selection to clipboard.", EUserInterfaceActionType::Button, FInputChord(EModifierKey::Control, EKeys::C));
@@ -75,7 +75,7 @@ public:
 		UI_COMMAND(Command_ExportEntireTreeToFile, "Export Entire Tree to File...", "Exports the entire tree/table content to a file. It exports also the collapsed tree nodes, including the leaf nodes. Filtered out nodes are not exported.", EUserInterfaceActionType::Button, FInputChord(EModifierKey::Shift, EKeys::E));
 		UI_COMMAND(Command_ExportEntireTreeToFileNoLeafs, "Export Entire Tree (No Leafs) to File...", "Exports the entire tree/table content to a file, but not the leaf nodes. It exports the collapsed tree nodes. Filtered out nodes are not exported.", EUserInterfaceActionType::Button, FInputChord(EModifierKey::Alt, EKeys::E));
 	}
-	PRAGMA_ENABLE_OPTIMIZATION
+	UE_ENABLE_OPTIMIZATION_SHIP
 
 	TSharedPtr<FUICommandInfo> Command_CopyToClipboard;
 	TSharedPtr<FUICommandInfo> Command_CopyColumnToClipboard;
diff --git a/Engine/Source/Developer/TraceInsights/Private/Insights/TaskGraphProfiler/ViewModels/TaskTimingTrack.cpp b/Engine/Source/Developer/TraceInsights/Private/Insights/TaskGraphProfiler/ViewModels/TaskTimingTrack.cpp
--- a/Engine/Source/Developer/TraceInsights/Private/Insights/TaskGraphProfiler/ViewModels/TaskTimingTrack.cpp
+++ b/Engine/Source/Developer/TraceInsights/Private/Insights/TaskGraphProfiler/ViewModels/TaskTimingTrack.cpp
@@ -48,7 +48,7 @@ public:
 	}
 
 	// UI_COMMAND takes long for the compiler to optimize
-	PRAGMA_DISABLE_OPTIMIZATION
+	UE_DISABLE_OPTIMIZATION_SHIP
 	virtual void RegisterCommands() override
 	{
 		UI_COMMAND(Command_ShowTaskTransitions,
@@ -104,7 +104,7 @@ public:
 				   EUserInterfaceActionType::ToggleButton,
 				   FInputChord());
 	}
-	PRAGMA_ENABLE_OPTIMIZATION
+	UE_ENABLE_OPTIMIZATION_SHIP
 
 	TSharedPtr<FUICommandInfo> Command_ShowTaskTransitions;
 	TSharedPtr<FUICommandInfo> Command_ShowTaskConnections;
diff --git a/Engine/Source/Developer/TraceInsights/Private/Insights/TaskGraphProfiler/Widgets/STaskTableTreeView.cpp b/Engine/Source/Developer/TraceInsights/Private/Insights/TaskGraphProfiler/Widgets/STaskTableTreeView.cpp
--- a/Engine/Source/Developer/TraceInsights/Private/Insights/TaskGraphProfiler/Widgets/STaskTableTreeView.cpp
+++ b/Engine/Source/Developer/TraceInsights/Private/Insights/TaskGraphProfiler/Widgets/STaskTableTreeView.cpp
@@ -56,13 +56,13 @@ public:
 	}
 
 	// UI_COMMAND takes long for the compiler to optimize
-	PRAGMA_DISABLE_OPTIMIZATION
+	UE_DISABLE_OPTIMIZATION_SHIP
 	virtual void RegisterCommands() override
 	{
 		UI_COMMAND(Command_GoToTask, "Go To Task", "Pan and zoom to the task in Timing View.", EUserInterfaceActionType::Button, FInputChord());
 		UI_COMMAND(Command_OpenInIDE, "Open in IDE", "Open the source location where the selected task was launched in IDE.", EUserInterfaceActionType::Button, FInputChord());
 	}
-	PRAGMA_ENABLE_OPTIMIZATION
+	UE_ENABLE_OPTIMIZATION_SHIP
 
 	TSharedPtr<FUICommandInfo> Command_GoToTask;
 	TSharedPtr<FUICommandInfo> Command_OpenInIDE;
diff --git a/Engine/Source/Developer/TraceInsights/Private/Insights/TimingProfilerCommands.cpp b/Engine/Source/Developer/TraceInsights/Private/Insights/TimingProfilerCommands.cpp
--- a/Engine/Source/Developer/TraceInsights/Private/Insights/TimingProfilerCommands.cpp
+++ b/Engine/Source/Developer/TraceInsights/Private/Insights/TimingProfilerCommands.cpp
@@ -54,7 +54,7 @@ FTimingProfilerCommands::~FTimingProfilerCommands()
 ////////////////////////////////////////////////////////////////////////////////////////////////////
 
 // UI_COMMAND takes long for the compiler to optimize
-PRAGMA_DISABLE_OPTIMIZATION
+UE_DISABLE_OPTIMIZATION_SHIP
 void FTimingProfilerCommands::RegisterCommands()
 {
 	UI_COMMAND(ToggleFramesTrackVisibility, "Frames", "Toggles the visibility of the Frames track.", EUserInterfaceActionType::ToggleButton, FInputChord());
@@ -65,7 +65,7 @@ void FTimingProfilerCommands::RegisterCommands()
 	UI_COMMAND(ToggleStatsCountersViewVisibility, "Counters", "Toggles the visibility of the Counters view.", EUserInterfaceActionType::ToggleButton, FInputChord());
 	UI_COMMAND(ToggleLogViewVisibility, "Log", "Toggles the visibility of the Log view.", EUserInterfaceActionType::ToggleButton,FInputChord());
 }
-PRAGMA_ENABLE_OPTIMIZATION
+UE_ENABLE_OPTIMIZATION_SHIP
 
 ////////////////////////////////////////////////////////////////////////////////////////////////////
 // FTimingViewCommands
@@ -89,7 +89,7 @@ FTimingViewCommands::~FTimingViewCommands()
 ////////////////////////////////////////////////////////////////////////////////////////////////////
 
 // UI_COMMAND takes long for the compiler to optimize
-PRAGMA_DISABLE_OPTIMIZATION
+UE_DISABLE_OPTIMIZATION_SHIP
 void FTimingViewCommands::RegisterCommands()
 {
 	UI_COMMAND(AutoHideEmptyTracks,
@@ -122,7 +122,7 @@ void FTimingViewCommands::RegisterCommands()
 		EUserInterfaceActionType::Button,
 		FInputChord(EModifierKey::Control, EKeys::F));
 }
-PRAGMA_ENABLE_OPTIMIZATION
+UE_ENABLE_OPTIMIZATION_SHIP
 
 ////////////////////////////////////////////////////////////////////////////////////////////////////
 // Toggle Commands
diff --git a/Engine/Source/Developer/TraceInsights/Private/Insights/ViewModels/FileActivityTimingTrack.cpp b/Engine/Source/Developer/TraceInsights/Private/Insights/ViewModels/FileActivityTimingTrack.cpp
--- a/Engine/Source/Developer/TraceInsights/Private/Insights/ViewModels/FileActivityTimingTrack.cpp
+++ b/Engine/Source/Developer/TraceInsights/Private/Insights/ViewModels/FileActivityTimingTrack.cpp
@@ -51,7 +51,7 @@ FFileActivityTimingViewCommands::~FFileActivityTimingViewCommands()
 ////////////////////////////////////////////////////////////////////////////////////////////////////
 
 // UI_COMMAND takes long for the compiler to optimize
-PRAGMA_DISABLE_OPTIMIZATION
+UE_DISABLE_OPTIMIZATION_SHIP
 void FFileActivityTimingViewCommands::RegisterCommands()
 {
 	// This command is used only for its key binding (to toggle both ShowHideIoOverviewTrack and ShowHideIoActivityTrack in the same time).
@@ -85,7 +85,7 @@ void FFileActivityTimingViewCommands::RegisterCommands()
 		EUserInterfaceActionType::ToggleButton,
 		FInputChord(EKeys::O));
 }
-PRAGMA_ENABLE_OPTIMIZATION
+UE_ENABLE_OPTIMIZATION_SHIP
 
 ////////////////////////////////////////////////////////////////////////////////////////////////////
 ////////////////////////////////////////////////////////////////////////////////////////////////////
diff --git a/Engine/Source/Developer/TraceInsights/Private/Insights/ViewModels/FrameTimingTrack.cpp b/Engine/Source/Developer/TraceInsights/Private/Insights/ViewModels/FrameTimingTrack.cpp
--- a/Engine/Source/Developer/TraceInsights/Private/Insights/ViewModels/FrameTimingTrack.cpp
+++ b/Engine/Source/Developer/TraceInsights/Private/Insights/ViewModels/FrameTimingTrack.cpp
@@ -47,7 +47,7 @@ FFrameTimingViewCommands::~FFrameTimingViewCommands()
 ////////////////////////////////////////////////////////////////////////////////////////////////////
 
 // UI_COMMAND takes long for the compiler to optimize
-PRAGMA_DISABLE_OPTIMIZATION
+UE_DISABLE_OPTIMIZATION_SHIP
 void FFrameTimingViewCommands::RegisterCommands()
 {
 	UI_COMMAND(ShowHideAllFrameTracks,
@@ -56,7 +56,7 @@ void FFrameTimingViewCommands::RegisterCommands()
 		EUserInterfaceActionType::ToggleButton,
 		FInputChord(EKeys::R));
 }
-PRAGMA_ENABLE_OPTIMIZATION
+UE_ENABLE_OPTIMIZATION_SHIP
 
 ////////////////////////////////////////////////////////////////////////////////////////////////////
 // FFrameSharedState
diff --git a/Engine/Source/Developer/TraceInsights/Private/Insights/ViewModels/LoadingTimingTrack.cpp b/Engine/Source/Developer/TraceInsights/Private/Insights/ViewModels/LoadingTimingTrack.cpp
--- a/Engine/Source/Developer/TraceInsights/Private/Insights/ViewModels/LoadingTimingTrack.cpp
+++ b/Engine/Source/Developer/TraceInsights/Private/Insights/ViewModels/LoadingTimingTrack.cpp
@@ -45,7 +45,7 @@ FLoadingTimingViewCommands::~FLoadingTimingViewCommands()
 ////////////////////////////////////////////////////////////////////////////////////////////////////
 
 // UI_COMMAND takes long for the compiler to optimize
-PRAGMA_DISABLE_OPTIMIZATION
+UE_DISABLE_OPTIMIZATION_SHIP
 void FLoadingTimingViewCommands::RegisterCommands()
 {
 	UI_COMMAND(ShowHideAllLoadingTracks,
@@ -54,7 +54,7 @@ void FLoadingTimingViewCommands::RegisterCommands()
 		EUserInterfaceActionType::ToggleButton,
 		FInputChord(EKeys::L));
 }
-PRAGMA_ENABLE_OPTIMIZATION
+UE_ENABLE_OPTIMIZATION_SHIP
 
 ////////////////////////////////////////////////////////////////////////////////////////////////////
 // FLoadingSharedState
diff --git a/Engine/Source/Developer/TraceInsights/Private/Insights/ViewModels/ThreadTimingTrack.cpp b/Engine/Source/Developer/TraceInsights/Private/Insights/ViewModels/ThreadTimingTrack.cpp
--- a/Engine/Source/Developer/TraceInsights/Private/Insights/ViewModels/ThreadTimingTrack.cpp
+++ b/Engine/Source/Developer/TraceInsights/Private/Insights/ViewModels/ThreadTimingTrack.cpp
@@ -58,7 +58,7 @@ FThreadTimingViewCommands::~FThreadTimingViewCommands()
 ////////////////////////////////////////////////////////////////////////////////////////////////////
 
 // UI_COMMAND takes long for the compiler to optimize
-PRAGMA_DISABLE_OPTIMIZATION
+UE_DISABLE_OPTIMIZATION_SHIP
 void FThreadTimingViewCommands::RegisterCommands()
 {
 	UI_COMMAND(ShowHideAllGpuTracks,
@@ -73,7 +73,7 @@ void FThreadTimingViewCommands::RegisterCommands()
 		EUserInterfaceActionType::ToggleButton,
 		FInputChord(EKeys::U));
 }
-PRAGMA_ENABLE_OPTIMIZATION
+UE_ENABLE_OPTIMIZATION_SHIP
 
 ////////////////////////////////////////////////////////////////////////////////////////////////////
 ////////////////////////////////////////////////////////////////////////////////////////////////////
diff --git a/Engine/Source/Developer/TraceInsights/Private/Insights/Widgets/SLogView.cpp b/Engine/Source/Developer/TraceInsights/Private/Insights/Widgets/SLogView.cpp
--- a/Engine/Source/Developer/TraceInsights/Private/Insights/Widgets/SLogView.cpp
+++ b/Engine/Source/Developer/TraceInsights/Private/Insights/Widgets/SLogView.cpp
@@ -60,7 +60,7 @@ public:
 	}
 
 	// UI_COMMAND takes long for the compiler to optimize
-	PRAGMA_DISABLE_OPTIMIZATION
+	UE_DISABLE_OPTIMIZATION_SHIP
 	virtual void RegisterCommands() override
 	{
 		UI_COMMAND(Command_HideSelectedCategory,
@@ -119,7 +119,7 @@ public:
 			"Opens the source file of the selected message in the registered IDE.",
 			EUserInterfaceActionType::Button, FInputChord());
 	}
-	PRAGMA_ENABLE_OPTIMIZATION
+	UE_ENABLE_OPTIMIZATION_SHIP
 
 	TSharedPtr<FUICommandInfo> Command_HideSelectedCategory;
 	TSharedPtr<FUICommandInfo> Command_ShowOnlySelectedCategory;
diff --git a/Engine/Source/Developer/TraceInsights/Private/Insights/Widgets/SStatsView.cpp b/Engine/Source/Developer/TraceInsights/Private/Insights/Widgets/SStatsView.cpp
--- a/Engine/Source/Developer/TraceInsights/Private/Insights/Widgets/SStatsView.cpp
+++ b/Engine/Source/Developer/TraceInsights/Private/Insights/Widgets/SStatsView.cpp
@@ -66,7 +66,7 @@ public:
 	}
 
 	// UI_COMMAND takes long for the compiler to optimize
-	PRAGMA_DISABLE_OPTIMIZATION
+	UE_DISABLE_OPTIMIZATION_SHIP
 	virtual void RegisterCommands() override
 	{
 		UI_COMMAND(Command_CopyToClipboard,
@@ -99,7 +99,7 @@ public:
 			EUserInterfaceActionType::Button,
 			FInputChord());
 	}
-	PRAGMA_ENABLE_OPTIMIZATION
+	UE_ENABLE_OPTIMIZATION_SHIP
 
 	TSharedPtr<FUICommandInfo> Command_CopyToClipboard;
 	TSharedPtr<FUICommandInfo> Command_Export;
diff --git a/Engine/Source/Developer/TraceInsights/Private/Insights/Widgets/STimerTreeView.cpp b/Engine/Source/Developer/TraceInsights/Private/Insights/Widgets/STimerTreeView.cpp
--- a/Engine/Source/Developer/TraceInsights/Private/Insights/Widgets/STimerTreeView.cpp
+++ b/Engine/Source/Developer/TraceInsights/Private/Insights/Widgets/STimerTreeView.cpp
@@ -47,7 +47,7 @@ public:
 	}
 
 	// UI_COMMAND takes long for the compiler to optimize
-	PRAGMA_DISABLE_OPTIMIZATION
+	UE_DISABLE_OPTIMIZATION_SHIP
 	virtual void RegisterCommands() override
 	{
 		UI_COMMAND(Command_CopyToClipboard,
@@ -56,7 +56,7 @@ public:
 			EUserInterfaceActionType::Button,
 			FInputChord(EModifierKey::Control, EKeys::C));
 	}
-	PRAGMA_ENABLE_OPTIMIZATION
+	UE_ENABLE_OPTIMIZATION_SHIP
 
 	TSharedPtr<FUICommandInfo> Command_CopyToClipboard;
 };
diff --git a/Engine/Source/Developer/TraceInsights/Private/Insights/Widgets/STimersView.cpp b/Engine/Source/Developer/TraceInsights/Private/Insights/Widgets/STimersView.cpp
--- a/Engine/Source/Developer/TraceInsights/Private/Insights/Widgets/STimersView.cpp
+++ b/Engine/Source/Developer/TraceInsights/Private/Insights/Widgets/STimersView.cpp
@@ -71,7 +71,7 @@ public:
 	}
 
 	// UI_COMMAND takes long for the compiler to optimize
-	PRAGMA_DISABLE_OPTIMIZATION
+	UE_DISABLE_OPTIMIZATION_SHIP
 	virtual void RegisterCommands() override
 	{
 		UI_COMMAND(Command_CopyToClipboard,
@@ -116,7 +116,7 @@ public:
 			EUserInterfaceActionType::Button,
 			FInputChord());
 	}
-	PRAGMA_ENABLE_OPTIMIZATION
+	UE_ENABLE_OPTIMIZATION_SHIP
 
 	TSharedPtr<FUICommandInfo> Command_CopyToClipboard;
 	TSharedPtr<FUICommandInfo> Command_Export;
diff --git a/Engine/Source/Editor/Cascade/Private/CascadeActions.cpp b/Engine/Source/Editor/Cascade/Private/CascadeActions.cpp
--- a/Engine/Source/Editor/Cascade/Private/CascadeActions.cpp
+++ b/Engine/Source/Editor/Cascade/Private/CascadeActions.cpp
@@ -11,7 +11,7 @@
 #define LOCTEXT_NAMESPACE "CascadeCommands"
 
 /** UI_COMMAND takes long for the compile to optimize */
-PRAGMA_DISABLE_OPTIMIZATION
+UE_DISABLE_OPTIMIZATION_SHIP
 void FCascadeCommands::RegisterCommands()
 {
 	UI_COMMAND(RestartSimulation, "Restart Sim", "Restart Simulation", EUserInterfaceActionType::Button, FInputChord());
@@ -91,6 +91,6 @@ void FCascadeCommands::RegisterCommands()
 	UI_COMMAND(RemoveDuplicateModules, "Remove Duplicate Modules", "Remove Duplicate Modules", EUserInterfaceActionType::Button, FInputChord());
 }
 
-PRAGMA_ENABLE_OPTIMIZATION
+UE_ENABLE_OPTIMIZATION_SHIP
 
 #undef LOCTEXT_NAMESPACE
diff --git a/Engine/Source/Editor/Kismet/Private/BlueprintEditorCommands.cpp b/Engine/Source/Editor/Kismet/Private/BlueprintEditorCommands.cpp
--- a/Engine/Source/Editor/Kismet/Private/BlueprintEditorCommands.cpp
+++ b/Engine/Source/Editor/Kismet/Private/BlueprintEditorCommands.cpp
@@ -48,7 +48,7 @@
 #define LOCTEXT_NAMESPACE "BlueprintEditorCommands"
 
 /** UI_COMMAND takes long for the compile to optimize */
-PRAGMA_DISABLE_OPTIMIZATION
+UE_DISABLE_OPTIMIZATION_SHIP
 void FBlueprintEditorCommands::RegisterCommands()
 {
 	// Edit commands
@@ -106,7 +106,7 @@ void FBlueprintEditorCommands::RegisterCommands()
 	UI_COMMAND( ToggleHideUnrelatedNodes, "Hide Unrelated", "Toggles automatically hiding nodes which are unrelated to the selected nodes.", EUserInterfaceActionType::ToggleButton, FInputChord() );
 }
 
-PRAGMA_ENABLE_OPTIMIZATION
+UE_ENABLE_OPTIMIZATION_SHIP
 
 namespace NodeSpawnInfoHelpers
 {
diff --git a/Engine/Source/Editor/LevelEditor/Private/LevelEditorActions.cpp b/Engine/Source/Editor/LevelEditor/Private/LevelEditorActions.cpp
--- a/Engine/Source/Editor/LevelEditor/Private/LevelEditorActions.cpp
+++ b/Engine/Source/Editor/LevelEditor/Private/LevelEditorActions.cpp
@@ -3445,7 +3445,7 @@ FLevelEditorCommands::FLevelEditorCommands()
 }
 
 /** UI_COMMAND takes long for the compile to optimize */
-PRAGMA_DISABLE_OPTIMIZATION
+UE_DISABLE_OPTIMIZATION_SHIP
 void FLevelEditorCommands::RegisterCommands()
 {
 	UI_COMMAND( BrowseDocumentation, "Level Editor Documentation", "Details on how to use the Level Editor", EUserInterfaceActionType::Button, FInputChord( EKeys::F1 ) );
@@ -3781,6 +3781,6 @@ void FLevelEditorCommands::RegisterCommands()
 	UI_COMMAND(OpenMergeActor, "Merge Actors", "Opens the Merge Actor panel", EUserInterfaceActionType::Button, FInputChord());
 }
 
-PRAGMA_ENABLE_OPTIMIZATION
+UE_ENABLE_OPTIMIZATION_SHIP
 
 #undef LOCTEXT_NAMESPACE
diff --git a/Engine/Source/Editor/LevelEditor/Private/LevelEditorModesActions.cpp b/Engine/Source/Editor/LevelEditor/Private/LevelEditorModesActions.cpp
--- a/Engine/Source/Editor/LevelEditor/Private/LevelEditorModesActions.cpp
+++ b/Engine/Source/Editor/LevelEditor/Private/LevelEditorModesActions.cpp
@@ -9,7 +9,7 @@ DEFINE_LOG_CATEGORY_STATIC(LevelEditorModesActions, Log, All);
 #define LOCTEXT_NAMESPACE "LevelEditorModesActions"
 
 /** UI_COMMAND takes long for the compile to optimize */
-PRAGMA_DISABLE_OPTIMIZATION
+UE_DISABLE_OPTIMIZATION_SHIP
 void FLevelEditorModesCommands::RegisterCommands()
 {
 	EditorModeCommands.Empty();
@@ -59,6 +59,6 @@ void FLevelEditorModesCommands::RegisterCommands()
 	}
 }
 
-PRAGMA_ENABLE_OPTIMIZATION
+UE_ENABLE_OPTIMIZATION_SHIP
 
 #undef LOCTEXT_NAMESPACE
diff --git a/Engine/Source/Editor/LevelEditor/Private/LevelViewportActions.cpp b/Engine/Source/Editor/LevelEditor/Private/LevelViewportActions.cpp
--- a/Engine/Source/Editor/LevelEditor/Private/LevelViewportActions.cpp
+++ b/Engine/Source/Editor/LevelEditor/Private/LevelViewportActions.cpp
@@ -31,7 +31,7 @@ FLevelViewportCommands::~FLevelViewportCommands()
 PRAGMA_ENABLE_DEPRECATION_WARNINGS
 
 /** UI_COMMAND takes long for the compile to optimize */
-PRAGMA_DISABLE_OPTIMIZATION
+UE_DISABLE_OPTIMIZATION_SHIP
 void FLevelViewportCommands::RegisterCommands()
 {
 	UI_COMMAND( ToggleMaximize, "Maximize Viewport", "Toggles the Maximize state of the current viewport", EUserInterfaceActionType::ToggleButton, FInputChord() );
@@ -314,6 +314,6 @@ void FLevelViewportCommands::RegisterShowSpriteCommands()
 	}
 }
 
-PRAGMA_ENABLE_OPTIMIZATION
+UE_ENABLE_OPTIMIZATION_SHIP
 
 #undef LOCTEXT_NAMESPACE
diff --git a/Engine/Source/Editor/Sequencer/Private/SSequencer.cpp b/Engine/Source/Editor/Sequencer/Private/SSequencer.cpp
--- a/Engine/Source/Editor/Sequencer/Private/SSequencer.cpp
+++ b/Engine/Source/Editor/Sequencer/Private/SSequencer.cpp
@@ -122,7 +122,7 @@
 
 /* SSequencer interface
  *****************************************************************************/
-PRAGMA_DISABLE_OPTIMIZATION
+UE_DISABLE_OPTIMIZATION_SHIP
 void SSequencer::Construct(const FArguments& InArgs, TSharedRef<FSequencer> InSequencer)
 {
 	using namespace UE::Sequencer;
@@ -856,7 +856,7 @@ void SSequencer::Construct(const FArguments& InArgs, TSharedRef<FSequencer> InSe
 
 	ResetBreadcrumbs();
 }
-PRAGMA_ENABLE_OPTIMIZATION
+UE_ENABLE_OPTIMIZATION_SHIP
 
 void SSequencer::BindCommands(TSharedRef<FUICommandList> SequencerCommandBindings, TSharedRef<FUICommandList> CurveEditorSharedBindings)
 {
diff --git a/Engine/Source/Runtime/ApplicationCore/Private/Windows/WindowsPlatformApplicationMisc.cpp b/Engine/Source/Runtime/ApplicationCore/Private/Windows/WindowsPlatformApplicationMisc.cpp
--- a/Engine/Source/Runtime/ApplicationCore/Private/Windows/WindowsPlatformApplicationMisc.cpp
+++ b/Engine/Source/Runtime/ApplicationCore/Private/Windows/WindowsPlatformApplicationMisc.cpp
@@ -378,7 +378,7 @@ float FWindowsPlatformApplicationMisc::GetDPIScaleFactorAtPoint(float X, float Y
 // Disabling optimizations helps to reduce the frequency of OpenClipboard failing with error code 0. It still happens
 // though only with really large text buffers and we worked around this by changing the editor to use an intermediate
 // text buffer for internal operations.
-PRAGMA_DISABLE_OPTIMIZATION 
+UE_DISABLE_OPTIMIZATION_SHIP
 
 void FWindowsPlatformApplicationMisc::ClipboardCopy(const TCHAR* Str)
 {
@@ -447,4 +447,4 @@ void FWindowsPlatformApplicationMisc::ClipboardPaste(class FString& Result)
 	}
 }
 
-PRAGMA_ENABLE_OPTIMIZATION 
+UE_ENABLE_OPTIMIZATION_SHIP
diff --git a/Engine/Source/Runtime/Core/Private/Containers/String.cpp b/Engine/Source/Runtime/Core/Private/Containers/String.cpp
--- a/Engine/Source/Runtime/Core/Private/Containers/String.cpp
+++ b/Engine/Source/Runtime/Core/Private/Containers/String.cpp
@@ -1209,7 +1209,7 @@ bool FString::ToHexBlob( const FString& Source, uint8* DestBuffer, const uint32
 	return false;
 }
 
-PRAGMA_DISABLE_OPTIMIZATION
+UE_DISABLE_OPTIMIZATION_SHIP
 void StripNegativeZero(double& InFloat)
 {
 	// This works for translating a negative zero into a positive zero,
@@ -1217,7 +1217,7 @@ void StripNegativeZero(double& InFloat)
 	// or /fp:fast, the compiler can strip it out.
 	InFloat += 0.0f;
 }
-PRAGMA_ENABLE_OPTIMIZATION
+UE_ENABLE_OPTIMIZATION_SHIP
 
 FString FString::SanitizeFloat( double InFloat, const int32 InMinFractionalDigits )
 {
diff --git a/Engine/Source/Runtime/CoreUObject/Private/UObject/CoreRedirects.cpp b/Engine/Source/Runtime/CoreUObject/Private/UObject/CoreRedirects.cpp
--- a/Engine/Source/Runtime/CoreUObject/Private/UObject/CoreRedirects.cpp
+++ b/Engine/Source/Runtime/CoreUObject/Private/UObject/CoreRedirects.cpp
@@ -1443,7 +1443,7 @@ ECoreRedirectFlags FCoreRedirects::GetFlagsForTypeClass(UClass *TypeClass)
 
 #if UE_WITH_CORE_REDIRECTS
 
-PRAGMA_DISABLE_OPTIMIZATION
+UE_DISABLE_OPTIMIZATION_SHIP
 
 // The compiler doesn't like having a massive string table in a single function so split it up
 #define CLASS_REDIRECT(OldName, NewName) Redirects.Emplace_GetRef(ECoreRedirectFlags::Type_Class, TEXT(OldName), TEXT(NewName))
@@ -2184,7 +2184,7 @@ static void RegisterNativeRedirects49(TArray<FCoreRedirect>& Redirects)
 	ENUM_REDIRECT("EInterpolationBlend", "/Script/AnimGraphRuntime.EInterpolationBlend");
 }
 
-PRAGMA_ENABLE_OPTIMIZATION
+UE_ENABLE_OPTIMIZATION_SHIP
 
 void FCoreRedirects::RegisterNativeRedirects()
 {
diff --git a/Engine/Source/Runtime/CoreUObject/Private/UObject/ScriptCore.cpp b/Engine/Source/Runtime/CoreUObject/Private/UObject/ScriptCore.cpp
--- a/Engine/Source/Runtime/CoreUObject/Private/UObject/ScriptCore.cpp
+++ b/Engine/Source/Runtime/CoreUObject/Private/UObject/ScriptCore.cpp
@@ -329,7 +329,7 @@ bool FBlueprintContextTracker::RecordAccessViolation(const UObject* Object)
 }
 
 // This is meant to be called from the immediate mode, and for confusing reasons the optimized code isn't always safe in that case
-PRAGMA_DISABLE_OPTIMIZATION
+UE_DISABLE_OPTIMIZATION_SHIP
 
 void PrintScriptCallStackImpl()
 {
@@ -349,7 +349,7 @@ void PrintScriptCallStackImpl()
 	}
 }
 
-PRAGMA_ENABLE_OPTIMIZATION
+UE_ENABLE_OPTIMIZATION_SHIP
 
 extern CORE_API void (*GPrintScriptCallStackFn)();
 
diff --git a/Engine/Source/Runtime/Engine/Private/ShaderCompiler/ShaderCompiler.cpp b/Engine/Source/Runtime/Engine/Private/ShaderCompiler/ShaderCompiler.cpp
--- a/Engine/Source/Runtime/Engine/Private/ShaderCompiler/ShaderCompiler.cpp
+++ b/Engine/Source/Runtime/Engine/Private/ShaderCompiler/ShaderCompiler.cpp
@@ -1839,7 +1839,7 @@ static void LogQueuedCompileJobs(const TArray<FShaderCommonCompileJobPtr>& Queue
 }
 
 // Disable optimization for this crash handler to get full access to the entire stack frame when debugging a crash dump
-PRAGMA_DISABLE_OPTIMIZATION;
+UE_DISABLE_OPTIMIZATION_SHIP
 static void HandleWorkerCrash(const TArray<FShaderCommonCompileJobPtr>& QueuedJobs, FArchive& OutputFile, int32 OutputVersion, int64 FileSize, ESCWErrorCode ErrorCode, int32 NumProcessedJobs, int32 CallstackLength, int32 ExceptionInfoLength)
 {
 	TArray<TCHAR> Callstack;
@@ -1941,7 +1941,7 @@ static void HandleWorkerCrash(const TArray<FShaderCommonCompileJobPtr>& QueuedJo
 		break;
 	}
 }
-PRAGMA_ENABLE_OPTIMIZATION;
+UE_ENABLE_OPTIMIZATION_SHIP
 
 // Process results from Worker Process
 void FShaderCompileUtilities::DoReadTaskResults(const TArray<FShaderCommonCompileJobPtr>& QueuedJobs, FArchive& OutputFile)
diff --git a/Engine/Source/Runtime/InputCore/Private/InputCoreTypes.cpp b/Engine/Source/Runtime/InputCore/Private/InputCoreTypes.cpp
--- a/Engine/Source/Runtime/InputCore/Private/InputCoreTypes.cpp
+++ b/Engine/Source/Runtime/InputCore/Private/InputCoreTypes.cpp
@@ -529,7 +529,7 @@ void FKeyDetails::CommonInit(const uint32 InKeyFlags)
 }
 
 #if !FAST_BOOT_HACKS
-PRAGMA_DISABLE_OPTIMIZATION
+UE_DISABLE_OPTIMIZATION_SHIP
 #endif
 void EKeys::Initialize()
 {
@@ -992,7 +992,7 @@ void EKeys::Initialize()
 	FInputKeyManager::Get();
 }
 #if !FAST_BOOT_HACKS
-PRAGMA_ENABLE_OPTIMIZATION
+UE_ENABLE_OPTIMIZATION_SHIP
 #endif
 
 void EKeys::AddKey(const FKeyDetails& KeyDetails)
@@ -1089,7 +1089,7 @@ FName EKeys::GetMenuCategoryPaletteIcon(const FName CategoryName)
 
 EConsoleForGamepadLabels::Type EKeys::ConsoleForGamepadLabels = EConsoleForGamepadLabels::None;
 
-PRAGMA_DISABLE_OPTIMIZATION
+UE_DISABLE_OPTIMIZATION_SHIP
 FText EKeys::GetGamepadDisplayName(const FKey Key)
 {
 	switch (ConsoleForGamepadLabels)
@@ -1275,7 +1275,7 @@ FText EKeys::GetGamepadDisplayName(const FKey Key)
 	ensureMsgf(false, TEXT("Unexpected key %s using EKeys::GetGamepadDisplayName"), *Key.ToString());
 	return FText::FromString(Key.ToString());
 }
-PRAGMA_ENABLE_OPTIMIZATION
+UE_ENABLE_OPTIMIZATION_SHIP
 
 bool FKey::IsValid() const
 {
diff --git a/Engine/Source/Runtime/Slate/Public/SlateOptMacros.h b/Engine/Source/Runtime/Slate/Public/SlateOptMacros.h
--- a/Engine/Source/Runtime/Slate/Public/SlateOptMacros.h
+++ b/Engine/Source/Runtime/Slate/Public/SlateOptMacros.h
@@ -4,5 +4,5 @@
 
 #include "CoreMinimal.h"
 
-#define BEGIN_SLATE_FUNCTION_BUILD_OPTIMIZATION BEGIN_FUNCTION_BUILD_OPTIMIZATION
-#define END_SLATE_FUNCTION_BUILD_OPTIMIZATION   END_FUNCTION_BUILD_OPTIMIZATION
+#define BEGIN_SLATE_FUNCTION_BUILD_OPTIMIZATION UE_DISABLE_OPTIMIZATION_SHIP
+#define END_SLATE_FUNCTION_BUILD_OPTIMIZATION   UE_ENABLE_OPTIMIZATION_SHIP
diff --git a/Engine/Source/Runtime/SlateCore/Private/Styling/CoreStyle.cpp b/Engine/Source/Runtime/SlateCore/Private/Styling/CoreStyle.cpp
--- a/Engine/Source/Runtime/SlateCore/Private/Styling/CoreStyle.cpp
+++ b/Engine/Source/Runtime/SlateCore/Private/Styling/CoreStyle.cpp
@@ -127,7 +127,7 @@ void FCoreStyle::SetFocusBrush(FSlateBrush* NewBrush)
 }
 
 
-PRAGMA_DISABLE_OPTIMIZATION
+UE_DISABLE_OPTIMIZATION_SHIP
 
 #define IMAGE_BRUSH( RelativePath, ... ) FSlateImageBrush(Style->RootToContentDir(RelativePath, TEXT(".png")), __VA_ARGS__)
 #define BOX_BRUSH( RelativePath, ... ) FSlateBoxBrush(Style->RootToContentDir(RelativePath, TEXT(".png")), __VA_ARGS__)
@@ -1340,7 +1340,7 @@ TSharedRef<ISlateStyle> FCoreStyle::Create( const FName& InStyleSetName )
 	return Style;
 
 }
-PRAGMA_ENABLE_OPTIMIZATION
+UE_ENABLE_OPTIMIZATION_SHIP
 
 #undef IMAGE_BRUSH
 #undef BOX_BRUSH
