Recent changes made StartAnalysisForLastLiveSession always fail to connect, and prior to that there was a small chance of a failure, due to the timing of async network communication code.
Updated that function to queue up a retry for 1 second by default (with an optional parameter to override or disable the retry timeout)

diff --git a/Engine/Source/Developer/TraceInsights/Private/Insights/InsightsManager.cpp b/Engine/Source/Developer/TraceInsights/Private/Insights/InsightsManager.cpp
--- a/Engine/Source/Developer/TraceInsights/Private/Insights/InsightsManager.cpp
+++ b/Engine/Source/Developer/TraceInsights/Private/Insights/InsightsManager.cpp
@@ -497,6 +497,19 @@ FInsightsSettings& FInsightsManager::GetSettings()
 
 bool FInsightsManager::Tick(float DeltaTime)
 {
+	if (RetryLoadLastLiveSessionTimer > 0.0f)
+	{
+		LoadLastLiveSession(0.0f);
+		if (Session)
+		{
+			RetryLoadLastLiveSessionTimer = 0.0f;
+		}
+		else
+		{
+			RetryLoadLastLiveSessionTimer -= DeltaTime;
+		}
+	}
+	
 	AutoLoadLiveSession();
 
 	UpdateSessionDuration();
@@ -834,7 +847,7 @@ void FInsightsManager::AutoLoadLiveSession()
 
 ////////////////////////////////////////////////////////////////////////////////////////////////////
 
-void FInsightsManager::LoadLastLiveSession()
+void FInsightsManager::LoadLastLiveSession(float RetryTime)
 {
 	ResetSession();
 
@@ -852,6 +865,11 @@ void FInsightsManager::LoadLastLiveSession()
 			LoadTrace(SessionInfo->GetTraceId());
 		}
 	}
+
+	if (Session == nullptr && RetryTime > 0)
+	{
+		RetryLoadLastLiveSessionTimer = RetryTime;
+	}
 }
 
 ////////////////////////////////////////////////////////////////////////////////////////////////////
diff --git a/Engine/Source/Developer/TraceInsights/Private/Insights/InsightsManager.h b/Engine/Source/Developer/TraceInsights/Private/Insights/InsightsManager.h
--- a/Engine/Source/Developer/TraceInsights/Private/Insights/InsightsManager.h
+++ b/Engine/Source/Developer/TraceInsights/Private/Insights/InsightsManager.h
@@ -269,9 +269,10 @@ public:
 
 	/**
 	 * Creates a new analysis session instance and loads the latest available trace that is live.
-	 * Replaces the current analysis session.
+	 * Replaces the current analysis session.  On failure, if InRetryTime is > 0, will retry connecting every frame for RetryTime seconds 
+	 * @param InRetryTime - On failure, how many seconds to retry connecting during FInsightsManager::Tick
 	 */
-	void LoadLastLiveSession();
+	void LoadLastLiveSession(float InRetryTime = 1.0f);
 
 	/**
 	 * Creates a new analysis session instance using the specified trace id.
@@ -457,6 +458,7 @@ private:
 	bool bIsAnalysisComplete = false;
 	bool bSessionAnalysisCompletedAutoQuit = false;
 
+	float RetryLoadLastLiveSessionTimer = 0.0f;
 	bool bIsAutoLoadLiveSessionEnabled = false;
 	TSet<uint32> AutoLoadedTraceIds; // list of trace ids for the auto loaded live sessions
 
diff --git a/Engine/Source/Developer/TraceInsights/Private/Insights/TraceInsightsModule.cpp b/Engine/Source/Developer/TraceInsights/Private/Insights/TraceInsightsModule.cpp
--- a/Engine/Source/Developer/TraceInsights/Private/Insights/TraceInsightsModule.cpp
+++ b/Engine/Source/Developer/TraceInsights/Private/Insights/TraceInsightsModule.cpp
@@ -542,9 +542,9 @@ void FTraceInsightsModule::StartAnalysisForTrace(uint32 InTraceId, bool InAutoQu
 
 ////////////////////////////////////////////////////////////////////////////////////////////////////
 
-void FTraceInsightsModule::StartAnalysisForLastLiveSession()
+void FTraceInsightsModule::StartAnalysisForLastLiveSession(float InRetryTime)
 {
-	FInsightsManager::Get()->LoadLastLiveSession();
+	FInsightsManager::Get()->LoadLastLiveSession(InRetryTime);
 	UpdateAppTitle();
 }
 
diff --git a/Engine/Source/Developer/TraceInsights/Private/Insights/TraceInsightsModule.h b/Engine/Source/Developer/TraceInsights/Private/Insights/TraceInsightsModule.h
--- a/Engine/Source/Developer/TraceInsights/Private/Insights/TraceInsightsModule.h
+++ b/Engine/Source/Developer/TraceInsights/Private/Insights/TraceInsightsModule.h
@@ -51,7 +51,7 @@ public:
 
 	virtual TSharedPtr<const TraceServices::IAnalysisSession> GetAnalysisSession() const override;
 	virtual void StartAnalysisForTrace(uint32 InTraceId, bool InAutoQuit = false) override;
-	virtual void StartAnalysisForLastLiveSession() override;
+	virtual void StartAnalysisForLastLiveSession(float RetryTime = 1.0f) override;
 	virtual void StartAnalysisForTraceFile(const TCHAR* InTraceFile, bool InAutoQuit = false) override;
 
 	virtual void ShutdownUserInterface() override;
diff --git a/Engine/Source/Developer/TraceInsights/Public/Insights/IUnrealInsightsModule.h b/Engine/Source/Developer/TraceInsights/Public/Insights/IUnrealInsightsModule.h
--- a/Engine/Source/Developer/TraceInsights/Public/Insights/IUnrealInsightsModule.h
+++ b/Engine/Source/Developer/TraceInsights/Public/Insights/IUnrealInsightsModule.h
@@ -209,8 +209,11 @@ public:
 
 	/**
 	 * Starts analysis of the last live session. Called when the application starts in "Viewer" mode.
+	 * On failure, if InRetryTime is > 0, retry connecting every frame for RetryTime seconds 
+	 * 
+	 * @param InRetryTime How many seconds to retry connecting asyncronously
 	 */
-	virtual void StartAnalysisForLastLiveSession() = 0;
+	virtual void StartAnalysisForLastLiveSession(float InRetryTime = 1.0f) = 0;
 
 	/**
 	 * Starts analysis of the specified *.utrace file. Called when the application starts in "Viewer" mode.
