[Insights] Added support for *?-type wildcard name patterns for regions specified to ExportTimerStatisticsAsText API. It is now possible to specify the following export command:
    TimingInsights.ExportTimerStatistics C:\Out\Region_*.CSV -region="SameRegion,Game_*,OtherRegion" and it is expected to produce the following files:
        C:\Out\Region_SameRegion.CSV
        C:\Out\Region_Game_[...].CSV // multiple files
        C:\Out\Region_OtherRegion.CSV
    (if respective regions are found; i.e. defined by "RegionStart:" and "RegionEnd:" bookmarks).

diff --git a/Engine/Source/Developer/TraceInsights/Private/Insights/TimingProfilerManager.cpp b/Engine/Source/Developer/TraceInsights/Private/Insights/TimingProfilerManager.cpp
--- a/Engine/Source/Developer/TraceInsights/Private/Insights/TimingProfilerManager.cpp
+++ b/Engine/Source/Developer/TraceInsights/Private/Insights/TimingProfilerManager.cpp
@@ -745,7 +745,7 @@ bool FTimingProfilerManager::Exec(const TCHAR* Cmd, FOutputDevice& Ar)
 		Insights::FTimingExporter Exporter(*FInsightsManager::Get()->GetSession().Get());
 		Insights::FTimingExporter::FExportTimerStatisticsParams Params; // default (all timing events)
 
-		// These variables needs to be in the same scope with the call to Exporter.ExportTimingEventsAsText().
+		// These variables needs to be in the same scope with the call to Exporter.ExportTimerStatisticsAsText().
 		TArray<FName> Columns; // referenced by Params.Columns
 		TSet<uint32> IncludedThreads; // referenced in Params.ThreadFilter lambda function
 		TSet<uint32> IncludedTimers; // referenced in Params.TimingEventFilter lambda function
@@ -814,9 +814,13 @@ bool FTimingProfilerManager::Exec(const TCHAR* Cmd, FOutputDevice& Ar)
 				}
 				else if (Token.StartsWith(RegionToken))
 				{
+					// Comma-delimited list of region names. Supports *?-type wildcard.
+					// Each region is exported to a separate file. The '*' char in Filename (if any)
+					// will be replaced with the resolved name of the region.
 					// Default: -Region=
-					// Example: -Region=RegionName
+					// Example: -Region="RegionName,Game_*,OtherRegion"
 					Token.RightChopInline(UE_ARRAY_COUNT(RegionToken) - 1);
+					Token.TrimQuotesInline();
 					Params.Region = TCHAR_TO_ANSI(*Token);
 				}
 				else
@@ -832,11 +836,11 @@ bool FTimingProfilerManager::Exec(const TCHAR* Cmd, FOutputDevice& Ar)
 		{
 			Params.ThreadFilter = [](unsigned int){ return true; };
 		}
-		
+
 		int32 Result = Exporter.ExportTimerStatisticsAsText(Filename, Params);
 		return Result > 0;
 	}
-	
+
 	return false;
 }
 
diff --git a/Engine/Source/Developer/TraceInsights/Private/Insights/ViewModels/TimingExporter.cpp b/Engine/Source/Developer/TraceInsights/Private/Insights/ViewModels/TimingExporter.cpp
--- a/Engine/Source/Developer/TraceInsights/Private/Insights/ViewModels/TimingExporter.cpp
+++ b/Engine/Source/Developer/TraceInsights/Private/Insights/ViewModels/TimingExporter.cpp
@@ -24,13 +24,13 @@ namespace Insights
 
 FTimingExporter::FTimingExporter(const TraceServices::IAnalysisSession& InSession)
 	: Session(InSession)
-{	
+{
 }
 
 ////////////////////////////////////////////////////////////////////////////////////////////////////
 
 FTimingExporter::~FTimingExporter()
-{	
+{
 }
 
 ////////////////////////////////////////////////////////////////////////////////////////////////////
@@ -548,62 +548,157 @@ int32 FTimingExporter::ExportTimingEventsAsText(const FString& Filename, FExport
 }
 
 ////////////////////////////////////////////////////////////////////////////////////////////////////
-	
+
 int32 FTimingExporter::ExportTimerStatisticsAsText(const FString& Filename, FExportTimerStatisticsParams& Params) const
 {
 	TraceServices::ITable<TraceServices::FTimingProfilerAggregatedStats>* StatsTable;
 
-	double RegionStartTime = 0;
-	double RegionEndTime = std::numeric_limits<double>::max();
-	
+	if (!Params.Region.IsEmpty())
 	{
-		TraceServices::FAnalysisSessionReadScope SessionReadScope(Session);
-		
-		const TraceServices::ITimingProfilerProvider* TimingProfilerProvider = TraceServices::ReadTimingProfilerProvider(Session);
-		if (!TimingProfilerProvider)
-		{
-			UE_LOG(TraceInsights, Error, TEXT("Unable to access TimingProfilerProvider for ExportTimerStatisticsAsText"));
-			return -1;
-		}
-		
-		if (!Params.Region.IsEmpty())
+		class FRegionNameSpec
 		{
-			
-			const TraceServices::IBookmarkProvider& BookmarkProvider = TraceServices::ReadBookmarkProvider(Session);
-			const FString RegionStartName = TEXT("RegionStart:") + Params.Region;
-			const FString RegionEndName = TEXT("RegionEnd:") + Params.Region;
-			UE_LOG(TraceInsights, Log, TEXT("Looking for the following bookmarks: '%s' and '%s'"), *RegionStartName, *RegionEndName);
-			
-			BookmarkProvider.EnumerateBookmarks(0, std::numeric_limits<double>::max(),
-				[&RegionStartTime, &RegionEndTime, &RegionStartName, &RegionEndName](const TraceServices::FBookmark& InBookmark){
-					if (RegionStartTime == 0 && RegionStartName.Equals(InBookmark.Text))
+		public:
+			FRegionNameSpec(FString& InNamePatternList)
+			{
+				InNamePatternList.ParseIntoArray(NamePatterns, TEXT(","), true);
+			}
+
+			bool Match(const FString& InBookmarkName, FString& OutRegionName, bool& bOutIsRegionStart)
+			{
+				if (InBookmarkName.StartsWith(TEXT("RegionStart:")))
+				{
+					bOutIsRegionStart = true;
+					OutRegionName = InBookmarkName.RightChop(12);
+				}
+				else if (InBookmarkName.StartsWith(TEXT("RegionEnd:")))
+				{
+					bOutIsRegionStart = false;
+					OutRegionName = InBookmarkName.RightChop(10);
+				}
+				else
+				{
+					return false;
+				}
+				for (const FString& NamePattern : NamePatterns)
+				{
+					if (OutRegionName.MatchesWildcard(NamePattern))
 					{
-						RegionStartTime = InBookmark.Time;
+						return true;
 					}
-					if (RegionEndTime == std::numeric_limits<double>::max() && RegionEndName.Equals(InBookmark.Text))
+				}
+				return false;
+			}
+
+		private:
+			TArray<FString> NamePatterns;
+		};
+		FRegionNameSpec RegionNameSpec(Params.Region);
+
+		struct FTimeRegion
+		{
+			double Start;
+			double End;
+		};
+		TMap<FString, FTimeRegion> Regions;
+
+		// Detect regions
+		{
+			TraceServices::FAnalysisSessionReadScope SessionReadScope(Session);
+
+			const TraceServices::IBookmarkProvider& BookmarkProvider = TraceServices::ReadBookmarkProvider(Session);
+
+			UE_LOG(TraceInsights, Log, TEXT("Looking for regions: '%s'"), *Params.Region);
+
+			BookmarkProvider.EnumerateBookmarks(0.0, std::numeric_limits<double>::max(),
+				[&RegionNameSpec, &Regions](const TraceServices::FBookmark& InBookmark)
+				{
+					FString BaseRegionName;
+					bool bIsRegionStart;
+					if (RegionNameSpec.Match(InBookmark.Text, BaseRegionName, bIsRegionStart))
 					{
-						RegionEndTime = InBookmark.Time;
+						FTimeRegion* Region = nullptr;
+						FString RegionName = BaseRegionName;
+						int32 Index = 2;
+
+						while (true)
+						{
+							Region = Regions.Find(RegionName);
+
+							if (!Region)
+							{
+								Region = &Regions.Add(RegionName);
+								Region->Start = -std::numeric_limits<double>::infinity();
+								Region->End = std::numeric_limits<double>::infinity();
+							}
+
+							if (bIsRegionStart)
+							{
+								if (Region->Start == -std::numeric_limits<double>::infinity() &&
+									InBookmark.Time <= Region->End)
+								{
+									Region->Start = InBookmark.Time;
+									break;
+								}
+							}
+							else
+							{
+								if (Region->End == std::numeric_limits<double>::infinity() &&
+									InBookmark.Time >= Region->Start)
+								{
+									Region->End = InBookmark.Time;
+									break;
+								}
+							}
+
+							RegionName = FString::Printf(TEXT("%s%d"), *BaseRegionName, Index++);
+						}
 					}
-			});
+				});
+		}
 
-			if (RegionStartTime == 0 || RegionEndTime == std::numeric_limits<double>::max())
-			{
-				UE_LOG(TraceInsights, Error, TEXT("Unable to find 'RegionStart:' and 'RegionEnd:' bookmarks for region '%s'"), *Params.Region);
-				return -1;
-			}
+		if (Regions.Num() == 0)
+		{
+			UE_LOG(TraceInsights, Error, TEXT("Unable to find any region ('RegionStart:' and 'RegionEnd:' bookmarks) with name pattern '%s'."), *Params.Region);
+			return -1;
 		}
 
-		if (RegionStartTime != 0 && RegionEndTime != std::numeric_limits<double>::max())
+		FStopwatch Stopwatch;
+		Stopwatch.Start();
+
+		// Export timing statistics for each region.
+		for (auto& KV : Regions)
 		{
-			Params.IntervalStartTime = RegionStartTime;
-			Params.IntervalEndTime = RegionEndTime;
-			UE_LOG(TraceInsights, Display, TEXT("Limit selection to region '%s' from %f to %f"), *Params.Region, RegionStartTime, RegionEndTime);
+			FString RegionFilename = Filename.Replace(TEXT("*"), *KV.Key);
+			FExportTimerStatisticsParams RegionParams = Params;
+			RegionParams.Region.Reset();
+			RegionParams.IntervalStartTime = KV.Value.Start;
+			RegionParams.IntervalEndTime = KV.Value.End;
+			UE_LOG(TraceInsights, Display, TEXT("Exporting timing statistics for region '%s' [%f .. %f] to '%s'"), *KV.Key, RegionParams.IntervalStartTime, RegionParams.IntervalEndTime, *RegionFilename);
+			ExportTimerStatisticsAsText(RegionFilename, RegionParams);
+		}
+
+		Stopwatch.Stop();
+		const double TotalTime = Stopwatch.GetAccumulatedTime();
+		UE_LOG(TraceInsights, Log, TEXT("Exported timing statistics for %d regions in %.3fs."), Regions.Num(), TotalTime);
+		return Regions.Num();
+	}
+
+	check(Params.Region.IsEmpty());
+
+	{
+		TraceServices::FAnalysisSessionReadScope SessionReadScope(Session);
+
+		const TraceServices::ITimingProfilerProvider* TimingProfilerProvider = TraceServices::ReadTimingProfilerProvider(Session);
+		if (!TimingProfilerProvider)
+		{
+			UE_LOG(TraceInsights, Error, TEXT("Unable to access TimingProfilerProvider for ExportTimerStatisticsAsText"));
+			return -1;
 		}
 
 		//@Todo: this does not yet handle the -column and -timers parameters.
 		StatsTable = TimingProfilerProvider->CreateAggregation(Params.IntervalStartTime, Params.IntervalEndTime, Params.ThreadFilter, true);
 	}
-		
+
 	FStopwatch Stopwatch;
 	Stopwatch.Start();
 
@@ -713,7 +808,7 @@ TFunction<bool(uint32)> FTimingExporter::MakeThreadFilterExclusive(const TSet<ui
 			return ThreadId != ExcludedThreadId;
 		};
 	}
-	
+
 	return [&ExcludedThreads](uint32 ThreadId) -> bool
 	{
 		return !ExcludedThreads.Contains(ThreadId);
@@ -730,7 +825,7 @@ FTimingExporter::FTimingEventFilterFunc FTimingExporter::MakeTimingEventFilterBy
 	}
 
 	OutIncludedTimers.Reset();
-	
+
 	TMap<FString, uint32> Timers;
 
 	// Iterate the GPU & CPU timers.
@@ -811,7 +906,7 @@ FTimingExporter::FTimingEventFilterFunc FTimingExporter::MakeTimingEventFilterBy
 	{
 		return nullptr;
 	}
-		
+
 	if (ExcludedTimers.Num() == 1)
 	{
 		const uint32 ExcludedTimerId = ExcludedTimers[FSetElementId::FromInteger(0)];
