[Insights] LogView: Fixed log spam when filtering the LogView in Animation Insights.

diff --git a/Engine/Source/Developer/TraceInsights/Private/Insights/ViewModels/LogFilter.cpp b/Engine/Source/Developer/TraceInsights/Private/Insights/ViewModels/LogFilter.cpp
--- a/Engine/Source/Developer/TraceInsights/Private/Insights/ViewModels/LogFilter.cpp
+++ b/Engine/Source/Developer/TraceInsights/Private/Insights/ViewModels/LogFilter.cpp
@@ -271,6 +271,7 @@ void FLogFilteringAsyncTask::DoWork()
 	bool bCanceled = false;
 	FilteredMessages.Reset();
 
+#if !WITH_EDITOR
 	if (Filter.IsFilterSetByText())
 	{
 		UE_LOG(TraceInsights, Log, TEXT("[LogView] FLogFilteringAsyncTask::DoWork [%d to %d] by Text (\"%s\")"), StartIndex, EndIndex, *Filter.GetFilterText().ToString());
@@ -279,6 +280,7 @@ void FLogFilteringAsyncTask::DoWork()
 	{
 		UE_LOG(TraceInsights, Log, TEXT("[LogView] FLogFilteringAsyncTask::DoWork [%d to %d]"), StartIndex, EndIndex);
 	}
+#endif // !WITH_EDITOR
 
 	for (int32 Index = StartIndex; Index < EndIndex && !bCanceled; ++Index)
 	{
@@ -295,11 +297,15 @@ void FLogFilteringAsyncTask::DoWork()
 	if (bCanceled)
 	{
 		FilteredMessages.Reset();
+#if !WITH_EDITOR
 		UE_LOG(TraceInsights, Log, TEXT("[LogView] FLogFilteringAsyncTask::DoWork CANCELED"));
+#endif // !WITH_EDITOR
 	}
 	else
 	{
+#if !WITH_EDITOR
 		UE_LOG(TraceInsights, Log, TEXT("[LogView] FLogFilteringAsyncTask::DoWork DONE (%d filtered messages)"), FilteredMessages.Num());
+#endif // !WITH_EDITOR
 	}
 }
 
diff --git a/Engine/Source/Developer/TraceInsights/Private/Insights/Widgets/SLogView.cpp b/Engine/Source/Developer/TraceInsights/Private/Insights/Widgets/SLogView.cpp
--- a/Engine/Source/Developer/TraceInsights/Private/Insights/Widgets/SLogView.cpp
+++ b/Engine/Source/Developer/TraceInsights/Private/Insights/Widgets/SLogView.cpp
@@ -833,12 +833,14 @@ void SLogView::Tick(const FGeometry& AllottedGeometry, const double InCurrentTim
 					FilteringEndIndex = NewMessageCount;
 					FilteringChangeNumber = Filter.GetChangeNumber();
 					FilteringAsyncTask = MakeUnique<FAsyncTask<FLogFilteringAsyncTask>>(FilteringStartIndex, FilteringEndIndex, Filter, SharedThis(this));
+#if !WITH_EDITOR
 					UE_LOG(TraceInsights, Log, TEXT("[LogView] Start async task for filtering by%s%s%s (\"%s\") (%d to %d)"),
 						Filter.IsFilterSetByVerbosity() ? TEXT(" Verbosity,") : TEXT(""),
 						Filter.IsFilterSetByCategory() ? TEXT(" Category,") : TEXT(""),
 						Filter.IsFilterSetByText() ? TEXT(" Text") : TEXT(""),
 						*Filter.GetFilterText().ToString(),
 						FilteringStartIndex, FilteringEndIndex);
+#endif // !WITH_EDITOR
 					FilteringAsyncTask->StartBackgroundTask();
 				}
 				else
@@ -880,18 +882,22 @@ void SLogView::Tick(const FGeometry& AllottedGeometry, const double InCurrentTim
 				UpdateStatsText();
 
 				FilteringStopwatch.Stop();
+#if !WITH_EDITOR
 				uint64 DurationMs = FilteringStopwatch.GetAccumulatedTimeMs();
 				if (DurationMs > 10) // avoids spams
 				{
 					UE_LOG(TraceInsights, Log, TEXT("[LogView] Updated (no filter; %d added / %d total messages) in %llu ms."),
 						NumAddedMessages, NewMessageCount, DurationMs);
 				}
+#endif // !WITH_EDITOR
 			}
 		}
 		else // if (NewMessageCount < TotalNumMessages)
 		{
 			// Just reset. On next Tick() the list will grow if needed.
+#if !WITH_EDITOR
 			UE_LOG(TraceInsights, Log, TEXT("[LogView] RESET"));
+#endif // !WITH_EDITOR
 			Cache.Reset();
 			Messages.Reset();
 			TotalNumMessages = 0;
@@ -947,6 +953,7 @@ void SLogView::Tick(const FGeometry& AllottedGeometry, const double InCurrentTim
 			UpdateStatsText();
 
 			FilteringStopwatch.Stop();
+#if !WITH_EDITOR
 			uint64 DurationMs = FilteringStopwatch.GetAccumulatedTimeMs();
 			if (DurationMs > 10) // avoids spams
 			{
@@ -955,6 +962,7 @@ void SLogView::Tick(const FGeometry& AllottedGeometry, const double InCurrentTim
 				UE_LOG(TraceInsights, Log, TEXT("[LogView] Updated (%d added / %d async filtered / %d total messages) in %llu ms (%.2f messages/second)."),
 					NumFilteredMessages, NumAsyncFilteredMessages, TotalNumMessages, DurationMs, Speed);
 			}
+#endif // !WITH_EDITOR
 		}
 
 		FilteringAsyncTask.Reset();
@@ -1053,7 +1061,9 @@ void SLogView::FilterTextBox_OnTextChanged(const FText& InFilterText)
 void SLogView::OnFilterChanged()
 {
 	const FString FilterText = FilterTextBox->GetText().ToString();
+#if !WITH_EDITOR
 	UE_LOG(TraceInsights, Log, TEXT("[LogView] OnFilterChanged: \"%s\""), *FilterText);
+#endif // !WITH_EDITOR
 	Cache.Reset();
 	Messages.Reset();
 	TotalNumMessages = 0;
