[Insights] Added GetName() to ITimingViewSession. The Timing View widgets from TimingInsights, MemoryInsights and AssetLoadingInsights will have this name initialized with the corresponding FInsightsManagerTabs::*ProfilerTabId.

diff --git a/Engine/Source/Developer/TraceInsights/Private/Insights/LoadingProfiler/Widgets/SLoadingProfilerWindow.cpp b/Engine/Source/Developer/TraceInsights/Private/Insights/LoadingProfiler/Widgets/SLoadingProfilerWindow.cpp
--- a/Engine/Source/Developer/TraceInsights/Private/Insights/LoadingProfiler/Widgets/SLoadingProfilerWindow.cpp
+++ b/Engine/Source/Developer/TraceInsights/Private/Insights/LoadingProfiler/Widgets/SLoadingProfilerWindow.cpp
@@ -228,10 +228,9 @@ TSharedRef<SDockTab> SLoadingProfilerWindow::SpawnTab_TimingView(const FSpawnTab
 		.ShouldAutosize(false)
 		.TabRole(ETabRole::PanelTab)
 		[
-			SAssignNew(TimingView, STimingView)
+			SAssignNew(TimingView, STimingView, FInsightsManagerTabs::LoadingProfilerTabId)
 		];
 
-	TimingView->EnableAssetLoadingMode();
 	TimingView->Reset(true);
 	TimingView->OnSelectionChanged().AddSP(this, &SLoadingProfilerWindow::OnTimeSelectionChanged);
 	TimingView->SelectTimeInterval(SelectionStartTime, SelectionEndTime - SelectionStartTime);
diff --git a/Engine/Source/Developer/TraceInsights/Private/Insights/MemoryProfiler/Widgets/SMemoryProfilerWindow.cpp b/Engine/Source/Developer/TraceInsights/Private/Insights/MemoryProfiler/Widgets/SMemoryProfilerWindow.cpp
--- a/Engine/Source/Developer/TraceInsights/Private/Insights/MemoryProfiler/Widgets/SMemoryProfilerWindow.cpp
+++ b/Engine/Source/Developer/TraceInsights/Private/Insights/MemoryProfiler/Widgets/SMemoryProfilerWindow.cpp
@@ -97,7 +97,7 @@ TSharedRef<SDockTab> SMemoryProfilerWindow::SpawnTab_TimingView(const FSpawnTabA
 		.ShouldAutosize(false)
 		.TabRole(ETabRole::PanelTab)
 		[
-			SAssignNew(TimingView, STimingView)
+			SAssignNew(TimingView, STimingView, FInsightsManagerTabs::MemoryProfilerTabId)
 		];
 
 	SharedState->SetTimingView(TimingView);
diff --git a/Engine/Source/Developer/TraceInsights/Private/Insights/ViewModels/LoadingTimingTrack.cpp b/Engine/Source/Developer/TraceInsights/Private/Insights/ViewModels/LoadingTimingTrack.cpp
--- a/Engine/Source/Developer/TraceInsights/Private/Insights/ViewModels/LoadingTimingTrack.cpp
+++ b/Engine/Source/Developer/TraceInsights/Private/Insights/ViewModels/LoadingTimingTrack.cpp
@@ -60,6 +60,18 @@ UE_ENABLE_OPTIMIZATION_SHIP
 // FLoadingSharedState
 ////////////////////////////////////////////////////////////////////////////////////////////////////
 
+FLoadingSharedState::FLoadingSharedState(STimingView* InTimingView)
+	: TimingView(InTimingView)
+	, bShowHideAllLoadingTracks(false)
+	//, LoadingTracks
+	, LoadTimeProfilerTimelineCount(0)
+	//, GetEventNameDelegate
+{
+	check(TimingView != nullptr);
+}
+
+////////////////////////////////////////////////////////////////////////////////////////////////////
+
 void FLoadingSharedState::OnBeginSession(Insights::ITimingViewSession& InSession)
 {
 	if (&InSession != TimingView)
@@ -67,7 +79,7 @@ void FLoadingSharedState::OnBeginSession(Insights::ITimingViewSession& InSession
 		return;
 	}
 
-	if (TimingView && TimingView->IsAssetLoadingModeEnabled())
+	if (TimingView->GetName() == FInsightsManagerTabs::LoadingProfilerTabId)
 	{
 		bShowHideAllLoadingTracks = true;
 	}
@@ -189,10 +201,7 @@ void FLoadingSharedState::SetAllLoadingTracksToggle(bool bOnOff)
 		Track.SetVisibilityFlag(bShowHideAllLoadingTracks);
 	}
 
-	if (TimingView)
-	{
-		TimingView->OnTrackVisibilityChanged();
-	}
+	TimingView->OnTrackVisibilityChanged();
 }
 
 ////////////////////////////////////////////////////////////////////////////////////////////////////
diff --git a/Engine/Source/Developer/TraceInsights/Private/Insights/ViewModels/LoadingTimingTrack.h b/Engine/Source/Developer/TraceInsights/Private/Insights/ViewModels/LoadingTimingTrack.h
--- a/Engine/Source/Developer/TraceInsights/Private/Insights/ViewModels/LoadingTimingTrack.h
+++ b/Engine/Source/Developer/TraceInsights/Private/Insights/ViewModels/LoadingTimingTrack.h
@@ -38,7 +38,7 @@ DECLARE_DELEGATE_RetVal_TwoParams(const TCHAR*, FLoadingTrackGetEventNameDelegat
 class FLoadingSharedState : public Insights::ITimingViewExtender, public TSharedFromThis<FLoadingSharedState>
 {
 public:
-	explicit FLoadingSharedState(STimingView* InTimingView) : TimingView(InTimingView) {}
+	explicit FLoadingSharedState(STimingView* InTimingView);
 	virtual ~FLoadingSharedState() = default;
 
 	//////////////////////////////////////////////////
diff --git a/Engine/Source/Developer/TraceInsights/Private/Insights/ViewModels/ThreadTimingTrack.cpp b/Engine/Source/Developer/TraceInsights/Private/Insights/ViewModels/ThreadTimingTrack.cpp
--- a/Engine/Source/Developer/TraceInsights/Private/Insights/ViewModels/ThreadTimingTrack.cpp
+++ b/Engine/Source/Developer/TraceInsights/Private/Insights/ViewModels/ThreadTimingTrack.cpp
@@ -359,6 +359,22 @@ static void AddTimingEventToBuilder(ITimingEventsTrackDrawStateBuilder& Builder,
 // FThreadTimingSharedState
 ////////////////////////////////////////////////////////////////////////////////////////////////////
 
+FThreadTimingSharedState::FThreadTimingSharedState(STimingView* InTimingView)
+	: TimingView(InTimingView)
+	, bShowHideAllGpuTracks(false)
+	, bShowHideAllCpuTracks(false)
+	, GpuTrack()
+	, Gpu2Track()
+	//, CpuTracks
+	//, ThreadGroups
+	, TimingProfilerTimelineCount(0)
+	, LoadTimeProfilerTimelineCount(0)
+{
+	check(TimingView != nullptr);
+}
+
+////////////////////////////////////////////////////////////////////////////////////////////////////
+
 TSharedPtr<FCpuTimingTrack> FThreadTimingSharedState::GetCpuTrack(uint32 InThreadId)
 {
 	TSharedPtr<FCpuTimingTrack>*const TrackPtrPtr = CpuTracks.Find(InThreadId);
@@ -404,7 +420,7 @@ void FThreadTimingSharedState::OnBeginSession(Insights::ITimingViewSession& InSe
 		return;
 	}
 
-	if (TimingView && TimingView->IsAssetLoadingModeEnabled())
+	if (TimingView->GetName() == FInsightsManagerTabs::LoadingProfilerTabId)
 	{
 		bShowHideAllGpuTracks = false;
 		bShowHideAllCpuTracks = false;
@@ -551,7 +567,8 @@ void FThreadTimingSharedState::Tick(Insights::ITimingViewSession& InSession, con
 						FThreadGroup& ThreadGroup = ThreadGroups[GroupName];
 						ThreadGroup.NumTimelines++;
 
-						if (TimingView && TimingView->IsAssetLoadingModeEnabled() && bIsLoadingThread)
+						if (bIsLoadingThread &&
+							TimingView->GetName() == FInsightsManagerTabs::LoadingProfilerTabId)
 						{
 							Track->SetVisibilityFlag(true);
 							ThreadGroup.bIsVisible = true;
@@ -695,10 +712,7 @@ void FThreadTimingSharedState::SetAllCpuTracksToggle(bool bOnOff)
 		KV.Value.bIsVisible = bShowHideAllCpuTracks;
 	}
 
-	if (TimingView)
-	{
-		TimingView->OnTrackVisibilityChanged();
-	}
+	TimingView->OnTrackVisibilityChanged();
 }
 
 ////////////////////////////////////////////////////////////////////////////////////////////////////
@@ -717,10 +731,7 @@ void FThreadTimingSharedState::SetAllGpuTracksToggle(bool bOnOff)
 	}
 	if (GpuTrack.IsValid() || Gpu2Track.IsValid())
 	{
-		if (TimingView)
-		{
-			TimingView->OnTrackVisibilityChanged();
-		}
+		TimingView->OnTrackVisibilityChanged();
 	}
 }
 
@@ -754,10 +765,7 @@ void FThreadTimingSharedState::ToggleTrackVisibilityByGroup_Execute(const TCHAR*
 			}
 		}
 
-		if (TimingView)
-		{
-			TimingView->OnTrackVisibilityChanged();
-		}
+		TimingView->OnTrackVisibilityChanged();
 	}
 }
 
diff --git a/Engine/Source/Developer/TraceInsights/Private/Insights/ViewModels/ThreadTimingTrack.h b/Engine/Source/Developer/TraceInsights/Private/Insights/ViewModels/ThreadTimingTrack.h
--- a/Engine/Source/Developer/TraceInsights/Private/Insights/ViewModels/ThreadTimingTrack.h
+++ b/Engine/Source/Developer/TraceInsights/Private/Insights/ViewModels/ThreadTimingTrack.h
@@ -55,7 +55,7 @@ private:
 	};
 
 public:
-	explicit FThreadTimingSharedState(STimingView* InTimingView) : TimingView(InTimingView) {}
+	explicit FThreadTimingSharedState(STimingView* InTimingView);
 	virtual ~FThreadTimingSharedState() = default;
 
 	TSharedPtr<FGpuTimingTrack> GetGpuTrack() { return GpuTrack; }
diff --git a/Engine/Source/Developer/TraceInsights/Private/Insights/Widgets/STimingProfilerWindow.cpp b/Engine/Source/Developer/TraceInsights/Private/Insights/Widgets/STimingProfilerWindow.cpp
--- a/Engine/Source/Developer/TraceInsights/Private/Insights/Widgets/STimingProfilerWindow.cpp
+++ b/Engine/Source/Developer/TraceInsights/Private/Insights/Widgets/STimingProfilerWindow.cpp
@@ -147,7 +147,7 @@ TSharedRef<SDockTab> STimingProfilerWindow::SpawnTab_TimingView(const FSpawnTabA
 		.ShouldAutosize(false)
 		.TabRole(ETabRole::PanelTab)
 		[
-			SAssignNew(TimingView, STimingView)
+			SAssignNew(TimingView, STimingView, FInsightsManagerTabs::TimingProfilerTabId)
 		];
 
 	TimingView->Reset(true);
diff --git a/Engine/Source/Developer/TraceInsights/Private/Insights/Widgets/STimingView.cpp b/Engine/Source/Developer/TraceInsights/Private/Insights/Widgets/STimingView.cpp
--- a/Engine/Source/Developer/TraceInsights/Private/Insights/Widgets/STimingView.cpp
+++ b/Engine/Source/Developer/TraceInsights/Private/Insights/Widgets/STimingView.cpp
@@ -98,7 +98,6 @@ STimingView::STimingView()
 	, FrameSharedState(MakeShared<FFrameSharedState>(this))
 	, ThreadTimingSharedState(MakeShared<FThreadTimingSharedState>(this))
 	, LoadingSharedState(MakeShared<FLoadingSharedState>(this))
-	, bAssetLoadingMode(false)
 	, FileActivitySharedState(MakeShared<FFileActivitySharedState>(this))
 	, TimeRulerTrack(MakeShared<FTimeRulerTrack>())
 	, DefaultTimeMarker(MakeShared<Insights::FTimeMarker>())
@@ -151,8 +150,10 @@ STimingView::~STimingView()
 
 ////////////////////////////////////////////////////////////////////////////////////////////////////
 
-void STimingView::Construct(const FArguments& InArgs)
+void STimingView::Construct(const FArguments& InArgs, FName InViewName)
 {
+	ViewName = InViewName;
+
 	FSlimHorizontalToolBarBuilder LeftToolbar(nullptr, FMultiBoxCustomization::None);
 	LeftToolbar.SetStyle(&FInsightsStyle::Get(), "SecondaryToolbar");
 
@@ -3649,13 +3650,10 @@ void STimingView::UpdateHoveredTimingEvent(float InMousePosX, float InMousePosY)
 
 void STimingView::OnSelectedTimingEventChanged()
 {
-	if (!bAssetLoadingMode)
+	if (SelectedEvent.IsValid())
 	{
-		if (SelectedEvent.IsValid())
-		{
-			SelectedEvent->GetTrack()->UpdateEventStats(const_cast<ITimingEvent&>(*SelectedEvent));
-			SelectedEvent->GetTrack()->OnEventSelected(*SelectedEvent);
-		}
+		SelectedEvent->GetTrack()->UpdateEventStats(const_cast<ITimingEvent&>(*SelectedEvent));
+		SelectedEvent->GetTrack()->OnEventSelected(*SelectedEvent);
 	}
 
 	OnSelectedEventChangedDelegate.Broadcast(SelectedEvent);
diff --git a/Engine/Source/Developer/TraceInsights/Private/Insights/Widgets/STimingView.h b/Engine/Source/Developer/TraceInsights/Private/Insights/Widgets/STimingView.h
--- a/Engine/Source/Developer/TraceInsights/Private/Insights/Widgets/STimingView.h
+++ b/Engine/Source/Developer/TraceInsights/Private/Insights/Widgets/STimingView.h
@@ -69,7 +69,10 @@ public:
 	 *
 	 * @param	InArgs	The declaration data for this widget
 	 */
-	void Construct(const FArguments& InArgs);
+	void Construct(const FArguments& InArgs, FName InViewName = NAME_None);
+
+	/** Gets the name of the view. */
+	virtual const FName& GetName() const { return ViewName; }
 
 	bool IsCompactModeEnabled() const;
 	void ToggleCompactMode();
@@ -91,9 +94,6 @@ public:
 	TSharedPtr<FLoadingSharedState> GetLoadingSharedState() const { return LoadingSharedState; }
 	TSharedPtr<FFileActivitySharedState> GetFileActivitySharedState() const { return FileActivitySharedState; }
 
-	bool IsAssetLoadingModeEnabled() const { return bAssetLoadingMode; }
-	void EnableAssetLoadingMode() { bAssetLoadingMode = true; }
-
 	void HideAllDefaultTracks();
 
 	/** Gets the time ruler track. It includes the custom time markers (ones user can drag with mouse). */
@@ -509,6 +509,9 @@ protected:
 	void EnumerateFilteredTracks(TSharedPtr<Insights::FFilterConfigurator> FilterConfigurator, EnumerateFilteredTracksCallback Callback);
 
 protected:
+	/** The name of the view. */
+	FName ViewName;
+
 	/** The track's viewport. Encapsulates info about position and scale. */
 	FTimingTrackViewport Viewport;
 
@@ -537,7 +540,6 @@ protected:
 
 	// Shared state for Asset Loading tracks
 	TSharedPtr<FLoadingSharedState> LoadingSharedState;
-	bool bAssetLoadingMode;
 
 	// Shared state for File Activity (I/O) tracks
 	TSharedPtr<FFileActivitySharedState> FileActivitySharedState;
diff --git a/Engine/Source/Developer/TraceInsights/Public/Insights/ITimingViewSession.h b/Engine/Source/Developer/TraceInsights/Public/Insights/ITimingViewSession.h
--- a/Engine/Source/Developer/TraceInsights/Public/Insights/ITimingViewSession.h
+++ b/Engine/Source/Developer/TraceInsights/Public/Insights/ITimingViewSession.h
@@ -66,6 +66,9 @@ class TRACEINSIGHTS_API ITimingViewSession
 public:
 	virtual ~ITimingViewSession() = default;
 
+	/** Gets the name of the view. */
+	virtual const FName& GetName() const = 0;
+
 	/** Adds a new top docked track. */
 	virtual void AddTopDockedTrack(TSharedPtr<FBaseTimingTrack> Track) = 0;
 	/** Removes a top docked track. Returns whether the track was removed or not. */
@@ -137,7 +140,7 @@ public:
 	/** Prevents mouse movements from throttling application updates. */
 	virtual void PreventThrottling() = 0;
 
-	/** Add a slot to the overlay */
+	/** Adds a slot to the overlay. */
 	virtual void AddOverlayWidget(const TSharedRef<SWidget>& InWidget) = 0;
 };
 
