[Insights] Fix crashes when closing the Packages table and the Tasks table.

diff --git a/Engine/Source/Developer/TraceInsights/Private/Insights/CookProfiler/CookProfilerManager.cpp b/Engine/Source/Developer/TraceInsights/Private/Insights/CookProfiler/CookProfilerManager.cpp
--- a/Engine/Source/Developer/TraceInsights/Private/Insights/CookProfiler/CookProfilerManager.cpp
+++ b/Engine/Source/Developer/TraceInsights/Private/Insights/CookProfiler/CookProfilerManager.cpp
@@ -227,6 +227,11 @@ bool FCookProfilerManager::CanSpawnTab_PackageTableTreeView(const FSpawnTabArgs&
 
 void FCookProfilerManager::OnPackageTableTreeViewTabClosed(TSharedRef<SDockTab> TabBeingClosed)
 {
+	if (PackageTableTreeView.IsValid())
+	{
+		PackageTableTreeView->OnClose();
+	}
+
 	PackageTableTreeView.Reset();
 }
 
@@ -234,6 +239,11 @@ void FCookProfilerManager::OnPackageTableTreeViewTabClosed(TSharedRef<SDockTab>
 
 void FCookProfilerManager::OnWindowClosedEvent()
 {
+	if (PackageTableTreeView.IsValid())
+	{
+		PackageTableTreeView->OnClose();
+	}
+
 	TSharedPtr<FTabManager> TimingTabManagerSharedPtr = TimingTabManager.Pin();
 
 	if (TimingTabManagerSharedPtr.IsValid())
diff --git a/Engine/Source/Developer/TraceInsights/Private/Insights/Table/Widgets/STableTreeView.cpp b/Engine/Source/Developer/TraceInsights/Private/Insights/Table/Widgets/STableTreeView.cpp
--- a/Engine/Source/Developer/TraceInsights/Private/Insights/Table/Widgets/STableTreeView.cpp
+++ b/Engine/Source/Developer/TraceInsights/Private/Insights/Table/Widgets/STableTreeView.cpp
@@ -128,6 +128,11 @@ STableTreeView::STableTreeView()
 
 STableTreeView::~STableTreeView()
 {
+	if (bRunInAsyncMode)
+	{
+		checkf(bIsCloseScheduled, TEXT("TableTreeView running in async mode was closed but OnClose() was not called. This can lead to a crash. Call OnClose() from the owner tab/window."))
+	}
+
 	// Remove ourselves from the Insights manager.
 	if (FInsightsManager::Get().IsValid())
 	{
@@ -2989,9 +2994,15 @@ FGraphEventRef STableTreeView::StartApplyFiltersTask(FGraphEventRef Prerequisite
 
 void STableTreeView::OnClose()
 {
-	if (bIsUpdateRunning && !bIsCloseScheduled && InProgressAsyncOperationEvent.IsValid() && !InProgressAsyncOperationEvent->IsComplete())
+	if (bIsCloseScheduled)
+	{
+		return;
+	}
+
+	bIsCloseScheduled = true;
+
+	if (bIsUpdateRunning && InProgressAsyncOperationEvent.IsValid() && !InProgressAsyncOperationEvent->IsComplete())
 	{
-		bIsCloseScheduled = true;
 		CancelCurrentAsyncOp();
 
 		FGraphEventArray Prerequisites;
diff --git a/Engine/Source/Developer/TraceInsights/Private/Insights/TaskGraphProfiler/TaskGraphProfilerManager.cpp b/Engine/Source/Developer/TraceInsights/Private/Insights/TaskGraphProfiler/TaskGraphProfilerManager.cpp
--- a/Engine/Source/Developer/TraceInsights/Private/Insights/TaskGraphProfiler/TaskGraphProfilerManager.cpp
+++ b/Engine/Source/Developer/TraceInsights/Private/Insights/TaskGraphProfiler/TaskGraphProfilerManager.cpp
@@ -254,6 +254,11 @@ bool FTaskGraphProfilerManager::CanSpawnTab_TaskTableTreeView(const FSpawnTabArg
 
 void FTaskGraphProfilerManager::OnTaskTableTreeViewTabClosed(TSharedRef<SDockTab> TabBeingClosed)
 {
+	if (TaskTableTreeView.IsValid())
+	{
+		TaskTableTreeView->OnClose();
+	}
+
 	TaskTableTreeView.Reset();
 }
 
