[Insights] Make MessageLog tab docked by default

This change makes space for the message log window in the tab manager. Whenever the window is opened it shows up as a tab.

diff --git a/Engine/Source/Developer/TraceInsights/Private/Insights/InsightsManager.cpp b/Engine/Source/Developer/TraceInsights/Private/Insights/InsightsManager.cpp
--- a/Engine/Source/Developer/TraceInsights/Private/Insights/InsightsManager.cpp
+++ b/Engine/Source/Developer/TraceInsights/Private/Insights/InsightsManager.cpp
@@ -64,6 +64,7 @@ const FName FInsightsManagerTabs::LoadingProfilerTabId(TEXT("LoadingProfiler"));
 const FName FInsightsManagerTabs::NetworkingProfilerTabId(TEXT("NetworkingProfiler"));
 const FName FInsightsManagerTabs::MemoryProfilerTabId(TEXT("MemoryProfiler"));
 const FName FInsightsManagerTabs::AutomationWindowTabId(TEXT("AutomationWindow"));
+const FName FInsightsManagerTabs::MessageLogTabId(TEXT("MessageLog"));
 
 ////////////////////////////////////////////////////////////////////////////////////////////////////
 // FAvailabilityCheck
@@ -1088,6 +1089,26 @@ void FInsightsManager::ScheduleCommand(const FString& InCmd)
 
 void FInsightsManager::OnSessionAnalysisCompleted()
 {
+	// Report any errors or warning that occurred
+	FMessageLog* Log = Session->GetLog();
+	if (Log)
+	{
+		if (const int32 NumErrors = Log->NumMessages(EMessageSeverity::Error); NumErrors > 0)
+		{
+			Log->Notify(FText::Format(NSLOCTEXT("TraceAnalysis", "FinishedWithErrors", "Analysis finished with {0} error(s)"),
+				FText::AsNumber(NumErrors)));
+		}
+		else if (const int32 NumWarnings = Log->NumMessages(EMessageSeverity::Warning); NumWarnings > 0)
+		{
+			Log->Notify(FText::Format(NSLOCTEXT("TraceAnalysis", "FinishedWithWarnings", "Analysis finished with {0} warning(s)"),
+				FText::AsNumber(NumWarnings)));
+		}
+		else
+		{
+			Log->Notify(NSLOCTEXT("TraceAnalysis", "Finished", "Analysis finished."));
+		}
+	}
+	
 	if (!SessionAnalysisCompletedCmd.IsEmpty())
 	{
 		FOutputDevice& Ar = *GLog;
diff --git a/Engine/Source/Developer/TraceInsights/Private/Insights/TraceInsightsModule.cpp b/Engine/Source/Developer/TraceInsights/Private/Insights/TraceInsightsModule.cpp
--- a/Engine/Source/Developer/TraceInsights/Private/Insights/TraceInsightsModule.cpp
+++ b/Engine/Source/Developer/TraceInsights/Private/Insights/TraceInsightsModule.cpp
@@ -371,7 +371,7 @@ void FTraceInsightsModule::CreateSessionViewer(bool bAllowDebugTools)
 	//////////////////////////////////////////////////
 	// Setup the window's content.
 
-	TSharedRef<FTabManager::FLayout> DefaultLayout = FTabManager::NewLayout("UnrealInsightsLayout_v1.0");
+	TSharedRef<FTabManager::FLayout> DefaultLayout = FTabManager::NewLayout("UnrealInsightsLayout_v1.1");
 
 	AddAreaForSessionViewer(DefaultLayout);
 
@@ -425,6 +425,7 @@ void FTraceInsightsModule::AddAreaForSessionViewer(TSharedRef<FTabManager::FLayo
 	Stack->AddTab(FInsightsManagerTabs::LoadingProfilerTabId, ETabState::ClosedTab);
 	Stack->AddTab(FInsightsManagerTabs::NetworkingProfilerTabId, ETabState::ClosedTab);
 	Stack->AddTab(FInsightsManagerTabs::MemoryProfilerTabId, ETabState::ClosedTab);
+	Stack->AddTab(FInsightsManagerTabs::MessageLogTabId, ETabState::ClosedTab);
 	Stack->SetForegroundTab(FTabId(FInsightsManagerTabs::TimingProfilerTabId));
 
 	Layout->AddArea
diff --git a/Engine/Source/Developer/TraceInsights/Public/Insights/IUnrealInsightsModule.h b/Engine/Source/Developer/TraceInsights/Public/Insights/IUnrealInsightsModule.h
--- a/Engine/Source/Developer/TraceInsights/Public/Insights/IUnrealInsightsModule.h
+++ b/Engine/Source/Developer/TraceInsights/Public/Insights/IUnrealInsightsModule.h
@@ -39,6 +39,7 @@ struct TRACEINSIGHTS_API FInsightsManagerTabs
 	static const FName NetworkingProfilerTabId;
 	static const FName MemoryProfilerTabId;
 	static const FName AutomationWindowTabId;
+	static const FName MessageLogTabId;
 };
 
 ////////////////////////////////////////////////////////////////////////////////////////////////////
