[Insights] Minor code style fixes.

diff --git a/Engine/Source/Developer/TraceInsights/Private/Insights/Table/Widgets/STableTreeView.h b/Engine/Source/Developer/TraceInsights/Private/Insights/Table/Widgets/STableTreeView.h
--- a/Engine/Source/Developer/TraceInsights/Private/Insights/Table/Widgets/STableTreeView.h
+++ b/Engine/Source/Developer/TraceInsights/Private/Insights/Table/Widgets/STableTreeView.h
@@ -172,7 +172,7 @@ public:
 	virtual FText GetCurrentOperationName() const override;
 
 	////////////////////////////////////////////////////////////////////////////////////////////////////
-	
+
 	/** Set a log listing name to be used for any errors or warnings. Must be preregistered by the caller with the MessageLog module. */
 	void SetLogListingName(const FName& InLogListingName) { LogListingName = InLogListingName; }
 	const FName& GetLogListingName() { return LogListingName; }
@@ -583,12 +583,13 @@ protected:
 
 	//////////////////////////////////////////////////
 	// Async
+
 	bool bRunInAsyncMode = false;
 	bool bIsUpdateRunning = false;
 	bool bIsCloseScheduled = false;
 
 	TArray<FTableTreeNodePtr> DummyGroupNodes;
-	FGraphEventRef InProgressAsyncOperationEvent;;
+	FGraphEventRef InProgressAsyncOperationEvent;
 	FGraphEventRef AsyncCompleteTaskEvent;
 	EAsyncOperationType InProgressAsyncOperations = static_cast<EAsyncOperationType>(0);
 	TSharedPtr<class SAsyncOperationStatus> AsyncOperationStatus;
@@ -605,6 +606,7 @@ protected:
 	TArray<FTableTreeNodePtr> NodesToExpand;
 
 	//////////////////////////////////////////////////
+
 	TSharedPtr<FFilterConfigurator> FilterConfigurator;
 	FDelegateHandle OnFilterChangesCommitedHandle;
 	FFilterContext Context;
diff --git a/Engine/Source/Developer/TraceInsights/Private/Insights/ViewModels/FilterConfigurator.cpp b/Engine/Source/Developer/TraceInsights/Private/Insights/ViewModels/FilterConfigurator.cpp
--- a/Engine/Source/Developer/TraceInsights/Private/Insights/ViewModels/FilterConfigurator.cpp
+++ b/Engine/Source/Developer/TraceInsights/Private/Insights/ViewModels/FilterConfigurator.cpp
@@ -73,11 +73,15 @@ bool FFilterConfigurator::ApplyFilters(const FFilterContext& Context) const
 	return RootNode->ApplyFilters(Context);
 }
 
+////////////////////////////////////////////////////////////////////////////////////////////////////
+
 bool FFilterConfigurator::IsKeyUsed(int32 Key) const
 {
 	return KeysUsed.Contains(Key);
 }
 
+////////////////////////////////////////////////////////////////////////////////////////////////////
+
 void FFilterConfigurator::ComputeUsedKeys()
 {
 	KeysUsed.Reset();
@@ -86,4 +90,4 @@ void FFilterConfigurator::ComputeUsedKeys()
 
 ////////////////////////////////////////////////////////////////////////////////////////////////////
 
-} // namespace Insights
\ No newline at end of file
+} // namespace Insights
diff --git a/Engine/Source/Developer/TraceInsights/Private/Insights/ViewModels/FilterConfigurator.h b/Engine/Source/Developer/TraceInsights/Private/Insights/ViewModels/FilterConfigurator.h
--- a/Engine/Source/Developer/TraceInsights/Private/Insights/ViewModels/FilterConfigurator.h
+++ b/Engine/Source/Developer/TraceInsights/Private/Insights/ViewModels/FilterConfigurator.h
@@ -9,7 +9,7 @@
 
 namespace Insights
 {
- 
+
 ////////////////////////////////////////////////////////////////////////////////////////////////////
 
 class FFilterConfigurator : public IFilterExecutor
@@ -36,6 +36,7 @@ public:
 
 	////////////////////////////////////////////////////////////////////////////////////////////////////
 	// OnDestroyedEvent
+
 public:
 	/** The event to execute when an instance is destroyed. */
 	DECLARE_MULTICAST_DELEGATE(FOnDestroyedEvent);
@@ -46,6 +47,7 @@ private:
 
 	////////////////////////////////////////////////////////////////////////////////////////////////////
 	// OnChangesCommitedEvent
+
 public:
 	/** The event to execute when the changes to the Filter Widget are saved by clicking on the OK Button. */
 	DECLARE_MULTICAST_DELEGATE(FOnChangesCommitedEvent);
@@ -57,9 +59,9 @@ private:
 	////////////////////////////////////////////////////////////////////////////////////////////////////
 
 private:
-
 	void ComputeUsedKeys();
 
+private:
 	FFilterConfiguratorNodePtr RootNode;
 
 	TSharedPtr<TArray<TSharedPtr<struct FFilter>>> AvailableFilters;
@@ -69,4 +71,4 @@ private:
 
 ////////////////////////////////////////////////////////////////////////////////////////////////////
 
-} // namespace Insights
\ No newline at end of file
+} // namespace Insights
