UnrealInsights: implement diff mode for csv tables.
- Added raw snapshot export for MemoryAllocTable
- Added ImportTable/DiffTables to StoreWindow
- Refactored UntypedTable to support opening without opened session
- Added UntypedDiffTable widget
- Added MergeService to merge two tables together

#tests diff tables

diff --git a/Engine/Source/Developer/TraceInsights/Private/Insights/CookProfiler/Widgets/SPackageTableTreeView.h b/Engine/Source/Developer/TraceInsights/Private/Insights/CookProfiler/Widgets/SPackageTableTreeView.h
--- a/Engine/Source/Developer/TraceInsights/Private/Insights/CookProfiler/Widgets/SPackageTableTreeView.h
+++ b/Engine/Source/Developer/TraceInsights/Private/Insights/CookProfiler/Widgets/SPackageTableTreeView.h
@@ -16,7 +16,7 @@ namespace Insights
 
 ////////////////////////////////////////////////////////////////////////////////////////////////////
 
-class SPackageTableTreeView : public STableTreeView
+class SPackageTableTreeView : public SSessionTableTreeView
 {
 public:
 	/** Default constructor. */
diff --git a/Engine/Source/Developer/TraceInsights/Private/Insights/ImportTool/TableImportTool.cpp b/Engine/Source/Developer/TraceInsights/Private/Insights/ImportTool/TableImportTool.cpp
--- a/Engine/Source/Developer/TraceInsights/Private/Insights/ImportTool/TableImportTool.cpp
+++ b/Engine/Source/Developer/TraceInsights/Private/Insights/ImportTool/TableImportTool.cpp
@@ -19,6 +19,7 @@
 #include "Insights/InsightsStyle.h"
 #include "Insights/Log.h"
 #include "Insights/Table/ViewModels/UntypedTable.h"
+#include "Insights/Table/Widgets/SUntypedDiffTableTreeView.h"
 #include "Insights/Table/Widgets/SUntypedTableTreeView.h"
 
 #define LOCTEXT_NAMESPACE "ImportTableTool"
@@ -95,12 +96,19 @@ void FTableImportTool::StartImportProcess()
 		return;
 	}
 
+	ImportFile(Filenames[0]);
+}
+
+////////////////////////////////////////////////////////////////////////////////////////////////////
+
+void FTableImportTool::ImportFile(const FString& Filename)
+{
 	using namespace TraceServices;
 
-	FName TableId = GetTableID(Filenames[0]);
-	DisplayTable(TableId);
+	FName TableId = GetTableID(Filename);
+	DisplayImportTable(TableId);
 
-	FTableImportService::ImportTable(Filenames[0],
+	FTableImportService::ImportTable(Filename,
 		TableId,
 		[this](TSharedPtr<FTableImportCallbackParams> Params) {
 			this->TableImportServiceCallback(Params);
@@ -109,7 +117,119 @@ void FTableImportTool::StartImportProcess()
 
 ////////////////////////////////////////////////////////////////////////////////////////////////////
 
-void FTableImportTool::DisplayTable(FName TableViewID)
+void FTableImportTool::StartDiffProcess()
+{
+	TArray<FString> Filenames;
+	bool bDialogResult = false;
+
+	IDesktopPlatform* DesktopPlatform = FDesktopPlatformModule::Get();
+	if (DesktopPlatform)
+	{
+		const FString DefaultPath = FPaths::ProjectSavedDir();
+		const FString DefaultFile = TEXT("");
+		bDialogResult = DesktopPlatform->OpenFileDialog(
+			FSlateApplication::Get().FindBestParentWindowHandleForDialogs(nullptr),
+			LOCTEXT("DiffTableTitle", "Choose Tables").ToString(),
+			DefaultPath,
+			DefaultFile,
+			TEXT("Comma-Separated Values (*.csv)|*.csv|Tab-Separated Values (*.tsv)|*.tsv"),
+			EFileDialogFlags::Multiple,
+			Filenames
+		);
+	}
+
+	if (!bDialogResult)
+	{
+		return;
+	}
+
+	if (DesktopPlatform && Filenames.Num() <  2)
+	{
+		const FString DefaultPath = FPaths::ProjectSavedDir();
+		const FString DefaultFile = TEXT("");
+		bDialogResult = DesktopPlatform->OpenFileDialog(
+			FSlateApplication::Get().FindBestParentWindowHandleForDialogs(nullptr),
+			LOCTEXT("DiffTableBTitle", "Choose Table B").ToString(),
+			DefaultPath,
+			DefaultFile,
+			TEXT("Comma-Separated Values (*.csv)|*.csv|Tab-Separated Values (*.tsv)|*.tsv"),
+			EFileDialogFlags::None,
+			Filenames
+		);
+	}
+
+	if (!bDialogResult)
+	{
+		return;
+	}
+
+	if (Filenames.Num() != 2)
+	{
+		FMessageLog ReportMessageLog(FInsightsManager::Get()->GetLogListingName());
+		ReportMessageLog.Error(LOCTEXT("NotEnoughFilesSelected", "Exactly 2 files are required to be chosen!"));
+		ReportMessageLog.Notify();
+		return;
+	}
+
+	DiffFiles(Filenames[0], Filenames[1]);
+}
+
+////////////////////////////////////////////////////////////////////////////////////////////////////
+
+void FTableImportTool::DiffFiles(const FString& FilenameA, const FString& FilenameB)
+{
+	using namespace TraceServices;
+
+	FName TableViewID = GetTableID(FilenameA+TEXT(",")+FilenameB);
+	DisplayDiffTable(TableViewID);
+
+	FTableImportService::ImportTable(FilenameA, TableViewID, 
+		[this, Name = FilenameA](TSharedPtr<FTableImportCallbackParams> Params)
+		{
+			ensure(Params.IsValid());
+
+			if (const FOpenImportedTableTabData* TableViewPtr = OpenTablesMap.Find(Params->TableId))
+			{
+				TSharedPtr<SUntypedDiffTableTreeView> TableView = StaticCastSharedPtr<SUntypedDiffTableTreeView>(TableViewPtr->TableTreeView);
+
+				if (Params->Result == TraceServices::ETableImportResult::ESuccess)
+				{
+					TableView->UpdateSourceTableA(Name, Params->Table);
+				}
+				else
+				{
+					FMessageLog ReportMessageLog(FInsightsManager::Get()->GetLogListingName());
+					ReportMessageLog.AddMessages(Params->Messages);
+					ReportMessageLog.Notify();
+				}
+			}
+		});
+	FTableImportService::ImportTable(FilenameB, TableViewID, 
+		[this, Name = FilenameB](TSharedPtr<FTableImportCallbackParams> Params)
+		{
+			ensure(Params.IsValid());
+
+			if (const FOpenImportedTableTabData* TableViewPtr = OpenTablesMap.Find(Params->TableId))
+			{
+				TSharedPtr<SUntypedDiffTableTreeView> TableView = StaticCastSharedPtr<SUntypedDiffTableTreeView>(TableViewPtr->TableTreeView);
+
+				if (Params->Result == TraceServices::ETableImportResult::ESuccess)
+				{
+					TableView->UpdateSourceTableB(Name, Params->Table);
+				}
+				else
+				{
+					FMessageLog ReportMessageLog(FInsightsManager::Get()->GetLogListingName());
+					ReportMessageLog.AddMessages(Params->Messages);
+					ReportMessageLog.Notify();
+				}
+			}
+		});
+}
+
+////////////////////////////////////////////////////////////////////////////////////////////////////
+
+void FTableImportTool::DisplayImportTable(FName TableViewID)
 {
 	FText ImportTableTreeViewTabDisplayName;
 	FString FileName = FPaths::GetCleanFilename(TableViewID.GetPlainNameString());
@@ -135,6 +255,45 @@ void FTableImportTool::DisplayTable(FName TableViewID)
 
 ////////////////////////////////////////////////////////////////////////////////////////////////////
 
+void FTableImportTool::DisplayDiffTable(FName TableViewID)
+{
+	FText ImportTableTreeViewTabDisplayName;
+	TArray<FString> Filenames;
+	ensure(TableViewID.GetPlainNameString().ParseIntoArray(Filenames, TEXT(",")) == 2);
+
+	const FText FilenameA = FText::FromString(FPaths::GetCleanFilename(Filenames[0]));
+	const FText FilenameB = FText::FromString(FPaths::GetCleanFilename(Filenames[1]));
+
+	if (TableViewID.GetNumber())
+	{
+		ImportTableTreeViewTabDisplayName = FText::Format(
+			LOCTEXT("TableImportTool.DiffTablesNumericTreeViewTabDisplayName", "{0} vs {1} ({2})"),
+			FilenameA,
+			FilenameB,
+			FText::AsNumber(TableViewID.GetNumber()));
+	}
+	else
+	{
+		ImportTableTreeViewTabDisplayName = FText::Format(
+			LOCTEXT("TableImportTool.DiffTablesTableTreeViewTabDisplayName", "{0} vs {1}"),
+			FilenameA,
+			FilenameB);
+	}
+	const FText ImportTableTreeViewTabTooltipText = FText::Format(
+			LOCTEXT("TableImportTool.DiffTablesTableTreeViewTabTooltip", "{0} vs {1}"),
+			FilenameA,
+			FilenameB);
+
+	FGlobalTabmanager::Get()->RegisterNomadTabSpawner(TableViewID,
+		FOnSpawnTab::CreateSP(this, &FTableImportTool::SpawnTab_TableDiffTreeView, TableViewID, ImportTableTreeViewTabDisplayName))
+		.SetDisplayName(ImportTableTreeViewTabDisplayName)
+		.SetTooltipText(ImportTableTreeViewTabTooltipText)
+		.SetIcon(FSlateIcon(FInsightsStyle::GetStyleSetName(), "Icons.TableTreeView"));
+
+	FGlobalTabmanager::Get()->TryInvokeTab(TableViewID);
+}
+////////////////////////////////////////////////////////////////////////////////////////////////////
+
 void FTableImportTool::TableImportServiceCallback(TSharedPtr<TraceServices::FTableImportCallbackParams> Params)
 {
 	ensure(Params.IsValid());
@@ -192,6 +351,38 @@ END_SLATE_FUNCTION_BUILD_OPTIMIZATION
 
 ////////////////////////////////////////////////////////////////////////////////////////////////////
 
+BEGIN_SLATE_FUNCTION_BUILD_OPTIMIZATION
+TSharedRef<SDockTab> FTableImportTool::SpawnTab_TableDiffTreeView(const FSpawnTabArgs& Args, FName TableViewID, FText InDisplayName)
+{
+	TSharedRef<FUntypedTable> UntypedTable = MakeShared<Insights::FUntypedTable>();
+	UntypedTable->Reset();
+	UntypedTable->SetDisplayName(InDisplayName);
+
+	TSharedPtr<Insights::SUntypedDiffTableTreeView> TableTreeView;
+
+	const TSharedRef<SDockTab> DockTab = SNew(SDockTab)
+		.ShouldAutosize(false)
+		.TabRole(ETabRole::PanelTab)
+		[
+			SAssignNew(TableTreeView, SUntypedDiffTableTreeView, UntypedTable)
+			.RunInAsyncMode(true)
+		];
+
+	TableTreeView->SetCurrentOperationNameOverride(LOCTEXT("TableImportMessage", "Importing Data"));
+
+	FOpenImportedTableTabData OpenTabData;
+	OpenTabData.TableTreeView = TableTreeView;
+	OpenTabData.Tab = DockTab;
+
+	OpenTablesMap.Add(TableViewID, OpenTabData);
+	DockTab->SetOnTabClosed(SDockTab::FOnTabClosedCallback::CreateSP(this, &FTableImportTool::OnTableImportTreeViewTabClosed));
+
+	return DockTab;
+}
+END_SLATE_FUNCTION_BUILD_OPTIMIZATION
+
+////////////////////////////////////////////////////////////////////////////////////////////////////
+
 void FTableImportTool::OnTableImportTreeViewTabClosed(TSharedRef<SDockTab> TabBeingClosed)
 {
 	TSharedRef<SUntypedTableTreeView> TableTreeView = StaticCastSharedRef<SUntypedTableTreeView>(TabBeingClosed->GetContent());
diff --git a/Engine/Source/Developer/TraceInsights/Private/Insights/ImportTool/TableImportTool.h b/Engine/Source/Developer/TraceInsights/Private/Insights/ImportTool/TableImportTool.h
--- a/Engine/Source/Developer/TraceInsights/Private/Insights/ImportTool/TableImportTool.h
+++ b/Engine/Source/Developer/TraceInsights/Private/Insights/ImportTool/TableImportTool.h
@@ -59,15 +59,20 @@ public:
 	virtual void OnWindowClosedEvent() override;
 
 	void StartImportProcess();
+	void ImportFile(const FString& Filename);
+	void StartDiffProcess();
+	void DiffFiles(const FString& FilenameA, const FString& FilenameB);
 
-	void DisplayTable(FName TableViewID);
-	void TableImportServiceCallback(TSharedPtr<TraceServices::FTableImportCallbackParams> Params);
+	void CloseAllOpenTabs();
 
+private:
 	TSharedRef<SDockTab> SpawnTab_TableImportTreeView(const FSpawnTabArgs& Args, FName TableViewID, FText InDisplayName);
+	TSharedRef<SDockTab> SpawnTab_TableDiffTreeView(const FSpawnTabArgs& Args, FName TableViewID, FText InDisplayName);
 	void OnTableImportTreeViewTabClosed(TSharedRef<SDockTab> TabBeingClosed);
 
-	void CloseAllOpenTabs();
-private:
+	void DisplayImportTable(FName TableViewID);
+	void DisplayDiffTable(FName TableViewID);
+	void TableImportServiceCallback(TSharedPtr<TraceServices::FTableImportCallbackParams> Params);
 
 	FName GetTableID(const FString& Path);
 	static TSharedPtr<FTableImportTool> Instance;
diff --git a/Engine/Source/Developer/TraceInsights/Private/Insights/MemoryProfiler/Widgets/SMemAllocTableTreeView.cpp b/Engine/Source/Developer/TraceInsights/Private/Insights/MemoryProfiler/Widgets/SMemAllocTableTreeView.cpp
--- a/Engine/Source/Developer/TraceInsights/Private/Insights/MemoryProfiler/Widgets/SMemAllocTableTreeView.cpp
+++ b/Engine/Source/Developer/TraceInsights/Private/Insights/MemoryProfiler/Widgets/SMemAllocTableTreeView.cpp
@@ -13,6 +13,9 @@
 #include "Widgets/Images/SImage.h"
 #include "Widgets/Input/SCheckBox.h"
 #include "Widgets/SToolTip.h"
+#include "DesktopPlatformModule.h"
+#include "HAL/PlatformFileManager.h"
+#include "Logging/MessageLog.h"
 
 // TraceServices
 #include "Common/ProviderLock.h"
@@ -27,6 +30,7 @@
 #include "Insights/InsightsStyle.h"
 #include "Insights/MemoryProfiler/Common/SymbolSearchPathsHelper.h"
 #include "Insights/MemoryProfiler/MemoryProfilerManager.h"
+#include "Insights/MemoryProfiler/ViewModels/CallstackFormatting.h"
 #include "Insights/MemoryProfiler/ViewModels/MemAllocGroupingByCallstack.h"
 #include "Insights/MemoryProfiler/ViewModels/MemAllocGroupingByHeap.h"
 #include "Insights/MemoryProfiler/ViewModels/MemAllocGroupingBySize.h"
@@ -1470,6 +1474,25 @@ void SMemAllocTableTreeView::ExtendMenu(FMenuBuilder& MenuBuilder)
 		}
 		MenuBuilder.EndSection();
 	}
+
+	{
+		const FText ItemLabel = LOCTEXT("ContextMenu_Export_SubMenu", "Export snaphot");
+		const FText ItemToolTip = LOCTEXT("ContextMenu_Export_Desc_SubMenu", "Export memory snaphot to construct diff later");
+
+		FUIAction Action_ExportSnapshot
+		(
+			FExecuteAction::CreateSP(this, &SMemAllocTableTreeView::ExportMemorySnapshot),
+			FCanExecuteAction::CreateSP(this, &SMemAllocTableTreeView::IsExportMemorySnaphotAvailable)
+		);
+		MenuBuilder.AddMenuEntry(
+			ItemLabel,
+			ItemToolTip,
+			FSlateIcon(FAppStyle::GetAppStyleSetName(), "Icons.Save"),
+			Action_ExportSnapshot,
+			NAME_None,
+			EUserInterfaceActionType::Button
+		);
+	}
 }
 
 ////////////////////////////////////////////////////////////////////////////////////////////////////
@@ -1648,6 +1671,206 @@ void SMemAllocTableTreeView::OpenSourceFileInIDE(const TCHAR* InFile, uint32 Lin
 
 ////////////////////////////////////////////////////////////////////////////////////////////////////
 
+void SMemAllocTableTreeView::ExportMemorySnapshot() const
+{
+	// 1. Choose file
+	FString DefaultFile = TEXT("Table.tsv");
+	if (Table.IsValid() && !Table->GetDisplayName().IsEmpty())
+	{
+		DefaultFile = Table->GetDisplayName().ToString();
+		DefaultFile.RemoveSpacesInline();
+	}
+
+	TArray<FString> SaveFilenames;
+	bool bDialogResult = false;
+
+	if (IDesktopPlatform* DesktopPlatform = FDesktopPlatformModule::Get())
+	{
+		const FString DefaultPath = FPaths::ProjectSavedDir();
+		bDialogResult = DesktopPlatform->SaveFileDialog(
+			FSlateApplication::Get().FindBestParentWindowHandleForDialogs(nullptr),
+			LOCTEXT("ExportFileTitle", "Export Table").ToString(),
+			DefaultPath,
+			DefaultFile,
+			TEXT("Comma-Separated Values (*.csv)|*.csv|Tab-Separated Values (*.tsv)|*.tsv|Text Files (*.txt)|*.txt|All Files (*.*)|*.*"),
+			EFileDialogFlags::None,
+			SaveFilenames
+		);
+	}
+
+	if (!bDialogResult || SaveFilenames.Num() == 0)
+	{
+		return;
+	}
+
+	const FString& Path = SaveFilenames[0];
+	const TCHAR Separator = Path.EndsWith(TEXT(".csv")) ? TEXT(',') : TEXT('\t');
+
+	const auto ExportFileHandle = std::unique_ptr<IFileHandle>(FPlatformFileManager::Get().GetPlatformFile().OpenWrite(*Path));
+
+	if (ExportFileHandle == nullptr)
+	{
+		FMessageLog ReportMessageLog((LogListingName != NAME_None) ? LogListingName : TEXT("Other"));
+		ReportMessageLog.Error(LOCTEXT("FailedToOpenFile", "Export failed. Failed to open file for write."));
+		ReportMessageLog.Notify();
+		return;
+	}
+	
+	UTF16CHAR BOM = UNICODE_BOM;
+	ExportFileHandle->Write((uint8*)&BOM, sizeof(UTF16CHAR));
+	
+	static constexpr TCHAR LineEnd = TEXT('\n');
+	static constexpr TCHAR QuotationMarkBegin = TEXT('\"');
+	static constexpr TCHAR QuotationMarkEnd = TEXT('\"');
+
+	bool IsFirstRow = true;
+	auto WriteColumnHeader = [Separator](FStringBuilderBase& OutData, const FTableColumn& Column)
+		{
+			FString Value = Column.GetShortName().ToString().ReplaceCharWithEscapedChar();
+			int32 CharIndex;
+			if (Value.FindChar(Separator, CharIndex))
+			{
+				OutData += QuotationMarkBegin;
+				OutData += Value;
+				OutData += QuotationMarkEnd;
+			}
+			else
+			{
+				OutData += Value;
+			}
+		};
+	auto WriteColumn = [Separator](FStringBuilderBase& OutData, const FTableColumn& Column, const FTableTreeNode& Node)
+		{
+			FString Value = Column.GetValue(Node)->AsString().ReplaceCharWithEscapedChar();
+			int32 CharIndex;
+			if (Value.FindChar(Separator, CharIndex))
+			{
+				OutData << QuotationMarkBegin;
+				OutData << Value;
+				OutData << QuotationMarkEnd;
+			}
+			else
+			{
+				OutData << Value;
+			}
+		};
+	auto WriteCallstackColumn = [Separator](FStringBuilderBase& OutData, const FTableTreeNode& Node)
+		{
+			const FMemAllocNode& MemAllocNode = static_cast<const FMemAllocNode&>(Node);
+			const FMemoryAlloc* Alloc = MemAllocNode.GetMemAlloc();
+			if (Alloc)
+			{
+				if (!Alloc->Callstack)
+				{
+					OutData << QuotationMarkBegin;
+					OutData << GetCallstackNotAvailableString();
+					OutData << QuotationMarkEnd;
+					return;
+				}
+
+				if (Alloc->Callstack->Num() == 0)
+				{
+					OutData << QuotationMarkBegin;
+					OutData << GetEmptyCallstackString();
+					OutData << QuotationMarkEnd;
+					return;
+				}
+				
+				const uint32 NumCallstackFrames = Alloc->Callstack->Num();
+				check(NumCallstackFrames <= 256);
+				OutData << QuotationMarkBegin;
+				for (uint32 Index = 0; Index < NumCallstackFrames; ++Index)
+				{
+					if (Index != 0)
+					{
+						OutData << TEXT("/");
+					}
+					const TraceServices::FStackFrame* Frame = Alloc->Callstack->Frame(static_cast<uint8>(Index));
+					check(Frame != nullptr);
+					FormatStackFrame(*Frame, OutData, EStackFrameFormatFlags::ModuleAndSymbol);
+				}
+				OutData << QuotationMarkEnd;
+			}
+		};
+	const FTableColumn& StartEventIndexColumn = *GetTable()->FindColumn(FMemAllocTableColumns::StartEventIndexColumnId);
+	const FTableColumn& EndEventIndexColumn = *GetTable()->FindColumn(FMemAllocTableColumns::EndEventIndexColumnId);
+	const FTableColumn& EventDistanceColumn = *GetTable()->FindColumn(FMemAllocTableColumns::EventDistanceColumnId);
+	const FTableColumn& StartTimeColumn = *GetTable()->FindColumn(FMemAllocTableColumns::StartTimeColumnId);
+	const FTableColumn& EndTimeColumn = *GetTable()->FindColumn(FMemAllocTableColumns::EndTimeColumnId);
+	const FTableColumn& DurationColumn = *GetTable()->FindColumn(FMemAllocTableColumns::DurationColumnId);
+	const FTableColumn& AddressColumn = *GetTable()->FindColumn(FMemAllocTableColumns::AddressColumnId);
+	const FTableColumn& MemoryPageColumn = *GetTable()->FindColumn(FMemAllocTableColumns::MemoryPageColumnId);
+	const FTableColumn& SizeColumn = *GetTable()->FindColumn(FMemAllocTableColumns::SizeColumnId);
+	const FTableColumn& TagColumn = *GetTable()->FindColumn(FMemAllocTableColumns::TagColumnId);
+	const FTableColumn& AssetColumn = *GetTable()->FindColumn(FMemAllocTableColumns::AssetColumnId);
+	const FTableColumn& ClassNameColumn = *GetTable()->FindColumn(FMemAllocTableColumns::ClassNameColumnId);
+	const FTableColumn& FunctionColumn = *GetTable()->FindColumn(FMemAllocTableColumns::FunctionColumnId);
+	const FTableColumn& SourceFileColumn = *GetTable()->FindColumn(FMemAllocTableColumns::SourceFileColumnId);
+
+	// 2. Iterate over TreeNodes
+	TStringBuilder<2048> Buffer;
+	for (const TSharedPtr<FTableTreeNode>& Node : TableTreeNodes)
+	{
+		// Export only leaves
+		if (Node->IsGroup()) continue;
+		// String buffer optimization
+		Buffer.Reset();
+
+		if (IsFirstRow)
+		{
+			WriteColumnHeader(Buffer, StartEventIndexColumn); Buffer += Separator;
+			WriteColumnHeader(Buffer, EndEventIndexColumn); Buffer += Separator;
+			WriteColumnHeader(Buffer, EventDistanceColumn); Buffer += Separator;
+			WriteColumnHeader(Buffer, StartTimeColumn); Buffer += Separator;
+			WriteColumnHeader(Buffer, EndTimeColumn); Buffer += Separator;
+			WriteColumnHeader(Buffer, DurationColumn); Buffer += Separator;
+			WriteColumnHeader(Buffer, AddressColumn); Buffer += Separator;
+			WriteColumnHeader(Buffer, MemoryPageColumn); Buffer += Separator;
+			WriteColumnHeader(Buffer, SizeColumn); Buffer += Separator;
+			WriteColumnHeader(Buffer, TagColumn); Buffer += Separator;
+			WriteColumnHeader(Buffer, AssetColumn); Buffer += Separator;
+			WriteColumnHeader(Buffer, ClassNameColumn); Buffer += Separator;
+			WriteColumnHeader(Buffer, FunctionColumn); Buffer += Separator;
+			WriteColumnHeader(Buffer, SourceFileColumn); Buffer += Separator;
+			Buffer += TEXT("Callstack"); Buffer += LineEnd;
+
+			IsFirstRow = false;
+		}
+
+		// 3. Export these column values as is:
+		WriteColumn(Buffer, StartEventIndexColumn, *Node); Buffer += Separator;
+		WriteColumn(Buffer, EndEventIndexColumn, *Node); Buffer += Separator;
+		WriteColumn(Buffer, EventDistanceColumn, *Node); Buffer += Separator;
+		WriteColumn(Buffer, StartTimeColumn, *Node); Buffer += Separator;
+		WriteColumn(Buffer, EndTimeColumn, *Node); Buffer += Separator;
+		WriteColumn(Buffer, DurationColumn, *Node); Buffer += Separator;
+		WriteColumn(Buffer, AddressColumn, *Node); Buffer += Separator;
+		WriteColumn(Buffer, MemoryPageColumn, *Node); Buffer += Separator;
+		WriteColumn(Buffer, SizeColumn, *Node); Buffer += Separator;
+		WriteColumn(Buffer, TagColumn, *Node); Buffer += Separator;
+		WriteColumn(Buffer, AssetColumn, *Node); Buffer += Separator;
+		WriteColumn(Buffer, ClassNameColumn, *Node); Buffer += Separator;
+		WriteColumn(Buffer, FunctionColumn, *Node); Buffer += Separator;
+		WriteColumn(Buffer, SourceFileColumn, *Node); Buffer += Separator;
+		WriteCallstackColumn(Buffer, *Node); Buffer += LineEnd;
+
+		// 5. Write rows to file
+		FTCHARToUTF16 UTF16String(*Buffer, Buffer.Len());
+		ExportFileHandle->Write((const uint8*)UTF16String.Get(), UTF16String.Length() * sizeof(UTF16CHAR));
+	}
+
+	ExportFileHandle->Flush();
+}
+
+////////////////////////////////////////////////////////////////////////////////////////////////////
+
+bool SMemAllocTableTreeView::IsExportMemorySnaphotAvailable() const
+{
+	return !TableTreeNodes.IsEmpty();
+}
+
+////////////////////////////////////////////////////////////////////////////////////////////////////
+
 bool SMemAllocTableTreeView::CanOpenCallstackFrameSourceFileInIDE() const
 {
 	if (TreeView->GetNumItemsSelected() != 1)
diff --git a/Engine/Source/Developer/TraceInsights/Private/Insights/MemoryProfiler/Widgets/SMemAllocTableTreeView.h b/Engine/Source/Developer/TraceInsights/Private/Insights/MemoryProfiler/Widgets/SMemAllocTableTreeView.h
--- a/Engine/Source/Developer/TraceInsights/Private/Insights/MemoryProfiler/Widgets/SMemAllocTableTreeView.h
+++ b/Engine/Source/Developer/TraceInsights/Private/Insights/MemoryProfiler/Widgets/SMemAllocTableTreeView.h
@@ -18,7 +18,7 @@ class FMemoryRuleSpec;
 
 ////////////////////////////////////////////////////////////////////////////////////////////////////
 
-class SMemAllocTableTreeView : public STableTreeView
+class SMemAllocTableTreeView : public SSessionTableTreeView
 {
 public:
 	/** Default constructor. */
@@ -95,6 +95,8 @@ protected:
 	void OpenCallstackFrameSourceFileInIDE();
 	FText GetSelectedCallstackFrameFileName() const;
 	void OpenSourceFileInIDE(const TCHAR* File, uint32 Line) const;
+	void ExportMemorySnapshot() const;
+	bool IsExportMemorySnaphotAvailable() const;
 
 private:
 	void OnQueryInvalidated();
diff --git a/Engine/Source/Developer/TraceInsights/Private/Insights/Table/Widgets/STableTreeView.cpp b/Engine/Source/Developer/TraceInsights/Private/Insights/Table/Widgets/STableTreeView.cpp
--- a/Engine/Source/Developer/TraceInsights/Private/Insights/Table/Widgets/STableTreeView.cpp
+++ b/Engine/Source/Developer/TraceInsights/Private/Insights/Table/Widgets/STableTreeView.cpp
@@ -177,12 +177,6 @@ STableTreeView::~STableTreeView()
 		checkf(bIsCloseScheduled, TEXT("TableTreeView running in async mode was closed but OnClose() was not called. This can lead to a crash. Call OnClose() from the owner tab/window."))
 	}
 
-	// Remove ourselves from the Insights manager.
-	if (FInsightsManager::Get().IsValid())
-	{
-		FInsightsManager::Get()->GetSessionChangedEvent().RemoveAll(this);
-	}
-
 	if (CurrentAsyncOpFilterConfigurator)
 	{
 		delete CurrentAsyncOpFilterConfigurator;
@@ -190,8 +184,6 @@ STableTreeView::~STableTreeView()
 	}
 
 	//FTableTreeViewCommands::Unregister();
-
-	Session.Reset();
 }
 
 ////////////////////////////////////////////////////////////////////////////////////////////////////
@@ -375,15 +367,6 @@ void STableTreeView::ConstructWidget(TSharedPtr<FTable> InTablePtr)
 	InitializeAndShowHeaderColumns();
 
 	InitCommandList();
-
-	// Register ourselves with the Insights manager.
-	FInsightsManager::Get()->GetSessionChangedEvent().AddSP(this, &STableTreeView::InsightsManager_OnSessionChanged);
-
-	// Update the Session (i.e. when analysis session was already started).
-	InsightsManager_OnSessionChanged();
-
-	CreateGroupings();
-	CreateSortings();
 }
 END_SLATE_FUNCTION_BUILD_OPTIMIZATION
 
@@ -953,6 +936,21 @@ TSharedRef<SWidget> STableTreeView::TreeViewHeaderRow_GenerateColumnMenu(const F
 			NAME_None,
 			EUserInterfaceActionType::Button
 		);
+
+		FUIAction Action_HideAllColumns
+		(
+			FExecuteAction::CreateSP(this, &STableTreeView::ContextMenu_HideAllColumns_Execute),
+			FCanExecuteAction::CreateSP(this, &STableTreeView::ContextMenu_HideAllColumns_CanExecute)
+		);
+		MenuBuilder.AddMenuEntry
+		(
+			LOCTEXT("ContextMenu_HideAllColumns", "Hide All Columns"),
+			LOCTEXT("ContextMenu_HideAllColumns_Desc", "Resets tree view to hide all columns (except hierarchy)."),
+			FSlateIcon(FInsightsStyle::GetStyleSetName(), "Icons.ResetColumn"),
+			Action_HideAllColumns,
+			NAME_None,
+			EUserInterfaceActionType::Button
+		);
 	}
 	MenuBuilder.EndSection();
 
@@ -961,22 +959,6 @@ TSharedRef<SWidget> STableTreeView::TreeViewHeaderRow_GenerateColumnMenu(const F
 
 ////////////////////////////////////////////////////////////////////////////////////////////////////
 
-void STableTreeView::InsightsManager_OnSessionChanged()
-{
-	TSharedPtr<const TraceServices::IAnalysisSession> NewSession = FInsightsManager::Get()->GetSession();
-	if (NewSession != Session)
-	{
-		Session = NewSession;
-		Reset();
-	}
-	else
-	{
-		UpdateTree();
-	}
-}
-
-////////////////////////////////////////////////////////////////////////////////////////////////////
-
 void STableTreeView::Tick(const FGeometry& AllottedGeometry, const double InCurrentTime, const float InDeltaTime)
 {
 	if (bRunInAsyncMode && bIsUpdateRunning && !bIsCloseScheduled)
@@ -2810,6 +2792,36 @@ void STableTreeView::ContextMenu_ResetColumns_Execute()
 
 ////////////////////////////////////////////////////////////////////////////////////////////////////
 
+bool STableTreeView::ContextMenu_HideAllColumns_CanExecute() const
+{
+	if (bIsUpdateRunning)
+	{
+		return false;
+	}
+
+	return true;
+}
+
+////////////////////////////////////////////////////////////////////////////////////////////////////
+
+void STableTreeView::ContextMenu_HideAllColumns_Execute()
+{
+	ColumnBeingSorted = GetDefaultColumnBeingSorted();
+	ColumnSortMode = GetDefaultColumnSortMode();
+	UpdateCurrentSortingByColumn();
+
+	for (const TSharedRef<FTableColumn>& ColumnRef : Table->GetColumns())
+	{
+		FTableColumn& Column = ColumnRef.Get();
+		if (Column.IsVisible() && Column.CanBeHidden())
+		{
+			HideColumn(Column);
+		}
+	}
+}
+
+////////////////////////////////////////////////////////////////////////////////////////////////////
+
 void STableTreeView::Reset()
 {
 	StatsStartTime = 0.0;
@@ -3757,6 +3769,52 @@ void STableTreeView::StopAllTableDataTasks(bool bWait)
 	}
 }
 
+////////////////////////////////////////////////////////////////////////////////////////////////////
+// SSessionTableTreeView
+////////////////////////////////////////////////////////////////////////////////////////////////////
+
+SSessionTableTreeView::~SSessionTableTreeView()
+{
+	// Remove ourselves from the Insights manager.
+	if (FInsightsManager::Get().IsValid())
+	{
+		FInsightsManager::Get()->GetSessionChangedEvent().RemoveAll(this);
+	}
+
+	Session.Reset();
+}
+
+////////////////////////////////////////////////////////////////////////////////////////////////////
+
+void SSessionTableTreeView::ConstructWidget(TSharedPtr<FTable> InTablePtr)
+{
+	STableTreeView::ConstructWidget(InTablePtr);
+	// Register ourselves with the Insights manager.
+	FInsightsManager::Get()->GetSessionChangedEvent().AddSP(this, &SSessionTableTreeView::InsightsManager_OnSessionChanged);
+
+	// Update the Session (i.e. when analysis session was already started).
+	InsightsManager_OnSessionChanged();
+
+	CreateGroupings();
+	CreateSortings();
+}
+
+////////////////////////////////////////////////////////////////////////////////////////////////////
+
+void SSessionTableTreeView::InsightsManager_OnSessionChanged()
+{
+	TSharedPtr<const TraceServices::IAnalysisSession> NewSession = FInsightsManager::Get()->GetSession();
+	if (NewSession != Session)
+	{
+		Session = NewSession;
+		Reset();
+	}
+	else
+	{
+		UpdateTree();
+	}
+}
+
 ////////////////////////////////////////////////////////////////////////////////////////////////////
 
 } // namespace Insights
diff --git a/Engine/Source/Developer/TraceInsights/Private/Insights/Table/Widgets/STableTreeView.h b/Engine/Source/Developer/TraceInsights/Private/Insights/Table/Widgets/STableTreeView.h
--- a/Engine/Source/Developer/TraceInsights/Private/Insights/Table/Widgets/STableTreeView.h
+++ b/Engine/Source/Developer/TraceInsights/Private/Insights/Table/Widgets/STableTreeView.h
@@ -33,7 +33,7 @@ class FUICommandList;
 
 namespace TraceServices
 {
-	class IAnalysisSession;
+class IAnalysisSession;
 }
 
 namespace Insights
@@ -180,7 +180,7 @@ public:
 protected:
 	void InitCommandList();
 
-	void ConstructWidget(TSharedPtr<FTable> InTablePtr);
+	virtual void ConstructWidget(TSharedPtr<FTable> InTablePtr);
 	virtual TSharedRef<SWidget> ConstructSearchBox();
 	virtual TSharedRef<SWidget> ConstructAdvancedFiltersButton();
 	virtual TSharedRef<SWidget> ConstructHierarchyBreadcrumbTrail();
@@ -191,9 +191,6 @@ protected:
 
 	void UpdateTree();
 
-	/** Called when the analysis session has changed. */
-	void InsightsManager_OnSessionChanged();
-
 	/**
 	 * Populates OutSearchStrings with the strings that should be used in searching.
 	 *
@@ -426,6 +423,10 @@ protected:
 	bool ContextMenu_ResetColumns_CanExecute() const;
 	void ContextMenu_ResetColumns_Execute();
 
+	// HideAllColumns (ContextMenu)
+	bool ContextMenu_HideAllColumns_CanExecute() const;
+	void ContextMenu_HideAllColumns_Execute();
+
 	//Async
 	virtual void OnPreAsyncUpdate();
 	virtual void OnPostAsyncUpdate();
@@ -494,10 +495,6 @@ protected:
 protected:
 	/** Table view model. */
 	TSharedPtr<Insights::FTable> Table;
-
-	/** The analysis session used to populate this widget. */
-	TSharedPtr<const TraceServices::IAnalysisSession> Session;
-
 	//////////////////////////////////////////////////
 	// Tree View, Columns
 
@@ -632,6 +629,23 @@ protected:
 
 ////////////////////////////////////////////////////////////////////////////////////////////////////
 
+class SSessionTableTreeView : public STableTreeView
+{
+public:
+	virtual ~SSessionTableTreeView() override;
+
+protected:
+	virtual void ConstructWidget(TSharedPtr<FTable> InTablePtr) override;
+
+	/** Called when the analysis session has changed. */
+	void InsightsManager_OnSessionChanged();
+
+	/** The analysis session used to populate this widget. */
+	TSharedPtr<const TraceServices::IAnalysisSession> Session;
+};
+
+////////////////////////////////////////////////////////////////////////////////////////////////////
+
 class FTableTreeViewFilterAsyncTask
 {
 public:
diff --git a/Engine/Source/Developer/TraceInsights/Private/Insights/Table/Widgets/SUntypedDiffTableTreeView.cpp b/Engine/Source/Developer/TraceInsights/Private/Insights/Table/Widgets/SUntypedDiffTableTreeView.cpp
new file mode 100644
--- /dev/null
+++ b/Engine/Source/Developer/TraceInsights/Private/Insights/Table/Widgets/SUntypedDiffTableTreeView.cpp
@@ -0,0 +1,117 @@
+// Copyright Epic Games, Inc. All Rights Reserved.
+
+#include "SUntypedDiffTableTreeView.h"
+
+#include "SlateOptMacros.h"
+#include "Logging/MessageLog.h"
+#include "TraceServices/Model/TableMerge.h"
+
+#define LOCTEXT_NAMESPACE "SUntypedDiffTableTreeView"
+
+namespace Insights
+{
+
+////////////////////////////////////////////////////////////////////////////////////////////////////
+
+void SUntypedDiffTableTreeView::UpdateSourceTableA(const FString& Name, TSharedPtr<TraceServices::IUntypedTable> SourceTable)
+{
+	TableNameA = Name;
+	SourceTableA = SourceTable;
+
+	if (SourceTableA.IsValid() && SourceTableB.IsValid())
+	{
+		RequestMergeTables();
+	}
+}
+
+////////////////////////////////////////////////////////////////////////////////////////////////////
+
+void SUntypedDiffTableTreeView::UpdateSourceTableB(const FString& Name, TSharedPtr<TraceServices::IUntypedTable> SourceTable)
+{
+	TableNameB = Name;
+	SourceTableB = SourceTable;
+
+	if (SourceTableA.IsValid() && SourceTableB.IsValid())
+	{
+		RequestMergeTables();
+	}
+}
+
+////////////////////////////////////////////////////////////////////////////////////////////////////
+
+FReply SUntypedDiffTableTreeView::SwapTables_OnClicked()
+{
+	Swap(TableNameA, TableNameB);
+	Swap(SourceTableA, SourceTableB);
+
+	if (SourceTableA.IsValid() && SourceTableB.IsValid())
+	{
+		RequestMergeTables();
+	}
+
+	return FReply::Handled();
+}
+
+////////////////////////////////////////////////////////////////////////////////////////////////////
+
+FText SUntypedDiffTableTreeView::GetSwapButtonText() const
+{
+	const FText CleanA = FText::FromString(FPaths::GetCleanFilename(TableNameA));
+	const FText CleanB = FText::FromString(FPaths::GetCleanFilename(TableNameB));
+	return FText::Format(LOCTEXT("ToolBar_SwapButtonText", "SWAP (A: {0} - B: {1})"), CleanA, CleanB);
+}
+
+////////////////////////////////////////////////////////////////////////////////////////////////////
+
+BEGIN_SLATE_FUNCTION_BUILD_OPTIMIZATION
+TSharedPtr<SWidget> SUntypedDiffTableTreeView::ConstructToolbar()
+{
+	return
+		SNew(SHorizontalBox)
+		+ SHorizontalBox::Slot()
+		[
+			SNew(SButton)
+			.ButtonStyle(&FAppStyle::Get().GetWidgetStyle<FButtonStyle>("PrimaryButton"))
+			.ForegroundColor(FSlateColor::UseStyle())
+			.HAlign(HAlign_Center)
+			.Text(this, &SUntypedDiffTableTreeView::GetSwapButtonText)
+			.OnClicked(this, &SUntypedDiffTableTreeView::SwapTables_OnClicked)
+		];
+}
+END_SLATE_FUNCTION_BUILD_OPTIMIZATION
+
+////////////////////////////////////////////////////////////////////////////////////////////////////
+
+void SUntypedDiffTableTreeView::RequestMergeTables()
+{
+	SetCurrentOperationNameOverride(LOCTEXT("TableDiffMessage", "Merging Tables"));
+
+	// Clear current table
+	if (TreeView.IsValid())
+	{
+		TreeView->SetTreeItemsSource(&DummyGroupNodes);
+		TreeView->RequestTreeRefresh();
+	}
+
+	TraceServices::FTableMergeService::MergeTables(SourceTableA, SourceTableB,
+		[this](TSharedPtr<TraceServices::FTableDiffCallbackParams> Params)
+		{
+			ClearCurrentOperationNameOverride();
+			if (Params->Result == TraceServices::ETableDiffResult::ESuccess)
+			{
+				this->UpdateSourceTable(Params->Table);
+			}
+			else
+			{
+				FMessageLog ReportMessageLog(FInsightsManager::Get()->GetLogListingName());
+				ReportMessageLog.AddMessages(Params->Messages);
+				ReportMessageLog.Notify();
+			}
+		});
+}
+
+////////////////////////////////////////////////////////////////////////////////////////////////////
+
+} // namespace Insights
+
+#undef LOCTEXT_NAMESPACE
diff --git a/Engine/Source/Developer/TraceInsights/Private/Insights/Table/Widgets/SUntypedDiffTableTreeView.h b/Engine/Source/Developer/TraceInsights/Private/Insights/Table/Widgets/SUntypedDiffTableTreeView.h
new file mode 100644
--- /dev/null
+++ b/Engine/Source/Developer/TraceInsights/Private/Insights/Table/Widgets/SUntypedDiffTableTreeView.h
@@ -0,0 +1,39 @@
+// Copyright Epic Games, Inc. All Rights Reserved.
+
+#pragma once
+
+#include "CoreMinimal.h"
+
+// Insights
+#include "SUntypedTableTreeView.h"
+#include "Insights/Table/ViewModels/UntypedTable.h"
+
+namespace Insights
+{
+
+////////////////////////////////////////////////////////////////////////////////////////////////////
+
+class SUntypedDiffTableTreeView : public SUntypedTableTreeView
+{
+public:
+	void UpdateSourceTableA(const FString& Name, TSharedPtr<TraceServices::IUntypedTable> SourceTable);
+	void UpdateSourceTableB(const FString& Name, TSharedPtr<TraceServices::IUntypedTable> SourceTable);
+
+protected:
+	FReply SwapTables_OnClicked();
+	FText GetSwapButtonText() const;
+
+	virtual TSharedPtr<SWidget> ConstructToolbar() override;
+
+	void RequestMergeTables();
+
+private:
+	TSharedPtr<TraceServices::IUntypedTable> SourceTableA;
+	TSharedPtr<TraceServices::IUntypedTable> SourceTableB;
+	FString TableNameA;
+	FString TableNameB;
+};
+
+////////////////////////////////////////////////////////////////////////////////////////////////////
+
+} // namespace Insights
diff --git a/Engine/Source/Developer/TraceInsights/Private/Insights/Table/Widgets/SUntypedTableTreeView.cpp b/Engine/Source/Developer/TraceInsights/Private/Insights/Table/Widgets/SUntypedTableTreeView.cpp
--- a/Engine/Source/Developer/TraceInsights/Private/Insights/Table/Widgets/SUntypedTableTreeView.cpp
+++ b/Engine/Source/Developer/TraceInsights/Private/Insights/Table/Widgets/SUntypedTableTreeView.cpp
@@ -81,7 +81,7 @@ void SUntypedTableTreeView::RebuildTree(bool bResync)
 	TSharedPtr<TraceServices::IUntypedTable> SourceTable = UntypedTable->GetSourceTable();
 	TSharedPtr<TraceServices::IUntypedTableReader> TableReader = UntypedTable->GetTableReader();
 
-	if (Session.IsValid() && SourceTable.IsValid() && TableReader.IsValid())
+	if (SourceTable.IsValid() && TableReader.IsValid())
 	{
 		const int32 TotalRowCount = static_cast<int32>(SourceTable->GetRowCount());
 		if (TotalRowCount != TableTreeNodes.Num())
diff --git a/Engine/Source/Developer/TraceInsights/Private/Insights/TaskGraphProfiler/Widgets/STaskTableTreeView.h b/Engine/Source/Developer/TraceInsights/Private/Insights/TaskGraphProfiler/Widgets/STaskTableTreeView.h
--- a/Engine/Source/Developer/TraceInsights/Private/Insights/TaskGraphProfiler/Widgets/STaskTableTreeView.h
+++ b/Engine/Source/Developer/TraceInsights/Private/Insights/TaskGraphProfiler/Widgets/STaskTableTreeView.h
@@ -16,7 +16,7 @@ namespace Insights
 
 ////////////////////////////////////////////////////////////////////////////////////////////////////
 
-class STaskTableTreeView : public STableTreeView
+class STaskTableTreeView : public SSessionTableTreeView
 {
 private:
 	struct FColumnConfig
diff --git a/Engine/Source/Developer/TraceInsights/Private/Insights/Widgets/SStartPageWindow.cpp b/Engine/Source/Developer/TraceInsights/Private/Insights/Widgets/SStartPageWindow.cpp
--- a/Engine/Source/Developer/TraceInsights/Private/Insights/Widgets/SStartPageWindow.cpp
+++ b/Engine/Source/Developer/TraceInsights/Private/Insights/Widgets/SStartPageWindow.cpp
@@ -44,6 +44,7 @@
 #include "Insights/StoreService/StoreBrowser.h"
 #include "Insights/InsightsStyle.h"
 #include "Insights/Version.h"
+#include "Insights/ImportTool/TableImportTool.h"
 #include "Insights/Widgets/SInsightsSettings.h"
 #include "Insights/Widgets/SLazyToolTip.h"
 
@@ -1954,6 +1955,10 @@ FReply STraceStoreWindow::OnDragOver(const FGeometry& MyGeometry, const FDragDro
 				{
 					return FReply::Handled();
 				}
+				if (DraggedFileExtension == TEXT(".csv") || DraggedFileExtension == TEXT(".tsv"))
+				{
+					return FReply::Handled();
+				}
 			}
 		}
 	}
@@ -1980,6 +1985,11 @@ FReply STraceStoreWindow::OnDrop(const FGeometry& MyGeometry, const FDragDropEve
 					OpenTraceFile(Files[0]);
 					return FReply::Handled();
 				}
+				if (DraggedFileExtension == TEXT(".csv") || DraggedFileExtension == TEXT(".tsv"))
+				{
+					Insights::FTableImportTool::Get()->ImportFile(Files[0]);
+					return FReply::Handled();
+				}
 			}
 		}
 	}
@@ -2095,7 +2105,24 @@ TSharedRef<SWidget> STraceStoreWindow::MakeTraceListMenu()
 			NAME_None,
 			EUserInterfaceActionType::Button
 		);
+		MenuBuilder.AddMenuEntry(
+			LOCTEXT("ImportTableButtonLabel", "Import table..."),
+			LOCTEXT("ImportTableButtonTooltip", "Open .csv or .tsv file."),
+			FSlateIcon(FAppStyle::Get().GetStyleSetName(), "Icons.FolderOpen"),
+			FUIAction(FExecuteAction::CreateLambda([]{ Insights::FTableImportTool::Get()->StartImportProcess(); })),
+			NAME_None,
+			EUserInterfaceActionType::Button
+		);
+		MenuBuilder.AddMenuEntry(
+			LOCTEXT("DiffTablesButtonLabel", "Diff tables..."),
+			LOCTEXT("DiffTablesButtonTooltip", "Open 2 table files in diff mode."),
+			FSlateIcon(FAppStyle::Get().GetStyleSetName(), "Icons.FolderOpen"),
+			FUIAction(FExecuteAction::CreateLambda([]{ Insights::FTableImportTool::Get()->StartDiffProcess(); })),
+			NAME_None,
+			EUserInterfaceActionType::Button
+		);
 	}
+
 	MenuBuilder.EndSection();
 
 	MenuBuilder.BeginSection("AvailableTraces", LOCTEXT("TraceListMenu_Section_AvailableTraces", "Top Most Recently Created Traces"));
diff --git a/Engine/Source/Developer/TraceServices/Private/Model/TableImportTask.cpp b/Engine/Source/Developer/TraceServices/Private/Model/TableImportTask.cpp
--- a/Engine/Source/Developer/TraceServices/Private/Model/TableImportTask.cpp
+++ b/Engine/Source/Developer/TraceServices/Private/Model/TableImportTask.cpp
@@ -91,7 +91,7 @@ ETableImportResult FTableImportTask::ImportTable()
 
 bool FTableImportTask::ParseHeader(const FString& HeaderLine)
 {
-	HeaderLine.ParseIntoArray(ColumnNames, *Separator);
+	HeaderLine.ParseIntoArray(ColumnNames, *Separator, false);
 	if (ColumnNames.Num() == 0)
 	{
 		AddError(FText::Format(LOCTEXT("NoColumnsMsg", "Import failed because the file did not contain any columns."), FText::FromString(FilePath)));
@@ -119,16 +119,21 @@ bool FTableImportTask::CreateLayout(const FString& Line)
 		{
 			return Row.GetValue(Index);
 		};
+		
+		if (ColumnNames[Index].IsEmpty())
+		{
+			ColumnNames[Index] = FString::Format(TEXT("Column {0}"), {Index});
+		}
 
 		if (Values[Index].IsNumeric())
 		{
 			if (Values[Index].Contains(TEXT(".")))
 			{
-				Layout.AddColumn<double>(*ColumnNames[Index], ProjectorFunc);
+				Layout.AddColumn<double>(*ColumnNames[Index], ProjectorFunc, TableColumnDisplayHint_Summable);
 			}
 			else
 			{
-				Layout.AddColumn<uint32>(*ColumnNames[Index], ProjectorFunc);
+				Layout.AddColumn<uint32>(*ColumnNames[Index], ProjectorFunc, TableColumnDisplayHint_Summable);
 			}
 		}
 		else
@@ -162,33 +167,33 @@ bool FTableImportTask::ParseData(TArray<FString>& Lines)
 		{
 			ETableColumnType ColumnType = Layout.GetColumnType(ValueIndex);
 			const TCHAR* Value = *Values[ValueIndex];
-			if (ColumnType == ETableColumnType::TableColumnType_CString)
+			if (ColumnType == TableColumnType_CString)
 			{
 				const TCHAR* StoredValue = Table->GetStringStore().Store(Value);
 				NewRow.SetValue(ValueIndex, StoredValue);
 			}
-			else if (ColumnType == ETableColumnType::TableColumnType_Double)
+			else if (ColumnType == TableColumnType_Double)
 			{
 				if (!Values[ValueIndex].IsNumeric())
 				{
-					Layout.SetColumnType(ValueIndex, ETableColumnType::TableColumnType_CString);
+					Layout.SetColumnType(ValueIndex, TableColumnType_CString);
 					Restart = true;
 					break;
 				}
 
 				NewRow.SetValue(ValueIndex, FCString::Atod(Value));
 			}
-			else if (ColumnType == ETableColumnType::TableColumnType_Int)
+			else if (ColumnType == TableColumnType_Int)
 			{
 				if (!Values[ValueIndex].IsNumeric())
 				{
-					Layout.SetColumnType(ValueIndex, ETableColumnType::TableColumnType_CString);
+					Layout.SetColumnType(ValueIndex, TableColumnType_CString);
 					Restart = true;
 					break;
 				}
 				else if (Values[ValueIndex].Contains(TEXT(".")))
 				{
-					Layout.SetColumnType(ValueIndex, ETableColumnType::TableColumnType_Double);
+					Layout.SetColumnType(ValueIndex, TableColumnType_Double);
 					Restart = true;
 					break;
 				}
@@ -225,14 +230,14 @@ void FTableImportTask::SplitLineIntoValues(const FString& InLine, TArray<FString
 		}
 		else if(!IsInQuotes && InLine[Index] == Separator[0])
 		{
-			FString Value = InLine.Mid(Start, Index - Start);
-			OutValues.Add(Value.TrimQuotes());
+			FString Value = InLine.Mid(Start, Index - Start).TrimQuotes();
+			OutValues.Add(std::move(Value));
 			Start = Index + 1;
 		}
 	}
 
-	FString Value = InLine.Mid(Start, Index - Start);
-	OutValues.Add(Value.TrimQuotes());
+	FString Value = InLine.Mid(Start, Index - Start).TrimQuotes();
+	OutValues.Add(std::move(Value));
 }
 
 bool FTableImportTask::LoadFileToStringArray(const FString &InFilePath, TArray<FString>& Lines)
diff --git a/Engine/Source/Developer/TraceServices/Private/Model/TableMergeTask.cpp b/Engine/Source/Developer/TraceServices/Private/Model/TableMergeTask.cpp
new file mode 100644
--- /dev/null
+++ b/Engine/Source/Developer/TraceServices/Private/Model/TableMergeTask.cpp
@@ -0,0 +1,209 @@
+// Copyright Epic Games, Inc. All Rights Reserved.
+
+#include "TableMergeTask.h"
+
+#include "Logging/TokenizedMessage.h"
+#include "Tasks/Task.h"
+
+namespace TraceServices
+{
+
+#define LOCTEXT_NAMESPACE "TableMergeTask"
+	
+FTableMergeTask::FTableMergeTask(const TSharedPtr<IUntypedTable>& InTableA, const TSharedPtr<IUntypedTable>& InTableB, FTableMergeService::TableDiffCallback InCallback)
+	: TableA(InTableA)
+	, TableB(InTableB)
+	, Callback(InCallback)
+{
+}
+
+void FTableMergeTask::operator()()
+{
+	const ETableDiffResult Result = MergeTables();
+
+	TSharedPtr<FTableDiffCallbackParams> Params = MakeShared<FTableDiffCallbackParams>();
+	Params->Result = Result;
+	Params->Table = TableC;
+	Params->Messages = Messages;
+
+	FFunctionGraphTask::CreateAndDispatchWhenReady(
+		[Callback = this->Callback, Params]()
+		{
+			Callback(Params);
+		},
+		TStatId(),
+		nullptr,
+		ENamedThreads::GameThread);
+}
+
+TraceServices::ETableDiffResult FTableMergeTask::MergeTables()
+{
+	const ITableLayout& LayoutA = TableA->GetLayout();
+	const ITableLayout& LayoutB = TableB->GetLayout();
+
+	if (LayoutA.GetColumnCount() != LayoutB.GetColumnCount())
+	{
+		AddError(LOCTEXT("IncompatibleLayoutColumnCount", "Layout incompatible: column count is different."));
+		return ETableDiffResult::EFail;
+	}
+
+	TableC = MakeShared<TImportTable<FImportTableRow>>();
+		
+	if (!BuildCLayout(LayoutA, LayoutB, TableC->EditLayout()))
+	{
+		return ETableDiffResult::EFail;
+	}
+
+	const TSharedPtr<IUntypedTableReader> TableReaderA = MakeShareable(TableA->CreateReader());
+	const TSharedPtr<IUntypedTableReader> TableReaderB = MakeShareable(TableB->CreateReader());
+	FillResultUsing<FOnlyA, Sign::Positive>(*TableReaderA);
+	FillResultUsing<FOnlyB, Sign::Negative>(*TableReaderB);
+
+	return ETableDiffResult::ESuccess;
+}
+
+bool FTableMergeTask::BuildCLayout(const ITableLayout& LayoutA, const ITableLayout& LayoutB, TTableLayout<FImportTableRow>& LayoutC)
+{
+	for (uint64 ColIndex = 0; ColIndex < LayoutA.GetColumnCount(); ++ColIndex)
+	{
+		const ETableColumnType ColumnType = LayoutA.GetColumnType(ColIndex);
+		const TCHAR* ColumnName = LayoutA.GetColumnName(ColIndex);
+		const uint32 Flags = LayoutA.GetColumnDisplayHintFlags(ColIndex);
+
+		ensure(ColumnType == LayoutB.GetColumnType(ColIndex));
+		ensure(FCString::Stricmp(ColumnName, LayoutB.GetColumnName(ColIndex)) == 0);
+
+		auto ProjectorFuncCreator = [&LayoutC]
+		{
+			return [ColIndex = LayoutC.GetColumnCount()](const FImportTableRow& Row)
+			{
+				return Row.GetValue(ColIndex);
+			};
+		};
+
+		FString DiffColumnName = FString::Format(TEXT("{0} (A-B)"), {ColumnName});
+		FString AColumnName = FString::Format(TEXT("{0} (A)"), {ColumnName});
+		FString BColumnName = FString::Format(TEXT("{0} (B)"), {ColumnName});
+
+		switch (ColumnType)
+		{
+		case TableColumnType_CString:
+			LayoutC.AddColumn<const TCHAR*>(ColumnName, ProjectorFuncCreator(), Flags);
+			ColumnsMap.AddDefaulted_GetRef().Emplace<FBoth>(ColIndex);
+			break;
+		case TableColumnType_Bool:
+			{
+				if ((Flags & TableColumnDisplayHint_Summable) != 0)
+				{
+					// DIFF column
+					LayoutC.AddColumn<bool>(*DiffColumnName, ProjectorFuncCreator(), Flags);
+					ColumnsMap.AddDefaulted_GetRef().Emplace<FMerged>(ColIndex);
+					// OriginalA
+					LayoutC.AddColumn<bool>(*AColumnName, ProjectorFuncCreator(), Flags);
+					ColumnsMap.AddDefaulted_GetRef().Emplace<FOnlyA>(ColIndex);
+					// OriginalB
+					LayoutC.AddColumn<bool>(*BColumnName, ProjectorFuncCreator(), Flags);
+					ColumnsMap.AddDefaulted_GetRef().Emplace<FOnlyB>(ColIndex);
+				}
+				else
+				{
+					LayoutC.AddColumn<bool>(ColumnName, ProjectorFuncCreator(), Flags);
+					ColumnsMap.AddDefaulted_GetRef().Emplace<FBoth>(ColIndex);
+				}
+			}
+			break;
+		case TableColumnType_Int:
+			{
+				if ((Flags & TableColumnDisplayHint_Summable) != 0)
+				{
+					// DIFF column
+					LayoutC.AddColumn<int64>(*DiffColumnName, ProjectorFuncCreator(), Flags);
+					ColumnsMap.AddDefaulted_GetRef().Emplace<FMerged>(ColIndex);
+					// OriginalA
+					LayoutC.AddColumn<uint64>(*AColumnName, ProjectorFuncCreator(), Flags);
+					ColumnsMap.AddDefaulted_GetRef().Emplace<FOnlyA>(ColIndex);
+					// OriginalB
+					LayoutC.AddColumn<uint64>(*BColumnName, ProjectorFuncCreator(), Flags);
+					ColumnsMap.AddDefaulted_GetRef().Emplace<FOnlyB>(ColIndex);
+				}
+				else
+				{
+					AddError(FText::Format(
+						LOCTEXT("LayoutInclompatibleIntNotSummmable", "Column {0} has int values but is not summable."), 
+						{FText::FromStringView(ColumnName)}));
+					return false;
+				}
+			}
+			break;
+		case TableColumnType_Float:
+			{
+				if ((Flags & TableColumnDisplayHint_Summable) != 0)
+				{
+					// DIFF column
+					LayoutC.AddColumn<float>(*DiffColumnName, ProjectorFuncCreator(), Flags);
+					ColumnsMap.AddDefaulted_GetRef().Emplace<FMerged>(ColIndex);
+					// OriginalA
+					LayoutC.AddColumn<float>(*AColumnName, ProjectorFuncCreator(), Flags);
+					ColumnsMap.AddDefaulted_GetRef().Emplace<FOnlyA>(ColIndex);
+					// OriginalB
+					LayoutC.AddColumn<float>(*BColumnName, ProjectorFuncCreator(), Flags);
+					ColumnsMap.AddDefaulted_GetRef().Emplace<FOnlyB>(ColIndex);
+				}
+				else
+				{
+					AddError(FText::Format(
+						LOCTEXT("LayoutInclompatibleFloatNotSummmable", "Column {0} has float values but is not summable."), 
+						{FText::FromStringView(ColumnName)}));
+					return false;
+				}
+			}
+			break;
+		case TableColumnType_Double:
+			{
+				if ((Flags & TableColumnDisplayHint_Summable) != 0)
+				{
+					// DIFF column
+					LayoutC.AddColumn<double>(*DiffColumnName, ProjectorFuncCreator(), Flags);
+					ColumnsMap.AddDefaulted_GetRef().Emplace<FMerged>(ColIndex);
+					// OriginalA
+					LayoutC.AddColumn<double>(*AColumnName, ProjectorFuncCreator(), Flags);
+					ColumnsMap.AddDefaulted_GetRef().Emplace<FOnlyA>(ColIndex);
+					// OriginalB
+					LayoutC.AddColumn<double>(*BColumnName, ProjectorFuncCreator(), Flags);
+					ColumnsMap.AddDefaulted_GetRef().Emplace<FOnlyB>(ColIndex);
+				}
+				else
+				{
+					AddError(FText::Format(
+						LOCTEXT("LayoutInclompatibleDoubleNotSummmable", "Column {0} has double values but is not summable."), 
+						{FText::FromStringView(ColumnName)}));
+					return false;
+				}
+			}
+			break;
+		case TableColumnType_Invalid: 
+			AddError(FText::Format(
+				LOCTEXT("LayoutInclompatibleInvalidColumn", "Column {0} has invalid type."), 
+				{FText::FromStringView(ColumnName)}));
+			return false;
+		default: ;
+		}
+	}
+
+	return true;
+}
+
+void FTableMergeTask::AddError(const FText& Msg)
+{
+	Messages.Add(FTokenizedMessage::Create(EMessageSeverity::Error, Msg));
+}
+
+void FTableMergeService::MergeTables(const TSharedPtr<IUntypedTable>& TableA, const TSharedPtr<IUntypedTable>& TableB, TableDiffCallback InCallback)
+{
+	UE::Tasks::Launch(UE_SOURCE_LOCATION, FTableMergeTask(TableA, TableB, InCallback));
+}
+
+#undef LOCTEXT_NAMESPACE
+
+} // namespace TraceServices
+
diff --git a/Engine/Source/Developer/TraceServices/Private/Model/TableMergeTask.h b/Engine/Source/Developer/TraceServices/Private/Model/TableMergeTask.h
new file mode 100644
--- /dev/null
+++ b/Engine/Source/Developer/TraceServices/Private/Model/TableMergeTask.h
@@ -0,0 +1,121 @@
+// Copyright Epic Games, Inc. All Rights Reserved.
+
+#pragma once
+#include "TableImportTask.h"
+#include "Tables.h"
+#include "TraceServices/Model/TableMerge.h"
+
+namespace TraceServices
+{
+class FImportTableRow;
+
+class FTableMergeTask
+{
+public:
+	FTableMergeTask(const TSharedPtr<IUntypedTable>& InTableA, const TSharedPtr<IUntypedTable>& InTableB, FTableMergeService::TableDiffCallback InCallback);
+
+	void operator()();
+
+private:
+	ETableDiffResult MergeTables();
+
+	struct FBoth { FBoth(uint64 Index) : ColIndex(Index) {} uint64 ColIndex; };
+	struct FOnlyA { FOnlyA(uint64 Index) : ColIndex(Index) {} uint64 ColIndex; };
+	struct FOnlyB { FOnlyB(uint64 Index) : ColIndex(Index) {} uint64 ColIndex; };
+	struct FMerged { FMerged(uint64 Index) : ColIndex(Index) {} uint64 ColIndex; };
+	enum class Sign { Positive, Negative };
+	
+	bool BuildCLayout(const ITableLayout& LayoutA, const ITableLayout& LayoutB, TTableLayout<FImportTableRow>& LayoutC);
+
+	template<typename Variant, Sign Operation>
+	void FillResultUsing(IUntypedTableReader& SourceTableReader) const;
+
+	void AddError(const FText& Msg);
+
+private:
+	const TSharedPtr<IUntypedTable> TableA;
+	const TSharedPtr<IUntypedTable> TableB;
+	const FTableMergeService::TableDiffCallback Callback;
+
+	TSharedPtr<TImportTable<FImportTableRow>> TableC;
+	TArray<TSharedRef<FTokenizedMessage>> Messages;
+	TArray<TVariant<FEmptyVariantState, FBoth, FOnlyA, FOnlyB, FMerged>> ColumnsMap;
+};
+
+template <typename Variant, FTableMergeTask::Sign Operation>
+void FTableMergeTask::FillResultUsing(IUntypedTableReader& SourceTableReader) const
+{
+	for (SourceTableReader.SetRowIndex(0); SourceTableReader.IsValid(); SourceTableReader.NextRow())
+	{
+		FImportTableRow& Row = TableC->AddRow();
+		Row.SetNumValues(ColumnsMap.Num());
+
+		for (uint64 DestIndex = 0; DestIndex < ColumnsMap.Num(); ++DestIndex)
+		{
+			const ETableColumnType ColumnType = TableC->GetLayout().GetColumnType(DestIndex);
+
+			uint64 SourceIndex;
+			int64 Multiplier = 1;
+			if (const auto* Column = ColumnsMap[DestIndex].TryGet<FBoth>())
+			{
+				SourceIndex = Column->ColIndex;
+			}
+			else if (const auto* ColumnA = ColumnsMap[DestIndex].TryGet<Variant>())
+			{
+				SourceIndex = ColumnA->ColIndex;
+			}
+			else if (const auto* ColumnM = ColumnsMap[DestIndex].TryGet<FMerged>())
+			{
+				SourceIndex = ColumnM->ColIndex;
+				Multiplier = Operation == Sign::Negative ? -1 : 1;
+			}
+			else
+			{
+				switch (ColumnType)
+				{
+				case TableColumnType_Bool:
+					Row.SetValue(DestIndex, false);
+					break;
+				case TableColumnType_Int:
+					Row.SetValue(DestIndex, 0);
+					break;
+				case TableColumnType_Float:
+					Row.SetValue(DestIndex, 0.0f);
+					break;
+				case TableColumnType_Double:
+					Row.SetValue(DestIndex, 0.0);
+					break;
+				case TableColumnType_CString:
+					Row.SetValue(DestIndex, TEXT("<>"));
+					break;
+				case TableColumnType_Invalid: break;
+				default: break;
+				}
+				continue;
+			}
+
+			switch (ColumnType)
+			{
+			case TableColumnType_Bool:
+				Row.SetValue(DestIndex, SourceTableReader.GetValueBool(SourceIndex));
+				break;
+			case TableColumnType_Int:
+				Row.SetValue(DestIndex, SourceTableReader.GetValueInt(SourceIndex) * Multiplier);
+				break;
+			case TableColumnType_Float:
+				Row.SetValue(DestIndex, SourceTableReader.GetValueFloat(SourceIndex) * Multiplier);
+				break;
+			case TableColumnType_Double:
+				Row.SetValue(DestIndex, SourceTableReader.GetValueDouble(SourceIndex) * Multiplier);
+				break;
+			case TableColumnType_CString:
+				Row.SetValue(DestIndex, TableC->GetStringStore().Store(SourceTableReader.GetValueCString(SourceIndex)));
+				break;
+			case TableColumnType_Invalid: break;
+			default: break;
+			}
+		}
+	}
+	SourceTableReader.SetRowIndex(0);
+}
+} // namespace TraceServices
diff --git a/Engine/Source/Developer/TraceServices/Public/TraceServices/Model/TableMerge.h b/Engine/Source/Developer/TraceServices/Public/TraceServices/Model/TableMerge.h
new file mode 100644
--- /dev/null
+++ b/Engine/Source/Developer/TraceServices/Public/TraceServices/Model/TableMerge.h
@@ -0,0 +1,32 @@
+// Copyright Epic Games, Inc. All Rights Reserved.
+
+#pragma once
+
+#include "CoreMinimal.h"
+#include "TraceServices/Containers/Tables.h"
+
+class FTokenizedMessage;
+
+namespace TraceServices
+{
+enum class ETableDiffResult : uint32
+{
+	ESuccess = 0,
+	EFail = 1,
+};
+
+struct FTableDiffCallbackParams
+{
+	ETableDiffResult Result;
+	TSharedPtr<IUntypedTable> Table;
+	TArray<TSharedRef<FTokenizedMessage>> Messages;
+};
+
+class FTableMergeService
+{
+public:
+	typedef TFunction<void(TSharedPtr<FTableDiffCallbackParams>)> TableDiffCallback;
+
+	TRACESERVICES_API static void MergeTables(const TSharedPtr<IUntypedTable>& TableA, const TSharedPtr<IUntypedTable>& TableB, TableDiffCallback InCallback);
+};
+} // namespace TraceServices
