[Insights] Timing view: Reset timing event highlight by double clicking on highlighted event. Previously it was possible to reset the highlight only by double clicking on an empty space.

diff --git a/Engine/Source/Developer/TraceInsights/Private/Insights/Widgets/STimingView.cpp b/Engine/Source/Developer/TraceInsights/Private/Insights/Widgets/STimingView.cpp
--- a/Engine/Source/Developer/TraceInsights/Private/Insights/Widgets/STimingView.cpp
+++ b/Engine/Source/Developer/TraceInsights/Private/Insights/Widgets/STimingView.cpp
@@ -79,7 +79,7 @@
 
 #define LOCTEXT_NAMESPACE "STimingView"
 
-#define ACTIVATE_BENCHMARK 0
+#define INSIGHTS_ACTIVATE_BENCHMARK 0
 
 // start auto generated ids from a big number (MSB set to 1) to avoid collisions with ids for gpu/cpu tracks based on 32bit timeline index
 uint64 FBaseTimingTrack::IdGenerator = (1ULL << 63);
@@ -2285,7 +2285,15 @@ FReply STimingView::OnMouseButtonDoubleClick(const FGeometry& MyGeometry, const
 			}
 			else
 			{
-				SetEventFilter(HoveredEvent->GetTrack()->GetFilterByEvent(HoveredEvent));
+				if (HoveredEvent->Is<FTimingEvent>() &&
+					IsFilterByEventType(HoveredEvent->As<FTimingEvent>().GetType()))
+				{
+					SetEventFilter(nullptr); // reset filter
+				}
+				else
+				{
+					SetEventFilter(HoveredEvent->GetTrack()->GetFilterByEvent(HoveredEvent));
+				}
 			}
 		}
 		else
@@ -2808,7 +2816,7 @@ FReply STimingView::OnKeyDown(const FGeometry& MyGeometry, const FKeyEvent& InKe
 		ChooseNextEventDepthLimit();
 		return FReply::Handled();
 	}
-#if ACTIVATE_BENCHMARK
+#if INSIGHTS_ACTIVATE_BENCHMARK
 	else if (InKeyEvent.GetKey() == EKeys::Z)
 	{
 		FTimingProfilerTests::FCheckValues CheckValues;
@@ -4910,7 +4918,11 @@ void STimingView::ClearRelations()
 
 void STimingView::FindFirstEvent()
 {
-	SelectedEvent.Reset();
+	if (SelectedEvent.IsValid())
+	{
+		SelectedEvent.Reset();
+		OnSelectedEventChangedDelegate.Broadcast(SelectedEvent);
+	}
 	FindNextEvent();
 }
 
@@ -5024,7 +5036,11 @@ void STimingView::FindNextEvent()
 
 void STimingView::FindLastEvent()
 {
-	SelectedEvent.Reset();
+	if (SelectedEvent.IsValid())
+	{
+		SelectedEvent.Reset();
+		OnSelectedEventChangedDelegate.Broadcast(SelectedEvent);
+	}
 	FindPrevEvent();
 }
 
@@ -5118,5 +5134,5 @@ void STimingView::EnumerateFilteredTracks(TSharedPtr<Insights::FFilterConfigurat
 
 ////////////////////////////////////////////////////////////////////////////////////////////////////
 
-#undef ACTIVATE_BENCHMARK
+#undef INSIGHTS_ACTIVATE_BENCHMARK
 #undef LOCTEXT_NAMESPACE
