[Insights] Session Browser: Improved note about "initial channels" in the "Connection" tab.

diff --git a/Engine/Source/Developer/TraceInsights/Private/Insights/Widgets/SStartPageWindow.cpp b/Engine/Source/Developer/TraceInsights/Private/Insights/Widgets/SStartPageWindow.cpp
--- a/Engine/Source/Developer/TraceInsights/Private/Insights/Widgets/SStartPageWindow.cpp
+++ b/Engine/Source/Developer/TraceInsights/Private/Insights/Widgets/SStartPageWindow.cpp
@@ -3120,7 +3120,36 @@ TSharedRef<SWidget> SConnectionWindow::ConstructConnectPanel()
 		.Padding(198.0f, 4.0f, 12.0f, 0.0f)
 		[
 			SNew(STextBlock)
-			.Text(LOCTEXT("InitialChannelsNoteText", "Comma-separated list of channel names (or \"default\"=cpu,gpu,frame,log,bookmark) to enable when connected."))
+			.Text(LOCTEXT("InitialChannelsNoteText", "Comma-separated list of channels/presets to enable when connected."))
+		]
+		+ SVerticalBox::Slot()
+		.AutoHeight()
+		.HAlign(HAlign_Fill)
+		.Padding(198.0f, 2.0f, 12.0f, 0.0f)
+		[
+			SNew(SHorizontalBox)
+			+ SHorizontalBox::Slot()
+			.AutoWidth()
+			.VAlign(VAlign_Top)
+			[
+				SNew(STextBlock)
+				.Text(LOCTEXT("InitialChannelsExamplesTitle", "Examples"))
+			]
+			+ SHorizontalBox::Slot()
+			.Padding(6.0f, 0.0f, 0.0f, 0.0f)
+			[
+				SNew(SEditableTextBox)
+				.IsReadOnly(true)
+				.Text(FText::FromString(TEXT("default,counter,stats,file,loadtime,assetloadtime,task\ndefault=cpu,gpu,frame,log,bookmark,screenshot")))
+			]
+		]
+		+ SVerticalBox::Slot()
+		.AutoHeight()
+		.HAlign(HAlign_Fill)
+		.Padding(198.0f, 2.0f, 12.0f, 0.0f)
+		[
+			SNew(STextBlock)
+			.Text(LOCTEXT("InitialChannelsNote2Text", "Same channels/presets (like \"memory\") cannot be enabled on late connections."))
 		]
 
 		+ SVerticalBox::Slot()
@@ -3134,7 +3163,7 @@ TSharedRef<SWidget> SConnectionWindow::ConstructConnectPanel()
 				SNew(SButton)
 				.ButtonStyle(&FAppStyle::Get().GetWidgetStyle<FButtonStyle>("PrimaryButton"))
 				.Text(LOCTEXT("Connect", "Connect"))
-				.ToolTipText(LOCTEXT("ConnectToolTip", "Connect the running instance at specified address with the trace recorder."))
+				.ToolTipText(LOCTEXT("ConnectToolTip", "Late connect the running instance at specified address with the trace recorder."))
 				.OnClicked(this, &SConnectionWindow::Connect_OnClicked)
 				.IsEnabled_Lambda([this]() { return !bIsConnecting; })
 			]
