[Insights] Fix crash when closing the Insights Session Browser after importing a table.

diff --git a/Engine/Source/Developer/TraceInsights/Private/Insights/TraceInsightsModule.cpp b/Engine/Source/Developer/TraceInsights/Private/Insights/TraceInsightsModule.cpp
--- a/Engine/Source/Developer/TraceInsights/Private/Insights/TraceInsightsModule.cpp
+++ b/Engine/Source/Developer/TraceInsights/Private/Insights/TraceInsightsModule.cpp
@@ -302,6 +302,7 @@ void FTraceInsightsModule::CreateSessionBrowser(const FCreateSessionBrowserParam
 	const EOutputCanBeNullptr OutputCanBeNullptr = EOutputCanBeNullptr::Never;
 	TSharedPtr<SWidget> Content = FGlobalTabmanager::Get()->RestoreFrom(PersistentLayout.ToSharedRef(), RootWindow, bEmbedTitleAreaContent, OutputCanBeNullptr);
 	RootWindow->SetContent(Content.ToSharedRef());
+	RootWindow->GetOnWindowClosedEvent().AddRaw(this, &FTraceInsightsModule::OnWindowClosedEvent);
 
 	//////////////////////////////////////////////////
 	// Show the window.
