[Insights] Memory Insights: Fixed typo "snaphot" --> "snapshot".

diff --git a/Engine/Source/Developer/TraceInsights/Private/Insights/MemoryProfiler/Widgets/SMemAllocTableTreeView.cpp b/Engine/Source/Developer/TraceInsights/Private/Insights/MemoryProfiler/Widgets/SMemAllocTableTreeView.cpp
--- a/Engine/Source/Developer/TraceInsights/Private/Insights/MemoryProfiler/Widgets/SMemAllocTableTreeView.cpp
+++ b/Engine/Source/Developer/TraceInsights/Private/Insights/MemoryProfiler/Widgets/SMemAllocTableTreeView.cpp
@@ -1477,13 +1477,13 @@ void SMemAllocTableTreeView::ExtendMenu(FMenuBuilder& MenuBuilder)
 	}
 
 	{
-		const FText ItemLabel = LOCTEXT("ContextMenu_Export_SubMenu", "Export snaphot");
-		const FText ItemToolTip = LOCTEXT("ContextMenu_Export_Desc_SubMenu", "Export memory snaphot to construct diff later");
+		const FText ItemLabel = LOCTEXT("ContextMenu_Export_SubMenu", "Export Snapshot");
+		const FText ItemToolTip = LOCTEXT("ContextMenu_Export_Desc_SubMenu", "Export memory snapshot to construct diff later.");
 
 		FUIAction Action_ExportSnapshot
 		(
 			FExecuteAction::CreateSP(this, &SMemAllocTableTreeView::ExportMemorySnapshot),
-			FCanExecuteAction::CreateSP(this, &SMemAllocTableTreeView::IsExportMemorySnaphotAvailable)
+			FCanExecuteAction::CreateSP(this, &SMemAllocTableTreeView::IsExportMemorySnapshotAvailable)
 		);
 		MenuBuilder.AddMenuEntry(
 			ItemLabel,
@@ -1865,7 +1865,7 @@ void SMemAllocTableTreeView::ExportMemorySnapshot() const
 
 ////////////////////////////////////////////////////////////////////////////////////////////////////
 
-bool SMemAllocTableTreeView::IsExportMemorySnaphotAvailable() const
+bool SMemAllocTableTreeView::IsExportMemorySnapshotAvailable() const
 {
 	return !TableTreeNodes.IsEmpty();
 }
diff --git a/Engine/Source/Developer/TraceInsights/Private/Insights/MemoryProfiler/Widgets/SMemAllocTableTreeView.h b/Engine/Source/Developer/TraceInsights/Private/Insights/MemoryProfiler/Widgets/SMemAllocTableTreeView.h
--- a/Engine/Source/Developer/TraceInsights/Private/Insights/MemoryProfiler/Widgets/SMemAllocTableTreeView.h
+++ b/Engine/Source/Developer/TraceInsights/Private/Insights/MemoryProfiler/Widgets/SMemAllocTableTreeView.h
@@ -96,7 +96,7 @@ protected:
 	FText GetSelectedCallstackFrameFileName() const;
 	void OpenSourceFileInIDE(const TCHAR* File, uint32 Line) const;
 	void ExportMemorySnapshot() const;
-	bool IsExportMemorySnaphotAvailable() const;
+	bool IsExportMemorySnapshotAvailable() const;
 
 private:
 	void OnQueryInvalidated();
