Add trace store connection status UX and reconnection logic

* Adds user interface for displaying connection to trace store
* Adds logic that checks the connection in store browser and periodically attempts to reconnect if a severed connection is detected.

diff --git a/Engine/Source/Developer/TraceAnalysis/Private/Store/StoreClient.cpp b/Engine/Source/Developer/TraceAnalysis/Private/Store/StoreClient.cpp
--- a/Engine/Source/Developer/TraceAnalysis/Private/Store/StoreClient.cpp
+++ b/Engine/Source/Developer/TraceAnalysis/Private/Store/StoreClient.cpp
@@ -478,6 +478,13 @@ FStoreClient* FStoreClient::Connect(const TCHAR* Host, uint32 Port)
 	return (FStoreClient*)Impl;
 }
 
+////////////////////////////////////////////////////////////////////////////////
+bool FStoreClient::Reconnect(const TCHAR* Host, uint32 Port)
+{
+	auto* Self = (FStoreCborClient*)this;
+	return Self->Connect(Host, Port);
+}
+
 ////////////////////////////////////////////////////////////////////////////////
 void FStoreClient::operator delete (void* Addr)
 {
diff --git a/Engine/Source/Developer/TraceAnalysis/Public/Trace/StoreClient.h b/Engine/Source/Developer/TraceAnalysis/Public/Trace/StoreClient.h
--- a/Engine/Source/Developer/TraceAnalysis/Public/Trace/StoreClient.h
+++ b/Engine/Source/Developer/TraceAnalysis/Public/Trace/StoreClient.h
@@ -42,6 +42,7 @@ public:
 	
 						~FStoreClient() = default;
 	static FStoreClient*Connect(const TCHAR* Host, uint32 Port=0);
+	bool				Reconnect(const TCHAR* Host, uint32 Port);
 	void				operator delete (void* Addr);
 	bool				IsValid() const;
 	uint32				GetStoreAddress() const;
diff --git a/Engine/Source/Developer/TraceInsights/Private/Insights/InsightsManager.cpp b/Engine/Source/Developer/TraceInsights/Private/Insights/InsightsManager.cpp
--- a/Engine/Source/Developer/TraceInsights/Private/Insights/InsightsManager.cpp
+++ b/Engine/Source/Developer/TraceInsights/Private/Insights/InsightsManager.cpp
@@ -448,6 +448,9 @@ bool FInsightsManager::ConnectToStore(const TCHAR* Host, uint32 Port)
 	{
 		return false;
 	}
+	
+	LastStoreHost = Host;
+	LastStorePort = Port;
 
 	const FStoreClient::FStatus* Status = StoreClient->GetStatus();
 	FString RemoteStoreDir(Status->GetStoreDir());
@@ -461,6 +464,13 @@ bool FInsightsManager::ConnectToStore(const TCHAR* Host, uint32 Port)
 
 ////////////////////////////////////////////////////////////////////////////////////////////////////
 
+bool FInsightsManager::ReconnectToStore() const
+{
+	return StoreClient->Reconnect(*LastStoreHost, LastStorePort);
+}
+
+////////////////////////////////////////////////////////////////////////////////////////////////////
+
 TSharedPtr<const TraceServices::IAnalysisSession> FInsightsManager::GetSession() const
 {
 	return Session;
diff --git a/Engine/Source/Developer/TraceInsights/Private/Insights/InsightsManager.h b/Engine/Source/Developer/TraceInsights/Private/Insights/InsightsManager.h
--- a/Engine/Source/Developer/TraceInsights/Private/Insights/InsightsManager.h
+++ b/Engine/Source/Developer/TraceInsights/Private/Insights/InsightsManager.h
@@ -118,6 +118,13 @@ public:
 	const FString& GetStoreDir() const { return StoreDir; }
 
 	bool ConnectToStore(const TCHAR* Host, uint32 Port=0);
+
+	/**
+	 * Attempt to reconnect to the store if the connection was severed, without recreating the store client.
+	 * @return True on success, false on failure
+	 */
+	bool ReconnectToStore() const;
+	
 	UE::Trace::FStoreClient* GetStoreClient() const { return StoreClient.Get(); }
 	FCriticalSection& GetStoreClientCriticalSection() const { return StoreClientCriticalSection; }
 
@@ -479,4 +486,7 @@ private:
 
 	/** A shared pointer to the global instance of the main manager. */
 	static TSharedPtr<FInsightsManager> Instance;
+
+	FString LastStoreHost;
+	uint32 LastStorePort = 0;
 };
diff --git a/Engine/Source/Developer/TraceInsights/Private/Insights/InsightsStyle.cpp b/Engine/Source/Developer/TraceInsights/Private/Insights/InsightsStyle.cpp
--- a/Engine/Source/Developer/TraceInsights/Private/Insights/InsightsStyle.cpp
+++ b/Engine/Source/Developer/TraceInsights/Private/Insights/InsightsStyle.cpp
@@ -151,6 +151,8 @@ void FInsightsStyle::FStyle::Initialize()
 
 	Set("Icons.TraceStore", new IMAGE_BRUSH_SVG("TraceStore", Icon16x16));
 	Set("Icons.Connection", new IMAGE_BRUSH_SVG("Connection", Icon16x16));
+	Set("Icons.Online", new CORE_IMAGE_BRUSH_SVG("Starship/Common/check-circle", Icon16x16, FStyleColors::AccentGreen));
+	Set("Icons.Offline", new CORE_IMAGE_BRUSH_SVG("Starship/Common/alert-triangle", Icon16x16, FStyleColors::Warning));
 	Set("Icons.Launcher", new TODO_IMAGE_BRUSH(Icon16x16));
 
 	//////////////////////////////////////////////////
diff --git a/Engine/Source/Developer/TraceInsights/Private/Insights/StoreService/StoreBrowser.cpp b/Engine/Source/Developer/TraceInsights/Private/Insights/StoreService/StoreBrowser.cpp
--- a/Engine/Source/Developer/TraceInsights/Private/Insights/StoreService/StoreBrowser.cpp
+++ b/Engine/Source/Developer/TraceInsights/Private/Insights/StoreService/StoreBrowser.cpp
@@ -120,7 +120,39 @@ void FStoreBrowser::UpdateTraces()
 
 	FStopwatch StopwatchTotal;
 	StopwatchTotal.Start();
+	
+	// Check if connection to store is still active We want to do this without locking the critical
+	// section, in case the UI needs to read values. The output is an atomic anyway.
+	{
+		bool bIsValid = StoreClient->IsValid();
+		while (!bIsValid && bRunning) // TODO: Perhaps a max reconnection attempts is needed?
+		{
+			constexpr EConnectionStatus ReconnectionFrequency = static_cast<EConnectionStatus>(5);
+			uint32 SecondsToSleep = static_cast<uint32>(ReconnectionFrequency);
+			do
+			{
+				FPlatformProcess::Sleep(1.0);
+				ConnectionStatus.store(static_cast<EConnectionStatus>(--SecondsToSleep));			
+			} while (SecondsToSleep);
+
+			// At this point ConnectionStatus is zero (EConnectionStatus::Connecting)
+			bIsValid = FInsightsManager::Get()->ReconnectToStore();
+			
+			if (bIsValid)
+			{
+				// If we gave just reconnected we need to reset known traces and change serials
+				Refresh();
+			}
+		}
+		ConnectionStatus.store(bIsValid ? EConnectionStatus::Connected : EConnectionStatus::Connecting);
 
+		if (!bIsValid)
+		{
+			// If there is no connection, no point in checking for sessions
+			return;
+		}
+	}
+	
 	FScopeLock StoreClientLock(&GetStoreClientCriticalSection());
 
 	// Check if the list of trace sessions has changed.
diff --git a/Engine/Source/Developer/TraceInsights/Private/Insights/StoreService/StoreBrowser.h b/Engine/Source/Developer/TraceInsights/Private/Insights/StoreService/StoreBrowser.h
--- a/Engine/Source/Developer/TraceInsights/Private/Insights/StoreService/StoreBrowser.h
+++ b/Engine/Source/Developer/TraceInsights/Private/Insights/StoreService/StoreBrowser.h
@@ -78,6 +78,22 @@ public:
 
 	//////////////////////////////////////////////////
 
+	enum class EConnectionStatus : uint8
+	{
+		// Attempting connection
+		Connecting = 0,
+		// Values between (start,end) is interpreted as
+		// number of seconds until next reconnection attemp
+		SecondsToReconnectStart,
+		SecondsToReconnectEnd = 0xfd,
+		// No connection could be made, no more reconnection
+		// attempts are made.
+		NoConnection = 0xfe,
+		// Connection is active
+		Connected = 0xff
+	};
+
+	EConnectionStatus GetConnectionStatus() const { return ConnectionStatus; }
 	bool IsRunning() const { return bRunning; }
 
 	bool IsLocked() const { return bTracesLocked; }
@@ -107,6 +123,7 @@ private:
 
 	mutable FCriticalSection TracesCriticalSection;
 	FThreadSafeBool bTracesLocked; // for debugging, to ensure GetLocked*() methods are only called between Lock() - Unlock() calls.
+	std::atomic<EConnectionStatus> ConnectionStatus;
 	uint32 StoreChangeSerial;
 	uint64 TracesChangeSerial;
 	TArray<TSharedPtr<FStoreBrowserTraceInfo>> Traces;
diff --git a/Engine/Source/Developer/TraceInsights/Private/Insights/Widgets/SStartPageWindow.cpp b/Engine/Source/Developer/TraceInsights/Private/Insights/Widgets/SStartPageWindow.cpp
--- a/Engine/Source/Developer/TraceInsights/Private/Insights/Widgets/SStartPageWindow.cpp
+++ b/Engine/Source/Developer/TraceInsights/Private/Insights/Widgets/SStartPageWindow.cpp
@@ -1317,6 +1317,16 @@ TSharedRef<SWidget> STraceStoreWindow::ConstructTraceStoreDirectoryPanel()
 
 		SNew(SHorizontalBox)
 
+		+ SHorizontalBox::Slot()
+		.AutoWidth()
+		.Padding(0.0f, 0.0f, 4.0f, 0.0f)
+		.VAlign(VAlign_Center)
+		[
+			 SNew(SImage)
+			 .Image_Raw(this, &STraceStoreWindow::GetConnectionStatusIcon)
+			 .ToolTip(SNew(SToolTip).Text_Raw(this, &STraceStoreWindow::GetConnectionStatusTooltip))
+		]
+
 		+ SHorizontalBox::Slot()
 		.AutoWidth()
 		.Padding(0.0f, 0.0f)
@@ -1422,6 +1432,46 @@ END_SLATE_FUNCTION_BUILD_OPTIMIZATION
 
 ////////////////////////////////////////////////////////////////////////////////////////////////////
 
+FText STraceStoreWindow::GetConnectionStatusTooltip() const
+{
+	using Insights::FStoreBrowser;
+	
+	static FText Connected = LOCTEXT("Connected", "Connected to the trace server");
+	static FText NotConnected = LOCTEXT("NoConnection", "Unable to connect trace server."); 
+	static FText Connecting = LOCTEXT("Connecting", "Trying to connect to trace server.");
+	static FText Disconnected = LOCTEXT("Disconnected", "Connection to trace server has been lost. Attempting to reconnect in {0} seconds.");
+	
+	const FStoreBrowser::EConnectionStatus Status = StoreBrowser->GetConnectionStatus();
+	switch(Status)
+	{
+	case FStoreBrowser::EConnectionStatus::Connected:
+		return Connected;
+	case FStoreBrowser::EConnectionStatus::NoConnection:
+		return NotConnected;
+	case FStoreBrowser::EConnectionStatus::Connecting:
+		return Connecting;
+	default:
+		{
+			return FText::Format(Disconnected, FText::AsNumber(static_cast<uint32>(Status)));
+		}
+	}
+}
+
+////////////////////////////////////////////////////////////////////////////////////////////////////
+
+const FSlateBrush* STraceStoreWindow::GetConnectionStatusIcon() const
+{
+	using Insights::FStoreBrowser;
+	const FStoreBrowser::EConnectionStatus Status = StoreBrowser->GetConnectionStatus();
+	if (Status == FStoreBrowser::EConnectionStatus::Connected)
+	{
+		return FInsightsStyle::Get().GetBrush("Icons.Online");
+	}
+	return FInsightsStyle::Get().GetBrush("Icons.Offline");
+}
+
+////////////////////////////////////////////////////////////////////////////////////////////////////
+
 TSharedPtr<SWidget> STraceStoreWindow::TraceList_GetMenuContent()
 {
 	const bool bShouldCloseWindowAfterMenuSelection = true;
diff --git a/Engine/Source/Developer/TraceInsights/Private/Insights/Widgets/SStartPageWindow.h b/Engine/Source/Developer/TraceInsights/Private/Insights/Widgets/SStartPageWindow.h
--- a/Engine/Source/Developer/TraceInsights/Private/Insights/Widgets/SStartPageWindow.h
+++ b/Engine/Source/Developer/TraceInsights/Private/Insights/Widgets/SStartPageWindow.h
@@ -268,6 +268,8 @@ private:
 	TSharedRef<SWidget> ConstructFiltersToolbar();
 	TSharedRef<SWidget> ConstructSessionsPanel();
 	TSharedRef<SWidget> ConstructLoadPanel();
+	FText GetConnectionStatusTooltip() const;
+	const FSlateBrush* GetConnectionStatusIcon() const;
 	TSharedRef<SWidget> ConstructTraceStoreDirectoryPanel();
 	TSharedRef<SWidget> ConstructAutoStartPanel();
 
