[Insights] Implement Frame stats timer graph series

diff --git a/Engine/Source/Developer/TraceInsights/Private/Insights/ViewModels/TimerNode.cpp b/Engine/Source/Developer/TraceInsights/Private/Insights/ViewModels/TimerNode.cpp
--- a/Engine/Source/Developer/TraceInsights/Private/Insights/ViewModels/TimerNode.cpp
+++ b/Engine/Source/Developer/TraceInsights/Private/Insights/ViewModels/TimerNode.cpp
@@ -22,7 +22,7 @@ FTimerNode::FTimerNode(uint32 InTimerId, const TCHAR* InName, ETimerNodeType InT
 	, TimerId(InTimerId)
 	, MetaGroupName(InType == ETimerNodeType::CpuScope ? CpuGroup : InType == ETimerNodeType::GpuScope ? GpuGroup : NAME_None)
 	, Type(InType)
-	, bIsAddedToGraph(false)
+	, NumGraphs(0)
 	, bIsHotPath(false)
 {
 	uint32 Color32 = FTimingEvent::ComputeEventColor(InName);
@@ -42,7 +42,7 @@ FTimerNode::FTimerNode(const FName InGroupName)
 	, TimerId(InvalidTimerId)
 	, Type(ETimerNodeType::Group)
 	, Color(0.0, 0.0, 0.0, 1.0)
-	, bIsAddedToGraph(false)
+	, NumGraphs(0)
 	, bIsHotPath(false)
 {
 	ResetAggregatedStats();
diff --git a/Engine/Source/Developer/TraceInsights/Private/Insights/ViewModels/TimerNode.h b/Engine/Source/Developer/TraceInsights/Private/Insights/ViewModels/TimerNode.h
--- a/Engine/Source/Developer/TraceInsights/Private/Insights/ViewModels/TimerNode.h
+++ b/Engine/Source/Developer/TraceInsights/Private/Insights/ViewModels/TimerNode.h
@@ -93,12 +93,18 @@ public:
 
 	bool IsAddedToGraph() const
 	{
-		return bIsAddedToGraph;
+		return NumGraphs > 0;
 	}
 
-	void SetAddedToGraphFlag(bool bOnOff)
+	void OnAddedToGraph()
 	{
-		bIsAddedToGraph = bOnOff;
+		++NumGraphs;
+	}
+
+	void OnRemovedFromGraph()
+	{
+		ensure(NumGraphs > 0);
+		--NumGraphs;
 	}
 
 	/**
@@ -127,8 +133,9 @@ private:
 
 	/** Color of the node. */
 	FLinearColor Color;
-
-	bool bIsAddedToGraph;
+	
+	/** The number of graphs created for this timer. */
+	int32 NumGraphs;
 
 	/** True if this tree node is on the hot path. */
 	bool bIsHotPath;
diff --git a/Engine/Source/Developer/TraceInsights/Private/Insights/ViewModels/TimingGraphTrack.cpp b/Engine/Source/Developer/TraceInsights/Private/Insights/ViewModels/TimingGraphTrack.cpp
--- a/Engine/Source/Developer/TraceInsights/Private/Insights/ViewModels/TimingGraphTrack.cpp
+++ b/Engine/Source/Developer/TraceInsights/Private/Insights/ViewModels/TimingGraphTrack.cpp
@@ -51,6 +51,7 @@ FString FTimingGraphSeries::FormatValue(double Value) const
 		return FString::Printf(TEXT("%s (%g fps)"), *TimeUtils::FormatTimeAuto(Value), 1.0 / Value);
 
 	case FTimingGraphSeries::ESeriesType::Timer:
+	case FTimingGraphSeries::ESeriesType::FrameStatsTimer:
 		return TimeUtils::FormatTimeAuto(Value);
 
 	case FTimingGraphSeries::ESeriesType::StatsCounter:
@@ -185,6 +186,10 @@ void FTimingGraphTrack::Update(const ITimingTrackUpdateContext& Context)
 					UpdateTimerSeries(*TimingSeries, Viewport);
 					break;
 
+				case FTimingGraphSeries::ESeriesType::FrameStatsTimer:
+					UpdateFrameStatsTimerSeries(*TimingSeries, Viewport);
+					break;
+
 				case FTimingGraphSeries::ESeriesType::StatsCounter:
 					UpdateStatsCounterSeries(*TimingSeries, Viewport);
 					break;
@@ -399,6 +404,179 @@ void FTimingGraphTrack::UpdateTimerSeries(FTimingGraphSeries& Series, const FTim
 	}
 }
 
+////////////////////////////////////////////////////////////////////////////////////////////////////
+// Frams Stats Timer Series
+////////////////////////////////////////////////////////////////////////////////////////////////////
+
+TSharedPtr<FTimingGraphSeries> FTimingGraphTrack::GetFrameStatsTimerSeries(uint32 TimerId)
+{
+	TSharedPtr<FGraphSeries>* Ptr = AllSeries.FindByPredicate([TimerId](const TSharedPtr<FGraphSeries>& Series)
+		{
+			const TSharedPtr<FTimingGraphSeries> TimingSeries = StaticCastSharedPtr<FTimingGraphSeries>(Series);
+			return TimingSeries->Type == FTimingGraphSeries::ESeriesType::FrameStatsTimer && TimingSeries->TimerId == TimerId;
+		});
+	return (Ptr != nullptr) ? StaticCastSharedPtr<FTimingGraphSeries>(*Ptr) : nullptr;
+}
+
+////////////////////////////////////////////////////////////////////////////////////////////////////
+
+TSharedPtr<FTimingGraphSeries> FTimingGraphTrack::AddFrameStatsTimerSeries(uint32 TimerId, FLinearColor Color)
+{
+	TSharedRef<FTimingGraphSeries> Series = MakeShared<FTimingGraphSeries>(FTimingGraphSeries::ESeriesType::FrameStatsTimer);
+
+	Series->SetName(TEXT("<Frame Stats Timer>"));
+	Series->SetDescription(TEXT("Frame Stats Timer series"));
+
+	const FLinearColor BorderColor(Color.R + 0.4f, Color.G + 0.4f, Color.B + 0.4f, 1.0f);
+	Series->SetColor(Color, BorderColor);
+
+	Series->TimerId = TimerId;
+	//Series->CpuOrGpu = ;
+	//Series->TimelineIndex = ;
+
+	// Use shared viewport.
+	Series->SetBaselineY(SharedValueViewport.GetBaselineY());
+	Series->SetScaleY(SharedValueViewport.GetScaleY());
+	Series->EnableSharedViewport();
+
+	Series->CachedSessionDuration = 0.0;
+
+	AllSeries.Add(Series);
+	return Series;
+}
+
+////////////////////////////////////////////////////////////////////////////////////////////////////
+
+void FTimingGraphTrack::RemoveFrameStatsTimerSeries(uint32 TimerId)
+{
+	AllSeries.RemoveAll([TimerId](const TSharedPtr<FGraphSeries>& Series)
+		{
+			const TSharedPtr<FTimingGraphSeries> TimingSeries = StaticCastSharedPtr<FTimingGraphSeries>(Series);
+			return TimingSeries->Type == FTimingGraphSeries::ESeriesType::FrameStatsTimer && TimingSeries->TimerId == TimerId;
+		});
+}
+
+////////////////////////////////////////////////////////////////////////////////////////////////////
+
+void FTimingGraphTrack::UpdateFrameStatsTimerSeries(FTimingGraphSeries& Series, const FTimingTrackViewport& Viewport)
+{
+	FGraphTrackBuilder Builder(*this, Series, Viewport);
+	TSharedPtr<const TraceServices::IAnalysisSession> Session = FInsightsManager::Get()->GetSession();
+	if (Session.IsValid())
+	{
+		TraceServices::FAnalysisSessionReadScope SessionReadScope(*Session.Get());
+
+		const double SessionDuration = Session->GetDurationSeconds();
+		if (Series.CachedSessionDuration != SessionDuration)
+		{
+			Series.CachedSessionDuration = SessionDuration;
+
+			const TraceServices::IFrameProvider& FramesProvider = ReadFrameProvider(*Session.Get());
+
+			Series.FrameStatsCachedEvents.Empty();
+			uint32 FrameCount = FramesProvider.GetFrameCount(ETraceFrameType::TraceFrameType_Game);
+			if (FrameCount == 0)
+			{
+				return;
+			}
+
+			FramesProvider.EnumerateFrames(ETraceFrameType::TraceFrameType_Game, 0, FrameCount, [&Series](const TraceServices::FFrame& Frame)
+				{
+					FTimingGraphSeries::FAtomicTimingEvent Event;
+					Event.FrameStartTime = Frame.StartTime;
+					Event.FrameEndTime = Frame.EndTime;
+					Event.Duration.store(0.0f);
+					Series.FrameStatsCachedEvents.Add(Event);
+				});
+
+			const TraceServices::ITimingProfilerProvider& TimingProfilerProvider = *TraceServices::ReadTimingProfilerProvider(*Session.Get());
+
+			const TraceServices::ITimingProfilerTimerReader* TimerReader;
+			TimingProfilerProvider.ReadTimers([&TimerReader](const TraceServices::ITimingProfilerTimerReader& Out) { TimerReader = &Out; });
+
+			const uint32 TimelineCount = TimingProfilerProvider.GetTimelineCount();
+			for (uint32 TimelineIndex = 0; TimelineIndex < TimelineCount; ++TimelineIndex)
+			{
+				TimingProfilerProvider.ReadTimeline(TimelineIndex,
+					[SessionDuration, &Series, TimerReader](const TraceServices::ITimingProfilerProvider::Timeline& Timeline)
+					{
+						TArray<TArray<FTimingGraphSeries::FSimpleTimingEvent>> Events;
+						TraceServices::ITimeline<TraceServices::FTimingProfilerEvent>::EnumerateAsyncParams Params;
+						Params.IntervalStart = 0;
+						Params.IntervalEnd = SessionDuration;
+						Params.Resolution = 0.0;
+						Params.SetupCallback = [&Events](uint32 NumTasks) {};
+						Params.Callback = [&Events, TimerReader, &Series, SessionDuration](double StartTime, double EndTime, uint32 Depth, const TraceServices::FTimingProfilerEvent& Event, uint32 TaskIndex)
+						{
+							const TraceServices::FTimingProfilerTimer* Timer = TimerReader->GetTimer(Event.TimerIndex);
+							if (ensure(Timer != nullptr))
+							{
+								if (Timer->Id == Series.TimerId)
+								{
+									int32 Index = Algo::UpperBoundBy(Series.FrameStatsCachedEvents, StartTime, &FTimingGraphSeries::FAtomicTimingEvent::FrameStartTime);
+									if (Index > 0)
+									{
+										--Index;
+									}
+
+									// This can can happen when the event is between frames.
+									if (StartTime > Series.FrameStatsCachedEvents[Index].FrameEndTime)
+									{
+										Index++;
+										if (Index >= Series.FrameStatsCachedEvents.Num())
+										{
+											return TraceServices::EEventEnumerate::Continue;
+										}
+									}
+
+									do
+									{
+										FTimingGraphSeries::FAtomicTimingEvent& Entry = Series.FrameStatsCachedEvents[Index];
+
+										if (EndTime < Entry.FrameStartTime)
+										{
+											return TraceServices::EEventEnumerate::Continue;
+										}
+
+										if (StartTime < Entry.FrameStartTime)
+										{
+											StartTime = Entry.FrameStartTime;
+										}
+
+										const double Duration = FMath::Min(EndTime, Entry.FrameEndTime) - StartTime;
+										ensure(Duration >= 0.0f);
+										for (double Value = Entry.Duration.load(); !Entry.Duration.compare_exchange_strong(Value, Value + Duration););
+
+										Index++;
+									} while (Index < Series.FrameStatsCachedEvents.Num());
+								}
+							}
+							return TraceServices::EEventEnumerate::Continue;
+						};
+
+						Timeline.EnumerateEventsDownSampledAsync(Params);
+					});
+			}
+		}
+
+		int32 StartIndex = Algo::UpperBoundBy(Series.FrameStatsCachedEvents, Viewport.GetStartTime(), &FTimingGraphSeries::FAtomicTimingEvent::FrameStartTime);
+		if (StartIndex > 0)
+		{
+			StartIndex--;
+		}
+		int32 EndIndex = Algo::UpperBoundBy(Series.FrameStatsCachedEvents, Viewport.GetEndTime(), &FTimingGraphSeries::FAtomicTimingEvent::FrameStartTime);
+		if (EndIndex < Series.FrameStatsCachedEvents.Num())
+		{
+			EndIndex++;
+		}
+		for (int32 Index = StartIndex; Index < EndIndex; ++Index)
+		{
+			const FTimingGraphSeries::FAtomicTimingEvent& Entry = Series.FrameStatsCachedEvents[Index];
+			Builder.AddEvent(Entry.FrameStartTime, Entry.Duration.load(), Entry.Duration.load());
+		}
+	}
+}
+
 ////////////////////////////////////////////////////////////////////////////////////////////////////
 // Stats Counter Series
 ////////////////////////////////////////////////////////////////////////////////////////////////////
diff --git a/Engine/Source/Developer/TraceInsights/Private/Insights/ViewModels/TimingGraphTrack.h b/Engine/Source/Developer/TraceInsights/Private/Insights/ViewModels/TimingGraphTrack.h
--- a/Engine/Source/Developer/TraceInsights/Private/Insights/ViewModels/TimingGraphTrack.h
+++ b/Engine/Source/Developer/TraceInsights/Private/Insights/ViewModels/TimingGraphTrack.h
@@ -17,7 +17,8 @@ public:
 	{
 		Frame,
 		Timer,
-		StatsCounter
+		StatsCounter,
+		FrameStatsTimer
 	};
 
 	struct FSimpleTimingEvent
@@ -26,6 +27,27 @@ public:
 		double Duration;
 	};
 
+	struct FAtomicTimingEvent
+	{
+		FAtomicTimingEvent()
+		{
+			FrameStartTime = 0.0f;
+			FrameEndTime = 0.0f;
+			Duration.store(0.0f);
+		}
+
+		FAtomicTimingEvent(const FAtomicTimingEvent& Other)
+		{
+			FrameStartTime = Other.FrameStartTime;
+			FrameEndTime = Other.FrameEndTime;
+			Duration.store(Other.Duration.load());
+		}
+
+		double FrameEndTime;
+		double FrameStartTime;
+		std::atomic<double> Duration;
+	};
+
 public:
 	explicit FTimingGraphSeries(FTimingGraphSeries::ESeriesType Type);
 	virtual ~FTimingGraphSeries();
@@ -48,6 +70,7 @@ public:
 
 	double CachedSessionDuration;
 	TArray<FSimpleTimingEvent> CachedEvents; // used by Timer series
+	TArray<FAtomicTimingEvent> FrameStatsCachedEvents; // used by Frame Stats Timer series
 
 	bool bIsTime; // the unit for values is [second]
 	bool bIsMemory; // the unit for value is [byte]
@@ -72,6 +95,10 @@ public:
 	TSharedPtr<FTimingGraphSeries> AddTimerSeries(uint32 TimerId, FLinearColor Color);
 	void RemoveTimerSeries(uint32 TimerId);
 
+	TSharedPtr<FTimingGraphSeries> GetFrameStatsTimerSeries(uint32 TimerId);
+	TSharedPtr<FTimingGraphSeries> AddFrameStatsTimerSeries(uint32 TimerId, FLinearColor Color);
+	void RemoveFrameStatsTimerSeries(uint32 TimerId);
+
 	TSharedPtr<FTimingGraphSeries> GetStatsCounterSeries(uint32 CounterId);
 	TSharedPtr<FTimingGraphSeries> AddStatsCounterSeries(uint32 CounterId, FLinearColor Color);
 	void RemoveStatsCounterSeries(uint32 CounterId);
@@ -79,6 +106,7 @@ public:
 protected:
 	void UpdateFrameSeries(FTimingGraphSeries& Series, const FTimingTrackViewport& Viewport);
 	void UpdateTimerSeries(FTimingGraphSeries& Series, const FTimingTrackViewport& Viewport);
+	void UpdateFrameStatsTimerSeries(FTimingGraphSeries& Series, const FTimingTrackViewport& Viewport);
 	void UpdateStatsCounterSeries(FTimingGraphSeries& Series, const FTimingTrackViewport& Viewport);
 
 	virtual void DrawVerticalAxisGrid(const ITimingTrackDrawContext& Context) const override;
diff --git a/Engine/Source/Developer/TraceInsights/Private/Insights/Widgets/STimersView.cpp b/Engine/Source/Developer/TraceInsights/Private/Insights/Widgets/STimersView.cpp
--- a/Engine/Source/Developer/TraceInsights/Private/Insights/Widgets/STimersView.cpp
+++ b/Engine/Source/Developer/TraceInsights/Private/Insights/Widgets/STimersView.cpp
@@ -534,6 +534,40 @@ TSharedPtr<SWidget> STimersView::TreeView_GetMenuContent()
 			}
 		}
 
+		// Add/remove frame stats series to/from graph track
+		{
+			FUIAction Action_ToggleFrameStatsTimerInGraphTrack;
+			Action_ToggleFrameStatsTimerInGraphTrack.CanExecuteAction = FCanExecuteAction::CreateLambda(CanExecute); 
+			Action_ToggleFrameStatsTimerInGraphTrack.ExecuteAction = FExecuteAction::CreateSP(this, &STimersView::ToggleTimingViewMainGraphEventFrameStatsSeries, SelectedNode);
+
+			if (SelectedNode.IsValid() &&
+				SelectedNode->GetType() != ETimerNodeType::Group &&
+				IsFrameStatsSeriesInTimingViewMainGraph(SelectedNode))
+			{
+				MenuBuilder.AddMenuEntry
+				(
+					LOCTEXT("ContextMenu_RemoveFrameStatsFromGraphTrack", "Remove frame stats series from graph track"),
+					LOCTEXT("ContextMenu_RemoveFrameStatsFromGraphTrack_Desc", "Removes the frame stats series containing event instances of the selected timer from the Main Graph track."),
+					FSlateIcon(FInsightsStyle::GetStyleSetName(), "Icons.RemoveGraphSeries"),
+					Action_ToggleFrameStatsTimerInGraphTrack,
+					NAME_None,
+					EUserInterfaceActionType::Button
+				);
+			}
+			else
+			{
+				MenuBuilder.AddMenuEntry
+				(
+					LOCTEXT("ContextMenu_AddFrameStatsSeriesToGraphTrack", "Add frame stats series to graph track"),
+					LOCTEXT("ContextMenu_AddFrameStatsSeriesToGraphTrack_Desc", "Adds a frame stats series containing event instances of the selected timer to the Main Graph track."),
+					FSlateIcon(FInsightsStyle::GetStyleSetName(), "Icons.AddGraphSeries"),
+					Action_ToggleFrameStatsTimerInGraphTrack,
+					NAME_None,
+					EUserInterfaceActionType::Button
+				);
+			}
+		}
+
 		// Open Source in IDE
 		{
 			ISourceCodeAccessModule& SourceCodeAccessModule = FModuleManager::LoadModuleChecked<ISourceCodeAccessModule>("SourceCodeAccess");
@@ -2181,7 +2215,7 @@ void STimersView::ToggleGraphSeries(TSharedRef<FTimingGraphTrack> GraphTrack, FT
 	{
 		GraphTrack->RemoveTimerSeries(TimerId);
 		GraphTrack->SetDirtyFlag();
-		NodePtr->SetAddedToGraphFlag(false);
+		NodePtr->OnRemovedFromGraph();
 	}
 	else
 	{
@@ -2189,7 +2223,7 @@ void STimersView::ToggleGraphSeries(TSharedRef<FTimingGraphTrack> GraphTrack, FT
 		Series = GraphTrack->AddTimerSeries(TimerId, NodePtr->GetColor());
 		Series->SetName(FText::FromName(NodePtr->GetName()));
 		GraphTrack->SetDirtyFlag();
-		NodePtr->SetAddedToGraphFlag(true);
+		NodePtr->OnAddedToGraph();
 	}
 }
 
@@ -2223,6 +2257,56 @@ void STimersView::ToggleTimingViewMainGraphEventSeries(FTimerNodePtr TimerNode)
 
 ////////////////////////////////////////////////////////////////////////////////////////////////////
 
+void STimersView::ToggleGraphFrameStatsSeries(TSharedRef<FTimingGraphTrack> GraphTrack, FTimerNodeRef NodePtr) const
+{
+	const uint32 TimerId = NodePtr->GetTimerId();
+	TSharedPtr<FTimingGraphSeries> Series = GraphTrack->GetFrameStatsTimerSeries(TimerId);
+	if (Series.IsValid())
+	{
+		GraphTrack->RemoveFrameStatsTimerSeries(TimerId);
+		GraphTrack->SetDirtyFlag();
+		NodePtr->OnRemovedFromGraph();
+	}
+	else
+	{
+		GraphTrack->Show();
+		Series = GraphTrack->AddFrameStatsTimerSeries(TimerId, NodePtr->GetColor());
+		Series->SetName(FText::FromName(NodePtr->GetName()));
+		GraphTrack->SetDirtyFlag();
+		NodePtr->OnAddedToGraph();
+	}
+}
+
+////////////////////////////////////////////////////////////////////////////////////////////////////
+
+bool STimersView::IsFrameStatsSeriesInTimingViewMainGraph(FTimerNodePtr TimerNode) const
+{
+	TSharedPtr<FTimingGraphTrack> GraphTrack = GetTimingViewMainGraphTrack();
+
+	if (GraphTrack.IsValid())
+	{
+		const uint32 TimerId = TimerNode->GetTimerId();
+		TSharedPtr<FTimingGraphSeries> Series = GraphTrack->GetFrameStatsTimerSeries(TimerId);
+
+		return Series.IsValid();
+	}
+
+	return false;
+}
+
+////////////////////////////////////////////////////////////////////////////////////////////////////
+
+void STimersView::ToggleTimingViewMainGraphEventFrameStatsSeries(FTimerNodePtr TimerNode) const
+{
+	TSharedPtr<FTimingGraphTrack> GraphTrack = GetTimingViewMainGraphTrack();
+	if (GraphTrack.IsValid())
+	{
+		ToggleGraphFrameStatsSeries(GraphTrack.ToSharedRef(), TimerNode.ToSharedRef());
+	}
+}
+
+////////////////////////////////////////////////////////////////////////////////////////////////////
+
 bool STimersView::ContextMenu_CopyToClipboard_CanExecute() const
 {
 	const TArray<FTimerNodePtr> SelectedNodes = TreeView->GetSelectedItems();
diff --git a/Engine/Source/Developer/TraceInsights/Private/Insights/Widgets/STimersView.h b/Engine/Source/Developer/TraceInsights/Private/Insights/Widgets/STimersView.h
--- a/Engine/Source/Developer/TraceInsights/Private/Insights/Widgets/STimersView.h
+++ b/Engine/Source/Developer/TraceInsights/Private/Insights/Widgets/STimersView.h
@@ -287,6 +287,10 @@ private:
 	bool IsSeriesInTimingViewMainGraph(FTimerNodePtr TimerNode) const;
 	void ToggleTimingViewMainGraphEventSeries(FTimerNodePtr TimerNode) const;
 
+	void ToggleGraphFrameStatsSeries(TSharedRef<FTimingGraphTrack> GraphTrack, FTimerNodeRef NodePtr) const;
+	bool IsFrameStatsSeriesInTimingViewMainGraph(FTimerNodePtr TimerNode) const;
+	void ToggleTimingViewMainGraphEventFrameStatsSeries(FTimerNodePtr TimerNode) const;
+
 	////////////////////////////////////////////////////////////////////////////////////////////////////
 
 	/**
