[Insights] Session Browser: Renamed "Open File..." to "Open Trace File..." to distinguish from importing CSV/TSV files. Also updated "Import Table..." and "Diff Tables...".

diff --git a/Engine/Source/Developer/TraceInsights/Private/Insights/Widgets/SStartPageWindow.cpp b/Engine/Source/Developer/TraceInsights/Private/Insights/Widgets/SStartPageWindow.cpp
--- a/Engine/Source/Developer/TraceInsights/Private/Insights/Widgets/SStartPageWindow.cpp
+++ b/Engine/Source/Developer/TraceInsights/Private/Insights/Widgets/SStartPageWindow.cpp
@@ -2149,7 +2149,7 @@ TSharedRef<SWidget> STraceStoreWindow::MakeTraceListMenu()
 	MenuBuilder.BeginSection("Misc", LOCTEXT("TraceListMenu_Section_Misc", "Misc"));
 	{
 		MenuBuilder.AddMenuEntry(
-			LOCTEXT("OpenFileButtonLabel", "Open File..."),
+			LOCTEXT("OpenFileButtonLabel", "Open Trace File..."),
 			LOCTEXT("OpenFileButtonTooltip", "Start analysis for a specified trace file."),
 			FSlateIcon(FAppStyle::Get().GetStyleSetName(), "Icons.FolderOpen"),
 			FUIAction(FExecuteAction::CreateSP(this, &STraceStoreWindow::OpenTraceFile)),
@@ -2157,7 +2157,7 @@ TSharedRef<SWidget> STraceStoreWindow::MakeTraceListMenu()
 			EUserInterfaceActionType::Button
 		);
 		MenuBuilder.AddMenuEntry(
-			LOCTEXT("ImportTableButtonLabel", "Import table..."),
+			LOCTEXT("ImportTableButtonLabel", "Import Table..."),
 			LOCTEXT("ImportTableButtonTooltip", "Open .csv or .tsv file."),
 			FSlateIcon(FAppStyle::Get().GetStyleSetName(), "Icons.FolderOpen"),
 			FUIAction(FExecuteAction::CreateLambda([]{ Insights::FTableImportTool::Get()->StartImportProcess(); })),
@@ -2165,8 +2165,8 @@ TSharedRef<SWidget> STraceStoreWindow::MakeTraceListMenu()
 			EUserInterfaceActionType::Button
 		);
 		MenuBuilder.AddMenuEntry(
-			LOCTEXT("DiffTablesButtonLabel", "Diff tables..."),
-			LOCTEXT("DiffTablesButtonTooltip", "Open 2 table files in diff mode."),
+			LOCTEXT("DiffTablesButtonLabel", "Diff Tables..."),
+			LOCTEXT("DiffTablesButtonTooltip", "Open two table files in diff mode."),
 			FSlateIcon(FAppStyle::Get().GetStyleSetName(), "Icons.FolderOpen"),
 			FUIAction(FExecuteAction::CreateLambda([]{ Insights::FTableImportTool::Get()->StartDiffProcess(); })),
 			NAME_None,
