[Insights] Memory Insights: Fixed crash when grouping By Heap (missing the read lock when using the Allocations provider).

diff --git a/Engine/Source/Developer/TraceInsights/Private/Insights/MemoryProfiler/ViewModels/MemAllocGroupingByHeap.cpp b/Engine/Source/Developer/TraceInsights/Private/Insights/MemoryProfiler/ViewModels/MemAllocGroupingByHeap.cpp
--- a/Engine/Source/Developer/TraceInsights/Private/Insights/MemoryProfiler/ViewModels/MemAllocGroupingByHeap.cpp
+++ b/Engine/Source/Developer/TraceInsights/Private/Insights/MemoryProfiler/ViewModels/MemAllocGroupingByHeap.cpp
@@ -2,6 +2,7 @@
 
 #include "MemAllocGroupingByHeap.h"
 #include "CallstackFormatting.h"
+#include "Common/ProviderLock.h" // TraceServices
 #include "Internationalization/Internationalization.h"
 
 // Insights
@@ -59,10 +60,13 @@ void FMemAllocGroupingByHeap::GroupNodes(const TArray<FTableTreeNodePtr>& Nodes,
 	TArray<FTableTreeNodePtr> HeapNodes;
 	HeapNodes.AddZeroed(256);
 
-	AllocProvider.EnumerateRootHeaps([&](HeapId Id, const TraceServices::IAllocationsProvider::FHeapSpec& Spec)
 	{
-		ParentGroup.AddChildAndSetGroupPtr(MakeGroupNodeHierarchy(Spec, InParentTable, HeapNodes));
-	});
+		TraceServices::FProviderReadScopeLock _(AllocProvider);
+		AllocProvider.EnumerateRootHeaps([&](HeapId Id, const TraceServices::IAllocationsProvider::FHeapSpec& Spec)
+		{
+			ParentGroup.AddChildAndSetGroupPtr(MakeGroupNodeHierarchy(Spec, InParentTable, HeapNodes));
+		});
+	}
 
 	// Add allocations nodes to heaps
 	for (FTableTreeNodePtr NodePtr : Nodes)
