[Insights] Fix bug causing the Task Track to appear twice in some cases.
Fix crash when closing the Timing Insights window

diff --git a/Engine/Source/Developer/TraceInsights/Private/Insights/ContextSwitches/ViewModels/ContextSwitchesSharedState.cpp b/Engine/Source/Developer/TraceInsights/Private/Insights/ContextSwitches/ViewModels/ContextSwitchesSharedState.cpp
--- a/Engine/Source/Developer/TraceInsights/Private/Insights/ContextSwitches/ViewModels/ContextSwitchesSharedState.cpp
+++ b/Engine/Source/Developer/TraceInsights/Private/Insights/ContextSwitches/ViewModels/ContextSwitchesSharedState.cpp
@@ -108,7 +108,7 @@ UE_ENABLE_OPTIMIZATION_SHIP
 ////////////////////////////////////////////////////////////////////////////////////////////////////
 
 FContextSwitchesSharedState::FContextSwitchesSharedState(STimingView* InTimingView)
-	: TimingView(InTimingView)
+	: TimingViewSession(InTimingView)
 	, ThreadsSerial(0)
 	, CpuCoresSerial(0)
 	, bAreCoreTracksVisible(true)
@@ -123,21 +123,11 @@ FContextSwitchesSharedState::FContextSwitchesSharedState(STimingView* InTimingVi
 
 void FContextSwitchesSharedState::OnBeginSession(Insights::ITimingViewSession& InSession)
 {
-	if (&InSession != TimingView)
+	if (&InSession != TimingViewSession)
 	{
-		TSharedPtr<STimingProfilerWindow> Window = FTimingProfilerManager::Get()->GetProfilerWindow();
-		if (Window.IsValid())
+		if (InSession.GetName() == FInsightsManagerTabs::TimingProfilerTabId)
 		{
-			TSharedPtr<STimingView> WindowTimingView = Window->GetTimingView();
-			if (WindowTimingView.IsValid() && WindowTimingView.Get() == &InSession)
-			{
-				TimingView = WindowTimingView.Get();
-				AddCommands();
-			}
-		}
-		else if (TimingView == nullptr && FTimingProfilerManager::Get()->IsTimingViewVisible())
-		{
-			TimingView = (STimingView*)&InSession;
+			TimingViewSession = &InSession;
 			AddCommands();
 		}
 		else
@@ -166,7 +156,7 @@ void FContextSwitchesSharedState::OnBeginSession(Insights::ITimingViewSession& I
 
 void FContextSwitchesSharedState::OnEndSession(Insights::ITimingViewSession& InSession)
 {
-	if (&InSession != TimingView)
+	if (&InSession != TimingViewSession)
 	{
 		return;
 	}
@@ -185,15 +175,14 @@ void FContextSwitchesSharedState::OnEndSession(Insights::ITimingViewSession& InS
 	bSyncWithProviders = false;
 
 	TargetTimingEvent = nullptr;
-
-	TimingView = nullptr;
+	TimingViewSession = nullptr;
 }
 
 ////////////////////////////////////////////////////////////////////////////////////////////////////
 
 void FContextSwitchesSharedState::Tick(Insights::ITimingViewSession& InSession, const TraceServices::IAnalysisSession& InAnalysisSession)
 {
-	if (&InSession != TimingView)
+	if (&InSession != TimingViewSession)
 	{
 		return;
 	}
@@ -222,7 +211,7 @@ void FContextSwitchesSharedState::Tick(Insights::ITimingViewSession& InSession,
 
 void FContextSwitchesSharedState::ExtendCpuTracksFilterMenu(Insights::ITimingViewSession& InSession, FMenuBuilder& InMenuBuilder)
 {
-	if (&InSession != TimingView)
+	if (&InSession != TimingViewSession)
 	{
 		return;
 	}
@@ -301,6 +290,12 @@ void FContextSwitchesSharedState::AddCommands()
 {
 	FContextSwitchesStateCommands::Register();
 
+	TSharedPtr<STimingView> TimingView = GetTimingView();
+	
+	if (!TimingView.IsValid())
+	{
+		return;
+	}
 	TSharedPtr<FUICommandList> CommandList = TimingView->GetCommandList();
 	ensure(CommandList.IsValid());
 
@@ -360,7 +355,9 @@ bool FContextSwitchesSharedState::AreContextSwitchesAvailable() const
 
 void FContextSwitchesSharedState::AddCoreTracks()
 {
-	if (!TimingView)
+	TSharedPtr<STimingView> TimingView = GetTimingView();
+
+	if (!TimingView.IsValid())
 	{
 		return;
 	}
@@ -387,7 +384,7 @@ void FContextSwitchesSharedState::AddCoreTracks()
 
 	CpuCoresSerial = NewCpuCoresSerial;
 
-	ContextSwitchesProvider->EnumerateCpuCores([this](const TraceServices::FCpuCoreInfo& CpuCoreInfo)
+	ContextSwitchesProvider->EnumerateCpuCores([this, TimingView](const TraceServices::FCpuCoreInfo& CpuCoreInfo)
 		{
 			TSharedPtr<FCpuCoreTimingTrack>* TrackPtrPtr = CpuCoreTimingTracks.Find(CpuCoreInfo.CoreNumber);
 			if (TrackPtrPtr == nullptr)
@@ -408,6 +405,7 @@ void FContextSwitchesSharedState::AddCoreTracks()
 
 void FContextSwitchesSharedState::RemoveCoreTracks()
 {
+	TSharedPtr<STimingView> TimingView = GetTimingView();
 	CpuCoresSerial = 0;
 
 	if (TimingView)
@@ -452,6 +450,12 @@ void FContextSwitchesSharedState::AddContextSwitchesChildTracks()
 	//TODO: Create "Cpu Thread" timing tracks also for threads without cpu timing events (i.e. only with context switch events).
 #endif
 
+	TSharedPtr<STimingView> TimingView = GetTimingView();
+	if (!TimingView.IsValid())
+	{
+		return;
+	}
+
 	TSharedPtr<FThreadTimingSharedState> TimingSharedState = TimingView->GetThreadTimingSharedState();
 
 	if (!TimingSharedState.IsValid())
@@ -478,7 +482,12 @@ void FContextSwitchesSharedState::AddContextSwitchesChildTracks()
 void FContextSwitchesSharedState::RemoveContextSwitchesChildTracks()
 {
 	ThreadsSerial = 0;
-
+	
+	TSharedPtr<STimingView> TimingView = GetTimingView();
+	if (!TimingView.IsValid())
+	{
+		return;
+	}
 	TSharedPtr<FThreadTimingSharedState> TimingSharedState = TimingView->GetThreadTimingSharedState();
 
 	if (!TimingSharedState.IsValid())
@@ -554,6 +563,7 @@ void FContextSwitchesSharedState::SetExtendedLinesVisible(bool bOnOff)
 
 bool FContextSwitchesSharedState::IsValidCpuCoreEventSelected() const
 {
+	TSharedPtr<STimingView> TimingView = GetTimingView();
 	if (TimingView == nullptr)
 	{
 		return false;
@@ -569,6 +579,7 @@ bool FContextSwitchesSharedState::IsValidCpuCoreEventSelected() const
 
 bool FContextSwitchesSharedState::IsValidContextSwitchEventSelected() const
 {
+	TSharedPtr<STimingView> TimingView = GetTimingView();
 	if (TimingView == nullptr)
 	{
 		return false;
@@ -584,6 +595,8 @@ bool FContextSwitchesSharedState::IsValidContextSwitchEventSelected() const
 
 void FContextSwitchesSharedState::Command_NavigateToCpuThreadEvent_Execute()
 {
+	TSharedPtr<STimingView> TimingView = GetTimingView();
+
 	if (TimingView && ensure(IsValidCpuCoreEventSelected()))
 	{
 		if (!TargetTimingEvent.IsValid())
@@ -615,6 +628,8 @@ void FContextSwitchesSharedState::Command_NavigateToCpuThreadEvent_Execute()
 
 void FContextSwitchesSharedState::Command_DockCpuThreadTrackToBottom_Execute()
 {
+	TSharedPtr<STimingView> TimingView = GetTimingView();
+
 	if (TimingView && ensure(IsValidCpuCoreEventSelected()))
 	{
 		if (!TargetTimingEvent.IsValid())
@@ -650,6 +665,8 @@ void FContextSwitchesSharedState::Command_DockCpuThreadTrackToBottom_Execute()
 
 void FContextSwitchesSharedState::Command_NavigateToCpuCoreEvent_Execute()
 {
+	TSharedPtr<STimingView> TimingView = GetTimingView();
+
 	if (TimingView && ensure(IsValidContextSwitchEventSelected()))
 	{
 		if (!TargetTimingEvent.IsValid())
@@ -678,6 +695,8 @@ void FContextSwitchesSharedState::Command_NavigateToCpuCoreEvent_Execute()
 
 void FContextSwitchesSharedState::Command_DockCpuCoreTrackToTop_Execute()
 {
+	TSharedPtr<STimingView> TimingView = GetTimingView();
+
 	if (TimingView && ensure(IsValidContextSwitchEventSelected()))
 	{
 		if (!TargetTimingEvent.IsValid())
@@ -737,6 +756,7 @@ void FContextSwitchesSharedState::GetThreadInfo(uint32 InSystemThreadId, uint32&
 TSharedPtr<FThreadTimingTrack> FContextSwitchesSharedState::GetThreadTimingTrack(uint32 ThreadId) const
 {
 	TSharedPtr<FThreadTimingTrack> FoundTrack;
+	TSharedPtr<STimingView> TimingView = GetTimingView();
 
 	if (TimingView)
 	{
@@ -760,6 +780,7 @@ TSharedPtr<FThreadTimingTrack> FContextSwitchesSharedState::GetThreadTimingTrack
 TSharedPtr<Insights::FCpuCoreTimingTrack> FContextSwitchesSharedState::GetCpuCoreTimingTrack(uint32 CoreNumber) const
 {
 	TSharedPtr<FCpuCoreTimingTrack> FoundTrack;
+	TSharedPtr<STimingView> TimingView = GetTimingView();
 
 	if (TimingView)
 	{
@@ -780,6 +801,20 @@ TSharedPtr<Insights::FCpuCoreTimingTrack> FContextSwitchesSharedState::GetCpuCor
 
 ////////////////////////////////////////////////////////////////////////////////////////////////////
 
+TSharedPtr<STimingView> FContextSwitchesSharedState::GetTimingView()
+{
+	TSharedPtr<STimingView> TimingView;
+	TSharedPtr<STimingProfilerWindow> Window = FTimingProfilerManager::Get()->GetProfilerWindow();
+	if (Window.IsValid())
+	{
+		TimingView = Window->GetTimingView();
+	}
+
+	return TimingView;
+}
+
+////////////////////////////////////////////////////////////////////////////////////////////////////
+
 } // namespace Insights
 
 #undef LOCTEXT_NAMESPACE
diff --git a/Engine/Source/Developer/TraceInsights/Private/Insights/ContextSwitches/ViewModels/ContextSwitchesSharedState.h b/Engine/Source/Developer/TraceInsights/Private/Insights/ContextSwitches/ViewModels/ContextSwitchesSharedState.h
--- a/Engine/Source/Developer/TraceInsights/Private/Insights/ContextSwitches/ViewModels/ContextSwitchesSharedState.h
+++ b/Engine/Source/Developer/TraceInsights/Private/Insights/ContextSwitches/ViewModels/ContextSwitchesSharedState.h
@@ -71,7 +71,7 @@ public:
 
 	//////////////////////////////////////////////////
 
-	STimingView* GetTimingView() { return TimingView; }
+	static TSharedPtr<STimingView> GetTimingView();
 
 	bool AreContextSwitchesAvailable() const;
 
@@ -153,11 +153,11 @@ private:
 	void PopulateCoreEventNameSuggestionList(const FString& Text, TArray<FString>& OutSuggestions);
 
 private:
-	STimingView* TimingView;
-
 	TMap<uint32, TSharedPtr<FCpuCoreTimingTrack>> CpuCoreTimingTracks;
 	TMap<uint32, TSharedPtr<FContextSwitchesTimingTrack>> ContextSwitchesTimingTracks;
 
+	Insights::ITimingViewSession* TimingViewSession;
+
 	uint64 ThreadsSerial;
 	uint64 CpuCoresSerial;
 
diff --git a/Engine/Source/Developer/TraceInsights/Private/Insights/ContextSwitches/ViewModels/CpuCoreTimingTrack.cpp b/Engine/Source/Developer/TraceInsights/Private/Insights/ContextSwitches/ViewModels/CpuCoreTimingTrack.cpp
--- a/Engine/Source/Developer/TraceInsights/Private/Insights/ContextSwitches/ViewModels/CpuCoreTimingTrack.cpp
+++ b/Engine/Source/Developer/TraceInsights/Private/Insights/ContextSwitches/ViewModels/CpuCoreTimingTrack.cpp
@@ -106,7 +106,7 @@ void FCpuCoreTimingTrack::BuildFilteredDrawState(ITimingEventsTrackDrawStateBuil
 
 	if (HasCustomFilter())
 	{
-		STimingView* TimingView = SharedState.GetTimingView();
+		TSharedPtr<STimingView> TimingView = SharedState.GetTimingView();
 		if (!TimingView)
 		{
 			return;
@@ -420,7 +420,7 @@ void FCpuCoreTimingTrack::BuildContextMenu(FMenuBuilder& InOutMenuBuilder)
 
 bool FCpuCoreTimingTrack::HasCustomFilter() const
 {
-	STimingView* TimingView = SharedState.GetTimingView();
+	TSharedPtr<STimingView> TimingView = SharedState.GetTimingView();
 	if (!TimingView)
 	{
 		return false;
diff --git a/Engine/Source/Developer/TraceInsights/Private/Insights/TaskGraphProfiler/TaskGraphProfilerManager.cpp b/Engine/Source/Developer/TraceInsights/Private/Insights/TaskGraphProfiler/TaskGraphProfilerManager.cpp
--- a/Engine/Source/Developer/TraceInsights/Private/Insights/TaskGraphProfiler/TaskGraphProfilerManager.cpp
+++ b/Engine/Source/Developer/TraceInsights/Private/Insights/TaskGraphProfiler/TaskGraphProfilerManager.cpp
@@ -238,6 +238,8 @@ TSharedRef<SDockTab> FTaskGraphProfilerManager::SpawnTab_TaskTableTreeView(const
 			SAssignNew(TaskTableTreeView, STaskTableTreeView, TaskTable)
 		];
 
+	RegisterOnWindowClosedEventHandle();
+
 	DockTab->SetOnTabClosed(SDockTab::FOnTabClosedCallback::CreateRaw(this, &FTaskGraphProfilerManager::OnTaskTableTreeViewTabClosed));
 
 	return DockTab;
@@ -491,6 +493,8 @@ void FTaskGraphProfilerManager::ShowTaskRelations(TaskTrace::FId TaskId)
 
 void FTaskGraphProfilerManager::OnWindowClosedEvent()
 {
+	OnWindowClosedEventHandle.Reset();
+
 	if (TaskTableTreeView.IsValid())
 	{
 		TaskTableTreeView->OnClose();
@@ -819,6 +823,20 @@ void FTaskGraphProfilerManager::SelectTaskInTaskTable(TaskTrace::FId InId)
 
 ////////////////////////////////////////////////////////////////////////////////////////////////////
 
+void FTaskGraphProfilerManager::RegisterOnWindowClosedEventHandle()
+{
+	if (!OnWindowClosedEventHandle.IsValid())
+	{
+		TSharedPtr<STimingProfilerWindow> Window = FTimingProfilerManager::Get()->GetProfilerWindow();
+		if (!Window.IsValid())
+		{
+			return;
+		}
+
+		OnWindowClosedEventHandle = Window->GetWindowClosedEvent().AddSP(this, &FTaskGraphProfilerManager::OnWindowClosedEvent);
+	}
+}
+
 } // namespace Insights
 
 #undef LOCTEXT_NAMESPACE
diff --git a/Engine/Source/Developer/TraceInsights/Private/Insights/TaskGraphProfiler/TaskGraphProfilerManager.h b/Engine/Source/Developer/TraceInsights/Private/Insights/TaskGraphProfiler/TaskGraphProfilerManager.h
--- a/Engine/Source/Developer/TraceInsights/Private/Insights/TaskGraphProfiler/TaskGraphProfilerManager.h
+++ b/Engine/Source/Developer/TraceInsights/Private/Insights/TaskGraphProfiler/TaskGraphProfilerManager.h
@@ -135,6 +135,8 @@ public:
 
 	void SelectTaskInTaskTable(TaskTrace::FId);
 
+	void RegisterOnWindowClosedEventHandle();
+
 private:
 	/** Updates this manager, done through FCoreTicker. */
 	bool Tick(float DeltaTime);
@@ -180,6 +182,8 @@ private:
 	bool bShowSubsequents = false;
 	bool bShowParentTasks = false;
 	bool bShowNestedTasks = false;
+	
+	FDelegateHandle OnWindowClosedEventHandle;
 };
 
 } // namespace Insights
diff --git a/Engine/Source/Developer/TraceInsights/Private/Insights/TaskGraphProfiler/ViewModels/TaskTimingTrack.cpp b/Engine/Source/Developer/TraceInsights/Private/Insights/TaskGraphProfiler/ViewModels/TaskTimingTrack.cpp
--- a/Engine/Source/Developer/TraceInsights/Private/Insights/TaskGraphProfiler/ViewModels/TaskTimingTrack.cpp
+++ b/Engine/Source/Developer/TraceInsights/Private/Insights/TaskGraphProfiler/ViewModels/TaskTimingTrack.cpp
@@ -122,7 +122,7 @@ public:
 ////////////////////////////////////////////////////////////////////////////////////////////////////
 
 FTaskTimingSharedState::FTaskTimingSharedState(STimingView* InTimingView)
-	: TimingView(InTimingView)
+	: TimingViewSession(InTimingView)
 {
 }
 
@@ -130,20 +130,11 @@ FTaskTimingSharedState::FTaskTimingSharedState(STimingView* InTimingView)
 
 void FTaskTimingSharedState::OnBeginSession(Insights::ITimingViewSession& InSession)
 {
-	if (&InSession != TimingView)
+	if (&InSession != TimingViewSession)
 	{
-		TSharedPtr<STimingProfilerWindow> Window = FTimingProfilerManager::Get()->GetProfilerWindow();
-		if (Window.IsValid())
+		if (InSession.GetName() == FInsightsManagerTabs::TimingProfilerTabId)
 		{
-			TSharedPtr<STimingView> WindowTimingView = Window->GetTimingView();
-			if (WindowTimingView.IsValid() && WindowTimingView.Get() == &InSession)
-			{
-				TimingView = WindowTimingView.Get();
-			}
-		}
-		else if (TimingView == nullptr && FTimingProfilerManager::Get()->IsTimingViewVisible())
-		{
-			TimingView = (STimingView*)&InSession;
+			TimingViewSession = &InSession;
 		}
 		else
 		{
@@ -158,28 +149,32 @@ void FTaskTimingSharedState::OnBeginSession(Insights::ITimingViewSession& InSess
 
 void FTaskTimingSharedState::OnEndSession(Insights::ITimingViewSession& InSession)
 {
-	if (&InSession != TimingView)
+	if (&InSession != TimingViewSession)
 	{
 		return;
 	}
 
 	FTaskGraphProfilerManager::Get()->ClearTaskRelations();
 	TaskTrack = nullptr;
-	TimingView = nullptr;
+	TimingViewSession = nullptr;
 }
 
 ////////////////////////////////////////////////////////////////////////////////////////////////////
 
 void FTaskTimingSharedState::Tick(Insights::ITimingViewSession& InSession, const TraceServices::IAnalysisSession& InAnalysisSession)
 {
-	if (&InSession != TimingView)
+	if (&InSession != TimingViewSession)
 	{
 		return;
 	}
 
 	if (!TaskTrack.IsValid() && FTaskGraphProfilerManager::Get()->GetIsAvailable())
 	{
-		InitCommandList();
+		TSharedPtr<STimingView> TimingView = GetTimingView();
+
+		InitCommandList(TimingView);
+
+		FTaskGraphProfilerManager::Get()->RegisterOnWindowClosedEventHandle();
 
 		TaskTrack = MakeShared<FTaskTimingTrack>(*this, TEXT("Task Overview"), 0);
 		TaskTrack->SetVisibilityFlag(true);
@@ -192,6 +187,8 @@ void FTaskTimingSharedState::Tick(Insights::ITimingViewSession& InSession, const
 
 	if (bResetOnNextTick)
 	{
+		TSharedPtr<STimingView> TimingView = GetTimingView();
+
 		bResetOnNextTick = false;
 		if (!TimingView->GetSelectedEvent().IsValid() &&
 			(!TimingView->GetSelectedTrack().IsValid() || TimingView->GetSelectedTrack().Get() != TaskTrack.Get()))
@@ -204,9 +201,23 @@ void FTaskTimingSharedState::Tick(Insights::ITimingViewSession& InSession, const
 
 ////////////////////////////////////////////////////////////////////////////////////////////////////
 
+TSharedPtr<STimingView> FTaskTimingSharedState::GetTimingView()
+{
+	TSharedPtr<STimingView> TimingView;
+	TSharedPtr<STimingProfilerWindow> Window = FTimingProfilerManager::Get()->GetProfilerWindow();
+	if (Window.IsValid())
+	{
+		TimingView = Window->GetTimingView();
+	}
+
+	return TimingView;
+}
+
+////////////////////////////////////////////////////////////////////////////////////////////////////
+
 void FTaskTimingSharedState::ExtendFilterMenu(Insights::ITimingViewSession& InSession, FMenuBuilder& InOutMenuBuilder)
 {
-	if (&InSession != TimingView)
+	if (&InSession != TimingViewSession)
 	{
 		return;
 	}
@@ -226,7 +237,7 @@ void FTaskTimingSharedState::SetTaskId(TaskTrace::FId InTaskId)
 
 void FTaskTimingSharedState::ExtendOtherTracksFilterMenu(ITimingViewSession& InSession, FMenuBuilder& InMenuBuilder)
 {
-	if (&InSession != TimingView)
+	if (&InSession != TimingViewSession)
 	{
 		return;
 	}
@@ -335,7 +346,7 @@ void FTaskTimingSharedState::BuildTasksSubMenu(FMenuBuilder& MenuBuilder)
 
 ////////////////////////////////////////////////////////////////////////////////////////////////////
 
-void FTaskTimingSharedState::InitCommandList()
+void FTaskTimingSharedState::InitCommandList(TSharedPtr<STimingView> TimingView)
 {
 	TSharedPtr<FUICommandList> CommandList = TimingView->GetCommandList();
 	ensure(CommandList);
@@ -665,6 +676,12 @@ void FTaskTimingSharedState::OnTaskSettingsChanged()
 	TaskTrack->SetTaskId(TaskTrace::InvalidId);
 
 	TSharedPtr<Insights::FTaskGraphProfilerManager> TaskGraphManager = Insights::FTaskGraphProfilerManager::Get();
+	TSharedPtr<STimingView> TimingView = GetTimingView();
+
+	if (!TimingView)
+	{
+		return;
+	}
 
 	TSharedPtr<const ITimingEvent> SelectedEvent = TimingView->GetSelectedEvent();
 	if (SelectedEvent.IsValid() && SelectedEvent->Is<const FThreadTrackEvent>())
@@ -1074,7 +1091,12 @@ void FTaskTimingTrack::GetEventRelations(const FThreadTrackEvent& InSelectedEven
 	const TraceServices::ITasksProvider* TasksProvider = TraceServices::ReadTasksProvider(*Session.Get());
 	if (TasksProvider)
 	{
-		STimingView* TimingView = SharedState.GetTimingView();
+		TSharedPtr<STimingView> TimingView = SharedState.GetTimingView();
+		if (!TimingView.IsValid())
+		{
+			return;
+		}
+
 		TSharedRef<const FThreadTimingTrack> EventTrack = StaticCastSharedRef<const FThreadTimingTrack>(InSelectedEvent.GetTrack());
 		uint32 ThreadId = EventTrack->GetThreadId();
 
diff --git a/Engine/Source/Developer/TraceInsights/Private/Insights/TaskGraphProfiler/ViewModels/TaskTimingTrack.h b/Engine/Source/Developer/TraceInsights/Private/Insights/TaskGraphProfiler/ViewModels/TaskTimingTrack.h
--- a/Engine/Source/Developer/TraceInsights/Private/Insights/TaskGraphProfiler/ViewModels/TaskTimingTrack.h
+++ b/Engine/Source/Developer/TraceInsights/Private/Insights/TaskGraphProfiler/ViewModels/TaskTimingTrack.h
@@ -51,10 +51,10 @@ public:
 
 	void SetResetOnNextTick(bool bInValue) { bResetOnNextTick = bInValue; }
 
-	STimingView* GetTimingView() { return TimingView; }
+	static TSharedPtr<STimingView> GetTimingView();
 
 private:
-	void InitCommandList();
+	void InitCommandList(TSharedPtr<STimingView> TimingView);
 
 	void BuildTasksSubMenu(FMenuBuilder& MenuBuilder);
 
@@ -97,7 +97,7 @@ private:
 	void OnTaskSettingsChanged();
 
 private:
-	STimingView* TimingView;
+	Insights::ITimingViewSession* TimingViewSession;
 
 	bool bShowHideTaskTrack;
 	bool bResetOnNextTick = false;
diff --git a/Engine/Source/Developer/TraceInsights/Private/Insights/Widgets/SMajorTabWindow.cpp b/Engine/Source/Developer/TraceInsights/Private/Insights/Widgets/SMajorTabWindow.cpp
--- a/Engine/Source/Developer/TraceInsights/Private/Insights/Widgets/SMajorTabWindow.cpp
+++ b/Engine/Source/Developer/TraceInsights/Private/Insights/Widgets/SMajorTabWindow.cpp
@@ -119,6 +119,8 @@ void SMajorTabWindow::RemoveOpenTab(const TSharedRef<SDockTab>& DockTab)
 
 void SMajorTabWindow::CloseAllOpenTabs()
 {
+	WindowClosedEvent.Broadcast();
+
 	// Close all tabs.
 	TArray<TSharedPtr<SDockTab>> LocalOpenTabs;
 	for (TSharedPtr<SDockTab>& Tab : OpenTabs)
diff --git a/Engine/Source/Developer/TraceInsights/Public/Insights/Widgets/SMajorTabWindow.h b/Engine/Source/Developer/TraceInsights/Public/Insights/Widgets/SMajorTabWindow.h
--- a/Engine/Source/Developer/TraceInsights/Public/Insights/Widgets/SMajorTabWindow.h
+++ b/Engine/Source/Developer/TraceInsights/Public/Insights/Widgets/SMajorTabWindow.h
@@ -52,6 +52,16 @@ public:
 	void ShowHideTab(const FName& TabId, bool bShow) { bShow ? ShowTab(TabId) : HideTab(TabId); }
 	void CloseAllOpenTabs();
 
+	////////////////////////////////////////////////////////////////////////////////////////////////////
+	// OnWindowClosedEvent
+public:
+	/** The event to execute when the window is closed. */
+	DECLARE_MULTICAST_DELEGATE(FWindowClosedEvent);
+	FWindowClosedEvent& GetWindowClosedEvent() { return WindowClosedEvent; }
+private:
+	/** The event to execute when the window is closed. */
+	FWindowClosedEvent WindowClosedEvent;
+
 protected:
 	virtual const TCHAR* GetAnalyticsEventName() const;
 	virtual TSharedRef<FWorkspaceItem> CreateWorkspaceMenuGroup();
