[Insights] Update TimingExporter to use new RegionProvider

diff --git a/Engine/Source/Developer/TraceInsights/Private/Insights/ViewModels/TimingExporter.cpp b/Engine/Source/Developer/TraceInsights/Private/Insights/ViewModels/TimingExporter.cpp
--- a/Engine/Source/Developer/TraceInsights/Private/Insights/ViewModels/TimingExporter.cpp
+++ b/Engine/Source/Developer/TraceInsights/Private/Insights/ViewModels/TimingExporter.cpp
@@ -12,6 +12,7 @@
 #include "Insights/Log.h"
 #include "Insights/ViewModels/ThreadTimingTrack.h"
 #include "Insights/TimingProfilerManager.h"
+#include "TraceServices/Model/Regions.h"
 
 namespace Insights
 {
@@ -563,25 +564,11 @@ int32 FTimingExporter::ExportTimerStatisticsAsText(const FString& Filename, FExp
 				InNamePatternList.ParseIntoArray(NamePatterns, TEXT(","), true);
 			}
 
-			bool Match(const FString& InBookmarkName, FString& OutRegionName, bool& bOutIsRegionStart)
+			bool Match(const FString& InRegionName)
 			{
-				if (InBookmarkName.StartsWith(TEXT("RegionStart:")))
-				{
-					bOutIsRegionStart = true;
-					OutRegionName = InBookmarkName.RightChop(12);
-				}
-				else if (InBookmarkName.StartsWith(TEXT("RegionEnd:")))
-				{
-					bOutIsRegionStart = false;
-					OutRegionName = InBookmarkName.RightChop(10);
-				}
-				else
-				{
-					return false;
-				}
 				for (const FString& NamePattern : NamePatterns)
 				{
-					if (OutRegionName.MatchesWildcard(NamePattern))
+					if (InRegionName.MatchesWildcard(NamePattern))
 					{
 						return true;
 					}
@@ -593,72 +580,50 @@ int32 FTimingExporter::ExportTimerStatisticsAsText(const FString& Filename, FExp
 			TArray<FString> NamePatterns;
 		};
 		FRegionNameSpec RegionNameSpec(Params.Region);
-
-		struct FTimeRegion
-		{
-			double Start;
-			double End;
-		};
-		TMap<FString, FTimeRegion> Regions;
+		
+		TMap<FString, TraceServices::FTimeRegion> Regions;
 
 		// Detect regions
 		{
 			TraceServices::FAnalysisSessionReadScope SessionReadScope(Session);
 
-			const TraceServices::IBookmarkProvider& BookmarkProvider = TraceServices::ReadBookmarkProvider(Session);
+			const TraceServices::IRegionProvider& RegionProvider = TraceServices::ReadRegionProvider(Session);
 
 			UE_LOG(TraceInsights, Log, TEXT("Looking for regions: '%s'"), *Params.Region);
 
-			BookmarkProvider.EnumerateBookmarks(0.0, std::numeric_limits<double>::max(),
-				[&RegionNameSpec, &Regions](const TraceServices::FBookmark& InBookmark)
+			RegionProvider.EnumerateRegions(0.0, std::numeric_limits<double>::max(),
+				[&RegionNameSpec, &Regions](const TraceServices::FTimeRegion& InRegion) -> bool
 				{
-					FString BaseRegionName;
-					bool bIsRegionStart;
-					if (RegionNameSpec.Match(InBookmark.Text, BaseRegionName, bIsRegionStart))
+					if (!RegionNameSpec.Match(InRegion.Text))
 					{
-						FTimeRegion* Region = nullptr;
-						FString RegionName = BaseRegionName;
-						int32 Index = 2;
-
-						while (true)
-						{
-							Region = Regions.Find(RegionName);
-
-							if (!Region)
-							{
-								Region = &Regions.Add(RegionName);
-								Region->Start = -std::numeric_limits<double>::infinity();
-								Region->End = std::numeric_limits<double>::infinity();
-							}
+						return true;
+					}
+					
+					// Handle duplicate region names, individual regions may appear multiple times
+					// we append numbers to allow for unique export filenames.
+					FString RegionName = InRegion.Text;
+					int32 Index = 2;
 
-							if (bIsRegionStart)
-							{
-								if (Region->Start == -std::numeric_limits<double>::infinity() &&
-									InBookmark.Time <= Region->End)
-								{
-									Region->Start = InBookmark.Time;
-									break;
-								}
-							}
-							else
-							{
-								if (Region->End == std::numeric_limits<double>::infinity() &&
-									InBookmark.Time >= Region->Start)
-								{
-									Region->End = InBookmark.Time;
-									break;
-								}
-							}
+					while (true)
+					{
+						const TraceServices::FTimeRegion* ExistingRegion = Regions.Find(RegionName);
 
-							RegionName = FString::Printf(TEXT("%s%d"), *BaseRegionName, Index++);
+						if (!ExistingRegion)
+						{
+							Regions.Add(RegionName, InRegion);
+							break;
 						}
+
+						RegionName = FString::Printf(TEXT("%s%d"), *RegionName, Index++);
 					}
+
+					return true;
 				});
 		}
 
 		if (Regions.Num() == 0)
 		{
-			UE_LOG(TraceInsights, Error, TEXT("Unable to find any region ('RegionStart:' and 'RegionEnd:' bookmarks) with name pattern '%s'."), *Params.Region);
+			UE_LOG(TraceInsights, Error, TEXT("Unable to find any region with name pattern '%s'."), *Params.Region);
 			return -1;
 		}
 
@@ -671,8 +636,8 @@ int32 FTimingExporter::ExportTimerStatisticsAsText(const FString& Filename, FExp
 			FString RegionFilename = Filename.Replace(TEXT("*"), *KV.Key);
 			FExportTimerStatisticsParams RegionParams = Params;
 			RegionParams.Region.Reset();
-			RegionParams.IntervalStartTime = KV.Value.Start;
-			RegionParams.IntervalEndTime = KV.Value.End;
+			RegionParams.IntervalStartTime = KV.Value.BeginTime;
+			RegionParams.IntervalEndTime = KV.Value.EndTime;
 			UE_LOG(TraceInsights, Display, TEXT("Exporting timing statistics for region '%s' [%f .. %f] to '%s'"), *KV.Key, RegionParams.IntervalStartTime, RegionParams.IntervalEndTime, *RegionFilename);
 			ExportTimerStatisticsAsText(RegionFilename, RegionParams);
 		}
