[Insights] Add visualization of regions through FTimingRegionsTrack

diff --git a/Engine/Source/Developer/TraceInsights/Private/Insights/ViewModels/RegionsTimingTrack.cpp b/Engine/Source/Developer/TraceInsights/Private/Insights/ViewModels/RegionsTimingTrack.cpp
new file mode 100644
--- /dev/null
+++ b/Engine/Source/Developer/TraceInsights/Private/Insights/ViewModels/RegionsTimingTrack.cpp
@@ -0,0 +1,255 @@
+﻿// Copyright Epic Games, Inc. All Rights Reserved.
+
+#include "RegionsTimingTrack.h"
+
+#include "TimingViewDrawHelper.h"
+#include "Common/ProviderLock.h"
+#include "Insights/InsightsManager.h"
+#include "Insights/InsightsStyle.h"
+#include "Insights/ITimingViewSession.h"
+#include "Insights/Log.h"
+#include "Insights/TimingProfilerCommon.h"
+#include "Insights/Common/InsightsMenuBuilder.h"
+#include "Insights/Common/TimeUtils.h"
+#include "Insights/ViewModels/TimingTrackViewport.h"
+#include "Insights/Widgets/STimingView.h"
+#include "TraceServices/Model/Regions.h"
+#include "Framework/MultiBox/MultiBoxBuilder.h"
+
+#define LOCTEXT_NAMESPACE "RegionsTimingTrack"
+
+////////////////////////////////////////////////////////////////////////////////////////////////////
+// FFileActivityTimingViewCommands
+////////////////////////////////////////////////////////////////////////////////////////////////////
+
+FTimingRegionsViewCommands::FTimingRegionsViewCommands()
+: TCommands<FTimingRegionsViewCommands>(
+	TEXT("FTimingRegionsViewCommands"),
+	NSLOCTEXT("Contexts", "FTimingRegionsViewCommands", "Insights - Timing View - Timing Regions"),
+	NAME_None,
+	FInsightsStyle::GetStyleSetName())
+{
+}
+
+////////////////////////////////////////////////////////////////////////////////////////////////////
+
+FTimingRegionsViewCommands::~FTimingRegionsViewCommands()
+{
+}
+
+////////////////////////////////////////////////////////////////////////////////////////////////////
+
+// UI_COMMAND takes long for the compiler to optimize
+UE_DISABLE_OPTIMIZATION_SHIP
+void FTimingRegionsViewCommands::RegisterCommands()
+{
+	UI_COMMAND(ShowHideRegionTrack,
+		"Timing Regions Track",
+		"Shows/hides the Timing Regions track.",
+		EUserInterfaceActionType::ToggleButton,
+		FInputChord(EKeys::R));
+}
+UE_ENABLE_OPTIMIZATION_SHIP
+
+////////////////////////////////////////////////////////////////////////////////////////////////////
+
+
+
+
+void FTimingRegionsSharedState::OnBeginSession(Insights::ITimingViewSession& InSession)
+{
+	TimingRegionsTrack.Reset();
+}
+
+void FTimingRegionsSharedState::OnEndSession(Insights::ITimingViewSession& InSession)
+{
+	TimingRegionsTrack.Reset();
+}
+
+void FTimingRegionsSharedState::Tick(Insights::ITimingViewSession& InSession,
+	const TraceServices::IAnalysisSession& InAnalysisSession)
+{
+	if (&InSession != TimingView)
+	{
+		return;
+	}
+
+	if (!TimingRegionsTrack.IsValid())
+	{
+		TimingRegionsTrack = MakeShared<FTimingRegionsTrack>(*this);
+		TimingRegionsTrack->SetOrder(FTimingTrackOrder::First);
+		TimingRegionsTrack->SetVisibilityFlag(true);
+		InSession.AddScrollableTrack(TimingRegionsTrack);
+	}
+}
+
+void FTimingRegionsSharedState::ShowHideRegionsTrack()
+{
+	bShowHideRegionsTrack = !bShowHideRegionsTrack;
+
+	if (TimingRegionsTrack.IsValid())
+	{
+		TimingRegionsTrack->SetVisibilityFlag(bShowHideRegionsTrack);
+	}
+
+	if (TimingView)
+	{
+		TimingView->OnTrackVisibilityChanged();
+	}
+
+	if (bShowHideRegionsTrack)
+	{
+		TimingRegionsTrack->SetDirtyFlag();
+	}
+}
+
+void FTimingRegionsSharedState::ExtendOtherTracksFilterMenu(Insights::ITimingViewSession& InSession,
+                                                            FMenuBuilder& InOutMenuBuilder)
+{
+	InOutMenuBuilder.BeginSection("Timing Regions", LOCTEXT("ContextMenu_Section_Regions", "Timing Regions"));
+	{
+		// Note: We use the custom AddMenuEntry in order to set the same key binding text for multiple menu items.
+
+		//InOutMenuBuilder.AddMenuEntry(FFileActivityTimingViewCommands::Get().ShowHideIoOverviewTrack);
+		FInsightsMenuBuilder::AddMenuEntry(InOutMenuBuilder,
+			FUIAction(
+				FExecuteAction::CreateSP(this, &FTimingRegionsSharedState::ShowHideRegionsTrack),
+				FCanExecuteAction(),
+				FIsActionChecked::CreateSP(this, &FTimingRegionsSharedState::IsRegionsTrackVisible)),
+			FTimingRegionsViewCommands::Get().ShowHideRegionTrack->GetLabel(),
+			FTimingRegionsViewCommands::Get().ShowHideRegionTrack->GetDescription(),
+			LOCTEXT("TimingRegionsTracksKeybinding", "R"),
+			EUserInterfaceActionType::ToggleButton);
+	}
+	InOutMenuBuilder.EndSection();
+}
+
+void FTimingRegionsSharedState::BindCommands()
+{
+	FTimingRegionsViewCommands::Register();
+
+	TSharedPtr<FUICommandList> CommandList = TimingView->GetCommandList();
+	ensure(CommandList.IsValid());
+
+	CommandList->MapAction(
+		FTimingRegionsViewCommands::Get().ShowHideRegionTrack,
+		FExecuteAction::CreateSP(this, &FTimingRegionsSharedState::ShowHideRegionsTrack),
+		FCanExecuteAction(),
+		FIsActionChecked::CreateSP(this, &FTimingRegionsSharedState::IsRegionsTrackVisible));
+}
+
+
+//////////////////////////////////////////////////////////////////////////
+
+INSIGHTS_IMPLEMENT_RTTI(FTimingRegionsTrack)
+
+FTimingRegionsTrack::~FTimingRegionsTrack()
+{
+}
+
+void FTimingRegionsTrack::BuildContextMenu(FMenuBuilder& MenuBuilder)
+{
+	FTimingEventsTrack::BuildContextMenu(MenuBuilder);
+}
+
+
+void FTimingRegionsTrack::InitTooltip(FTooltipDrawState& InOutTooltip, const ITimingEvent& InTooltipEvent) const
+{
+	if (InTooltipEvent.CheckTrack(this) && InTooltipEvent.Is<FTimingEvent>())
+	{
+		const FTimingEvent& TooltipEvent = InTooltipEvent.As<FTimingEvent>();
+
+		auto MatchEvent = [this, &TooltipEvent](double InStartTime, double InEndTime, uint32 InDepth)
+		{
+			return InDepth == TooltipEvent.GetDepth()
+				&& InStartTime == TooltipEvent.GetStartTime()
+				&& InEndTime == TooltipEvent.GetEndTime();
+		};
+
+		FTimingEventSearchParameters SearchParameters(TooltipEvent.GetStartTime(), TooltipEvent.GetEndTime(), ETimingEventSearchFlags::StopAtFirstMatch, MatchEvent);
+		FindRegionEvent(SearchParameters, [this, &InOutTooltip, &TooltipEvent](double InFoundStartTime, double InFoundEndTime, uint32 InFoundDepth, const TraceServices::FTimeRegion& InRegion)
+		{
+			InOutTooltip.Reset();
+			InOutTooltip.AddTitle(InRegion.Text, FLinearColor::White);
+			InOutTooltip.AddNameValueTextLine(TEXT("Duration:"),  TimeUtils::FormatTimeAuto(InRegion.EndTime-InRegion.BeginTime));
+			InOutTooltip.AddNameValueTextLine(TEXT("Depth:"),  FString::FromInt(InRegion.Depth));
+			InOutTooltip.UpdateLayout();
+		});
+	}
+}
+
+void FTimingRegionsTrack::BuildDrawState(ITimingEventsTrackDrawStateBuilder& Builder,
+	const ITimingTrackUpdateContext& Context)
+{
+	const FTimingTrackViewport& Viewport = Context.GetViewport();
+	TSharedPtr<const TraceServices::IAnalysisSession> Session = FInsightsManager::Get()->GetSession();
+
+	const TraceServices::IRegionProvider& RegionProvider = TraceServices::ReadRegionProvider(*Session);
+	TraceServices::FProviderReadScopeLock RegionProviderScopedLock(RegionProvider);
+
+	FStopwatch Stopwatch;
+	Stopwatch.Start();
+	
+	// whe're counting only non-empty lanes, so we can collapse empty ones in the visualization.
+	int32 CurDepth = 0;
+	RegionProvider.EnumerateLanes([this, Viewport, &CurDepth, &Builder](const TraceServices::FRegionLane& Lane, const int32 Depth)
+	{
+		bool RegionHadEvents = false;
+		Lane.EnumerateRegions(Viewport.GetStartTime(), Viewport.GetEndTime(), [&Builder, &RegionHadEvents, &CurDepth](const TraceServices::FTimeRegion& Region) -> bool
+		{
+			RegionHadEvents = true;
+			Builder.AddEvent(Region.BeginTime, Region.EndTime,CurDepth, Region.Text);
+			return true;
+		});
+
+		if (RegionHadEvents) CurDepth++;
+	});
+
+	Stopwatch.Stop();
+	const double TotalTime = Stopwatch.GetAccumulatedTime();
+	UE_CLOG(TotalTime > 1.0,TimingProfiler, Verbose, TEXT("[Regions] Updated draw state in %s."), *TimeUtils::FormatTimeAuto(TotalTime));
+}
+
+const TSharedPtr<const ITimingEvent> FTimingRegionsTrack::SearchEvent(
+	const FTimingEventSearchParameters& InSearchParameters) const
+{
+	TSharedPtr<const ITimingEvent> FoundEvent;
+
+	FindRegionEvent(InSearchParameters, [this, &FoundEvent](double InFoundStartTime, double InFoundEndTime, uint32 InFoundDepth, const TraceServices::FTimeRegion& InEvent)
+	{
+		FoundEvent = MakeShared<FTimingEvent>(SharedThis(this), InFoundStartTime, InFoundEndTime, InFoundDepth);
+	});
+
+	return FoundEvent;
+}
+
+bool FTimingRegionsTrack::FindRegionEvent(const FTimingEventSearchParameters& InParameters,
+	TFunctionRef<void(double, double, uint32, const TraceServices::FTimeRegion&)> InFoundPredicate) const
+{
+	return TTimingEventSearch<TraceServices::FTimeRegion>::Search(
+	InParameters,
+
+	// Search...
+	[this](TTimingEventSearch<TraceServices::FTimeRegion>::FContext& InContext)
+	{
+		TSharedPtr<const TraceServices::IAnalysisSession> Session = FInsightsManager::Get()->GetSession();
+		const TraceServices::IRegionProvider& RegionProvider = TraceServices::ReadRegionProvider(*Session);
+			TraceServices::FProviderReadScopeLock RegionProviderScopedLock(RegionProvider);
+
+		RegionProvider.EnumerateRegions(InContext.GetParameters().StartTime, InContext.GetParameters().EndTime, [&InContext](const TraceServices::FTimeRegion& Region)
+		{
+			InContext.Check(Region.BeginTime, Region.EndTime, Region.Depth, Region);
+			
+			if (!InContext.ShouldContinueSearching())
+			{
+				return false;
+			}
+			
+			return true;
+		});
+	},
+	// Found!
+	InFoundPredicate);
+}
+
+#undef LOCTEXT_NAMESPACE
\ No newline at end of file
diff --git a/Engine/Source/Developer/TraceInsights/Private/Insights/ViewModels/RegionsTimingTrack.h b/Engine/Source/Developer/TraceInsights/Private/Insights/ViewModels/RegionsTimingTrack.h
new file mode 100644
--- /dev/null
+++ b/Engine/Source/Developer/TraceInsights/Private/Insights/ViewModels/RegionsTimingTrack.h
@@ -0,0 +1,66 @@
+﻿// Copyright Epic Games, Inc. All Rights Reserved.
+#pragma once
+#include "FileActivityTimingTrack.h"
+#include "Insights/ViewModels/TimingEventsTrack.h"
+#include "TraceServices/Model/Regions.h"
+
+
+class FTimingRegionsTrack;
+
+////////////////////////////////////////////////////////////////////////////////////////////////////
+
+class FTimingRegionsViewCommands : public TCommands<FTimingRegionsViewCommands>
+{
+public:
+	FTimingRegionsViewCommands();
+	virtual ~FTimingRegionsViewCommands() override;
+	virtual void RegisterCommands() override;
+
+public:
+	TSharedPtr<FUICommandInfo> ShowHideRegionTrack;
+};
+
+class FTimingRegionsSharedState : public Insights::ITimingViewExtender, public TSharedFromThis<FTimingRegionsSharedState>
+{
+	friend class FTimingRegionsTrack;
+	
+public:
+	explicit FTimingRegionsSharedState(STimingView* InTimingView) : TimingView(InTimingView) {}
+	virtual ~FTimingRegionsSharedState() override = default;
+
+	//////////////////////////////////////////////////
+	// ITimingViewExtender interface
+
+	virtual void OnBeginSession(Insights::ITimingViewSession& InSession) override;
+	virtual void OnEndSession(Insights::ITimingViewSession& InSession) override;
+	virtual void Tick(Insights::ITimingViewSession& InSession, const TraceServices::IAnalysisSession& InAnalysisSession) override;
+	void ShowHideRegionsTrack();
+	bool IsRegionsTrackVisible() const {return bShowHideRegionsTrack;};
+	virtual void ExtendOtherTracksFilterMenu(Insights::ITimingViewSession& InSession, FMenuBuilder& InOutMenuBuilder) override;
+	void BindCommands();
+
+private:
+	STimingView* TimingView;
+	TSharedPtr<FTimingRegionsTrack> TimingRegionsTrack;
+	
+	bool bShowHideRegionsTrack = true;
+};
+
+class FTimingRegionsTrack : public FTimingEventsTrack
+{
+	INSIGHTS_DECLARE_RTTI(FTimingRegionsTrack, FTimingEventsTrack)
+
+public:
+	FTimingRegionsTrack(FTimingRegionsSharedState& InSharedState) : FTimingEventsTrack(TEXT("Timing Regions")), SharedState(InSharedState) {}
+	virtual ~FTimingRegionsTrack() override;
+
+	virtual void BuildContextMenu(FMenuBuilder& MenuBuilder) override;
+	virtual void InitTooltip(FTooltipDrawState& InOutTooltip, const ITimingEvent& InTooltipEvent) const override;
+	virtual const TSharedPtr<const ITimingEvent> SearchEvent(const FTimingEventSearchParameters& InSearchParameters) const override;
+	virtual void BuildDrawState(ITimingEventsTrackDrawStateBuilder& Builder, const ITimingTrackUpdateContext& Context) override;
+protected:
+	bool FindRegionEvent(const FTimingEventSearchParameters& InParameters, TFunctionRef<void(double, double, uint32, const TraceServices::FTimeRegion&)> InFoundPredicate) const;
+private:
+	
+	FTimingRegionsSharedState& SharedState;
+};
diff --git a/Engine/Source/Developer/TraceInsights/Private/Insights/Widgets/STimingView.cpp b/Engine/Source/Developer/TraceInsights/Private/Insights/Widgets/STimingView.cpp
--- a/Engine/Source/Developer/TraceInsights/Private/Insights/Widgets/STimingView.cpp
+++ b/Engine/Source/Developer/TraceInsights/Private/Insights/Widgets/STimingView.cpp
@@ -60,6 +60,7 @@
 #include "Insights/ViewModels/GraphTrack.h"
 #include "Insights/ViewModels/LoadingTimingTrack.h"
 #include "Insights/ViewModels/MarkersTimingTrack.h"
+#include "Insights/ViewModels/RegionsTimingTrack.h"
 #include "Insights/ViewModels/ThreadTimingTrack.h"
 #include "Insights/ViewModels/TimeFilterValueConverter.h"
 #include "Insights/ViewModels/TimeRulerTrack.h"
@@ -75,6 +76,7 @@
 
 #include <limits>
 
+
 ////////////////////////////////////////////////////////////////////////////////////////////////////
 
 #define LOCTEXT_NAMESPACE "STimingView"
@@ -99,6 +101,7 @@ STimingView::STimingView()
 	, ThreadTimingSharedState(MakeShared<FThreadTimingSharedState>(this))
 	, LoadingSharedState(MakeShared<FLoadingSharedState>(this))
 	, FileActivitySharedState(MakeShared<FFileActivitySharedState>(this))
+	, TimingRegionsSharedState(MakeShared<FTimingRegionsSharedState>(this))
 	, TimeRulerTrack(MakeShared<FTimeRulerTrack>())
 	, DefaultTimeMarker(MakeShared<Insights::FTimeMarker>())
 	, MarkersTrack(MakeShared<FMarkersTimingTrack>())
@@ -120,6 +123,7 @@ STimingView::STimingView()
 	IModularFeatures::Get().RegisterModularFeature(Insights::TimingViewExtenderFeatureName, ThreadTimingSharedState.Get());
 	IModularFeatures::Get().RegisterModularFeature(Insights::TimingViewExtenderFeatureName, LoadingSharedState.Get());
 	IModularFeatures::Get().RegisterModularFeature(Insights::TimingViewExtenderFeatureName, FileActivitySharedState.Get());
+	IModularFeatures::Get().RegisterModularFeature(Insights::TimingViewExtenderFeatureName, TimingRegionsSharedState.Get());
 
 	ExtensionOverlay = SNew(SOverlay).Visibility(EVisibility::SelfHitTestInvisible);
 }
@@ -140,6 +144,7 @@ STimingView::~STimingView()
 		Extender->OnEndSession(*this);
 	}
 
+	IModularFeatures::Get().UnregisterModularFeature(Insights::TimingViewExtenderFeatureName, TimingRegionsSharedState.Get());
 	IModularFeatures::Get().UnregisterModularFeature(Insights::TimingViewExtenderFeatureName, FileActivitySharedState.Get());
 	IModularFeatures::Get().UnregisterModularFeature(Insights::TimingViewExtenderFeatureName, LoadingSharedState.Get());
 	IModularFeatures::Get().UnregisterModularFeature(Insights::TimingViewExtenderFeatureName, ThreadTimingSharedState.Get());
@@ -3114,6 +3119,7 @@ void STimingView::BindCommands()
 	ThreadTimingSharedState->BindCommands();
 	LoadingSharedState->BindCommands();
 	FileActivitySharedState->BindCommands();
+	TimingRegionsSharedState->BindCommands();
 }
 
 ////////////////////////////////////////////////////////////////////////////////////////////////////
diff --git a/Engine/Source/Developer/TraceInsights/Private/Insights/Widgets/STimingView.h b/Engine/Source/Developer/TraceInsights/Private/Insights/Widgets/STimingView.h
--- a/Engine/Source/Developer/TraceInsights/Private/Insights/Widgets/STimingView.h
+++ b/Engine/Source/Developer/TraceInsights/Private/Insights/Widgets/STimingView.h
@@ -24,6 +24,7 @@
 #include "Insights/ViewModels/TimingViewDrawHelper.h"
 #include "Insights/ViewModels/TooltipDrawState.h"
 
+class FTimingRegionsSharedState;
 class FFileActivitySharedState;
 class FFrameSharedState;
 class FLoadingSharedState;
@@ -544,6 +545,9 @@ protected:
 	// Shared state for File Activity (I/O) tracks
 	TSharedPtr<FFileActivitySharedState> FileActivitySharedState;
 
+	// Shared state for Regions tracks
+	TSharedPtr<FTimingRegionsSharedState> TimingRegionsSharedState;
+	
 	////////////////////////////////////////////////////////////
 
 	/** The time ruler track. It includes the custom time markers (ones user can drag with mouse). */
