[Insights]
  - Counters: Fixed "Export Values/Operations" incorrectly saving all values/ops for selected counter even if a timing selection is made in the Timing View.
  - Counters: Added name of the selected counter in the default filename when exporting values/ops.
  - Fixed typos in comments.

diff --git a/Engine/Source/Developer/TraceInsights/Private/Insights/Widgets/SStatsView.cpp b/Engine/Source/Developer/TraceInsights/Private/Insights/Widgets/SStatsView.cpp
--- a/Engine/Source/Developer/TraceInsights/Private/Insights/Widgets/SStatsView.cpp
+++ b/Engine/Source/Developer/TraceInsights/Private/Insights/Widgets/SStatsView.cpp
@@ -2395,9 +2395,19 @@ void SStatsView::ContextMenu_ExportValues_Execute() const
 	}
 
 	const uint32 CounterId = SelectedNodes[0]->GetCounterId();
+	FString CounterName = SelectedNodes[0]->GetName().ToString();
 
 	const FString DialogTitle = LOCTEXT("ExportValues_Title", "Export Counter Values").ToString();
-	const FString DefaultFile = TEXT("CounterValues.tsv");
+
+	FString DefaultFile = CounterName.Replace(TEXT("(1/frame)"), TEXT("(1 per frame)"));
+	DefaultFile.ReplaceCharInline(TEXT('.'), TEXT('_'));
+	DefaultFile = FPaths::MakeValidFileName(DefaultFile, TCHAR('_'));
+	if (DefaultFile.IsEmpty() || DefaultFile[DefaultFile.Len() - 1] == TEXT(' '))
+	{
+		DefaultFile += TEXT('_');
+	}
+	DefaultFile = TEXT("CounterValues ") + DefaultFile + TEXT(".tsv");
+
 	FString Filename;
 	if (!OpenSaveTextFileDialog(DialogTitle, DefaultFile, Filename))
 	{
@@ -2406,6 +2416,28 @@ void SStatsView::ContextMenu_ExportValues_Execute() const
 
 	Insights::FTimingExporter Exporter(*Session.Get());
 	Insights::FTimingExporter::FExportCounterParams Params; // default
+
+	////////////////////////////////////////////////////////////////////////////////////////////////////
+	// Limit the time interval for enumeration (if a time range selection is made in Timing view).
+
+	Params.IntervalStartTime = -std::numeric_limits<double>::infinity();
+	Params.IntervalEndTime = +std::numeric_limits<double>::infinity();
+
+	TSharedPtr<STimingProfilerWindow> Wnd = FTimingProfilerManager::Get()->GetProfilerWindow();
+	TSharedPtr<STimingView> TimingView = Wnd.IsValid() ? Wnd->GetTimingView() : nullptr;
+	if (TimingView.IsValid())
+	{
+		const double SelectionStartTime = TimingView->GetSelectionStartTime();
+		const double SelectionEndTime = TimingView->GetSelectionEndTime();
+		if (SelectionStartTime < SelectionEndTime)
+		{
+			Params.IntervalStartTime = SelectionStartTime;
+			Params.IntervalEndTime = SelectionEndTime;
+		}
+	}
+
+	////////////////////////////////////////////////////////////////////////////////////////////////////
+
 	Exporter.ExportCounterAsText(*Filename, CounterId, Params);
 }
 
@@ -2433,9 +2465,19 @@ void SStatsView::ContextMenu_ExportOps_Execute() const
 	}
 
 	const uint32 CounterId = SelectedNodes[0]->GetCounterId();
+	FString CounterName = SelectedNodes[0]->GetName().ToString();
 
 	const FString DialogTitle = LOCTEXT("ExportOps_Title", "Export Counter Ops").ToString();
-	const FString DefaultFile = TEXT("CounterOps.tsv");
+
+	FString DefaultFile = CounterName.Replace(TEXT("(1/frame)"), TEXT("(1 per frame)"));
+	DefaultFile.ReplaceCharInline(TEXT('.'), TEXT('_'));
+	DefaultFile = FPaths::MakeValidFileName(DefaultFile, TCHAR('_'));
+	if (DefaultFile.IsEmpty() || DefaultFile[DefaultFile.Len() - 1] == TEXT(' '))
+	{
+		DefaultFile += TEXT('_');
+	}
+	DefaultFile = TEXT("CounterOps ") + DefaultFile + TEXT(".tsv");
+
 	FString Filename;
 	if (!OpenSaveTextFileDialog(DialogTitle, DefaultFile, Filename))
 	{
@@ -2444,7 +2486,30 @@ void SStatsView::ContextMenu_ExportOps_Execute() const
 
 	Insights::FTimingExporter Exporter(*Session.Get());
 	Insights::FTimingExporter::FExportCounterParams Params; // default
-	Params.bExportOps = true;
+
+	Params.bExportOps = true; // export "operations" instead of "values"
+
+	////////////////////////////////////////////////////////////////////////////////////////////////////
+	// Limit the time interval for enumeration (if a time range selection is made in Timing view).
+
+	Params.IntervalStartTime = -std::numeric_limits<double>::infinity();
+	Params.IntervalEndTime = +std::numeric_limits<double>::infinity();
+
+	TSharedPtr<STimingProfilerWindow> Wnd = FTimingProfilerManager::Get()->GetProfilerWindow();
+	TSharedPtr<STimingView> TimingView = Wnd.IsValid() ? Wnd->GetTimingView() : nullptr;
+	if (TimingView.IsValid())
+	{
+		const double SelectionStartTime = TimingView->GetSelectionStartTime();
+		const double SelectionEndTime = TimingView->GetSelectionEndTime();
+		if (SelectionStartTime < SelectionEndTime)
+		{
+			Params.IntervalStartTime = SelectionStartTime;
+			Params.IntervalEndTime = SelectionEndTime;
+		}
+	}
+
+	////////////////////////////////////////////////////////////////////////////////////////////////////
+
 	Exporter.ExportCounterAsText(*Filename, CounterId, Params);
 }
 
diff --git a/Engine/Source/Developer/TraceInsights/Private/Insights/Widgets/STimersView.cpp b/Engine/Source/Developer/TraceInsights/Private/Insights/Widgets/STimersView.cpp
--- a/Engine/Source/Developer/TraceInsights/Private/Insights/Widgets/STimersView.cpp
+++ b/Engine/Source/Developer/TraceInsights/Private/Insights/Widgets/STimersView.cpp
@@ -2546,7 +2546,7 @@ void STimersView::ContextMenu_ExportTimingEventsSelection_Execute() const
 	// Filter by thread (visible Gpu/Cpu tracks in TimingView).
 
 	// These variables needs to be in the same scope with the call to Exporter.ExportTimingEventsAsText()
-	// (becasue they are referenced in the ThreadFilter lambda function).
+	// (because they are referenced in the ThreadFilter lambda function).
 	TSet<uint32> IncludedThreads;
 	TSet<uint32> ExcludedThreads;
 
@@ -2596,7 +2596,7 @@ void STimersView::ContextMenu_ExportTimingEventsSelection_Execute() const
 	// Filter by timing event (e.g.: by timer, by duration, by depth).
 
 	// This variable needs to be in the same scope with the call to Exporter.ExportTimingEventsAsText()
-	// (becasue it is referenced in the TimingEventFilter lambda function).
+	// (because it is referenced in the TimingEventFilter lambda function).
 	TSet<uint32> IncludedTimers;
 
 	TArray<FTimerNodePtr> SelectedNodes = TreeView->GetSelectedItems();
