[Insights]
  - MemoryInsights: Added alloc event callstack id and free event callstack id in tooltip of the Callstack column.

diff --git a/Engine/Source/Developer/TraceInsights/Private/Insights/MemoryProfiler/ViewModels/MemAllocNode.cpp b/Engine/Source/Developer/TraceInsights/Private/Insights/MemoryProfiler/ViewModels/MemAllocNode.cpp
--- a/Engine/Source/Developer/TraceInsights/Private/Insights/MemoryProfiler/ViewModels/MemAllocNode.cpp
+++ b/Engine/Source/Developer/TraceInsights/Private/Insights/MemoryProfiler/ViewModels/MemAllocNode.cpp
@@ -14,9 +14,16 @@ INSIGHTS_IMPLEMENT_RTTI(FMemAllocNode)
 
 ////////////////////////////////////////////////////////////////////////////////////////////////////
 
-uint64 FMemAllocNode::GetCallstackId() const
+uint32 FMemAllocNode::GetCallstackId() const
 {
-	return IsValidMemAlloc() ? uint64(GetMemAllocChecked().GetCallstack()) : 0ull;
+	return IsValidMemAlloc() ? GetMemAllocChecked().GetCallstackId() : 0u;
+}
+
+////////////////////////////////////////////////////////////////////////////////////////////////////
+
+uint32 FMemAllocNode::GetFreeCallstackId() const
+{
+	return IsValidMemAlloc() ? GetMemAllocChecked().GetFreeCallstackId() : 0u;
 }
 
 ////////////////////////////////////////////////////////////////////////////////////////////////////
diff --git a/Engine/Source/Developer/TraceInsights/Private/Insights/MemoryProfiler/ViewModels/MemAllocNode.h b/Engine/Source/Developer/TraceInsights/Private/Insights/MemoryProfiler/ViewModels/MemAllocNode.h
--- a/Engine/Source/Developer/TraceInsights/Private/Insights/MemoryProfiler/ViewModels/MemAllocNode.h
+++ b/Engine/Source/Developer/TraceInsights/Private/Insights/MemoryProfiler/ViewModels/MemAllocNode.h
@@ -61,7 +61,9 @@ public:
 	const FMemoryAlloc* GetMemAlloc() const { return GetMemTableChecked().GetMemAlloc(GetRowIndex()); }
 	const FMemoryAlloc& GetMemAllocChecked() const { return GetMemTableChecked().GetMemAllocChecked(GetRowIndex()); }
 
-	uint64 GetCallstackId() const;
+	uint32 GetCallstackId() const;
+	uint32 GetFreeCallstackId() const;
+
 	FText GetFullCallstack() const;
 	FText GetFullCallstackSourceFiles() const;
 	FText GetTopFunction() const;
diff --git a/Engine/Source/Developer/TraceInsights/Private/Insights/MemoryProfiler/ViewModels/MemAllocTable.cpp b/Engine/Source/Developer/TraceInsights/Private/Insights/MemoryProfiler/ViewModels/MemAllocTable.cpp
--- a/Engine/Source/Developer/TraceInsights/Private/Insights/MemoryProfiler/ViewModels/MemAllocTable.cpp
+++ b/Engine/Source/Developer/TraceInsights/Private/Insights/MemoryProfiler/ViewModels/MemAllocTable.cpp
@@ -1,6 +1,7 @@
 // Copyright Epic Games, Inc. All Rights Reserved.
 
 #include "MemAllocTable.h"
+
 #include "Styling/StyleColors.h"
 #include "Widgets/SBoxPanel.h"
 #include "Widgets/SToolTip.h"
@@ -1128,6 +1129,10 @@ void FMemAllocTable::AddDefaultColumns()
 			{
 				const FMemAllocNode& MemAllocNode = static_cast<const FMemAllocNode&>(Node);
 
+				FText HeaderText = FText::Format(LOCTEXT("CallstackHeaderFmt", "CallstackId={0} FreeCallstackId={1}"),
+					FText::AsNumber(MemAllocNode.GetCallstackId(), &FNumberFormattingOptions::DefaultNoGrouping()),
+					FText::AsNumber(MemAllocNode.GetFreeCallstackId(), &FNumberFormattingOptions::DefaultNoGrouping()));
+
 				return SNew(SToolTip)
 					.Visibility(TAttribute<EVisibility>::Create(TAttribute<EVisibility>::FGetter::CreateStatic(&FTableCellValueFormatter::GetTooltipVisibility)))
 					[
@@ -1135,6 +1140,14 @@ void FMemAllocTable::AddDefaultColumns()
 						+ SVerticalBox::Slot()
 						.AutoHeight()
 						.Padding(2.0f)
+						[
+							SNew(STextBlock)
+							.Text(HeaderText)
+							.ColorAndOpacity(FSlateColor(EStyleColor::White25))
+						]
+						+ SVerticalBox::Slot()
+						.AutoHeight()
+						.Padding(2.0f)
 						[
 							SNew(STextBlock)
 							.Text(&MemAllocNode, &FMemAllocNode::GetFullCallstack)
diff --git a/Engine/Source/Developer/TraceInsights/Private/Insights/MemoryProfiler/ViewModels/MemoryAlloc.h b/Engine/Source/Developer/TraceInsights/Private/Insights/MemoryProfiler/ViewModels/MemoryAlloc.h
--- a/Engine/Source/Developer/TraceInsights/Private/Insights/MemoryProfiler/ViewModels/MemoryAlloc.h
+++ b/Engine/Source/Developer/TraceInsights/Private/Insights/MemoryProfiler/ViewModels/MemoryAlloc.h
@@ -33,6 +33,8 @@ public:
 	const TCHAR* GetTag() const { return Tag; }
 	const TCHAR* GetAsset() const { return Asset; }
 	const TCHAR* GetClassName() const { return ClassName; }
+	uint32 GetCallstackId() const { return CallstackId; }
+	uint32 GetFreeCallstackId() const { return FreeCallstackId; }
 	const TraceServices::FCallstack* GetCallstack() const { return Callstack; }
 	const TraceServices::FCallstack* GetFreeCallstack() const { return FreeCallstack; }
 	FText GetFullCallstack() const;
@@ -50,6 +52,8 @@ private:
 	const TCHAR* Tag;
 	const TCHAR* Asset;
 	const TCHAR* ClassName;
+	uint32 CallstackId;
+	uint32 FreeCallstackId;
 	const TraceServices::FCallstack* Callstack;
 	const TraceServices::FCallstack* FreeCallstack;
 	HeapId RootHeap;
diff --git a/Engine/Source/Developer/TraceInsights/Private/Insights/MemoryProfiler/Widgets/SMemAllocTableTreeView.cpp b/Engine/Source/Developer/TraceInsights/Private/Insights/MemoryProfiler/Widgets/SMemAllocTableTreeView.cpp
--- a/Engine/Source/Developer/TraceInsights/Private/Insights/MemoryProfiler/Widgets/SMemAllocTableTreeView.cpp
+++ b/Engine/Source/Developer/TraceInsights/Private/Insights/MemoryProfiler/Widgets/SMemAllocTableTreeView.cpp
@@ -434,6 +434,9 @@ void SMemAllocTableTreeView::UpdateQuery(TraceServices::IAllocationsProvider::EQ
 							});
 					}
 
+					Alloc.CallstackId = Allocation->GetCallstackId();
+					Alloc.FreeCallstackId = Allocation->GetFreeCallstackId();
+
 					if (CallstacksProvider)
 					{
 						Alloc.Callstack = CallstacksProvider->GetCallstack(Allocation->GetCallstackId());
