Unreal Insights:
Shift-click in frame track to select multiple frames
Right-click context menu for events has an option to select time range of the event

diff --git a/Engine/Source/Developer/TraceInsights/Private/Insights/Widgets/SFrameTrack.cpp b/Engine/Source/Developer/TraceInsights/Private/Insights/Widgets/SFrameTrack.cpp
--- a/Engine/Source/Developer/TraceInsights/Private/Insights/Widgets/SFrameTrack.cpp
+++ b/Engine/Source/Developer/TraceInsights/Private/Insights/Widgets/SFrameTrack.cpp
@@ -367,7 +367,7 @@ FFrameTrackSampleRef SFrameTrack::GetSampleAtMousePosition(double X, double Y)
 
 ////////////////////////////////////////////////////////////////////////////////////////////////////
 
-void SFrameTrack::SelectFrameAtMousePosition(double X, double Y)
+void SFrameTrack::SelectFrameAtMousePosition(double X, double Y, bool JoinCurrentSelection)
 {
 	FFrameTrackSampleRef SampleRef = GetSampleAtMousePosition(X, Y);
 	if (!SampleRef.IsValid())
@@ -387,8 +387,16 @@ void SFrameTrack::SelectFrameAtMousePosition(double X, double Y)
 			TSharedPtr<STimingView> TimingView = Window->GetTimingView();
 			if (TimingView.IsValid())
 			{
-				const double StartTime = SampleRef.Sample->LargestFrameStartTime;
-				const double Duration = SampleRef.Sample->LargestFrameDuration;
+				double StartTime = SampleRef.Sample->LargestFrameStartTime;
+				double Duration = SampleRef.Sample->LargestFrameDuration;
+
+				if (JoinCurrentSelection)
+				{
+					double EndTime = StartTime + Duration;
+					StartTime = FMath::Min(StartTime, TimingView->GetSelectionStartTime());
+					EndTime = FMath::Max(EndTime, TimingView->GetSelectionEndTime());
+					Duration = EndTime - StartTime;
+				}
 
 				if (bZoomTimingViewOnFrameSelection)
 				{
@@ -937,7 +945,12 @@ FReply SFrameTrack::OnMouseButtonUp(const FGeometry& MyGeometry, const FPointerE
 			}
 			else if (bIsValidForMouseClick)
 			{
-				SelectFrameAtMousePosition(static_cast<float>(MousePositionOnButtonUp.X), static_cast<float>(MousePositionOnButtonUp.Y));
+				const bool JoinCurrentSelection = MouseEvent.IsShiftDown();
+
+				SelectFrameAtMousePosition(
+					static_cast<float>(MousePositionOnButtonUp.X), 
+					static_cast<float>(MousePositionOnButtonUp.Y), 
+					JoinCurrentSelection);
 			}
 
 			bIsLMB_Pressed = false;
diff --git a/Engine/Source/Developer/TraceInsights/Private/Insights/Widgets/SFrameTrack.h b/Engine/Source/Developer/TraceInsights/Private/Insights/Widgets/SFrameTrack.h
--- a/Engine/Source/Developer/TraceInsights/Private/Insights/Widgets/SFrameTrack.h
+++ b/Engine/Source/Developer/TraceInsights/Private/Insights/Widgets/SFrameTrack.h
@@ -137,7 +137,7 @@ protected:
 	void DrawVerticalAxisGrid(FDrawContext& DrawContext, const FSlateBrush* Brush, const FSlateFontInfo& Font) const;
 
 	FFrameTrackSampleRef GetSampleAtMousePosition(double X, double Y);
-	void SelectFrameAtMousePosition(double X, double Y);
+	void SelectFrameAtMousePosition(double X, double Y, bool JoinCurrentSelection);
 
 	void ShowContextMenu(const FPointerEvent& MouseEvent);
 
diff --git a/Engine/Source/Developer/TraceInsights/Private/Insights/Widgets/STimingView.cpp b/Engine/Source/Developer/TraceInsights/Private/Insights/Widgets/STimingView.cpp
--- a/Engine/Source/Developer/TraceInsights/Private/Insights/Widgets/STimingView.cpp
+++ b/Engine/Source/Developer/TraceInsights/Private/Insights/Widgets/STimingView.cpp
@@ -2881,6 +2881,25 @@ void STimingView::ShowContextMenu(const FPointerEvent& MouseEvent)
 			TAttribute<FText>(),
 			TAttribute<FText>(),
 			FSlateIcon(FInsightsStyle::GetStyleSetName(), "Icons.Find"));
+
+		if (HoveredEvent)
+		{
+			double RangeStart = HoveredEvent->GetStartTime();
+			double RangeDuration = HoveredEvent->GetDuration();
+			
+			MenuBuilder.AddMenuEntry(
+				FText::FromString(TEXT("Select Time Range of Event")),
+				FText(),
+				FSlateIcon(),
+				FUIAction(FExecuteAction::CreateLambda([this, RangeStart, RangeDuration]()
+					{
+						SelectTimeInterval(RangeStart, RangeDuration);
+					})),
+				NAME_None,
+				EUserInterfaceActionType::Button
+			);
+		}
+
 		bHasAnyActions = true;
 	}
 	MenuBuilder.EndSection();
