[Insights] Add a frame stats mode to the timer panel.

diff --git a/Engine/Source/Developer/TraceInsights/Private/Insights/ViewModels/TimerAggregation.cpp b/Engine/Source/Developer/TraceInsights/Private/Insights/ViewModels/TimerAggregation.cpp
--- a/Engine/Source/Developer/TraceInsights/Private/Insights/ViewModels/TimerAggregation.cpp
+++ b/Engine/Source/Developer/TraceInsights/Private/Insights/ViewModels/TimerAggregation.cpp
@@ -17,12 +17,13 @@ namespace Insights
 class FTimerAggregationWorker : public IStatsAggregationWorker
 {
 public:
-	FTimerAggregationWorker(TSharedPtr<const TraceServices::IAnalysisSession> InSession, double InStartTime, double InEndTime, const TSet<uint32>& InCpuThreads, bool bInIncludeGpuThread)
+	FTimerAggregationWorker(TSharedPtr<const TraceServices::IAnalysisSession> InSession, double InStartTime, double InEndTime, const TSet<uint32>& InCpuThreads, bool bInIncludeGpuThread, ETraceFrameType InFrameType)
 		: Session(InSession)
 		, StartTime(InStartTime)
 		, EndTime(InEndTime)
 		, CpuThreads(InCpuThreads)
 		, bIncludeGpuThread(bInIncludeGpuThread)
+		, FrameType(InFrameType)
 		, ResultTable()
 	{
 	}
@@ -40,6 +41,7 @@ private:
 	double EndTime;
 	TSet<uint32> CpuThreads;
 	bool bIncludeGpuThread;
+	ETraceFrameType FrameType;
 	TUniquePtr<TraceServices::ITable<TraceServices::FTimingProfilerAggregatedStats>> ResultTable;
 };
 
@@ -60,7 +62,7 @@ void FTimerAggregationWorker::DoWork()
 			return CpuThreads.Contains(ThreadId);
 		};
 
-		ResultTable.Reset(TimingProfilerProvider.CreateAggregation(StartTime, EndTime, CpuThreadFilter, bIncludeGpuThread));
+		ResultTable.Reset(TimingProfilerProvider.CreateAggregation(StartTime, EndTime, CpuThreadFilter, bIncludeGpuThread, FrameType));
 	}
 }
 
@@ -89,7 +91,7 @@ IStatsAggregationWorker* FTimerAggregator::CreateWorker(TSharedPtr<const TraceSe
 		}
 	}
 
-	return new FTimerAggregationWorker(InSession, GetIntervalStartTime(), GetIntervalEndTime(), CpuThreads, bIsGpuTrackVisible);
+	return new FTimerAggregationWorker(InSession, GetIntervalStartTime(), GetIntervalEndTime(), CpuThreads, bIsGpuTrackVisible, FrameType);
 }
 
 ////////////////////////////////////////////////////////////////////////////////////////////////////
diff --git a/Engine/Source/Developer/TraceInsights/Private/Insights/ViewModels/TimerAggregation.h b/Engine/Source/Developer/TraceInsights/Private/Insights/ViewModels/TimerAggregation.h
--- a/Engine/Source/Developer/TraceInsights/Private/Insights/ViewModels/TimerAggregation.h
+++ b/Engine/Source/Developer/TraceInsights/Private/Insights/ViewModels/TimerAggregation.h
@@ -24,8 +24,14 @@ public:
 	TraceServices::ITable<TraceServices::FTimingProfilerAggregatedStats>* GetResultTable() const;
 	void ResetResults();
 
+	ETraceFrameType GetFrameType() { return FrameType; }
+	void SetFrameType(ETraceFrameType InFrameType) { FrameType = InFrameType; }
+
 protected:
 	virtual IStatsAggregationWorker* CreateWorker(TSharedPtr<const TraceServices::IAnalysisSession> InSession) override;
+
+private:
+	ETraceFrameType FrameType = ETraceFrameType::TraceFrameType_Count;
 };
 
 ////////////////////////////////////////////////////////////////////////////////////////////////////
diff --git a/Engine/Source/Developer/TraceInsights/Private/Insights/ViewModels/TimingGraphTrack.cpp b/Engine/Source/Developer/TraceInsights/Private/Insights/ViewModels/TimingGraphTrack.cpp
--- a/Engine/Source/Developer/TraceInsights/Private/Insights/ViewModels/TimingGraphTrack.cpp
+++ b/Engine/Source/Developer/TraceInsights/Private/Insights/ViewModels/TimingGraphTrack.cpp
@@ -480,7 +480,7 @@ void FTimingGraphTrack::UpdateFrameStatsTimerSeries(FTimingGraphSeries& Series,
 				return;
 			}
 
-			FramesProvider.EnumerateFrames(ETraceFrameType::TraceFrameType_Game, 0, FrameCount, [&Series](const TraceServices::FFrame& Frame)
+			FramesProvider.EnumerateFrames(ETraceFrameType::TraceFrameType_Game, (uint64) 0, (uint64) FrameCount, [&Series](const TraceServices::FFrame& Frame)
 				{
 					FTimingGraphSeries::FAtomicTimingEvent Event;
 					Event.FrameStartTime = Frame.StartTime;
diff --git a/Engine/Source/Developer/TraceInsights/Private/Insights/Widgets/STimersView.cpp b/Engine/Source/Developer/TraceInsights/Private/Insights/Widgets/STimersView.cpp
--- a/Engine/Source/Developer/TraceInsights/Private/Insights/Widgets/STimersView.cpp
+++ b/Engine/Source/Developer/TraceInsights/Private/Insights/Widgets/STimersView.cpp
@@ -297,6 +297,43 @@ void STimersView::Construct(const FArguments& InArgs)
 						]
 					]
 				]
+
+				+ SHorizontalBox::Slot()
+				.FillWidth(1.0f)
+				.Padding(0.0f, 0.0f, 4.0f, 0.0f)
+				.VAlign(VAlign_Center)
+				.HAlign(HAlign_Right)
+				[
+					SNew(SHorizontalBox)
+
+					+ SHorizontalBox::Slot()
+					.AutoWidth()
+					.Padding(0.0f, 0.0f, 4.0f, 0.0f)
+					.VAlign(VAlign_Center)
+					[
+						SNew(STextBlock)
+						.Text(LOCTEXT("FrameType", "Frame Type"))
+					]
+
+					+ SHorizontalBox::Slot()
+					.AutoWidth()
+					.VAlign(VAlign_Center)
+					[
+						SNew(SBox)
+						.MinDesiredWidth(128.0f)
+						[
+							SAssignNew(FrameTypeComboBox, SComboBox<TSharedPtr<ETraceFrameType>>)
+							.ToolTipText(this, &STimersView::FrameType_GetSelectedTooltipText)
+							.OptionsSource(&FrameTypeOptionsSource)
+							.OnSelectionChanged(this, &STimersView::FrameType_OnSelectionChanged)
+							.OnGenerateWidget(this, &STimersView::FrameType_OnGenerateWidget)
+							[
+								SNew(STextBlock)
+								.Text(this, &STimersView::FrameType_GetSelectedText)
+							]
+						]
+					]
+				]
 			]
 		]
 
@@ -383,6 +420,7 @@ void STimersView::Construct(const FArguments& InArgs)
 	Filters->Add(TextFilter);
 
 	CreateGroupByOptionsSources();
+	CreateFrameTypeOptionsSources();
 	CreateSortings();
 
 	InitCommandList();
@@ -1537,6 +1575,113 @@ FText STimersView::GroupBy_GetSelectedTooltipText() const
 	return TimerNodeGroupingHelper::ToDescription(GroupingMode);
 }
 
+
+////////////////////////////////////////////////////////////////////////////////////////////////////
+// Frame Type
+////////////////////////////////////////////////////////////////////////////////////////////////////
+
+void STimersView::CreateFrameTypeOptionsSources()
+{
+	FrameTypeOptionsSource.Reset(3);
+
+	FrameTypeOptionsSource.Add(MakeShared<ETraceFrameType>(ETraceFrameType::TraceFrameType_Count));
+	FrameTypeOptionsSource.Add(MakeShared<ETraceFrameType>(ETraceFrameType::TraceFrameType_Game));
+	FrameTypeOptionsSource.Add(MakeShared<ETraceFrameType>(ETraceFrameType::TraceFrameType_Rendering));
+
+	FrameTypeComboBox->SetSelectedItem(FrameTypeOptionsSource[0]);
+
+	FrameTypeComboBox->RefreshOptions();
+}
+
+////////////////////////////////////////////////////////////////////////////////////////////////////
+
+void STimersView::FrameType_OnSelectionChanged(TSharedPtr<ETraceFrameType> NewFrameType, ESelectInfo::Type SelectInfo)
+{
+	if (*NewFrameType == FrameType)
+	{
+		return;
+	}
+
+	FrameType = *NewFrameType;
+
+	TSharedPtr<STimingProfilerWindow> Window = FTimingProfilerManager::Get()->GetProfilerWindow();
+	if (Window.IsValid())
+	{
+		TSharedPtr<STimingView> TimingView = Window->GetTimingView();
+		if (TimingView.IsValid())
+		{
+			TimingView->SetFrameTypeToSnapTo(FrameType);
+		}
+	}
+
+	Aggregator->SetFrameType(FrameType);
+	Aggregator->Cancel();
+	Aggregator->Start();
+}
+
+////////////////////////////////////////////////////////////////////////////////////////////////////
+
+TSharedRef<SWidget> STimersView::FrameType_OnGenerateWidget(TSharedPtr<ETraceFrameType> InFrameTypeMode) const
+{
+	return SNew(STextBlock)
+	.Text(FrameType_GetText(*InFrameTypeMode))
+	.ToolTipText(FrameType_GetTooltipText(*InFrameTypeMode));
+}
+
+////////////////////////////////////////////////////////////////////////////////////////////////////
+
+FText STimersView::FrameType_GetSelectedText() const
+{
+	return FrameType_GetText(FrameType);
+}
+
+////////////////////////////////////////////////////////////////////////////////////////////////////
+
+FText STimersView::FrameType_GetText(ETraceFrameType InFrameType) const
+{
+	switch (InFrameType)
+	{
+	case ETraceFrameType::TraceFrameType_Count:
+		return LOCTEXT("FrameTypeNone", "None");
+		break;
+	case ETraceFrameType::TraceFrameType_Game:
+		return LOCTEXT("FrameTypeGame", "Game");
+		break;
+	case ETraceFrameType::TraceFrameType_Rendering:
+		return LOCTEXT("FrameTypeRendering", "Rendering");
+		break;
+	}
+
+	return FText();
+}
+
+////////////////////////////////////////////////////////////////////////////////////////////////////
+
+FText STimersView::FrameType_GetSelectedTooltipText() const
+{
+	return FrameType_GetTooltipText(FrameType);
+}
+
+////////////////////////////////////////////////////////////////////////////////////////////////////
+
+FText STimersView::FrameType_GetTooltipText(ETraceFrameType InFrameType) const
+{
+	switch (InFrameType)
+	{
+	case ETraceFrameType::TraceFrameType_Count:
+		return LOCTEXT("FrameTypeNoneTooltip", "Do not calculate the timer stats per frame");
+		break;
+	case ETraceFrameType::TraceFrameType_Game:
+		return LOCTEXT("FrameTypeGameTooltip", "Calculate the timer stats per game frame");
+		break;
+	case ETraceFrameType::TraceFrameType_Rendering:
+		return LOCTEXT("FrameTypeRenderingTooltip", "Calculate the timer stats per rendering frame");
+		break;
+	}
+
+	return FText();
+}
+ 
 ////////////////////////////////////////////////////////////////////////////////////////////////////
 // Sorting
 ////////////////////////////////////////////////////////////////////////////////////////////////////
@@ -2104,9 +2249,12 @@ void STimersView::ResetStats()
 
 void STimersView::UpdateStats(double StartTime, double EndTime)
 {
-	Aggregator->Cancel();
-	Aggregator->SetTimeInterval(StartTime, EndTime);
-	Aggregator->Start();
+	if (Aggregator->GetIntervalStartTime() != StartTime || Aggregator->GetIntervalEndTime() != EndTime)
+	{
+		Aggregator->Cancel();
+		Aggregator->SetTimeInterval(StartTime, EndTime);
+		Aggregator->Start();
+	}
 }
 
 ////////////////////////////////////////////////////////////////////////////////////////////////////
diff --git a/Engine/Source/Developer/TraceInsights/Private/Insights/Widgets/STimersView.h b/Engine/Source/Developer/TraceInsights/Private/Insights/Widgets/STimersView.h
--- a/Engine/Source/Developer/TraceInsights/Private/Insights/Widgets/STimersView.h
+++ b/Engine/Source/Developer/TraceInsights/Private/Insights/Widgets/STimersView.h
@@ -216,6 +216,23 @@ private:
 
 	FText GroupBy_GetSelectedTooltipText() const;
 
+	////////////////////////////////////////////////////////////////////////////////////////////////////
+	// Frame Type
+
+	void CreateFrameTypeOptionsSources();
+
+	void FrameType_OnSelectionChanged(TSharedPtr<ETraceFrameType> NewGroupingMode, ESelectInfo::Type SelectInfo);
+
+	TSharedRef<SWidget> FrameType_OnGenerateWidget(TSharedPtr<ETraceFrameType> InGroupingMode) const;
+
+	FText FrameType_GetSelectedText() const;
+	
+	FText FrameType_GetText(ETraceFrameType InFrameType) const;
+
+	FText FrameType_GetSelectedTooltipText() const;
+
+	FText FrameType_GetTooltipText(ETraceFrameType InFrameType) const;
+
 	////////////////////////////////////////////////////////////////////////////////////////////////////
 	// Sorting
 
@@ -387,6 +404,15 @@ private:
 	/** How we group the timers? */
 	ETimerGroupingMode GroupingMode;
 
+	//////////////////////////////////////////////////
+	// Frame Type
+
+	TArray<TSharedPtr<ETraceFrameType>> FrameTypeOptionsSource;
+
+	TSharedPtr<SComboBox<TSharedPtr<ETraceFrameType>>> FrameTypeComboBox;
+
+	ETraceFrameType FrameType;
+
 	//////////////////////////////////////////////////
 	// Sorting
 
diff --git a/Engine/Source/Developer/TraceInsights/Private/Insights/Widgets/STimingView.cpp b/Engine/Source/Developer/TraceInsights/Private/Insights/Widgets/STimingView.cpp
--- a/Engine/Source/Developer/TraceInsights/Private/Insights/Widgets/STimingView.cpp
+++ b/Engine/Source/Developer/TraceInsights/Private/Insights/Widgets/STimingView.cpp
@@ -470,6 +470,7 @@ void STimingView::Reset(bool bIsFirstReset)
 	Tooltip.Reset();
 
 	LastSelectionType = ESelectionType::None;
+	FrameTypeToSnapTo = ETraceFrameType::TraceFrameType_Count;
 
 	//ThisGeometry
 
@@ -2142,7 +2143,7 @@ FReply STimingView::OnMouseButtonUp(const FGeometry& MyGeometry, const FPointerE
 			}
 			else if (bIsSelecting)
 			{
-				RaiseSelectionChanged();
+				SelectTimeInterval(SelectionStartTime, SelectionEndTime - SelectionStartTime);
 				bIsSelecting = false;
 			}
 			else if (TimeRulerTrack->IsScrubbing())
@@ -3392,12 +3393,91 @@ void STimingView::SelectTimeInterval(double IntervalStartTime, double IntervalDu
 {
 	SelectionStartTime = IntervalStartTime;
 	SelectionEndTime = IntervalStartTime + IntervalDuration;
+
+	if (FrameTypeToSnapTo != ETraceFrameType::TraceFrameType_Count)
+	{
+		SnapToFrameBound(SelectionStartTime, SelectionEndTime);
+	}
+
 	LastSelectionType = ESelectionType::TimeRange;
 	RaiseSelectionChanged();
 }
 
 ////////////////////////////////////////////////////////////////////////////////////////////////////
 
+void STimingView::SnapToFrameBound(double& StartTime, double& EndTime)
+{
+	TSharedPtr<const TraceServices::IAnalysisSession> Session = FInsightsManager::Get()->GetSession();
+
+	if (!Session.IsValid())
+	{
+		return;
+	}
+
+	TraceServices::FAnalysisSessionReadScope SessionReadScope(*Session.Get());
+	const TraceServices::IFrameProvider& FramesProvider = TraceServices::ReadFrameProvider(*Session.Get());
+
+	uint32 FrameNum = FramesProvider.GetFrameNumberForTimestamp(FrameTypeToSnapTo, StartTime);
+	const TraceServices::FFrame* StartFrame = FramesProvider.GetFrame(FrameTypeToSnapTo, FrameNum);
+
+	if (StartFrame == nullptr)
+	{
+		return;
+	}
+
+	if (StartFrame->EndTime < StartTime)
+	{
+		if (FrameNum + 1 < FramesProvider.GetFrameCount(FrameTypeToSnapTo))
+		{
+			StartFrame = FramesProvider.GetFrame(FrameTypeToSnapTo, FrameNum + 1);
+		}
+	}
+
+	FrameNum = FramesProvider.GetFrameNumberForTimestamp(FrameTypeToSnapTo, EndTime);
+	const TraceServices::FFrame* EndFrame = FramesProvider.GetFrame(FrameTypeToSnapTo, FrameNum);
+
+	if (EndFrame == nullptr)
+	{
+		EndFrame = FramesProvider.GetFrame(FrameTypeToSnapTo, FrameNum - 1);
+		if (EndFrame == nullptr)
+		{
+			return;
+		}
+	}
+
+	// If the interval is before the first frame or after the last frame.
+	if (StartFrame->StartTime > EndTime || EndFrame->EndTime < StartTime)
+	{
+		return;
+	}
+
+	// If the interval is between frames.
+	if (EndFrame->Index < StartFrame->Index)
+	{
+		return;
+	}
+
+	StartTime = StartFrame->StartTime;
+	EndTime = FMath::Min(EndFrame->EndTime, Session->GetDurationSeconds());
+}
+
+////////////////////////////////////////////////////////////////////////////////////////////////////
+
+void STimingView::SetFrameTypeToSnapTo(ETraceFrameType InFrameType)
+{
+	if (FrameTypeToSnapTo == InFrameType)
+	{
+		return;
+	}
+
+	FrameTypeToSnapTo = InFrameType;
+
+	SelectTimeInterval(SelectionStartTime, SelectionEndTime - SelectionStartTime);
+}
+
+////////////////////////////////////////////////////////////////////////////////////////////////////
+
+
 void STimingView::RaiseSelectionChanging()
 {
 	OnSelectionChangedDelegate.Broadcast(Insights::ETimeChangedFlags::Interactive, SelectionStartTime, SelectionEndTime);
diff --git a/Engine/Source/Developer/TraceInsights/Private/Insights/Widgets/STimingView.h b/Engine/Source/Developer/TraceInsights/Private/Insights/Widgets/STimingView.h
--- a/Engine/Source/Developer/TraceInsights/Private/Insights/Widgets/STimingView.h
+++ b/Engine/Source/Developer/TraceInsights/Private/Insights/Widgets/STimingView.h
@@ -345,8 +345,12 @@ public:
 	void ZoomOnTimeInterval(double IntervalStartTime, double IntervalDuration);
 	void BringIntoView(double StartTime, double EndTime);
 	void SelectTimeInterval(double IntervalStartTime, double IntervalDuration);
+	void SnapToFrameBound(double& IntervalStartTime, double& IntervalDuration);
 	void SelectToTimeMarker(double InTimeMarker);
 
+	void SetFrameTypeToSnapTo(ETraceFrameType InFrameType);
+	ETraceFrameType GetFrameTypeToSnapTo() { return FrameTypeToSnapTo; }
+
 	//bool AreTimeMarkersVisible() { return MarkersTrack->IsVisible(); }
 	void SetTimeMarkersVisible(bool bOnOff);
 	//bool IsDrawOnlyBookmarksEnabled() { return MarkersTrack->IsBookmarksTrack(); }
@@ -675,6 +679,8 @@ protected:
 	};
 	ESelectionType LastSelectionType;
 
+	ETraceFrameType FrameTypeToSnapTo;
+
 	/** Throttle flag, allowing tracks to control whether Slate throttle should take place */
 	bool bPreventThrottling;
 
diff --git a/Engine/Source/Developer/TraceServices/Private/Common/TimelineStatistics.h b/Engine/Source/Developer/TraceServices/Private/Common/TimelineStatistics.h
--- a/Engine/Source/Developer/TraceServices/Private/Common/TimelineStatistics.h
+++ b/Engine/Source/Developer/TraceServices/Private/Common/TimelineStatistics.h
@@ -25,6 +25,12 @@ struct FAggregatedTimingStats
 	double MedianExclusiveTime = 0.0;
 };
 
+struct FFrameData
+{
+	double StartTime;
+	double EndTime;
+};
+
 class FTimelineStatistics
 {
 public:
@@ -58,6 +64,68 @@ public:
 		for (auto& KV : InternalResult)
 		{
 			PostProcessTimerStats(KV.Value, ComputeMedian);
+			Result.Add(KV.Key, KV.Value.Stats);
+		}
+	};
+
+	template<typename TimelineType, typename BucketMappingFunc, typename BucketKeyType>
+	static void CreateFrameStatsAggregation(const TArray<const TimelineType*>& Timelines, BucketMappingFunc BucketMapper, TArray<FFrameData> Frames, TMap<BucketKeyType, FAggregatedTimingStats>& Result)
+	{
+		TMap<BucketKeyType, FInternalAggregationEntry> FrameResult;
+		TMap<BucketKeyType, FInternalAggregationEntry> GlobalResult;
+		double StartTime = Frames[0].StartTime;
+		double EndTime = Frames[Frames.Num() - 1].EndTime;
+
+		// Gather the keys for all timers so we have a stable map (no insertions necessary) for the next step.
+		// This is to guarantee that GlobalResult and FrameResult have the same iteration order.
+		for (const TimelineType* Timeline : Timelines)
+		{
+			GatherKeysFromTimeline(Timeline, BucketMapper, StartTime, EndTime, FrameResult);
+		}
+
+		GlobalResult = FrameResult;
+
+		for (int32 Index = 0; Index < Frames.Num(); ++Index)
+		{
+			// Compute instance count and total/min/max inclusive/exclusive times for each timer.
+			for (const TimelineType* Timeline : Timelines)
+			{
+				ProcessTimeline(Timeline, BucketMapper, UpdateTotalMinMaxTimerStats, Frames[Index].StartTime, Frames[Index].EndTime, FrameResult);
+			}
+
+			typename TMap<BucketKeyType, FInternalAggregationEntry>::TIterator FrameResultIterator(FrameResult);
+			typename TMap<BucketKeyType, FInternalAggregationEntry>::TIterator ResultIterator(GlobalResult);
+			while (ResultIterator)
+			{
+				FAggregatedTimingStats& FrameStats = FrameResultIterator->Value.Stats;
+				FAggregatedTimingStats& ResultStats = ResultIterator->Value.Stats;
+
+				// TODO: Add Instance Count per frame column
+				// ResultStats.InstanceCount = (ResultStats.InstanceCount * Index + FrameStats.InstanceCount) / (Index + 1);
+
+				ResultStats.InstanceCount += FrameStats.InstanceCount;
+
+				ResultStats.TotalInclusiveTime += FrameStats.TotalInclusiveTime;
+				ResultStats.MinInclusiveTime = FMath::Min(ResultStats.MinInclusiveTime, FrameStats.TotalInclusiveTime);
+				ResultStats.MaxInclusiveTime = FMath::Max(ResultStats.MaxInclusiveTime, FrameStats.TotalInclusiveTime);
+
+				ResultStats.TotalExclusiveTime += FrameStats.TotalExclusiveTime;
+				ResultStats.MinExclusiveTime = FMath::Min(ResultStats.MinExclusiveTime, FrameStats.TotalExclusiveTime);
+				ResultStats.MaxExclusiveTime = FMath::Max(ResultStats.MaxExclusiveTime, FrameStats.TotalExclusiveTime);
+
+				// Reset the per frame stats.
+				FrameStats = FAggregatedTimingStats();
+
+				++FrameResultIterator;
+				++ResultIterator;
+			}
+		}
+
+		for (auto& KV : GlobalResult)
+		{
+			KV.Value.Stats.AverageInclusiveTime = KV.Value.Stats.TotalInclusiveTime / Frames.Num();
+			KV.Value.Stats.AverageExclusiveTime = KV.Value.Stats.TotalExclusiveTime / Frames.Num();
+
 			Result.Add(KV.Key, KV.Value.Stats);
 		}
 	}
@@ -285,6 +353,25 @@ private:
 			}
 		}
 	}
+
+	template<typename TimelineType, typename BucketMappingFunc, typename BucketKeyType>
+	static void GatherKeysFromTimeline(const TimelineType* Timeline, BucketMappingFunc BucketMapper, double IntervalStart, double IntervalEnd, TMap<BucketKeyType, FInternalAggregationEntry>& InternalResult)
+	{
+		Timeline->EnumerateEvents(IntervalStart, IntervalEnd, [BucketMapper, IntervalStart, IntervalEnd, &InternalResult](bool IsEnter, double Time, const typename TimelineType::EventType& Event)
+			{
+				BucketKeyType BucketKey = BucketMapper(Event);
+
+				if (IsEnter)
+				{
+					if (!InternalResult.Contains(BucketKey))
+					{
+						InternalResult.Add(BucketKey);
+					}
+				}
+
+				return EEventEnumerate::Continue;
+			});
+	}
 };
 
 } // namespace TraceServices
diff --git a/Engine/Source/Developer/TraceServices/Private/Model/Frames.cpp b/Engine/Source/Developer/TraceServices/Private/Model/Frames.cpp
--- a/Engine/Source/Developer/TraceServices/Private/Model/Frames.cpp
+++ b/Engine/Source/Developer/TraceServices/Private/Model/Frames.cpp
@@ -29,21 +29,43 @@ uint64 FFrameProvider::GetFrameCount(ETraceFrameType FrameType) const
 	return Frames[FrameType].Num();
 }
 
-void FFrameProvider::EnumerateFrames(ETraceFrameType FrameType, uint64 Start, uint64 End, TFunctionRef<void(const FFrame&)> Callback) const
+void FFrameProvider::EnumerateFrames(ETraceFrameType FrameType, uint64 StartIndex, uint64 EndIndex, TFunctionRef<void(const FFrame&)> Callback) const
 {
 	Session.ReadAccessCheck();
 
-	End = FMath::Min(End, Frames[FrameType].Num());
-	if (Start >= End)
+	EndIndex = FMath::Min(EndIndex, Frames[FrameType].Num());
+	if (StartIndex >= EndIndex)
 	{
 		return;
 	}
-	for (auto Iterator = Frames[FrameType].GetIteratorFromItem(Start); Iterator && Iterator->Index < End; ++Iterator)
+	for (auto Iterator = Frames[FrameType].GetIteratorFromItem(StartIndex); Iterator && Iterator->Index < EndIndex; ++Iterator)
 	{
 		Callback(*Iterator);
 	}
 }
 
+void FFrameProvider::EnumerateFrames(ETraceFrameType FrameType, double StartTime, double EndTime, TFunctionRef<void(const FFrame&)> Callback) const
+{
+	Session.ReadAccessCheck();
+
+	int64 StartIndex = Algo::LowerBound(FrameStartTimes[FrameType], StartTime);
+	if (StartIndex > 0 && Frames[FrameType][StartIndex - 1].EndTime > StartTime)
+	{
+		--StartIndex;
+	}
+
+	for (uint64 Index = StartIndex; Index < Frames[FrameType].Num(); ++Index)
+	{
+		const FFrame& CurrentFrame = Frames[FrameType][Index];
+		if (CurrentFrame.StartTime > EndTime)
+		{
+			break;
+		}
+
+		Callback(CurrentFrame);
+	}
+}
+
 bool FFrameProvider::GetFrameFromTime(ETraceFrameType FrameType, double Time, FFrame& OutFrame) const
 {
 	Session.ReadAccessCheck();
diff --git a/Engine/Source/Developer/TraceServices/Private/Model/FramesPrivate.h b/Engine/Source/Developer/TraceServices/Private/Model/FramesPrivate.h
--- a/Engine/Source/Developer/TraceServices/Private/Model/FramesPrivate.h
+++ b/Engine/Source/Developer/TraceServices/Private/Model/FramesPrivate.h
@@ -24,7 +24,8 @@ public:
 	// Read operations
 
 	virtual uint64 GetFrameCount(ETraceFrameType FrameType) const override;
-	virtual void EnumerateFrames(ETraceFrameType FrameType, uint64 Start, uint64 End, TFunctionRef<void(const FFrame&)> Callback) const override;
+	virtual void EnumerateFrames(ETraceFrameType FrameType, uint64 StartIndex, uint64 EndIndex, TFunctionRef<void(const FFrame&)> Callback) const override;
+	virtual void EnumerateFrames(ETraceFrameType FrameType, double StartTime, double EndTime, TFunctionRef<void(const FFrame&)> Callback) const override;
 	virtual const TArray64<double>& GetFrameStartTimes(ETraceFrameType FrameType) const override { return FrameStartTimes[FrameType]; }
 	virtual bool GetFrameFromTime(ETraceFrameType FrameType, double Time, FFrame& OutFrame) const override;
 	virtual const FFrame* GetFrame(ETraceFrameType FrameType, uint64 Index) const override;
diff --git a/Engine/Source/Developer/TraceServices/Private/Model/TimingProfiler.cpp b/Engine/Source/Developer/TraceServices/Private/Model/TimingProfiler.cpp
--- a/Engine/Source/Developer/TraceServices/Private/Model/TimingProfiler.cpp
+++ b/Engine/Source/Developer/TraceServices/Private/Model/TimingProfiler.cpp
@@ -1,6 +1,8 @@
 // Copyright Epic Games, Inc. All Rights Reserved.
 
 #include "TraceServices/Model/TimingProfiler.h"
+
+#include "TraceServices/Model/Frames.h"
 #include "Model/TimingProfilerPrivate.h"
 #include "AnalysisServicePrivate.h"
 #include "Common/StringStore.h"
@@ -263,7 +265,7 @@ void FTimingProfilerProvider::ReadTimers(TFunctionRef<void(const ITimingProfiler
 	Callback(*this);
 }
 
-ITable<FTimingProfilerAggregatedStats>* FTimingProfilerProvider::CreateAggregation(double IntervalStart, double IntervalEnd, TFunctionRef<bool(uint32)> CpuThreadFilter, bool IncludeGpu) const
+ITable<FTimingProfilerAggregatedStats>* FTimingProfilerProvider::CreateAggregation(double IntervalStart, double IntervalEnd, TFunctionRef<bool(uint32)> CpuThreadFilter, bool IncludeGpu, ETraceFrameType FrameType) const
 {
 	Session.ReadAccessCheck();
 
@@ -287,7 +289,29 @@ ITable<FTimingProfilerAggregatedStats>* FTimingProfilerProvider::CreateAggregati
 	};
 
 	TMap<const FTimingProfilerTimer*, FAggregatedTimingStats> Aggregation;
-	FTimelineStatistics::CreateAggregation(IncludedTimelines, BucketMappingFunc, IntervalStart, IntervalEnd, Aggregation);
+	if (FrameType == ETraceFrameType::TraceFrameType_Count)
+	{
+		FTimelineStatistics::CreateAggregation(IncludedTimelines, BucketMappingFunc, IntervalStart, IntervalEnd, Aggregation);
+	}
+	else
+	{
+		TArray<FFrameData> Frames;
+		const IFrameProvider& FrameProvider = ReadFrameProvider(Session);
+		FrameProvider.EnumerateFrames(FrameType, IntervalStart, IntervalEnd, [&Frames](const FFrame& Frame)
+			{
+				FFrameData NewFrameData;
+				NewFrameData.StartTime = Frame.StartTime;
+				NewFrameData.EndTime = Frame.EndTime;
+
+				Frames.Add(NewFrameData);
+			});
+
+		if (Frames.Num() > 0)
+		{
+			FTimelineStatistics::CreateFrameStatsAggregation(IncludedTimelines, BucketMappingFunc, Frames, Aggregation);
+		}
+	}
+
 	TTable<FTimingProfilerAggregatedStats>* Table = new TTable<FTimingProfilerAggregatedStats>(AggregatedStatsTableLayout);
 	for (const auto& KV : Aggregation)
 	{
diff --git a/Engine/Source/Developer/TraceServices/Private/Model/TimingProfilerPrivate.h b/Engine/Source/Developer/TraceServices/Private/Model/TimingProfilerPrivate.h
--- a/Engine/Source/Developer/TraceServices/Private/Model/TimingProfilerPrivate.h
+++ b/Engine/Source/Developer/TraceServices/Private/Model/TimingProfilerPrivate.h
@@ -35,7 +35,7 @@ public:
 	virtual uint32 GetTimelineCount() const override { return Timelines.Num(); }
 	virtual void EnumerateTimelines(TFunctionRef<void(const Timeline&)> Callback) const override;
 	virtual void ReadTimers(TFunctionRef<void(const ITimingProfilerTimerReader&)> Callback) const override;
-	virtual ITable<FTimingProfilerAggregatedStats>* CreateAggregation(double IntervalStart, double IntervalEnd, TFunctionRef<bool(uint32)> CpuThreadFilter, bool IncludeGpu) const override;
+	virtual ITable<FTimingProfilerAggregatedStats>* CreateAggregation(double IntervalStart, double IntervalEnd, TFunctionRef<bool(uint32)> CpuThreadFilter, bool IncludeGpu, ETraceFrameType FrameType = ETraceFrameType::TraceFrameType_Count) const override;
 	virtual ITimingProfilerButterfly* CreateButterfly(double IntervalStart, double IntervalEnd, TFunctionRef<bool(uint32)> CpuThreadFilter, bool IncludeGpu) const override;
 	virtual const FTimingProfilerTimer* GetTimer(uint32 TimerId) const override;
 	virtual uint32 GetTimerCount() const override;
diff --git a/Engine/Source/Developer/TraceServices/Public/TraceServices/Model/Frames.h b/Engine/Source/Developer/TraceServices/Public/TraceServices/Model/Frames.h
--- a/Engine/Source/Developer/TraceServices/Public/TraceServices/Model/Frames.h
+++ b/Engine/Source/Developer/TraceServices/Public/TraceServices/Model/Frames.h
@@ -43,7 +43,16 @@ public:
 	 * @param Callback  - a callback called for each frame
 	 */
 	virtual void EnumerateFrames(ETraceFrameType FrameType, uint64 Start, uint64 End, TFunctionRef<void(const FFrame&)> Callback) const = 0;
-
+	/**
+	* Enumerate frames whose duration intersects [StartTime, EndTime].
+	* 
+	* @param FrameType	The frame type to enumerate.
+	* @param StartTime	The start timestamp in seconds.
+	* @param EndTime	The end timestamp in seconds.
+	* @param Callback	The callback to be called for each frame.
+	*/
+	virtual void EnumerateFrames(ETraceFrameType FrameType, double StartTime, double EndTime, TFunctionRef<void(const FFrame&)> Callback) const = 0;
+	
 	/**
 	 * Gets the array of frame start times, for the specified frame type.
 	 * @param FrameType - type of frames
diff --git a/Engine/Source/Developer/TraceServices/Public/TraceServices/Model/TimingProfiler.h b/Engine/Source/Developer/TraceServices/Public/TraceServices/Model/TimingProfiler.h
--- a/Engine/Source/Developer/TraceServices/Public/TraceServices/Model/TimingProfiler.h
+++ b/Engine/Source/Developer/TraceServices/Public/TraceServices/Model/TimingProfiler.h
@@ -95,7 +95,17 @@ public:
 	virtual uint32 GetTimelineCount() const = 0;
 	virtual void EnumerateTimelines(TFunctionRef<void(const Timeline&)> Callback) const = 0;
 	virtual void ReadTimers(TFunctionRef<void(const ITimingProfilerTimerReader&)> Callback) const = 0;
-	virtual ITable<FTimingProfilerAggregatedStats>* CreateAggregation(double IntervalStart, double IntervalEnd, TFunctionRef<bool(uint32)> CpuThreadFilter, bool IncludeGpu) const = 0;
+
+	/**
+	* Create a table of aggregated stats.
+	*
+	* @param IntervalStart		The start timestamp in seconds.
+	* @param IntervalEnd		The end timestamp in seconds.
+	* @param CpuThreadFilter	A function to filter the CPU threads to aggregate.
+	* @param IncludeGpu			A boolean value to specify if aggregaton should include GPU timelines.
+	* @param FrameType			The type of frame to use for frame stats aggregation. ETraceFrameType::TraceFrameType_Count means no frame aggregation.
+	*/
+	virtual ITable<FTimingProfilerAggregatedStats>* CreateAggregation(double IntervalStart, double IntervalEnd, TFunctionRef<bool(uint32)> CpuThreadFilter, bool IncludeGpu, ETraceFrameType FrameType = ETraceFrameType::TraceFrameType_Count) const = 0;
 	virtual ITimingProfilerButterfly* CreateButterfly(double IntervalStart, double IntervalEnd, TFunctionRef<bool(uint32)> CpuThreadFilter, bool IncludeGpu) const = 0;
 };
 
