[Insights] Timing View: Changed zooming in/out with mouse wheel to auto-pin the hovered track.

diff --git a/Engine/Source/Developer/TraceInsights/Private/Insights/Widgets/STimingView.cpp b/Engine/Source/Developer/TraceInsights/Private/Insights/Widgets/STimingView.cpp
--- a/Engine/Source/Developer/TraceInsights/Private/Insights/Widgets/STimingView.cpp
+++ b/Engine/Source/Developer/TraceInsights/Private/Insights/Widgets/STimingView.cpp
@@ -765,7 +765,7 @@ void STimingView::Tick(const FGeometry& AllottedGeometry, const double InCurrent
 	}
 
 	////////////////////////////////////////////////////////////////////////////////////////////////////
-	// Update Y postion for the visible top docked tracks.
+	// Update Y position for the visible top docked tracks.
 	// Compute the total height of top docked areas.
 
 	int32 NumVisibleTopDockedTracks = 0;
@@ -791,7 +791,7 @@ void STimingView::Tick(const FGeometry& AllottedGeometry, const double InCurrent
 	Viewport.SetTopOffset(TopOffset);
 
 	////////////////////////////////////////////////////////////////////////////////////////////////////
-	// Update Y postion for the visible bottom docked tracks.
+	// Update Y position for the visible bottom docked tracks.
 	// Compute the total height of bottom docked areas.
 
 	float BottomOffset = 0.0f;
@@ -1027,7 +1027,7 @@ void STimingView::Tick(const FGeometry& AllottedGeometry, const double InCurrent
 	Tooltip.Update();
 	if (!MousePosition.IsZero())
 	{
-		Tooltip.SetPosition(MousePosition, 0.0f, Viewport.GetWidth() - 12.0f, Viewport.GetPosY(), Viewport.GetPosY() + Viewport.GetHeight() - 12.0f); // -12.0f is to avoid overlaping the scrollbars
+		Tooltip.SetPosition(MousePosition, 0.0f, Viewport.GetWidth() - 12.0f, Viewport.GetPosY(), Viewport.GetPosY() + Viewport.GetHeight() - 12.0f); // -12.0f is to avoid overlapping the scrollbars
 	}
 
 	////////////////////////////////////////////////////////////////////////////////////////////////////
@@ -1092,7 +1092,7 @@ void STimingView::Tick(const FGeometry& AllottedGeometry, const double InCurrent
 
 void STimingView::UpdatePositionForScrollableTracks()
 {
-	// Update the Y postion for the visible scrollable tracks.
+	// Update the Y position for the visible scrollable tracks.
 	float ScrollableTrackPosY = Viewport.GetPosY() + Viewport.GetTopOffset() - Viewport.GetScrollPosY();
 	for (TSharedPtr<FBaseTimingTrack>& TrackPtr : ScrollableTracks)
 	{
@@ -2517,6 +2517,11 @@ FReply STimingView::OnMouseWheel(const FGeometry& MyGeometry, const FPointerEven
 	}
 	else if (MouseEvent.GetModifierKeys().IsControlDown())
 	{
+		if (HoveredTrack.IsValid())
+		{
+			SelectHoveredTimingTrack();
+		}
+
 		// Scroll horizontally.
 		const double ScrollSpeedX = Viewport.GetDurationForViewportDX(16.0 * 3);
 		const double NewStartTime = Viewport.GetStartTime() - ScrollSpeedX * MouseEvent.GetWheelDelta();
@@ -2524,6 +2529,11 @@ FReply STimingView::OnMouseWheel(const FGeometry& MyGeometry, const FPointerEven
 	}
 	else
 	{
+		if (HoveredTrack.IsValid())
+		{
+			SelectHoveredTimingTrack();
+		}
+
 		// Zoom in/out horizontally.
 		const float Delta = MouseEvent.GetWheelDelta();
 		if (Viewport.RelativeZoomWithFixedX(Delta, static_cast<float>(MousePosition.X)))
