Fix C++20 compile error. Make Insights::FFilterConfigurator's comparison operators const to address ambiguity issues.

diff --git a/Engine/Source/Developer/TraceInsights/Private/Insights/ViewModels/FilterConfigurator.cpp b/Engine/Source/Developer/TraceInsights/Private/Insights/ViewModels/FilterConfigurator.cpp
--- a/Engine/Source/Developer/TraceInsights/Private/Insights/ViewModels/FilterConfigurator.cpp
+++ b/Engine/Source/Developer/TraceInsights/Private/Insights/ViewModels/FilterConfigurator.cpp
@@ -54,7 +54,7 @@ FFilterConfigurator& FFilterConfigurator::operator=(const FFilterConfigurator& O
 
 ////////////////////////////////////////////////////////////////////////////////////////////////////
 
-bool FFilterConfigurator::operator==(const FFilterConfigurator& Other)
+bool FFilterConfigurator::operator==(const FFilterConfigurator& Other) const
 {
 	bool bIsEqual = AvailableFilters.Get() == Other.AvailableFilters.Get();
 	bIsEqual &= RootNode.IsValid() == Other.RootNode.IsValid();
diff --git a/Engine/Source/Developer/TraceInsights/Private/Insights/ViewModels/FilterConfigurator.h b/Engine/Source/Developer/TraceInsights/Private/Insights/ViewModels/FilterConfigurator.h
--- a/Engine/Source/Developer/TraceInsights/Private/Insights/ViewModels/FilterConfigurator.h
+++ b/Engine/Source/Developer/TraceInsights/Private/Insights/ViewModels/FilterConfigurator.h
@@ -20,9 +20,9 @@ public:
 	FFilterConfigurator(const FFilterConfigurator& Other);
 	FFilterConfigurator& operator=(const FFilterConfigurator& Other);
 
-	bool operator==(const FFilterConfigurator& Other);
+	bool operator==(const FFilterConfigurator& Other) const;
 
-	bool operator!=(const FFilterConfigurator& Other) { return !(*this == Other); }
+	bool operator!=(const FFilterConfigurator& Other) const { return !(*this == Other); }
 
 	virtual ~FFilterConfigurator();
 
