[Insights] Frame stats improvements to options name and graph series creation

diff --git a/Engine/Source/Developer/TraceInsights/Private/Insights/ViewModels/TimingGraphTrack.cpp b/Engine/Source/Developer/TraceInsights/Private/Insights/ViewModels/TimingGraphTrack.cpp
--- a/Engine/Source/Developer/TraceInsights/Private/Insights/ViewModels/TimingGraphTrack.cpp
+++ b/Engine/Source/Developer/TraceInsights/Private/Insights/ViewModels/TimingGraphTrack.cpp
@@ -408,19 +408,19 @@ void FTimingGraphTrack::UpdateTimerSeries(FTimingGraphSeries& Series, const FTim
 // Frams Stats Timer Series
 ////////////////////////////////////////////////////////////////////////////////////////////////////
 
-TSharedPtr<FTimingGraphSeries> FTimingGraphTrack::GetFrameStatsTimerSeries(uint32 TimerId)
+TSharedPtr<FTimingGraphSeries> FTimingGraphTrack::GetFrameStatsTimerSeries(uint32 TimerId, ETraceFrameType FrameType)
 {
-	TSharedPtr<FGraphSeries>* Ptr = AllSeries.FindByPredicate([TimerId](const TSharedPtr<FGraphSeries>& Series)
+	TSharedPtr<FGraphSeries>* Ptr = AllSeries.FindByPredicate([TimerId, FrameType](const TSharedPtr<FGraphSeries>& Series)
 		{
 			const TSharedPtr<FTimingGraphSeries> TimingSeries = StaticCastSharedPtr<FTimingGraphSeries>(Series);
-			return TimingSeries->Type == FTimingGraphSeries::ESeriesType::FrameStatsTimer && TimingSeries->TimerId == TimerId;
+			return TimingSeries->Type == FTimingGraphSeries::ESeriesType::FrameStatsTimer && TimingSeries->TimerId == TimerId && TimingSeries->FrameType == FrameType;
 		});
 	return (Ptr != nullptr) ? StaticCastSharedPtr<FTimingGraphSeries>(*Ptr) : nullptr;
 }
 
 ////////////////////////////////////////////////////////////////////////////////////////////////////
 
-TSharedPtr<FTimingGraphSeries> FTimingGraphTrack::AddFrameStatsTimerSeries(uint32 TimerId, FLinearColor Color)
+TSharedPtr<FTimingGraphSeries> FTimingGraphTrack::AddFrameStatsTimerSeries(uint32 TimerId, ETraceFrameType FrameType, FLinearColor Color)
 {
 	TSharedRef<FTimingGraphSeries> Series = MakeShared<FTimingGraphSeries>(FTimingGraphSeries::ESeriesType::FrameStatsTimer);
 
@@ -431,8 +431,7 @@ TSharedPtr<FTimingGraphSeries> FTimingGraphTrack::AddFrameStatsTimerSeries(uint3
 	Series->SetColor(Color, BorderColor);
 
 	Series->TimerId = TimerId;
-	//Series->CpuOrGpu = ;
-	//Series->TimelineIndex = ;
+	Series->FrameType = FrameType;
 
 	// Use shared viewport.
 	Series->SetBaselineY(SharedValueViewport.GetBaselineY());
@@ -447,12 +446,12 @@ TSharedPtr<FTimingGraphSeries> FTimingGraphTrack::AddFrameStatsTimerSeries(uint3
 
 ////////////////////////////////////////////////////////////////////////////////////////////////////
 
-void FTimingGraphTrack::RemoveFrameStatsTimerSeries(uint32 TimerId)
+void FTimingGraphTrack::RemoveFrameStatsTimerSeries(uint32 TimerId, ETraceFrameType FrameType)
 {
-	AllSeries.RemoveAll([TimerId](const TSharedPtr<FGraphSeries>& Series)
+	AllSeries.RemoveAll([TimerId, FrameType](const TSharedPtr<FGraphSeries>& Series)
 		{
 			const TSharedPtr<FTimingGraphSeries> TimingSeries = StaticCastSharedPtr<FTimingGraphSeries>(Series);
-			return TimingSeries->Type == FTimingGraphSeries::ESeriesType::FrameStatsTimer && TimingSeries->TimerId == TimerId;
+			return TimingSeries->Type == FTimingGraphSeries::ESeriesType::FrameStatsTimer && TimingSeries->TimerId == TimerId && TimingSeries->FrameType == FrameType;
 		});
 }
 
@@ -480,7 +479,7 @@ void FTimingGraphTrack::UpdateFrameStatsTimerSeries(FTimingGraphSeries& Series,
 				return;
 			}
 
-			FramesProvider.EnumerateFrames(ETraceFrameType::TraceFrameType_Game, (uint64) 0, (uint64) FrameCount, [&Series](const TraceServices::FFrame& Frame)
+			FramesProvider.EnumerateFrames(Series.FrameType, (uint64) 0, (uint64) FrameCount, [&Series](const TraceServices::FFrame& Frame)
 				{
 					FTimingGraphSeries::FAtomicTimingEvent Event;
 					Event.FrameStartTime = Frame.StartTime;
diff --git a/Engine/Source/Developer/TraceInsights/Private/Insights/ViewModels/TimingGraphTrack.h b/Engine/Source/Developer/TraceInsights/Private/Insights/ViewModels/TimingGraphTrack.h
--- a/Engine/Source/Developer/TraceInsights/Private/Insights/ViewModels/TimingGraphTrack.h
+++ b/Engine/Source/Developer/TraceInsights/Private/Insights/ViewModels/TimingGraphTrack.h
@@ -63,11 +63,11 @@ public:
 	ESeriesType Type;
 	union
 	{
-		ETraceFrameType FrameType;
 		uint32 TimerId;
 		uint32 CounterId;
 	};
 
+	ETraceFrameType FrameType;
 	double CachedSessionDuration;
 	TArray<FSimpleTimingEvent> CachedEvents; // used by Timer series
 	TArray<FAtomicTimingEvent> FrameStatsCachedEvents; // used by Frame Stats Timer series
@@ -95,9 +95,9 @@ public:
 	TSharedPtr<FTimingGraphSeries> AddTimerSeries(uint32 TimerId, FLinearColor Color);
 	void RemoveTimerSeries(uint32 TimerId);
 
-	TSharedPtr<FTimingGraphSeries> GetFrameStatsTimerSeries(uint32 TimerId);
-	TSharedPtr<FTimingGraphSeries> AddFrameStatsTimerSeries(uint32 TimerId, FLinearColor Color);
-	void RemoveFrameStatsTimerSeries(uint32 TimerId);
+	TSharedPtr<FTimingGraphSeries> GetFrameStatsTimerSeries(uint32 TimerId, ETraceFrameType FrameType);
+	TSharedPtr<FTimingGraphSeries> AddFrameStatsTimerSeries(uint32 TimerId, ETraceFrameType FrameType, FLinearColor Color);
+	void RemoveFrameStatsTimerSeries(uint32 TimerId, ETraceFrameType FrameType);
 
 	TSharedPtr<FTimingGraphSeries> GetStatsCounterSeries(uint32 CounterId);
 	TSharedPtr<FTimingGraphSeries> AddStatsCounterSeries(uint32 CounterId, FLinearColor Color);
diff --git a/Engine/Source/Developer/TraceInsights/Private/Insights/Widgets/STimersView.cpp b/Engine/Source/Developer/TraceInsights/Private/Insights/Widgets/STimersView.cpp
--- a/Engine/Source/Developer/TraceInsights/Private/Insights/Widgets/STimersView.cpp
+++ b/Engine/Source/Developer/TraceInsights/Private/Insights/Widgets/STimersView.cpp
@@ -312,7 +312,7 @@ void STimersView::Construct(const FArguments& InArgs)
 					.VAlign(VAlign_Center)
 					[
 						SNew(STextBlock)
-						.Text(LOCTEXT("FrameType", "Frame Type"))
+						.Text(LOCTEXT("Mode", "Mode"))
 					]
 
 					+ SHorizontalBox::Slot()
@@ -322,14 +322,14 @@ void STimersView::Construct(const FArguments& InArgs)
 						SNew(SBox)
 						.MinDesiredWidth(128.0f)
 						[
-							SAssignNew(FrameTypeComboBox, SComboBox<TSharedPtr<ETraceFrameType>>)
-							.ToolTipText(this, &STimersView::FrameType_GetSelectedTooltipText)
-							.OptionsSource(&FrameTypeOptionsSource)
-							.OnSelectionChanged(this, &STimersView::FrameType_OnSelectionChanged)
-							.OnGenerateWidget(this, &STimersView::FrameType_OnGenerateWidget)
+							SAssignNew(ModeComboBox, SComboBox<TSharedPtr<ETraceFrameType>>)
+							.ToolTipText(this, &STimersView::Mode_GetSelectedTooltipText)
+							.OptionsSource(&ModeOptionsSource)
+							.OnSelectionChanged(this, &STimersView::Mode_OnSelectionChanged)
+							.OnGenerateWidget(this, &STimersView::Mode_OnGenerateWidget)
 							[
 								SNew(STextBlock)
-								.Text(this, &STimersView::FrameType_GetSelectedText)
+								.Text(this, &STimersView::Mode_GetSelectedText)
 							]
 						]
 					]
@@ -420,7 +420,7 @@ void STimersView::Construct(const FArguments& InArgs)
 	Filters->Add(TextFilter);
 
 	CreateGroupByOptionsSources();
-	CreateFrameTypeOptionsSources();
+	CreateModeOptionsSources();
 	CreateSortings();
 
 	InitCommandList();
@@ -538,73 +538,14 @@ TSharedPtr<SWidget> STimersView::TreeView_GetMenuContent()
 			}
 		}
 
-		// Add/remove series to/from graph track
-		{
-			FUIAction Action_ToggleTimerInGraphTrack;
-			Action_ToggleTimerInGraphTrack.CanExecuteAction = FCanExecuteAction::CreateLambda(CanExecute);
-			Action_ToggleTimerInGraphTrack.ExecuteAction = FExecuteAction::CreateSP(this, &STimersView::ToggleTimingViewMainGraphEventSeries, SelectedNode);
-
-			if (SelectedNode.IsValid() &&
-				SelectedNode->GetType() != ETimerNodeType::Group &&
-				IsSeriesInTimingViewMainGraph(SelectedNode))
-			{
-				MenuBuilder.AddMenuEntry
-				(
-					LOCTEXT("ContextMenu_RemoveFromGraphTrack", "Remove series from graph track"),
-					LOCTEXT("ContextMenu_RemoveFromGraphTrack_Desc", "Removes the series containing event instances of the selected timer from the Main Graph track."),
-					FSlateIcon(FInsightsStyle::GetStyleSetName(), "Icons.RemoveGraphSeries"),
-					Action_ToggleTimerInGraphTrack,
-					NAME_None,
-					EUserInterfaceActionType::Button
-				);
-			}
-			else
-			{
-				MenuBuilder.AddMenuEntry
-				(
-					LOCTEXT("ContextMenu_AddToGraphTrack", "Add series to graph track"),
-					LOCTEXT("ContextMenu_AddToGraphTrack_Desc", "Adds a series containing event instances of the selected timer to the Main Graph track."),
-					FSlateIcon(FInsightsStyle::GetStyleSetName(), "Icons.AddGraphSeries"),
-					Action_ToggleTimerInGraphTrack,
-					NAME_None,
-					EUserInterfaceActionType::Button
-				);
-			}
-		}
-
-		// Add/remove frame stats series to/from graph track
-		{
-			FUIAction Action_ToggleFrameStatsTimerInGraphTrack;
-			Action_ToggleFrameStatsTimerInGraphTrack.CanExecuteAction = FCanExecuteAction::CreateLambda(CanExecute); 
-			Action_ToggleFrameStatsTimerInGraphTrack.ExecuteAction = FExecuteAction::CreateSP(this, &STimersView::ToggleTimingViewMainGraphEventFrameStatsSeries, SelectedNode);
-
-			if (SelectedNode.IsValid() &&
-				SelectedNode->GetType() != ETimerNodeType::Group &&
-				IsFrameStatsSeriesInTimingViewMainGraph(SelectedNode))
-			{
-				MenuBuilder.AddMenuEntry
-				(
-					LOCTEXT("ContextMenu_RemoveFrameStatsFromGraphTrack", "Remove frame stats series from graph track"),
-					LOCTEXT("ContextMenu_RemoveFrameStatsFromGraphTrack_Desc", "Removes the frame stats series containing event instances of the selected timer from the Main Graph track."),
-					FSlateIcon(FInsightsStyle::GetStyleSetName(), "Icons.RemoveGraphSeries"),
-					Action_ToggleFrameStatsTimerInGraphTrack,
-					NAME_None,
-					EUserInterfaceActionType::Button
-				);
-			}
-			else
-			{
-				MenuBuilder.AddMenuEntry
-				(
-					LOCTEXT("ContextMenu_AddFrameStatsSeriesToGraphTrack", "Add frame stats series to graph track"),
-					LOCTEXT("ContextMenu_AddFrameStatsSeriesToGraphTrack_Desc", "Adds a frame stats series containing event instances of the selected timer to the Main Graph track."),
-					FSlateIcon(FInsightsStyle::GetStyleSetName(), "Icons.AddGraphSeries"),
-					Action_ToggleFrameStatsTimerInGraphTrack,
-					NAME_None,
-					EUserInterfaceActionType::Button
-				);
-			}
-		}
+		MenuBuilder.AddSubMenu
+		(
+			LOCTEXT("ContextMenu_PlotTimer_SubMenu", "Plot Timer"),
+			LOCTEXT("ContextMenu_PlotTimer_SubMenu_Desc", "Options to add the timer series to graph or frame tracks."),
+			FNewMenuDelegate::CreateSP(this, &STimersView::TreeView_BuildPlotTimerMenu),
+			false,
+			FSlateIcon(FAppStyle::Get().GetStyleSetName(), "Icons.AddGraphSeries")
+		);
 
 		// Open Source in IDE
 		{
@@ -824,6 +765,124 @@ void STimersView::TreeView_BuildExportMenu(FMenuBuilder& MenuBuilder)
 
 ////////////////////////////////////////////////////////////////////////////////////////////////////
 
+void STimersView::TreeView_BuildPlotTimerMenu(FMenuBuilder& MenuBuilder)
+{
+	const TArray<FTimerNodePtr> SelectedNodes = TreeView->GetSelectedItems();
+	const int32 NumSelectedNodes = SelectedNodes.Num();
+	FTimerNodePtr SelectedNode = NumSelectedNodes ? SelectedNodes[0] : nullptr;
+
+	auto CanExecute = [NumSelectedNodes, SelectedNode]()
+	{
+		TSharedPtr<STimingProfilerWindow> Wnd = FTimingProfilerManager::Get()->GetProfilerWindow();
+		TSharedPtr<STimingView> TimingView = Wnd.IsValid() ? Wnd->GetTimingView() : nullptr;
+		return TimingView.IsValid() && NumSelectedNodes == 1 && SelectedNode.IsValid() && SelectedNode->GetType() != ETimerNodeType::Group;
+	};
+
+	// Add/remove series to/from graph track
+	{
+		FUIAction Action_ToggleTimerInGraphTrack;
+		Action_ToggleTimerInGraphTrack.CanExecuteAction = FCanExecuteAction::CreateLambda(CanExecute);
+		Action_ToggleTimerInGraphTrack.ExecuteAction = FExecuteAction::CreateSP(this, &STimersView::ToggleTimingViewMainGraphEventSeries, SelectedNode);
+
+		if (SelectedNode.IsValid() &&
+			SelectedNode->GetType() != ETimerNodeType::Group &&
+			IsSeriesInTimingViewMainGraph(SelectedNode))
+		{
+			MenuBuilder.AddMenuEntry
+			(
+				LOCTEXT("ContextMenu_RemoveFromGraphTrack", "Remove series from graph track"),
+				LOCTEXT("ContextMenu_RemoveFromGraphTrack_Desc", "Removes the series containing event instances of the selected timer from the Main Graph track."),
+				FSlateIcon(FInsightsStyle::GetStyleSetName(), "Icons.RemoveGraphSeries"),
+				Action_ToggleTimerInGraphTrack,
+				NAME_None,
+				EUserInterfaceActionType::Button
+			);
+		}
+		else
+		{
+			MenuBuilder.AddMenuEntry
+			(
+				LOCTEXT("ContextMenu_AddToGraphTrack", "Add series to graph track"),
+				LOCTEXT("ContextMenu_AddToGraphTrack_Desc", "Adds a series containing event instances of the selected timer to the Main Graph track."),
+				FSlateIcon(FInsightsStyle::GetStyleSetName(), "Icons.AddGraphSeries"),
+				Action_ToggleTimerInGraphTrack,
+				NAME_None,
+				EUserInterfaceActionType::Button
+			);
+		}
+	}
+
+	// Add/remove game frame stats series to/from graph track
+	{
+		FUIAction Action_ToggleFrameStatsTimerInGraphTrack;
+		Action_ToggleFrameStatsTimerInGraphTrack.CanExecuteAction = FCanExecuteAction::CreateLambda(CanExecute);
+		Action_ToggleFrameStatsTimerInGraphTrack.ExecuteAction = FExecuteAction::CreateSP(this, &STimersView::ToggleTimingViewMainGraphEventFrameStatsSeries, SelectedNode, ETraceFrameType::TraceFrameType_Game);
+
+		if (SelectedNode.IsValid() &&
+			SelectedNode->GetType() != ETimerNodeType::Group &&
+			IsFrameStatsSeriesInTimingViewMainGraph(SelectedNode, ETraceFrameType::TraceFrameType_Game))
+		{
+			MenuBuilder.AddMenuEntry
+			(
+				LOCTEXT("ContextMenu_RemoveGameFrameStatsFromGraphTrack", "Remove game frame stats series from graph track"),
+				LOCTEXT("ContextMenu_RemoveGameFrameStatsFromGraphTrack_Desc", "Remove the game frame stats series for this timer."),
+				FSlateIcon(FInsightsStyle::GetStyleSetName(), "Icons.RemoveGraphSeries"),
+				Action_ToggleFrameStatsTimerInGraphTrack,
+				NAME_None,
+				EUserInterfaceActionType::Button
+			);
+		}
+		else
+		{
+			MenuBuilder.AddMenuEntry
+			(
+				LOCTEXT("ContextMenu_AddGameFrameStatsSeriesToGraphTrack", "Add game frame stats series to graph track"),
+				LOCTEXT("ContextMenu_AddGameFrameStatsSeriesToGraphTrack_Desc", "Adds a game frame stats series for this timer. Each data entry is computed as the sum of all instances of this timer in a game frame."),
+				FSlateIcon(FInsightsStyle::GetStyleSetName(), "Icons.AddGraphSeries"),
+				Action_ToggleFrameStatsTimerInGraphTrack,
+				NAME_None,
+				EUserInterfaceActionType::Button
+			);
+		}
+	}
+
+	// Add/remove rendering frame stats series to/from graph track
+	{
+		FUIAction Action_ToggleFrameStatsTimerInGraphTrack;
+		Action_ToggleFrameStatsTimerInGraphTrack.CanExecuteAction = FCanExecuteAction::CreateLambda(CanExecute);
+		Action_ToggleFrameStatsTimerInGraphTrack.ExecuteAction = FExecuteAction::CreateSP(this, &STimersView::ToggleTimingViewMainGraphEventFrameStatsSeries, SelectedNode, ETraceFrameType::TraceFrameType_Rendering);
+
+		if (SelectedNode.IsValid() &&
+			SelectedNode->GetType() != ETimerNodeType::Group &&
+			IsFrameStatsSeriesInTimingViewMainGraph(SelectedNode, ETraceFrameType::TraceFrameType_Rendering))
+		{
+			MenuBuilder.AddMenuEntry
+			(
+				LOCTEXT("ContextMenu_RemoveRenderingFrameStatsFromGraphTrack", "Remove rendering frame stats series from graph track"),
+				LOCTEXT("ContextMenu_RemoveRenderingFrameStatsFromGraphTrack_Desc", "Remove the rendering frame stats series for this timer."),
+				FSlateIcon(FInsightsStyle::GetStyleSetName(), "Icons.RemoveGraphSeries"),
+				Action_ToggleFrameStatsTimerInGraphTrack,
+				NAME_None,
+				EUserInterfaceActionType::Button
+			);
+		}
+		else
+		{
+			MenuBuilder.AddMenuEntry
+			(
+				LOCTEXT("ContextMenu_AddRenderingFrameStatsSeriesToGraphTrack", "Add rendering frame stats series to graph track"),
+				LOCTEXT("ContextMenu_AddRenderingFrameStatsSeriesToGraphTrack_Desc", "Adds a rendering frame stats series for this timer. Each data entry is computed as the sum of all instances of this timer in a rendering frame."),
+				FSlateIcon(FInsightsStyle::GetStyleSetName(), "Icons.AddGraphSeries"),
+				Action_ToggleFrameStatsTimerInGraphTrack,
+				NAME_None,
+				EUserInterfaceActionType::Button
+			);
+		}
+	}
+}
+
+////////////////////////////////////////////////////////////////////////////////////////////////////
+
 void STimersView::InitializeAndShowHeaderColumns()
 {
 	// Create columns.
@@ -1580,29 +1639,29 @@ FText STimersView::GroupBy_GetSelectedTooltipText() const
 // Frame Type
 ////////////////////////////////////////////////////////////////////////////////////////////////////
 
-void STimersView::CreateFrameTypeOptionsSources()
+void STimersView::CreateModeOptionsSources()
 {
-	FrameTypeOptionsSource.Reset(3);
+	ModeOptionsSource.Reset(3);
 
-	FrameTypeOptionsSource.Add(MakeShared<ETraceFrameType>(ETraceFrameType::TraceFrameType_Count));
-	FrameTypeOptionsSource.Add(MakeShared<ETraceFrameType>(ETraceFrameType::TraceFrameType_Game));
-	FrameTypeOptionsSource.Add(MakeShared<ETraceFrameType>(ETraceFrameType::TraceFrameType_Rendering));
+	ModeOptionsSource.Add(MakeShared<ETraceFrameType>(ETraceFrameType::TraceFrameType_Count));
+	ModeOptionsSource.Add(MakeShared<ETraceFrameType>(ETraceFrameType::TraceFrameType_Game));
+	ModeOptionsSource.Add(MakeShared<ETraceFrameType>(ETraceFrameType::TraceFrameType_Rendering));
 
-	FrameTypeComboBox->SetSelectedItem(FrameTypeOptionsSource[0]);
+	ModeComboBox->SetSelectedItem(ModeOptionsSource[0]);
 
-	FrameTypeComboBox->RefreshOptions();
+	ModeComboBox->RefreshOptions();
 }
 
 ////////////////////////////////////////////////////////////////////////////////////////////////////
 
-void STimersView::FrameType_OnSelectionChanged(TSharedPtr<ETraceFrameType> NewFrameType, ESelectInfo::Type SelectInfo)
+void STimersView::Mode_OnSelectionChanged(TSharedPtr<ETraceFrameType> NewFrameType, ESelectInfo::Type SelectInfo)
 {
-	if (*NewFrameType == FrameType)
+	if (*NewFrameType == ModeFrameType)
 	{
 		return;
 	}
 
-	FrameType = *NewFrameType;
+	ModeFrameType = *NewFrameType;
 
 	TSharedPtr<STimingProfilerWindow> Window = FTimingProfilerManager::Get()->GetProfilerWindow();
 	if (Window.IsValid())
@@ -1610,45 +1669,45 @@ void STimersView::FrameType_OnSelectionChanged(TSharedPtr<ETraceFrameType> NewFr
 		TSharedPtr<STimingView> TimingView = Window->GetTimingView();
 		if (TimingView.IsValid())
 		{
-			TimingView->SetFrameTypeToSnapTo(FrameType);
+			TimingView->SetFrameTypeToSnapTo(ModeFrameType);
 		}
 	}
 
-	Aggregator->SetFrameType(FrameType);
+	Aggregator->SetFrameType(ModeFrameType);
 	Aggregator->Cancel();
 	Aggregator->Start();
 }
 
 ////////////////////////////////////////////////////////////////////////////////////////////////////
 
-TSharedRef<SWidget> STimersView::FrameType_OnGenerateWidget(TSharedPtr<ETraceFrameType> InFrameTypeMode) const
+TSharedRef<SWidget> STimersView::Mode_OnGenerateWidget(TSharedPtr<ETraceFrameType> InFrameTypeMode) const
 {
 	return SNew(STextBlock)
-	.Text(FrameType_GetText(*InFrameTypeMode))
-	.ToolTipText(FrameType_GetTooltipText(*InFrameTypeMode));
+	.Text(Mode_GetText(*InFrameTypeMode))
+	.ToolTipText(Mode_GetTooltipText(*InFrameTypeMode));
 }
 
 ////////////////////////////////////////////////////////////////////////////////////////////////////
 
-FText STimersView::FrameType_GetSelectedText() const
+FText STimersView::Mode_GetSelectedText() const
 {
-	return FrameType_GetText(FrameType);
+	return Mode_GetText(ModeFrameType);
 }
 
 ////////////////////////////////////////////////////////////////////////////////////////////////////
 
-FText STimersView::FrameType_GetText(ETraceFrameType InFrameType) const
+FText STimersView::Mode_GetText(ETraceFrameType InFrameType) const
 {
 	switch (InFrameType)
 	{
 	case ETraceFrameType::TraceFrameType_Count:
-		return LOCTEXT("FrameTypeNone", "None");
+		return LOCTEXT("InstanceMode", "Instance");
 		break;
 	case ETraceFrameType::TraceFrameType_Game:
-		return LOCTEXT("FrameTypeGame", "Game");
+		return LOCTEXT("GameFrameMode", "Game Frame");
 		break;
 	case ETraceFrameType::TraceFrameType_Rendering:
-		return LOCTEXT("FrameTypeRendering", "Rendering");
+		return LOCTEXT("RenderingFrameMode", "Rendering Frame");
 		break;
 	}
 
@@ -1657,25 +1716,25 @@ FText STimersView::FrameType_GetText(ETraceFrameType InFrameType) const
 
 ////////////////////////////////////////////////////////////////////////////////////////////////////
 
-FText STimersView::FrameType_GetSelectedTooltipText() const
+FText STimersView::Mode_GetSelectedTooltipText() const
 {
-	return FrameType_GetTooltipText(FrameType);
+	return Mode_GetTooltipText(ModeFrameType);
 }
 
 ////////////////////////////////////////////////////////////////////////////////////////////////////
 
-FText STimersView::FrameType_GetTooltipText(ETraceFrameType InFrameType) const
+FText STimersView::Mode_GetTooltipText(ETraceFrameType InFrameType) const
 {
 	switch (InFrameType)
 	{
 	case ETraceFrameType::TraceFrameType_Count:
-		return LOCTEXT("FrameTypeNoneTooltip", "Do not calculate the timer stats per frame");
+		return LOCTEXT("ModeInstanceTooltip", "Calculate the stats per event instance.");
 		break;
 	case ETraceFrameType::TraceFrameType_Game:
-		return LOCTEXT("FrameTypeGameTooltip", "Calculate the timer stats per game frame");
+		return LOCTEXT("ModeGameFrameTooltip", "Calculate the timer stats per game frame.");
 		break;
 	case ETraceFrameType::TraceFrameType_Rendering:
-		return LOCTEXT("FrameTypeRenderingTooltip", "Calculate the timer stats per rendering frame");
+		return LOCTEXT("ModeRenderingFrameTooltip", "Calculate the timer stats per rendering frame.");
 		break;
 	}
 
@@ -2405,21 +2464,32 @@ void STimersView::ToggleTimingViewMainGraphEventSeries(FTimerNodePtr TimerNode)
 
 ////////////////////////////////////////////////////////////////////////////////////////////////////
 
-void STimersView::ToggleGraphFrameStatsSeries(TSharedRef<FTimingGraphTrack> GraphTrack, FTimerNodeRef NodePtr) const
+void STimersView::ToggleGraphFrameStatsSeries(TSharedRef<FTimingGraphTrack> GraphTrack, FTimerNodeRef NodePtr, ETraceFrameType FrameType) const
 {
 	const uint32 TimerId = NodePtr->GetTimerId();
-	TSharedPtr<FTimingGraphSeries> Series = GraphTrack->GetFrameStatsTimerSeries(TimerId);
+	TSharedPtr<FTimingGraphSeries> Series = GraphTrack->GetFrameStatsTimerSeries(TimerId, FrameType);
 	if (Series.IsValid())
 	{
-		GraphTrack->RemoveFrameStatsTimerSeries(TimerId);
+		GraphTrack->RemoveFrameStatsTimerSeries(TimerId, FrameType);
 		GraphTrack->SetDirtyFlag();
 		NodePtr->OnRemovedFromGraph();
 	}
 	else
 	{
 		GraphTrack->Show();
-		Series = GraphTrack->AddFrameStatsTimerSeries(TimerId, NodePtr->GetColor());
-		Series->SetName(FText::FromName(NodePtr->GetName()));
+		Series = GraphTrack->AddFrameStatsTimerSeries(TimerId, FrameType, NodePtr->GetColor());
+		FText NameText = FText::FromName(NodePtr->GetName());
+
+		if (FrameType == ETraceFrameType::TraceFrameType_Game)
+		{
+			NameText = FText::Format((LOCTEXT("GameFrameSeriesName", "{0} Game Frame")), NameText);
+		}
+		else if (FrameType == ETraceFrameType::TraceFrameType_Rendering)
+		{
+			NameText = FText::Format((LOCTEXT("RenderingFrameSeriesName", "{0} Rendering Frame")), NameText);
+		}
+
+		Series->SetName(NameText);
 		GraphTrack->SetDirtyFlag();
 		NodePtr->OnAddedToGraph();
 	}
@@ -2427,14 +2497,14 @@ void STimersView::ToggleGraphFrameStatsSeries(TSharedRef<FTimingGraphTrack> Grap
 
 ////////////////////////////////////////////////////////////////////////////////////////////////////
 
-bool STimersView::IsFrameStatsSeriesInTimingViewMainGraph(FTimerNodePtr TimerNode) const
+bool STimersView::IsFrameStatsSeriesInTimingViewMainGraph(FTimerNodePtr TimerNode, ETraceFrameType FrameType) const
 {
 	TSharedPtr<FTimingGraphTrack> GraphTrack = GetTimingViewMainGraphTrack();
 
 	if (GraphTrack.IsValid())
 	{
 		const uint32 TimerId = TimerNode->GetTimerId();
-		TSharedPtr<FTimingGraphSeries> Series = GraphTrack->GetFrameStatsTimerSeries(TimerId);
+		TSharedPtr<FTimingGraphSeries> Series = GraphTrack->GetFrameStatsTimerSeries(TimerId, FrameType);
 
 		return Series.IsValid();
 	}
@@ -2444,12 +2514,12 @@ bool STimersView::IsFrameStatsSeriesInTimingViewMainGraph(FTimerNodePtr TimerNod
 
 ////////////////////////////////////////////////////////////////////////////////////////////////////
 
-void STimersView::ToggleTimingViewMainGraphEventFrameStatsSeries(FTimerNodePtr TimerNode) const
+void STimersView::ToggleTimingViewMainGraphEventFrameStatsSeries(FTimerNodePtr TimerNode, ETraceFrameType FrameType) const
 {
 	TSharedPtr<FTimingGraphTrack> GraphTrack = GetTimingViewMainGraphTrack();
 	if (GraphTrack.IsValid())
 	{
-		ToggleGraphFrameStatsSeries(GraphTrack.ToSharedRef(), TimerNode.ToSharedRef());
+		ToggleGraphFrameStatsSeries(GraphTrack.ToSharedRef(), TimerNode.ToSharedRef(), FrameType);
 	}
 }
 
diff --git a/Engine/Source/Developer/TraceInsights/Private/Insights/Widgets/STimersView.h b/Engine/Source/Developer/TraceInsights/Private/Insights/Widgets/STimersView.h
--- a/Engine/Source/Developer/TraceInsights/Private/Insights/Widgets/STimersView.h
+++ b/Engine/Source/Developer/TraceInsights/Private/Insights/Widgets/STimersView.h
@@ -219,19 +219,19 @@ private:
 	////////////////////////////////////////////////////////////////////////////////////////////////////
 	// Frame Type
 
-	void CreateFrameTypeOptionsSources();
+	void CreateModeOptionsSources();
 
-	void FrameType_OnSelectionChanged(TSharedPtr<ETraceFrameType> NewGroupingMode, ESelectInfo::Type SelectInfo);
+	void Mode_OnSelectionChanged(TSharedPtr<ETraceFrameType> NewGroupingMode, ESelectInfo::Type SelectInfo);
 
-	TSharedRef<SWidget> FrameType_OnGenerateWidget(TSharedPtr<ETraceFrameType> InGroupingMode) const;
+	TSharedRef<SWidget> Mode_OnGenerateWidget(TSharedPtr<ETraceFrameType> InGroupingMode) const;
 
-	FText FrameType_GetSelectedText() const;
+	FText Mode_GetSelectedText() const;
 	
-	FText FrameType_GetText(ETraceFrameType InFrameType) const;
+	FText Mode_GetText(ETraceFrameType InFrameType) const;
 
-	FText FrameType_GetSelectedTooltipText() const;
+	FText Mode_GetSelectedTooltipText() const;
 
-	FText FrameType_GetTooltipText(ETraceFrameType InFrameType) const;
+	FText Mode_GetTooltipText(ETraceFrameType InFrameType) const;
 
 	////////////////////////////////////////////////////////////////////////////////////////////////////
 	// Sorting
@@ -299,14 +299,16 @@ private:
 
 	void ToggleTimingViewEventFilter(FTimerNodePtr TimerNode) const;
 
+	void TreeView_BuildPlotTimerMenu(FMenuBuilder& MenuBuilder);
+
 	TSharedPtr<FTimingGraphTrack> GetTimingViewMainGraphTrack() const;
 	void ToggleGraphSeries(TSharedRef<FTimingGraphTrack> GraphTrack, FTimerNodeRef NodePtr) const;
 	bool IsSeriesInTimingViewMainGraph(FTimerNodePtr TimerNode) const;
 	void ToggleTimingViewMainGraphEventSeries(FTimerNodePtr TimerNode) const;
 
-	void ToggleGraphFrameStatsSeries(TSharedRef<FTimingGraphTrack> GraphTrack, FTimerNodeRef NodePtr) const;
-	bool IsFrameStatsSeriesInTimingViewMainGraph(FTimerNodePtr TimerNode) const;
-	void ToggleTimingViewMainGraphEventFrameStatsSeries(FTimerNodePtr TimerNode) const;
+	void ToggleGraphFrameStatsSeries(TSharedRef<FTimingGraphTrack> GraphTrack, FTimerNodeRef NodePtr, ETraceFrameType FrameType) const;
+	bool IsFrameStatsSeriesInTimingViewMainGraph(FTimerNodePtr TimerNode, ETraceFrameType FrameType) const;
+	void ToggleTimingViewMainGraphEventFrameStatsSeries(FTimerNodePtr TimerNode, ETraceFrameType FrameType) const;
 
 	////////////////////////////////////////////////////////////////////////////////////////////////////
 
@@ -407,11 +409,11 @@ private:
 	//////////////////////////////////////////////////
 	// Frame Type
 
-	TArray<TSharedPtr<ETraceFrameType>> FrameTypeOptionsSource;
+	TArray<TSharedPtr<ETraceFrameType>> ModeOptionsSource;
 
-	TSharedPtr<SComboBox<TSharedPtr<ETraceFrameType>>> FrameTypeComboBox;
+	TSharedPtr<SComboBox<TSharedPtr<ETraceFrameType>>> ModeComboBox;
 
-	ETraceFrameType FrameType;
+	ETraceFrameType ModeFrameType;
 
 	//////////////////////////////////////////////////
 	// Sorting
