UX for settings trace store and additional watch directories

diff --git a/Engine/Source/Developer/TraceAnalysis/Private/Store/CborPayload.h b/Engine/Source/Developer/TraceAnalysis/Private/Store/CborPayload.h
--- a/Engine/Source/Developer/TraceAnalysis/Private/Store/CborPayload.h
+++ b/Engine/Source/Developer/TraceAnalysis/Private/Store/CborPayload.h
@@ -73,6 +73,7 @@ public:
 	template <int N>			TPayloadBuilder(const char (&Path)[N]);
 	template <int N> void		AddInteger(const char (&Name)[N], int64 Value);
 	template <int N> void		AddString(const char (&Name)[N], const char* Value, int32 Length=-1);
+	template <int N> void		AddStringArray(const char(&Name)[N], const TArray<FString>& Values);
 	FPayload					Done();
 
 private:
@@ -119,7 +120,20 @@ inline void TPayloadBuilder<Size>::AddString(
 	CborWriter.WriteValue(Name, N - 1);
 	CborWriter.WriteValue(Value, Length);
 }
-
+	
+////////////////////////////////////////////////////////////////////////////////
+template<int Size>
+template<int N>
+inline void TPayloadBuilder<Size>::AddStringArray(const char(&Name)[N], const TArray<FString>& Values)
+{
+	CborWriter.WriteValue(Name, N - 1);
+	CborWriter.WriteContainerStart(ECborCode::Array, Values.Num());
+	for (const FString& Elem : Values)
+	{
+		CborWriter.WriteValue(Elem);
+	}
+}
+	
 ////////////////////////////////////////////////////////////////////////////////
 template <int Size>
 inline FPayload TPayloadBuilder<Size>::Done()
@@ -142,6 +156,7 @@ public:
 	int64			GetInteger(const char* Key, int64 Default) const { return GetInt64Checked(Key, Default); }
 	template <int N>
 	FUtf8StringView	GetString(const char* Key, const char (&Default)[N]) const;
+	bool			GetStringArray(const char* Key, TArray<FString>& OutArray) const;
 	const uint8*	GetData() const;
 	uint32			GetSize() const;
 	uint8*			Reserve(uint32 Size);
@@ -159,6 +174,35 @@ inline EStatusCode FResponse::GetStatusCode() const
 	return Code ? EStatusCode(Code) : EStatusCode::Unknown;
 }
 
+////////////////////////////////////////////////////////////////////////////////
+inline bool FResponse::GetStringArray(const char* Key, TArray<FString>& OutArray) const
+{
+	FMemoryReader MemoryReader(Buffer);
+	FCborReader CborReader(&MemoryReader, ECborEndianness::StandardCompliant);
+	FCborContext Context;
+	while (CborReader.ReadNext(Context))
+	{
+		if (Context.IsString())
+		{
+			uint32 Length = Context.AsLength();
+			uint32 Offset = uint32(MemoryReader.Tell());
+			auto* String = (const char*)(Buffer.GetData() + Offset - Length);
+			if (FCStringAnsi::Strncmp(Key, String, Length) == 0)
+			{
+				if (CborReader.ReadNext(Context) && Context.IsContainer() && Context.MajorType() == ECborCode::Array)
+				{
+					while (CborReader.ReadNext(Context) && Context.IsString())
+					{
+						OutArray.Emplace(Context.AsString());
+					}
+					return true;
+				}
+			}
+		}
+	}
+	return false;
+}
+
 ////////////////////////////////////////////////////////////////////////////////
 inline const uint8* FResponse::GetData() const
 {
@@ -174,7 +218,7 @@ inline uint32 FResponse::GetSize() const
 ////////////////////////////////////////////////////////////////////////////////
 inline uint8* FResponse::Reserve(uint32 Size)
 {
-	Buffer.SetNumUninitialized(Size);
+	Buffer.SetNumUninitialized(Size, false);
 	return Buffer.GetData();
 }
 
diff --git a/Engine/Source/Developer/TraceAnalysis/Private/Store/StoreClient.cpp b/Engine/Source/Developer/TraceAnalysis/Private/Store/StoreClient.cpp
--- a/Engine/Source/Developer/TraceAnalysis/Private/Store/StoreClient.cpp
+++ b/Engine/Source/Developer/TraceAnalysis/Private/Store/StoreClient.cpp
@@ -1,6 +1,7 @@
 // Copyright Epic Games, Inc. All Rights Reserved.
 
 #include "Trace/StoreClient.h"
+#include "Algo/Transform.h"
 #include "Asio/Asio.h"
 #include "AsioStore.h"
 #include "CborPayload.h"
@@ -115,12 +116,16 @@ public:
 	bool					GetSessionInfo(uint32 Index);
 	bool					GetSessionInfoById(uint32 Id);
 	bool					GetSessionInfoByTraceId(uint32 TraceId);
+							bool SetStoreDirectories(const TCHAR* StoreDir, const TArray<FString>& AddWatchDirs,
+													const TArray<FString>& RemoveWatchDirs);
 
 private:
 	bool					Communicate(const FPayload& Payload);
 	asio::io_context&		IoContext;
 	asio::ip::tcp::socket	Socket;
 	FResponse				Response;
+	FString					Host;
+	uint16					Port;
 };
 
 ////////////////////////////////////////////////////////////////////////////////
@@ -180,13 +185,13 @@ const FResponse& FStoreCborClient::GetResponse() const
 }
 
 ////////////////////////////////////////////////////////////////////////////////
-bool FStoreCborClient::Connect(const TCHAR* Host, uint16 Port)
+bool FStoreCborClient::Connect(const TCHAR* InHost, uint16 InPort)
 {
-	Port = (Port == 0) ? 1989 : Port;
+	InPort = (InPort == 0) ? 1989 : InPort;
 
-	FTCHARToUTF8 HostUtf8(Host);
+	FTCHARToUTF8 HostUtf8(InHost);
 	char PortString[8];
-	FCStringAnsi::Sprintf(PortString, "%d", Port);
+	FCStringAnsi::Sprintf(PortString, "%d", InPort);
 
 	asio::ip::tcp::resolver Resolver(IoContext);
 	asio::ip::tcp::resolver::results_type Endpoints = Resolver.resolve(
@@ -209,6 +214,10 @@ bool FStoreCborClient::Connect(const TCHAR* Host, uint16 Port)
 		return false;
 	}
 
+	// Save connection details if we need to reconnect
+	Host = InHost;
+	Port = InPort;
+
 	return true;
 }
 
@@ -371,6 +380,27 @@ bool FStoreCborClient::GetSessionInfoByTraceId(uint32 TraceId)
 	return Communicate(Payload);
 }
 
+////////////////////////////////////////////////////////////////////////////////
+bool FStoreCborClient::SetStoreDirectories(const TCHAR* StoreDir, const TArray<FString>& AddWatchDirs, const TArray<FString>& RemoveWatchDir)
+{
+	TPayloadBuilder<> Builder("v1/settings/write");
+	if (StoreDir)
+	{
+		Builder.AddString("StoreDir", TCHAR_TO_ANSI(StoreDir));
+	}
+	if (!AddWatchDirs.IsEmpty())
+	{
+		Builder.AddStringArray("Additionalwatchdirs", AddWatchDirs);
+	}
+	if (!RemoveWatchDir.IsEmpty())
+	{
+		TArray<FString> RemoveDirs;
+		Algo::Transform(RemoveWatchDir, RemoveDirs, [](const FString& In) { return TEXT("-") + In; });
+		Builder.AddStringArray("Additionalwatchdirs", RemoveDirs);
+	}
+	return Communicate(Builder.Done());
+}
+	
 } // namespace Trace
 } // namespace UE
 
@@ -399,8 +429,19 @@ uint32 FStoreClient::FStatus::GetChangeSerial() const
 	return Response->GetUint32Checked("change_serial", 0);
 }
 
+////////////////////////////////////////////////////////////////////////////////
+uint32 FStoreClient::FStatus::GetSettingsSerial() const
+{
+	const auto* Response = (const FResponse*)this;
+	return Response->GetUint32Checked("settings_serial", 0);
+}
 
-
+////////////////////////////////////////////////////////////////////////////////
+void FStoreClient::FStatus::GetWatchDirectories(TArray<FString>& OutDirs) const
+{
+	const auto* Response = (const FResponse*)this;
+	Response->GetStringArray("watch_dirs", OutDirs);
+}
 
 
 ////////////////////////////////////////////////////////////////////////////////
@@ -424,6 +465,13 @@ uint64 FStoreClient::FTraceInfo::GetTimestamp() const
 	return Response->GetUint64Checked("timestamp", 0);
 }
 
+////////////////////////////////////////////////////////////////////////////////
+FUtf8StringView FStoreClient::FTraceInfo::GetUri() const
+{
+	const auto* Response = (const FResponse*)this;
+	return Response->GetString("uri", "");
+}
+
 ////////////////////////////////////////////////////////////////////////////////
 FUtf8StringView FStoreClient::FTraceInfo::GetName() const
 {
@@ -571,6 +619,15 @@ FStoreClient::FTraceData FStoreClient::ReadTrace(uint32 Id)
 	return FTraceData(Self->ReadTrace(Id));
 }
 
+////////////////////////////////////////////////////////////////////////////////
+bool FStoreClient::SetStoreDirectories(const TCHAR* StoreDir,
+										const TArray<FString>& AddWatchDirs,
+										const TArray<FString>& RemoveWatchDirs)
+{
+	auto* Self = (FStoreCborClient*) this;
+	return Self->SetStoreDirectories(StoreDir, AddWatchDirs, RemoveWatchDirs);
+}
+
 ////////////////////////////////////////////////////////////////////////////////
 uint32 FStoreClient::GetSessionCount() const
 {
diff --git a/Engine/Source/Developer/TraceAnalysis/Public/Trace/StoreClient.h b/Engine/Source/Developer/TraceAnalysis/Public/Trace/StoreClient.h
--- a/Engine/Source/Developer/TraceAnalysis/Public/Trace/StoreClient.h
+++ b/Engine/Source/Developer/TraceAnalysis/Public/Trace/StoreClient.h
@@ -22,6 +22,8 @@ public:
 		FUtf8StringView	GetStoreDir() const;
 		uint32			GetRecorderPort() const;
 		uint32			GetChangeSerial() const;
+		uint32			GetSettingsSerial() const;
+		void			GetWatchDirectories(TArray<FString>& OutDirs) const;
 	};
 
 	struct TRACEANALYSIS_API FTraceInfo
@@ -30,6 +32,7 @@ public:
 		uint32			GetId() const;
 		uint64			GetSize() const;
 		uint64			GetTimestamp() const;
+		FUtf8StringView GetUri() const;
 		//const TCHAR*	GetMetadata(const TCHAR* Key) const;
 		//template <typename Lambda> uint32 ReadMetadata(Lambda&& Callback) const;
 	};
@@ -52,6 +55,8 @@ public:
 	const FTraceInfo*	GetTraceInfo(uint32 Index);
 	const FTraceInfo*	GetTraceInfoById(uint32 Id);
 	FTraceData			ReadTrace(uint32 Id);
+	bool SetStoreDirectories(const TCHAR* StoreDir, const TArray<FString>& AddWatchDirs,
+							const TArray<FString>& RemoveWatchDirs);
 #if 0
 	template <typename Lambda> uint32 GetTraceInfos(uint32 StartIndex, uint32 Count, Lambda&& Callback) const;
 #endif // 0
diff --git a/Engine/Source/Developer/TraceInsights/Private/Insights/InsightsManager.h b/Engine/Source/Developer/TraceInsights/Private/Insights/InsightsManager.h
--- a/Engine/Source/Developer/TraceInsights/Private/Insights/InsightsManager.h
+++ b/Engine/Source/Developer/TraceInsights/Private/Insights/InsightsManager.h
@@ -124,6 +124,8 @@ public:
 	 * @return True on success, false on failure
 	 */
 	bool ReconnectToStore() const;
+
+	const FString& GetLastStoreHost() const { return LastStoreHost; }
 	
 	UE::Trace::FStoreClient* GetStoreClient() const { return StoreClient.Get(); }
 	FCriticalSection& GetStoreClientCriticalSection() const { return StoreClientCriticalSection; }
diff --git a/Engine/Source/Developer/TraceInsights/Private/Insights/InsightsStyle.cpp b/Engine/Source/Developer/TraceInsights/Private/Insights/InsightsStyle.cpp
--- a/Engine/Source/Developer/TraceInsights/Private/Insights/InsightsStyle.cpp
+++ b/Engine/Source/Developer/TraceInsights/Private/Insights/InsightsStyle.cpp
@@ -150,10 +150,16 @@ void FInsightsStyle::FStyle::Initialize()
 	// Trace Store, Connection, Launcher
 
 	Set("Icons.TraceStore", new IMAGE_BRUSH_SVG("TraceStore", Icon16x16));
+	Set("Icons.Expand", new CORE_IMAGE_BRUSH_SVG("Starship/Common/chevron-right", Icon16x16));
+	Set("Icons.Expanded", new CORE_IMAGE_BRUSH_SVG("Starship/Common/chevron-down", Icon16x16));
+	Set("Icons.AddWatchDir", new CORE_IMAGE_BRUSH_SVG("Starship/Common/folder-plus", Icon16x16));
+	Set("Icons.RemoveWatchDir", new CORE_IMAGE_BRUSH_SVG("Starship/Common/delete", Icon16x16));
 	Set("Icons.Connection", new IMAGE_BRUSH_SVG("Connection", Icon16x16));
 	Set("Icons.Online", new CORE_IMAGE_BRUSH_SVG("Starship/Common/check-circle", Icon16x16, FStyleColors::AccentGreen));
 	Set("Icons.Offline", new CORE_IMAGE_BRUSH_SVG("Starship/Common/alert-triangle", Icon16x16, FStyleColors::Warning));
 	Set("Icons.Launcher", new TODO_IMAGE_BRUSH(Icon16x16));
+	Set("Icons.UTrace", new IMAGE_BRUSH_SVG("UTrace", Icon16x16));
+	Set("Icons.UTraceStack", new IMAGE_BRUSH_SVG("UTrace", Icon16x16));
 
 	//////////////////////////////////////////////////
 	// Timing Insights
diff --git a/Engine/Source/Developer/TraceInsights/Private/Insights/StoreService/StoreBrowser.cpp b/Engine/Source/Developer/TraceInsights/Private/Insights/StoreService/StoreBrowser.cpp
--- a/Engine/Source/Developer/TraceInsights/Private/Insights/StoreService/StoreBrowser.cpp
+++ b/Engine/Source/Developer/TraceInsights/Private/Insights/StoreService/StoreBrowser.cpp
@@ -140,7 +140,7 @@ void FStoreBrowser::UpdateTraces()
 			
 			if (bIsValid)
 			{
-				// If we gave just reconnected we need to reset known traces and change serials
+				// If we just reconnected we need to reset known traces and change serials
 				Refresh();
 			}
 		}
@@ -151,13 +151,30 @@ void FStoreBrowser::UpdateTraces()
 			// If there is no connection, no point in checking for sessions
 			return;
 		}
+		else
+		{
+			FScopeLock Lock(&TracesCriticalSection);
+			Host = FInsightsManager::Get()->GetLastStoreHost();
+		}
 	}
 	
 	FScopeLock StoreClientLock(&GetStoreClientCriticalSection());
 
-	// Check if the list of trace sessions has changed.
+	// Check if the list of trace sessions or store settings have changed.
 	{
 		const UE::Trace::FStoreClient::FStatus* Status = StoreClient->GetStatus();
+
+		if (Status != nullptr && SettingsChangeSerial != Status->GetSettingsSerial())
+		{
+			SettingsChangeSerial = Status->GetSettingsSerial();
+			
+			// Update watch dirs
+			FScopeLock Lock(&TracesCriticalSection);
+			WatchDirectories.Empty();
+			Status->GetWatchDirectories(WatchDirectories);
+			StoreDirectory = FString(Status->GetStoreDir());
+		}
+		
 		if (Status != nullptr && StoreChangeSerial != Status->GetChangeSerial())
 		{
 			StoreChangeSerial = Status->GetChangeSerial();
@@ -198,6 +215,7 @@ void FStoreBrowser::UpdateTraces()
 
 							const FUtf8StringView Utf8NameView = TraceInfo->GetName();
 							Trace.Name = FString(Utf8NameView);
+							Trace.Uri = FString(TraceInfo->GetUri());
 
 							Trace.Timestamp = FStoreBrowserTraceInfo::ConvertTimestamp(TraceInfo->GetTimestamp());
 							Trace.Size = TraceInfo->GetSize();
diff --git a/Engine/Source/Developer/TraceInsights/Private/Insights/StoreService/StoreBrowser.h b/Engine/Source/Developer/TraceInsights/Private/Insights/StoreService/StoreBrowser.h
--- a/Engine/Source/Developer/TraceInsights/Private/Insights/StoreService/StoreBrowser.h
+++ b/Engine/Source/Developer/TraceInsights/Private/Insights/StoreService/StoreBrowser.h
@@ -31,7 +31,7 @@ struct FStoreBrowserTraceInfo
 	uint64 ChangeSerial = 0;
 
 	FString Name;
-	//FString Uri;
+	FString Uri;
 
 	FDateTime Timestamp = 0;
 	uint64 Size = 0;
@@ -99,10 +99,13 @@ public:
 	bool IsLocked() const { return bTracesLocked; }
 	void Lock() { check(!bTracesLocked); bTracesLocked = true; TracesCriticalSection.Lock(); }
 	uint64 GetLockedTracesChangeSerial() const { check(bTracesLocked); return TracesChangeSerial; }
+	uint32 GetLockedSettingsSerial() const { check(bTracesLocked); return SettingsChangeSerial; };
 	const TArray<TSharedPtr<FStoreBrowserTraceInfo>>& GetLockedTraces() const { check(bTracesLocked); return Traces; }
 	const TMap<uint32, TSharedPtr<FStoreBrowserTraceInfo>>& GetLockedTraceMap() const { check(bTracesLocked); return TraceMap; }
+	const TArray<FString>& GetLockedWatchDirectories() const { check(bTracesLocked); return WatchDirectories; }
+	const FString& GetLockedStoreDirectory() const { check(bTracesLocked); return StoreDirectory; }
+	const FString& GetLockedHost() const { check(bTracesLocked); return Host; }
 	void Unlock() { check(bTracesLocked); bTracesLocked = false; TracesCriticalSection.Unlock(); }
-
 	void Refresh();
 
 private:
@@ -126,9 +129,13 @@ private:
 	std::atomic<EConnectionStatus> ConnectionStatus;
 	uint32 StoreChangeSerial;
 	uint64 TracesChangeSerial;
+	uint32 SettingsChangeSerial;
 	TArray<TSharedPtr<FStoreBrowserTraceInfo>> Traces;
 	TMap<uint32, TSharedPtr<FStoreBrowserTraceInfo>> TraceMap;
 	TMap<uint32, TSharedPtr<FStoreBrowserTraceInfo>> LiveTraceMap;
+	TArray<FString> WatchDirectories;
+	FString StoreDirectory;
+	FString Host;
 };
 
 ////////////////////////////////////////////////////////////////////////////////////////////////////
diff --git a/Engine/Source/Developer/TraceInsights/Private/Insights/Widgets/SStartPageWindow.cpp b/Engine/Source/Developer/TraceInsights/Private/Insights/Widgets/SStartPageWindow.cpp
--- a/Engine/Source/Developer/TraceInsights/Private/Insights/Widgets/SStartPageWindow.cpp
+++ b/Engine/Source/Developer/TraceInsights/Private/Insights/Widgets/SStartPageWindow.cpp
@@ -3,33 +3,36 @@
 #include "SStartPageWindow.h"
 
 #include "DesktopPlatformModule.h"
+#include "IPAddress.h"
+#include "SlateOptMacros.h"
+#include "SocketSubsystem.h"
 #include "Framework/Application/SlateApplication.h"
 #include "Framework/MultiBox/MultiBoxBuilder.h"
 #include "HAL/FileManager.h"
 #include "HAL/PlatformProcess.h"
 #include "Internationalization/Text.h"
-#include "IPAddress.h"
+#include "Logging/MessageLog.h"
 #include "Misc/MessageDialog.h"
-#include "SlateOptMacros.h"
-#include "SocketSubsystem.h"
 #include "Styling/AppStyle.h"
 #include "Styling/StyleColors.h"
 #include "Trace/ControlClient.h"
 #include "Trace/StoreClient.h"
+#include "Widgets/SBoxPanel.h"
+#include "Widgets/SToolTip.h"
+#include "Widgets/Colors/SColorBlock.h"
 #include "Widgets/Images/SImage.h"
 #include "Widgets/Input/SButton.h"
-#include "Widgets/Input/SEditableTextBox.h"
-#include "Widgets/Input/SComboButton.h"
 #include "Widgets/Input/SCheckBox.h"
+#include "Widgets/Input/SComboButton.h"
+#include "Widgets/Input/SEditableTextBox.h"
 #include "Widgets/Input/SSearchBox.h"
 #include "Widgets/Layout/SBorder.h"
 #include "Widgets/Layout/SBox.h"
 #include "Widgets/Layout/SGridPanel.h"
 #include "Widgets/Notifications/SNotificationList.h"
-#include "Widgets/SBoxPanel.h"
-#include "Widgets/SToolTip.h"
 #include "Widgets/Testing/SStarshipSuite.h"
 #include "Widgets/Text/STextBlock.h"
+#include "Widgets/Views/STileView.h"
 
 #if WITH_EDITOR
 	#include "EngineAnalytics.h"
@@ -38,13 +41,13 @@
 #endif // WITH_EDITOR
 
 // Insights
-#include "Insights/Common/Stopwatch.h"
 #include "Insights/InsightsManager.h"
-#include "Insights/Log.h"
-#include "Insights/StoreService/StoreBrowser.h"
 #include "Insights/InsightsStyle.h"
+#include "Insights/Log.h"
 #include "Insights/Version.h"
+#include "Insights/Common/Stopwatch.h"
 #include "Insights/ImportTool/TableImportTool.h"
+#include "Insights/StoreService/StoreBrowser.h"
 #include "Insights/Widgets/SInsightsSettings.h"
 #include "Insights/Widgets/SLazyToolTip.h"
 
@@ -83,33 +86,50 @@ public:
 		{
 			TSharedPtr<SEditableTextBox> RenameTextBox;
 
-			TSharedRef<SWidget> Widget = SNew(SBox)
-				.Padding(FMargin(4.0f, 0.0f))
+			TSharedRef<SWidget> Widget = SNew(SHorizontalBox)
+			+ SHorizontalBox::Slot()
+			.AutoWidth()
+			[
+				SNew(SBox)
+				.WidthOverride(16.0)
+				.HeightOverride(16.0)
+				.Padding(FMargin(2.0))
 				[
-					SNew(SOverlay)
-
-					+ SOverlay::Slot()
-					[
-						SNew(STextBlock)
-						.Visibility_Lambda([this]() { return IsRenaming() ? EVisibility::Collapsed : EVisibility::Visible; })
-						.Text(this, &STraceListRow::GetTraceName)
-						.HighlightText(this, &STraceListRow::GetTraceNameHighlightText)
-						.HighlightColor(FLinearColor(0.75f, 0.75f, 0.75f, 1.0f))
-						.HighlightShape(FInsightsStyle::Get().GetBrush("DarkGreenBrush"))
-						.ToolTip(STraceListRow::GetTraceTooltip())
-					]
+					SNew(SImage)
+					.Image(FInsightsStyle::Get().GetBrush("Icons.UTrace"))
+					.ColorAndOpacity(this, &STraceListRow::GetColorForPath)
+				]
+				
+			]
+			
+			+ SHorizontalBox::Slot()
+			.Padding(4.0f, 0.0f)
+			.AutoWidth()
+			[
+				SNew(SOverlay)
+				
+				+ SOverlay::Slot()
+				[
+					SNew(STextBlock)
+					.Visibility_Lambda([this]() { return IsRenaming() ? EVisibility::Collapsed : EVisibility::Visible; })
+					.Text(this, &STraceListRow::GetTraceName)
+					.HighlightText(this, &STraceListRow::GetTraceNameHighlightText)
+					.HighlightColor(FLinearColor(0.75f, 0.75f, 0.75f, 1.0f))
+					.HighlightShape(FInsightsStyle::Get().GetBrush("DarkGreenBrush"))
+					.ToolTip(STraceListRow::GetTraceTooltip())
+				]
 
-					+ SOverlay::Slot()
-					[
-						SAssignNew(RenameTextBox, SEditableTextBox)
-						.Padding(FMargin(-1.0f, 0.0f, 0.0f, 0.0f))
-						.Visibility_Lambda([this]() { return IsRenaming() ? EVisibility::Visible : EVisibility::Collapsed; })
-						.Text(this, &STraceListRow::GetTraceName)
-						.ClearKeyboardFocusOnCommit(true)
-						.OnTextCommitted(this, &STraceListRow::RenameTextBox_OnValueCommitted)
-						.ToolTip(STraceListRow::GetTraceTooltip())
-					]
-				];
+				+ SOverlay::Slot()
+				[
+					SAssignNew(RenameTextBox, SEditableTextBox)
+					.Padding(0.0f)
+					.Visibility_Lambda([this]() { return IsRenaming() ? EVisibility::Visible : EVisibility::Collapsed; })
+					.Text(this, &STraceListRow::GetTraceName)
+					.ClearKeyboardFocusOnCommit(true)
+					.OnTextCommitted(this, &STraceListRow::RenameTextBox_OnValueCommitted)
+					.ToolTip(STraceListRow::GetTraceTooltip())
+				]
+			];
 
 			TSharedPtr<FTraceViewModel> TracePin = WeakTrace.Pin();
 			if (TracePin.IsValid())
@@ -362,6 +382,19 @@ public:
 		}
 	}
 
+	FSlateColor GetColorForPath() const
+	{
+		TSharedPtr<FTraceViewModel> TracePin = WeakTrace.Pin();
+		if (TracePin.IsValid())
+		{
+			return TracePin->DirectoryColor;
+		}
+		else
+		{
+			return FSlateColor::UseForeground();
+		}
+	}
+
 	FText GetTracePlatform() const
 	{
 		TSharedPtr<FTraceViewModel> TracePin = WeakTrace.Pin();
@@ -1260,6 +1293,205 @@ TSharedRef<SWidget> STraceStoreWindow::ConstructSessionsPanel()
 
 ////////////////////////////////////////////////////////////////////////////////////////////////////
 
+class STraceDirectoryItem : public SCompoundWidget
+{
+public:
+
+	SLATE_BEGIN_ARGS(STraceDirectoryItem) {}
+	SLATE_END_ARGS()
+
+	void Construct(const FArguments& InArgs, TSharedPtr<STraceStoreWindow::FTraceDirectoryModel> InClientItem);
+	
+private:
+	FSlateColor GetColor() const;
+	TSharedRef<SWidget> ConstructOperations();
+	FReply OnModifyStore();
+	FReply OnExplore();
+	FReply OnDelete();
+
+
+	TSharedPtr<STraceStoreWindow::FTraceDirectoryModel> Model;
+};
+
+////////////////////////////////////////////////////////////////////////////////////////////////////
+
+void STraceDirectoryItem::Construct(const FArguments& InArgs, TSharedPtr<STraceStoreWindow::FTraceDirectoryModel> InModel)
+{
+	Model = InModel;
+	
+	ChildSlot
+	[
+		SNew(SBorder)
+		.Padding(2.0f , 1.0f)
+		[
+			SNew(SHorizontalBox)
+
+			+ SHorizontalBox::Slot()
+			.Padding(4.0, 2.0)
+			.AutoWidth()
+			.VAlign(EVerticalAlignment::VAlign_Center)
+			[
+				SNew(SBox)
+				.HeightOverride(16.0)
+				.WidthOverride(16.0)
+				[
+					SNew(SImage)
+					.Image(FInsightsStyle::Get().GetBrush("Icons.UTraceStack"))
+					.ColorAndOpacity(GetColor())
+				]
+			]
+			
+			+ SHorizontalBox::Slot()
+			.FillWidth(1.0)
+			.Padding(4.0, 2.0)
+			.VAlign(EVerticalAlignment::VAlign_Center)
+			[
+				 SNew(STextBlock)
+				 .Text(FText::FromString(Model->Path))
+			]
+			
+			+ SHorizontalBox::Slot()
+			.Padding(4.0, 2.0)
+			.AutoWidth()
+			[
+				ConstructOperations()
+			]
+		]
+	];
+}
+
+////////////////////////////////////////////////////////////////////////////////////////////////////
+
+FSlateColor STraceDirectoryItem::GetColor() const
+{
+	if (Model && Model->Color != NAME_None)
+	{
+		return FAppStyle::Get().GetSlateColor(Model->Color);			
+	}
+	return FSlateColor::UseForeground();
+}
+
+////////////////////////////////////////////////////////////////////////////////////////////////////
+
+TSharedRef<SWidget> STraceDirectoryItem::ConstructOperations() 
+{
+	TSharedRef<SHorizontalBox> Box = SNew(SHorizontalBox);
+	
+	if (Model && EnumHasAllFlags(Model->Operations, ETraceDirOperations::ModifyStore) )
+	{
+		Box->AddSlot()
+			.AutoWidth()
+			[
+				SNew(SButton)
+				.ButtonStyle(&FAppStyle::Get().GetWidgetStyle<FButtonStyle>("SimpleButton"))
+				.ToolTipText(LOCTEXT("SetTraceStoreDirButtonToolTip", "Set the Trace Store Directory"))
+				.OnClicked_Raw(this, &STraceDirectoryItem::OnModifyStore)
+				[
+					SNew(SImage)
+					.Image(FInsightsStyle::Get().GetBrush("Icons.Edit"))
+					.ColorAndOpacity(FSlateColor::UseForeground())
+				]
+			 ];
+	}
+	if (Model && EnumHasAllFlags(Model->Operations, ETraceDirOperations::Delete) )
+	{
+		Box->AddSlot()
+			.AutoWidth()
+			[
+				 SNew(SButton)
+				 .ButtonStyle(&FAppStyle::Get().GetWidgetStyle<FButtonStyle>("SimpleButton"))
+				 .ToolTipText(LOCTEXT("WatchDirsRemoveTooltip", "Remove watch directory (files will not be deleted)"))
+				 .OnClicked_Raw(this, &STraceDirectoryItem::OnDelete)
+				 [
+					  SNew(SImage)
+					  .Image(FInsightsStyle::Get().GetBrush("Icons.RemoveWatchDir"))
+					  .ColorAndOpacity(FSlateColor::UseForeground())
+				 ]
+			 ];
+	}
+	if (Model && EnumHasAllFlags(Model->Operations, ETraceDirOperations::Explore))
+	{
+		Box->AddSlot()
+			.AutoWidth()
+			[
+				SNew(SButton)
+				.ButtonStyle(&FAppStyle::Get().GetWidgetStyle<FButtonStyle>("SimpleButton"))
+				.ToolTipText(LOCTEXT("ExploreTraceStoreDirButtonToolTip", "Explore the Trace Store Directory"))
+				.OnClicked_Raw(this, &STraceDirectoryItem::OnExplore)
+				[
+					SNew(SImage)
+					.Image(FInsightsStyle::Get().GetBrush("Icons.FolderExplore"))
+					.ColorAndOpacity(FSlateColor::UseForeground())
+				]
+			];
+	}
+
+	return Box;
+}
+
+////////////////////////////////////////////////////////////////////////////////////////////////////
+
+FReply STraceDirectoryItem::OnModifyStore() 
+{
+	if (IDesktopPlatform* DesktopPlatform = FDesktopPlatformModule::Get())
+	{
+		const FString Title = LOCTEXT("SetTraceStoreDirectorySelectPopupTitle", "Set Trace store to ...").ToString();
+
+		FString CurrentDirectory = FInsightsManager::Get()->GetStoreDir();
+		FString SelectedDirectory;
+		const bool bHasSelected = DesktopPlatform->OpenDirectoryDialog(
+			FSlateApplication::Get().FindBestParentWindowHandleForDialogs(AsShared()),
+			Title,
+			CurrentDirectory,
+			SelectedDirectory
+		);
+
+		if (bHasSelected)
+		{
+			FPaths::MakePlatformFilename(SelectedDirectory);
+			FPaths::MakePlatformFilename(CurrentDirectory);
+			TArray<FString> AddWatchDirs = { CurrentDirectory };
+			if (!FInsightsManager::Get()->GetStoreClient()->SetStoreDirectories(*SelectedDirectory, AddWatchDirs, TArray<FString>()))
+			{
+				FMessageLog(FInsightsManager::Get()->GetLogListingName()).Error(LOCTEXT("StoreCommunicationFail", "Failed to change settings on the store service."));
+			}
+		}
+	}
+	return FReply::Handled();
+}
+
+////////////////////////////////////////////////////////////////////////////////////////////////////
+
+FReply STraceDirectoryItem::OnExplore()
+{
+	if (Model)
+	{
+		FString FullPath(FPaths::ConvertRelativePathToFull(Model->Path));
+		FPlatformProcess::ExploreFolder(*FullPath);
+	}
+	return FReply::Handled();
+}
+
+////////////////////////////////////////////////////////////////////////////////////////////////////
+
+FReply STraceDirectoryItem::OnDelete()
+{
+	if (Model)
+	{
+		const EAppReturnType::Type bConfirmed = FMessageDialog::Open(EAppMsgType::OkCancel, FText::Format(
+			LOCTEXT("WatchDirRemoveConfirmBody", "This will remove {0} from watch directories?"),
+			LOCTEXT("WatchDirRemoveConfirmTitle", "Confirm removing watch directory")));
+		if (bConfirmed == EAppReturnType::Ok)
+		{
+			const TCHAR* Path = *Model->Path;
+			FInsightsManager::Get()->GetStoreClient()->SetStoreDirectories(nullptr, {}, { Path });
+		}
+	}
+	return FReply::Handled();
+}
+
+////////////////////////////////////////////////////////////////////////////////////////////////////
+
 TSharedRef<SWidget> STraceStoreWindow::ConstructLoadPanel()
 {
 	TSharedRef<SWidget> Widget = SNew(SHorizontalBox)
@@ -1315,51 +1547,173 @@ TSharedRef<SWidget> STraceStoreWindow::ConstructTraceStoreDirectoryPanel()
 {
 	TSharedRef<SWidget> Widget =
 
-		SNew(SHorizontalBox)
+		SNew(SVerticalBox)
 
-		+ SHorizontalBox::Slot()
-		.AutoWidth()
-		.Padding(0.0f, 0.0f, 4.0f, 0.0f)
-		.VAlign(VAlign_Center)
+		+ SVerticalBox::Slot()
+		.AutoHeight()
+		.Padding(.0f, .0f)
+		.HAlign(HAlign_Left)
 		[
-			 SNew(SImage)
-			 .Image_Raw(this, &STraceStoreWindow::GetConnectionStatusIcon)
-			 .ToolTip(SNew(SToolTip).Text_Raw(this, &STraceStoreWindow::GetConnectionStatusTooltip))
-		]
+			SNew(SHorizontalBox)
+			
+			+ SHorizontalBox::Slot()
+			.AutoWidth()
+			.Padding(0.0f, 0.0f, 0.0f, 0.0f)
+			.VAlign(VAlign_Center)
+			[
+				SNew(SButton)
+				.ButtonStyle(&FAppStyle::Get().GetWidgetStyle<FButtonStyle>("SimpleButton"))
+				.ToolTipText(LOCTEXT("ManageWatchDirsTooltip", "Manage store settings."))
+				.OnClicked(this, &STraceStoreWindow::StoreSettingsArea_Toggle)
+				[
+					SNew(SImage)
+					.Image_Raw(this, &STraceStoreWindow::StoreSettingsToggle_Icon)
+					.ColorAndOpacity(FSlateColor::UseForeground())
+				]
+			]
 
-		+ SHorizontalBox::Slot()
-		.AutoWidth()
-		.Padding(0.0f, 0.0f)
-		.VAlign(VAlign_Center)
-		[
-			SNew(STextBlock)
-			.Text(LOCTEXT("TraceStoreDirText", "Trace Store Directory"))
-		]
+			+ SHorizontalBox::Slot()
+			.AutoWidth()
+			.Padding(4.0f, 0.0f, 4.0f, 0.0f)
+			.VAlign(VAlign_Center)
+			[
+				SNew(SImage)
+				.Image_Raw(this, &STraceStoreWindow::GetConnectionStatusIcon)
+				.ToolTip(SNew(SToolTip)
+				.Text_Raw(this, &STraceStoreWindow::GetConnectionStatusTooltip))
+			]
+				
+			+ SHorizontalBox::Slot()
+			.AutoWidth()
+			.Padding(4.0f, 0.0f, 0.0f, 0.0f)
+			.VAlign(VAlign_Center)
+			[
+				SNew(STextBlock)
+				.Text(LOCTEXT("StoreHostText", "Store Host:"))
+			]
 
-		+ SHorizontalBox::Slot()
-		.FillWidth(1.0f)
-		.Padding(4.0f, 0.0f, 0.0f, 0.0f)
-		.VAlign(VAlign_Center)
-		[
-			SNew(SEditableTextBox)
-			.IsReadOnly(true)
-			.BackgroundColor(FSlateColor(EStyleColor::Background))
-			.Text(this, &STraceStoreWindow::GetTraceStoreDirectory)
+			+ SHorizontalBox::Slot()
+			.Padding(4.0f, 0.0f, 0.0f, 0.0f)
+			.VAlign(VAlign_Center)
+			[
+				SAssignNew(StoreHostTextBox, SEditableTextBox)
+				.IsReadOnly(true)
+				.BackgroundColor(FSlateColor(EStyleColor::Background))
+			]
+
+			+ SHorizontalBox::Slot()
+			.AutoWidth()
+			.Padding(10.0f, 0.0f, 0.0f, 0.0f)
+			.VAlign(VAlign_Center)
+			[
+				SNew(STextBlock)
+				.Text(LOCTEXT("TraceStoreDirText", "Directory: "))
+			]
+
+			+ SHorizontalBox::Slot()
+			//.FillWidth(1.0f)
+			.AutoWidth()
+			.Padding(4.0f, 0.0f, 0.0f, 0.0f)
+			.VAlign(VAlign_Center)
+			[
+				SAssignNew(StoreDirTextBox, SEditableTextBox)
+				.IsReadOnly(true)
+				.BackgroundColor(FSlateColor(EStyleColor::Background))
+				.Text(this, &STraceStoreWindow::GetTraceStoreDirectory)
+			]
+
+			+ SHorizontalBox::Slot()
+			.AutoWidth()
+			.Padding(4.0f, 0.0f, 0.0f, 0.0f)
+			.VAlign(VAlign_Center)
+			[
+				SNew(SButton)
+				.ButtonStyle(&FAppStyle::Get().GetWidgetStyle<FButtonStyle>("SimpleButton"))
+				.ToolTipText(LOCTEXT("ExploreTraceStoreDirButtonToolTip", "Explore the Trace Store Directory"))
+				.OnClicked(this, &STraceStoreWindow::ExploreTraceStoreDirectory_OnClicked)
+				[
+					SNew(SImage)
+					.Image(FInsightsStyle::Get().GetBrush("Icons.FolderExplore"))
+					.ColorAndOpacity(FSlateColor::UseForeground())
+				]
+			]
 		]
 
-		+ SHorizontalBox::Slot()
-		.AutoWidth()
-		.Padding(4.0f, 0.0f, 0.0f, 0.0f)
-		.VAlign(VAlign_Center)
+		+ SVerticalBox::Slot()
+		.AutoHeight()
+		.Padding(.0f, 8.0f)
+		.HAlign(HAlign_Left)
 		[
-			SNew(SButton)
-			.ButtonStyle(&FAppStyle::Get().GetWidgetStyle<FButtonStyle>("SimpleButton"))
-			.ToolTipText(LOCTEXT("ExploreTraceStoreDirButtonToolTip", "Explore the Trace Store Directory"))
-			.OnClicked(this, &STraceStoreWindow::ExploreTraceStoreDirectory_OnClicked)
+			SAssignNew(StoreSettingsArea, SVerticalBox)
+			.Visibility(EVisibility::Collapsed)
+
+			+ SVerticalBox::Slot()
+			.AutoHeight()
+			.Padding(0.0f, 4.0f)
+			.VAlign(VAlign_Center)
 			[
-				SNew(SImage)
-				.Image(FInsightsStyle::Get().GetBrush("Icons.FolderExplore"))
-				.ColorAndOpacity(FSlateColor::UseForeground())
+				SNew(STextBlock)
+				.Text(LOCTEXT("StoreDirLabel", "Trace store directory. New traces will be stored here:"))
+			]
+			
+			+ SVerticalBox::Slot()
+			.AutoHeight()
+			.Padding(0.0f, 4.0f)
+			[
+				SAssignNew(StoreDirListView, SListView<TSharedPtr<FTraceDirectoryModel>>)
+				.ListItemsSource(&StoreDirectoryModel)
+				.OnGenerateRow(this, &STraceStoreWindow::TraceDirs_OnGenerateRow)
+				.SelectionMode(ESelectionMode::None)
+			]
+
+			+ SVerticalBox::Slot()
+			.AutoHeight()
+			.Padding(0.0f, 4.0f)
+			.VAlign(VAlign_Center)
+			[
+				SNew(STextBlock)
+				.Text(LOCTEXT("WatchDirsLabel", "Additional directories to monitor for traces:"))
+			]
+			
+			+ SVerticalBox::Slot()
+			.AutoHeight()
+			.Padding(0.0f, 4.0f)
+			[
+				SAssignNew(WatchDirsListView, SListView<TSharedPtr<FTraceDirectoryModel>>)
+				.ListItemsSource(&WatchDirectoriesModel)
+				.OnGenerateRow(this, &STraceStoreWindow::TraceDirs_OnGenerateRow)
+				//.Orientation(Orient_Horizontal)
+				.SelectionMode(ESelectionMode::None)
+			]
+
+			+ SVerticalBox::Slot()
+			.AutoHeight()
+			.Padding(0.0f, 4.0f)
+			.HAlign(HAlign_Left)
+			[
+				SNew(SButton)
+				.ButtonStyle(&FAppStyle::Get().GetWidgetStyle<FButtonStyle>("Button"))
+				.ToolTipText(LOCTEXT("WatchDirsAddTooltip", "Add additonal watch directory..."))
+				.OnClicked(this, &STraceStoreWindow::AddWatchDir_Clicked)
+				[
+					SNew(SHorizontalBox)
+
+					+ SHorizontalBox::Slot()
+					.AutoWidth()
+					[
+						SNew(SImage)
+						.Image(FInsightsStyle::Get().GetBrush("Icons.AddWatchDir"))
+						.ColorAndOpacity(FSlateColor::UseForeground())
+					]
+
+					+ SHorizontalBox::Slot()
+					.Padding(4.0, 0.0, 0.0, 0.0)
+					.AutoWidth()
+					[
+						SNew(STextBlock)
+						.Text(FText::FromStringView(TEXTVIEW("Add directory")))
+					]
+				]
 			]
 		];
 
@@ -1666,6 +2020,22 @@ FReply STraceStoreWindow::RefreshTraces_OnClicked()
 
 ////////////////////////////////////////////////////////////////////////////////////////////////////
 
+FSlateColor STraceStoreWindow::GetColorByPath(const FString& Uri)
+{
+	const TSharedPtr<FTraceDirectoryModel>* Dir = WatchDirectoriesModel.FindByPredicate([&](const TSharedPtr<FTraceDirectoryModel>& Dir)
+	{
+		return FPaths::IsUnderDirectory(Uri, Dir->Path);
+	});
+	if (Dir)
+	{
+		return FAppStyle::Get().GetSlateColor((*Dir)->Color);
+	}
+	// If this is default trace store directory, use forground
+	return FSlateColor(FSlateColor::UseForeground());
+}
+
+////////////////////////////////////////////////////////////////////////////////////////////////////
+
 void STraceStoreWindow::RefreshTraceList()
 {
 	FStopwatch StopwatchTotal;
@@ -1678,6 +2048,48 @@ void STraceStoreWindow::RefreshTraceList()
 	{
 		StoreBrowser->Lock();
 
+		const uint32 NewSettingsSerial = StoreBrowser->GetLockedSettingsSerial();
+		if (NewSettingsSerial != SettingsChangeSerial)
+		{
+			SettingsChangeSerial = NewSettingsSerial;
+
+			// Update the host text
+			StoreHostTextBox->SetText(FText::FromString(StoreBrowser->GetLockedHost()));
+
+			// Update the store text box
+			StoreDirTextBox->SetText(FText::FromString(StoreBrowser->GetLockedStoreDirectory()));
+
+			// Update store directory model
+			StoreDirectoryModel.Empty(1);
+			StoreDirectoryModel.Push(MakeShared<FTraceDirectoryModel>(
+				FString(StoreBrowser->GetLockedStoreDirectory()),
+				NAME_None,
+				ETraceDirOperations::ModifyStore|ETraceDirOperations::Explore
+			));
+			StoreDirListView->RequestListRefresh();
+
+			// Update additional watch directory models
+			static const FName DirColor[] = {
+				FName("Colors.AccentBlue"),
+				FName("Colors.AccentGreen"),
+				FName("Colors.AccentYellow"),
+				FName("Colors.AccentOrange"),
+				FName("Colors.AccentRed")
+			};
+			uint64 ColorIdx = 0;
+			WatchDirectoriesModel.Empty();
+			for (const auto& Dir : StoreBrowser->GetLockedWatchDirectories())
+			{
+				WatchDirectoriesModel.Emplace(MakeShared<FTraceDirectoryModel>(
+					FString(Dir),
+					DirColor[ColorIdx],
+					ETraceDirOperations::Delete|ETraceDirOperations::Explore
+				));
+				ColorIdx = FMath::Wrap(ColorIdx + 1, uint64(0u), uint64(UE_ARRAY_COUNT(DirColor)));
+			}
+			WatchDirsListView->RequestListRefresh();
+		}
+		
 		const uint64 NewChangeSerial = StoreBrowser->GetLockedTracesChangeSerial();
 		if (NewChangeSerial != TracesChangeSerial)
 		{
@@ -1730,9 +2142,10 @@ void STraceStoreWindow::RefreshTraceList()
 					FTraceViewModel& Trace = *TracePtr;
 					Trace.TraceId = SourceTrace.TraceId;
 					Trace.Name = FText::FromString(SourceTrace.Name);
-					FString Uri = FInsightsManager::Get()->GetStoreDir() + TEXT("/") + SourceTrace.Name + TEXT(".utrace");
+					FString Uri = SourceTrace.Uri;
 					FPaths::NormalizeFilename(Uri);
 					Trace.Uri = FText::FromString(Uri);
+					Trace.DirectoryColor = GetColorByPath(Uri);
 					UpdateTrace(Trace, SourceTrace);
 					TraceViewModels.Add(TracePtr);
 					TraceViewModelMap.Add(SourceTrace.TraceId, TracePtr);
@@ -2466,6 +2879,78 @@ FReply STraceStoreWindow::ExploreTraceStoreDirectory_OnClicked()
 
 ////////////////////////////////////////////////////////////////////////////////////////////////////
 
+TSharedRef<ITableRow> STraceStoreWindow::TraceDirs_OnGenerateRow(TSharedPtr<FTraceDirectoryModel> Item, const TSharedRef<STableViewBase>& Owner) const
+{
+	return SNew(STableRow<TSharedPtr<FTraceDirectoryModel>>, Owner)
+		.Content()
+		[
+			SNew(STraceDirectoryItem, Item)
+		];
+}
+
+////////////////////////////////////////////////////////////////////////////////////////////////////
+
+FReply STraceStoreWindow::StoreSettingsArea_Toggle() const
+{
+	if (StoreSettingsArea)
+	{
+		if(StoreSettingsArea->GetVisibility() == EVisibility::Visible)
+		{
+			StoreSettingsArea->SetVisibility(EVisibility::Collapsed);
+		}
+		else
+		{
+			StoreSettingsArea->SetVisibility(EVisibility::Visible);
+		}
+	}
+	return FReply::Handled();
+}
+
+////////////////////////////////////////////////////////////////////////////////////////////////////
+
+const FSlateBrush* STraceStoreWindow::StoreSettingsToggle_Icon() const
+{
+	if (StoreSettingsArea)
+	{
+		if(StoreSettingsArea->GetVisibility() == EVisibility::Visible)
+		{
+			return FInsightsStyle::Get().GetBrush("Icons.Expanded");
+		}
+	}
+	return FInsightsStyle::Get().GetBrush("Icons.Expand");
+}
+
+////////////////////////////////////////////////////////////////////////////////////////////////////
+
+FReply STraceStoreWindow::AddWatchDir_Clicked()
+{
+	if (IDesktopPlatform* DesktopPlatform = FDesktopPlatformModule::Get())
+	{
+		const FString Title = LOCTEXT("AddWatchDirectorySelectPopupTitle", "Add watch directory ...").ToString();
+
+		const FString StoreDirectory = FInsightsManager::Get()->GetStoreDir();
+		FString SelectedDirectory;
+		const bool bHasSelected = DesktopPlatform->OpenDirectoryDialog(
+			FSlateApplication::Get().FindBestParentWindowHandleForDialogs(AsShared()),
+			Title,
+			StoreDirectory,
+			SelectedDirectory
+		);
+
+		if (bHasSelected && !SelectedDirectory.Equals(StoreDirectory))
+		{
+			FPaths::MakePlatformFilename(SelectedDirectory);
+			if (!FInsightsManager::Get()->GetStoreClient()->SetStoreDirectories(nullptr, { (*SelectedDirectory) }, {}))
+			{
+				FMessageLog(FInsightsManager::Get()->GetLogListingName()).Error(LOCTEXT("StoreCommunicationFail", "Failed to change settings on the store service."));
+			}
+		}
+	}
+	return FReply::Handled();
+}
+
+////////////////////////////////////////////////////////////////////////////////////////////////////
+
 void STraceStoreWindow::OpenSettings()
 {
 	MainContentPanel->SetEnabled(false);
diff --git a/Engine/Source/Developer/TraceInsights/Private/Insights/Widgets/SStartPageWindow.h b/Engine/Source/Developer/TraceInsights/Private/Insights/Widgets/SStartPageWindow.h
--- a/Engine/Source/Developer/TraceInsights/Private/Insights/Widgets/SStartPageWindow.h
+++ b/Engine/Source/Developer/TraceInsights/Private/Insights/Widgets/SStartPageWindow.h
@@ -53,6 +53,17 @@ typedef TWeakPtr<class SNotificationItem> SNotificationItemWeak;
 
 ////////////////////////////////////////////////////////////////////////////////////////////////////
 
+enum class ETraceDirOperations : uint8
+{
+	None		= 0,
+	ModifyStore	= 1 << 0,
+	Delete		= 1 << 1,
+	Explore		= 1 << 2,
+};
+ENUM_CLASS_FLAGS(ETraceDirOperations);
+
+////////////////////////////////////////////////////////////////////////////////////////////////////
+
 struct FTraceViewModel
 {
 	uint32 TraceId = 0;
@@ -62,6 +73,7 @@ struct FTraceViewModel
 
 	FText Name;
 	FText Uri;
+	FSlateColor DirectoryColor;
 
 	FDateTime Timestamp = 0;
 	uint64 Size = 0;
@@ -233,6 +245,7 @@ protected:
 class STraceStoreWindow : public SCompoundWidget
 {
 	friend class STraceListRow;
+	friend class STraceDirectoryItem;
 
 public:
 	/** Default constructor. */
@@ -329,6 +342,7 @@ private:
 	void BuildBranchFilterSubMenu(FMenuBuilder& InMenuBuilder);
 
 	FReply RefreshTraces_OnClicked();
+	FSlateColor GetColorByPath(const FString& Uri);
 
 	void RefreshTraceList();
 	void UpdateTrace(FTraceViewModel& InOutTrace, const Insights::FStoreBrowserTraceInfo& InSourceTrace);
@@ -349,6 +363,34 @@ private:
 	FText GetTraceStoreDirectory() const;
 	FReply ExploreTraceStoreDirectory_OnClicked();
 
+	//////////////////////////////////////////////////
+	// Store settings
+
+	/** View model for a source of traces. This could be the trace store default directory or a watch directory.*/
+	struct FTraceDirectoryModel
+	{
+		
+		FTraceDirectoryModel(FString&& InPath, const FName& InColor, ETraceDirOperations InOperations)
+			: Path(MoveTemp(InPath)), Color(InColor), Operations(InOperations) {}
+
+		/** Path to directory */
+		FString		Path;
+		/** Assigned color */
+		FName		Color;
+		/** Supported operations */
+		ETraceDirOperations	Operations;
+	};
+	
+	TSharedRef<ITableRow> TraceDirs_OnGenerateRow(TSharedPtr<FTraceDirectoryModel> Item, const TSharedRef<STableViewBase>& Owner) const;
+	FReply StoreSettingsArea_Toggle() const;
+	const FSlateBrush* StoreSettingsToggle_Icon() const;
+	FReply AddWatchDir_Clicked();
+
+	TArray<TSharedPtr<FTraceDirectoryModel>> StoreDirectoryModel;
+	TArray<TSharedPtr<FTraceDirectoryModel>> WatchDirectoriesModel;
+	
+	friend class STraceWatchDirTableRow;
+
 	//////////////////////////////////////////////////
 
 	/** Updates the amount of time the profiler has been active. */
@@ -464,11 +506,17 @@ private:
 
 	TUniquePtr<Insights::FStoreBrowser> StoreBrowser;
 	uint64 TracesChangeSerial;
+	uint32 SettingsChangeSerial;
 
 	TArray<TSharedPtr<FTraceViewModel>> TraceViewModels; // all available trace view models
 	TArray<TSharedPtr<FTraceViewModel>> FilteredTraceViewModels; // the filtered list of trace view models
 	TMap<uint32, TSharedPtr<FTraceViewModel>> TraceViewModelMap;
 
+	TSharedPtr<SEditableTextBox> StoreDirTextBox;
+	TSharedPtr<SEditableTextBox> StoreHostTextBox;
+	TSharedPtr<STableViewBase> StoreDirListView;
+	TSharedPtr<SVerticalBox> StoreSettingsArea;
+	TSharedPtr<STableViewBase> WatchDirsListView;
 	TSharedPtr<SListView<TSharedPtr<FTraceViewModel>>> TraceListView;
 	TSharedPtr<FTraceViewModel> SelectedTrace;
 	bool bIsUserSelectedTrace;
