[Insights]
  - SessionBrowser: Restored Uri functionality for traces (backward compatibility with current/older UTS).
  - SessionBrowser: Disabled tooltip when pressing button to add new watch directory.
  - SessionBrowser: Disabled a trace session after it was renamed, until the trace id is properly updated (~3sec); the trace name changes color to grey.
  - SessionBrowser: Fixed focus for the trace sessions list (initial focus and after renaming a trace).

diff --git a/Engine/Source/Developer/TraceInsights/Private/Insights/Widgets/SStartPageWindow.cpp b/Engine/Source/Developer/TraceInsights/Private/Insights/Widgets/SStartPageWindow.cpp
--- a/Engine/Source/Developer/TraceInsights/Private/Insights/Widgets/SStartPageWindow.cpp
+++ b/Engine/Source/Developer/TraceInsights/Private/Insights/Widgets/SStartPageWindow.cpp
@@ -99,20 +99,21 @@ public:
 					.Image(FInsightsStyle::Get().GetBrush("Icons.UTrace"))
 					.ColorAndOpacity(this, &STraceListRow::GetColorForPath)
 				]
-				
+
 			]
-			
+
 			+ SHorizontalBox::Slot()
 			.Padding(4.0f, 0.0f)
 			.AutoWidth()
 			[
 				SNew(SOverlay)
-				
+
 				+ SOverlay::Slot()
 				[
 					SNew(STextBlock)
 					.Visibility_Lambda([this]() { return IsRenaming() ? EVisibility::Collapsed : EVisibility::Visible; })
 					.Text(this, &STraceListRow::GetTraceName)
+					.ColorAndOpacity(this, &STraceListRow::GetTraceTextColor)
 					.HighlightText(this, &STraceListRow::GetTraceNameHighlightText)
 					.HighlightColor(FLinearColor(0.75f, 0.75f, 0.75f, 1.0f))
 					.HighlightShape(FInsightsStyle::Get().GetBrush("DarkGreenBrush"))
@@ -125,6 +126,7 @@ public:
 					.Padding(0.0f)
 					.Visibility_Lambda([this]() { return IsRenaming() ? EVisibility::Visible : EVisibility::Collapsed; })
 					.Text(this, &STraceListRow::GetTraceName)
+					.RevertTextOnEscape(true)
 					.ClearKeyboardFocusOnCommit(true)
 					.OnTextCommitted(this, &STraceListRow::RenameTextBox_OnValueCommitted)
 					.ToolTip(STraceListRow::GetTraceTooltip())
@@ -244,6 +246,12 @@ public:
 
 	void Rename(FTraceViewModel& Trace, const FText& InText)
 	{
+		TSharedPtr<STraceStoreWindow> TraceStoreWindow = WeakParentWidget.Pin();
+		if (!TraceStoreWindow.IsValid())
+		{
+			return;
+		}
+
 		const FString TraceName = Trace.Name.ToString();
 		const FString NewTraceName = InText.ToString().TrimStartAndEnd();
 
@@ -255,7 +263,7 @@ public:
 		if (Trace.bIsLive)
 		{
 			FText Message = LOCTEXT("RenameLive", "Cannot rename a live session!");
-			WeakParentWidget.Pin()->ShowFailMessage(Message);
+			TraceStoreWindow->ShowFailMessage(Message);
 			return;
 		}
 
@@ -278,7 +286,7 @@ public:
 		if (bHasReservedCharacters || !FPaths::ValidatePath(NewTraceFile, &Reason))
 		{
 			FText Message = FText::Format(LOCTEXT("RenameFailFmt3", "Failed to rename \"{0}\" to \"{1}\"!\n{2}"), Trace.Name, FText::FromString(NewTraceName), Reason);
-			WeakParentWidget.Pin()->ShowFailMessage(Message);
+			TraceStoreWindow->ShowFailMessage(Message);
 			return;
 		}
 
@@ -287,7 +295,7 @@ public:
 			UE_LOG(TraceInsights, Warning, TEXT("[TraceStore] Failed to rename \"%s\" to \"%s\"! File already exists."), *TraceName, *NewTraceName);
 
 			FText Message = FText::Format(LOCTEXT("RenameFailFmt1", "Failed to rename \"{0}\" to \"{1}\"!\nFile already exists."), Trace.Name, FText::FromString(NewTraceName));
-			WeakParentWidget.Pin()->ShowFailMessage(Message);
+			TraceStoreWindow->ShowFailMessage(Message);
 			return;
 		}
 
@@ -297,13 +305,15 @@ public:
 			UE_LOG(TraceInsights, Warning, TEXT("[TraceStore] Failed to rename \"%s\" to \"%s\"!"), *TraceName, *NewTraceName);
 
 			FText Message = FText::Format(LOCTEXT("RenameFailFmt2", "Failed to rename \"{0}\" to \"{1}\"!"), Trace.Name, FText::FromString(NewTraceName));
-			WeakParentWidget.Pin()->ShowFailMessage(Message);
+			TraceStoreWindow->ShowFailMessage(Message);
 			return;
 		}
 
 		UE_LOG(TraceInsights, Verbose, TEXT("[TraceStore] Renamed utrace file (\"%s\")."), *NewTraceFile);
 		Trace.Name = FText::FromString(NewTraceName);
 		Trace.Uri = FText::FromString(NewTraceFile);
+		Trace.bWasJustRenamed = true; // cannot be open until its TraceId is not updated
+		Trace.ChangeSerial = 0; // to force update
 
 		const FString CacheFile = FPaths::ChangeExtension(TraceFile, TEXT("ucache"));
 		if (FPaths::FileExists(CacheFile))
@@ -316,7 +326,9 @@ public:
 		}
 
 		FText Message = FText::Format(LOCTEXT("RenameSuccessFmt", "Renamed \"{0}\" to \"{1}\"."), FText::FromString(TraceName), FText::FromString(NewTraceName));
-		WeakParentWidget.Pin()->ShowSuccessMessage(Message);
+		TraceStoreWindow->ShowSuccessMessage(Message);
+
+		TraceStoreWindow->bSetKeyboardFocusOnNextTick = true;
 	}
 
 	FText GetTraceIndexAndId() const
@@ -346,6 +358,19 @@ public:
 		}
 	}
 
+	FSlateColor GetTraceTextColor() const
+	{
+		TSharedPtr<FTraceViewModel> TracePin = WeakTrace.Pin();
+		if (TracePin.IsValid())
+		{
+			if (TracePin->bWasJustRenamed)
+			{
+				return FSlateColor(EStyleColor::White25);
+			}
+		}
+		return FSlateColor(EStyleColor::Foreground);
+	}
+
 	FText GetTraceNameHighlightText() const
 	{
 		TSharedPtr<STraceStoreWindow> ParentWidgetPin = WeakParentWidget.Pin();
@@ -950,7 +975,7 @@ void STraceStoreWindow::Construct(const FArguments& InArgs)
 			+ SVerticalBox::Slot()
 			.HAlign(HAlign_Fill)
 			.AutoHeight()
-			.Padding(12.0f, 8.0f, 12.0f, 0.0f)
+			.Padding(6.0f, 8.0f, 12.0f, 0.0f)
 			[
 				ConstructTraceStoreDirectoryPanel()
 			]
@@ -1032,8 +1057,7 @@ void STraceStoreWindow::Construct(const FArguments& InArgs)
 
 	RefreshTraceList();
 
-	FSlateApplication::Get().SetKeyboardFocus(TraceListView);
-	FSlateApplication::Get().SetUserFocus(0, TraceListView);
+	bSetKeyboardFocusOnNextTick = true;
 }
 
 ////////////////////////////////////////////////////////////////////////////////////////////////////
@@ -1301,7 +1325,7 @@ public:
 	SLATE_END_ARGS()
 
 	void Construct(const FArguments& InArgs, TSharedPtr<STraceStoreWindow::FTraceDirectoryModel> InClientItem);
-	
+
 private:
 	FSlateColor GetColor() const;
 	TSharedRef<SWidget> ConstructOperations();
@@ -1318,7 +1342,7 @@ private:
 void STraceDirectoryItem::Construct(const FArguments& InArgs, TSharedPtr<STraceStoreWindow::FTraceDirectoryModel> InModel)
 {
 	Model = InModel;
-	
+
 	ChildSlot
 	[
 		SNew(SBorder)
@@ -1340,7 +1364,7 @@ void STraceDirectoryItem::Construct(const FArguments& InArgs, TSharedPtr<STraceS
 					.ColorAndOpacity(GetColor())
 				]
 			]
-			
+
 			+ SHorizontalBox::Slot()
 			.FillWidth(1.0)
 			.Padding(4.0, 2.0)
@@ -1349,7 +1373,7 @@ void STraceDirectoryItem::Construct(const FArguments& InArgs, TSharedPtr<STraceS
 				 SNew(STextBlock)
 				 .Text(FText::FromString(Model->Path))
 			]
-			
+
 			+ SHorizontalBox::Slot()
 			.Padding(4.0, 2.0)
 			.AutoWidth()
@@ -1366,17 +1390,17 @@ FSlateColor STraceDirectoryItem::GetColor() const
 {
 	if (Model && Model->Color != NAME_None)
 	{
-		return FAppStyle::Get().GetSlateColor(Model->Color);			
+		return FAppStyle::Get().GetSlateColor(Model->Color);
 	}
 	return FSlateColor::UseForeground();
 }
 
 ////////////////////////////////////////////////////////////////////////////////////////////////////
 
-TSharedRef<SWidget> STraceDirectoryItem::ConstructOperations() 
+TSharedRef<SWidget> STraceDirectoryItem::ConstructOperations()
 {
 	TSharedRef<SHorizontalBox> Box = SNew(SHorizontalBox);
-	
+
 	if (Model && EnumHasAllFlags(Model->Operations, ETraceDirOperations::ModifyStore) )
 	{
 		Box->AddSlot()
@@ -1431,11 +1455,11 @@ TSharedRef<SWidget> STraceDirectoryItem::ConstructOperations()
 
 ////////////////////////////////////////////////////////////////////////////////////////////////////
 
-FReply STraceDirectoryItem::OnModifyStore() 
+FReply STraceDirectoryItem::OnModifyStore()
 {
 	if (IDesktopPlatform* DesktopPlatform = FDesktopPlatformModule::Get())
 	{
-		const FString Title = LOCTEXT("SetTraceStoreDirectorySelectPopupTitle", "Set Trace store to ...").ToString();
+		const FString Title = LOCTEXT("SetTraceStoreDirectorySelectPopupTitle", "Set Trace Store Directory").ToString();
 
 		FString CurrentDirectory = FInsightsManager::Get()->GetStoreDir();
 		FString SelectedDirectory;
@@ -1551,11 +1575,11 @@ TSharedRef<SWidget> STraceStoreWindow::ConstructTraceStoreDirectoryPanel()
 
 		+ SVerticalBox::Slot()
 		.AutoHeight()
-		.Padding(.0f, .0f)
-		.HAlign(HAlign_Left)
+		.Padding(0.0f, 0.0f)
+		.HAlign(HAlign_Fill)
 		[
 			SNew(SHorizontalBox)
-			
+
 			+ SHorizontalBox::Slot()
 			.AutoWidth()
 			.Padding(0.0f, 0.0f, 0.0f, 0.0f)
@@ -1563,7 +1587,7 @@ TSharedRef<SWidget> STraceStoreWindow::ConstructTraceStoreDirectoryPanel()
 			[
 				SNew(SButton)
 				.ButtonStyle(&FAppStyle::Get().GetWidgetStyle<FButtonStyle>("SimpleButton"))
-				.ToolTipText(LOCTEXT("ManageWatchDirsTooltip", "Manage store settings."))
+				.ToolTipText(LOCTEXT("ManageStoreSettingsTooltip", "Manage store settings."))
 				.OnClicked(this, &STraceStoreWindow::StoreSettingsArea_Toggle)
 				[
 					SNew(SImage)
@@ -1582,7 +1606,7 @@ TSharedRef<SWidget> STraceStoreWindow::ConstructTraceStoreDirectoryPanel()
 				.ToolTip(SNew(SToolTip)
 				.Text_Raw(this, &STraceStoreWindow::GetConnectionStatusTooltip))
 			]
-				
+
 			+ SHorizontalBox::Slot()
 			.AutoWidth()
 			.Padding(4.0f, 0.0f, 0.0f, 0.0f)
@@ -1594,6 +1618,7 @@ TSharedRef<SWidget> STraceStoreWindow::ConstructTraceStoreDirectoryPanel()
 
 			+ SHorizontalBox::Slot()
 			.Padding(4.0f, 0.0f, 0.0f, 0.0f)
+			.AutoWidth()
 			.VAlign(VAlign_Center)
 			[
 				SAssignNew(StoreHostTextBox, SEditableTextBox)
@@ -1603,16 +1628,15 @@ TSharedRef<SWidget> STraceStoreWindow::ConstructTraceStoreDirectoryPanel()
 
 			+ SHorizontalBox::Slot()
 			.AutoWidth()
-			.Padding(10.0f, 0.0f, 0.0f, 0.0f)
+			.Padding(8.0f, 0.0f, 0.0f, 0.0f)
 			.VAlign(VAlign_Center)
 			[
 				SNew(STextBlock)
-				.Text(LOCTEXT("TraceStoreDirText", "Directory: "))
+				.Text(LOCTEXT("TraceStoreDirText", "Directory:"))
 			]
 
 			+ SHorizontalBox::Slot()
-			//.FillWidth(1.0f)
-			.AutoWidth()
+			.FillWidth(1.0f)
 			.Padding(4.0f, 0.0f, 0.0f, 0.0f)
 			.VAlign(VAlign_Center)
 			[
@@ -1655,7 +1679,7 @@ TSharedRef<SWidget> STraceStoreWindow::ConstructTraceStoreDirectoryPanel()
 				SNew(STextBlock)
 				.Text(LOCTEXT("StoreDirLabel", "Trace store directory. New traces will be stored here:"))
 			]
-			
+
 			+ SVerticalBox::Slot()
 			.AutoHeight()
 			.Padding(0.0f, 4.0f)
@@ -1674,7 +1698,7 @@ TSharedRef<SWidget> STraceStoreWindow::ConstructTraceStoreDirectoryPanel()
 				SNew(STextBlock)
 				.Text(LOCTEXT("WatchDirsLabel", "Additional directories to monitor for traces:"))
 			]
-			
+
 			+ SVerticalBox::Slot()
 			.AutoHeight()
 			.Padding(0.0f, 4.0f)
@@ -1693,7 +1717,7 @@ TSharedRef<SWidget> STraceStoreWindow::ConstructTraceStoreDirectoryPanel()
 			[
 				SNew(SButton)
 				.ButtonStyle(&FAppStyle::Get().GetWidgetStyle<FButtonStyle>("Button"))
-				.ToolTipText(LOCTEXT("WatchDirsAddTooltip", "Add additonal watch directory..."))
+				.ToolTipText(LOCTEXT("WatchDirsAddTooltip", "Add additional watch directory..."))
 				.OnClicked(this, &STraceStoreWindow::AddWatchDir_Clicked)
 				[
 					SNew(SHorizontalBox)
@@ -1711,7 +1735,7 @@ TSharedRef<SWidget> STraceStoreWindow::ConstructTraceStoreDirectoryPanel()
 					.AutoWidth()
 					[
 						SNew(STextBlock)
-						.Text(FText::FromStringView(TEXTVIEW("Add directory")))
+						.Text(FText::FromStringView(TEXTVIEW("Add Directory...")))
 					]
 				]
 			]
@@ -1789,12 +1813,12 @@ END_SLATE_FUNCTION_BUILD_OPTIMIZATION
 FText STraceStoreWindow::GetConnectionStatusTooltip() const
 {
 	using Insights::FStoreBrowser;
-	
+
 	static FText Connected = LOCTEXT("Connected", "Connected to the trace server");
-	static FText NotConnected = LOCTEXT("NoConnection", "Unable to connect trace server."); 
+	static FText NotConnected = LOCTEXT("NoConnection", "Unable to connect trace server.");
 	static FText Connecting = LOCTEXT("Connecting", "Trying to connect to trace server.");
 	static FText Disconnected = LOCTEXT("Disconnected", "Connection to trace server has been lost. Attempting to reconnect in {0} seconds.");
-	
+
 	const FStoreBrowser::EConnectionStatus Status = StoreBrowser->GetConnectionStatus();
 	switch(Status)
 	{
@@ -1871,14 +1895,14 @@ TSharedPtr<SWidget> STraceStoreWindow::TraceList_GetMenuContent()
 
 bool STraceStoreWindow::CanEditTraceFile() const
 {
-	return SelectedTrace.IsValid() && !SelectedTrace->bIsLive;
+	return SelectedTrace.IsValid() && !SelectedTrace->bIsLive && !SelectedTrace->bWasJustRenamed;
 }
 
 ////////////////////////////////////////////////////////////////////////////////////////////////////
 
 void STraceStoreWindow::RenameTraceFile()
 {
-	if (SelectedTrace.IsValid() && !SelectedTrace->bIsLive)
+	if (SelectedTrace.IsValid() && !SelectedTrace->bIsLive && !SelectedTrace->bWasJustRenamed)
 	{
 		SelectedTrace->bIsRenaming = true;
 
@@ -1896,7 +1920,7 @@ void STraceStoreWindow::RenameTraceFile()
 
 void STraceStoreWindow::DeleteTraceFile()
 {
-	if (!SelectedTrace.IsValid() || SelectedTrace->bIsLive)
+	if (!SelectedTrace.IsValid() || SelectedTrace->bIsLive || SelectedTrace->bWasJustRenamed)
 	{
 		return;
 	}
@@ -2030,7 +2054,7 @@ FSlateColor STraceStoreWindow::GetColorByPath(const FString& Uri)
 	{
 		return FAppStyle::Get().GetSlateColor((*Dir)->Color);
 	}
-	// If this is default trace store directory, use forground
+	// If this is default trace store directory, use foreground
 	return FSlateColor(FSlateColor::UseForeground());
 }
 
@@ -2089,7 +2113,7 @@ void STraceStoreWindow::RefreshTraceList()
 			}
 			WatchDirsListView->RequestListRefresh();
 		}
-		
+
 		const uint64 NewChangeSerial = StoreBrowser->GetLockedTracesChangeSerial();
 		if (NewChangeSerial != TracesChangeSerial)
 		{
@@ -2139,14 +2163,7 @@ void STraceStoreWindow::RefreshTraceList()
 					// This trace was added.
 					AddedTraces++;
 					TSharedPtr<FTraceViewModel> TracePtr = MakeShared<FTraceViewModel>();
-					FTraceViewModel& Trace = *TracePtr;
-					Trace.TraceId = SourceTrace.TraceId;
-					Trace.Name = FText::FromString(SourceTrace.Name);
-					FString Uri = SourceTrace.Uri;
-					FPaths::NormalizeFilename(Uri);
-					Trace.Uri = FText::FromString(Uri);
-					Trace.DirectoryColor = GetColorByPath(Uri);
-					UpdateTrace(Trace, SourceTrace);
+					UpdateTrace(*TracePtr, SourceTrace);
 					TraceViewModels.Add(TracePtr);
 					TraceViewModelMap.Add(SourceTrace.TraceId, TracePtr);
 				}
@@ -2175,13 +2192,21 @@ void STraceStoreWindow::RefreshTraceList()
 
 void STraceStoreWindow::UpdateTrace(FTraceViewModel& InOutTrace, const Insights::FStoreBrowserTraceInfo& InSourceTrace)
 {
-	//TraceId -- no need to update
+	InOutTrace.TraceId = InSourceTrace.TraceId;
+	InOutTrace.bWasJustRenamed = false;
 
 	InOutTrace.ChangeSerial = InSourceTrace.ChangeSerial;
 	InOutTrace.TraceIndex = InSourceTrace.TraceIndex;
 
-	//Name -- no need to update
-	//Uri -- no need to update
+	InOutTrace.Name = FText::FromString(InSourceTrace.Name);
+	FString Uri = InSourceTrace.Uri;
+	if (Uri.IsEmpty())
+	{
+		// For backward compatibility to older UTS.
+		Uri = FInsightsManager::Get()->GetStoreDir() + TEXT("/") + InSourceTrace.Name + TEXT(".utrace");
+	}
+	InOutTrace.Uri = FText::FromString(Uri);
+	InOutTrace.DirectoryColor = GetColorByPath(Uri);
 
 	InOutTrace.Timestamp = InSourceTrace.Timestamp;
 	InOutTrace.Size = InSourceTrace.Size;
@@ -2246,9 +2271,20 @@ void STraceStoreWindow::UpdateTraceListView()
 	TSharedPtr<FTraceViewModel> NewSelectedTrace;
 	if (bIsUserSelectedTrace && SelectedTrace)
 	{
-		// Identify the previously selected trace (if stil available) to ensure selection remains unchanged.
-		TSharedPtr<FTraceViewModel>* TracePtrPtr = TraceViewModelMap.Find(SelectedTrace->TraceId);
-		NewSelectedTrace = TracePtrPtr ? *TracePtrPtr : nullptr;
+		// Identify the previously selected trace (if still available) to ensure selection remains unchanged.
+		TSharedPtr<FTraceViewModel>* FoundNewTrace = TraceViewModelMap.Find(SelectedTrace->TraceId);
+		if (!FoundNewTrace)
+		{
+			FoundNewTrace = TraceViewModels.FindByPredicate([this](const TSharedPtr<FTraceViewModel>& Trace) { return Trace->Uri.EqualTo(SelectedTrace->Uri); });
+		}
+		if (!FoundNewTrace)
+		{
+			FoundNewTrace = TraceViewModels.FindByPredicate([this](const TSharedPtr<FTraceViewModel>& Trace) { return Trace->Name.EqualTo(SelectedTrace->Name); });
+		}
+		if (FoundNewTrace)
+		{
+			NewSelectedTrace = *FoundNewTrace;
+		}
 	}
 
 	TraceListView->RebuildList();
@@ -2333,6 +2369,13 @@ void STraceStoreWindow::Tick(const FGeometry& AllottedGeometry, const double InC
 		}
 	}
 
+	if (bSetKeyboardFocusOnNextTick)
+	{
+		bSetKeyboardFocusOnNextTick = false;
+		FSlateApplication::Get().ClearKeyboardFocus();
+		FSlateApplication::Get().SetKeyboardFocus(TraceListView);
+	}
+
 	TickSplashScreenOverlay(InDeltaTime);
 }
 
@@ -2465,7 +2508,9 @@ FReply STraceStoreWindow::OnDrop(const FGeometry& MyGeometry, const FDragDropEve
 
 bool STraceStoreWindow::Open_IsEnabled() const
 {
-	return FilteredTraceViewModels.Num() > 0;
+	return FilteredTraceViewModels.Num() > 0
+		&& SelectedTrace.IsValid()
+		&& !SelectedTrace->bWasJustRenamed;
 }
 
 ////////////////////////////////////////////////////////////////////////////////////////////////////
@@ -2509,7 +2554,7 @@ void STraceStoreWindow::OpenTraceFile(const FString& InTraceFile)
 
 void STraceStoreWindow::OpenTraceSession(TSharedPtr<FTraceViewModel> InTraceSession)
 {
-	if (InTraceSession.IsValid())
+	if (InTraceSession.IsValid() && !InTraceSession->bWasJustRenamed)
 	{
 		OpenTraceSession(InTraceSession->TraceId);
 	}
@@ -2894,7 +2939,7 @@ FReply STraceStoreWindow::StoreSettingsArea_Toggle() const
 {
 	if (StoreSettingsArea)
 	{
-		if(StoreSettingsArea->GetVisibility() == EVisibility::Visible)
+		if (StoreSettingsArea->GetVisibility() == EVisibility::Visible)
 		{
 			StoreSettingsArea->SetVisibility(EVisibility::Collapsed);
 		}
@@ -2912,7 +2957,7 @@ const FSlateBrush* STraceStoreWindow::StoreSettingsToggle_Icon() const
 {
 	if (StoreSettingsArea)
 	{
-		if(StoreSettingsArea->GetVisibility() == EVisibility::Visible)
+		if (StoreSettingsArea->GetVisibility() == EVisibility::Visible)
 		{
 			return FInsightsStyle::Get().GetBrush("Icons.Expanded");
 		}
@@ -2924,9 +2969,10 @@ const FSlateBrush* STraceStoreWindow::StoreSettingsToggle_Icon() const
 
 FReply STraceStoreWindow::AddWatchDir_Clicked()
 {
+	FSlateApplication::Get().CloseToolTip();
 	if (IDesktopPlatform* DesktopPlatform = FDesktopPlatformModule::Get())
 	{
-		const FString Title = LOCTEXT("AddWatchDirectorySelectPopupTitle", "Add watch directory ...").ToString();
+		const FString Title = LOCTEXT("AddWatchDirectorySelectPopupTitle", "Add Watch Directory").ToString();
 
 		const FString StoreDirectory = FInsightsManager::Get()->GetStoreDir();
 		FString SelectedDirectory;
diff --git a/Engine/Source/Developer/TraceInsights/Private/Insights/Widgets/SStartPageWindow.h b/Engine/Source/Developer/TraceInsights/Private/Insights/Widgets/SStartPageWindow.h
--- a/Engine/Source/Developer/TraceInsights/Private/Insights/Widgets/SStartPageWindow.h
+++ b/Engine/Source/Developer/TraceInsights/Private/Insights/Widgets/SStartPageWindow.h
@@ -89,6 +89,7 @@ struct FTraceViewModel
 
 	bool bIsMetadataUpdated = false;
 	bool bIsRenaming = false;
+	bool bWasJustRenamed = false;
 	bool bIsLive = false;
 	uint32 IpAddress = 0;
 
@@ -369,18 +370,17 @@ private:
 	/** View model for a source of traces. This could be the trace store default directory or a watch directory.*/
 	struct FTraceDirectoryModel
 	{
-		
 		FTraceDirectoryModel(FString&& InPath, const FName& InColor, ETraceDirOperations InOperations)
 			: Path(MoveTemp(InPath)), Color(InColor), Operations(InOperations) {}
 
 		/** Path to directory */
-		FString		Path;
+		FString Path;
 		/** Assigned color */
-		FName		Color;
+		FName Color;
 		/** Supported operations */
-		ETraceDirOperations	Operations;
+		ETraceDirOperations Operations;
 	};
-	
+
 	TSharedRef<ITableRow> TraceDirs_OnGenerateRow(TSharedPtr<FTraceDirectoryModel> Item, const TSharedRef<STableViewBase>& Owner) const;
 	FReply StoreSettingsArea_Toggle() const;
 	const FSlateBrush* StoreSettingsToggle_Icon() const;
@@ -388,7 +388,7 @@ private:
 
 	TArray<TSharedPtr<FTraceDirectoryModel>> StoreDirectoryModel;
 	TArray<TSharedPtr<FTraceDirectoryModel>> WatchDirectoriesModel;
-	
+
 	friend class STraceWatchDirTableRow;
 
 	//////////////////////////////////////////////////
@@ -564,6 +564,8 @@ private:
 	bool bEnableAutomaticTesting = false;
 	bool bEnableDebugTools = false;
 	bool bStartProcessWithStompMalloc = false;
+
+	bool bSetKeyboardFocusOnNextTick = false;
 };
 
 ////////////////////////////////////////////////////////////////////////////////////////////////////
