[Insights] Memory Insights: Added option to include heap allocs in mem queries.

diff --git a/Engine/Source/Developer/TraceInsights/Private/Insights/MemoryProfiler/ViewModels/MemoryAlloc.cpp b/Engine/Source/Developer/TraceInsights/Private/Insights/MemoryProfiler/ViewModels/MemoryAlloc.cpp
--- a/Engine/Source/Developer/TraceInsights/Private/Insights/MemoryProfiler/ViewModels/MemoryAlloc.cpp
+++ b/Engine/Source/Developer/TraceInsights/Private/Insights/MemoryProfiler/ViewModels/MemoryAlloc.cpp
@@ -23,7 +23,7 @@ FMemoryAlloc::FMemoryAlloc()
 	, Callstack(nullptr)
 	, FreeCallstack(nullptr)
 	, RootHeap(0)
-	, bIsBlock(false)
+	, bIsHeap(false)
 	, bIsDecline(false)
 {
 }
diff --git a/Engine/Source/Developer/TraceInsights/Private/Insights/MemoryProfiler/ViewModels/MemoryAlloc.h b/Engine/Source/Developer/TraceInsights/Private/Insights/MemoryProfiler/ViewModels/MemoryAlloc.h
--- a/Engine/Source/Developer/TraceInsights/Private/Insights/MemoryProfiler/ViewModels/MemoryAlloc.h
+++ b/Engine/Source/Developer/TraceInsights/Private/Insights/MemoryProfiler/ViewModels/MemoryAlloc.h
@@ -29,6 +29,7 @@ public:
 	uint64 GetAddress() const { return Address; }
 	uint64 GetPage() const { return Address & ~(4llu*1024-1); }
 	int64 GetSize() const { return Size; }
+	bool IsHeap() const { return bIsHeap; }
 	TraceServices::TagIdType GetTagId() const { return TagId; }
 	const TCHAR* GetTag() const { return Tag; }
 	const TCHAR* GetAsset() const { return Asset; }
@@ -59,7 +60,7 @@ private:
 	const TraceServices::FCallstack* Callstack;
 	const TraceServices::FCallstack* FreeCallstack;
 	HeapId RootHeap;
-	bool bIsBlock;
+	bool bIsHeap;
 	bool bIsDecline;
 };
 
diff --git a/Engine/Source/Developer/TraceInsights/Private/Insights/MemoryProfiler/Widgets/SMemAllocTableTreeView.cpp b/Engine/Source/Developer/TraceInsights/Private/Insights/MemoryProfiler/Widgets/SMemAllocTableTreeView.cpp
--- a/Engine/Source/Developer/TraceInsights/Private/Insights/MemoryProfiler/Widgets/SMemAllocTableTreeView.cpp
+++ b/Engine/Source/Developer/TraceInsights/Private/Insights/MemoryProfiler/Widgets/SMemAllocTableTreeView.cpp
@@ -141,26 +141,27 @@ void SMemAllocTableTreeView::RebuildTree(bool bResync)
 				}
 				TableTreeNodes.Reserve(TotalAllocCount);
 
-				uint32 HeapAllocCount(0);
+				uint32 HeapAllocCount = 0;
 				const FName BaseNodeName(TEXT("alloc"));
 				const FName BaseHeapName(TEXT("heap"));
 				for (int32 AllocIndex = TableTreeNodes.Num(); AllocIndex < TotalAllocCount; ++AllocIndex)
 				{
 					const FMemoryAlloc* Alloc = MemAllocTable->GetMemAlloc(AllocIndex);
 
-					// Until we have an UX story around heap allocations
-					// remove them from the list
-					if (Alloc->bIsBlock)
+					if (Alloc->bIsHeap)
 					{
 						++HeapAllocCount;
-						continue;
+						if (!bIncludeHeapAllocs)
+						{
+							continue;
+						}
 					}
 
-					FName NodeName(Alloc->bIsBlock ? BaseHeapName : BaseNodeName, static_cast<int32>(Alloc->GetStartEventIndex() + 1));
+					FName NodeName(Alloc->bIsHeap ? BaseHeapName : BaseNodeName, static_cast<int32>(Alloc->GetStartEventIndex() + 1));
 					FMemAllocNodePtr NodePtr = MakeShared<FMemAllocNode>(NodeName, MemAllocTable, AllocIndex);
 					TableTreeNodes.Add(NodePtr);
 				}
-				ensure(TableTreeNodes.Num() == TotalAllocCount - HeapAllocCount);
+				ensure(TableTreeNodes.Num() == (bIncludeHeapAllocs ? TotalAllocCount : TotalAllocCount - HeapAllocCount));
 				UpdateQueryInfo();
 			}
 		}
@@ -460,7 +461,7 @@ void SMemAllocTableTreeView::UpdateQuery(TraceServices::IAllocationsProvider::EQ
 					}
 
 					Alloc.RootHeap = Allocation->GetRootHeap();
-					Alloc.bIsBlock = Allocation->IsHeap();
+					Alloc.bIsHeap = Allocation->IsHeap();
 
 					Alloc.bIsDecline = false;
 					if (Rule->GetValue() == TraceServices::IAllocationsProvider::EQueryRule::aAfaBf)
diff --git a/Engine/Source/Developer/TraceInsights/Private/Insights/MemoryProfiler/Widgets/SMemAllocTableTreeView.h b/Engine/Source/Developer/TraceInsights/Private/Insights/MemoryProfiler/Widgets/SMemAllocTableTreeView.h
--- a/Engine/Source/Developer/TraceInsights/Private/Insights/MemoryProfiler/Widgets/SMemAllocTableTreeView.h
+++ b/Engine/Source/Developer/TraceInsights/Private/Insights/MemoryProfiler/Widgets/SMemAllocTableTreeView.h
@@ -61,13 +61,21 @@ public:
 	 */
 	virtual void RebuildTree(bool bResync);
 
-	void SetQueryParams(TSharedPtr<FMemoryRuleSpec> InRule, double TimeA = 0.0, double TimeB = 0.0, double TimeC = 0.0, double TimeD = 0.0)
+	struct FQueryParams
 	{
-		Rule = InRule;
-		TimeMarkers[0] = TimeA;
-		TimeMarkers[1] = TimeB;
-		TimeMarkers[2] = TimeC;
-		TimeMarkers[3] = TimeD;
+		TSharedPtr<FMemoryRuleSpec> Rule;
+		double TimeMarkers[4] = { 0.0, 0.0, 0.0, 0.0 };
+		bool bIncludeHeapAllocs = false;
+	};
+
+	void SetQueryParams(const FQueryParams& InQueryParams)
+	{
+		Rule = InQueryParams.Rule;
+		TimeMarkers[0] = InQueryParams.TimeMarkers[0];
+		TimeMarkers[1] = InQueryParams.TimeMarkers[1];
+		TimeMarkers[2] = InQueryParams.TimeMarkers[2];
+		TimeMarkers[3] = InQueryParams.TimeMarkers[3];
+		bIncludeHeapAllocs = InQueryParams.bIncludeHeapAllocs;
 		OnQueryInvalidated();
 	}
 
@@ -127,6 +135,7 @@ private:
 	int32 TabIndex = -1;
 	TSharedPtr<FMemoryRuleSpec> Rule = nullptr;
 	double TimeMarkers[4];
+	bool bIncludeHeapAllocs = false;
 	TraceServices::IAllocationsProvider::FQueryHandle Query = 0;
 	FText QueryInfo;
 	FText QueryInfoTooltip;
diff --git a/Engine/Source/Developer/TraceInsights/Private/Insights/MemoryProfiler/Widgets/SMemInvestigationView.cpp b/Engine/Source/Developer/TraceInsights/Private/Insights/MemoryProfiler/Widgets/SMemInvestigationView.cpp
--- a/Engine/Source/Developer/TraceInsights/Private/Insights/MemoryProfiler/Widgets/SMemInvestigationView.cpp
+++ b/Engine/Source/Developer/TraceInsights/Private/Insights/MemoryProfiler/Widgets/SMemInvestigationView.cpp
@@ -5,6 +5,7 @@
 #include "DesktopPlatformModule.h"
 #include "SlateOptMacros.h"
 #include "Widgets/Input/SButton.h"
+#include "Widgets/Input/SCheckBox.h"
 #include "Widgets/Input/SEditableTextBox.h"
 #include "Widgets/Images/SImage.h"
 
@@ -26,6 +27,7 @@
 
 SMemInvestigationView::SMemInvestigationView()
 	: ProfilerWindowWeakPtr()
+	, bIncludeHeapAllocs(false)
 {
 }
 
@@ -140,6 +142,27 @@ TSharedRef<SWidget> SMemInvestigationView::ConstructInvestigationWidgetArea()
 			.AutoWrapText(true)
 		]
 
+		+ SVerticalBox::Slot()
+		.AutoHeight()
+		.Padding(0.0f, 4.0f, 0.0f, 0.0f)
+		[
+			SNew(SCheckBox)
+			.IsChecked_Lambda([this]()
+			{
+				return bIncludeHeapAllocs ? ECheckBoxState::Checked : ECheckBoxState::Unchecked;
+			})
+			.OnCheckStateChanged_Lambda([this](ECheckBoxState InCheckBoxState)
+			{
+				bIncludeHeapAllocs = (InCheckBoxState == ECheckBoxState::Checked);
+			})
+			.Content()
+			[
+				SNew(STextBlock)
+				.Text(LOCTEXT("IncludeHeapAllocsText", "Include Heap Allocs"))
+			]
+			.ToolTipText(LOCTEXT("IncludeHeapAllocsToolTipText", "Include heap allocs."))
+		]
+
 		+ SVerticalBox::Slot()
 		.AutoHeight()
 		.Padding(0.0f, 4.0f, 0.0f, 0.0f)
@@ -540,12 +563,14 @@ FReply SMemInvestigationView::RunQuery()
 	TSharedPtr<Insights::SMemAllocTableTreeView> MemAllocTableTreeView = ProfilerWindow->ShowMemAllocTableTreeViewTab();
 	if (MemAllocTableTreeView)
 	{
-		double TimeMarkers[4];
-		TimeMarkers[0] = (RuleNumTimeMarkers > 0) ? ProfilerWindow->GetCustomTimeMarker(0)->GetTime() : 0.0;
-		TimeMarkers[1] = (RuleNumTimeMarkers > 1) ? ProfilerWindow->GetCustomTimeMarker(1)->GetTime() : 0.0;
-		TimeMarkers[2] = (RuleNumTimeMarkers > 2) ? ProfilerWindow->GetCustomTimeMarker(2)->GetTime() : 0.0;
-		TimeMarkers[3] = (RuleNumTimeMarkers > 3) ? ProfilerWindow->GetCustomTimeMarker(3)->GetTime() : 0.0;
-		MemAllocTableTreeView->SetQueryParams(Rule, TimeMarkers[0], TimeMarkers[1], TimeMarkers[2], TimeMarkers[3]);
+		Insights::SMemAllocTableTreeView::FQueryParams QueryParams;
+		QueryParams.Rule = Rule;
+		QueryParams.TimeMarkers[0] = (RuleNumTimeMarkers > 0) ? ProfilerWindow->GetCustomTimeMarker(0)->GetTime() : 0.0;
+		QueryParams.TimeMarkers[1] = (RuleNumTimeMarkers > 1) ? ProfilerWindow->GetCustomTimeMarker(1)->GetTime() : 0.0;
+		QueryParams.TimeMarkers[2] = (RuleNumTimeMarkers > 2) ? ProfilerWindow->GetCustomTimeMarker(2)->GetTime() : 0.0;
+		QueryParams.TimeMarkers[3] = (RuleNumTimeMarkers > 3) ? ProfilerWindow->GetCustomTimeMarker(3)->GetTime() : 0.0;
+		QueryParams.bIncludeHeapAllocs = bIncludeHeapAllocs;
+		MemAllocTableTreeView->SetQueryParams(QueryParams);
 	}
 
 	return FReply::Handled();
diff --git a/Engine/Source/Developer/TraceInsights/Private/Insights/MemoryProfiler/Widgets/SMemInvestigationView.h b/Engine/Source/Developer/TraceInsights/Private/Insights/MemoryProfiler/Widgets/SMemInvestigationView.h
--- a/Engine/Source/Developer/TraceInsights/Private/Insights/MemoryProfiler/Widgets/SMemInvestigationView.h
+++ b/Engine/Source/Developer/TraceInsights/Private/Insights/MemoryProfiler/Widgets/SMemInvestigationView.h
@@ -98,6 +98,8 @@ private:
 
 	TSharedPtr<SComboBox<TSharedPtr<Insights::FMemoryRuleSpec>>> QueryRuleComboBox;
 
+	bool bIncludeHeapAllocs;
+
 	TSharedPtr<SComboBox<TSharedPtr<Insights::FQueryTargetWindowSpec>>> QueryTargetComboBox;
 	
 	TSharedPtr<STextBlock> SymbolPathsTextBlock;
