[Insights] Minor fixes.

#rnx

diff --git a/Engine/Source/Developer/TraceInsights/Private/Insights/Common/InsightsMenuBuilder.cpp b/Engine/Source/Developer/TraceInsights/Private/Insights/Common/InsightsMenuBuilder.cpp
--- a/Engine/Source/Developer/TraceInsights/Private/Insights/Common/InsightsMenuBuilder.cpp
+++ b/Engine/Source/Developer/TraceInsights/Private/Insights/Common/InsightsMenuBuilder.cpp
@@ -96,7 +96,7 @@ void FInsightsMenuBuilder::PopulateMenu(FMenuBuilder& MenuBuilder)
 			LOCTEXT("OpenWidgetReflector", "Widget Reflector"),
 			LOCTEXT("OpenWidgetReflector_ToolTip", "Opens the Widget Reflector, a handy tool for diagnosing problems with live widgets."),
 			FSlateIcon(FAppStyle::Get().GetStyleSetName(), "WidgetReflector.Icon"),
-			FUIAction(FExecuteAction::CreateLambda([=] { FGlobalTabmanager::Get()->TryInvokeTab(WidgetReflectorTabId); })));
+			FUIAction(FExecuteAction::CreateLambda([] { FGlobalTabmanager::Get()->TryInvokeTab(WidgetReflectorTabId); })));
 		MenuBuilder.EndSection();
 	}
 
diff --git a/Engine/Source/Developer/TraceInsights/Private/Insights/ContextSwitches/ContextSwitchesProfilerManager.cpp b/Engine/Source/Developer/TraceInsights/Private/Insights/ContextSwitches/ContextSwitchesProfilerManager.cpp
--- a/Engine/Source/Developer/TraceInsights/Private/Insights/ContextSwitches/ContextSwitchesProfilerManager.cpp
+++ b/Engine/Source/Developer/TraceInsights/Private/Insights/ContextSwitches/ContextSwitchesProfilerManager.cpp
@@ -124,7 +124,7 @@ void FContextSwitchesProfilerManager::UnregisterMajorTabs()
 
 bool FContextSwitchesProfilerManager::Tick(float DeltaTime)
 {
-	// Check if session has Memory events (to spawn the tab), but not too often.
+	// Check if session has ContextSwitch events (to spawn the tab), but not too often.
 	if (!bIsAvailable && AvailabilityCheck.Tick())
 	{
 		bool bShouldBeAvailable = false;
