Network insights
- Include injected trace events outside of normal packet range used to visualize split data

diff --git a/Engine/Source/Developer/TraceInsights/Private/Insights/NetworkingProfiler/Widgets/SPacketContentView.cpp b/Engine/Source/Developer/TraceInsights/Private/Insights/NetworkingProfiler/Widgets/SPacketContentView.cpp
--- a/Engine/Source/Developer/TraceInsights/Private/Insights/NetworkingProfiler/Widgets/SPacketContentView.cpp
+++ b/Engine/Source/Developer/TraceInsights/Private/Insights/NetworkingProfiler/Widgets/SPacketContentView.cpp
@@ -1044,10 +1044,9 @@ void SPacketContentView::UpdateState(float FontScale)
 			{
 				const FAxisViewportDouble& ViewportX = Viewport.GetHorizontalAxisViewport();
 
-				//const int64 StartPos = static_cast<int64>(FMath::FloorToDouble(ViewportX.GetValueAtOffset(0.0f)));
-				//const int64 EndPos = static_cast<int64>(FMath::CeilToDouble(ViewportX.GetValueAtOffset(ViewportX.GetSize())));
-				const uint32 StartPos = 0;
-				const uint32 EndPos = static_cast<uint32>(PacketBitSize);
+				// Count all events in packet, including split data
+				const uint32 StartPos = 0U;
+				const uint32 EndPos = ~0U;
 				NetProfilerProvider->EnumeratePacketContentEventsByPosition(ConnectionIndex, ConnectionMode, PacketIndex, StartPos, EndPos, [this, &Builder, &FilteredDrawStateBuilder, NetProfilerProvider](const TraceServices::FNetProfilerContentEvent& Event)
 				{
 					const TCHAR* Name = nullptr;
diff --git a/Engine/Source/Developer/TraceInsights/Private/Insights/NetworkingProfiler/Widgets/SPacketView.cpp b/Engine/Source/Developer/TraceInsights/Private/Insights/NetworkingProfiler/Widgets/SPacketView.cpp
--- a/Engine/Source/Developer/TraceInsights/Private/Insights/NetworkingProfiler/Widgets/SPacketView.cpp
+++ b/Engine/Source/Developer/TraceInsights/Private/Insights/NetworkingProfiler/Widgets/SPacketView.cpp
@@ -355,8 +355,9 @@ void SPacketView::UpdateState()
 								uint32 FilterMatchAggregatedEventSizeInBits = 0U;
 								uint32 FilterMatchMaxEventSizeBits = 0U;
 
+								// Filter all events in packet, including split data
 								const uint32 StartPos = 0;
-								const uint32 EndPos = Packet.ContentSizeInBits;
+								const uint32 EndPos = ~0U;
 								NetProfilerProvider->EnumeratePacketContentEventsByPosition(ConnectionIndex, ConnectionMode, PacketIndex - 1, StartPos, EndPos, [this, &bFilterMatch, &Filter, NetProfilerProvider, &FilterMatchAggregatedEventSizeInBits,&FilterMatchMaxEventSizeBits,  &FilterMatchEventTypeIndex](const TraceServices::FNetProfilerContentEvent& Event)
 								{
 									bool bEventMatchesFilter = true;
