Hide UI for store settings on platforms that doesn't support it.

diff --git a/Engine/Source/Developer/TraceInsights/Private/Insights/Widgets/SStartPageWindow.cpp b/Engine/Source/Developer/TraceInsights/Private/Insights/Widgets/SStartPageWindow.cpp
--- a/Engine/Source/Developer/TraceInsights/Private/Insights/Widgets/SStartPageWindow.cpp
+++ b/Engine/Source/Developer/TraceInsights/Private/Insights/Widgets/SStartPageWindow.cpp
@@ -54,6 +54,8 @@
 ////////////////////////////////////////////////////////////////////////////////////////////////////
 
 #define LOCTEXT_NAMESPACE "STraceStoreWindow"
+// For now we only support store settings on platforms that run Unreal Trace Server.
+#define STORE_SUPPORTS_SETTINGS PLATFORM_WINDOWS
 
 ////////////////////////////////////////////////////////////////////////////////////////////////////
 // STraceListRow
@@ -1580,6 +1582,7 @@ TSharedRef<SWidget> STraceStoreWindow::ConstructTraceStoreDirectoryPanel()
 		[
 			SNew(SHorizontalBox)
 
+#if STORE_SUPPORTS_SETTINGS
 			+ SHorizontalBox::Slot()
 			.AutoWidth()
 			.Padding(0.0f, 0.0f, 0.0f, 0.0f)
@@ -1595,6 +1598,7 @@ TSharedRef<SWidget> STraceStoreWindow::ConstructTraceStoreDirectoryPanel()
 					.ColorAndOpacity(FSlateColor::UseForeground())
 				]
 			]
+#endif
 
 			+ SHorizontalBox::Slot()
 			.AutoWidth()
@@ -1607,6 +1611,8 @@ TSharedRef<SWidget> STraceStoreWindow::ConstructTraceStoreDirectoryPanel()
 				.Text_Raw(this, &STraceStoreWindow::GetConnectionStatusTooltip))
 			]
 
+#if STORE_SUPPORTS_SETTINGS
+
 			+ SHorizontalBox::Slot()
 			.AutoWidth()
 			.Padding(4.0f, 0.0f, 0.0f, 0.0f)
@@ -1625,6 +1631,7 @@ TSharedRef<SWidget> STraceStoreWindow::ConstructTraceStoreDirectoryPanel()
 				.IsReadOnly(true)
 				.BackgroundColor(FSlateColor(EStyleColor::Background))
 			]
+#endif
 
 			+ SHorizontalBox::Slot()
 			.AutoWidth()
@@ -1663,6 +1670,7 @@ TSharedRef<SWidget> STraceStoreWindow::ConstructTraceStoreDirectoryPanel()
 			]
 		]
 
+#if STORE_SUPPORTS_SETTINGS
 		+ SVerticalBox::Slot()
 		.AutoHeight()
 		.Padding(.0f, 8.0f)
@@ -1739,7 +1747,9 @@ TSharedRef<SWidget> STraceStoreWindow::ConstructTraceStoreDirectoryPanel()
 					]
 				]
 			]
-		];
+		]
+#endif
+		;
 
 	return Widget;
 }
