Insights fixes to store settings

* Compare paths instead of using IsUnder to get color for trace.
* Hide tooltips when buttons are pressed
* Remove use of red accent color (associated with error)
* Fix wrapping code to restart on first
* Define max height for settings panel and wrap in scroll box in case the number of directories get too long.
* Update trace information when settings change. Might need to update color.
* Guard against multiple operations on a single directory by clicking button multiple times.
* Changed InsightsManagers function to get store directory. Will now always ask server since it can't keep a reliable local version any more.
* Disable modify trace store button while active live traces are detected.

diff --git a/Engine/Source/Developer/TraceInsights/Private/Insights/InsightsManager.cpp b/Engine/Source/Developer/TraceInsights/Private/Insights/InsightsManager.cpp
--- a/Engine/Source/Developer/TraceInsights/Private/Insights/InsightsManager.cpp
+++ b/Engine/Source/Developer/TraceInsights/Private/Insights/InsightsManager.cpp
@@ -439,6 +439,16 @@ void FInsightsManager::OnSessionInfoTabClosed(TSharedRef<SDockTab> TabBeingClose
 
 ////////////////////////////////////////////////////////////////////////////////////////////////////
 
+FString FInsightsManager::GetStoreDir() 
+{
+	using namespace UE::Trace;
+	FScopeLock _(&StoreClientCriticalSection);
+	const FStoreClient::FStatus* Status = StoreClient->GetStatus();
+	return Status ? FString(Status->GetStoreDir()) : FString();
+}
+
+////////////////////////////////////////////////////////////////////////////////////////////////////
+
 bool FInsightsManager::ConnectToStore(const TCHAR* Host, uint32 Port)
 {
 	using namespace UE::Trace;
@@ -452,13 +462,6 @@ bool FInsightsManager::ConnectToStore(const TCHAR* Host, uint32 Port)
 	LastStoreHost = Host;
 	LastStorePort = Port;
 
-	const FStoreClient::FStatus* Status = StoreClient->GetStatus();
-	FString RemoteStoreDir(Status->GetStoreDir());
-	if (RemoteStoreDir.Len() > 0)
-	{
-		SetStoreDir(RemoteStoreDir);
-	}
-
 	return true;
 }
 
diff --git a/Engine/Source/Developer/TraceInsights/Private/Insights/InsightsManager.h b/Engine/Source/Developer/TraceInsights/Private/Insights/InsightsManager.h
--- a/Engine/Source/Developer/TraceInsights/Private/Insights/InsightsManager.h
+++ b/Engine/Source/Developer/TraceInsights/Private/Insights/InsightsManager.h
@@ -114,8 +114,7 @@ public:
 	TSharedRef<TraceServices::IAnalysisService> GetAnalysisService() const { return AnalysisService; }
 	TSharedRef<TraceServices::IModuleService> GetModuleService() const { return ModuleService; }
 
-	void SetStoreDir(const FString& InStoreDir) { StoreDir = InStoreDir; }
-	const FString& GetStoreDir() const { return StoreDir; }
+	FString GetStoreDir();
 
 	bool ConnectToStore(const TCHAR* Host, uint32 Port=0);
 
@@ -419,9 +418,6 @@ private:
 	TSharedRef<TraceServices::IAnalysisService> AnalysisService;
 	TSharedRef<TraceServices::IModuleService> ModuleService;
 
-	/** The location of the trace files managed by the trace store. */
-	FString StoreDir;
-
 	/** The client used to connect to the trace store. It is not thread safe! */
 	TUniquePtr<UE::Trace::FStoreClient> StoreClient;
 
diff --git a/Engine/Source/Developer/TraceInsights/Private/Insights/Widgets/SStartPageWindow.cpp b/Engine/Source/Developer/TraceInsights/Private/Insights/Widgets/SStartPageWindow.cpp
--- a/Engine/Source/Developer/TraceInsights/Private/Insights/Widgets/SStartPageWindow.cpp
+++ b/Engine/Source/Developer/TraceInsights/Private/Insights/Widgets/SStartPageWindow.cpp
@@ -13,6 +13,7 @@
 #include "Internationalization/Text.h"
 #include "Logging/MessageLog.h"
 #include "Misc/MessageDialog.h"
+#include "Misc/PathViews.h"
 #include "Styling/AppStyle.h"
 #include "Styling/StyleColors.h"
 #include "Trace/ControlClient.h"
@@ -29,6 +30,7 @@
 #include "Widgets/Layout/SBorder.h"
 #include "Widgets/Layout/SBox.h"
 #include "Widgets/Layout/SGridPanel.h"
+#include "Widgets/Layout/SScrollBox.h"
 #include "Widgets/Notifications/SNotificationList.h"
 #include "Widgets/Testing/SStarshipSuite.h"
 #include "Widgets/Text/STextBlock.h"
@@ -1326,7 +1328,7 @@ public:
 	SLATE_BEGIN_ARGS(STraceDirectoryItem) {}
 	SLATE_END_ARGS()
 
-	void Construct(const FArguments& InArgs, TSharedPtr<STraceStoreWindow::FTraceDirectoryModel> InClientItem);
+	void Construct(const FArguments& InArgs, TSharedPtr<STraceStoreWindow::FTraceDirectoryModel> InClientItem, const STraceStoreWindow* InWindow);
 
 private:
 	FSlateColor GetColor() const;
@@ -1334,16 +1336,21 @@ private:
 	FReply OnModifyStore();
 	FReply OnExplore();
 	FReply OnDelete();
-
+	bool CanModifyStore() const;
+	FText ModifyStoreTooltip() const;
+	
+	bool bInOperation = false;
+	const STraceStoreWindow* Window = nullptr;
 
 	TSharedPtr<STraceStoreWindow::FTraceDirectoryModel> Model;
 };
 
 ////////////////////////////////////////////////////////////////////////////////////////////////////
 
-void STraceDirectoryItem::Construct(const FArguments& InArgs, TSharedPtr<STraceStoreWindow::FTraceDirectoryModel> InModel)
+void STraceDirectoryItem::Construct(const FArguments& InArgs, TSharedPtr<STraceStoreWindow::FTraceDirectoryModel> InModel, const STraceStoreWindow* InWindow)
 {
 	Model = InModel;
+	Window = InWindow;
 
 	ChildSlot
 	[
@@ -1410,8 +1417,9 @@ TSharedRef<SWidget> STraceDirectoryItem::ConstructOperations()
 			[
 				SNew(SButton)
 				.ButtonStyle(&FAppStyle::Get().GetWidgetStyle<FButtonStyle>("SimpleButton"))
-				.ToolTipText(LOCTEXT("SetTraceStoreDirButtonToolTip", "Set the Trace Store Directory"))
+				.ToolTipText_Raw(this, &STraceDirectoryItem::ModifyStoreTooltip)
 				.OnClicked_Raw(this, &STraceDirectoryItem::OnModifyStore)
+				.IsEnabled_Raw(this, &STraceDirectoryItem::CanModifyStore)
 				[
 					SNew(SImage)
 					.Image(FInsightsStyle::Get().GetBrush("Icons.Edit"))
@@ -1459,11 +1467,21 @@ TSharedRef<SWidget> STraceDirectoryItem::ConstructOperations()
 
 FReply STraceDirectoryItem::OnModifyStore()
 {
+	FSlateApplication::Get().CloseToolTip();
+	
+	// Avoid executing more than one operation
+	if (bInOperation)
+	{
+		return FReply::Handled();
+	}
+	
+	TGuardValue<bool> OperationGuard(bInOperation, true);
+	
 	if (IDesktopPlatform* DesktopPlatform = FDesktopPlatformModule::Get())
 	{
 		const FString Title = LOCTEXT("SetTraceStoreDirectorySelectPopupTitle", "Set Trace Store Directory").ToString();
 
-		FString CurrentDirectory = FInsightsManager::Get()->GetStoreDir();
+		FString CurrentDirectory = Window->GetStoreDirectory();
 		FString SelectedDirectory;
 		const bool bHasSelected = DesktopPlatform->OpenDirectoryDialog(
 			FSlateApplication::Get().FindBestParentWindowHandleForDialogs(AsShared()),
@@ -1490,6 +1508,7 @@ FReply STraceDirectoryItem::OnModifyStore()
 
 FReply STraceDirectoryItem::OnExplore()
 {
+	FSlateApplication::Get().CloseToolTip();
 	if (Model)
 	{
 		FString FullPath(FPaths::ConvertRelativePathToFull(Model->Path));
@@ -1500,8 +1519,18 @@ FReply STraceDirectoryItem::OnExplore()
 
 ////////////////////////////////////////////////////////////////////////////////////////////////////
 
-FReply STraceDirectoryItem::OnDelete()
+FReply STraceDirectoryItem::OnDelete() 
 {
+	FSlateApplication::Get().CloseToolTip();
+
+	// Avoid executing more than one operation
+	if (bInOperation)
+	{
+		return FReply::Handled();
+	}
+	
+	TGuardValue<bool> OperationGuard(bInOperation, true);
+	
 	if (Model)
 	{
 		const EAppReturnType::Type bConfirmed = FMessageDialog::Open(EAppMsgType::OkCancel, FText::Format(
@@ -1518,6 +1547,22 @@ FReply STraceDirectoryItem::OnDelete()
 
 ////////////////////////////////////////////////////////////////////////////////////////////////////
 
+bool STraceDirectoryItem::CanModifyStore() const
+{
+	check(Window);
+	return !Window->HasAnyLiveTrace();
+}
+
+////////////////////////////////////////////////////////////////////////////////////////////////////
+
+FText STraceDirectoryItem::ModifyStoreTooltip() const
+{
+	return CanModifyStore() ? LOCTEXT("SetTraceStoreDirButtonToolTip", "Set the Trace Store Directory.") :
+		LOCTEXT("SetTraceStoreDirButtonTooltipInactive", "Set the Trace Store Directory.\nNot available while live trace sessions are running.");
+}
+
+////////////////////////////////////////////////////////////////////////////////////////////////////
+
 TSharedRef<SWidget> STraceStoreWindow::ConstructLoadPanel()
 {
 	TSharedRef<SWidget> Widget = SNew(SHorizontalBox)
@@ -1673,77 +1718,85 @@ TSharedRef<SWidget> STraceStoreWindow::ConstructTraceStoreDirectoryPanel()
 #if STORE_SUPPORTS_SETTINGS
 		+ SVerticalBox::Slot()
 		.AutoHeight()
+		.MaxHeight(400)
 		.Padding(.0f, 8.0f)
 		.HAlign(HAlign_Left)
 		[
-			SAssignNew(StoreSettingsArea, SVerticalBox)
+			SAssignNew(StoreSettingsArea, SScrollBox)
+			.Orientation(Orient_Vertical)
 			.Visibility(EVisibility::Collapsed)
-
-			+ SVerticalBox::Slot()
-			.AutoHeight()
-			.Padding(0.0f, 4.0f)
-			.VAlign(VAlign_Center)
+			
+			+ SScrollBox::Slot()
 			[
-				SNew(STextBlock)
-				.Text(LOCTEXT("StoreDirLabel", "Trace store directory. New traces will be stored here:"))
-			]
 
-			+ SVerticalBox::Slot()
-			.AutoHeight()
-			.Padding(0.0f, 4.0f)
-			[
-				SAssignNew(StoreDirListView, SListView<TSharedPtr<FTraceDirectoryModel>>)
-				.ListItemsSource(&StoreDirectoryModel)
-				.OnGenerateRow(this, &STraceStoreWindow::TraceDirs_OnGenerateRow)
-				.SelectionMode(ESelectionMode::None)
-			]
+				SNew(SVerticalBox)
+				
+				+ SVerticalBox::Slot()
+				.AutoHeight()
+				.Padding(0.0f, 4.0f)
+				.VAlign(VAlign_Center)
+				[
+					SNew(STextBlock)
+					.Text(LOCTEXT("StoreDirLabel", "Trace store directory. New traces will be stored here:"))
+				]
 
-			+ SVerticalBox::Slot()
-			.AutoHeight()
-			.Padding(0.0f, 4.0f)
-			.VAlign(VAlign_Center)
-			[
-				SNew(STextBlock)
-				.Text(LOCTEXT("WatchDirsLabel", "Additional directories to monitor for traces:"))
-			]
+				+ SVerticalBox::Slot()
+				.AutoHeight()
+				.Padding(0.0f, 4.0f)
+				[
+					SAssignNew(StoreDirListView, SListView<TSharedPtr<FTraceDirectoryModel>>)
+					.ListItemsSource(&StoreDirectoryModel)
+					.OnGenerateRow(this, &STraceStoreWindow::TraceDirs_OnGenerateRow)
+					.SelectionMode(ESelectionMode::None)
+				]
 
-			+ SVerticalBox::Slot()
-			.AutoHeight()
-			.Padding(0.0f, 4.0f)
-			[
-				SAssignNew(WatchDirsListView, SListView<TSharedPtr<FTraceDirectoryModel>>)
-				.ListItemsSource(&WatchDirectoriesModel)
-				.OnGenerateRow(this, &STraceStoreWindow::TraceDirs_OnGenerateRow)
-				//.Orientation(Orient_Horizontal)
-				.SelectionMode(ESelectionMode::None)
-			]
+				+ SVerticalBox::Slot()
+				.AutoHeight()
+				.Padding(0.0f, 4.0f)
+				.VAlign(VAlign_Center)
+				[
+					SNew(STextBlock)
+					.Text(LOCTEXT("WatchDirsLabel", "Additional directories to monitor for traces:"))
+				]
 
-			+ SVerticalBox::Slot()
-			.AutoHeight()
-			.Padding(0.0f, 4.0f)
-			.HAlign(HAlign_Left)
-			[
-				SNew(SButton)
+				+ SVerticalBox::Slot()
+				.AutoHeight()
+				.Padding(0.0f, 4.0f)
+				[
+					SAssignNew(WatchDirsListView, SListView<TSharedPtr<FTraceDirectoryModel>>)
+					.ListItemsSource(&WatchDirectoriesModel)
+					.OnGenerateRow(this, &STraceStoreWindow::TraceDirs_OnGenerateRow)
+					//.Orientation(Orient_Horizontal)
+					.SelectionMode(ESelectionMode::None)
+				]
+
+				+ SVerticalBox::Slot()
+				.AutoHeight()
+				.Padding(0.0f, 4.0f)
+				.HAlign(HAlign_Left)
+				[
+					SNew(SButton)
 				.ButtonStyle(&FAppStyle::Get().GetWidgetStyle<FButtonStyle>("Button"))
 				.ToolTipText(LOCTEXT("WatchDirsAddTooltip", "Add additional watch directory..."))
 				.OnClicked(this, &STraceStoreWindow::AddWatchDir_Clicked)
-				[
-					SNew(SHorizontalBox)
-
-					+ SHorizontalBox::Slot()
-					.AutoWidth()
 					[
-						SNew(SImage)
+						SNew(SHorizontalBox)
+
+						+ SHorizontalBox::Slot()
+						.AutoWidth()
+						[
+							SNew(SImage)
 						.Image(FInsightsStyle::Get().GetBrush("Icons.AddWatchDir"))
 						.ColorAndOpacity(FSlateColor::UseForeground())
-					]
+						]
 
-					+ SHorizontalBox::Slot()
-					.Padding(4.0, 0.0, 0.0, 0.0)
-					.AutoWidth()
-					[
-						SNew(STextBlock)
-						.Text(FText::FromStringView(TEXTVIEW("Add Directory...")))
+						+ SHorizontalBox::Slot()
+						.Padding(4.0, 0.0, 0.0, 0.0)
+						.AutoWidth()
+						[
+							SNew(STextBlock)
+							.Text(FText::FromStringView(TEXTVIEW("Add Directory...")))
+						]
 					]
 				]
 			]
@@ -1983,6 +2036,13 @@ void STraceStoreWindow::DeleteTraceFile()
 
 ////////////////////////////////////////////////////////////////////////////////////////////////////
 
+bool STraceStoreWindow::HasAnyLiveTrace() const
+{
+	return TraceViewModels.FindByPredicate(&FTraceViewModel::bIsLive) != nullptr;
+}
+
+////////////////////////////////////////////////////////////////////////////////////////////////////
+
 void STraceStoreWindow::ShowSplashScreenOverlay()
 {
 	SplashScreenOverlayFadeTime = 3.5f;
@@ -2056,9 +2116,10 @@ FReply STraceStoreWindow::RefreshTraces_OnClicked()
 
 FSlateColor STraceStoreWindow::GetColorByPath(const FString& Uri)
 {
+	const FStringView UriBase = FPathViews::GetPath(Uri);
 	const TSharedPtr<FTraceDirectoryModel>* Dir = WatchDirectoriesModel.FindByPredicate([&](const TSharedPtr<FTraceDirectoryModel>& Dir)
 	{
-		return FPaths::IsUnderDirectory(Uri, Dir->Path);
+		return FPathViews::Equals(UriBase, Dir->Path);
 	});
 	if (Dir)
 	{
@@ -2078,6 +2139,7 @@ void STraceStoreWindow::RefreshTraceList()
 	int32 AddedTraces = 0;
 	int32 RemovedTraces = 0;
 	int32 UpdatedTraces = 0;
+	bool bSettingsChanged = false;
 
 	{
 		StoreBrowser->Lock();
@@ -2108,9 +2170,10 @@ void STraceStoreWindow::RefreshTraceList()
 				FName("Colors.AccentGreen"),
 				FName("Colors.AccentYellow"),
 				FName("Colors.AccentOrange"),
-				FName("Colors.AccentRed")
+				FName("Colors.AccentPurple"),
+				FName("Colors.AccentPink")
 			};
-			uint64 ColorIdx = 0;
+			int32 ColorIdx = 0;
 			WatchDirectoriesModel.Empty();
 			for (const auto& Dir : StoreBrowser->GetLockedWatchDirectories())
 			{
@@ -2119,16 +2182,17 @@ void STraceStoreWindow::RefreshTraceList()
 					DirColor[ColorIdx],
 					ETraceDirOperations::Delete|ETraceDirOperations::Explore
 				));
-				ColorIdx = FMath::Wrap(ColorIdx + 1, uint64(0u), uint64(UE_ARRAY_COUNT(DirColor)));
+				ColorIdx = FMath::Wrap(++ColorIdx, int32(-1), int32(UE_ARRAY_COUNT(DirColor))-1);
 			}
 			WatchDirsListView->RequestListRefresh();
+
+			bSettingsChanged = true;
 		}
 
 		const uint64 NewChangeSerial = StoreBrowser->GetLockedTracesChangeSerial();
-		if (NewChangeSerial != TracesChangeSerial)
+		if (NewChangeSerial != TracesChangeSerial || bSettingsChanged)
 		{
 			TracesChangeSerial = NewChangeSerial;
-
 			//UE_LOG(TraceInsights, Log, TEXT("[TraceStore] Synching the trace list with StoreBrowser..."));
 
 			const TArray<TSharedPtr<Insights::FStoreBrowserTraceInfo>>& InTraces = StoreBrowser->GetLockedTraces();
@@ -2161,9 +2225,9 @@ void STraceStoreWindow::RefreshTraceList()
 				if (TracePtrPtr)
 				{
 					FTraceViewModel& Trace = **TracePtrPtr;
-					if (Trace.ChangeSerial != SourceTrace.ChangeSerial)
+					if (Trace.ChangeSerial != SourceTrace.ChangeSerial || bSettingsChanged)
 					{
-						// This trace was updated.
+						// This trace was updated or settings updated
 						UpdatedTraces++;
 						UpdateTrace(Trace, SourceTrace);
 					}
@@ -2920,14 +2984,14 @@ void STraceStoreWindow::BuildBranchFilterSubMenu(FMenuBuilder& InMenuBuilder)
 
 FText STraceStoreWindow::GetTraceStoreDirectory() const
 {
-	return FText::FromString(FPaths::ConvertRelativePathToFull(FInsightsManager::Get()->GetStoreDir()));
+	return FText::FromString(FPaths::ConvertRelativePathToFull(GetStoreDirectory()));
 }
 
 ////////////////////////////////////////////////////////////////////////////////////////////////////
 
 FReply STraceStoreWindow::ExploreTraceStoreDirectory_OnClicked()
 {
-	FString FullPath(FPaths::ConvertRelativePathToFull(FInsightsManager::Get()->GetStoreDir()));
+	FString FullPath(FPaths::ConvertRelativePathToFull(GetStoreDirectory()));
 	FPlatformProcess::ExploreFolder(*FullPath);
 	return FReply::Handled();
 }
@@ -2939,7 +3003,7 @@ TSharedRef<ITableRow> STraceStoreWindow::TraceDirs_OnGenerateRow(TSharedPtr<FTra
 	return SNew(STableRow<TSharedPtr<FTraceDirectoryModel>>, Owner)
 		.Content()
 		[
-			SNew(STraceDirectoryItem, Item)
+			SNew(STraceDirectoryItem, Item, this)
 		];
 }
 
@@ -2984,16 +3048,16 @@ FReply STraceStoreWindow::AddWatchDir_Clicked()
 	{
 		const FString Title = LOCTEXT("AddWatchDirectorySelectPopupTitle", "Add Watch Directory").ToString();
 
-		const FString StoreDirectory = FInsightsManager::Get()->GetStoreDir();
+		const FString& CurrentStoreDirectory = StoreDirectoryModel.IsEmpty() ? FString() : StoreDirectoryModel.Last()->Path;
 		FString SelectedDirectory;
 		const bool bHasSelected = DesktopPlatform->OpenDirectoryDialog(
 			FSlateApplication::Get().FindBestParentWindowHandleForDialogs(AsShared()),
 			Title,
-			StoreDirectory,
+			CurrentStoreDirectory,
 			SelectedDirectory
 		);
 
-		if (bHasSelected && !SelectedDirectory.Equals(StoreDirectory))
+		if (bHasSelected && !SelectedDirectory.Equals(CurrentStoreDirectory))
 		{
 			FPaths::MakePlatformFilename(SelectedDirectory);
 			if (!FInsightsManager::Get()->GetStoreClient()->SetStoreDirectories(nullptr, { (*SelectedDirectory) }, {}))
@@ -3007,6 +3071,13 @@ FReply STraceStoreWindow::AddWatchDir_Clicked()
 
 ////////////////////////////////////////////////////////////////////////////////////////////////////
 
+FString STraceStoreWindow::GetStoreDirectory() const
+{
+	return StoreDirectoryModel.IsEmpty() ? FString() : StoreDirectoryModel.Last()->Path;
+}
+
+////////////////////////////////////////////////////////////////////////////////////////////////////
+
 void STraceStoreWindow::OpenSettings()
 {
 	MainContentPanel->SetEnabled(false);
diff --git a/Engine/Source/Developer/TraceInsights/Private/Insights/Widgets/SStartPageWindow.h b/Engine/Source/Developer/TraceInsights/Private/Insights/Widgets/SStartPageWindow.h
--- a/Engine/Source/Developer/TraceInsights/Private/Insights/Widgets/SStartPageWindow.h
+++ b/Engine/Source/Developer/TraceInsights/Private/Insights/Widgets/SStartPageWindow.h
@@ -17,6 +17,7 @@ class FActiveTimerHandle;
 class ITableRow;
 class SEditableTextBox;
 class SNotificationList;
+class SScrollBox;
 class SSearchBox;
 class STableViewBase;
 class SVerticalBox;
@@ -295,6 +296,7 @@ private:
 	bool CanEditTraceFile() const;
 	void RenameTraceFile();
 	void DeleteTraceFile();
+	bool HasAnyLiveTrace() const;
 
 	//////////////////////////////////////////////////
 	// "Starting Analysis" Splash Screen
@@ -385,6 +387,7 @@ private:
 	FReply StoreSettingsArea_Toggle() const;
 	const FSlateBrush* StoreSettingsToggle_Icon() const;
 	FReply AddWatchDir_Clicked();
+	FString GetStoreDirectory() const;
 
 	TArray<TSharedPtr<FTraceDirectoryModel>> StoreDirectoryModel;
 	TArray<TSharedPtr<FTraceDirectoryModel>> WatchDirectoriesModel;
@@ -515,7 +518,7 @@ private:
 	TSharedPtr<SEditableTextBox> StoreDirTextBox;
 	TSharedPtr<SEditableTextBox> StoreHostTextBox;
 	TSharedPtr<STableViewBase> StoreDirListView;
-	TSharedPtr<SVerticalBox> StoreSettingsArea;
+	TSharedPtr<SScrollBox> StoreSettingsArea;
 	TSharedPtr<STableViewBase> WatchDirsListView;
 	TSharedPtr<SListView<TSharedPtr<FTraceViewModel>>> TraceListView;
 	TSharedPtr<FTraceViewModel> SelectedTrace;
