[Insights] SessionBrowser: Fixed color of trace name when session is hovered or selected.

#rnx

diff --git a/Engine/Source/Developer/TraceInsights/Private/Insights/Widgets/SStartPageWindow.cpp b/Engine/Source/Developer/TraceInsights/Private/Insights/Widgets/SStartPageWindow.cpp
--- a/Engine/Source/Developer/TraceInsights/Private/Insights/Widgets/SStartPageWindow.cpp
+++ b/Engine/Source/Developer/TraceInsights/Private/Insights/Widgets/SStartPageWindow.cpp
@@ -372,7 +372,9 @@ public:
 				return FSlateColor(EStyleColor::White25);
 			}
 		}
-		return FSlateColor(EStyleColor::Foreground);
+		return IsSelected() || IsHovered() ?
+			FSlateColor(EStyleColor::ForegroundHover) :
+			FSlateColor(EStyleColor::Foreground);
 	}
 
 	FText GetTraceNameHighlightText() const
@@ -1379,8 +1381,8 @@ void STraceDirectoryItem::Construct(const FArguments& InArgs, TSharedPtr<STraceS
 			.Padding(4.0, 2.0)
 			.VAlign(EVerticalAlignment::VAlign_Center)
 			[
-				 SNew(STextBlock)
-				 .Text(FText::FromString(Model->Path))
+				SNew(STextBlock)
+				.Text(FText::FromString(Model->Path))
 			]
 
 			+ SHorizontalBox::Slot()
@@ -1425,23 +1427,23 @@ TSharedRef<SWidget> STraceDirectoryItem::ConstructOperations()
 					.Image(FInsightsStyle::Get().GetBrush("Icons.Edit"))
 					.ColorAndOpacity(FSlateColor::UseForeground())
 				]
-			 ];
+			];
 	}
 	if (Model && EnumHasAllFlags(Model->Operations, ETraceDirOperations::Delete) )
 	{
 		Box->AddSlot()
 			.AutoWidth()
 			[
-				 SNew(SButton)
-				 .ButtonStyle(&FAppStyle::Get().GetWidgetStyle<FButtonStyle>("SimpleButton"))
-				 .ToolTipText(LOCTEXT("WatchDirsRemoveTooltip", "Remove watch directory (files will not be deleted)"))
-				 .OnClicked_Raw(this, &STraceDirectoryItem::OnDelete)
-				 [
-					  SNew(SImage)
-					  .Image(FInsightsStyle::Get().GetBrush("Icons.RemoveWatchDir"))
-					  .ColorAndOpacity(FSlateColor::UseForeground())
-				 ]
-			 ];
+				SNew(SButton)
+				.ButtonStyle(&FAppStyle::Get().GetWidgetStyle<FButtonStyle>("SimpleButton"))
+				.ToolTipText(LOCTEXT("WatchDirsRemoveTooltip", "Remove watch directory (files will not be deleted)"))
+				.OnClicked_Raw(this, &STraceDirectoryItem::OnDelete)
+				[
+					SNew(SImage)
+					.Image(FInsightsStyle::Get().GetBrush("Icons.RemoveWatchDir"))
+					.ColorAndOpacity(FSlateColor::UseForeground())
+				]
+			];
 	}
 	if (Model && EnumHasAllFlags(Model->Operations, ETraceDirOperations::Explore))
 	{
