Moved fallback uri construction to session browser to avoid deadlock

With new code that does a remote call to get store directory, there was a deadlock when StartPageWindow was trying to get the lock which was held by [REDACTED_PERSON] which was waiting for StartPageWindow.

diff --git a/Engine/Source/Developer/TraceInsights/Private/Insights/StoreService/StoreBrowser.cpp b/Engine/Source/Developer/TraceInsights/Private/Insights/StoreService/StoreBrowser.cpp
--- a/Engine/Source/Developer/TraceInsights/Private/Insights/StoreService/StoreBrowser.cpp
+++ b/Engine/Source/Developer/TraceInsights/Private/Insights/StoreService/StoreBrowser.cpp
@@ -215,7 +215,17 @@ void FStoreBrowser::UpdateTraces()
 
 							const FUtf8StringView Utf8NameView = TraceInfo->GetName();
 							Trace.Name = FString(Utf8NameView);
-							Trace.Uri = FString(TraceInfo->GetUri());
+							FUtf8StringView Uri = TraceInfo->GetUri();
+							if (Uri.Len() > 0)
+							{
+								Trace.Uri = FString(Uri);
+							}
+							else
+							{
+								// Fallback for older versions of UTS which didn't write uri
+								Trace.Uri = FPaths::SetExtension(FPaths::Combine(StoreDirectory, Trace.Name), TEXT(".utrace"));
+								FPaths::MakePlatformFilename(Trace.Uri);
+							}
 
 							Trace.Timestamp = FStoreBrowserTraceInfo::ConvertTimestamp(TraceInfo->GetTimestamp());
 							Trace.Size = TraceInfo->GetSize();
diff --git a/Engine/Source/Developer/TraceInsights/Private/Insights/Widgets/SStartPageWindow.cpp b/Engine/Source/Developer/TraceInsights/Private/Insights/Widgets/SStartPageWindow.cpp
--- a/Engine/Source/Developer/TraceInsights/Private/Insights/Widgets/SStartPageWindow.cpp
+++ b/Engine/Source/Developer/TraceInsights/Private/Insights/Widgets/SStartPageWindow.cpp
@@ -2275,14 +2275,8 @@ void STraceStoreWindow::UpdateTrace(FTraceViewModel& InOutTrace, const Insights:
 	InOutTrace.TraceIndex = InSourceTrace.TraceIndex;
 
 	InOutTrace.Name = FText::FromString(InSourceTrace.Name);
-	FString Uri = InSourceTrace.Uri;
-	if (Uri.IsEmpty())
-	{
-		// For backward compatibility to older UTS.
-		Uri = FInsightsManager::Get()->GetStoreDir() + TEXT("/") + InSourceTrace.Name + TEXT(".utrace");
-	}
-	InOutTrace.Uri = FText::FromString(Uri);
-	InOutTrace.DirectoryColor = GetColorByPath(Uri);
+	InOutTrace.Uri = FText::FromString(InSourceTrace.Uri);
+	InOutTrace.DirectoryColor = GetColorByPath(InSourceTrace.Uri);
 
 	InOutTrace.Timestamp = InSourceTrace.Timestamp;
 	InOutTrace.Size = InSourceTrace.Size;
