[Insights] Timing View: When Ctrl is pressed while clicking an event to be selected, the time interval that includes both the old selected event (if any) and the new selected event will be also selected. This is similar with feature contributed by [REDACTED_PERSON] in PR #10298, but implemented in a different way.

diff --git a/Engine/Source/Developer/TraceInsights/Private/Insights/Widgets/STimingView.cpp b/Engine/Source/Developer/TraceInsights/Private/Insights/Widgets/STimingView.cpp
--- a/Engine/Source/Developer/TraceInsights/Private/Insights/Widgets/STimingView.cpp
+++ b/Engine/Source/Developer/TraceInsights/Private/Insights/Widgets/STimingView.cpp
@@ -2130,7 +2130,7 @@ FReply STimingView::OnMouseButtonUp(const FGeometry& MyGeometry, const FPointerE
 	MousePositionOnButtonUp = MyGeometry.AbsoluteToLocal(MouseEvent.GetScreenSpacePosition());
 	MousePosition = MousePositionOnButtonUp;
 
-	const bool bIsValidForMouseClick = MousePositionOnButtonUp.Equals(MousePositionOnButtonDown, 2.0f);
+	bool bIsMouseClick = MousePositionOnButtonUp.Equals(MousePositionOnButtonDown, 2.0f);
 
 	if (MouseEvent.GetEffectingButton() == EKeys::LeftMouseButton)
 	{
@@ -2145,17 +2145,31 @@ FReply STimingView::OnMouseButtonUp(const FGeometry& MyGeometry, const FPointerE
 			{
 				SelectTimeInterval(SelectionStartTime, SelectionEndTime - SelectionStartTime);
 				bIsSelecting = false;
+				bIsMouseClick = false;
 			}
 			else if (TimeRulerTrack->IsScrubbing())
 			{
 				RaiseTimeMarkerChanged(TimeRulerTrack->GetScrubbingTimeMarker());
 				TimeRulerTrack->StopScrubbing();
+				bIsMouseClick = false;
 			}
 
-			if (bIsValidForMouseClick)
+			if (bIsMouseClick)
 			{
-				// Select the hovered timing event (if any).
 				UpdateHoveredTimingEvent(static_cast<float>(MousePositionOnButtonUp.X), static_cast<float>(MousePositionOnButtonUp.Y));
+
+				if (MouseEvent.GetModifierKeys().IsControlDown())
+				{
+					if (SelectedEvent.IsValid() && HoveredEvent.IsValid())
+					{
+						// Select the time region that includes both the current selected event and the new event to be selected.
+						const double RegionStartTime = FMath::Min(SelectedEvent->GetStartTime(), HoveredEvent->GetStartTime());
+						const double RegionEndTime = FMath::Max(Viewport.RestrictEndTime(SelectedEvent->GetEndTime()), Viewport.RestrictEndTime(HoveredEvent->GetEndTime()));
+						SelectTimeInterval(RegionStartTime, RegionEndTime - RegionStartTime);
+					}
+				}
+
+				// Select the hovered timing event (if any).
 				SelectHoveredTimingTrack();
 				SelectHoveredTimingEvent();
 
@@ -2190,14 +2204,16 @@ FReply STimingView::OnMouseButtonUp(const FGeometry& MyGeometry, const FPointerE
 			{
 				RaiseSelectionChanged();
 				bIsSelecting = false;
+				bIsMouseClick = false;
 			}
 			else if (TimeRulerTrack->IsScrubbing())
 			{
 				RaiseTimeMarkerChanged(TimeRulerTrack->GetScrubbingTimeMarker());
 				TimeRulerTrack->StopScrubbing();
+				bIsMouseClick = false;
 			}
 
-			if (bIsValidForMouseClick)
+			if (bIsMouseClick)
 			{
 				SelectHoveredTimingTrack();
 				ShowContextMenu(MouseEvent);
