Small fixes to setting trace store directory and watch directories.

* Check that watch directory path is not the same as store directory. Use FPathViews comparison rather than string comparison.
* Check new trace store directory is not in watch directories. Also check that it is not the existing store directory to avoid unecessary operations.

#rnx

diff --git a/Engine/Source/Developer/TraceInsights/Private/Insights/Widgets/SStartPageWindow.cpp b/Engine/Source/Developer/TraceInsights/Private/Insights/Widgets/SStartPageWindow.cpp
--- a/Engine/Source/Developer/TraceInsights/Private/Insights/Widgets/SStartPageWindow.cpp
+++ b/Engine/Source/Developer/TraceInsights/Private/Insights/Widgets/SStartPageWindow.cpp
@@ -1483,20 +1483,23 @@ FReply STraceDirectoryItem::OnModifyStore()
 	{
 		const FString Title = LOCTEXT("SetTraceStoreDirectorySelectPopupTitle", "Set Trace Store Directory").ToString();
 
-		FString CurrentDirectory = Window->GetStoreDirectory();
+		FString CurrentStoreDirectory = Window->GetStoreDirectory();
 		FString SelectedDirectory;
 		const bool bHasSelected = DesktopPlatform->OpenDirectoryDialog(
 			FSlateApplication::Get().FindBestParentWindowHandleForDialogs(AsShared()),
 			Title,
-			CurrentDirectory,
+			CurrentStoreDirectory,
 			SelectedDirectory
 		);
 
-		if (bHasSelected)
+		const bool bIsWatchDir = Window->WatchDirectoriesModel.FindByPredicate([&](const auto& Directory){ return FPathViews::Equals(SelectedDirectory, Directory->Path);}) != nullptr;
+		const bool bIsCurrentStoreDir = FPathViews::Equals(SelectedDirectory, CurrentStoreDirectory);
+
+		if (bHasSelected && !bIsWatchDir && !bIsCurrentStoreDir)
 		{
 			FPaths::MakePlatformFilename(SelectedDirectory);
-			FPaths::MakePlatformFilename(CurrentDirectory);
-			TArray<FString> AddWatchDirs = { CurrentDirectory };
+			FPaths::MakePlatformFilename(CurrentStoreDirectory);
+			TArray<FString> AddWatchDirs = { CurrentStoreDirectory };
 			if (!FInsightsManager::Get()->GetStoreClient()->SetStoreDirectories(*SelectedDirectory, AddWatchDirs, TArray<FString>()))
 			{
 				FMessageLog(FInsightsManager::Get()->GetLogListingName()).Error(LOCTEXT("StoreCommunicationFail", "Failed to change settings on the store service."));
@@ -3053,7 +3056,7 @@ FReply STraceStoreWindow::AddWatchDir_Clicked()
 			SelectedDirectory
 		);
 
-		if (bHasSelected && !SelectedDirectory.Equals(CurrentStoreDirectory))
+		if (bHasSelected && !FPathViews::Equals(SelectedDirectory, CurrentStoreDirectory))
 		{
 			FPaths::MakePlatformFilename(SelectedDirectory);
 			if (!FInsightsManager::Get()->GetStoreClient()->SetStoreDirectories(nullptr, { (*SelectedDirectory) }, {}))
