Show ProjectName as AppName if available.

With a previous change the session info event separated app name and project name fields. To simulate previous behaviour we use project name if available, otherwise fallback to app name (useful for programs that doesn't set project name for example).

diff --git a/Engine/Source/Developer/TraceInsights/Private/Insights/StoreService/DiagnosticsSessionAnalyzer.cpp b/Engine/Source/Developer/TraceInsights/Private/Insights/StoreService/DiagnosticsSessionAnalyzer.cpp
--- a/Engine/Source/Developer/TraceInsights/Private/Insights/StoreService/DiagnosticsSessionAnalyzer.cpp
+++ b/Engine/Source/Developer/TraceInsights/Private/Insights/StoreService/DiagnosticsSessionAnalyzer.cpp
@@ -54,6 +54,7 @@ bool FDiagnosticsSessionAnalyzer::OnEvent(uint16 RouteId, EStyle, const FOnEvent
 	{
 		EventData.GetString("Platform", Platform);
 		EventData.GetString("AppName", AppName);
+		EventData.GetString("ProjectName", ProjectName);
 		EventData.GetString("CommandLine", CommandLine);
 		EventData.GetString("Branch", Branch);
 		EventData.GetString("BuildVersion", BuildVersion);
diff --git a/Engine/Source/Developer/TraceInsights/Private/Insights/StoreService/DiagnosticsSessionAnalyzer.h b/Engine/Source/Developer/TraceInsights/Private/Insights/StoreService/DiagnosticsSessionAnalyzer.h
--- a/Engine/Source/Developer/TraceInsights/Private/Insights/StoreService/DiagnosticsSessionAnalyzer.h
+++ b/Engine/Source/Developer/TraceInsights/Private/Insights/StoreService/DiagnosticsSessionAnalyzer.h
@@ -21,6 +21,7 @@ struct FDiagnosticsSessionAnalyzer : public UE::Trace::IAnalyzer
 
 	FString Platform;
 	FString AppName;
+	FString ProjectName;
 	FString CommandLine;
 	FString Branch;
 	FString BuildVersion;
diff --git a/Engine/Source/Developer/TraceInsights/Private/Insights/StoreService/StoreBrowser.cpp b/Engine/Source/Developer/TraceInsights/Private/Insights/StoreService/StoreBrowser.cpp
--- a/Engine/Source/Developer/TraceInsights/Private/Insights/StoreService/StoreBrowser.cpp
+++ b/Engine/Source/Developer/TraceInsights/Private/Insights/StoreService/StoreBrowser.cpp
@@ -506,6 +506,7 @@ void FStoreBrowser::UpdateMetadata(FStoreBrowserTraceInfo& Trace)
 			Trace.ChangeSerial++;
 			Trace.Platform = Analyzer.Platform;
 			Trace.AppName = Analyzer.AppName;
+			Trace.ProjectName = Analyzer.ProjectName;
 			Trace.CommandLine = Analyzer.CommandLine;
 			Trace.Branch = Analyzer.Branch;
 			Trace.BuildVersion = Analyzer.BuildVersion;
diff --git a/Engine/Source/Developer/TraceInsights/Private/Insights/StoreService/StoreBrowser.h b/Engine/Source/Developer/TraceInsights/Private/Insights/StoreService/StoreBrowser.h
--- a/Engine/Source/Developer/TraceInsights/Private/Insights/StoreService/StoreBrowser.h
+++ b/Engine/Source/Developer/TraceInsights/Private/Insights/StoreService/StoreBrowser.h
@@ -38,6 +38,7 @@ struct FStoreBrowserTraceInfo
 
 	FString Platform;
 	FString AppName;
+	FString ProjectName;
 	FString CommandLine;
 	FString Branch;
 	FString BuildVersion;
diff --git a/Engine/Source/Developer/TraceInsights/Private/Insights/Widgets/SStartPageWindow.cpp b/Engine/Source/Developer/TraceInsights/Private/Insights/Widgets/SStartPageWindow.cpp
--- a/Engine/Source/Developer/TraceInsights/Private/Insights/Widgets/SStartPageWindow.cpp
+++ b/Engine/Source/Developer/TraceInsights/Private/Insights/Widgets/SStartPageWindow.cpp
@@ -2292,7 +2292,14 @@ void STraceStoreWindow::UpdateTrace(FTraceViewModel& InOutTrace, const Insights:
 	{
 		InOutTrace.bIsMetadataUpdated = true;
 		InOutTrace.Platform = FText::FromString(InSourceTrace.Platform);
-		InOutTrace.AppName = FText::FromString(InSourceTrace.AppName);
+		if (!InSourceTrace.ProjectName.IsEmpty())
+		{
+			InOutTrace.AppName = FText::FromString(InSourceTrace.ProjectName);
+		}
+		else
+		{
+			InOutTrace.AppName = FText::FromString(InSourceTrace.AppName);
+		}
 		InOutTrace.CommandLine = FText::FromString(InSourceTrace.CommandLine);
 		InOutTrace.Branch = FText::FromString(InSourceTrace.Branch);
 		InOutTrace.BuildVersion = FText::FromString(InSourceTrace.BuildVersion);
