Add support for plotting the timer series in the graph track

diff --git a/Engine/Source/Developer/TraceInsights/Private/Insights/ViewModels/FrameStatsHelper.cpp b/Engine/Source/Developer/TraceInsights/Private/Insights/ViewModels/FrameStatsHelper.cpp
new file mode 100644
--- /dev/null
+++ b/Engine/Source/Developer/TraceInsights/Private/Insights/ViewModels/FrameStatsHelper.cpp
@@ -0,0 +1,93 @@
+// Copyright Epic Games, Inc. All Rights Reserved.
+
+#include "FrameStatsHelper.h"
+
+#include "TraceServices/Model/TimingProfiler.h"
+
+#include "Insights/InsightsManager.h"
+
+namespace Insights
+{
+
+void FFrameStatsHelper::ComputeFrameStatsForTimer(TArray<FFrameStatsCachedEvent>& FrameStatsEvents, uint32 TimerId)
+{
+	TSharedPtr<const TraceServices::IAnalysisSession> Session = FInsightsManager::Get()->GetSession();
+	if (Session.IsValid())
+	{
+		Session->ReadAccessCheck();
+
+		const double SessionDuration = Session->GetDurationSeconds();
+
+		const TraceServices::ITimingProfilerProvider& TimingProfilerProvider = *TraceServices::ReadTimingProfilerProvider(*Session.Get());
+
+		const TraceServices::ITimingProfilerTimerReader* TimerReader;
+		TimingProfilerProvider.ReadTimers([&TimerReader](const TraceServices::ITimingProfilerTimerReader& Out) { TimerReader = &Out; });
+
+		const uint32 TimelineCount = TimingProfilerProvider.GetTimelineCount();
+		for (uint32 TimelineIndex = 0; TimelineIndex < TimelineCount; ++TimelineIndex)
+		{
+			TimingProfilerProvider.ReadTimeline(TimelineIndex,
+				[SessionDuration, &FrameStatsEvents, TimerReader, TimerId](const TraceServices::ITimingProfilerProvider::Timeline& Timeline)
+				{
+					TraceServices::ITimeline<TraceServices::FTimingProfilerEvent>::EnumerateAsyncParams Params;
+					Params.IntervalStart = 0;
+					Params.IntervalEnd = SessionDuration;
+					Params.Resolution = 0.0;
+					Params.SetupCallback = [](uint32 NumTasks) {};
+					Params.Callback = [TimerReader, &FrameStatsEvents, SessionDuration, TimerId](double StartTime, double EndTime, uint32 Depth, const TraceServices::FTimingProfilerEvent& Event, uint32 TaskIndex)
+					{
+						const TraceServices::FTimingProfilerTimer* Timer = TimerReader->GetTimer(Event.TimerIndex);
+						if (ensure(Timer != nullptr))
+						{
+							if (Timer->Id == TimerId)
+							{
+								int32 Index = Algo::UpperBoundBy(FrameStatsEvents, StartTime, &FFrameStatsCachedEvent::FrameStartTime);
+								if (Index > 0)
+								{
+									--Index;
+								}
+
+								// This can can happen when the event is between frames.
+								if (StartTime > FrameStatsEvents[Index].FrameEndTime)
+								{
+									Index++;
+									if (Index >= FrameStatsEvents.Num())
+									{
+										return TraceServices::EEventEnumerate::Continue;
+									}
+								}
+
+								do
+								{
+									FFrameStatsCachedEvent& Entry = FrameStatsEvents[Index];
+
+									if (EndTime < Entry.FrameStartTime)
+									{
+										return TraceServices::EEventEnumerate::Continue;
+									}
+
+									if (StartTime < Entry.FrameStartTime)
+									{
+										StartTime = Entry.FrameStartTime;
+									}
+
+									const double Duration = FMath::Min(EndTime, Entry.FrameEndTime) - StartTime;
+									ensure(Duration >= 0.0f);
+									for (double Value = Entry.Duration.load(); !Entry.Duration.compare_exchange_strong(Value, Value + Duration););
+
+									Index++;
+								} while (Index < FrameStatsEvents.Num());
+							}
+						}
+						return TraceServices::EEventEnumerate::Continue;
+					};
+
+					Timeline.EnumerateEventsDownSampledAsync(Params);
+				});
+		}
+	}
+}
+
+// namespace Insights
+};
+
diff --git a/Engine/Source/Developer/TraceInsights/Private/Insights/ViewModels/FrameStatsHelper.h b/Engine/Source/Developer/TraceInsights/Private/Insights/ViewModels/FrameStatsHelper.h
new file mode 100644
--- /dev/null
+++ b/Engine/Source/Developer/TraceInsights/Private/Insights/ViewModels/FrameStatsHelper.h
@@ -0,0 +1,41 @@
+// Copyright Epic Games, Inc. All Rights Reserved.
+
+#pragma once
+
+#include "Containers/Array.h"
+
+#include <atomic>
+
+namespace Insights
+{
+struct FFrameStatsCachedEvent
+{
+	FFrameStatsCachedEvent()
+		: FrameStartTime(0.0f)
+		, FrameEndTime(0.0f)
+	{
+		Duration.store(0.0f);
+	}
+
+	FFrameStatsCachedEvent(const FFrameStatsCachedEvent& Other)
+		: FrameStartTime(Other.FrameStartTime)
+		, FrameEndTime(Other.FrameEndTime)
+	{
+		Duration.store(Other.Duration.load());
+	}
+
+	double FrameStartTime;
+	double FrameEndTime;
+
+	// The sum of the all the instances of a timing event in a frame.
+	std::atomic<double> Duration;
+};
+
+class FFrameStatsHelper
+{
+public:
+	static void ComputeFrameStatsForTimer(TArray<FFrameStatsCachedEvent>& FrameStatsEvents, uint32 TimerId);
+};
+
+// namespace Insights
+}
diff --git a/Engine/Source/Developer/TraceInsights/Private/Insights/ViewModels/FrameTrackHelper.h b/Engine/Source/Developer/TraceInsights/Private/Insights/ViewModels/FrameTrackHelper.h
--- a/Engine/Source/Developer/TraceInsights/Private/Insights/ViewModels/FrameTrackHelper.h
+++ b/Engine/Source/Developer/TraceInsights/Private/Insights/ViewModels/FrameTrackHelper.h
@@ -66,15 +66,25 @@ struct FFrameTrackSample
 
 ////////////////////////////////////////////////////////////////////////////////////////////////////
 
+enum class EFrameTrackSeriesType : uint32
+{
+	Frame,
+	TimerFrameStats
+};
+
+////////////////////////////////////////////////////////////////////////////////////////////////////
+
 struct FFrameTrackSeries
 {
 	int32 FrameType;
+	EFrameTrackSeriesType Type;
 	bool bIsVisible;
 	int32 NumAggregatedFrames; // total number of frames aggregated in samples; i.e. sum of all Sample.NumFrames
 	TArray<FFrameTrackSample> Samples; // the aggregated samples
 
-	explicit FFrameTrackSeries(int32 InFrameType)
+	explicit FFrameTrackSeries(int32 InFrameType, EFrameTrackSeriesType InType)
 		: FrameType(InFrameType)
+		, Type(InType)
 		, bIsVisible(true)
 		, NumAggregatedFrames(0)
 		, Samples()
@@ -90,6 +100,20 @@ struct FFrameTrackSeries
 
 ////////////////////////////////////////////////////////////////////////////////////////////////////
 
+struct FTimerFrameStatsTrackSeries : public FFrameTrackSeries
+{
+	FTimerFrameStatsTrackSeries(int32 InFrameType, uint32 InTimerId)
+		: FFrameTrackSeries(InFrameType, EFrameTrackSeriesType::TimerFrameStats)
+		, TimerId(InTimerId)
+	{
+	}
+
+	uint32 TimerId;
+	FText TimerDisplayName;
+};
+
+////////////////////////////////////////////////////////////////////////////////////////////////////
+
 class FFrameTrackSeriesBuilder
 {
 public:
diff --git a/Engine/Source/Developer/TraceInsights/Private/Insights/ViewModels/TimingGraphTrack.cpp b/Engine/Source/Developer/TraceInsights/Private/Insights/ViewModels/TimingGraphTrack.cpp
--- a/Engine/Source/Developer/TraceInsights/Private/Insights/ViewModels/TimingGraphTrack.cpp
+++ b/Engine/Source/Developer/TraceInsights/Private/Insights/ViewModels/TimingGraphTrack.cpp
@@ -13,6 +13,7 @@
 #include "Insights/Common/TimeUtils.h"
 #include "Insights/InsightsManager.h"
 #include "Insights/ViewModels/AxisViewportDouble.h"
+#include "Insights/ViewModels/FrameStatsHelper.h"
 #include "Insights/ViewModels/GraphTrackBuilder.h"
 #include "Insights/ViewModels/ITimingViewDrawHelper.h"
 #include "Insights/ViewModels/TimingTrackViewport.h"
@@ -481,96 +482,29 @@ void FTimingGraphTrack::UpdateFrameStatsTimerSeries(FTimingGraphSeries& Series,
 
 			FramesProvider.EnumerateFrames(Series.FrameType, (uint64) 0, (uint64) FrameCount, [&Series](const TraceServices::FFrame& Frame)
 				{
-					FTimingGraphSeries::FAtomicTimingEvent Event;
+					Insights::FFrameStatsCachedEvent Event;
 					Event.FrameStartTime = Frame.StartTime;
 					Event.FrameEndTime = Frame.EndTime;
 					Event.Duration.store(0.0f);
 					Series.FrameStatsCachedEvents.Add(Event);
 				});
 
-			const TraceServices::ITimingProfilerProvider& TimingProfilerProvider = *TraceServices::ReadTimingProfilerProvider(*Session.Get());
-
-			const TraceServices::ITimingProfilerTimerReader* TimerReader;
-			TimingProfilerProvider.ReadTimers([&TimerReader](const TraceServices::ITimingProfilerTimerReader& Out) { TimerReader = &Out; });
-
-			const uint32 TimelineCount = TimingProfilerProvider.GetTimelineCount();
-			for (uint32 TimelineIndex = 0; TimelineIndex < TimelineCount; ++TimelineIndex)
-			{
-				TimingProfilerProvider.ReadTimeline(TimelineIndex,
-					[SessionDuration, &Series, TimerReader](const TraceServices::ITimingProfilerProvider::Timeline& Timeline)
-					{
-						TArray<TArray<FTimingGraphSeries::FSimpleTimingEvent>> Events;
-						TraceServices::ITimeline<TraceServices::FTimingProfilerEvent>::EnumerateAsyncParams Params;
-						Params.IntervalStart = 0;
-						Params.IntervalEnd = SessionDuration;
-						Params.Resolution = 0.0;
-						Params.SetupCallback = [&Events](uint32 NumTasks) {};
-						Params.Callback = [&Events, TimerReader, &Series, SessionDuration](double StartTime, double EndTime, uint32 Depth, const TraceServices::FTimingProfilerEvent& Event, uint32 TaskIndex)
-						{
-							const TraceServices::FTimingProfilerTimer* Timer = TimerReader->GetTimer(Event.TimerIndex);
-							if (ensure(Timer != nullptr))
-							{
-								if (Timer->Id == Series.TimerId)
-								{
-									int32 Index = Algo::UpperBoundBy(Series.FrameStatsCachedEvents, StartTime, &FTimingGraphSeries::FAtomicTimingEvent::FrameStartTime);
-									if (Index > 0)
-									{
-										--Index;
-									}
-
-									// This can can happen when the event is between frames.
-									if (StartTime > Series.FrameStatsCachedEvents[Index].FrameEndTime)
-									{
-										Index++;
-										if (Index >= Series.FrameStatsCachedEvents.Num())
-										{
-											return TraceServices::EEventEnumerate::Continue;
-										}
-									}
-
-									do
-									{
-										FTimingGraphSeries::FAtomicTimingEvent& Entry = Series.FrameStatsCachedEvents[Index];
-
-										if (EndTime < Entry.FrameStartTime)
-										{
-											return TraceServices::EEventEnumerate::Continue;
-										}
-
-										if (StartTime < Entry.FrameStartTime)
-										{
-											StartTime = Entry.FrameStartTime;
-										}
-
-										const double Duration = FMath::Min(EndTime, Entry.FrameEndTime) - StartTime;
-										ensure(Duration >= 0.0f);
-										for (double Value = Entry.Duration.load(); !Entry.Duration.compare_exchange_strong(Value, Value + Duration););
-
-										Index++;
-									} while (Index < Series.FrameStatsCachedEvents.Num());
-								}
-							}
-							return TraceServices::EEventEnumerate::Continue;
-						};
-
-						Timeline.EnumerateEventsDownSampledAsync(Params);
-					});
-			}
+			Insights::FFrameStatsHelper::ComputeFrameStatsForTimer(Series.FrameStatsCachedEvents, Series.TimerId);
 		}
 
-		int32 StartIndex = Algo::UpperBoundBy(Series.FrameStatsCachedEvents, Viewport.GetStartTime(), &FTimingGraphSeries::FAtomicTimingEvent::FrameStartTime);
+		int32 StartIndex = Algo::UpperBoundBy(Series.FrameStatsCachedEvents, Viewport.GetStartTime(), &Insights::FFrameStatsCachedEvent::FrameStartTime);
 		if (StartIndex > 0)
 		{
 			StartIndex--;
 		}
-		int32 EndIndex = Algo::UpperBoundBy(Series.FrameStatsCachedEvents, Viewport.GetEndTime(), &FTimingGraphSeries::FAtomicTimingEvent::FrameStartTime);
+		int32 EndIndex = Algo::UpperBoundBy(Series.FrameStatsCachedEvents, Viewport.GetEndTime(), &Insights::FFrameStatsCachedEvent::FrameStartTime);
 		if (EndIndex < Series.FrameStatsCachedEvents.Num())
 		{
 			EndIndex++;
 		}
 		for (int32 Index = StartIndex; Index < EndIndex; ++Index)
 		{
-			const FTimingGraphSeries::FAtomicTimingEvent& Entry = Series.FrameStatsCachedEvents[Index];
+			const Insights::FFrameStatsCachedEvent& Entry = Series.FrameStatsCachedEvents[Index];
 			Builder.AddEvent(Entry.FrameStartTime, Entry.Duration.load(), Entry.Duration.load());
 		}
 	}
diff --git a/Engine/Source/Developer/TraceInsights/Private/Insights/ViewModels/TimingGraphTrack.h b/Engine/Source/Developer/TraceInsights/Private/Insights/ViewModels/TimingGraphTrack.h
--- a/Engine/Source/Developer/TraceInsights/Private/Insights/ViewModels/TimingGraphTrack.h
+++ b/Engine/Source/Developer/TraceInsights/Private/Insights/ViewModels/TimingGraphTrack.h
@@ -8,6 +8,11 @@
 #include "Insights/ViewModels/GraphSeries.h"
 #include "Insights/ViewModels/GraphTrack.h"
 
+namespace Insights
+{
+	struct FFrameStatsCachedEvent;
+}
+
 ////////////////////////////////////////////////////////////////////////////////////////////////////
 
 class FTimingGraphSeries : public FGraphSeries
@@ -27,27 +32,6 @@ public:
 		double Duration;
 	};
 
-	struct FAtomicTimingEvent
-	{
-		FAtomicTimingEvent()
-		{
-			FrameStartTime = 0.0f;
-			FrameEndTime = 0.0f;
-			Duration.store(0.0f);
-		}
-
-		FAtomicTimingEvent(const FAtomicTimingEvent& Other)
-		{
-			FrameStartTime = Other.FrameStartTime;
-			FrameEndTime = Other.FrameEndTime;
-			Duration.store(Other.Duration.load());
-		}
-
-		double FrameEndTime;
-		double FrameStartTime;
-		std::atomic<double> Duration;
-	};
-
 public:
 	explicit FTimingGraphSeries(FTimingGraphSeries::ESeriesType Type);
 	virtual ~FTimingGraphSeries();
@@ -70,7 +54,7 @@ public:
 	ETraceFrameType FrameType;
 	double CachedSessionDuration;
 	TArray<FSimpleTimingEvent> CachedEvents; // used by Timer series
-	TArray<FAtomicTimingEvent> FrameStatsCachedEvents; // used by Frame Stats Timer series
+	TArray<Insights::FFrameStatsCachedEvent> FrameStatsCachedEvents; // used by Frame Stats Timer series
 
 	bool bIsTime; // the unit for values is [second]
 	bool bIsMemory; // the unit for value is [byte]
diff --git a/Engine/Source/Developer/TraceInsights/Private/Insights/Widgets/SFrameTrack.cpp b/Engine/Source/Developer/TraceInsights/Private/Insights/Widgets/SFrameTrack.cpp
--- a/Engine/Source/Developer/TraceInsights/Private/Insights/Widgets/SFrameTrack.cpp
+++ b/Engine/Source/Developer/TraceInsights/Private/Insights/Widgets/SFrameTrack.cpp
@@ -11,6 +11,7 @@
 #include "Rendering/DrawElements.h"
 #include "Styling/AppStyle.h"
 #include "TraceServices/Model/Frames.h"
+#include "TraceServices/Model/TimingProfiler.h"
 #include "Widgets/Layout/SScrollBar.h"
 
 // Insights
@@ -22,6 +23,7 @@
 #include "Insights/Log.h"
 #include "Insights/TimingProfilerCommon.h"
 #include "Insights/TimingProfilerManager.h"
+#include "Insights/ViewModels/FrameStatsHelper.h"
 #include "Insights/ViewModels/FrameTrackHelper.h"
 #include "Insights/Widgets/STimingProfilerWindow.h"
 #include "Insights/Widgets/STimingView.h"
@@ -60,16 +62,8 @@ void SFrameTrack::Reset()
 	ViewportY.SetScale(1500.0);
 	bIsViewportDirty = true;
 
-	SeriesMap.Reset();
-	SeriesOrder.Reset();
-	SeriesOrder.Add(TraceFrameType_Game);
-	SeriesOrder.Add(TraceFrameType_Rendering);
-
 	bIsStateDirty = true;
 
-	bShowGameFrames = true;
-	bShowRenderingFrames = true;
-
 	bIsAutoZoomEnabled = true;
 	AutoZoomViewportPos = ViewportX.GetPos();
 	AutoZoomViewportScale = ViewportX.GetScale();
@@ -106,6 +100,8 @@ void SFrameTrack::Reset()
 	DrawDurationHistory.Reset();
 	OnPaintDurationHistory.Reset();
 	LastOnPaintTime = FPlatformTime::Cycles64();
+
+	AllSeries.Empty();
 }
 
 ////////////////////////////////////////////////////////////////////////////////////////////////////
@@ -187,7 +183,7 @@ void SFrameTrack::Tick(const FGeometry& AllottedGeometry, const double InCurrent
 
 			for (int32 FrameType = 0; FrameType < TraceFrameType_Count; ++FrameType)
 			{
-				TSharedPtr<FFrameTrackSeries> SeriesPtr = FindOrAddSeries(FrameType);
+				TSharedPtr<FFrameTrackSeries> SeriesPtr = FindOrAddSeries(static_cast<ETraceFrameType>(FrameType));
 
 				int32 NumFrames = static_cast<int32>(FramesProvider.GetFrameCount(static_cast<ETraceFrameType>(FrameType)));
 				if (NumFrames > ViewportX.GetMaxValue())
@@ -219,37 +215,60 @@ void SFrameTrack::Tick(const FGeometry& AllottedGeometry, const double InCurrent
 
 ////////////////////////////////////////////////////////////////////////////////////////////////////
 
-TSharedRef<FFrameTrackSeries> SFrameTrack::FindOrAddSeries(int32 FrameType)
+TSharedRef<FFrameTrackSeries> SFrameTrack::FindOrAddSeries(ETraceFrameType FrameType)
 {
-	TSharedPtr<FFrameTrackSeries>* SeriesPtrPtr = SeriesMap.Find(FrameType);
-	if (SeriesPtrPtr)
-	{
-		ensure((**SeriesPtrPtr).FrameType == FrameType);
-		return (*SeriesPtrPtr).ToSharedRef();
-	}
-	else
+	TSharedPtr<FFrameTrackSeries>* ExistingSeries = AllSeries.FindByPredicate([FrameType](TSharedPtr<FFrameTrackSeries> Series)
+		{
+			return Series->Type == EFrameTrackSeriesType::Frame &&
+				Series->FrameType == FrameType;
+		});
+
+	if (ExistingSeries != nullptr)
 	{
-		LLM_SCOPE_BYTAG(Insights);
-		TSharedRef<FFrameTrackSeries> SeriesRef = MakeShared<FFrameTrackSeries>(FrameType);
-		SeriesMap.Add(FrameType, SeriesRef);
-		return SeriesRef;
+		return ExistingSeries->ToSharedRef();
 	}
+
+	LLM_SCOPE_BYTAG(Insights);
+	TSharedRef<FFrameTrackSeries> SeriesRef = MakeShared<FFrameTrackSeries>(FrameType, EFrameTrackSeriesType::Frame);
+	AllSeries.Add(SeriesRef);
+	return SeriesRef;
 }
 
 ////////////////////////////////////////////////////////////////////////////////////////////////////
 
-TSharedPtr<FFrameTrackSeries> SFrameTrack::FindSeries(int32 FrameType) const
+TSharedPtr<FFrameTrackSeries> SFrameTrack::FindSeries(ETraceFrameType FrameType) const
 {
-	const TSharedPtr<FFrameTrackSeries>* SeriesPtrPtr = SeriesMap.Find(FrameType);
-	if (SeriesPtrPtr)
+	const TSharedPtr<FFrameTrackSeries>* ExistingSeries = AllSeries.FindByPredicate([FrameType](TSharedPtr<FFrameTrackSeries> Series)
+		{
+			return Series->Type == EFrameTrackSeriesType::Frame &&
+				Series->FrameType == FrameType;
+		});
+
+	if (ExistingSeries != nullptr)
 	{
-		ensure((**SeriesPtrPtr).FrameType == FrameType);
-		return *SeriesPtrPtr;
+		return *ExistingSeries;
 	}
-	else
+
+	return nullptr;
+}
+
+////////////////////////////////////////////////////////////////////////////////////////////////////
+
+TSharedPtr<FFrameTrackSeries> SFrameTrack::FindFrameStatsSeries(ETraceFrameType FrameType, uint32 TimerId) const
+{
+	const TSharedPtr<FFrameTrackSeries>* ExistingSeries = AllSeries.FindByPredicate([FrameType, TimerId](TSharedPtr<FFrameTrackSeries> Series)
+		{
+			return Series->Type == EFrameTrackSeriesType::TimerFrameStats &&
+				Series->FrameType == FrameType &&
+				StaticCastSharedPtr<FTimerFrameStatsTrackSeries>(Series)->TimerId == TimerId;
+		});
+
+	if (ExistingSeries)
 	{
-		return nullptr;
+		return *ExistingSeries;
 	}
+
+	return nullptr;
 }
 
 ////////////////////////////////////////////////////////////////////////////////////////////////////
@@ -260,10 +279,9 @@ void SFrameTrack::UpdateState()
 	Stopwatch.Start();
 
 	// Reset stats.
-	for (TPair<int32, TSharedPtr<FFrameTrackSeries>>& KeyValuePair : SeriesMap)
+	for (TSharedPtr<FFrameTrackSeries> Series : AllSeries)
 	{
-		TSharedPtr<FFrameTrackSeries>& SeriesPtr = KeyValuePair.Value;
-		SeriesPtr->NumAggregatedFrames = 0;
+		Series->NumAggregatedFrames = 0;
 	}
 	NumUpdatedFrames = 0;
 
@@ -281,7 +299,7 @@ void SFrameTrack::UpdateState()
 
 		for (int32 FrameType = 0; FrameType < TraceFrameType_Count; ++FrameType)
 		{
-			TSharedPtr<FFrameTrackSeries> SeriesPtr = FindOrAddSeries(FrameType);
+			TSharedPtr<FFrameTrackSeries> SeriesPtr = FindOrAddSeries(static_cast<ETraceFrameType>(FrameType));
 
 			LLM_SCOPE_BYTAG(Insights);
 			FFrameTrackSeriesBuilder Builder(*SeriesPtr, Viewport);
@@ -293,6 +311,42 @@ void SFrameTrack::UpdateState()
 
 			NumUpdatedFrames += Builder.GetNumAddedFrames();
 		}
+
+		for (int32 Index = 0; Index < AllSeries.Num(); ++Index)
+		{
+			TSharedPtr<FFrameTrackSeries> Series = AllSeries[Index];
+			if (Series->Type != EFrameTrackSeriesType::TimerFrameStats)
+			{
+				continue;
+			}
+
+			TSharedPtr<FTimerFrameStatsTrackSeries> TimerSeries = StaticCastSharedPtr<FTimerFrameStatsTrackSeries>(Series);
+			TArray<Insights::FFrameStatsCachedEvent> Frames;
+			FramesProvider.EnumerateFrames(static_cast<ETraceFrameType>(Series->FrameType), StartIndex, EndIndex, [&Frames](const TraceServices::FFrame& Frame)
+				{
+					Insights::FFrameStatsCachedEvent Event;
+					Event.FrameStartTime = Frame.StartTime;
+					Event.FrameEndTime = Frame.EndTime;
+					Event.Duration.store(0.0f);
+					Frames.Add(Event);
+				});
+
+			FFrameTrackSeriesBuilder Builder(*Series, Viewport);
+
+			Insights::FFrameStatsHelper::ComputeFrameStatsForTimer(Frames, TimerSeries->TimerId);
+
+			int32 CurrentIndex = StartIndex;
+			for (Insights::FFrameStatsCachedEvent& Event : Frames)
+			{
+				TraceServices::FFrame NewFrame;
+				NewFrame.StartTime = Event.FrameStartTime;
+				NewFrame.EndTime = Event.FrameStartTime + Event.Duration.load();
+				NewFrame.Index = CurrentIndex++;
+				Builder.AddFrame(NewFrame);
+			}
+
+			NumUpdatedFrames += Builder.GetNumAddedFrames();
+		}
 	}
 
 	Stopwatch.Stop();
@@ -312,17 +366,15 @@ FFrameTrackSampleRef SFrameTrack::GetSampleAtMousePosition(double X, double Y)
 			const float MY = static_cast<float>(Y);
 
 			// Search in reverse paint order.
-			for (int32 SeriesIndex = SeriesOrder.Num() - 1; SeriesIndex >= 0; --SeriesIndex)
+			for (int32 SeriesIndex = AllSeries.Num() - 1; SeriesIndex >= 0; --SeriesIndex)
 			{
-				int32 FrameType = SeriesOrder[SeriesIndex];
+				TSharedPtr<FFrameTrackSeries> SeriesPtr = AllSeries[SeriesIndex];
 
-				if ((FrameType == TraceFrameType_Rendering && !bShowRenderingFrames) ||
-					(FrameType == TraceFrameType_Game && !bShowGameFrames))
+				if (!SeriesPtr->bIsVisible)
 				{
 					continue;
 				}
 
-				const TSharedPtr<FFrameTrackSeries> SeriesPtr = FindSeries(FrameType);
 				if (SeriesPtr.IsValid())
 				{
 					if (SeriesPtr->NumAggregatedFrames > 0 &&
@@ -447,20 +499,16 @@ int32 SFrameTrack::OnPaint(const FPaintArgs& Args, const FGeometry& AllottedGeom
 		DrawHorizontalAxisGrid(DrawContext, WhiteBrush, SummaryFont, true);
 
 		// Draw frames, for each visible Series.
-		for (int32 SeriesIndex = 0; SeriesIndex < SeriesOrder.Num(); ++SeriesIndex)
+		for (TSharedPtr<FFrameTrackSeries> Series : AllSeries)
 		{
-			int32 FrameType = SeriesOrder[SeriesIndex];
-
-			if ((FrameType == TraceFrameType_Rendering && !bShowRenderingFrames) ||
-				(FrameType == TraceFrameType_Game && !bShowGameFrames))
+			if (!Series->bIsVisible)
 			{
 				continue;
 			}
 
-			const TSharedPtr<FFrameTrackSeries> SeriesPtr = FindSeries(FrameType);
-			if (SeriesPtr.IsValid())
+			if (Series.IsValid())
 			{
-				Helper.DrawCached(*SeriesPtr);
+				Helper.DrawCached(*Series);
 			}
 		}
 
@@ -472,12 +520,8 @@ int32 SFrameTrack::OnPaint(const FPaintArgs& Args, const FGeometry& AllottedGeom
 			TSharedPtr<STimingView> TimingView = Window->GetTimingView();
 			if (TimingView)
 			{
-				TSharedPtr<FFrameTrackSeries> SeriesPtr = nullptr;
-				for (const TPair<int32, TSharedPtr<FFrameTrackSeries>>& KeyValuePair : SeriesMap)
-				{
-					SeriesPtr = KeyValuePair.Value;
-					break; // stop at first enumerated Series
-				}
+				TSharedPtr<FFrameTrackSeries> SeriesPtr = AllSeries[0];
+				
 				if (SeriesPtr.IsValid())
 				{
 					// Highlight the area corresponding to viewport of Timing View.
@@ -517,7 +561,18 @@ int32 SFrameTrack::OnPaint(const FPaintArgs& Args, const FGeometry& AllottedGeom
 
 			// First line: "Rendering Frame 1,234"
 			TStringBuilder<512> StringBuilder;
-			StringBuilder.Append(FFrameTrackDrawHelper::FrameTypeToString(HoveredSample.Series->FrameType));
+
+			if (HoveredSample.Series->Type == EFrameTrackSeriesType::TimerFrameStats)
+			{
+				StringBuilder.Append(StaticCastSharedPtr<FTimerFrameStatsTrackSeries>(HoveredSample.Series)->TimerDisplayName.ToString());
+				StringBuilder.Append(TEXT(" "));
+				StringBuilder.Append(FFrameTrackDrawHelper::FrameTypeToString(HoveredSample.Series->FrameType));
+			}
+			else
+			{
+				StringBuilder.Append(FFrameTrackDrawHelper::FrameTypeToString(HoveredSample.Series->FrameType));
+			}
+
 			StringBuilder.Append(TEXT(" Frame "));
 			StringBuilder.Append(FText::AsNumber(HoveredSample.Sample->LargestFrameIndex).ToString());
 			const FString Text1(StringBuilder);
@@ -1194,6 +1249,37 @@ void SFrameTrack::ShowContextMenu(const FPointerEvent& MouseEvent)
 	}
 	MenuBuilder.EndSection();
 
+	MenuBuilder.BeginSection("FrameStats", LOCTEXT("ContextMenu_Section_Stats", "Frame Stats"));
+
+	for (TSharedPtr<FFrameTrackSeries> Series : AllSeries)
+	{
+		if (Series->Type != EFrameTrackSeriesType::TimerFrameStats)
+		{
+			continue;
+		}
+
+		TSharedPtr<FTimerFrameStatsTrackSeries> FrameStatSeries = StaticCastSharedPtr<FTimerFrameStatsTrackSeries>(Series);
+
+		ETraceFrameType FrameType = static_cast<ETraceFrameType>(FrameStatSeries->FrameType);
+		FUIAction Action_ShowRenderingFrames
+		(
+			FExecuteAction::CreateSP(this, &SFrameTrack::ContextMenu_ShowFrameStats_Execute, FrameType, FrameStatSeries->TimerId),
+			FCanExecuteAction::CreateSP(this, &SFrameTrack::ContextMenu_ShowFrameStats_CanExecute, FrameType, FrameStatSeries->TimerId),
+			FIsActionChecked::CreateSP(this, &SFrameTrack::ContextMenu_ShowFrameStats_IsChecked, FrameType, FrameStatSeries->TimerId)
+		);
+		MenuBuilder.AddMenuEntry
+		(
+			FrameStatSeries->TimerDisplayName,
+			FText::Format(LOCTEXT("ContextMenu_ShowRenderingFrames_Desc", "Shows/hides the {0} series."), FrameStatSeries->TimerDisplayName),
+			FSlateIcon(),
+			Action_ShowRenderingFrames,
+			NAME_None,
+			EUserInterfaceActionType::ToggleButton
+		);
+	}
+
+	MenuBuilder.EndSection();
+
 	MenuBuilder.BeginSection("Zoom", LOCTEXT("ContextMenu_Section_Zoom", "Zoom"));
 	{
 		FUIAction Action_AutoZoom
@@ -1241,7 +1327,8 @@ void SFrameTrack::ShowContextMenu(const FPointerEvent& MouseEvent)
 
 void SFrameTrack::ContextMenu_ShowGameFrames_Execute()
 {
-	bShowGameFrames = !bShowGameFrames;
+	TSharedPtr<FFrameTrackSeries> Series = FindSeries(ETraceFrameType::TraceFrameType_Game);
+	Series->bIsVisible = !Series->bIsVisible;
 }
 
 ////////////////////////////////////////////////////////////////////////////////////////////////////
@@ -1255,14 +1342,16 @@ bool SFrameTrack::ContextMenu_ShowGameFrames_CanExecute()
 
 bool SFrameTrack::ContextMenu_ShowGameFrames_IsChecked()
 {
-	return bShowGameFrames;
+	TSharedPtr<FFrameTrackSeries> Series = FindSeries(ETraceFrameType::TraceFrameType_Game);
+	return Series.IsValid() ? Series->bIsVisible : false;
 }
 
 ////////////////////////////////////////////////////////////////////////////////////////////////////
 
 void SFrameTrack::ContextMenu_ShowRenderingFrames_Execute()
 {
-	bShowRenderingFrames = !bShowRenderingFrames;
+	TSharedPtr<FFrameTrackSeries> Series = FindSeries(ETraceFrameType::TraceFrameType_Rendering);
+	Series->bIsVisible = !Series->bIsVisible;
 }
 
 ////////////////////////////////////////////////////////////////////////////////////////////////////
@@ -1276,7 +1365,38 @@ bool SFrameTrack::ContextMenu_ShowRenderingFrames_CanExecute()
 
 bool SFrameTrack::ContextMenu_ShowRenderingFrames_IsChecked()
 {
-	return bShowRenderingFrames;
+	TSharedPtr<FFrameTrackSeries> Series = FindSeries(ETraceFrameType::TraceFrameType_Rendering);
+	return Series.IsValid() ? Series->bIsVisible : false;
+}
+
+////////////////////////////////////////////////////////////////////////////////////////////////////
+
+void SFrameTrack::ContextMenu_ShowFrameStats_Execute(ETraceFrameType FrameType, uint32 TimerId)
+{
+	TSharedPtr<FFrameTrackSeries> Series = FindFrameStatsSeries(FrameType, TimerId);
+
+	if (!Series.IsValid())
+	{
+		return;
+	}
+
+	Series->bIsVisible = !Series->bIsVisible;
+}
+
+////////////////////////////////////////////////////////////////////////////////////////////////////
+
+bool SFrameTrack::ContextMenu_ShowFrameStats_CanExecute(ETraceFrameType FrameType, uint32 TimerId)
+{
+	return true;
+}
+
+////////////////////////////////////////////////////////////////////////////////////////////////////
+
+bool SFrameTrack::ContextMenu_ShowFrameStats_IsChecked(ETraceFrameType FrameType, uint32 TimerId)
+{
+	TSharedPtr<FFrameTrackSeries> Series = FindFrameStatsSeries(FrameType, TimerId);
+
+	return Series.IsValid() ? Series->bIsVisible : false;
 }
 
 ////////////////////////////////////////////////////////////////////////////////////////////////////
@@ -1402,4 +1522,56 @@ void SFrameTrack::UpdateHorizontalScrollBar()
 
 ////////////////////////////////////////////////////////////////////////////////////////////////////
 
+bool SFrameTrack::HasFrameStatSeries(ETraceFrameType FrameType, uint32 TimerId)
+{
+	TSharedPtr<FFrameTrackSeries>* ExistingSeries = AllSeries.FindByPredicate([FrameType, TimerId](TSharedPtr<FFrameTrackSeries> Series)
+		{
+			return Series->Type == EFrameTrackSeriesType::TimerFrameStats &&
+				Series->FrameType == FrameType &&
+				StaticCastSharedPtr<FTimerFrameStatsTrackSeries>(Series)->TimerId == TimerId;
+		});
+
+	return ExistingSeries != nullptr;
+}
+
+////////////////////////////////////////////////////////////////////////////////////////////////////
+
+TSharedPtr<FTimerFrameStatsTrackSeries> SFrameTrack::AddTimerFrameStatSeries(ETraceFrameType FrameType, uint32 TimerId)
+{
+	TSharedPtr<FFrameTrackSeries> ExistingSeries = FindFrameStatsSeries(FrameType, TimerId);
+
+	if (ExistingSeries.IsValid())
+	{
+		return StaticCastSharedPtr<FTimerFrameStatsTrackSeries>(ExistingSeries);
+	}
+
+	TSharedRef<FTimerFrameStatsTrackSeries> SeriesRef = MakeShared<FTimerFrameStatsTrackSeries>(FrameType, TimerId);
+	SeriesRef->TimerId = TimerId;
+	AllSeries.Add(SeriesRef);
+
+	bIsStateDirty = true;
+
+	return SeriesRef;
+}
+
+////////////////////////////////////////////////////////////////////////////////////////////////////
+
+bool SFrameTrack::RemoveTimerFrameStatSeries(ETraceFrameType FrameType, uint32 TimerId)
+{
+	int32 NumRemoved = AllSeries.RemoveAll([FrameType, TimerId](TSharedPtr<FFrameTrackSeries> Series)
+		{
+			return Series->Type == EFrameTrackSeriesType::TimerFrameStats &&
+				Series->FrameType == FrameType &&
+				StaticCastSharedPtr<FTimerFrameStatsTrackSeries>(Series)->TimerId == TimerId;
+		});
+
+	ensure(NumRemoved == 1);
+
+	bIsStateDirty = true;
+
+	return NumRemoved >= 1;
+}
+
+////////////////////////////////////////////////////////////////////////////////////////////////////
+
 #undef LOCTEXT_NAMESPACE
diff --git a/Engine/Source/Developer/TraceInsights/Private/Insights/Widgets/SFrameTrack.h b/Engine/Source/Developer/TraceInsights/Private/Insights/Widgets/SFrameTrack.h
--- a/Engine/Source/Developer/TraceInsights/Private/Insights/Widgets/SFrameTrack.h
+++ b/Engine/Source/Developer/TraceInsights/Private/Insights/Widgets/SFrameTrack.h
@@ -128,9 +128,14 @@ public:
 
 	virtual FCursorReply OnCursorQuery(const FGeometry& MyGeometry, const FPointerEvent& CursorEvent) const override;
 
+	bool HasFrameStatSeries(ETraceFrameType FrameType, uint32 TimerId);
+	TSharedPtr<FTimerFrameStatsTrackSeries> AddTimerFrameStatSeries(ETraceFrameType FrameType, uint32 TimerId);
+	bool RemoveTimerFrameStatSeries(ETraceFrameType FrameType, uint32 TimerId);
+
 protected:
-	TSharedRef<FFrameTrackSeries> FindOrAddSeries(int32 FrameType);
-	TSharedPtr<FFrameTrackSeries> FindSeries(int32 FrameType) const;
+	TSharedRef<FFrameTrackSeries> FindOrAddSeries(ETraceFrameType FrameType);
+	TSharedPtr<FFrameTrackSeries> FindSeries(ETraceFrameType FrameType) const;
+	TSharedPtr<FFrameTrackSeries> FindFrameStatsSeries(ETraceFrameType FrameType, uint32 TimerId) const;
 	void UpdateState();
 
 	void DrawHorizontalAxisGrid(FDrawContext& DrawContext, const FSlateBrush* Brush, const FSlateFontInfo& Font, bool bDrawBackgroundLayer) const;
@@ -149,6 +154,10 @@ protected:
 	bool ContextMenu_ShowRenderingFrames_CanExecute();
 	bool ContextMenu_ShowRenderingFrames_IsChecked();
 
+	void ContextMenu_ShowFrameStats_Execute(ETraceFrameType FrameType, uint32 TimerId);
+	bool ContextMenu_ShowFrameStats_CanExecute(ETraceFrameType FrameType, uint32 TimerId);
+	bool ContextMenu_ShowFrameStats_IsChecked(ETraceFrameType FrameType, uint32 TimerId);
+
 	void ContextMenu_AutoZoom_Execute();
 	bool ContextMenu_AutoZoom_CanExecute();
 	bool ContextMenu_AutoZoom_IsChecked();
@@ -176,14 +185,10 @@ protected:
 	bool bIsViewportDirty;
 
 	/** Cached info for all frame series. */
-	TMap<int32, TSharedPtr<FFrameTrackSeries>> SeriesMap;
-	TArray<int32> SeriesOrder;
+	TArray<TSharedPtr<FFrameTrackSeries>> AllSeries;
 
 	bool bIsStateDirty;
 
-	bool bShowGameFrames;
-	bool bShowRenderingFrames;
-
 	bool bIsAutoZoomEnabled;
 
 	float AutoZoomViewportPos;
diff --git a/Engine/Source/Developer/TraceInsights/Private/Insights/Widgets/STimersView.cpp b/Engine/Source/Developer/TraceInsights/Private/Insights/Widgets/STimersView.cpp
--- a/Engine/Source/Developer/TraceInsights/Private/Insights/Widgets/STimersView.cpp
+++ b/Engine/Source/Developer/TraceInsights/Private/Insights/Widgets/STimersView.cpp
@@ -41,6 +41,7 @@
 #include "Insights/ViewModels/TimingExporter.h"
 #include "Insights/ViewModels/TimingGraphTrack.h"
 #include "Insights/Widgets/SAsyncOperationStatus.h"
+#include "Insights/Widgets/SFrameTrack.h"
 #include "Insights/Widgets/STimersViewTooltip.h"
 #include "Insights/Widgets/STimerTableRow.h"
 #include "Insights/Widgets/STimingProfilerWindow.h"
@@ -879,6 +880,73 @@ void STimersView::TreeView_BuildPlotTimerMenu(FMenuBuilder& MenuBuilder)
 			);
 		}
 	}
+
+	// Add/remove game frame stats series to/from frame track
+	{
+		FUIAction Action_ToggleFrameStatsTimerInFrameTrack;
+		Action_ToggleFrameStatsTimerInFrameTrack.CanExecuteAction = FCanExecuteAction::CreateLambda(CanExecute);
+		Action_ToggleFrameStatsTimerInFrameTrack.ExecuteAction = FExecuteAction::CreateSP(this, &STimersView::ToggleFrameTrackSeries, SelectedNode, ETraceFrameType::TraceFrameType_Game);
+
+		if (SelectedNode.IsValid() &&
+			SelectedNode->GetType() != ETimerNodeType::Group &&
+			IsSeriesInFrameTrack(SelectedNode, ETraceFrameType::TraceFrameType_Game))
+		{
+			MenuBuilder.AddMenuEntry
+			(
+				LOCTEXT("ContextMenu_RemoveGameFrameStatsSeriesFromFrameTrack", "Remove game frame stats series from frame track"),
+				LOCTEXT("ContextMenu_RemoveGameFrameStatsSeriesFromFrameTrack_Desc", "Remove the game frame stats series for this timer from the frame track."),
+				FSlateIcon(FInsightsStyle::GetStyleSetName(), "Icons.RemoveGraphSeries"),
+				Action_ToggleFrameStatsTimerInFrameTrack,
+				NAME_None,
+				EUserInterfaceActionType::Button
+			);
+		}
+		else
+		{
+			MenuBuilder.AddMenuEntry
+			(
+				LOCTEXT("ContextMenu_AddGameFrameStatsSeriesToFrameTrack", "Add game frame stats series to the frame track"),
+				LOCTEXT("ContextMenu_AddGameFrameStatsSeriesToFrameTrack_Desc", "Adds a game frame stats series for this timer to the frame track. Each data entry is computed as the sum of all instances of this timer in a game frame."),
+				FSlateIcon(FInsightsStyle::GetStyleSetName(), "Icons.AddGraphSeries"),
+				Action_ToggleFrameStatsTimerInFrameTrack,
+				NAME_None,
+				EUserInterfaceActionType::Button
+			);
+		}
+	}
+	// Add/remove rendering frame stats series to/from frame track
+	{
+		FUIAction Action_ToggleFrameStatsTimerInFrameTrack;
+		Action_ToggleFrameStatsTimerInFrameTrack.CanExecuteAction = FCanExecuteAction::CreateLambda(CanExecute);
+		Action_ToggleFrameStatsTimerInFrameTrack.ExecuteAction = FExecuteAction::CreateSP(this, &STimersView::ToggleFrameTrackSeries, SelectedNode, ETraceFrameType::TraceFrameType_Rendering);
+
+		if (SelectedNode.IsValid() &&
+			SelectedNode->GetType() != ETimerNodeType::Group &&
+			IsSeriesInFrameTrack(SelectedNode, ETraceFrameType::TraceFrameType_Rendering))
+		{
+			MenuBuilder.AddMenuEntry
+			(
+				LOCTEXT("ContextMenu_RemoveRenderingFrameStatsFromFrameTrac", "Remove rendering frame stats series from frame track"),
+				LOCTEXT("ContextMenu_RemoveRenderingFrameStatsFromFrameTrack_Desc", "Remove the rendering frame stats series for this timer from the frame track."),
+				FSlateIcon(FInsightsStyle::GetStyleSetName(), "Icons.RemoveGraphSeries"),
+				Action_ToggleFrameStatsTimerInFrameTrack,
+				NAME_None,
+				EUserInterfaceActionType::Button
+			);
+		}
+		else
+		{
+			MenuBuilder.AddMenuEntry
+			(
+				LOCTEXT("ContextMenu_AddRenderingFrameStatsSeriesToFrameTrack", "Add rendering frame stats series to the frame track"),
+				LOCTEXT("ContextMenu_AddRenderingFrameStatsSeriesToFrameTrack_Desc", "Adds a rendering frame stats series for this timer to the frame track. Each data entry is computed as the sum of all instances of this timer in a rendering frame."),
+				FSlateIcon(FInsightsStyle::GetStyleSetName(), "Icons.AddGraphSeries"),
+				Action_ToggleFrameStatsTimerInFrameTrack,
+				NAME_None,
+				EUserInterfaceActionType::Button
+			);
+		}
+	}
 }
 
 ////////////////////////////////////////////////////////////////////////////////////////////////////
@@ -2414,6 +2482,14 @@ TSharedPtr<FTimingGraphTrack> STimersView::GetTimingViewMainGraphTrack() const
 
 ////////////////////////////////////////////////////////////////////////////////////////////////////
 
+TSharedPtr<SFrameTrack> STimersView::GetFrameTrack() const
+{
+	TSharedPtr<STimingProfilerWindow> Wnd = FTimingProfilerManager::Get()->GetProfilerWindow();
+	return Wnd.IsValid() ? Wnd->GetFrameView() : nullptr;
+}
+
+////////////////////////////////////////////////////////////////////////////////////////////////////
+
 void STimersView::ToggleGraphSeries(TSharedRef<FTimingGraphTrack> GraphTrack, FTimerNodeRef NodePtr) const
 {
 	const uint32 TimerId = NodePtr->GetTimerId();
@@ -2525,6 +2601,42 @@ void STimersView::ToggleTimingViewMainGraphEventFrameStatsSeries(FTimerNodePtr T
 
 ////////////////////////////////////////////////////////////////////////////////////////////////////
 
+bool STimersView::IsSeriesInFrameTrack(FTimerNodePtr TimerNode, ETraceFrameType FrameType) const
+{
+	TSharedPtr<SFrameTrack> FrameTrack = GetFrameTrack();
+
+	if (!FrameTrack.IsValid())
+	{
+		return false;
+	}
+
+	return FrameTrack->HasFrameStatSeries(FrameType, TimerNode->GetTimerId());
+}
+
+////////////////////////////////////////////////////////////////////////////////////////////////////
+
+void STimersView::ToggleFrameTrackSeries(FTimerNodePtr TimerNode, ETraceFrameType FrameType) const
+{
+	TSharedPtr<SFrameTrack> FrameTrack = GetFrameTrack();
+
+	if (!FrameTrack.IsValid())
+	{
+		return;
+	}
+
+	if (FrameTrack->HasFrameStatSeries(FrameType, TimerNode->GetTimerId()))
+	{
+		FrameTrack->RemoveTimerFrameStatSeries(FrameType, TimerNode->GetTimerId());
+	}
+	else
+	{
+		TSharedPtr<FTimerFrameStatsTrackSeries> Series = FrameTrack->AddTimerFrameStatSeries(FrameType, TimerNode->GetTimerId());
+		Series->TimerDisplayName = TimerNode->GetDisplayName();
+	}
+}
+
+////////////////////////////////////////////////////////////////////////////////////////////////////
+
 bool STimersView::ContextMenu_CopyToClipboard_CanExecute() const
 {
 	const TArray<FTimerNodePtr> SelectedNodes = TreeView->GetSelectedItems();
diff --git a/Engine/Source/Developer/TraceInsights/Private/Insights/Widgets/STimersView.h b/Engine/Source/Developer/TraceInsights/Private/Insights/Widgets/STimersView.h
--- a/Engine/Source/Developer/TraceInsights/Private/Insights/Widgets/STimersView.h
+++ b/Engine/Source/Developer/TraceInsights/Private/Insights/Widgets/STimersView.h
@@ -21,6 +21,7 @@
 
 ////////////////////////////////////////////////////////////////////////////////////////////////////
 
+class SFrameTrack;
 class FMenuBuilder;
 class FTimingGraphTrack;
 class FUICommandList;
@@ -302,6 +303,8 @@ private:
 	void TreeView_BuildPlotTimerMenu(FMenuBuilder& MenuBuilder);
 
 	TSharedPtr<FTimingGraphTrack> GetTimingViewMainGraphTrack() const;
+	TSharedPtr<SFrameTrack> GetFrameTrack() const;
+
 	void ToggleGraphSeries(TSharedRef<FTimingGraphTrack> GraphTrack, FTimerNodeRef NodePtr) const;
 	bool IsSeriesInTimingViewMainGraph(FTimerNodePtr TimerNode) const;
 	void ToggleTimingViewMainGraphEventSeries(FTimerNodePtr TimerNode) const;
@@ -310,6 +313,9 @@ private:
 	bool IsFrameStatsSeriesInTimingViewMainGraph(FTimerNodePtr TimerNode, ETraceFrameType FrameType) const;
 	void ToggleTimingViewMainGraphEventFrameStatsSeries(FTimerNodePtr TimerNode, ETraceFrameType FrameType) const;
 
+	bool IsSeriesInFrameTrack(FTimerNodePtr TimerNode, ETraceFrameType FrameType) const;
+	void ToggleFrameTrackSeries(FTimerNodePtr TimerNode, ETraceFrameType FrameType) const;
+
 	////////////////////////////////////////////////////////////////////////////////////////////////////
 
 	/**
diff --git a/Engine/Source/Developer/TraceInsights/Private/Insights/Widgets/STimingProfilerWindow.h b/Engine/Source/Developer/TraceInsights/Private/Insights/Widgets/STimingProfilerWindow.h
--- a/Engine/Source/Developer/TraceInsights/Private/Insights/Widgets/STimingProfilerWindow.h
+++ b/Engine/Source/Developer/TraceInsights/Private/Insights/Widgets/STimingProfilerWindow.h
@@ -43,6 +43,7 @@ public:
 	TSharedPtr<STimerTreeView> GetCalleesTreeView() const { return CalleesTreeView; }
 	TSharedPtr<SStatsView> GetStatsView() const { return StatsView; }
 	TSharedPtr<SLogView> GetLogView() const { return LogView; }
+	TSharedPtr<SFrameTrack> GetFrameView() const { return FrameTrack; }
 
 protected:
 	virtual const TCHAR* GetAnalyticsEventName() const override;
