Fix localization text issue

diff --git a/Engine/Source/Developer/TraceInsights/Private/Insights/Widgets/SFrameTrack.cpp b/Engine/Source/Developer/TraceInsights/Private/Insights/Widgets/SFrameTrack.cpp
--- a/Engine/Source/Developer/TraceInsights/Private/Insights/Widgets/SFrameTrack.cpp
+++ b/Engine/Source/Developer/TraceInsights/Private/Insights/Widgets/SFrameTrack.cpp
@@ -1251,6 +1251,7 @@ void SFrameTrack::ShowContextMenu(const FPointerEvent& MouseEvent)
 
 	MenuBuilder.BeginSection("FrameStats", LOCTEXT("ContextMenu_Section_Stats", "Frame Stats"));
 
+	FText TooltipTextBase = LOCTEXT("ContextMenu_ShowRenderingFrames_Desc", "Shows/hides the {0} series.");
 	for (TSharedPtr<FFrameTrackSeries> Series : AllSeries)
 	{
 		if (Series->Type != EFrameTrackSeriesType::TimerFrameStats)
@@ -1270,7 +1271,7 @@ void SFrameTrack::ShowContextMenu(const FPointerEvent& MouseEvent)
 		MenuBuilder.AddMenuEntry
 		(
 			FrameStatSeries->TimerDisplayName,
-			FText::Format(LOCTEXT("ContextMenu_ShowRenderingFrames_Desc", "Shows/hides the {0} series."), FrameStatSeries->TimerDisplayName),
+			FText::Format(TooltipTextBase, FrameStatSeries->TimerDisplayName),
 			FSlateIcon(),
 			Action_ShowRenderingFrames,
 			NAME_None,
