Fix crashes in Linux/Mac Insights.

Fix invalid access to UI components that doesn't exist when store settings are not supported.

diff --git a/Engine/Source/Developer/TraceInsights/Private/Insights/Widgets/SStartPageWindow.cpp b/Engine/Source/Developer/TraceInsights/Private/Insights/Widgets/SStartPageWindow.cpp
--- a/Engine/Source/Developer/TraceInsights/Private/Insights/Widgets/SStartPageWindow.cpp
+++ b/Engine/Source/Developer/TraceInsights/Private/Insights/Widgets/SStartPageWindow.cpp
@@ -2155,10 +2155,16 @@ void STraceStoreWindow::RefreshTraceList()
 			SettingsChangeSerial = NewSettingsSerial;
 
 			// Update the host text
-			StoreHostTextBox->SetText(FText::FromString(StoreBrowser->GetLockedHost()));
+			if (StoreHostTextBox)
+			{
+				StoreHostTextBox->SetText(FText::FromString(StoreBrowser->GetLockedHost()));
+			}
 
 			// Update the store text box
-			StoreDirTextBox->SetText(FText::FromString(StoreBrowser->GetLockedStoreDirectory()));
+			if (StoreDirTextBox)
+			{
+				StoreDirTextBox->SetText(FText::FromString(StoreBrowser->GetLockedStoreDirectory()));
+			}
 
 			// Update store directory model
 			StoreDirectoryModel.Empty(1);
@@ -2167,7 +2173,10 @@ void STraceStoreWindow::RefreshTraceList()
 				NAME_None,
 				ETraceDirOperations::ModifyStore|ETraceDirOperations::Explore
 			));
-			StoreDirListView->RequestListRefresh();
+			if (StoreDirListView)
+			{
+				StoreDirListView->RequestListRefresh();
+			}
 
 			// Update additional watch directory models
 			static const FName DirColor[] = {
@@ -2189,7 +2198,10 @@ void STraceStoreWindow::RefreshTraceList()
 				));
 				ColorIdx = FMath::Wrap(++ColorIdx, int32(-1), int32(UE_ARRAY_COUNT(DirColor))-1);
 			}
-			WatchDirsListView->RequestListRefresh();
+			if (WatchDirsListView)
+			{
+				WatchDirsListView->RequestListRefresh();
+			}
 
 			bSettingsChanged = true;
 		}
