[Insights] Fixed crash in Frames track when there are no frame series created yet.

diff --git a/Engine/Source/Developer/TraceInsights/Private/Insights/Widgets/SFrameTrack.cpp b/Engine/Source/Developer/TraceInsights/Private/Insights/Widgets/SFrameTrack.cpp
--- a/Engine/Source/Developer/TraceInsights/Private/Insights/Widgets/SFrameTrack.cpp
+++ b/Engine/Source/Developer/TraceInsights/Private/Insights/Widgets/SFrameTrack.cpp
@@ -514,20 +514,19 @@ int32 SFrameTrack::OnPaint(const FPaintArgs& Args, const FGeometry& AllottedGeom
 
 		NumDrawSamples = Helper.GetNumDrawSamples();
 
-		TSharedPtr<STimingProfilerWindow> Window = FTimingProfilerManager::Get()->GetProfilerWindow();
-		if (Window)
+		TSharedPtr<FFrameTrackSeries> GameFrameSeries = FindSeries(ETraceFrameType::TraceFrameType_Game);
+		if (GameFrameSeries.IsValid())
 		{
-			TSharedPtr<STimingView> TimingView = Window->GetTimingView();
-			if (TimingView)
+			TSharedPtr<STimingProfilerWindow> Window = FTimingProfilerManager::Get()->GetProfilerWindow();
+			if (Window)
 			{
-				TSharedPtr<FFrameTrackSeries> SeriesPtr = AllSeries[0];
-				
-				if (SeriesPtr.IsValid())
+				TSharedPtr<STimingView> TimingView = Window->GetTimingView();
+				if (TimingView)
 				{
 					// Highlight the area corresponding to viewport of Timing View.
 					const double StartTime = TimingView->GetViewport().GetStartTime();
 					const double EndTime = TimingView->GetViewport().GetEndTime();
-					Helper.DrawHighlightedInterval(*SeriesPtr, StartTime, EndTime);
+					Helper.DrawHighlightedInterval(*GameFrameSeries, StartTime, EndTime);
 				}
 			}
 		}
@@ -1003,8 +1002,8 @@ FReply SFrameTrack::OnMouseButtonUp(const FGeometry& MyGeometry, const FPointerE
 				const bool JoinCurrentSelection = MouseEvent.IsShiftDown();
 
 				SelectFrameAtMousePosition(
-					static_cast<float>(MousePositionOnButtonUp.X), 
-					static_cast<float>(MousePositionOnButtonUp.Y), 
+					static_cast<float>(MousePositionOnButtonUp.X),
+					static_cast<float>(MousePositionOnButtonUp.Y),
 					JoinCurrentSelection);
 			}
 
@@ -1329,7 +1328,10 @@ void SFrameTrack::ShowContextMenu(const FPointerEvent& MouseEvent)
 void SFrameTrack::ContextMenu_ShowGameFrames_Execute()
 {
 	TSharedPtr<FFrameTrackSeries> Series = FindSeries(ETraceFrameType::TraceFrameType_Game);
-	Series->bIsVisible = !Series->bIsVisible;
+	if (Series.IsValid())
+	{
+		Series->bIsVisible = !Series->bIsVisible;
+	}
 }
 
 ////////////////////////////////////////////////////////////////////////////////////////////////////
@@ -1352,7 +1354,10 @@ bool SFrameTrack::ContextMenu_ShowGameFrames_IsChecked()
 void SFrameTrack::ContextMenu_ShowRenderingFrames_Execute()
 {
 	TSharedPtr<FFrameTrackSeries> Series = FindSeries(ETraceFrameType::TraceFrameType_Rendering);
-	Series->bIsVisible = !Series->bIsVisible;
+	if (Series.IsValid())
+	{
+		Series->bIsVisible = !Series->bIsVisible;
+	}
 }
 
 ////////////////////////////////////////////////////////////////////////////////////////////////////
@@ -1375,13 +1380,10 @@ bool SFrameTrack::ContextMenu_ShowRenderingFrames_IsChecked()
 void SFrameTrack::ContextMenu_ShowFrameStats_Execute(ETraceFrameType FrameType, uint32 TimerId)
 {
 	TSharedPtr<FFrameTrackSeries> Series = FindFrameStatsSeries(FrameType, TimerId);
-
-	if (!Series.IsValid())
+	if (Series.IsValid())
 	{
-		return;
+		Series->bIsVisible = !Series->bIsVisible;
 	}
-
-	Series->bIsVisible = !Series->bIsVisible;
 }
 
 ////////////////////////////////////////////////////////////////////////////////////////////////////
@@ -1396,7 +1398,6 @@ bool SFrameTrack::ContextMenu_ShowFrameStats_CanExecute(ETraceFrameType FrameTyp
 bool SFrameTrack::ContextMenu_ShowFrameStats_IsChecked(ETraceFrameType FrameType, uint32 TimerId)
 {
 	TSharedPtr<FFrameTrackSeries> Series = FindFrameStatsSeries(FrameType, TimerId);
-
 	return Series.IsValid() ? Series->bIsVisible : false;
 }
 
