[Insights] Timers View persistent settings

diff --git a/Engine/Programs/UnrealInsights/Config/DefaultUnrealInsightsSettings.ini b/Engine/Programs/UnrealInsights/Config/DefaultUnrealInsightsSettings.ini
--- a/Engine/Programs/UnrealInsights/Config/DefaultUnrealInsightsSettings.ini
+++ b/Engine/Programs/UnrealInsights/Config/DefaultUnrealInsightsSettings.ini
@@ -22,6 +22,32 @@ AutoScrollMinDelay=0.3
 ; If enabled, the panning is allowed to continue when the mouse cursor reaches the edges of the screen.
 bAllowPanningOnScreenEdges=False
 
+[Insights.TimingProfiler.TimersView]
+; Names of visible columns in the Timers Panel in the Game mode.
++GameFrameColumns=MaxInclTime
++GameFrameColumns=AverageInclTime
++GameFrameColumns=MedianInclTime
++GameFrameColumns=MinInclTime
+; Names of visible columns in the Timers Panel in the Rendering mode.
++RenderingFrameColumns=MaxInclTime
++RenderingFrameColumns=AverageInclTime
++RenderingFrameColumns=MedianInclTime
++RenderingFrameColumns=MinInclTime
+; Names of visible columns in the Timers Panel in the Instance mode.
++InstanceColumns=Count
++InstanceColumns=TotalInclTime
++InstanceColumns=TotalExclTime
+; Mode for the timers panel. 0 : Game Frame; 1 : Rendering Frame; 2 : Instance
+Mode=2
+; Grouping Mode for the timers panel. 0 : Flat; 1 : ByName; 3 : ByType; 4 : ByInstanceCount
+GroupingMode=3
+; Controls the CPU filter button from Timers View.
+ShowCpuTimers=1
+; Controls the GPU filter button from Timers View.
+ShowGpuTimers=1
+; Controls the "Filter Out Zero Count Timers" button from Timers View.
+ShowZeroCountTimers=1
+
 [Insights.MemoryProfiler]
 ; Add symbol search paths by adding to this array. Order is important.
 ;+SymbolSearchPaths=c:\symbols
diff --git a/Engine/Source/Developer/TraceInsights/Private/Insights/InsightsSettings.h b/Engine/Source/Developer/TraceInsights/Private/Insights/InsightsSettings.h
--- a/Engine/Source/Developer/TraceInsights/Private/Insights/InsightsSettings.h
+++ b/Engine/Source/Developer/TraceInsights/Private/Insights/InsightsSettings.h
@@ -23,11 +23,32 @@ public:
 		, AutoScrollFrameAlignment((int32)TraceFrameType_Game) // -1 = none, 0 = game, 1 = rendering
 		, AutoScrollViewportOffsetPercent(0.1) // scrolls forward 10% of viewport's width
 		, AutoScrollMinDelay(0.3) // [seconds]
+		, TimersViewMode((int32)TraceFrameType_Count)
+		, TimersViewGroupingMode(3) // ByType
+		, bTimersViewShowCpuTimers(true)
+		, bTimersViewShowGpuTimers(true)
+		, bTimersViewShowZeroCountTimers(true)
 	{
 		if (!bIsDefault)
 		{
 			LoadFromConfig();
 		}
+		else
+		{
+			TimersViewInstanceVisibleColumns.Add(TEXT("Count"));
+			TimersViewInstanceVisibleColumns.Add(TEXT("TotalInclTime"));
+			TimersViewInstanceVisibleColumns.Add(TEXT("TotalExclTime"));
+
+			TimersViewGameFrameVisibleColumns.Add(TEXT("MaxInclTime"));
+			TimersViewGameFrameVisibleColumns.Add(TEXT("AverageInclTime"));
+			TimersViewGameFrameVisibleColumns.Add(TEXT("MedianInclTime"));
+			TimersViewGameFrameVisibleColumns.Add(TEXT("MinInclTime"));
+
+			TimersViewRenderingFrameVisibleColumns.Add(TEXT("MaxInclTime"));
+			TimersViewRenderingFrameVisibleColumns.Add(TEXT("AverageInclTime"));
+			TimersViewRenderingFrameVisibleColumns.Add(TEXT("MedianInclTime"));
+			TimersViewRenderingFrameVisibleColumns.Add(TEXT("MinInclTime"));
+		}
 	}
 
 	~FInsightsSettings()
@@ -70,6 +91,17 @@ public:
 		GConfig->GetDouble(TEXT("Insights.TimingProfiler"), TEXT("AutoScrollMinDelay"), AutoScrollMinDelay, SettingsIni);
 
 		GConfig->GetArray(TEXT("Insights.MemoryProfiler"), TEXT("SymbolSearchPaths"), SymbolSearchPaths, SettingsIni);
+		
+		GConfig->GetArray(TEXT("Insights.TimingProfiler.TimersView"), TEXT("InstanceColumns"), TimersViewInstanceVisibleColumns, SettingsIni);
+		GConfig->GetArray(TEXT("Insights.TimingProfiler.TimersView"), TEXT("GameFrameColumns"), TimersViewGameFrameVisibleColumns, SettingsIni);
+		GConfig->GetArray(TEXT("Insights.TimingProfiler.TimersView"), TEXT("RenderingFrameColumns"), TimersViewRenderingFrameVisibleColumns, SettingsIni);
+
+		GConfig->GetInt(TEXT("Insights.TimingProfiler.TimersView"), TEXT("Mode"), TimersViewMode, SettingsIni);
+
+		GConfig->GetInt(TEXT("Insights.TimingProfiler.TimersView"), TEXT("GroupingMode"), TimersViewGroupingMode, SettingsIni);
+		GConfig->GetBool(TEXT("Insights.TimingProfiler.TimersView"), TEXT("ShowCpuTimers"), bTimersViewShowCpuTimers, SettingsIni);
+		GConfig->GetBool(TEXT("Insights.TimingProfiler.TimersView"), TEXT("ShowGpuTimers"), bTimersViewShowGpuTimers, SettingsIni);
+		GConfig->GetBool(TEXT("Insights.TimingProfiler.TimersView"), TEXT("ShowZeroCountTimers"), bTimersViewShowZeroCountTimers, SettingsIni);
 	}
 
 	void SaveToConfig()
@@ -87,6 +119,17 @@ public:
 
 		GConfig->SetArray(TEXT("Insights.MemoryProfiler"), TEXT("SymbolSearchPaths"), SymbolSearchPaths, SettingsIni);
 
+		GConfig->SetArray(TEXT("Insights.TimingProfiler.TimersView"), TEXT("InstanceColumns"), TimersViewInstanceVisibleColumns, SettingsIni);
+		GConfig->SetArray(TEXT("Insights.TimingProfiler.TimersView"), TEXT("GameFrameColumns"), TimersViewGameFrameVisibleColumns, SettingsIni);
+		GConfig->SetArray(TEXT("Insights.TimingProfiler.TimersView"), TEXT("RenderingFrameColumns"), TimersViewRenderingFrameVisibleColumns, SettingsIni);
+
+		GConfig->SetInt(TEXT("Insights.TimingProfiler.TimersView"), TEXT("Mode"), TimersViewMode, SettingsIni);
+
+		GConfig->SetInt(TEXT("Insights.TimingProfiler.TimersView"), TEXT("GroupingMode"), TimersViewGroupingMode, SettingsIni);
+		GConfig->SetInt(TEXT("Insights.TimingProfiler.TimersView"), TEXT("ShowCpuTimers"), bTimersViewShowCpuTimers, SettingsIni);
+		GConfig->SetInt(TEXT("Insights.TimingProfiler.TimersView"), TEXT("ShowGpuTimers"), bTimersViewShowGpuTimers, SettingsIni);
+		GConfig->SetInt(TEXT("Insights.TimingProfiler.TimersView"), TEXT("ShowZeroCountTimers"), bTimersViewShowZeroCountTimers, SettingsIni);
+
 		GConfig->Flush(false, SettingsIni);
 	}
 
@@ -155,6 +198,38 @@ public:
 	void SetSymbolSearchPaths(const TArray<FString>& SearchPaths) { SymbolSearchPaths = SearchPaths; }
 	void SetAndSaveSymbolSearchPaths(const TArray<FString>& SearchPaths) { SET_AND_SAVE(SymbolSearchPaths, SearchPaths); }
 
+	const TArray<FString>& GetTimersViewInstanceVisibleColumns() const { return TimersViewInstanceVisibleColumns; }
+	void SetTimersViewInstanceVisibleColumns(const TArray<FString>& Columns) { TimersViewInstanceVisibleColumns = Columns; }
+	void SetAndSaveTimersViewInstanceVisibleColumns(const TArray<FString>& Columns) { SET_AND_SAVE(TimersViewInstanceVisibleColumns, Columns); }
+
+	const TArray<FString>& GetTimersViewGameFrameVisibleColumns() const { return TimersViewGameFrameVisibleColumns; }
+	void SetTimersViewGameFrameVisibleColumns(const TArray<FString>& Columns) { TimersViewGameFrameVisibleColumns = Columns; }
+	void SetAndSaveTimersViewGameFrameVisibleColumns(const TArray<FString>& Columns) { SET_AND_SAVE(TimersViewGameFrameVisibleColumns, Columns); }
+
+	const TArray<FString>& GetTimersViewRenderingFrameVisibleColumns() const { return TimersViewRenderingFrameVisibleColumns; }
+	void SetTimersViewRenderingFrameVisibleColumns(const TArray<FString>& Columns) { TimersViewRenderingFrameVisibleColumns = Columns; }
+	void SetAndSaveTimersViewRenderingFrameVisibleColumns(const TArray<FString>& Columns) { SET_AND_SAVE(TimersViewRenderingFrameVisibleColumns, Columns); }
+
+	int32 GetTimersViewMode() const { return TimersViewMode; }
+	void SetTimersViewMode(int32 InMode) { TimersViewMode = InMode; }
+	void SetAndSaveTimersViewMode(int32 InMode) { SET_AND_SAVE(TimersViewMode, InMode); }
+
+	int32 GetTimersViewGroupingMode() const { return TimersViewGroupingMode; }
+	void SetTimersViewGroupingMode(int32 InValue) { TimersViewGroupingMode = InValue; }
+	void SetAndSaveTimersViewGroupingMode(int32 InValue) { SET_AND_SAVE(TimersViewGroupingMode, InValue); }
+
+	bool GetTimersViewShowCpuEvents() const { return bTimersViewShowCpuTimers; }
+	void SetTimersViewShowCpuEvents(bool InValue) { bTimersViewShowCpuTimers = InValue; }
+	void SetAndSaveTimersViewShowCpuEvents(bool InValue) { SET_AND_SAVE(bTimersViewShowCpuTimers, InValue); }
+
+	bool GetTimersViewShowGpuEvents() const { return bTimersViewShowGpuTimers; }
+	void SetTimersViewShowGpuEvents(bool InValue) { bTimersViewShowGpuTimers = InValue; }
+	void SetAndSaveTimersViewShowGpuEvents(bool InValue) { SET_AND_SAVE(bTimersViewShowGpuTimers, InValue); }
+
+	bool GetTimersViewShowZeroCountTimers() const { return bTimersViewShowZeroCountTimers; }
+	void SetTimersViewShowZeroCountTimers(bool InValue) { bTimersViewShowZeroCountTimers = InValue; }
+	void SetAndSaveTimersViewShowZeroCountTimers(bool InValue) { SET_AND_SAVE(bTimersViewShowZeroCountTimers, InValue); }
+
 	#undef SET_AND_SAVE
 
 private:
@@ -200,4 +275,28 @@ private:
 
 	/** Minimum time between two auto-scroll updates, in [seconds]. */
 	double AutoScrollMinDelay;
+
+	//** The list of visible columns in the Timers view in the Instance mode. */
+	TArray<FString> TimersViewInstanceVisibleColumns;
+
+	//** The list of visible columns in the Timers view in the Game Frame mode. */
+	TArray<FString> TimersViewGameFrameVisibleColumns;
+
+	//** The list of visible columns in the Timers view in the Rendering Frame mode. */
+	TArray<FString> TimersViewRenderingFrameVisibleColumns;
+
+	//** The mode for the timers panel. */
+	int32 TimersViewMode;
+
+	//** The grouping mode for the timers panel. */
+	int32 TimersViewGroupingMode;
+
+	//** If enabled, Cpu timers will be displayed in the Timing View. */
+	bool bTimersViewShowCpuTimers;
+
+	//** If enabled, Gpu timers will be displayed in the Timing View. */
+	bool bTimersViewShowGpuTimers;
+
+	//** If enabled, timers with no instances in the selected interval will still be displayed in the Timers View. */
+	bool bTimersViewShowZeroCountTimers;
 };
diff --git a/Engine/Source/Developer/TraceInsights/Private/Insights/Widgets/STimersView.cpp b/Engine/Source/Developer/TraceInsights/Private/Insights/Widgets/STimersView.cpp
--- a/Engine/Source/Developer/TraceInsights/Private/Insights/Widgets/STimersView.cpp
+++ b/Engine/Source/Developer/TraceInsights/Private/Insights/Widgets/STimersView.cpp
@@ -150,6 +150,9 @@ STimersView::STimersView()
 
 STimersView::~STimersView()
 {
+	SaveSettings();
+	SaveVisibleColumnsSettings();
+
 	// Remove ourselves from the Insights manager.
 	if (FInsightsManager::Get().IsValid())
 	{
@@ -414,6 +417,7 @@ void STimersView::Construct(const FArguments& InArgs)
 	];
 
 	InitializeAndShowHeaderColumns();
+	LoadSettings();
 
 	// Create the search filters: text based, type based etc.
 	TextFilter = MakeShared<FTimerNodeTextFilter>(FTimerNodeTextFilter::FItemToStringArray::CreateSP(this, &STimersView::HandleItemToStringArray));
@@ -1669,7 +1673,7 @@ void STimersView::CreateGroupByOptionsSources()
 
 void STimersView::GroupBy_OnSelectionChanged(TSharedPtr<ETimerGroupingMode> NewGroupingMode, ESelectInfo::Type SelectInfo)
 {
-	if (SelectInfo != ESelectInfo::Direct)
+	if (SelectInfo != ESelectInfo::Direct && *NewGroupingMode != GroupingMode)
 	{
 		GroupingMode = *NewGroupingMode;
 
@@ -1715,7 +1719,13 @@ void STimersView::CreateModeOptionsSources()
 	ModeOptionsSource.Add(MakeShared<ETraceFrameType>(ETraceFrameType::TraceFrameType_Game));
 	ModeOptionsSource.Add(MakeShared<ETraceFrameType>(ETraceFrameType::TraceFrameType_Rendering));
 
-	ModeComboBox->SetSelectedItem(ModeOptionsSource[0]);
+	for (TSharedPtr<ETraceFrameType> Option : ModeOptionsSource)
+	{
+		if (*Option == ModeFrameType)
+		{
+			ModeComboBox->SetSelectedItem(Option);
+		}
+	}
 
 	ModeComboBox->RefreshOptions();
 }
@@ -1729,7 +1739,9 @@ void STimersView::Mode_OnSelectionChanged(TSharedPtr<ETraceFrameType> NewFrameTy
 		return;
 	}
 
+	SaveVisibleColumnsSettings();
 	ModeFrameType = *NewFrameType;
+	LoadVisibleColumnsSettings();
 
 	TSharedPtr<STimingProfilerWindow> Window = FTimingProfilerManager::Get()->GetProfilerWindow();
 	if (Window.IsValid())
@@ -3163,4 +3175,111 @@ FReply STimersView::OnKeyDown(const FGeometry& MyGeometry, const FKeyEvent& InKe
 
 ////////////////////////////////////////////////////////////////////////////////////////////////////
 
+void STimersView::SaveVisibleColumnsSettings()
+{
+	FInsightsSettings& Settings = FInsightsManager::Get()->GetSettings();
+
+	TArray<FString> Columns;
+	for (const TSharedRef<Insights::FTableColumn>& ColumnRef : Table->GetColumns())
+	{
+		if (ColumnRef->IsVisible() && ColumnRef->CanBeHidden())
+		{
+			Columns.Add(ColumnRef->GetId().ToString());
+		}
+	}
+
+	switch (ModeFrameType)
+	{
+		case ETraceFrameType::TraceFrameType_Count:
+		{
+			Settings.SetTimersViewInstanceVisibleColumns(Columns);
+			break;
+		}
+		case ETraceFrameType::TraceFrameType_Game:
+		{
+			Settings.SetTimersViewGameFrameVisibleColumns(Columns);
+			break;
+		}
+		case ETraceFrameType::TraceFrameType_Rendering:
+		{
+			Settings.SetTimersViewRenderingFrameVisibleColumns(Columns);
+			break;
+		}
+	}
+
+	Settings.SaveToConfig();
+}
+
+////////////////////////////////////////////////////////////////////////////////////////////////////
+
+void STimersView::SaveSettings()
+{
+	FInsightsSettings& Settings = FInsightsManager::Get()->GetSettings();
+
+	Settings.SetTimersViewMode(static_cast<int32>(ModeFrameType));
+	Settings.SetTimersViewGroupingMode(static_cast<int32>(GroupingMode));
+	Settings.SetTimersViewShowCpuEvents(static_cast<int32>(bTimerTypeIsVisible[static_cast<int32>(ETimerNodeType::CpuScope)]));
+	Settings.SetTimersViewShowGpuEvents(static_cast<int32>(bTimerTypeIsVisible[static_cast<int32>(ETimerNodeType::GpuScope)]));
+	Settings.SetTimersViewShowZeroCountTimers(!bFilterOutZeroCountTimers);
+
+	Settings.SaveToConfig();
+}
+
+////////////////////////////////////////////////////////////////////////////////////////////////////
+
+void STimersView::LoadVisibleColumnsSettings()
+{
+	FInsightsSettings& Settings = FInsightsManager::Get()->GetSettings();
+	TArray<FString> Columns;
+
+	switch (ModeFrameType)
+	{
+		case ETraceFrameType::TraceFrameType_Count:
+		{
+			Columns = Settings.GetTimersViewInstanceVisibleColumns();
+			break;
+		}
+		case ETraceFrameType::TraceFrameType_Game:
+		{
+			Columns = Settings.GetTimersViewGameFrameVisibleColumns();
+			break;
+		}
+		case ETraceFrameType::TraceFrameType_Rendering:
+		{
+			Columns = Settings.GetTimersViewRenderingFrameVisibleColumns();
+			break;
+		}
+	}
+
+	for (const TSharedRef<Insights::FTableColumn>& ColumnRef : Table->GetColumns())
+	{
+		if (ColumnRef->CanBeHidden())
+		{
+			HideColumn(ColumnRef->GetId());
+		}
+	}
+
+	for (const FString& Column : Columns)
+	{
+		ShowColumn(FName(Column));
+	}
+}
+
+////////////////////////////////////////////////////////////////////////////////////////////////////
+
+void STimersView::LoadSettings()
+{
+	FInsightsSettings& Settings = FInsightsManager::Get()->GetSettings();
+
+	ModeFrameType = static_cast<ETraceFrameType>(Settings.GetTimersViewMode());
+	GroupingMode = static_cast<ETimerGroupingMode>(Settings.GetTimersViewGroupingMode());
+	bTimerTypeIsVisible[static_cast<int32>(ETimerNodeType::CpuScope)] = Settings.GetTimersViewShowCpuEvents();
+	bTimerTypeIsVisible[static_cast<int32>(ETimerNodeType::GpuScope)] = Settings.GetTimersViewShowGpuEvents();
+	bFilterOutZeroCountTimers = !Settings.GetTimersViewShowZeroCountTimers();
+
+	LoadVisibleColumnsSettings();
+}
+
+////////////////////////////////////////////////////////////////////////////////////////////////////
+
 #undef LOCTEXT_NAMESPACE
diff --git a/Engine/Source/Developer/TraceInsights/Private/Insights/Widgets/STimersView.h b/Engine/Source/Developer/TraceInsights/Private/Insights/Widgets/STimersView.h
--- a/Engine/Source/Developer/TraceInsights/Private/Insights/Widgets/STimersView.h
+++ b/Engine/Source/Developer/TraceInsights/Private/Insights/Widgets/STimersView.h
@@ -329,6 +329,11 @@ private:
 
 	void OpenSourceFileInIDE(FTimerNodePtr InNode) const;
 
+	void SaveSettings();
+	void SaveVisibleColumnsSettings();
+	void LoadSettings();
+	void LoadVisibleColumnsSettings();
+
 private:
 	/** Table view model. */
 	TSharedPtr<Insights::FTable> Table;
