[Insights] Fixed UnsafeTypeCast warnings.

diff --git a/Engine/Plugins/TraceUtilities/Source/EditorTraceUtilities/Private/InsightsStatusBar/SInsightsStatusBar.cpp b/Engine/Plugins/TraceUtilities/Source/EditorTraceUtilities/Private/InsightsStatusBar/SInsightsStatusBar.cpp
--- a/Engine/Plugins/TraceUtilities/Source/EditorTraceUtilities/Private/InsightsStatusBar/SInsightsStatusBar.cpp
+++ b/Engine/Plugins/TraceUtilities/Source/EditorTraceUtilities/Private/InsightsStatusBar/SInsightsStatusBar.cpp
@@ -1031,7 +1031,7 @@ void SInsightsStatusBarWidget::OpenTrace(int32 Index)
 				
 				if (TraceId)
 				{
-					FUnrealInsightsLauncher::Get()->OpenRemoteTrace(TEXT("localhost"), LiveSessionTracker->GetStorePort(), *TraceId);
+					FUnrealInsightsLauncher::Get()->OpenRemoteTrace(TEXT("localhost"), uint16(LiveSessionTracker->GetStorePort()), *TraceId);
 					return;
 				}
 			}
diff --git a/Engine/Source/Developer/TraceAnalysis/Private/Store/CborPayload.h b/Engine/Source/Developer/TraceAnalysis/Private/Store/CborPayload.h
--- a/Engine/Source/Developer/TraceAnalysis/Private/Store/CborPayload.h
+++ b/Engine/Source/Developer/TraceAnalysis/Private/Store/CborPayload.h
@@ -184,7 +184,7 @@ inline bool FResponse::GetStringArray(const char* Key, TArray<FString>& OutArray
 	{
 		if (Context.IsString())
 		{
-			uint32 Length = Context.AsLength();
+			uint32 Length = uint32(Context.AsLength());
 			uint32 Offset = uint32(MemoryReader.Tell());
 			auto* String = (const char*)(Buffer.GetData() + Offset - Length);
 			if (FCStringAnsi::Strncmp(Key, String, Length) == 0)
diff --git a/Engine/Source/Developer/TraceAnalysis/Private/Store/StoreClient.cpp b/Engine/Source/Developer/TraceAnalysis/Private/Store/StoreClient.cpp
--- a/Engine/Source/Developer/TraceAnalysis/Private/Store/StoreClient.cpp
+++ b/Engine/Source/Developer/TraceAnalysis/Private/Store/StoreClient.cpp
@@ -530,7 +530,7 @@ FStoreClient* FStoreClient::Connect(const TCHAR* Host, uint32 Port)
 bool FStoreClient::Reconnect(const TCHAR* Host, uint32 Port)
 {
 	auto* Self = (FStoreCborClient*)this;
-	return Self->Connect(Host, Port);
+	return Self->Connect(Host, static_cast<uint16>(Port));
 }
 
 ////////////////////////////////////////////////////////////////////////////////
diff --git a/Engine/Source/Developer/TraceInsights/Private/Insights/ViewModels/TimingGraphTrack.cpp b/Engine/Source/Developer/TraceInsights/Private/Insights/ViewModels/TimingGraphTrack.cpp
--- a/Engine/Source/Developer/TraceInsights/Private/Insights/ViewModels/TimingGraphTrack.cpp
+++ b/Engine/Source/Developer/TraceInsights/Private/Insights/ViewModels/TimingGraphTrack.cpp
@@ -474,13 +474,13 @@ void FTimingGraphTrack::UpdateFrameStatsTimerSeries(FTimingGraphSeries& Series,
 			const TraceServices::IFrameProvider& FramesProvider = ReadFrameProvider(*Session.Get());
 
 			Series.FrameStatsCachedEvents.Empty();
-			uint32 FrameCount = FramesProvider.GetFrameCount(ETraceFrameType::TraceFrameType_Game);
+			uint64 FrameCount = FramesProvider.GetFrameCount(ETraceFrameType::TraceFrameType_Game);
 			if (FrameCount == 0)
 			{
 				return;
 			}
 
-			FramesProvider.EnumerateFrames(Series.FrameType, (uint64) 0, (uint64) FrameCount, [&Series](const TraceServices::FFrame& Frame)
+			FramesProvider.EnumerateFrames(Series.FrameType, 0ull, FrameCount, [&Series](const TraceServices::FFrame& Frame)
 				{
 					Insights::FFrameStatsCachedEvent Event;
 					Event.FrameStartTime = Frame.StartTime;
diff --git a/Engine/Source/Developer/TraceInsights/Private/Insights/Widgets/SFrameTrack.cpp b/Engine/Source/Developer/TraceInsights/Private/Insights/Widgets/SFrameTrack.cpp
--- a/Engine/Source/Developer/TraceInsights/Private/Insights/Widgets/SFrameTrack.cpp
+++ b/Engine/Source/Developer/TraceInsights/Private/Insights/Widgets/SFrameTrack.cpp
@@ -335,7 +335,7 @@ void SFrameTrack::UpdateState()
 
 			Insights::FFrameStatsHelper::ComputeFrameStatsForTimer(Frames, TimerSeries->TimerId);
 
-			int32 CurrentIndex = StartIndex;
+			uint64 CurrentIndex = StartIndex;
 			for (Insights::FFrameStatsCachedEvent& Event : Frames)
 			{
 				TraceServices::FFrame NewFrame;
diff --git a/Engine/Source/Developer/TraceInsights/Private/Insights/Widgets/SStartPageWindow.cpp b/Engine/Source/Developer/TraceInsights/Private/Insights/Widgets/SStartPageWindow.cpp
--- a/Engine/Source/Developer/TraceInsights/Private/Insights/Widgets/SStartPageWindow.cpp
+++ b/Engine/Source/Developer/TraceInsights/Private/Insights/Widgets/SStartPageWindow.cpp
@@ -95,15 +95,14 @@ public:
 			.AutoWidth()
 			[
 				SNew(SBox)
-				.WidthOverride(16.0)
-				.HeightOverride(16.0)
-				.Padding(FMargin(2.0))
+				.WidthOverride(16.0f)
+				.HeightOverride(16.0f)
+				.Padding(FMargin(2.0f))
 				[
 					SNew(SImage)
 					.Image(FInsightsStyle::Get().GetBrush("Icons.UTrace"))
 					.ColorAndOpacity(this, &STraceListRow::GetColorForPath)
 				]
-
 			]
 
 			+ SHorizontalBox::Slot()
@@ -1357,18 +1356,18 @@ void STraceDirectoryItem::Construct(const FArguments& InArgs, TSharedPtr<STraceS
 	ChildSlot
 	[
 		SNew(SBorder)
-		.Padding(2.0f , 1.0f)
+		.Padding(2.0f, 1.0f)
 		[
 			SNew(SHorizontalBox)
 
 			+ SHorizontalBox::Slot()
-			.Padding(4.0, 2.0)
+			.Padding(4.0f, 2.0f)
 			.AutoWidth()
 			.VAlign(EVerticalAlignment::VAlign_Center)
 			[
 				SNew(SBox)
-				.HeightOverride(16.0)
-				.WidthOverride(16.0)
+				.HeightOverride(16.0f)
+				.WidthOverride(16.0f)
 				[
 					SNew(SImage)
 					.Image(FInsightsStyle::Get().GetBrush("Icons.UTraceStack"))
@@ -1377,8 +1376,8 @@ void STraceDirectoryItem::Construct(const FArguments& InArgs, TSharedPtr<STraceS
 			]
 
 			+ SHorizontalBox::Slot()
-			.FillWidth(1.0)
-			.Padding(4.0, 2.0)
+			.FillWidth(1.0f)
+			.Padding(4.0f, 2.0f)
 			.VAlign(EVerticalAlignment::VAlign_Center)
 			[
 				SNew(STextBlock)
@@ -1386,7 +1385,7 @@ void STraceDirectoryItem::Construct(const FArguments& InArgs, TSharedPtr<STraceS
 			]
 
 			+ SHorizontalBox::Slot()
-			.Padding(4.0, 2.0)
+			.Padding(4.0f, 2.0f)
 			.AutoWidth()
 			[
 				ConstructOperations()
@@ -1662,7 +1661,6 @@ TSharedRef<SWidget> STraceStoreWindow::ConstructTraceStoreDirectoryPanel()
 			]
 
 #if STORE_SUPPORTS_SETTINGS
-
 			+ SHorizontalBox::Slot()
 			.AutoWidth()
 			.Padding(4.0f, 0.0f, 0.0f, 0.0f)
@@ -1724,7 +1722,7 @@ TSharedRef<SWidget> STraceStoreWindow::ConstructTraceStoreDirectoryPanel()
 		+ SVerticalBox::Slot()
 		.AutoHeight()
 		.MaxHeight(400)
-		.Padding(.0f, 8.0f)
+		.Padding(0.0f, 8.0f)
 		.HAlign(HAlign_Left)
 		[
 			SAssignNew(StoreSettingsArea, SScrollBox)
@@ -1796,7 +1794,7 @@ TSharedRef<SWidget> STraceStoreWindow::ConstructTraceStoreDirectoryPanel()
 						]
 
 						+ SHorizontalBox::Slot()
-						.Padding(4.0, 0.0, 0.0, 0.0)
+						.Padding(4.0f, 0.0f, 0.0f, 0.0f)
 						.AutoWidth()
 						[
 							SNew(STextBlock)
diff --git a/Engine/Source/Developer/TraceServices/Private/Analyzers/NetTraceAnalyzer.cpp b/Engine/Source/Developer/TraceServices/Private/Analyzers/NetTraceAnalyzer.cpp
--- a/Engine/Source/Developer/TraceServices/Private/Analyzers/NetTraceAnalyzer.cpp
+++ b/Engine/Source/Developer/TraceServices/Private/Analyzers/NetTraceAnalyzer.cpp
@@ -285,7 +285,7 @@ void FNetTraceAnalyzer::HandlePacketContentEvent(const FOnEventContext& Context,
 
 						// Fill in the object data we currently have
 						ObjectInstance.LifeTime.Begin = GetLastTimestamp();
-						ObjectInstance.NameIndex = PendingNameIndex;
+						ObjectInstance.NameIndex = static_cast<uint16>(PendingNameIndex);
 						ObjectInstance.NetObjectId = DecodedNameOrObjectId;
 						ObjectInstance.TypeId = 0;
 
@@ -808,7 +808,7 @@ void FNetTraceAnalyzer::HandleObjectCreatedEvent(const FOnEventContext& Context,
 				ExistingInstance->TypeId = TypeId;
 
 				// Update name in both the persistent instance and the active one
-				ExistingInstance->NameIndex = NameIndex;
+				ExistingInstance->NameIndex = static_cast<uint16>(NameIndex);
 				ActiveObjectInstance->NameIndex = NameIndex;
 
 				return;
diff --git a/Engine/Source/Developer/TraceServices/Private/Model/Regions.cpp b/Engine/Source/Developer/TraceServices/Private/Model/Regions.cpp
--- a/Engine/Source/Developer/TraceServices/Private/Model/Regions.cpp
+++ b/Engine/Source/Developer/TraceServices/Private/Model/Regions.cpp
@@ -156,11 +156,11 @@ void FRegionProvider::OnAnalysisSessionEnded()
 
 int32 FRegionProvider::CalculateRegionDepth(const FTimeRegion& Region) const
 {
-	TArray<const FTimeRegion*> Overlaps;
-
 	constexpr int32 DepthLimit = 100;
+
 	int32 NewDepth = 0;
-	// find first free lane/depth
+
+	// Find first free lane/depth
 	while (NewDepth < DepthLimit)
 	{
 		if (!Lanes.IsValidIndex(NewDepth))
diff --git a/Engine/Source/Developer/TraceServices/Private/Model/TasksProfiler.cpp b/Engine/Source/Developer/TraceServices/Private/Model/TasksProfiler.cpp
--- a/Engine/Source/Developer/TraceServices/Private/Model/TasksProfiler.cpp
+++ b/Engine/Source/Developer/TraceServices/Private/Model/TasksProfiler.cpp
@@ -477,11 +477,11 @@ namespace TraceServices
 
 		if (TaskIndex >= Tasks.Num())
 		{
-			int32 PrevNum = Tasks.Num();
+			int64 PrevNum = Tasks.Num();
 			Tasks.AddDefaulted(TaskIndex - Tasks.Num() + 1);
-			for (int32 i = PrevNum; i != Tasks.Num(); ++i)
+			for (int64 Index = PrevNum; Index != Tasks.Num(); ++Index)
 			{
-				Tasks[i].Id = i + FirstTaskId;
+				Tasks[Index].Id = Index + FirstTaskId;
 			}
 		}
 
diff --git a/Engine/Source/Developer/TraceServices/Public/Common/PagedArray.h b/Engine/Source/Developer/TraceServices/Public/Common/PagedArray.h
--- a/Engine/Source/Developer/TraceServices/Public/Common/PagedArray.h
+++ b/Engine/Source/Developer/TraceServices/Public/Common/PagedArray.h
@@ -807,7 +807,7 @@ FInt32Interval GetElementRangeOverlappingGivenRange(const TPagedArray<ItemType,
 	{
 		// if we went past the end with StartPageIndex we still have to search the last page
 		const PageType& Page = PageData[StartPageIndex];
-		TArrayView<ItemType> PageValues(Page.Items, Page.Count);
+		TArrayView<ItemType, int32> PageValues(Page.Items, static_cast<int32>(Page.Count));
 		const int32 Index = Algo::UpperBoundBy(PageValues, StartTime, ItemEndProjection);
 		check(Index >= 0);
 		// Index may point past the end of the page, but the first item in the next page could be valid
@@ -835,7 +835,7 @@ FInt32Interval GetElementRangeOverlappingGivenRange(const TPagedArray<ItemType,
 	
 	{
 		const PageType& Page = PageData[EndPageIndex];
-		TArrayView<ItemType> PageValues(Page.Items, Page.Count);
+		TArrayView<ItemType, int32> PageValues(Page.Items, static_cast<int32>(Page.Count));
 		int32 Index = Algo::UpperBoundBy(PageValues, EndTime, ItemStartProjection);
 		check(Index >= 0);
 		// Index may point past the end of the page, that means include last element in result
diff --git a/Engine/Source/Developer/TraceServices/Public/TraceServices/Model/Regions.h b/Engine/Source/Developer/TraceServices/Public/TraceServices/Model/Regions.h
--- a/Engine/Source/Developer/TraceServices/Public/TraceServices/Model/Regions.h
+++ b/Engine/Source/Developer/TraceServices/Public/TraceServices/Model/Regions.h
@@ -16,7 +16,7 @@ struct TRACESERVICES_API FTimeRegion
 	double BeginTime = std::numeric_limits<double>::infinity();
 	double EndTime = std::numeric_limits<double>::infinity();
 	const TCHAR* Text = nullptr;
-	int8_t Depth = -1;
+	int32 Depth = -1;
 };
 
 class TRACESERVICES_API FRegionLane
@@ -26,7 +26,7 @@ class TRACESERVICES_API FRegionLane
 public:
 	FRegionLane(ILinearAllocator& InAllocator) : Regions(InAllocator, 512) {}
 
-	int32 Num() const { return Regions.Num(); }
+	int32 Num() const { return static_cast<int32>(Regions.Num()); }
 
 	/**
 	 * Call Callback for every region overlapping the interval defined by IntervalStart and IntervalEnd
