[Insights] Fixed opening a trace from a remote trace store. The new UnrealInsights process is launched now specifying also the store address in addition to store port. Also, the trace id can now be specified as hex number (prefixed by 0x).

diff --git a/Engine/Source/Developer/TraceInsights/Private/Insights/Widgets/SStartPageWindow.cpp b/Engine/Source/Developer/TraceInsights/Private/Insights/Widgets/SStartPageWindow.cpp
--- a/Engine/Source/Developer/TraceInsights/Private/Insights/Widgets/SStartPageWindow.cpp
+++ b/Engine/Source/Developer/TraceInsights/Private/Insights/Widgets/SStartPageWindow.cpp
@@ -2662,12 +2662,20 @@ void STraceStoreWindow::OpenTraceSession(uint32 InTraceId)
 		return;
 	}
 
+	uint32 StoreAddress = 0;
 	uint32 StorePort = 0;
 	{
 		FScopeLock StoreClientLock(&FInsightsManager::Get()->GetStoreClientCriticalSection());
+		StoreAddress = StoreClient->GetStoreAddress();
 		StorePort = StoreClient->GetStorePort();
 	}
-	FString CmdLine = FString::Printf(TEXT("-OpenTraceId=%d -StorePort=%d"), InTraceId, StorePort);
+	FString CmdLine = FString::Printf(TEXT("-OpenTraceId=0x%X -Store=%u.%u.%u.%u:%u"),
+		InTraceId,
+		(StoreAddress >> 24) & 0xFF,
+		(StoreAddress >> 16) & 0xFF,
+		(StoreAddress >>  8) & 0xFF,
+		(StoreAddress      ) & 0xFF,
+		StorePort);
 
 	FString ExtraCmdParams;
 	GetExtraCommandLineParams(ExtraCmdParams);
diff --git a/Engine/Source/Programs/UnrealInsights/Private/UserInterfaceCommand.cpp b/Engine/Source/Programs/UnrealInsights/Private/UserInterfaceCommand.cpp
--- a/Engine/Source/Programs/UnrealInsights/Private/UserInterfaceCommand.cpp
+++ b/Engine/Source/Programs/UnrealInsights/Private/UserInterfaceCommand.cpp
@@ -261,7 +261,18 @@ void FUserInterfaceCommand::InitializeSlateApplication(bool bOpenTraceFile, cons
 	IUnrealInsightsModule& TraceInsightsModule = FModuleManager::LoadModuleChecked<IUnrealInsightsModule>("TraceInsights");
 
 	uint32 TraceId = 0;
-	bool bUseTraceId = FParse::Value(FCommandLine::Get(), TEXT("-OpenTraceId="), TraceId);
+	FString TraceIdString;
+	bool bUseTraceId = FParse::Value(FCommandLine::Get(), TEXT("-OpenTraceId="), TraceIdString);
+	if (TraceIdString.StartsWith(TEXT("0x")))
+	{
+		TCHAR* End;
+		TraceId = FCString::Strtoi(*TraceIdString + 2, &End, 16);
+	}
+	else
+	{
+		TCHAR* End;
+		TraceId = FCString::Strtoi(*TraceIdString, &End, 10);
+	}
 
 	const uint32 MaxPath = FPlatformMisc::GetMaxPathLength();
 	TCHAR* StoreHost = new TCHAR[MaxPath + 1];
