Lock store settings UI if not on local machine

Since store settings cannot be changed on remote machines (server verifies this), lock the UI components that corresponds to modifications.

diff --git a/Engine/Source/Developer/TraceInsights/Private/Insights/InsightsManager.cpp b/Engine/Source/Developer/TraceInsights/Private/Insights/InsightsManager.cpp
--- a/Engine/Source/Developer/TraceInsights/Private/Insights/InsightsManager.cpp
+++ b/Engine/Source/Developer/TraceInsights/Private/Insights/InsightsManager.cpp
@@ -461,6 +461,8 @@ bool FInsightsManager::ConnectToStore(const TCHAR* Host, uint32 Port)
 	
 	LastStoreHost = Host;
 	LastStorePort = Port;
+	bCanChangeStoreSettings = LastStoreHost.Equals(TEXT("localhost"), ESearchCase::IgnoreCase)
+                              		|| LastStoreHost.Equals(TEXT("127.0.0.1"));
 
 	return true;
 }
diff --git a/Engine/Source/Developer/TraceInsights/Private/Insights/InsightsManager.h b/Engine/Source/Developer/TraceInsights/Private/Insights/InsightsManager.h
--- a/Engine/Source/Developer/TraceInsights/Private/Insights/InsightsManager.h
+++ b/Engine/Source/Developer/TraceInsights/Private/Insights/InsightsManager.h
@@ -125,6 +125,8 @@ public:
 	bool ReconnectToStore() const;
 
 	const FString& GetLastStoreHost() const { return LastStoreHost; }
+
+	const bool CanChangeStoreSettings() const { return bCanChangeStoreSettings; }
 	
 	UE::Trace::FStoreClient* GetStoreClient() const { return StoreClient.Get(); }
 	FCriticalSection& GetStoreClientCriticalSection() const { return StoreClientCriticalSection; }
@@ -487,4 +489,5 @@ private:
 
 	FString LastStoreHost;
 	uint32 LastStorePort = 0;
+	bool bCanChangeStoreSettings = false;
 };
diff --git a/Engine/Source/Developer/TraceInsights/Private/Insights/Widgets/SStartPageWindow.cpp b/Engine/Source/Developer/TraceInsights/Private/Insights/Widgets/SStartPageWindow.cpp
--- a/Engine/Source/Developer/TraceInsights/Private/Insights/Widgets/SStartPageWindow.cpp
+++ b/Engine/Source/Developer/TraceInsights/Private/Insights/Widgets/SStartPageWindow.cpp
@@ -1437,6 +1437,7 @@ TSharedRef<SWidget> STraceDirectoryItem::ConstructOperations()
 				.ButtonStyle(&FAppStyle::Get().GetWidgetStyle<FButtonStyle>("SimpleButton"))
 				.ToolTipText(LOCTEXT("WatchDirsRemoveTooltip", "Remove watch directory (files will not be deleted)"))
 				.OnClicked_Raw(this, &STraceDirectoryItem::OnDelete)
+				.IsEnabled_Static(&STraceStoreWindow::CanChangeStoreSettings)
 				[
 					SNew(SImage)
 					.Image(FInsightsStyle::Get().GetBrush("Icons.RemoveWatchDir"))
@@ -1453,6 +1454,7 @@ TSharedRef<SWidget> STraceDirectoryItem::ConstructOperations()
 				.ButtonStyle(&FAppStyle::Get().GetWidgetStyle<FButtonStyle>("SimpleButton"))
 				.ToolTipText(LOCTEXT("ExploreTraceStoreDirButtonToolTip", "Explore the Trace Store Directory"))
 				.OnClicked_Raw(this, &STraceDirectoryItem::OnExplore)
+				.IsEnabled_Static(&STraceStoreWindow::CanChangeStoreSettings)
 				[
 					SNew(SImage)
 					.Image(FInsightsStyle::Get().GetBrush("Icons.FolderExplore"))
@@ -1554,7 +1556,7 @@ FReply STraceDirectoryItem::OnDelete()
 bool STraceDirectoryItem::CanModifyStore() const
 {
 	check(Window);
-	return !Window->HasAnyLiveTrace();
+	return !Window->HasAnyLiveTrace() && STraceStoreWindow::CanChangeStoreSettings();
 }
 
 ////////////////////////////////////////////////////////////////////////////////////////////////////
@@ -1710,6 +1712,7 @@ TSharedRef<SWidget> STraceStoreWindow::ConstructTraceStoreDirectoryPanel()
 				.ButtonStyle(&FAppStyle::Get().GetWidgetStyle<FButtonStyle>("SimpleButton"))
 				.ToolTipText(LOCTEXT("ExploreTraceStoreDirButtonToolTip", "Explore the Trace Store Directory"))
 				.OnClicked(this, &STraceStoreWindow::ExploreTraceStoreDirectory_OnClicked)
+				.IsEnabled_Static(&STraceStoreWindow::CanChangeStoreSettings)
 				[
 					SNew(SImage)
 					.Image(FInsightsStyle::Get().GetBrush("Icons.FolderExplore"))
@@ -1769,7 +1772,6 @@ TSharedRef<SWidget> STraceStoreWindow::ConstructTraceStoreDirectoryPanel()
 					SAssignNew(WatchDirsListView, SListView<TSharedPtr<FTraceDirectoryModel>>)
 					.ListItemsSource(&WatchDirectoriesModel)
 					.OnGenerateRow(this, &STraceStoreWindow::TraceDirs_OnGenerateRow)
-					//.Orientation(Orient_Horizontal)
 					.SelectionMode(ESelectionMode::None)
 				]
 
@@ -1779,9 +1781,10 @@ TSharedRef<SWidget> STraceStoreWindow::ConstructTraceStoreDirectoryPanel()
 				.HAlign(HAlign_Left)
 				[
 					SNew(SButton)
-				.ButtonStyle(&FAppStyle::Get().GetWidgetStyle<FButtonStyle>("Button"))
-				.ToolTipText(LOCTEXT("WatchDirsAddTooltip", "Add additional watch directory..."))
-				.OnClicked(this, &STraceStoreWindow::AddWatchDir_Clicked)
+					.ButtonStyle(&FAppStyle::Get().GetWidgetStyle<FButtonStyle>("Button"))
+					.ToolTipText(LOCTEXT("WatchDirsAddTooltip", "Add additional watch directory..."))
+					.IsEnabled_Static(&STraceStoreWindow::CanChangeStoreSettings)
+					.OnClicked(this, &STraceStoreWindow::AddWatchDir_Clicked)
 					[
 						SNew(SHorizontalBox)
 
@@ -1789,8 +1792,8 @@ TSharedRef<SWidget> STraceStoreWindow::ConstructTraceStoreDirectoryPanel()
 						.AutoWidth()
 						[
 							SNew(SImage)
-						.Image(FInsightsStyle::Get().GetBrush("Icons.AddWatchDir"))
-						.ColorAndOpacity(FSlateColor::UseForeground())
+							.Image(FInsightsStyle::Get().GetBrush("Icons.AddWatchDir"))
+							.ColorAndOpacity(FSlateColor::UseForeground())
 						]
 
 						+ SHorizontalBox::Slot()
@@ -3022,6 +3025,13 @@ FReply STraceStoreWindow::ExploreTraceStoreDirectory_OnClicked()
 
 ////////////////////////////////////////////////////////////////////////////////////////////////////
 
+bool STraceStoreWindow::CanChangeStoreSettings()
+{
+	return FInsightsManager::Get()->CanChangeStoreSettings();
+}
+
+////////////////////////////////////////////////////////////////////////////////////////////////////
+
 TSharedRef<ITableRow> STraceStoreWindow::TraceDirs_OnGenerateRow(TSharedPtr<FTraceDirectoryModel> Item, const TSharedRef<STableViewBase>& Owner) const
 {
 	return SNew(STableRow<TSharedPtr<FTraceDirectoryModel>>, Owner)
diff --git a/Engine/Source/Developer/TraceInsights/Private/Insights/Widgets/SStartPageWindow.h b/Engine/Source/Developer/TraceInsights/Private/Insights/Widgets/SStartPageWindow.h
--- a/Engine/Source/Developer/TraceInsights/Private/Insights/Widgets/SStartPageWindow.h
+++ b/Engine/Source/Developer/TraceInsights/Private/Insights/Widgets/SStartPageWindow.h
@@ -383,6 +383,7 @@ private:
 		ETraceDirOperations Operations;
 	};
 
+	static bool CanChangeStoreSettings();
 	TSharedRef<ITableRow> TraceDirs_OnGenerateRow(TSharedPtr<FTraceDirectoryModel> Item, const TSharedRef<STableViewBase>& Owner) const;
 	FReply StoreSettingsArea_Toggle() const;
 	const FSlateBrush* StoreSettingsToggle_Icon() const;
