removed StatsThread and AudioThread from TasksInsights tooltip construction as they were removed some time ago

diff --git a/Engine/Source/Developer/TraceInsights/Private/Insights/ViewModels/ThreadTimingTrack.cpp b/Engine/Source/Developer/TraceInsights/Private/Insights/ViewModels/ThreadTimingTrack.cpp
--- a/Engine/Source/Developer/TraceInsights/Private/Insights/ViewModels/ThreadTimingTrack.cpp
+++ b/Engine/Source/Developer/TraceInsights/Private/Insights/ViewModels/ThreadTimingTrack.cpp
@@ -1176,7 +1176,7 @@ void FThreadTimingTrack::InitTooltip(FTooltipDrawState& InOutTooltip, const ITim
 				}
 
 				ENamedThreads::Type ThreadInfo = (ENamedThreads::Type)Task.ThreadToExecuteOn;
-				const TCHAR* NamedThreadsStr[] = { TEXT("Stats"), TEXT("RHI"), TEXT("Audio"), TEXT("Game"), TEXT("Rendering") };
+				const TCHAR* NamedThreadsStr[] = { TEXT("RHI"), TEXT("Game"), TEXT("Rendering") };
 				ENamedThreads::Type ThreadIndex = ENamedThreads::GetThreadIndex(ThreadInfo);
 
 				if (ThreadIndex == ENamedThreads::AnyThread)
