Add copy entry to context menu in timing tracks to increase discoverability.

diff --git a/Engine/Source/Developer/TraceInsights/Private/Insights/Widgets/STimingView.cpp b/Engine/Source/Developer/TraceInsights/Private/Insights/Widgets/STimingView.cpp
--- a/Engine/Source/Developer/TraceInsights/Private/Insights/Widgets/STimingView.cpp
+++ b/Engine/Source/Developer/TraceInsights/Private/Insights/Widgets/STimingView.cpp
@@ -2925,6 +2925,18 @@ void STimingView::ShowContextMenu(const FPointerEvent& MouseEvent)
 				NAME_None,
 				EUserInterfaceActionType::Button
 			);
+			
+			MenuBuilder.AddMenuEntry(
+				LOCTEXT("CtxMenuItem_Copy", "Copy"),
+				FText(),
+				FSlateIcon(FAppStyle::Get().GetStyleSetName(), "GenericCommands.Copy"),
+				FUIAction(FExecuteAction::CreateLambda([Event=HoveredEvent]()
+				{
+					Event->GetTrack()->OnClipboardCopyEvent(*Event);
+				})),
+				NAME_None,
+				EUserInterfaceActionType::Button
+			);
 		}
 
 		bHasAnyActions = true;
