[Insights] Fixed vertical alignment of text in SEditableTextBox and SSearchBox widgets (i.e. to be vertically centered).

diff --git a/Engine/Source/Developer/TraceInsights/Private/Insights/InsightsStyle.cpp b/Engine/Source/Developer/TraceInsights/Private/Insights/InsightsStyle.cpp
--- a/Engine/Source/Developer/TraceInsights/Private/Insights/InsightsStyle.cpp
+++ b/Engine/Source/Developer/TraceInsights/Private/Insights/InsightsStyle.cpp
@@ -360,6 +360,25 @@ void FInsightsStyle::FStyle::Initialize()
 
 	//////////////////////////////////////////////////
 
+	// NormalEditableTextBox && SearchBox
+	{
+		const FEditableTextBoxStyle& NormalEditableTextBoxStyle = FAppStyle::GetWidgetStyle<FEditableTextBoxStyle>("NormalEditableTextBox");
+		const FTextBlockStyle TextBlockStyle = FTextBlockStyle()
+			.SetColorAndOpacity(NormalEditableTextBoxStyle.ForegroundColor)
+			.SetHighlightColor(NormalEditableTextBoxStyle.FocusedForegroundColor)
+			.SetFont(NormalEditableTextBoxStyle.TextStyle.Font)
+			.SetFontSize(uint16(NormalEditableTextBoxStyle.TextStyle.Font.Size));
+		const FEditableTextBoxStyle EditableTextBoxStyle = FEditableTextBoxStyle(NormalEditableTextBoxStyle)
+			.SetPadding(FMargin(6.0f, 4.0f, 6.0f, 4.0f))
+			.SetTextStyle(TextBlockStyle);
+		Set("NormalEditableTextBox", EditableTextBoxStyle);
+
+		const FSearchBoxStyle& NormalSearchBoxStyle = FAppStyle::GetWidgetStyle<FSearchBoxStyle>("SearchBox");
+		const FSearchBoxStyle SearchBoxStyle = FSearchBoxStyle(NormalSearchBoxStyle)
+			.SetTextBoxStyle(EditableTextBoxStyle);
+		Set("SearchBox", SearchBoxStyle);
+	}
+
 	// PrimaryToolbar
 	{
 		FToolBarStyle PrimaryToolbarStyle = FAppStyle::Get().GetWidgetStyle<FToolBarStyle>("SlimToolBar");
