[Insights] Session Browser: Improved UI responsiveness when using a remote/slow trace store connection. Reduced pressure on StoreClient lock.

diff --git a/Engine/Source/Developer/TraceInsights/Private/Insights/StoreService/StoreBrowser.cpp b/Engine/Source/Developer/TraceInsights/Private/Insights/StoreService/StoreBrowser.cpp
--- a/Engine/Source/Developer/TraceInsights/Private/Insights/StoreService/StoreBrowser.cpp
+++ b/Engine/Source/Developer/TraceInsights/Private/Insights/StoreService/StoreBrowser.cpp
@@ -3,6 +3,7 @@
 #include "StoreBrowser.h"
 
 #include "Async/ParallelFor.h"
+#include "Containers/ContainerAllocationPolicies.h"
 #include "Trace/Analysis.h"
 #include "Trace/StoreClient.h"
 
@@ -120,7 +121,7 @@ void FStoreBrowser::UpdateTraces()
 
 	FStopwatch StopwatchTotal;
 	StopwatchTotal.Start();
-	
+
 	// Check if connection to store is still active We want to do this without locking the critical
 	// section, in case the UI needs to read values. The output is an atomic anyway.
 	{
@@ -132,12 +133,12 @@ void FStoreBrowser::UpdateTraces()
 			do
 			{
 				FPlatformProcess::Sleep(1.0);
-				ConnectionStatus.store(static_cast<EConnectionStatus>(--SecondsToSleep));			
+				ConnectionStatus.store(static_cast<EConnectionStatus>(--SecondsToSleep));
 			} while (SecondsToSleep);
 
 			// At this point ConnectionStatus is zero (EConnectionStatus::Connecting)
 			bIsValid = FInsightsManager::Get()->ReconnectToStore();
-			
+
 			if (bIsValid)
 			{
 				// If we just reconnected we need to reset known traces and change serials
@@ -157,39 +158,65 @@ void FStoreBrowser::UpdateTraces()
 			Host = FInsightsManager::Get()->GetLastStoreHost();
 		}
 	}
-	
-	FScopeLock StoreClientLock(&GetStoreClientCriticalSection());
+
 
 	// Check if the list of trace sessions or store settings have changed.
 	{
-		const UE::Trace::FStoreClient::FStatus* Status = StoreClient->GetStatus();
+		TArray<FString> NewWatchDirectories;
+		FString NewStoreDirectory;
+		uint32 NewSettingsChangeSerial = 0;
+		uint32 NewStoreChangeSerial = 0;
+
+		{
+			// Get Trace Store status.
+			FScopeLock StoreClientLock(&GetStoreClientCriticalSection());
+			const UE::Trace::FStoreClient::FStatus* Status = StoreClient->GetStatus();
+			if (Status)
+			{
+				NewSettingsChangeSerial = Status->GetSettingsSerial();
+				NewStoreChangeSerial = Status->GetChangeSerial();
+				if (SettingsChangeSerial != NewSettingsChangeSerial)
+				{
+					Status->GetWatchDirectories(NewWatchDirectories);
+					NewStoreDirectory = FString(Status->GetStoreDir());
+				}
+			}
+		}
 
-		if (Status != nullptr && SettingsChangeSerial != Status->GetSettingsSerial())
+		if (SettingsChangeSerial != NewSettingsChangeSerial)
 		{
-			SettingsChangeSerial = Status->GetSettingsSerial();
-			
+			SettingsChangeSerial = NewSettingsChangeSerial;
+
 			// Update watch dirs
 			FScopeLock Lock(&TracesCriticalSection);
 			WatchDirectories.Empty();
-			Status->GetWatchDirectories(WatchDirectories);
-			StoreDirectory = FString(Status->GetStoreDir());
+			WatchDirectories = NewWatchDirectories;
+			StoreDirectory = NewStoreDirectory;
 		}
-		
-		if (Status != nullptr && StoreChangeSerial != Status->GetChangeSerial())
+
+		if (StoreChangeSerial != NewStoreChangeSerial)
 		{
-			StoreChangeSerial = Status->GetChangeSerial();
+			StoreChangeSerial = NewStoreChangeSerial;
 
 			// Check for removed traces.
 			{
-				for (int32 TraceIndex = Traces.Num() - 1; TraceIndex >= 0; --TraceIndex)
+				TSet<uint32, DefaultKeyFuncs<uint32>, TInlineSetAllocator<4>> RemovedTraces;
+				for (const TSharedPtr<FStoreBrowserTraceInfo>& Trace : Traces)
 				{
-					const uint32 TraceId = Traces[TraceIndex]->TraceId;
-					const UE::Trace::FStoreClient::FTraceInfo* TraceInfo = StoreClient->GetTraceInfoById(TraceId);
+					FScopeLock StoreClientLock(&GetStoreClientCriticalSection());
+					const UE::Trace::FStoreClient::FTraceInfo* TraceInfo = StoreClient->GetTraceInfoById(Trace->TraceId);
 					if (TraceInfo == nullptr)
 					{
-						FScopeLock Lock(&TracesCriticalSection);
-						TracesChangeSerial++;
-						Traces.RemoveAtSwap(TraceIndex);
+						RemovedTraces.Add(Trace->TraceId);
+					}
+				}
+				if (RemovedTraces.Num() > 0)
+				{
+					FScopeLock Lock(&TracesCriticalSection);
+					TracesChangeSerial++;
+					Traces.RemoveAllSwap([&RemovedTraces](const TSharedPtr<FStoreBrowserTraceInfo>& Trace) { return RemovedTraces.Contains(Trace->TraceId); });
+					for (const uint32& TraceId : RemovedTraces)
+					{
 						TraceMap.Remove(TraceId);
 					}
 				}
@@ -197,15 +224,20 @@ void FStoreBrowser::UpdateTraces()
 
 			// Check for added traces.
 			{
-				const int32 TraceCount = StoreClient->GetTraceCount();
+				TArray<TSharedPtr<FStoreBrowserTraceInfo>, TInlineAllocator<8>> AddedTraces;
+				int32 TraceCount = 0;
+				{
+					FScopeLock StoreClientLock(&GetStoreClientCriticalSection());
+					TraceCount = StoreClient->GetTraceCount();
+				}
 				for (int32 TraceIndex = 0; TraceIndex < TraceCount; ++TraceIndex)
 				{
+					FScopeLock StoreClientLock(&GetStoreClientCriticalSection());
 					const UE::Trace::FStoreClient::FTraceInfo* TraceInfo = StoreClient->GetTraceInfo(TraceIndex);
 					if (TraceInfo != nullptr)
 					{
 						const uint32 TraceId = TraceInfo->GetId();
-						TSharedPtr<FStoreBrowserTraceInfo>* TracePtrPtr = TraceMap.Find(TraceId);
-						if (TracePtrPtr == nullptr)
+						if (!TraceMap.Find(TraceId))
 						{
 							TSharedPtr<FStoreBrowserTraceInfo> TracePtr = MakeShared<FStoreBrowserTraceInfo>();
 							FStoreBrowserTraceInfo& Trace = *TracePtr;
@@ -230,13 +262,34 @@ void FStoreBrowser::UpdateTraces()
 							Trace.Timestamp = FStoreBrowserTraceInfo::ConvertTimestamp(TraceInfo->GetTimestamp());
 							Trace.Size = TraceInfo->GetSize();
 
-							FScopeLock Lock(&TracesCriticalSection);
-							TracesChangeSerial++;
-							Traces.Add(TracePtr);
-							TraceMap.Add(TraceId, TracePtr);
+							AddedTraces.Add(TracePtr);
+
+							// Flush list from time to time.
+							if (AddedTraces.Num() > 5)
+							{
+								StoreClientLock.Unlock();
+								FScopeLock Lock(&TracesCriticalSection);
+								TracesChangeSerial++;
+								for (const TSharedPtr<FStoreBrowserTraceInfo>& AddedTracePtr : AddedTraces)
+								{
+									Traces.Add(AddedTracePtr);
+									TraceMap.Add(AddedTracePtr->TraceId, AddedTracePtr);
+								}
+								AddedTraces.Reset();
+							}
 						}
 					}
 				}
+				if (AddedTraces.Num() > 0)
+				{
+					FScopeLock Lock(&TracesCriticalSection);
+					TracesChangeSerial++;
+					for (TSharedPtr<FStoreBrowserTraceInfo> AddedTracePtr : AddedTraces)
+					{
+						Traces.Add(AddedTracePtr);
+						TraceMap.Add(AddedTracePtr->TraceId, AddedTracePtr);
+					}
+				}
 			}
 		}
 	}
@@ -253,10 +306,14 @@ void FStoreBrowser::UpdateTraces()
 			const uint32 TraceId = KV.Key;
 			FStoreBrowserTraceInfo& Trace = *KV.Value;
 
-			ensure(Trace.bIsLive);
-
 			FDateTime Timestamp(0);
 			uint64 Size = 0;
+			ensure(Trace.bIsLive);
+			bool bIsLive = true;
+			uint32 IpAddress = 0;
+
+			FScopeLock StoreClientLock(&GetStoreClientCriticalSection());
+
 			const UE::Trace::FStoreClient::FTraceInfo* TraceInfo = StoreClient->GetTraceInfoById(TraceId);
 			if (TraceInfo != nullptr)
 			{
@@ -267,30 +324,32 @@ void FStoreBrowser::UpdateTraces()
 			const UE::Trace::FStoreClient::FSessionInfo* SessionInfo = StoreClient->GetSessionInfoByTraceId(TraceId);
 			if (SessionInfo != nullptr)
 			{
-				// The trace is still live.
-				const uint32 IpAddress = SessionInfo->GetIpAddress();
-				if (IpAddress != Trace.IpAddress || Timestamp != Trace.Timestamp || Size != Trace.Size)
-				{
-					FScopeLock Lock(&TracesCriticalSection);
-					TracesChangeSerial++;
-					Trace.ChangeSerial++;
-					Trace.Timestamp = Timestamp;
-					Trace.Size = Size;
-					Trace.IpAddress = IpAddress;
-				}
+				bIsLive = true;
+				IpAddress = SessionInfo->GetIpAddress();
 			}
 			else
 			{
-				NotLiveTraces.Add(TraceId);
+				bIsLive = false;
+			}
+
+			StoreClientLock.Unlock();
+
+			if (!bIsLive || IpAddress != Trace.IpAddress || Timestamp != Trace.Timestamp || Size != Trace.Size)
+			{
+				StoreClientLock.Unlock();
 
-				// The trace is not live anymore.
 				FScopeLock Lock(&TracesCriticalSection);
 				TracesChangeSerial++;
 				Trace.ChangeSerial++;
 				Trace.Timestamp = Timestamp;
 				Trace.Size = Size;
-				Trace.bIsLive = false;
-				Trace.IpAddress = 0;
+				Trace.bIsLive = bIsLive;
+				Trace.IpAddress = IpAddress;
+			}
+
+			if (!bIsLive)
+			{
+				NotLiveTraces.Add(TraceId);
 			}
 		}
 
@@ -305,9 +364,14 @@ void FStoreBrowser::UpdateTraces()
 
 	// Check if we have new live sessions.
 	{
-		const uint32 SessionCount = StoreClient->GetSessionCount();
+		uint32 SessionCount = 0;
+		{
+			FScopeLock StoreClientLock(&GetStoreClientCriticalSection());
+			SessionCount = StoreClient->GetSessionCount();
+		}
 		for (uint32 SessionIndex = 0; SessionIndex < SessionCount; ++SessionIndex)
 		{
+			FScopeLock StoreClientLock(&GetStoreClientCriticalSection());
 			const UE::Trace::FStoreClient::FSessionInfo* SessionInfo = StoreClient->GetSessionInfo(SessionIndex);
 			if (SessionInfo != nullptr)
 			{
@@ -333,6 +397,7 @@ void FStoreBrowser::UpdateTraces()
 
 						FStoreBrowserTraceInfo& Trace = **TracePtrPtr;
 
+						StoreClientLock.Unlock();
 						FScopeLock Lock(&TracesCriticalSection);
 						TracesChangeSerial++;
 						Trace.ChangeSerial++;
@@ -346,8 +411,6 @@ void FStoreBrowser::UpdateTraces()
 		}
 	}
 
-	StoreClientLock.Unlock();
-
 	StopwatchTotal.Update();
 	const uint64 Step3 = StopwatchTotal.GetAccumulatedTimeMs();
 
