[Insights] Fixed the (non)localized text used in context menu for a track in Timing View.

diff --git a/Engine/Source/Developer/TraceInsights/Private/Insights/Widgets/STimingView.cpp b/Engine/Source/Developer/TraceInsights/Private/Insights/Widgets/STimingView.cpp
--- a/Engine/Source/Developer/TraceInsights/Private/Insights/Widgets/STimingView.cpp
+++ b/Engine/Source/Developer/TraceInsights/Private/Insights/Widgets/STimingView.cpp
@@ -2915,7 +2915,7 @@ void STimingView::ShowContextMenu(const FPointerEvent& MouseEvent)
 			double RangeDuration = HoveredEvent->GetDuration();
 			
 			MenuBuilder.AddMenuEntry(
-				FText::FromString(TEXT("Select Time Range of Event")),
+				LOCTEXT("ContextMenu_SelectEventTimeRange", "Select Time Range of Event"),
 				FText(),
 				FSlateIcon(),
 				FUIAction(FExecuteAction::CreateLambda([this, RangeStart, RangeDuration]()
@@ -2925,15 +2925,15 @@ void STimingView::ShowContextMenu(const FPointerEvent& MouseEvent)
 				NAME_None,
 				EUserInterfaceActionType::Button
 			);
-			
+
 			MenuBuilder.AddMenuEntry(
-				LOCTEXT("CtxMenuItem_Copy", "Copy"),
+				LOCTEXT("ContextMenu_Copy", "Copy"),
 				FText(),
 				FSlateIcon(FAppStyle::Get().GetStyleSetName(), "GenericCommands.Copy"),
 				FUIAction(FExecuteAction::CreateLambda([Event=HoveredEvent]()
-				{
-					Event->GetTrack()->OnClipboardCopyEvent(*Event);
-				})),
+					{
+						Event->GetTrack()->OnClipboardCopyEvent(*Event);
+					})),
 				NAME_None,
 				EUserInterfaceActionType::Button
 			);
