[Insights] Session Browser fixes:
  - Disabled Rename/Delete trace files when using a remote trace server.
  - Fixed updating the list of traces. If list is scrolled on bottom it will remain scrolled on bottom also when new traces are added.
  - Fixed initial channels example (including a typo) and usage of "default" preset to include "region" and "screenshot" channels (Connection tab).
  - Fixed trace recorder IP address to be prefilled with correct IP or host name of (remote) trace server (Connection tab).

diff --git a/Engine/Source/Developer/TraceInsights/Private/Insights/Widgets/SStartPageWindow.cpp b/Engine/Source/Developer/TraceInsights/Private/Insights/Widgets/SStartPageWindow.cpp
--- a/Engine/Source/Developer/TraceInsights/Private/Insights/Widgets/SStartPageWindow.cpp
+++ b/Engine/Source/Developer/TraceInsights/Private/Insights/Widgets/SStartPageWindow.cpp
@@ -56,9 +56,14 @@
 ////////////////////////////////////////////////////////////////////////////////////////////////////
 
 #define LOCTEXT_NAMESPACE "STraceStoreWindow"
-// For now we only support store settings on platforms that run Unreal Trace Server.
+
+// For now, we only support store settings on platforms that run Unreal Trace Server.
 #define STORE_SUPPORTS_SETTINGS PLATFORM_WINDOWS
 
+// Should match GDefaultChannels (from Runtime\Core\Private\ProfilingDebugging\TraceAuxiliary.cpp).
+// We cannot use GDefaultChannels directly as UE_TRACE_ENABLED may be off.
+static const TCHAR* GInsightsDefaultChannelPreset = TEXT("cpu,gpu,frame,log,bookmark,screenshot,region");
+
 ////////////////////////////////////////////////////////////////////////////////////////////////////
 // STraceListRow
 ////////////////////////////////////////////////////////////////////////////////////////////////////
@@ -1084,8 +1089,12 @@ TSharedRef<SWidget> STraceStoreWindow::ConstructFiltersToolbar()
 			.IsChecked_Lambda([this]() { return bSearchByCommandLine ? ECheckBoxState::Checked : ECheckBoxState::Unchecked; })
 			.ToolTipText(LOCTEXT("ToggleNameFilter_Tooltip", "Toggle between filtering the list of trace sessions by name or by command line."))
 			[
-				SNew(SImage)
-				.Image(FInsightsStyle::Get().GetBrush("Icons.Console"))
+				SNew(SBox)
+				.Padding(1.0f)
+				[
+					SNew(SImage)
+					.Image(FInsightsStyle::Get().GetBrush("Icons.Console"))
+				]
 			]);
 
 		// Text Filter (Search Box)
@@ -1339,7 +1348,7 @@ private:
 	FReply OnDelete();
 	bool CanModifyStore() const;
 	FText ModifyStoreTooltip() const;
-	
+
 	bool bInOperation = false;
 	const STraceStoreWindow* Window = nullptr;
 
@@ -1471,15 +1480,15 @@ TSharedRef<SWidget> STraceDirectoryItem::ConstructOperations()
 FReply STraceDirectoryItem::OnModifyStore()
 {
 	FSlateApplication::Get().CloseToolTip();
-	
+
 	// Avoid executing more than one operation
 	if (bInOperation)
 	{
 		return FReply::Handled();
 	}
-	
+
 	TGuardValue<bool> OperationGuard(bInOperation, true);
-	
+
 	if (IDesktopPlatform* DesktopPlatform = FDesktopPlatformModule::Get())
 	{
 		const FString Title = LOCTEXT("SetTraceStoreDirectorySelectPopupTitle", "Set Trace Store Directory").ToString();
@@ -1525,7 +1534,7 @@ FReply STraceDirectoryItem::OnExplore()
 
 ////////////////////////////////////////////////////////////////////////////////////////////////////
 
-FReply STraceDirectoryItem::OnDelete() 
+FReply STraceDirectoryItem::OnDelete()
 {
 	FSlateApplication::Get().CloseToolTip();
 
@@ -1534,9 +1543,9 @@ FReply STraceDirectoryItem::OnDelete()
 	{
 		return FReply::Handled();
 	}
-	
+
 	TGuardValue<bool> OperationGuard(bInOperation, true);
-	
+
 	if (Model)
 	{
 		const EAppReturnType::Type bConfirmed = FMessageDialog::Open(EAppMsgType::OkCancel, FText::Format(
@@ -1731,12 +1740,12 @@ TSharedRef<SWidget> STraceStoreWindow::ConstructTraceStoreDirectoryPanel()
 			SAssignNew(StoreSettingsArea, SScrollBox)
 			.Orientation(Orient_Vertical)
 			.Visibility(EVisibility::Collapsed)
-			
+
 			+ SScrollBox::Slot()
 			[
 
 				SNew(SVerticalBox)
-				
+
 				+ SVerticalBox::Slot()
 				.AutoHeight()
 				.Padding(0.0f, 4.0f)
@@ -1883,24 +1892,25 @@ FText STraceStoreWindow::GetConnectionStatusTooltip() const
 {
 	using Insights::FStoreBrowser;
 
-	static FText Connected = LOCTEXT("Connected", "Connected to the trace server");
-	static FText NotConnected = LOCTEXT("NoConnection", "Unable to connect trace server.");
-	static FText Connecting = LOCTEXT("Connecting", "Trying to connect to trace server.");
+	static FText Connected    = LOCTEXT("Connected",    "Connected to the trace server.");
+	static FText NotConnected = LOCTEXT("NoConnection", "Unable to connect to trace server.");
+	static FText Connecting   = LOCTEXT("Connecting",   "Trying to connect to trace server.");
 	static FText Disconnected = LOCTEXT("Disconnected", "Connection to trace server has been lost. Attempting to reconnect in {0} seconds.");
 
 	const FStoreBrowser::EConnectionStatus Status = StoreBrowser->GetConnectionStatus();
-	switch(Status)
+	switch (Status)
 	{
-	case FStoreBrowser::EConnectionStatus::Connected:
-		return Connected;
-	case FStoreBrowser::EConnectionStatus::NoConnection:
-		return NotConnected;
-	case FStoreBrowser::EConnectionStatus::Connecting:
-		return Connecting;
-	default:
-		{
+		case FStoreBrowser::EConnectionStatus::Connected:
+			return Connected;
+
+		case FStoreBrowser::EConnectionStatus::NoConnection:
+			return NotConnected;
+
+		case FStoreBrowser::EConnectionStatus::Connecting:
+			return Connecting;
+
+		default:
 			return FText::Format(Disconnected, FText::AsNumber(static_cast<uint32>(Status)));
-		}
 	}
 }
 
@@ -1964,24 +1974,29 @@ TSharedPtr<SWidget> STraceStoreWindow::TraceList_GetMenuContent()
 
 bool STraceStoreWindow::CanEditTraceFile() const
 {
-	return SelectedTrace.IsValid() && !SelectedTrace->bIsLive && !SelectedTrace->bWasJustRenamed;
+	return CanChangeStoreSettings()
+		&& SelectedTrace.IsValid()
+		&& !SelectedTrace->bIsLive
+		&& !SelectedTrace->bWasJustRenamed;
 }
 
 ////////////////////////////////////////////////////////////////////////////////////////////////////
 
 void STraceStoreWindow::RenameTraceFile()
 {
-	if (SelectedTrace.IsValid() && !SelectedTrace->bIsLive && !SelectedTrace->bWasJustRenamed)
+	if (!CanEditTraceFile())
 	{
-		SelectedTrace->bIsRenaming = true;
+		return;
+	}
 
-		FSlateApplication::Get().CloseToolTip();
+	SelectedTrace->bIsRenaming = true;
 
-		TSharedPtr<SEditableTextBox> RenameTextBox = SelectedTrace->RenameTextBox.Pin();
-		if (RenameTextBox.IsValid())
-		{
-			FSlateApplication::Get().SetKeyboardFocus(RenameTextBox.ToSharedRef(), EFocusCause::SetDirectly);
-		}
+	FSlateApplication::Get().CloseToolTip();
+
+	TSharedPtr<SEditableTextBox> RenameTextBox = SelectedTrace->RenameTextBox.Pin();
+	if (RenameTextBox.IsValid())
+	{
+		FSlateApplication::Get().SetKeyboardFocus(RenameTextBox.ToSharedRef(), EFocusCause::SetDirectly);
 	}
 }
 
@@ -1989,7 +2004,7 @@ void STraceStoreWindow::RenameTraceFile()
 
 void STraceStoreWindow::DeleteTraceFile()
 {
-	if (!SelectedTrace.IsValid() || SelectedTrace->bIsLive || SelectedTrace->bWasJustRenamed)
+	if (!CanEditTraceFile())
 	{
 		return;
 	}
@@ -2380,34 +2395,45 @@ void STraceStoreWindow::UpdateTraceListView()
 		}
 	}
 
+	double DistanceFromTop = TraceListView->GetScrollDistance().Y;
+	double DistanceFromBottom = TraceListView->GetScrollDistanceRemaining().Y;
+
 	TraceListView->RebuildList();
 
-	if ((SortColumn == TraceStoreColumns::Date && SortMode == EColumnSortMode::Ascending) ||
-		(SortColumn == TraceStoreColumns::Status && SortMode == EColumnSortMode::Ascending))
+	// If no selection...
+	if (!NewSelectedTrace.IsValid() && FilteredTraceViewModels.Num() > 0)
 	{
-		// If no selection, auto select the last (newest) trace.
-		if (!NewSelectedTrace.IsValid() && FilteredTraceViewModels.Num() > 0)
+		if ((SortColumn == TraceStoreColumns::Date && SortMode == EColumnSortMode::Ascending) ||
+			(SortColumn == TraceStoreColumns::Status && SortMode == EColumnSortMode::Ascending))
 		{
+			// Auto select the last (newest) trace.
 			NewSelectedTrace = FilteredTraceViewModels.Last();
+			DistanceFromTop = 1.0;
+			DistanceFromBottom = 0.0; // scroll to bottom
 		}
+		else
+		{
+			// Auto select the first trace.
+			NewSelectedTrace = FilteredTraceViewModels[0];
+			DistanceFromTop = 0.0; // scroll to top
+			DistanceFromBottom = 1.0;
+		}
+	}
 
+	if (FMath::IsNearlyZero(DistanceFromBottom, 1.0E-8))
+	{
 		TraceListView->ScrollToBottom();
 	}
-	else
+	else if (FMath::IsNearlyZero(DistanceFromTop, 1.0E-8))
 	{
-		// If no selection, auto select the first trace.
-		if (!NewSelectedTrace.IsValid() && FilteredTraceViewModels.Num() > 0)
-		{
-			NewSelectedTrace = FilteredTraceViewModels[0];
-		}
+		TraceListView->ScrollToTop();
 	}
 
-	// Restore selection and ensure it is visible.
+	// Restore selection.
 	if (NewSelectedTrace.IsValid())
 	{
 		TraceListView->ClearSelection();
 		TraceListView->SetItemSelection(NewSelectedTrace, true);
-		TraceListView->RequestScrollIntoView(NewSelectedTrace);
 	}
 }
 
@@ -3710,6 +3736,8 @@ void SConnectionWindow::Construct(const FArguments& InArgs)
 
 TSharedRef<SWidget> SConnectionWindow::ConstructConnectPanel()
 {
+	FText InitialChannelsExampleText = FText::FromString(FString::Printf(TEXT("default,counter,stats,file,loadtime,assetloadtime,task\ndefault=%s"), GInsightsDefaultChannelPreset));
+
 	TSharedRef<SWidget> Widget = SNew(SVerticalBox)
 
 		+ SVerticalBox::Slot()
@@ -3828,6 +3856,7 @@ TSharedRef<SWidget> SConnectionWindow::ConstructConnectPanel()
 			+ SHorizontalBox::Slot()
 			.AutoWidth()
 			.VAlign(VAlign_Top)
+			.Padding(FMargin(0.0f, 4.0f, 0.0f, 0.0f))
 			[
 				SNew(STextBlock)
 				.Text(LOCTEXT("InitialChannelsExamplesTitle", "Examples"))
@@ -3837,7 +3866,7 @@ TSharedRef<SWidget> SConnectionWindow::ConstructConnectPanel()
 			[
 				SNew(SEditableTextBox)
 				.IsReadOnly(true)
-				.Text(FText::FromString(TEXT("default,counter,stats,file,loadtime,assetloadtime,task\ndefault=cpu,gpu,frame,log,bookmark,screenshot")))
+				.Text(InitialChannelsExampleText)
 			]
 		]
 		+ SVerticalBox::Slot()
@@ -3846,7 +3875,7 @@ TSharedRef<SWidget> SConnectionWindow::ConstructConnectPanel()
 		.Padding(198.0f, 2.0f, 12.0f, 0.0f)
 		[
 			SNew(STextBlock)
-			.Text(LOCTEXT("InitialChannelsNote2Text", "Same channels/presets (like \"memory\") cannot be enabled on late connections."))
+			.Text(LOCTEXT("InitialChannelsNote2Text", "Some channels/presets (like \"memory\") cannot be enabled on late connections."))
 		]
 
 		+ SVerticalBox::Slot()
@@ -3878,20 +3907,40 @@ TSharedRef<SWidget> SConnectionWindow::ConstructConnectPanel()
 
 	const FText LocalHost = FText::FromString(TEXT("127.0.0.1"));
 
-	TSharedPtr<FInternetAddr> RecorderAddr;
-	if (ISocketSubsystem* Sockets = ISocketSubsystem::Get())
-	{
-		bool bCanBindAll = false;
-		RecorderAddr = Sockets->GetLocalHostAddr(*GLog, bCanBindAll);
-	}
-	if (RecorderAddr.IsValid())
+	if (FInsightsManager::Get()->CanChangeStoreSettings()) // local trace store
 	{
-		const FString RecorderAddrStr = RecorderAddr->ToString(false);
-		TraceRecorderAddressTextBox->SetText(FText::FromString(RecorderAddrStr));
+		TSharedPtr<FInternetAddr> RecorderAddr;
+		if (ISocketSubsystem* Sockets = ISocketSubsystem::Get())
+		{
+			bool bCanBindAll = false;
+			RecorderAddr = Sockets->GetLocalHostAddr(*GLog, bCanBindAll);
+		}
+		if (RecorderAddr.IsValid())
+		{
+			const FString RecorderAddrStr = RecorderAddr->ToString(false);
+			TraceRecorderAddressTextBox->SetText(FText::FromString(RecorderAddrStr));
+		}
+		else
+		{
+			TraceRecorderAddressTextBox->SetText(LocalHost);
+		}
 	}
-	else
+	else // remote trace store
 	{
-		TraceRecorderAddressTextBox->SetText(LocalHost);
+		FScopeLock StoreClientLock(&FInsightsManager::Get()->GetStoreClientCriticalSection());
+		UE::Trace::FStoreClient* StoreClient = FInsightsManager::Get()->GetStoreClient();
+		if (StoreClient)
+		{
+			const uint32 StoreAddress = StoreClient->GetStoreAddress();
+			StoreClientLock.Unlock();
+			const FString RecorderAddrStr = FString::Printf(TEXT("%u.%u.%u.%u"), (StoreAddress >> 24) & 0xFF, (StoreAddress >> 16) & 0xFF, (StoreAddress >> 8) & 0xFF, StoreAddress & 0xFF);
+			TraceRecorderAddressTextBox->SetText(FText::FromString(RecorderAddrStr));
+		}
+		else
+		{
+			StoreClientLock.Unlock();
+			TraceRecorderAddressTextBox->SetText(FText::FromString(FInsightsManager::Get()->GetLastStoreHost()));
+		}
 	}
 
 	RunningInstanceAddressTextBox->SetText(LocalHost);
@@ -3930,7 +3979,7 @@ FReply SConnectionWindow::Connect_OnClicked()
 		PrerequisitesPtr = &Prerequisites;
 	}
 
-	const FString ChannelsExpandedStr = ChannelsTextBox->GetText().ToString().Replace(TEXT("default"), TEXT("cpu,gpu,frame,log,bookmark"));
+	const FString ChannelsExpandedStr = ChannelsTextBox->GetText().ToString().Replace(TEXT("default"), GInsightsDefaultChannelPreset);
 
 	FGraphEventRef PreConnectTask = FFunctionGraphTask::CreateAndDispatchWhenReady(
 		[this, TraceRecorderAddressStr, RunningInstanceAddressStr, ChannelsExpandedStr]
