[Insights] Frame stats fixes

diff --git a/Engine/Source/Developer/TraceInsights/Private/Insights/ViewModels/FrameTrackHelper.cpp b/Engine/Source/Developer/TraceInsights/Private/Insights/ViewModels/FrameTrackHelper.cpp
--- a/Engine/Source/Developer/TraceInsights/Private/Insights/ViewModels/FrameTrackHelper.cpp
+++ b/Engine/Source/Developer/TraceInsights/Private/Insights/ViewModels/FrameTrackHelper.cpp
@@ -149,7 +149,7 @@ void FFrameTrackDrawHelper::DrawCached(const FFrameTrackSeries& Series) const
 
 	NumFrames += Series.NumAggregatedFrames;
 
-	FLinearColor SeriesColor = GetColorByFrameType(Series.FrameType);
+	FLinearColor SeriesColor = Series.Color;
 
 	const float SampleW = Viewport.GetSampleWidth();
 	const int32 NumSamples = Series.Samples.Num();
diff --git a/Engine/Source/Developer/TraceInsights/Private/Insights/ViewModels/FrameTrackHelper.h b/Engine/Source/Developer/TraceInsights/Private/Insights/ViewModels/FrameTrackHelper.h
--- a/Engine/Source/Developer/TraceInsights/Private/Insights/ViewModels/FrameTrackHelper.h
+++ b/Engine/Source/Developer/TraceInsights/Private/Insights/ViewModels/FrameTrackHelper.h
@@ -81,6 +81,8 @@ struct FFrameTrackSeries
 	bool bIsVisible;
 	int32 NumAggregatedFrames; // total number of frames aggregated in samples; i.e. sum of all Sample.NumFrames
 	TArray<FFrameTrackSample> Samples; // the aggregated samples
+	FLinearColor Color;
+	FText Name;
 
 	explicit FFrameTrackSeries(int32 InFrameType, EFrameTrackSeriesType InType)
 		: FrameType(InFrameType)
@@ -109,7 +111,6 @@ struct FTimerFrameStatsTrackSeries : public FFrameTrackSeries
 	}
 
 	uint32 TimerId;
-	FText TimerDisplayName;
 };
 
 ////////////////////////////////////////////////////////////////////////////////////////////////////
diff --git a/Engine/Source/Developer/TraceInsights/Private/Insights/Widgets/SFrameTrack.cpp b/Engine/Source/Developer/TraceInsights/Private/Insights/Widgets/SFrameTrack.cpp
--- a/Engine/Source/Developer/TraceInsights/Private/Insights/Widgets/SFrameTrack.cpp
+++ b/Engine/Source/Developer/TraceInsights/Private/Insights/Widgets/SFrameTrack.cpp
@@ -229,7 +229,10 @@ TSharedRef<FFrameTrackSeries> SFrameTrack::FindOrAddSeries(ETraceFrameType Frame
 	}
 
 	LLM_SCOPE_BYTAG(Insights);
+
 	TSharedRef<FFrameTrackSeries> SeriesRef = MakeShared<FFrameTrackSeries>(FrameType, EFrameTrackSeriesType::Frame);
+	SeriesRef->Color = FFrameTrackDrawHelper::GetColorByFrameType(FrameType);
+	SeriesRef->Name = FText::Format(LOCTEXT("FrameTrackSeriesName_Format", "{0} {1}"), FText::FromString(FFrameTrackDrawHelper::FrameTypeToString(FrameType)), LOCTEXT("Frame", "Frame"));
 	AllSeries.Add(SeriesRef);
 	return SeriesRef;
 }
@@ -561,18 +564,8 @@ int32 SFrameTrack::OnPaint(const FPaintArgs& Args, const FGeometry& AllottedGeom
 			// First line: "Rendering Frame 1,234"
 			TStringBuilder<512> StringBuilder;
 
-			if (HoveredSample.Series->Type == EFrameTrackSeriesType::TimerFrameStats)
-			{
-				StringBuilder.Append(StaticCastSharedPtr<FTimerFrameStatsTrackSeries>(HoveredSample.Series)->TimerDisplayName.ToString());
-				StringBuilder.Append(TEXT(" "));
-				StringBuilder.Append(FFrameTrackDrawHelper::FrameTypeToString(HoveredSample.Series->FrameType));
-			}
-			else
-			{
-				StringBuilder.Append(FFrameTrackDrawHelper::FrameTypeToString(HoveredSample.Series->FrameType));
-			}
-
-			StringBuilder.Append(TEXT(" Frame "));
+			StringBuilder.Append(HoveredSample.Series->Name.ToString());
+			StringBuilder.Append(TEXT(" "));
 			StringBuilder.Append(FText::AsNumber(HoveredSample.Sample->LargestFrameIndex).ToString());
 			const FString Text1(StringBuilder);
 
@@ -1269,8 +1262,8 @@ void SFrameTrack::ShowContextMenu(const FPointerEvent& MouseEvent)
 		);
 		MenuBuilder.AddMenuEntry
 		(
-			FrameStatSeries->TimerDisplayName,
-			FText::Format(TooltipTextBase, FrameStatSeries->TimerDisplayName),
+			FText::Format(LOCTEXT("ContextMenu_ShowFrameStatsSeries_Name", "{0}"), FrameStatSeries->Name),
+			FText::Format(TooltipTextBase, FrameStatSeries->Name),
 			FSlateIcon(),
 			Action_ShowRenderingFrames,
 			NAME_None,
@@ -1538,7 +1531,7 @@ bool SFrameTrack::HasFrameStatSeries(ETraceFrameType FrameType, uint32 TimerId)
 
 ////////////////////////////////////////////////////////////////////////////////////////////////////
 
-TSharedPtr<FTimerFrameStatsTrackSeries> SFrameTrack::AddTimerFrameStatSeries(ETraceFrameType FrameType, uint32 TimerId)
+TSharedPtr<FTimerFrameStatsTrackSeries> SFrameTrack::AddTimerFrameStatSeries(ETraceFrameType FrameType, uint32 TimerId, FLinearColor Color, FText Name)
 {
 	TSharedPtr<FFrameTrackSeries> ExistingSeries = FindFrameStatsSeries(FrameType, TimerId);
 
@@ -1549,6 +1542,8 @@ TSharedPtr<FTimerFrameStatsTrackSeries> SFrameTrack::AddTimerFrameStatSeries(ETr
 
 	TSharedRef<FTimerFrameStatsTrackSeries> SeriesRef = MakeShared<FTimerFrameStatsTrackSeries>(FrameType, TimerId);
 	SeriesRef->TimerId = TimerId;
+	SeriesRef->Color = Color;
+	SeriesRef->Name = Name;
 	AllSeries.Add(SeriesRef);
 
 	bIsStateDirty = true;
diff --git a/Engine/Source/Developer/TraceInsights/Private/Insights/Widgets/SFrameTrack.h b/Engine/Source/Developer/TraceInsights/Private/Insights/Widgets/SFrameTrack.h
--- a/Engine/Source/Developer/TraceInsights/Private/Insights/Widgets/SFrameTrack.h
+++ b/Engine/Source/Developer/TraceInsights/Private/Insights/Widgets/SFrameTrack.h
@@ -129,7 +129,7 @@ public:
 	virtual FCursorReply OnCursorQuery(const FGeometry& MyGeometry, const FPointerEvent& CursorEvent) const override;
 
 	bool HasFrameStatSeries(ETraceFrameType FrameType, uint32 TimerId);
-	TSharedPtr<FTimerFrameStatsTrackSeries> AddTimerFrameStatSeries(ETraceFrameType FrameType, uint32 TimerId);
+	TSharedPtr<FTimerFrameStatsTrackSeries> AddTimerFrameStatSeries(ETraceFrameType FrameType, uint32 TimerId, FLinearColor Color, FText Name);
 	bool RemoveTimerFrameStatSeries(ETraceFrameType FrameType, uint32 TimerId);
 
 protected:
diff --git a/Engine/Source/Developer/TraceInsights/Private/Insights/Widgets/STimersView.cpp b/Engine/Source/Developer/TraceInsights/Private/Insights/Widgets/STimersView.cpp
--- a/Engine/Source/Developer/TraceInsights/Private/Insights/Widgets/STimersView.cpp
+++ b/Engine/Source/Developer/TraceInsights/Private/Insights/Widgets/STimersView.cpp
@@ -783,6 +783,8 @@ void STimersView::TreeView_BuildPlotTimerMenu(FMenuBuilder& MenuBuilder)
 		return TimingView.IsValid() && NumSelectedNodes == 1 && SelectedNode.IsValid() && SelectedNode->GetType() != ETimerNodeType::Group;
 	};
 
+	MenuBuilder.BeginSection("MainGraphTrack", LOCTEXT("Plot_Series_MainGraphTrack_Section", "Main Graph Track"));
+
 	// Add/remove series to/from graph track
 	{
 		FUIAction Action_ToggleTimerInGraphTrack;
@@ -795,7 +797,7 @@ void STimersView::TreeView_BuildPlotTimerMenu(FMenuBuilder& MenuBuilder)
 		{
 			MenuBuilder.AddMenuEntry
 			(
-				LOCTEXT("ContextMenu_RemoveFromGraphTrack", "Remove series from graph track"),
+				LOCTEXT("ContextMenu_RemoveFromGraphTrack", "Remove instance series from graph track"),
 				LOCTEXT("ContextMenu_RemoveFromGraphTrack_Desc", "Removes the series containing event instances of the selected timer from the Main Graph track."),
 				FSlateIcon(FInsightsStyle::GetStyleSetName(), "Icons.RemoveGraphSeries"),
 				Action_ToggleTimerInGraphTrack,
@@ -807,7 +809,7 @@ void STimersView::TreeView_BuildPlotTimerMenu(FMenuBuilder& MenuBuilder)
 		{
 			MenuBuilder.AddMenuEntry
 			(
-				LOCTEXT("ContextMenu_AddToGraphTrack", "Add series to graph track"),
+				LOCTEXT("ContextMenu_AddToGraphTrack", "Add instance series to graph track"),
 				LOCTEXT("ContextMenu_AddToGraphTrack_Desc", "Adds a series containing event instances of the selected timer to the Main Graph track."),
 				FSlateIcon(FInsightsStyle::GetStyleSetName(), "Icons.AddGraphSeries"),
 				Action_ToggleTimerInGraphTrack,
@@ -885,6 +887,9 @@ void STimersView::TreeView_BuildPlotTimerMenu(FMenuBuilder& MenuBuilder)
 		}
 	}
 
+	MenuBuilder.EndSection();
+	MenuBuilder.BeginSection("FrameTrack", LOCTEXT("Plot_Series_FrameTrack_Section", "Frame Track"));
+
 	// Add/remove game frame stats series to/from frame track
 	{
 		FUIAction Action_ToggleFrameStatsTimerInFrameTrack;
@@ -951,6 +956,8 @@ void STimersView::TreeView_BuildPlotTimerMenu(FMenuBuilder& MenuBuilder)
 			);
 		}
 	}
+
+	MenuBuilder.EndSection();
 }
 
 ////////////////////////////////////////////////////////////////////////////////////////////////////
@@ -2639,11 +2646,13 @@ void STimersView::ToggleFrameTrackSeries(FTimerNodePtr TimerNode, ETraceFrameTyp
 	if (FrameTrack->HasFrameStatSeries(FrameType, TimerNode->GetTimerId()))
 	{
 		FrameTrack->RemoveTimerFrameStatSeries(FrameType, TimerNode->GetTimerId());
+		TimerNode->OnRemovedFromGraph();
 	}
 	else
 	{
-		TSharedPtr<FTimerFrameStatsTrackSeries> Series = FrameTrack->AddTimerFrameStatSeries(FrameType, TimerNode->GetTimerId());
-		Series->TimerDisplayName = TimerNode->GetDisplayName();
+		FText SeriesName = FText::Format(LOCTEXT("FrameTrackSeriesName_Format", "{0} {1} {2}"), TimerNode->GetDisplayName(), FText::FromString(FFrameTrackDrawHelper::FrameTypeToString(FrameType)), LOCTEXT("Frame", "Frame"));
+		TSharedPtr<FTimerFrameStatsTrackSeries> Series = FrameTrack->AddTimerFrameStatSeries(FrameType, TimerNode->GetTimerId(), TimerNode->GetColor(), SeriesName);
+		TimerNode->OnAddedToGraph();
 	}
 }
 
