[Insights] Made TimingView.AutoScroll option persistent in ini file.

diff --git a/Engine/Source/Developer/TraceInsights/Private/Insights/InsightsSettings.h b/Engine/Source/Developer/TraceInsights/Private/Insights/InsightsSettings.h
--- a/Engine/Source/Developer/TraceInsights/Private/Insights/InsightsSettings.h
+++ b/Engine/Source/Developer/TraceInsights/Private/Insights/InsightsSettings.h
@@ -68,6 +68,7 @@ public:
 		GConfig->GetBool(TEXT("Insights.TimingProfiler"), TEXT("bAutoZoomOnFrameSelection"), bAutoZoomOnFrameSelection, SettingsIni);
 
 		// Auto-scroll options
+		GConfig->GetBool(TEXT("Insights.AutoScroll"), TEXT("bAutoScroll"), bAutoScroll, SettingsIni);
 		FString FrameAlignment;
 		if (GConfig->GetString(TEXT("Insights.TimingProfiler"), TEXT("AutoScrollFrameAlignment"), FrameAlignment, SettingsIni))
 		{
@@ -112,6 +113,7 @@ public:
 		GConfig->SetBool(TEXT("Insights.TimingProfiler"), TEXT("bAutoZoomOnFrameSelection"), bAutoZoomOnFrameSelection, SettingsIni);
 
 		// Auto-scroll options
+		GConfig->SetBool(TEXT("Insights.AutoScroll"), TEXT("bAutoScroll"), bAutoScroll, SettingsIni);
 		const TCHAR* FrameAlignment = (AutoScrollFrameAlignment == 0) ? TEXT("game") : (AutoScrollFrameAlignment == 1) ? TEXT("rendering") : TEXT("none");
 		GConfig->SetString(TEXT("Insights.TimingProfiler"), TEXT("AutoScrollFrameAlignment"), FrameAlignment, SettingsIni);
 		GConfig->SetDouble(TEXT("Insights.TimingProfiler"), TEXT("AutoScrollViewportOffsetPercent"), AutoScrollViewportOffsetPercent, SettingsIni);
@@ -182,6 +184,10 @@ public:
 	void SetAutoZoomOnFrameSelection(bool bOnOff) { bAutoZoomOnFrameSelection = bOnOff; }
 	void SetAndSaveAutoZoomOnFrameSelection(bool bOnOff) { SET_AND_SAVE(bAutoZoomOnFrameSelection, bOnOff); }
 
+	bool IsAutoScrollEnabled() const { return bAutoScroll; }
+	void SetAutoScroll(bool bOnOff) { bAutoScroll = bOnOff; }
+	void SetAndSaveAutoScroll(bool bOnOff) { SET_AND_SAVE(bAutoScroll, bOnOff); }
+
 	int32 GetAutoScrollFrameAlignment() const { return AutoScrollFrameAlignment; }
 	void SetAutoScrollFrameAlignment(int32 FrameType) { AutoScrollFrameAlignment = FrameType; }
 	void SetAndSaveAutoScrollFrameAlignment(int32 FrameType) { SET_AND_SAVE(AutoScrollFrameAlignment, FrameType); }
@@ -266,6 +272,9 @@ private:
 	/** List of search paths to look for symbol files */
 	TArray<FString> SymbolSearchPaths;
 
+	/** If enabled, the Timing View will start with auto-scroll enabled. */
+	bool bAutoScroll;
+
 	/**
 	 * Viewport offset while auto-scrolling, as percent of viewport width.
 	 * If positive, it offsets the viewport forward, allowing an empty space at the right side of the viewport (i.e. after end of session).
diff --git a/Engine/Source/Developer/TraceInsights/Private/Insights/Widgets/STimingView.cpp b/Engine/Source/Developer/TraceInsights/Private/Insights/Widgets/STimingView.cpp
--- a/Engine/Source/Developer/TraceInsights/Private/Insights/Widgets/STimingView.cpp
+++ b/Engine/Source/Developer/TraceInsights/Private/Insights/Widgets/STimingView.cpp
@@ -207,7 +207,7 @@ void STimingView::Construct(const FArguments& InArgs, FName InViewName)
 	});
 	AutoScrollToggleButtonAction.ExecuteAction.BindLambda([this]
 	{
-		bAutoScroll = !bAutoScroll;
+		SetAutoScroll(!bAutoScroll);
 		Viewport.AddDirtyFlags(ETimingTrackViewportDirtyFlags::HInvalidated);
 	});
 
@@ -415,7 +415,7 @@ void STimingView::Reset(bool bIsFirstReset)
 	bIsSpaceBarKeyPressed = false;
 	bIsDragging = false;
 
-	bAutoScroll = false;
+	bAutoScroll = Settings.IsAutoScrollEnabled();
 	AutoScrollFrameAlignment = Settings.GetAutoScrollFrameAlignment();
 	AutoScrollViewportOffsetPercent = Settings.GetAutoScrollViewportOffsetPercent();
 	AutoScrollMinDelay = Settings.GetAutoScrollMinDelay();
@@ -608,7 +608,7 @@ void STimingView::Tick(const FGeometry& AllottedGeometry, const double InCurrent
 	if (bIsPanning)
 	{
 		// Disable auto-scroll if user starts panning manually.
-		bAutoScroll = false;
+		SetAutoScroll(false);
 	}
 
 	if (bAutoScroll)
@@ -3182,9 +3182,20 @@ void STimingView::BindCommands()
 
 ////////////////////////////////////////////////////////////////////////////////////////////////////
 
+void STimingView::SetAutoScroll(bool bOnOff)
+{
+	bAutoScroll = bOnOff;
+
+	// Persistent option. Save it to the config file.
+	FInsightsSettings& Settings = FInsightsManager::Get()->GetSettings();
+	Settings.SetAndSaveAutoScroll(bAutoScroll);
+}
+
+////////////////////////////////////////////////////////////////////////////////////////////////////
+
 void STimingView::AutoScroll_OnCheckStateChanged(ECheckBoxState NewRadioState)
 {
-	bAutoScroll = (NewRadioState == ECheckBoxState::Checked);
+	SetAutoScroll(NewRadioState == ECheckBoxState::Checked);
 	Viewport.AddDirtyFlags(ETimingTrackViewportDirtyFlags::HInvalidated);
 }
 
@@ -3303,7 +3314,7 @@ float STimingView::EnforceVerticalScrollLimits(const float InScrollPosY)
 void STimingView::HorizontalScrollBar_OnUserScrolled(float ScrollOffset)
 {
 	// Disable auto-scroll if user starts scrolling with horizontal scrollbar.
-	bAutoScroll = false;
+	SetAutoScroll(false);
 
 	Viewport.OnUserScrolled(HorizontalScrollBar, ScrollOffset);
 }
diff --git a/Engine/Source/Developer/TraceInsights/Private/Insights/Widgets/STimingView.h b/Engine/Source/Developer/TraceInsights/Private/Insights/Widgets/STimingView.h
--- a/Engine/Source/Developer/TraceInsights/Private/Insights/Widgets/STimingView.h
+++ b/Engine/Source/Developer/TraceInsights/Private/Insights/Widgets/STimingView.h
@@ -431,6 +431,8 @@ protected:
 	////////////////////////////////////////////////////////////////////////////////////////////////////
 	// Auto-Scroll
 
+	void SetAutoScroll(bool bOnOff);
+
 	void AutoScroll_OnCheckStateChanged(ECheckBoxState NewRadioState);
 	ECheckBoxState AutoScroll_IsChecked() const;
 
