[Insights] Timer aggregations are now canceled when a new selection is made.
 Fix timer aggregation bug when Insights is started and the mode is not instance .

diff --git a/Engine/Source/Developer/TraceInsights/Private/Insights/ViewModels/CounterAggregation.cpp b/Engine/Source/Developer/TraceInsights/Private/Insights/ViewModels/CounterAggregation.cpp
--- a/Engine/Source/Developer/TraceInsights/Private/Insights/ViewModels/CounterAggregation.cpp
+++ b/Engine/Source/Developer/TraceInsights/Private/Insights/ViewModels/CounterAggregation.cpp
@@ -353,7 +353,7 @@ public:
 
 	virtual ~FCounterAggregationWorker() {}
 
-	virtual void DoWork() override;
+	virtual void DoWork(TSharedPtr<TraceServices::FCancellationToken> CancellationToken) override;
 
 	void ApplyResultsTo(const TMap<uint32, FStatsNodePtr>& StatsNodesIdMap) const;
 	void ResetResults();
@@ -369,7 +369,7 @@ private:
 
 ////////////////////////////////////////////////////////////////////////////////////////////////////
 
-void FCounterAggregationWorker::DoWork()
+void FCounterAggregationWorker::DoWork(TSharedPtr<TraceServices::FCancellationToken> CancellationToken)
 {
 	CalculationHelperDbl.SetTimeInterval(StartTime, EndTime);
 	CalculationHelperInt.SetTimeInterval(StartTime, EndTime);
diff --git a/Engine/Source/Developer/TraceInsights/Private/Insights/ViewModels/StatsAggregator.cpp b/Engine/Source/Developer/TraceInsights/Private/Insights/ViewModels/StatsAggregator.cpp
--- a/Engine/Source/Developer/TraceInsights/Private/Insights/ViewModels/StatsAggregator.cpp
+++ b/Engine/Source/Developer/TraceInsights/Private/Insights/ViewModels/StatsAggregator.cpp
@@ -16,8 +16,9 @@ namespace Insights
 class FStatsAggregationTask : public FNonAbandonableTask
 {
 public:
-	FStatsAggregationTask(IStatsAggregationWorker* InWorker)
-		: Worker(InWorker)
+	FStatsAggregationTask(IStatsAggregationWorker* InWorker, TSharedPtr<TraceServices::FCancellationToken> InCancellationToken)
+		: CancellationToken(InCancellationToken)
+		, Worker(InWorker)
 	{
 		check(Worker != nullptr);
 	}
@@ -26,7 +27,7 @@ public:
 
 	IStatsAggregationWorker* GetWorker() { return Worker; }
 
-	void DoWork() { Worker->DoWork(); }
+	void DoWork() { Worker->DoWork(CancellationToken); }
 
 	FORCEINLINE TStatId GetStatId() const
 	{
@@ -34,6 +35,7 @@ public:
 	}
 
 private:
+	TSharedPtr<TraceServices::FCancellationToken> CancellationToken;
 	IStatsAggregationWorker* Worker;
 };
 
@@ -46,13 +48,13 @@ FStatsAggregator::FStatsAggregator(const FString InLogName)
 	, IntervalStartTime(0.0)
 	, IntervalEndTime(0.0)
 	, AsyncTask(nullptr)
-	, bIsCancelRequested(false)
 	, bIsStartRequested(false)
 	, bIsFinished(false)
 	, AllOpsStopwatch()
 	, CurrentOpStopwatch()
 	, OperationCount(0)
 {
+	CancellationToken = MakeShared<TraceServices::FCancellationToken>();
 }
 
 ////////////////////////////////////////////////////////////////////////////////////////////////////
@@ -71,7 +73,7 @@ void FStatsAggregator::ResetAsyncTask()
 	{
 		if (!AsyncTask->Cancel())
 		{
-			bIsCancelRequested = true;
+			CancellationToken->Cancel();
 			AsyncTask->EnsureCompletion();
 		}
 
@@ -105,7 +107,7 @@ void FStatsAggregator::Cancel()
 	if (AsyncTask)
 	{
 		UE_LOG(TimingProfiler, Log, TEXT("[%s] Cancel requested for async aggregation (op %d)..."), *LogName, OperationCount);
-		bIsCancelRequested = true;
+		CancellationToken->Cancel();
 	}
 }
 
@@ -120,7 +122,7 @@ void FStatsAggregator::Tick(TSharedPtr<const TraceServices::IAnalysisSession> In
 	{
 		double FinishedDuration = 0.0;
 
-		if (!bIsStartRequested && !bIsCancelRequested)
+		if (!bIsStartRequested && !CancellationToken->ShouldCancel())
 		{
 			FStopwatch FinishedStopwatch;
 			FinishedStopwatch.Start();
@@ -147,7 +149,7 @@ void FStatsAggregator::Tick(TSharedPtr<const TraceServices::IAnalysisSession> In
 		UE_LOG(TimingProfiler, Log, TEXT("[%s] Aggregated stats computed in %.4fs (%.4fs + %.4fs)%s%s - total: %.4fs (%d ops)"),
 			*LogName,
 			CurrentOpDuration, CurrentOpDuration - FinishedDuration, FinishedDuration,
-			bIsCancelRequested ? TEXT(" - CANCELED") : TEXT(""),
+			CancellationToken->ShouldCancel() ? TEXT(" - CANCELED") : TEXT(""),
 			bIsStartRequested ? TEXT(" - IGNORED") : TEXT(""),
 			AllOpsDuration, OperationCount);
 	}
@@ -156,11 +158,11 @@ void FStatsAggregator::Tick(TSharedPtr<const TraceServices::IAnalysisSession> In
 	{
 		if (AsyncTask)
 		{
-			if (!bIsCancelRequested)
+			if (!CancellationToken->ShouldCancel())
 			{
 				// Cancel and wait for the previous async task to finish.
 				UE_LOG(TimingProfiler, Log, TEXT("[%s] Cancel previous async aggregation (op %d)..."), *LogName, OperationCount);
-				bIsCancelRequested = true;
+				CancellationToken->Cancel();
 			}
 		}
 		else
@@ -169,14 +171,14 @@ void FStatsAggregator::Tick(TSharedPtr<const TraceServices::IAnalysisSession> In
 			++OperationCount;
 
 			bIsStartRequested = false;
-			bIsCancelRequested = false;
+			CancellationToken->Reset();
 
 			UE_LOG(TimingProfiler, Log, TEXT("[%s] Start async aggregation (op %d)..."), *LogName, OperationCount);
 
 			IStatsAggregationWorker* Worker = CreateWorker(InSession);
 			check(Worker);
 
-			AsyncTask = new FStatsAggregationAsyncTask(Worker);
+			AsyncTask = new FStatsAggregationAsyncTask(Worker, CancellationToken);
 			AsyncTask->StartBackgroundTask();
 		}
 	}
diff --git a/Engine/Source/Developer/TraceInsights/Private/Insights/ViewModels/StatsAggregator.h b/Engine/Source/Developer/TraceInsights/Private/Insights/ViewModels/StatsAggregator.h
--- a/Engine/Source/Developer/TraceInsights/Private/Insights/ViewModels/StatsAggregator.h
+++ b/Engine/Source/Developer/TraceInsights/Private/Insights/ViewModels/StatsAggregator.h
@@ -2,6 +2,8 @@
 
 #pragma once
 
+#include "TraceServices/Common/CancellationToken.h"
+
 #include "CoreMinimal.h"
 #include "Async/AsyncWork.h"
 #include "Insights/Common/InsightsAsyncWorkUtils.h"
@@ -34,7 +36,7 @@ class IStatsAggregationWorker
 {
 public:
 	virtual ~IStatsAggregationWorker() {}
-	virtual void DoWork() = 0;
+	virtual void DoWork(TSharedPtr<TraceServices::FCancellationToken> CancellationToken) = 0;
 };
 
 ////////////////////////////////////////////////////////////////////////////////////////////////////
@@ -68,7 +70,7 @@ public:
 	virtual void Start() override;
 	virtual void Cancel() override;
 
-	virtual bool IsCancelRequested() const override { return bIsCancelRequested; }
+	virtual bool IsCancelRequested() const override { return CancellationToken->ShouldCancel(); }
 	virtual bool IsRunning() const override { return AsyncTask != nullptr; }
 
 	virtual double GetAllOperationsDuration() override { AllOpsStopwatch.Update(); return AllOpsStopwatch.GetAccumulatedTime(); }
@@ -99,7 +101,7 @@ private:
 
 	FStatsAggregationAsyncTask* AsyncTask;
 
-	mutable volatile bool bIsCancelRequested; // true if we want the async task to finish asap
+	TSharedPtr<TraceServices::FCancellationToken> CancellationToken;
 	bool bIsStartRequested;
 	bool bIsFinished;
 
diff --git a/Engine/Source/Developer/TraceInsights/Private/Insights/ViewModels/TimerAggregation.cpp b/Engine/Source/Developer/TraceInsights/Private/Insights/ViewModels/TimerAggregation.cpp
--- a/Engine/Source/Developer/TraceInsights/Private/Insights/ViewModels/TimerAggregation.cpp
+++ b/Engine/Source/Developer/TraceInsights/Private/Insights/ViewModels/TimerAggregation.cpp
@@ -30,7 +30,7 @@ public:
 
 	virtual ~FTimerAggregationWorker() {}
 
-	virtual void DoWork() override;
+	virtual void DoWork(TSharedPtr<TraceServices::FCancellationToken> CancellationToken) override;
 
 	TraceServices::ITable<TraceServices::FTimingProfilerAggregatedStats>* GetResultTable() const { return ResultTable.Get(); }
 	void ResetResults() { ResultTable.Reset(); }
@@ -47,7 +47,7 @@ private:
 
 ////////////////////////////////////////////////////////////////////////////////////////////////////
 
-void FTimerAggregationWorker::DoWork()
+void FTimerAggregationWorker::DoWork(TSharedPtr<TraceServices::FCancellationToken> CancellationToken)
 {
 	if (Session.IsValid() && TraceServices::ReadTimingProfilerProvider(*Session.Get()))
 	{
@@ -62,7 +62,15 @@ void FTimerAggregationWorker::DoWork()
 			return CpuThreads.Contains(ThreadId);
 		};
 
-		ResultTable.Reset(TimingProfilerProvider.CreateAggregation(StartTime, EndTime, CpuThreadFilter, bIncludeGpuThread, FrameType));
+		TraceServices::FCreateAggreationParams Params;
+		Params.IntervalStart = StartTime;
+		Params.IntervalEnd = EndTime;
+		Params.CpuThreadFilter = CpuThreadFilter;
+		Params.IncludeGpu = bIncludeGpuThread;
+		Params.FrameType = FrameType;
+		Params.CancellationToken = CancellationToken;
+
+		ResultTable.Reset(TimingProfilerProvider.CreateAggregation(Params));
 	}
 }
 
diff --git a/Engine/Source/Developer/TraceInsights/Private/Insights/ViewModels/TimerButterflyAggregation.cpp b/Engine/Source/Developer/TraceInsights/Private/Insights/ViewModels/TimerButterflyAggregation.cpp
--- a/Engine/Source/Developer/TraceInsights/Private/Insights/ViewModels/TimerButterflyAggregation.cpp
+++ b/Engine/Source/Developer/TraceInsights/Private/Insights/ViewModels/TimerButterflyAggregation.cpp
@@ -29,7 +29,7 @@ public:
 
 	virtual ~FTimerButterflyAggregationWorker() {}
 
-	virtual void DoWork() override;
+	virtual void DoWork(TSharedPtr<TraceServices::FCancellationToken> CancellationToken) override;
 
 	TraceServices::ITimingProfilerButterfly* GetResultButterfly() const { return ResultButterfly.Get(); }
 	void ResetResults() { ResultButterfly.Reset(); }
@@ -45,7 +45,7 @@ private:
 
 ////////////////////////////////////////////////////////////////////////////////////////////////////
 
-void FTimerButterflyAggregationWorker::DoWork()
+void FTimerButterflyAggregationWorker::DoWork(TSharedPtr<TraceServices::FCancellationToken> CancellationToken)
 {
 	if (Session.IsValid() && TraceServices::ReadTimingProfilerProvider(*Session.Get()))
 	{
diff --git a/Engine/Source/Developer/TraceInsights/Private/Insights/ViewModels/TimingExporter.cpp b/Engine/Source/Developer/TraceInsights/Private/Insights/ViewModels/TimingExporter.cpp
--- a/Engine/Source/Developer/TraceInsights/Private/Insights/ViewModels/TimingExporter.cpp
+++ b/Engine/Source/Developer/TraceInsights/Private/Insights/ViewModels/TimingExporter.cpp
@@ -660,8 +660,14 @@ int32 FTimingExporter::ExportTimerStatisticsAsText(const FString& Filename, FExp
 			return -1;
 		}
 
+		TraceServices::FCreateAggreationParams CreateAggreagationParams;
+		CreateAggreagationParams.IntervalStart = Params.IntervalStartTime;
+		CreateAggreagationParams.IntervalEnd = Params.IntervalEndTime;
+		CreateAggreagationParams.CpuThreadFilter = Params.ThreadFilter;
+		CreateAggreagationParams.IncludeGpu = true;
+
 		//@Todo: this does not yet handle the -column and -timers parameters.
-		StatsTable = TimingProfilerProvider->CreateAggregation(Params.IntervalStartTime, Params.IntervalEndTime, Params.ThreadFilter, true);
+		StatsTable = TimingProfilerProvider->CreateAggregation(CreateAggreagationParams);
 	}
 
 	FStopwatch Stopwatch;
diff --git a/Engine/Source/Developer/TraceInsights/Private/Insights/Widgets/STimersView.cpp b/Engine/Source/Developer/TraceInsights/Private/Insights/Widgets/STimersView.cpp
--- a/Engine/Source/Developer/TraceInsights/Private/Insights/Widgets/STimersView.cpp
+++ b/Engine/Source/Developer/TraceInsights/Private/Insights/Widgets/STimersView.cpp
@@ -1154,6 +1154,7 @@ void STimersView::InsightsManager_OnSessionAnalysisCompleted()
 		constexpr double Delta = 0.0; // session padding
 
 		Aggregator->Cancel();
+		Aggregator->SetFrameType(ModeFrameType);
 		Aggregator->SetTimeInterval(0.0 - Delta, SessionDuration + Delta);
 		Aggregator->Start();
 
diff --git a/Engine/Source/Developer/TraceServices/Private/Common/TimelineStatistics.h b/Engine/Source/Developer/TraceServices/Private/Common/TimelineStatistics.h
--- a/Engine/Source/Developer/TraceServices/Private/Common/TimelineStatistics.h
+++ b/Engine/Source/Developer/TraceServices/Private/Common/TimelineStatistics.h
@@ -5,6 +5,8 @@
 #include "Containers/Map.h"
 #include "CoreTypes.h"
 #include <cfloat>
+#include "Templates/SharedPointer.h"
+#include "TraceServices/Common/CancellationToken.h"
 #include "TraceServices/Containers/Timelines.h"
 
 namespace TraceServices
@@ -35,12 +37,21 @@ class FTimelineStatistics
 {
 public:
 	template<typename TimelineType, typename BucketMappingFunc, typename BucketKeyType>
-	static void CreateAggregation(const TArray<const TimelineType*>& Timelines, BucketMappingFunc BucketMapper, double IntervalStart, double IntervalEnd, TMap<BucketKeyType, FAggregatedTimingStats>& Result)
+	static void CreateAggregation(const TArray<const TimelineType*>& Timelines,
+								  BucketMappingFunc BucketMapper,
+								  double IntervalStart,
+								  double IntervalEnd,
+								  TSharedPtr<FCancellationToken> CancellationToken,
+								  TMap<BucketKeyType, FAggregatedTimingStats>& Result)
 	{
 		TMap<BucketKeyType, FInternalAggregationEntry> InternalResult;
 		// Compute instance count and total/min/max inclusive/exclusive times for each timer.
 		for (const TimelineType* Timeline : Timelines)
 		{
+			if (CancellationToken.IsValid() && CancellationToken->ShouldCancel())
+			{
+				return;
+			}
 			ProcessTimeline(Timeline, BucketMapper, UpdateTotalMinMaxTimerStats, IntervalStart, IntervalEnd, InternalResult);
 		}
 		// Now, as we know min/max inclusive/exclusive times for each timer, we can compute histogram and median values.
@@ -56,6 +67,11 @@ public:
 			// Compute histogram.
 			for (const TimelineType* Timeline : Timelines)
 			{
+				if (CancellationToken.IsValid() && CancellationToken->ShouldCancel())
+				{
+					return;
+				}
+
 				ProcessTimeline(Timeline, BucketMapper, UpdateHistogramForTimerStats, IntervalStart, IntervalEnd, InternalResult);
 			}
 		}
@@ -69,7 +85,11 @@ public:
 	};
 
 	template<typename TimelineType, typename BucketMappingFunc, typename BucketKeyType>
-	static void CreateFrameStatsAggregation(const TArray<const TimelineType*>& Timelines, BucketMappingFunc BucketMapper, TArray<FFrameData> Frames, TMap<BucketKeyType, FAggregatedTimingStats>& Result)
+	static void CreateFrameStatsAggregation(const TArray<const TimelineType*>& Timelines,
+											BucketMappingFunc BucketMapper,
+											TArray<FFrameData> Frames,
+											TSharedPtr<FCancellationToken> CancellationToken,
+											TMap<BucketKeyType, FAggregatedTimingStats>& Result)
 	{
 		TMap<BucketKeyType, FInternalAggregationEntry> FrameResult;
 		TMap<BucketKeyType, FInternalFrameAggregationEntry> GlobalResult;
@@ -100,9 +120,13 @@ public:
 			GlobalResult.Add(KV.Key, Entry);
 		}
 
-
 		for (int32 Index = 0; Index < FramesNum; ++Index)
 		{
+			if (CancellationToken.IsValid() && CancellationToken->ShouldCancel())
+			{
+				return;
+			}
+
 			// Compute instance count and total/min/max inclusive/exclusive times for each timer.
 			for (const TimelineType* Timeline : Timelines)
 			{
diff --git a/Engine/Source/Developer/TraceServices/Private/Model/LoadTimeProfiler.cpp b/Engine/Source/Developer/TraceServices/Private/Model/LoadTimeProfiler.cpp
--- a/Engine/Source/Developer/TraceServices/Private/Model/LoadTimeProfiler.cpp
+++ b/Engine/Source/Developer/TraceServices/Private/Model/LoadTimeProfiler.cpp
@@ -4,6 +4,7 @@
 #include "Model/LoadTimeProfilerPrivate.h"
 #include "AnalysisServicePrivate.h"
 #include "Common/TimelineStatistics.h"
+#include "TraceServices/Common/CancellationToken.h"
 
 namespace TraceServices
 {
@@ -262,7 +263,11 @@ ITable<FLoadTimeProfilerAggregatedStats>* FLoadTimeProfilerProvider::CreateEvent
 		return Event.Package;
 	};
 	TMap<const FPackageInfo*, FAggregatedTimingStats> Aggregation;
-	FTimelineStatistics::CreateAggregation(Timelines, BucketMapper, IntervalStart, IntervalEnd, Aggregation);
+
+	// Cancellation not currently implemented for LoadTimeProfilerProvider.
+	TSharedPtr<TraceServices::FCancellationToken> bIsCancelRequested = MakeShared<TraceServices::FCancellationToken>();
+
+	FTimelineStatistics::CreateAggregation(Timelines, BucketMapper, IntervalStart, IntervalEnd, bIsCancelRequested, Aggregation);
 	TTable<FLoadTimeProfilerAggregatedStats>* Table = new TTable<FLoadTimeProfilerAggregatedStats>(AggregatedStatsTableLayout);
 	for (const auto& KV : Aggregation)
 	{
@@ -297,7 +302,11 @@ ITable<FLoadTimeProfilerAggregatedStats>* FLoadTimeProfilerProvider::CreateObjec
 		return Event.Export ? Event.Export->Class : nullptr;
 	};
 	TMap<const FClassInfo*, FAggregatedTimingStats> Aggregation;
-	FTimelineStatistics::CreateAggregation(Timelines, BucketMapper, IntervalStart, IntervalEnd, Aggregation);
+
+	// Cancellation not currently implemented for LoadTimeProfilerProvider.
+	TSharedPtr<TraceServices::FCancellationToken> CancellationToken = MakeShared<TraceServices::FCancellationToken>();
+
+	FTimelineStatistics::CreateAggregation(Timelines, BucketMapper, IntervalStart, IntervalEnd, CancellationToken, Aggregation);
 	TTable<FLoadTimeProfilerAggregatedStats>* Table = new TTable<FLoadTimeProfilerAggregatedStats>(AggregatedStatsTableLayout);
 	for (const auto& KV : Aggregation)
 	{
diff --git a/Engine/Source/Developer/TraceServices/Private/Model/TimingProfiler.cpp b/Engine/Source/Developer/TraceServices/Private/Model/TimingProfiler.cpp
--- a/Engine/Source/Developer/TraceServices/Private/Model/TimingProfiler.cpp
+++ b/Engine/Source/Developer/TraceServices/Private/Model/TimingProfiler.cpp
@@ -265,19 +265,19 @@ void FTimingProfilerProvider::ReadTimers(TFunctionRef<void(const ITimingProfiler
 	Callback(*this);
 }
 
-ITable<FTimingProfilerAggregatedStats>* FTimingProfilerProvider::CreateAggregation(double IntervalStart, double IntervalEnd, TFunctionRef<bool(uint32)> CpuThreadFilter, bool IncludeGpu, ETraceFrameType FrameType) const
+ITable<FTimingProfilerAggregatedStats>* FTimingProfilerProvider::CreateAggregation(const FCreateAggreationParams& Params) const
 {
 	Session.ReadAccessCheck();
 
 	TArray<const TimelineInternal*> IncludedTimelines;
-	if (IncludeGpu)
+	if (Params.IncludeGpu)
 	{
 		IncludedTimelines.Add(&Timelines[GpuTimelineIndex].Get());
 		IncludedTimelines.Add(&Timelines[Gpu2TimelineIndex].Get());
 	}
 	for (const auto& KV : CpuThreadTimelineIndexMap)
 	{
-		if (CpuThreadFilter(KV.Key))
+		if (Params.CpuThreadFilter(KV.Key))
 		{
 			IncludedTimelines.Add(&Timelines[KV.Value].Get());
 		}
@@ -289,15 +289,15 @@ ITable<FTimingProfilerAggregatedStats>* FTimingProfilerProvider::CreateAggregati
 	};
 
 	TMap<const FTimingProfilerTimer*, FAggregatedTimingStats> Aggregation;
-	if (FrameType == ETraceFrameType::TraceFrameType_Count)
+	if (Params.FrameType == ETraceFrameType::TraceFrameType_Count)
 	{
-		FTimelineStatistics::CreateAggregation(IncludedTimelines, BucketMappingFunc, IntervalStart, IntervalEnd, Aggregation);
+		FTimelineStatistics::CreateAggregation(IncludedTimelines, BucketMappingFunc, Params.IntervalStart, Params.IntervalEnd, Params.CancellationToken, Aggregation);
 	}
 	else
 	{
 		TArray<FFrameData> Frames;
 		const IFrameProvider& FrameProvider = ReadFrameProvider(Session);
-		FrameProvider.EnumerateFrames(FrameType, IntervalStart, IntervalEnd, [&Frames](const FFrame& Frame)
+		FrameProvider.EnumerateFrames(Params.FrameType, Params.IntervalStart, Params.IntervalEnd, [&Frames](const FFrame& Frame)
 			{
 				FFrameData NewFrameData;
 				NewFrameData.StartTime = Frame.StartTime;
@@ -308,11 +308,17 @@ ITable<FTimingProfilerAggregatedStats>* FTimingProfilerProvider::CreateAggregati
 
 		if (Frames.Num() > 0)
 		{
-			FTimelineStatistics::CreateFrameStatsAggregation(IncludedTimelines, BucketMappingFunc, Frames, Aggregation);
+			FTimelineStatistics::CreateFrameStatsAggregation(IncludedTimelines, BucketMappingFunc, Frames, Params.CancellationToken, Aggregation);
 		}
 	}
 
 	TTable<FTimingProfilerAggregatedStats>* Table = new TTable<FTimingProfilerAggregatedStats>(AggregatedStatsTableLayout);
+	
+	if (Params.CancellationToken.IsValid() && Params.CancellationToken->ShouldCancel())
+	{
+		return Table;
+	}
+
 	for (const auto& KV : Aggregation)
 	{
 		FTimingProfilerAggregatedStats& Row = Table->AddRow();
diff --git a/Engine/Source/Developer/TraceServices/Private/Model/TimingProfilerPrivate.h b/Engine/Source/Developer/TraceServices/Private/Model/TimingProfilerPrivate.h
--- a/Engine/Source/Developer/TraceServices/Private/Model/TimingProfilerPrivate.h
+++ b/Engine/Source/Developer/TraceServices/Private/Model/TimingProfilerPrivate.h
@@ -35,7 +35,7 @@ public:
 	virtual uint32 GetTimelineCount() const override { return Timelines.Num(); }
 	virtual void EnumerateTimelines(TFunctionRef<void(const Timeline&)> Callback) const override;
 	virtual void ReadTimers(TFunctionRef<void(const ITimingProfilerTimerReader&)> Callback) const override;
-	virtual ITable<FTimingProfilerAggregatedStats>* CreateAggregation(double IntervalStart, double IntervalEnd, TFunctionRef<bool(uint32)> CpuThreadFilter, bool IncludeGpu, ETraceFrameType FrameType = ETraceFrameType::TraceFrameType_Count) const override;
+	virtual ITable<FTimingProfilerAggregatedStats>* CreateAggregation(const FCreateAggreationParams& Params) const override;
 	virtual ITimingProfilerButterfly* CreateButterfly(double IntervalStart, double IntervalEnd, TFunctionRef<bool(uint32)> CpuThreadFilter, bool IncludeGpu) const override;
 	virtual const FTimingProfilerTimer* GetTimer(uint32 TimerId) const override;
 	virtual uint32 GetTimerCount() const override;
diff --git a/Engine/Source/Developer/TraceServices/Public/TraceServices/Common/CancellationToken.h b/Engine/Source/Developer/TraceServices/Public/TraceServices/Common/CancellationToken.h
new file mode 100644
--- /dev/null
+++ b/Engine/Source/Developer/TraceServices/Public/TraceServices/Common/CancellationToken.h
@@ -0,0 +1,25 @@
+// Copyright Epic Games, Inc. All Rights Reserved.
+
+#pragma once
+
+#include <atomic>
+
+namespace TraceServices
+{
+
+class FCancellationToken
+{
+public:
+	FCancellationToken()
+		: bCancel(false)
+	{}
+
+	bool ShouldCancel() { return bCancel.load(); }
+	void Cancel() { bCancel.store(true); }
+	void Reset() { bCancel.store(false); }
+
+private:
+	std::atomic<bool> bCancel;
+};
+
+} //namespace TraceServices
diff --git a/Engine/Source/Developer/TraceServices/Public/TraceServices/Model/TimingProfiler.h b/Engine/Source/Developer/TraceServices/Public/TraceServices/Model/TimingProfiler.h
--- a/Engine/Source/Developer/TraceServices/Public/TraceServices/Model/TimingProfiler.h
+++ b/Engine/Source/Developer/TraceServices/Public/TraceServices/Model/TimingProfiler.h
@@ -8,6 +8,7 @@
 #include "HAL/Platform.h"
 #include "HAL/PlatformMath.h"
 #include "Model/MonotonicTimeline.h"
+#include "TraceServices/Common/CancellationToken.h"
 #include "TraceServices/Containers/Tables.h"
 #include "TraceServices/Containers/Timelines.h"
 #include "TraceServices/Model/AnalysisSession.h"
@@ -65,6 +66,16 @@ struct FTimingProfilerButterflyNode
 	TArray<FTimingProfilerButterflyNode*> Children;
 };
 
+struct FCreateAggreationParams
+{
+	double IntervalStart;
+	double IntervalEnd;
+	TFunction<bool(uint32)> CpuThreadFilter;
+	bool IncludeGpu;
+	ETraceFrameType FrameType = ETraceFrameType::TraceFrameType_Count;
+	TSharedPtr<TraceServices::FCancellationToken> CancellationToken;
+};
+
 class ITimingProfilerButterfly
 {
 public:
@@ -96,6 +107,13 @@ public:
 	virtual void EnumerateTimelines(TFunctionRef<void(const Timeline&)> Callback) const = 0;
 	virtual void ReadTimers(TFunctionRef<void(const ITimingProfilerTimerReader&)> Callback) const = 0;
 
+	/**
+	* Create a table of aggregated stats.
+	*
+	* @param Params				The params for the aggregation.
+	*/
+	virtual ITable<FTimingProfilerAggregatedStats>* CreateAggregation(const FCreateAggreationParams& Params) const = 0;
+
 	/**
 	* Create a table of aggregated stats.
 	*
@@ -105,7 +123,6 @@ public:
 	* @param IncludeGpu			A boolean value to specify if aggregaton should include GPU timelines.
 	* @param FrameType			The type of frame to use for frame stats aggregation. ETraceFrameType::TraceFrameType_Count means no frame aggregation.
 	*/
-	virtual ITable<FTimingProfilerAggregatedStats>* CreateAggregation(double IntervalStart, double IntervalEnd, TFunctionRef<bool(uint32)> CpuThreadFilter, bool IncludeGpu, ETraceFrameType FrameType = ETraceFrameType::TraceFrameType_Count) const = 0;
 	virtual ITimingProfilerButterfly* CreateButterfly(double IntervalStart, double IntervalEnd, TFunctionRef<bool(uint32)> CpuThreadFilter, bool IncludeGpu) const = 0;
 };
 
