[Insights] Task Table header UI improvements

diff --git a/Engine/Source/Developer/TraceInsights/Private/Insights/TaskGraphProfiler/Widgets/STaskTableTreeView.cpp b/Engine/Source/Developer/TraceInsights/Private/Insights/TaskGraphProfiler/Widgets/STaskTableTreeView.cpp
--- a/Engine/Source/Developer/TraceInsights/Private/Insights/TaskGraphProfiler/Widgets/STaskTableTreeView.cpp
+++ b/Engine/Source/Developer/TraceInsights/Private/Insights/TaskGraphProfiler/Widgets/STaskTableTreeView.cpp
@@ -229,20 +229,50 @@ FText STaskTableTreeView::GetCurrentOperationName() const
 ////////////////////////////////////////////////////////////////////////////////////////////////////
 
 BEGIN_SLATE_FUNCTION_BUILD_OPTIMIZATION
-TSharedPtr<SWidget> STaskTableTreeView::ConstructToolbar()
+void STaskTableTreeView::ConstructHeaderArea(TSharedRef<SVerticalBox> InWidgetContent)
 {
-	return SNew(SHorizontalBox)
+	InWidgetContent->AddSlot()
+	.VAlign(VAlign_Center)
+	.AutoHeight()
+	.Padding(2.0f)
+	[
+		SNew(SHorizontalBox)
+
+		+ SHorizontalBox::Slot()
+		.FillWidth(1.0f)
+		.VAlign(VAlign_Center)
+		[
+			ConstructSearchBox()
+		]
+
+		+ SHorizontalBox::Slot()
+		.AutoWidth()
+		.Padding(4.0f, 0.0f, 0.0f, 0.0f)
+		.VAlign(VAlign_Center)
+		[
+			ConstructAdvancedFiltersButton()
+		]
+	];
+
+	InWidgetContent->AddSlot()
+	.VAlign(VAlign_Center)
+	.AutoHeight()
+	.Padding(2.0f)
+	[
+		SNew(SHorizontalBox)
 		+ SHorizontalBox::Slot()
 		.AutoWidth()
+		.HAlign(HAlign_Left)
 		.Padding(0.0f, 0.0f, 4.0f, 0.0f)
 		.VAlign(VAlign_Center)
 		[
 			SNew(STextBlock)
 			.Text(LOCTEXT("Tasks", "Tasks"))
 		]
-		
+
 		+ SHorizontalBox::Slot()
 		.AutoWidth()
+		.HAlign(HAlign_Left)
 		.Padding(4.0f, 0.0f, 10.0f, 0.0f)
 		[
 			SNew(SBox)
@@ -259,34 +289,59 @@ TSharedPtr<SWidget> STaskTableTreeView::ConstructToolbar()
 				]
 			]
 		]
-		
+
 		+ SHorizontalBox::Slot()
-		.AutoWidth()
+		.FillWidth(1.0f)
 		.Padding(0.0f, 0.0f, 4.0f, 0.0f)
 		.VAlign(VAlign_Center)
+		.HAlign(HAlign_Right)
 		[
-			SNew(STextBlock)
-			.Text(LOCTEXT("Timestamps", "Timestamps"))
-		]
-		
-		+ SHorizontalBox::Slot()
-		.AutoWidth()
-		.Padding(4.0f, 0.0f, 0.0f, 0.0f)
-		[
-			SNew(SBox)
-			.MinDesiredWidth(160.0f)
+			SNew(SHorizontalBox)
+
+			+ SHorizontalBox::Slot()
+			.Padding(0.0f, 0.0f, 4.0f, 0.0f)
+			.VAlign(VAlign_Center)
+			[
+				SNew(STextBlock)
+				.Text(LOCTEXT("Timestamps", "Timestamps"))
+			]
+
+			+ SHorizontalBox::Slot()
+			.AutoWidth()
+			.HAlign(HAlign_Right)
+			.Padding(4.0f, 0.0f, 0.0f, 0.0f)
 			[
-				SNew(SComboBox<TSharedPtr<ETimestampOptions>>)
-				.OptionsSource(GetAvailableTimestampOptions())
-				.OnSelectionChanged(this, &STaskTableTreeView::TimestampOptions_OnSelectionChanged)
-				.OnGenerateWidget(this, &STaskTableTreeView::TimestampOptions_OnGenerateWidget)
-				.IsEnabled(this, &STaskTableTreeView::TimestampOptions_IsEnabled)
+				SNew(SBox)
+				.MinDesiredWidth(160.0f)
 				[
-					SNew(STextBlock)
-					.Text(this, &STaskTableTreeView::TimestampOptions_GetSelectionText)
+					SNew(SComboBox<TSharedPtr<ETimestampOptions>>)
+					.OptionsSource(GetAvailableTimestampOptions())
+					.OnSelectionChanged(this, &STaskTableTreeView::TimestampOptions_OnSelectionChanged)
+					.OnGenerateWidget(this, &STaskTableTreeView::TimestampOptions_OnGenerateWidget)
+					.IsEnabled(this, &STaskTableTreeView::TimestampOptions_IsEnabled)
+					[
+						SNew(STextBlock)
+						.Text(this, &STaskTableTreeView::TimestampOptions_GetSelectionText)
+					]
 				]
 			]
-		];
+		]
+	];
+
+	InWidgetContent->AddSlot()
+	.VAlign(VAlign_Center)
+	.AutoHeight()
+	.Padding(2.0f)
+	[
+		SNew(SHorizontalBox)
+
+		+ SHorizontalBox::Slot()
+		.FillWidth(1.0f)
+		.VAlign(VAlign_Center)
+		[
+			ConstructHierarchyBreadcrumbTrail()
+		]
+	];
 }
 
 ////////////////////////////////////////////////////////////////////////////////////////////////////
diff --git a/Engine/Source/Developer/TraceInsights/Private/Insights/TaskGraphProfiler/Widgets/STaskTableTreeView.h b/Engine/Source/Developer/TraceInsights/Private/Insights/TaskGraphProfiler/Widgets/STaskTableTreeView.h
--- a/Engine/Source/Developer/TraceInsights/Private/Insights/TaskGraphProfiler/Widgets/STaskTableTreeView.h
+++ b/Engine/Source/Developer/TraceInsights/Private/Insights/TaskGraphProfiler/Widgets/STaskTableTreeView.h
@@ -49,7 +49,8 @@ public:
 	 */
 	void Construct(const FArguments& InArgs, TSharedPtr<FTaskTable> InTablePtr);
 
-	virtual TSharedPtr<SWidget> ConstructToolbar() override;
+	void ConstructHeaderArea(TSharedRef<SVerticalBox> InWidgetContent) override;
+
 	virtual TSharedPtr<SWidget> ConstructFooter() override;
 
 	TSharedPtr<FTaskTable> GetTaskTable() { return StaticCastSharedPtr<FTaskTable>(GetTable()); }
