[Insights] Tree Table View improvements:
  * Fixed grouping by CString values to also merge equal string values (when pointer if different).
  * Added default Icon and Color for a tree node.
  * Increased font size and updated text colors (ex.: when tree node is hoverd/selected/filtered).
  * Added FCustomTableTreeNode that can customize icon and/or color tint for a tree node.
  * Added custom color for group nodes created by Callstack and Tags groupings.

diff --git a/Engine/Source/Developer/TraceInsights/Private/Insights/InsightsStyle.cpp b/Engine/Source/Developer/TraceInsights/Private/Insights/InsightsStyle.cpp
--- a/Engine/Source/Developer/TraceInsights/Private/Insights/InsightsStyle.cpp
+++ b/Engine/Source/Developer/TraceInsights/Private/Insights/InsightsStyle.cpp
@@ -139,7 +139,6 @@ void FInsightsStyle::FStyle::Initialize()
 	Set("Border.R", new CORE_BOX_BRUSH("Icons/Profiler/Profiler_Border_R_16x", FMargin(4.0f / 16.0f)));
 
 	Set("Graph.Point", new EDITOR_IMAGE_BRUSH("Old/Graph/ExecutionBubble", Icon16x16));
-	Set("TreeViewBanner.WarningIcon", new CORE_IMAGE_BRUSH_SVG("Starship/Common/alert-circle", Icon20x20, FStyleColors::Warning));
 
 	//////////////////////////////////////////////////
 	// Icons for major components
@@ -293,8 +292,8 @@ void FInsightsStyle::FStyle::Initialize()
 	//Set("Icons.FolderClosed", new CORE_IMAGE_BRUSH_SVG("Starship/Common/folder-closed", Icon16x16));	//-> use FAppStyle "Icons.FolderClosed"
 
 	Set("Icons.SortBy", new CORE_IMAGE_BRUSH("Icons/Profiler/profiler_SortBy_32x", Icon16x16));
-	//Set("Icons.SortAscending", new CORE_IMAGE_BRUSH_SVG("Starship/Common/SortUp", Icon16x16));	//-> use FAppStyle "Icons.SortUp"
-	//Set("Icons.SortDescending", new CORE_IMAGE_BRUSH_SVG("Starship/Common/SortDown", Icon16x16));	//-> use FAppStyle "Icons.SortDown"
+	//Set("Icons.SortUp", new CORE_IMAGE_BRUSH_SVG("Starship/Common/SortUp", Icon16x16));		//-> use FAppStyle "Icons.SortUp"
+	//Set("Icons.SortDown", new CORE_IMAGE_BRUSH_SVG("Starship/Common/SortDown", Icon16x16));	//-> use FAppStyle "Icons.SortDown"
 
 	Set("Icons.ViewColumn", new CORE_IMAGE_BRUSH("Icons/Profiler/profiler_ViewColumn_32x", Icon16x16));
 	Set("Icons.ResetColumn", new CORE_IMAGE_BRUSH("Icons/Profiler/profiler_ResetColumn_32x", Icon16x16));
@@ -330,11 +329,12 @@ void FInsightsStyle::FStyle::Initialize()
 	//////////////////////////////////////////////////
 
 	Set("TreeTable.RowBackground", new EDITOR_IMAGE_BRUSH("Old/White", Icon16x16, FLinearColor(1.0f, 1.0f, 1.0f, 0.25f)));
+	Set("TreeViewBanner.WarningIcon", new CORE_IMAGE_BRUSH_SVG("Starship/Common/alert-circle", Icon20x20, FStyleColors::Warning));
 
 	Set("Icons.Hint.TreeItem", new IMAGE_BRUSH_SVG("InfoTag_12", Icon12x12));
 	Set("Icons.HotPath.TreeItem", new IMAGE_BRUSH_SVG("HotPath_12", Icon12x12));
 	Set("Icons.Group.TreeItem", new CORE_IMAGE_BRUSH_SVG("Starship/Common/folder-closed", Icon12x12));
-	Set("Icons.Leaf.TreeItem", new CORE_IMAGE_BRUSH_SVG("Starship/Common/file", Icon12x12));
+	Set("Icons.Leaf.TreeItem", new CORE_IMAGE_BRUSH_SVG("Starship/Common/bullet-point", Icon12x12));
 	Set("Icons.GpuTimer.TreeItem", new TODO_IMAGE_BRUSH(Icon12x12));
 	Set("Icons.CpuTimer.TreeItem", new TODO_IMAGE_BRUSH(Icon12x12));
 	Set("Icons.Counter.TreeItem", new TODO_IMAGE_BRUSH(Icon12x12));
@@ -358,6 +358,14 @@ void FInsightsStyle::FStyle::Initialize()
 		.SetShadowColorAndOpacity(FLinearColor(0.0f, 0.0f, 0.0f, 0.8f))
 	);
 
+	Set("TreeTable.NameText", FTextBlockStyle(NormalText)
+		.SetColorAndOpacity(FLinearColor::White)
+	);
+
+	Set("TreeTable.NormalText", FTextBlockStyle(NormalText)
+		.SetColorAndOpacity(FLinearColor::White)
+	);
+
 	//////////////////////////////////////////////////
 
 	// NormalEditableTextBox && SearchBox
@@ -441,6 +449,7 @@ void FInsightsStyle::FStyle::Initialize()
 		Set("SecondaryToolbar2.MaxUniformToolbarSize", 32.0f);
 	}
 
+	// ToggleButton
 	{
 		Set("ToggleButton", FButtonStyle(Button)
 			.SetNormal(FSlateNoResource())
@@ -448,6 +457,7 @@ void FInsightsStyle::FStyle::Initialize()
 			.SetPressed(EDITOR_BOX_BRUSH("Common/RoundedSelection_16x", 4.0f / 16.0f, SelectionColor_Pressed)));
 	}
 
+	// Common.GotoNativeCodeHyperlink
 	{
 		FTextBlockStyle InheritedFromNativeTextStyle = FTextBlockStyle(NormalText)
 			.SetFont(DEFAULT_FONT("Regular", 10));
@@ -467,7 +477,6 @@ void FInsightsStyle::FStyle::Initialize()
 		Set("Common.GotoNativeCodeHyperlink", EditNativeHyperlinkStyle);
 	}
 
-
 	//////////////////////////////////////////////////
 }
 
diff --git a/Engine/Source/Developer/TraceInsights/Private/Insights/MemoryProfiler/ViewModels/MemAllocGroupingByCallstack.cpp b/Engine/Source/Developer/TraceInsights/Private/Insights/MemoryProfiler/ViewModels/MemAllocGroupingByCallstack.cpp
--- a/Engine/Source/Developer/TraceInsights/Private/Insights/MemoryProfiler/ViewModels/MemAllocGroupingByCallstack.cpp
+++ b/Engine/Source/Developer/TraceInsights/Private/Insights/MemoryProfiler/ViewModels/MemAllocGroupingByCallstack.cpp
@@ -224,7 +224,10 @@ FMemAllocGroupingByCallstack::FCallstackGroup* FMemAllocGroupingByCallstack::Cre
 
 	InOutAllCallstackGroup.Add(NewGroupPtr);
 
-	FTableTreeNodePtr NodePtr = MakeShared<FTableTreeNode>(InGroupName, InParentTable);
+	const FSlateBrush* IconBrush = FBaseTreeNode::GetDefaultIcon(true);
+	FLinearColor GroupNodeColor(0.5f, 0.75f, 1.0f, 1.0f);
+
+	FTableTreeNodePtr NodePtr = MakeShared<FCustomTableTreeNode>(InGroupName, InParentTable, IconBrush, GroupNodeColor);
 	NodePtr->SetExpansion(false);
 	InParentGroup->Node->AddChildAndSetGroupPtr(NodePtr);
 
diff --git a/Engine/Source/Developer/TraceInsights/Private/Insights/MemoryProfiler/ViewModels/MemAllocGroupingByTag.cpp b/Engine/Source/Developer/TraceInsights/Private/Insights/MemoryProfiler/ViewModels/MemAllocGroupingByTag.cpp
--- a/Engine/Source/Developer/TraceInsights/Private/Insights/MemoryProfiler/ViewModels/MemAllocGroupingByTag.cpp
+++ b/Engine/Source/Developer/TraceInsights/Private/Insights/MemoryProfiler/ViewModels/MemAllocGroupingByTag.cpp
@@ -114,6 +114,9 @@ void FMemAllocGroupingByTag::GroupNodes(const TArray<FTableTreeNodePtr>& Nodes,
 	};
 	FScopedMemory _(TagNodes);
 
+	const FSlateBrush* IconBrush = FBaseTreeNode::GetDefaultIcon(true);
+	FLinearColor GroupNodeColor(0.75f, 0.5f, 1.0f, 1.0f);
+
 	// Create tag nodes.
 	{
 		TraceServices::FProviderReadScopeLock TagProviderReadScopeLock(TagProvider);
@@ -121,10 +124,10 @@ void FMemAllocGroupingByTag::GroupNodes(const TArray<FTableTreeNodePtr>& Nodes,
 			{
 				const FName NodeName(Name);
 #if INSIGHTS_MERGE_MEM_TAGS_BY_NAME
-				FTableTreeNodePtr TreeNode = MakeShared<FTableTreeNode>(NodeName, InParentTable);
+				FTableTreeNodePtr TreeNode = MakeShared<FCustomTableTreeNode>(NodeName, InParentTable, IconBrush, GroupNodeColor);
 #else
 				const FName NodeNameEx(Name, int32(Id));
-				FTableTreeNodePtr TreeNode = MakeShared<FTableTreeNode>(NodeNameEx, InParentTable);
+				FTableTreeNodePtr TreeNode = MakeShared<FCustomTableTreeNode>(NodeNameEx, InParentTable, IconBrush, GroupNodeColor);
 #endif
 				FTagNode* TagNode = new FTagNode(NodeName, Id, ParentId, TreeNode);
 				TagNodes.Add(TagNode);
@@ -138,7 +141,7 @@ void FMemAllocGroupingByTag::GroupNodes(const TArray<FTableTreeNodePtr>& Nodes,
 	if (!UntaggedNode)
 	{
 		const FName NodeName(TEXT("Untagged"));
-		FTableTreeNodePtr UntaggedTreeNode = MakeShared<FTableTreeNode>(NodeName, InParentTable);
+		FTableTreeNodePtr UntaggedTreeNode = MakeShared<FCustomTableTreeNode>(NodeName, InParentTable, IconBrush, GroupNodeColor);
 		UntaggedNode = new FTagNode(NodeName, UntaggedNodeId, InvalidTagId, UntaggedTreeNode);
 		TagNodes.Add(UntaggedNode);
 		IdToNodeMap.Add(UntaggedNodeId, UntaggedNode);
@@ -150,7 +153,7 @@ void FMemAllocGroupingByTag::GroupNodes(const TArray<FTableTreeNodePtr>& Nodes,
 	if (!UnknownTagNode)
 	{
 		const FName NodeName(TEXT("Unknown"));
-		FTableTreeNodePtr UnknownTagTreeNode = MakeShared<FTableTreeNode>(NodeName, InParentTable);
+		FTableTreeNodePtr UnknownTagTreeNode = MakeShared<FCustomTableTreeNode>(NodeName, InParentTable, IconBrush, GroupNodeColor);
 		UnknownTagNode = new FTagNode(NodeName, UnknownTagNodeId, InvalidTagId, UnknownTagTreeNode);
 		TagNodes.Add(UnknownTagNode);
 		IdToNodeMap.Add(UnknownTagNodeId, UnknownTagNode);
diff --git a/Engine/Source/Developer/TraceInsights/Private/Insights/Table/ViewModels/BaseTreeNode.cpp b/Engine/Source/Developer/TraceInsights/Private/Insights/Table/ViewModels/BaseTreeNode.cpp
--- a/Engine/Source/Developer/TraceInsights/Private/Insights/Table/ViewModels/BaseTreeNode.cpp
+++ b/Engine/Source/Developer/TraceInsights/Private/Insights/Table/ViewModels/BaseTreeNode.cpp
@@ -2,7 +2,11 @@
 
 #include "BaseTreeNode.h"
 
+#include "Styling/SlateBrush.h"
+#include "Styling/StyleColors.h"
+
 // Insights
+#include "Insights/InsightsStyle.h"
 #include "Insights/Table/ViewModels/TableCellValueSorter.h"
 
 #define LOCTEXT_NAMESPACE "Insights::FBaseTreeNode"
@@ -52,6 +56,34 @@ bool FBaseTreeNode::HasExtraDisplayName() const
 
 ////////////////////////////////////////////////////////////////////////////////////////////////////
 
+const FSlateBrush* FBaseTreeNode::GetDefaultIcon(bool bIsGroupNode)
+{
+	if (bIsGroupNode)
+	{
+		return FInsightsStyle::GetBrush("Icons.Group.TreeItem");
+	}
+	else
+	{
+		return FInsightsStyle::GetBrush("Icons.Leaf.TreeItem");
+	}
+}
+
+////////////////////////////////////////////////////////////////////////////////////////////////////
+
+FLinearColor FBaseTreeNode::GetDefaultColor(bool bIsGroupNode)
+{
+	if (bIsGroupNode)
+	{
+		return FLinearColor(1.0f, 0.9f, 0.6f, 1.0f);
+	}
+	else
+	{
+		return FLinearColor(1.0f, 1.0f, 1.0f, 1.0f);
+	}
+}
+
+////////////////////////////////////////////////////////////////////////////////////////////////////
+
 void FBaseTreeNode::SortChildrenAscending(const ITableCellValueSorter& Sorter)
 {
 	Sorter.Sort(GroupData->Children, ESortMode::Ascending);
diff --git a/Engine/Source/Developer/TraceInsights/Private/Insights/Table/ViewModels/BaseTreeNode.h b/Engine/Source/Developer/TraceInsights/Private/Insights/Table/ViewModels/BaseTreeNode.h
--- a/Engine/Source/Developer/TraceInsights/Private/Insights/Table/ViewModels/BaseTreeNode.h
+++ b/Engine/Source/Developer/TraceInsights/Private/Insights/Table/ViewModels/BaseTreeNode.h
@@ -3,9 +3,12 @@
 #pragma once
 
 #include "CoreMinimal.h"
+#include "Math/Color.h"
 
 #include "Insights/Common/SimpleRtti.h"
 
+struct FSlateBrush;
+
 namespace Insights
 {
 
@@ -54,7 +57,7 @@ protected:
 
 public:
 	/** Initialization constructor for the node. */
-	FBaseTreeNode(const FName InName, bool bInIsGroup)
+	explicit FBaseTreeNode(const FName InName, bool bInIsGroup)
 		: DefaultSortOrder(0)
 		, Name(InName)
 		, GroupData(bInIsGroup ? new FGroupNodeData() : &DefaultGroupData)
@@ -118,6 +121,32 @@ public:
 		}
 	}
 
+	/**
+	 * @return the default icon for a group/leaf node.
+	 */
+	static const FSlateBrush* GetDefaultIcon(bool bIsGroupNode);
+
+	/**
+	 * @return a brush icon for this node.
+	 */
+	virtual const FSlateBrush* GetIcon() const
+	{
+		return GetDefaultIcon(IsGroup());
+	}
+
+	/**
+	 * @return the default color tint for a group/leaf node.
+	 */
+	static FLinearColor GetDefaultColor(bool bIsGroupNode);
+
+	/**
+	 * @return the color tint to be used for the icon and the name text of this node.
+	 */
+	virtual FLinearColor GetColor() const
+	{
+		return GetDefaultColor(IsGroup());
+	}
+
 	/**
 	 * @return a pointer to a data context for this node.
 	 */
diff --git a/Engine/Source/Developer/TraceInsights/Private/Insights/Table/ViewModels/Table.cpp b/Engine/Source/Developer/TraceInsights/Private/Insights/Table/ViewModels/Table.cpp
--- a/Engine/Source/Developer/TraceInsights/Private/Insights/Table/ViewModels/Table.cpp
+++ b/Engine/Source/Developer/TraceInsights/Private/Insights/Table/ViewModels/Table.cpp
@@ -51,11 +51,11 @@ int32 FTable::GetColumnPositionIndex(const FName& ColumnId) const
 
 ////////////////////////////////////////////////////////////////////////////////////////////////////
 
-void FTable::SetColumns(const TArray<TSharedRef<Insights::FTableColumn>>& InColumns)
+void FTable::SetColumns(const TArray<TSharedRef<FTableColumn>>& InColumns)
 {
 	Columns.Reset(InColumns.Num());
 	ColumnIdToPtrMapping.Reset();
-	for (TSharedRef<Insights::FTableColumn> ColumnRef : InColumns)
+	for (TSharedRef<FTableColumn> ColumnRef : InColumns)
 	{
 		AddColumn(ColumnRef);
 	}
@@ -128,20 +128,20 @@ void FTable::GetVisibleColumns(TArray<TSharedRef<FTableColumn>>& InArray) const
 
 ////////////////////////////////////////////////////////////////////////////////////////////////////
 
-void FTable::GetVisibleColumnsData(const TArray<Insights::FBaseTreeNodePtr>& InNodes, const FName& LogListingName, TCHAR Separator, bool bIncludeHeaders, FString& OutData) const
+void FTable::GetVisibleColumnsData(const TArray<FBaseTreeNodePtr>& InNodes, const FName& LogListingName, TCHAR Separator, bool bIncludeHeaders, FString& OutData) const
 {
 	constexpr TCHAR LineEnd = TEXT('\n');
 	constexpr TCHAR QuotationMarkBegin = TEXT('\"');
 	constexpr TCHAR QuotationMarkEnd = TEXT('\"');
 
-	TArray<TSharedRef<Insights::FTableColumn>> VisibleColumns;
+	TArray<TSharedRef<FTableColumn>> VisibleColumns;
 	GetVisibleColumns(VisibleColumns);
 
 	// Table headers
 	if (bIncludeHeaders)
 	{
 		bool bIsFirstColumn = true;
-		for (const TSharedRef<Insights::FTableColumn>& ColumnRef : VisibleColumns)
+		for (const TSharedRef<FTableColumn>& ColumnRef : VisibleColumns)
 		{
 			if (bIsFirstColumn)
 			{
@@ -181,10 +181,10 @@ void FTable::GetVisibleColumnsData(const TArray<Insights::FBaseTreeNodePtr>& InN
 	// Selected items
 	for (int Index = 0; Index < NumItems; Index++)
 	{
-		const Insights::FBaseTreeNodePtr& Node = InNodes[Index];
+		const FBaseTreeNodePtr& Node = InNodes[Index];
 
 		bool bIsFirstColumn = true;
-		for (const TSharedRef<Insights::FTableColumn>& ColumnRef : VisibleColumns)
+		for (const TSharedRef<FTableColumn>& ColumnRef : VisibleColumns)
 		{
 			if (bIsFirstColumn)
 			{
diff --git a/Engine/Source/Developer/TraceInsights/Private/Insights/Table/ViewModels/Table.h b/Engine/Source/Developer/TraceInsights/Private/Insights/Table/ViewModels/Table.h
--- a/Engine/Source/Developer/TraceInsights/Private/Insights/Table/ViewModels/Table.h
+++ b/Engine/Source/Developer/TraceInsights/Private/Insights/Table/ViewModels/Table.h
@@ -32,7 +32,7 @@ public:
 	bool IsValid() const { return Columns.Num() > 0; }
 
 	const TArray<TSharedRef<FTableColumn>>& GetColumns() const { return Columns; }
-	void SetColumns(const TArray<TSharedRef<Insights::FTableColumn>>& InColumns);
+	void SetColumns(const TArray<TSharedRef<FTableColumn>>& InColumns);
 
 	void GetVisibleColumns(TArray<TSharedRef<FTableColumn>>& InArray) const;
 
diff --git a/Engine/Source/Developer/TraceInsights/Private/Insights/Table/ViewModels/TableCellValue.h b/Engine/Source/Developer/TraceInsights/Private/Insights/Table/ViewModels/TableCellValue.h
--- a/Engine/Source/Developer/TraceInsights/Private/Insights/Table/ViewModels/TableCellValue.h
+++ b/Engine/Source/Developer/TraceInsights/Private/Insights/Table/ViewModels/TableCellValue.h
@@ -4,6 +4,9 @@
 
 #include "CoreMinimal.h"
 
+namespace Insights
+{
+
 ////////////////////////////////////////////////////////////////////////////////////////////////////
 
 enum class ETableCellDataType : uint32
@@ -187,3 +190,5 @@ public:
 };
 
 ////////////////////////////////////////////////////////////////////////////////////////////////////
+
+} // namespace Insights
diff --git a/Engine/Source/Developer/TraceInsights/Private/Insights/Table/ViewModels/TableCellValueFormatter.cpp b/Engine/Source/Developer/TraceInsights/Private/Insights/Table/ViewModels/TableCellValueFormatter.cpp
--- a/Engine/Source/Developer/TraceInsights/Private/Insights/Table/ViewModels/TableCellValueFormatter.cpp
+++ b/Engine/Source/Developer/TraceInsights/Private/Insights/Table/ViewModels/TableCellValueFormatter.cpp
@@ -11,7 +11,7 @@
 #include "Insights/Table/ViewModels/TableCellValueGetter.h"
 #include "Insights/Table/ViewModels/TableColumn.h"
 
-#define LOCTEXT_NAMESPACE "TableCellValueFormatter"
+#define LOCTEXT_NAMESPACE "Insights::FTableCellValueFormatter"
 
 namespace Insights
 {
diff --git a/Engine/Source/Developer/TraceInsights/Private/Insights/Table/ViewModels/TableCellValueGetter.cpp b/Engine/Source/Developer/TraceInsights/Private/Insights/Table/ViewModels/TableCellValueGetter.cpp
--- a/Engine/Source/Developer/TraceInsights/Private/Insights/Table/ViewModels/TableCellValueGetter.cpp
+++ b/Engine/Source/Developer/TraceInsights/Private/Insights/Table/ViewModels/TableCellValueGetter.cpp
@@ -4,7 +4,7 @@
 
 #include "Insights/Table/ViewModels/BaseTreeNode.h"
 
-#define LOCTEXT_NAMESPACE "TableCellValueGetter"
+#define LOCTEXT_NAMESPACE "Insights::FTableCellValueGetter"
 
 namespace Insights
 {
diff --git a/Engine/Source/Developer/TraceInsights/Private/Insights/Table/ViewModels/TableCellValueSorter.cpp b/Engine/Source/Developer/TraceInsights/Private/Insights/Table/ViewModels/TableCellValueSorter.cpp
--- a/Engine/Source/Developer/TraceInsights/Private/Insights/Table/ViewModels/TableCellValueSorter.cpp
+++ b/Engine/Source/Developer/TraceInsights/Private/Insights/Table/ViewModels/TableCellValueSorter.cpp
@@ -4,7 +4,7 @@
 
 #include "Insights/Table/ViewModels/BaseTreeNode.h"
 
-#define LOCTEXT_NAMESPACE "TableCellValueSorter"
+#define LOCTEXT_NAMESPACE "Insights::FTableCellValueSorter"
 
 // Default pre-sorting (group nodes sorts above leaf nodes)
 #define INSIGHTS_DEFAULT_PRESORTING_NODES(A, B) \
diff --git a/Engine/Source/Developer/TraceInsights/Private/Insights/Table/ViewModels/TableColumn.cpp b/Engine/Source/Developer/TraceInsights/Private/Insights/Table/ViewModels/TableColumn.cpp
--- a/Engine/Source/Developer/TraceInsights/Private/Insights/Table/ViewModels/TableColumn.cpp
+++ b/Engine/Source/Developer/TraceInsights/Private/Insights/Table/ViewModels/TableColumn.cpp
@@ -6,7 +6,7 @@
 #include "Insights/Table/ViewModels/TableCellValueGetter.h"
 #include "Insights/Table/ViewModels/TableCellValueSorter.h"
 
-#define LOCTEXT_NAMESPACE "TableColumn"
+#define LOCTEXT_NAMESPACE "Insights::FTableColumn"
 
 namespace Insights
 {
diff --git a/Engine/Source/Developer/TraceInsights/Private/Insights/Table/ViewModels/TableTreeNode.cpp b/Engine/Source/Developer/TraceInsights/Private/Insights/Table/ViewModels/TableTreeNode.cpp
--- a/Engine/Source/Developer/TraceInsights/Private/Insights/Table/ViewModels/TableTreeNode.cpp
+++ b/Engine/Source/Developer/TraceInsights/Private/Insights/Table/ViewModels/TableTreeNode.cpp
@@ -8,6 +8,7 @@ namespace Insights
 {
 
 INSIGHTS_IMPLEMENT_RTTI(FTableTreeNode)
+INSIGHTS_IMPLEMENT_RTTI(FCustomTableTreeNode)
 
 } // namespace Insights
 
diff --git a/Engine/Source/Developer/TraceInsights/Private/Insights/Table/ViewModels/TableTreeNode.h b/Engine/Source/Developer/TraceInsights/Private/Insights/Table/ViewModels/TableTreeNode.h
--- a/Engine/Source/Developer/TraceInsights/Private/Insights/Table/ViewModels/TableTreeNode.h
+++ b/Engine/Source/Developer/TraceInsights/Private/Insights/Table/ViewModels/TableTreeNode.h
@@ -51,7 +51,7 @@ class FTableTreeNode : public FBaseTreeNode
 
 public:
 	/** Initialization constructor for a table record node. */
-	FTableTreeNode(const FName InName, TWeakPtr<FTable> InParentTable, int32 InRowIndex)
+	explicit FTableTreeNode(const FName InName, TWeakPtr<FTable> InParentTable, int32 InRowIndex)
 		: FBaseTreeNode(InName, false)
 		, ParentTable(InParentTable)
 		, RowId(InRowIndex)
@@ -60,7 +60,7 @@ public:
 	}
 
 	/** Initialization constructor for a group node. */
-	FTableTreeNode(const FName InGroupName, TWeakPtr<FTable> InParentTable)
+	explicit FTableTreeNode(const FName InGroupName, TWeakPtr<FTable> InParentTable)
 		: FBaseTreeNode(InGroupName, true)
 		, ParentTable(InParentTable)
 		, RowId(FTableRowId::InvalidRowIndex)
@@ -79,6 +79,7 @@ public:
 
 	void InitAggregatedValues()
 	{
+		ensure(IsGroup());
 		if (!AggregatedValues)
 		{
 			AggregatedValues = new TMap<FName, FTableCellValue>();
@@ -117,4 +118,63 @@ private:
 
 ////////////////////////////////////////////////////////////////////////////////////////////////////
 
+class FCustomTableTreeNode : public FTableTreeNode
+{
+	INSIGHTS_DECLARE_RTTI(FCustomTableTreeNode, FTableTreeNode)
+
+public:
+	/** Initialization constructor for a table record node. */
+	explicit FCustomTableTreeNode(const FName InName, TWeakPtr<FTable> InParentTable, int32 InRowIndex, const FSlateBrush* InIconBrush, FLinearColor InColor)
+		: FTableTreeNode(InName, InParentTable)
+		, IconBrush(InIconBrush)
+		, Color(InColor)
+	{
+	}
+
+	/** Initialization constructor for the group node. */
+	explicit FCustomTableTreeNode(const FName InName, TWeakPtr<FTable> InParentTable, const FSlateBrush* InIconBrush, FLinearColor InColor)
+		: FTableTreeNode(InName, InParentTable)
+		, IconBrush(InIconBrush)
+		, Color(InColor)
+	{
+	}
+
+	virtual ~FCustomTableTreeNode()
+	{
+	}
+
+	/**
+	 * @return a brush icon for this node.
+	 */
+	virtual const FSlateBrush* GetIcon() const override
+	{
+		return IconBrush;
+	}
+
+	/**
+	 * Sets an icon brush for this node.
+	 */
+	void SetIcon(const FSlateBrush* InIconBrush)
+	{
+		IconBrush = InIconBrush;
+	}
+
+	/**
+	 * @return the color tint for icon and name text.
+	 */
+	virtual FLinearColor GetColor() const override
+	{
+		return Color;
+	}
+
+private:
+	/** The icon of this node. */
+	const FSlateBrush* IconBrush;
+
+	/** The color tint of this node. */
+	FLinearColor Color;
+};
+
+////////////////////////////////////////////////////////////////////////////////////////////////////
+
 } // namespace Insights
diff --git a/Engine/Source/Developer/TraceInsights/Private/Insights/Table/ViewModels/TreeNodeGrouping.cpp b/Engine/Source/Developer/TraceInsights/Private/Insights/Table/ViewModels/TreeNodeGrouping.cpp
--- a/Engine/Source/Developer/TraceInsights/Private/Insights/Table/ViewModels/TreeNodeGrouping.cpp
+++ b/Engine/Source/Developer/TraceInsights/Private/Insights/Table/ViewModels/TreeNodeGrouping.cpp
@@ -5,7 +5,7 @@
 #include "Insights/Common/AsyncOperationProgress.h"
 #include "Insights/Table/ViewModels/TableTreeNode.h"
 
-#define LOCTEXT_NAMESPACE "Insights_TreeNode"
+#define LOCTEXT_NAMESPACE "Insights::FTreeNodeGrouping"
 
 namespace Insights
 {
@@ -152,12 +152,6 @@ template<> double TTreeNodeGroupingByUniqueValue<double>::GetValue(const FTableC
 
 FName FTreeNodeGroupingByUniqueValueCString::GetGroupName(const TCHAR* Value)
 {
-	if (Value == nullptr || *Value == TEXT('\0'))
-	{
-		static FName EmptyGroupName(TEXT("N/A"));
-		return EmptyGroupName;
-	}
-
 	FStringView StringView(Value);
 	if (StringView.Len() >= NAME_SIZE)
 	{
@@ -170,8 +164,10 @@ FName FTreeNodeGroupingByUniqueValueCString::GetGroupName(const TCHAR* Value)
 
 void FTreeNodeGroupingByUniqueValueCString::GroupNodes(const TArray<FTableTreeNodePtr>& Nodes, FTableTreeNode& ParentGroup, TWeakPtr<FTable> InParentTable, IAsyncOperationProgress& InAsyncOperationProgress) const
 {
-	TMap<const void*, FTableTreeNodePtr> GroupMap; // using void* as TMap<const TCHAR*, T> doesn't allow nullptr to be added
-	FTableTreeNodePtr UnsetGroupPtr = nullptr;
+	FTableTreeNodePtr UnsetGroupPtr = nullptr; // for unset FTableCellValue
+	FTableTreeNodePtr EmptyGroupPtr = nullptr; // for nullptr and empty strings
+	TMap<const TCHAR*, FTableTreeNodePtr> GroupMap; // for valid strings
+	TMap<FName, FTableTreeNodePtr> GroupNameMap; // for valid strings
 
 	ParentGroup.ClearChildren();
 
@@ -195,19 +191,40 @@ void FTreeNodeGroupingByUniqueValueCString::GroupNodes(const TArray<FTableTreeNo
 		if (CellValue.IsSet())
 		{
 			const TCHAR* Value = CellValue.GetValue().CString;
-
-			FTableTreeNodePtr* GroupPtrPtr = GroupMap.Find(Value);
-			if (!GroupPtrPtr)
+			if (Value != nullptr && Value[0] != TEXT('\0'))
 			{
-				const FName GroupName = GetGroupName(Value);
-				GroupPtr = MakeShared<FTableTreeNode>(GroupName, InParentTable);
-				GroupPtr->SetExpansion(false);
-				ParentGroup.AddChildAndSetGroupPtr(GroupPtr);
-				GroupMap.Add(Value, GroupPtr);
+				FTableTreeNodePtr* GroupPtrPtr = GroupMap.Find(Value);
+				if (!GroupPtrPtr)
+				{
+					const FName GroupName = GetGroupName(Value);
+					GroupPtrPtr = GroupNameMap.Find(Value);
+					if (!GroupPtrPtr)
+					{
+						GroupPtr = MakeShared<FTableTreeNode>(GroupName, InParentTable);
+						GroupPtr->SetExpansion(false);
+						ParentGroup.AddChildAndSetGroupPtr(GroupPtr);
+						GroupNameMap.Add(GroupName, GroupPtr);
+					}
+					else
+					{
+						GroupPtr = *GroupPtrPtr;
+					}
+					GroupMap.Add(Value, GroupPtr);
+				}
+				else
+				{
+					GroupPtr = *GroupPtrPtr;
+				}
 			}
 			else
 			{
-				GroupPtr = *GroupPtrPtr;
+				if (!EmptyGroupPtr)
+				{
+					EmptyGroupPtr = MakeShared<FTableTreeNode>(FName(TEXT("N/A")), InParentTable);
+					EmptyGroupPtr->SetExpansion(false);
+					ParentGroup.AddChildAndSetGroupPtr(EmptyGroupPtr);
+				}
+				GroupPtr = EmptyGroupPtr;
 			}
 		}
 		else
diff --git a/Engine/Source/Developer/TraceInsights/Private/Insights/Table/Widgets/STableTreeView.cpp b/Engine/Source/Developer/TraceInsights/Private/Insights/Table/Widgets/STableTreeView.cpp
--- a/Engine/Source/Developer/TraceInsights/Private/Insights/Table/Widgets/STableTreeView.cpp
+++ b/Engine/Source/Developer/TraceInsights/Private/Insights/Table/Widgets/STableTreeView.cpp
@@ -22,15 +22,15 @@
 #include "Widgets/Views/STableViewBase.h"
 
 // Insights
+#include "Insights/InsightsManager.h"
 #include "Insights/InsightsStyle.h"
 #include "Insights/Log.h"
 #include "Insights/Table/ViewModels/Table.h"
+#include "Insights/Table/ViewModels/TableCellValueSorter.h"
 #include "Insights/Table/ViewModels/TreeNodeGrouping.h"
-#include "Insights/Table/ViewModels/TreeNodeSorting.h"
-#include "Insights/Table/ViewModels/UntypedTable.h"
+
 #include "Insights/Table/Widgets/STableTreeViewTooltip.h"
 #include "Insights/Table/Widgets/STableTreeViewRow.h"
-#include "Insights/TimingProfilerCommon.h"
 #include "Insights/ViewModels/FilterConfigurator.h"
 #include "Insights/Widgets/SAsyncOperationStatus.h"
 
@@ -293,7 +293,7 @@ void STableTreeView::ConstructWidget(TSharedPtr<FTable> InTablePtr)
 			.VAlign(VAlign_Bottom)
 			.Padding(16.0f)
 			[
-				SAssignNew(AsyncOperationStatus, Insights::SAsyncOperationStatus, SharedThis(this))
+				SAssignNew(AsyncOperationStatus, SAsyncOperationStatus, SharedThis(this))
 			]
 
 			+ SOverlay::Slot()
@@ -2363,7 +2363,7 @@ void STableTreeView::CreateSortings()
 		const FTableColumn& Column = ColumnRef.Get();
 		if (Column.CanBeSorted())
 		{
-			TSharedPtr<Insights::ITableCellValueSorter> SorterPtr = Column.GetValueSorter();
+			TSharedPtr<ITableCellValueSorter> SorterPtr = Column.GetValueSorter();
 			if (ensure(SorterPtr.IsValid()))
 			{
 				AvailableSorters.Add(SorterPtr);
@@ -2378,7 +2378,7 @@ void STableTreeView::CreateSortings()
 
 void STableTreeView::UpdateCurrentSortingByColumn()
 {
-	TSharedPtr<Insights::FTableColumn> ColumnPtr = Table->FindColumn(ColumnBeingSorted);
+	TSharedPtr<FTableColumn> ColumnPtr = Table->FindColumn(ColumnBeingSorted);
 	CurrentSorter = ColumnPtr.IsValid() ? ColumnPtr->GetValueSorter() : nullptr;
 }
 
@@ -3356,7 +3356,7 @@ void STableTreeView::ContextMenu_CopySelectedToClipboard_Execute()
 		return;
 	}
 
-	TArray<Insights::FBaseTreeNodePtr> SelectedNodes;
+	TArray<FBaseTreeNodePtr> SelectedNodes;
 	for (FTableTreeNodePtr TimerPtr : TreeView->GetSelectedItems())
 	{
 		SelectedNodes.Add(TimerPtr);
@@ -3371,7 +3371,7 @@ void STableTreeView::ContextMenu_CopySelectedToClipboard_Execute()
 
 	if (CurrentSorter.IsValid())
 	{
-		CurrentSorter->Sort(SelectedNodes, ColumnSortMode == EColumnSortMode::Ascending ? Insights::ESortMode::Ascending : Insights::ESortMode::Descending);
+		CurrentSorter->Sort(SelectedNodes, ColumnSortMode == EColumnSortMode::Ascending ? ESortMode::Ascending : ESortMode::Descending);
 	}
 
 	Table->GetVisibleColumnsData(SelectedNodes, GetLogListingName(), TEXT('\t'), true, ClipboardText);
@@ -3604,7 +3604,7 @@ void STableTreeView::ContextMenu_ExportToFile_Execute(bool bInExportCollapsed, b
 	ExportFileHandle->Write((uint8*)&BOM, sizeof(UTF16CHAR));
 
 	bool bIncludeHeaders = true;
-	auto WriteToFile = [this, &Path, &bIncludeHeaders, ExportFileHandle, Separator](TArray<Insights::FBaseTreeNodePtr>& InNodes)
+	auto WriteToFile = [this, &Path, &bIncludeHeaders, ExportFileHandle, Separator](TArray<FBaseTreeNodePtr>& InNodes)
 	{
 		FString Text;
 		this->GetTable()->GetVisibleColumnsData(InNodes, GetLogListingName(), Separator, bIncludeHeaders, Text);
@@ -3618,7 +3618,7 @@ void STableTreeView::ContextMenu_ExportToFile_Execute(bool bInExportCollapsed, b
 		InNodes.Empty(InNodes.Num());
 	};
 
-	TArray<Insights::FBaseTreeNodePtr> NodesBuffer;
+	TArray<FBaseTreeNodePtr> NodesBuffer;
 	ExportToFileRec(Root, NodesBuffer, bInExportCollapsed, InExportLeafs, WriteToFile);
 
 	// Write the remaining lines
@@ -3635,7 +3635,7 @@ void STableTreeView::ContextMenu_ExportToFile_Execute(bool bInExportCollapsed, b
 
 ////////////////////////////////////////////////////////////////////////////////////////////////////
 
-void STableTreeView::ExportToFileRec(const FBaseTreeNodePtr& InGroupNode, TArray<Insights::FBaseTreeNodePtr>& InNodes, bool bInExportCollapsed, bool InExportLeafs, WriteToFileCallback Callback)
+void STableTreeView::ExportToFileRec(const FBaseTreeNodePtr& InGroupNode, TArray<FBaseTreeNodePtr>& InNodes, bool bInExportCollapsed, bool InExportLeafs, WriteToFileCallback Callback)
 {
 	constexpr int MaxBufferSize = 10000;
 
@@ -3789,6 +3789,7 @@ SSessionTableTreeView::~SSessionTableTreeView()
 void SSessionTableTreeView::ConstructWidget(TSharedPtr<FTable> InTablePtr)
 {
 	STableTreeView::ConstructWidget(InTablePtr);
+
 	// Register ourselves with the Insights manager.
 	FInsightsManager::Get()->GetSessionChangedEvent().AddSP(this, &SSessionTableTreeView::InsightsManager_OnSessionChanged);
 
diff --git a/Engine/Source/Developer/TraceInsights/Private/Insights/Table/Widgets/STableTreeView.h b/Engine/Source/Developer/TraceInsights/Private/Insights/Table/Widgets/STableTreeView.h
--- a/Engine/Source/Developer/TraceInsights/Private/Insights/Table/Widgets/STableTreeView.h
+++ b/Engine/Source/Developer/TraceInsights/Private/Insights/Table/Widgets/STableTreeView.h
@@ -33,7 +33,7 @@ class FUICommandList;
 
 namespace TraceServices
 {
-class IAnalysisSession;
+	class IAnalysisSession;
 }
 
 namespace Insights
@@ -129,10 +129,10 @@ public:
 	 * Construct this widget
 	 * @param InArgs - The declaration data for this widget
 	 */
-	void Construct(const FArguments& InArgs, TSharedPtr<Insights::FTable> InTablePtr);
+	void Construct(const FArguments& InArgs, TSharedPtr<FTable> InTablePtr);
 
-	TSharedPtr<Insights::FTable>& GetTable() { return Table; }
-	const TSharedPtr<Insights::FTable>& GetTable() const { return Table; }
+	TSharedPtr<FTable>& GetTable() { return Table; }
+	const TSharedPtr<FTable>& GetTable() const { return Table; }
 
 	virtual void Reset();
 
@@ -460,8 +460,8 @@ protected:
 
 	virtual void ExtendMenu(FMenuBuilder& Menu) {}
 
-	typedef TFunctionRef<void(TArray<Insights::FBaseTreeNodePtr>& InNodes)> WriteToFileCallback;
-	void ExportToFileRec(const FBaseTreeNodePtr& InGroupNode, TArray<Insights::FBaseTreeNodePtr>& InNodes, bool bInExportCollapsed, bool InExportLeafs, WriteToFileCallback Callback);
+	typedef TFunctionRef<void(TArray<FBaseTreeNodePtr>& InNodes)> WriteToFileCallback;
+	void ExportToFileRec(const FBaseTreeNodePtr& InGroupNode, TArray<FBaseTreeNodePtr>& InNodes, bool bInExportCollapsed, bool InExportLeafs, WriteToFileCallback Callback);
 
 	FText GetTreeViewBannerText() const { return TreeViewBannerText; }
 
@@ -494,7 +494,7 @@ protected:
 
 protected:
 	/** Table view model. */
-	TSharedPtr<Insights::FTable> Table;
+	TSharedPtr<FTable> Table;
 	//////////////////////////////////////////////////
 	// Tree View, Columns
 
diff --git a/Engine/Source/Developer/TraceInsights/Private/Insights/Table/Widgets/STableTreeViewCell.cpp b/Engine/Source/Developer/TraceInsights/Private/Insights/Table/Widgets/STableTreeViewCell.cpp
--- a/Engine/Source/Developer/TraceInsights/Private/Insights/Table/Widgets/STableTreeViewCell.cpp
+++ b/Engine/Source/Developer/TraceInsights/Private/Insights/Table/Widgets/STableTreeViewCell.cpp
@@ -3,6 +3,7 @@
 #include "STableTreeViewCell.h"
 
 #include "SlateOptMacros.h"
+#include "Styling/StyleColors.h"
 #include "Widgets/Images/SImage.h"
 #include "Widgets/Input/SButton.h"
 #include "Widgets/SBoxPanel.h"
@@ -26,8 +27,10 @@ namespace Insights
 
 BEGIN_SLATE_FUNCTION_BUILD_OPTIMIZATION
 
-void STableTreeViewCell::Construct(const FArguments& InArgs, const TSharedRef<ITableRow>& TableRow)
+void STableTreeViewCell::Construct(const FArguments& InArgs, const TSharedRef<ITableRow>& InTableRow)
 {
+	TableRow = InTableRow;
+
 	TablePtr = InArgs._TablePtr;
 	ColumnPtr = InArgs._ColumnPtr;
 	TableTreeNodePtr = InArgs._TableTreeNodePtr;
@@ -40,28 +43,31 @@ void STableTreeViewCell::Construct(const FArguments& InArgs, const TSharedRef<IT
 
 	ChildSlot
 	[
-		GenerateWidgetForColumn(InArgs, TableRow)
+		GenerateWidgetForColumn(InArgs)
 	];
 }
 
 ////////////////////////////////////////////////////////////////////////////////////////////////////
 
-TSharedRef<SWidget> STableTreeViewCell::GenerateWidgetForColumn(const FArguments& InArgs, const TSharedRef<ITableRow>& TableRow)
+TSharedRef<SWidget> STableTreeViewCell::GenerateWidgetForColumn(const FArguments& InArgs)
 {
 	if (InArgs._IsNameColumn)
 	{
-		return GenerateWidgetForNameColumn(InArgs, TableRow);
+		return GenerateWidgetForNameColumn(InArgs);
 	}
 	else
 	{
-		return GenerateWidgetForTableColumn(InArgs, TableRow);
+		return GenerateWidgetForTableColumn(InArgs);
 	}
 }
 
 ////////////////////////////////////////////////////////////////////////////////////////////////////
 
-TSharedRef<SWidget> STableTreeViewCell::GenerateWidgetForNameColumn(const FArguments& InArgs, const TSharedRef<ITableRow>& TableRow)
+TSharedRef<SWidget> STableTreeViewCell::GenerateWidgetForNameColumn(const FArguments& InArgs)
 {
+	TSharedPtr<STableTreeViewRow> Row = StaticCastSharedPtr<STableTreeViewRow, ITableRow>(TableRow);
+	TSharedPtr<IToolTip> RowToolTip = Row->GetRowToolTip();
+
 	return
 		SNew(SHorizontalBox)
 
@@ -73,60 +79,54 @@ TSharedRef<SWidget> STableTreeViewCell::GenerateWidgetForNameColumn(const FArgum
 			SNew(SExpanderArrow, TableRow)
 		]
 
-		// Info icon + tooltip
+		// Icon + tooltip
 		+ SHorizontalBox::Slot()
 		.AutoWidth()
 		.HAlign(HAlign_Center)
 		.VAlign(VAlign_Center)
 		[
 			SNew(SImage)
-			.Visibility(this, &STableTreeViewCell::GetHintIconVisibility)
-			.Image(FInsightsStyle::GetBrush("Icons.Hint.TreeItem"))
-			.ToolTip(GetRowToolTip(TableRow))
+			.Image(this, &STableTreeViewCell::GetIcon)
+			.ColorAndOpacity(this, &STableTreeViewCell::GetIconColorAndOpacity)
+			.ToolTip(RowToolTip)
 		]
 
 		// Name
 		+ SHorizontalBox::Slot()
 		.AutoWidth()
-		.VAlign(VAlign_Center)
 		.HAlign(HAlign_Left)
-		.Padding(FMargin(2.0f, 0.0f))
+		.VAlign(VAlign_Center)
+		.Padding(FMargin(2.0f, 0.0f, 2.0f, 0.0f))
 		[
 			SNew(STextBlock)
 			.Text(this, &STableTreeViewCell::GetDisplayName)
 			.HighlightText(InArgs._HighlightText)
-			.TextStyle(FInsightsStyle::Get(), TEXT("TreeTable.Tooltip"))
-			.ColorAndOpacity(this, &STableTreeViewCell::GetColorAndOpacity)
+			.TextStyle(FInsightsStyle::Get(), TEXT("TreeTable.NameText"))
+			.ColorAndOpacity(this, &STableTreeViewCell::GetDisplayNameColorAndOpacity)
 			.ShadowColorAndOpacity(this, &STableTreeViewCell::GetShadowColorAndOpacity)
+			.ToolTip(RowToolTip)
 		]
 
 		// Name Suffix
 		+ SHorizontalBox::Slot()
 		.AutoWidth()
-		.VAlign(VAlign_Center)
 		.HAlign(HAlign_Left)
-		.Padding(FMargin(2.0f, 0.0f))
+		.VAlign(VAlign_Center)
+		.Padding(FMargin(2.0f, 0.0f, 2.0f, 0.0f))
 		[
 			SNew(STextBlock)
 			.Visibility(this, &STableTreeViewCell::HasExtraDisplayName)
 			.Text(this, &STableTreeViewCell::GetExtraDisplayName)
-			.TextStyle(FInsightsStyle::Get(), TEXT("TreeTable.Tooltip"))
-			.ColorAndOpacity(this, &STableTreeViewCell::GetExtraColorAndOpacity)
+			.TextStyle(FInsightsStyle::Get(), TEXT("TreeTable.NameText"))
+			.ColorAndOpacity(this, &STableTreeViewCell::GetExtraDisplayNameColorAndOpacity)
 			.ShadowColorAndOpacity(this, &STableTreeViewCell::GetShadowColorAndOpacity)
+			.ToolTip(RowToolTip)
 		]
 	;
 }
 
 ////////////////////////////////////////////////////////////////////////////////////////////////////
 
-TSharedPtr<IToolTip> STableTreeViewCell::GetRowToolTip(const TSharedRef<ITableRow>& TableRow) const
-{
-	TSharedRef<STableTreeViewRow> Row = StaticCastSharedRef<STableTreeViewRow, ITableRow>(TableRow);
-	return Row->GetRowToolTip();
-}
-
-////////////////////////////////////////////////////////////////////////////////////////////////////
-
 FText STableTreeViewCell::GetValueAsText() const
 {
 	return ColumnPtr->GetValueAsText(*TableTreeNodePtr);
@@ -134,11 +134,11 @@ FText STableTreeViewCell::GetValueAsText() const
 
 ////////////////////////////////////////////////////////////////////////////////////////////////////
 
-TSharedRef<SWidget> STableTreeViewCell::GenerateWidgetForTableColumn(const FArguments& InArgs, const TSharedRef<ITableRow>& TableRow)
+TSharedRef<SWidget> STableTreeViewCell::GenerateWidgetForTableColumn(const FArguments& InArgs)
 {
 	TSharedRef<STextBlock> TextBox = SNew(STextBlock)
-		.TextStyle(FInsightsStyle::Get(), TEXT("TreeTable.Tooltip"))
-		.ColorAndOpacity(this, &STableTreeViewCell::GetStatsColorAndOpacity)
+		.TextStyle(FInsightsStyle::Get(), TEXT("TreeTable.NormalText"))
+		.ColorAndOpacity(this, &STableTreeViewCell::GetNormalTextColorAndOpacity)
 		.ShadowColorAndOpacity(this, &STableTreeViewCell::GetShadowColorAndOpacity);
 
 	if (ColumnPtr->IsDynamic())
@@ -151,15 +151,17 @@ TSharedRef<SWidget> STableTreeViewCell::GenerateWidgetForTableColumn(const FArgu
 		TextBox->SetText(CellText);
 	}
 
+	TSharedPtr<IToolTip> ColumnToolTip = ColumnPtr->GetValueFormatter()->GetCustomTooltip(*ColumnPtr, *TableTreeNodePtr);
+
 	return
 		SNew(SHorizontalBox)
-		.ToolTip(GetTooltip())
+		.ToolTip(ColumnToolTip)
 
 		+ SHorizontalBox::Slot()
 		.FillWidth(1.0f)
-		.VAlign(VAlign_Center)
 		.HAlign(ColumnPtr->GetHorizontalAlignment())
-		.Padding(FMargin(2.0f, 0.0f))
+		.VAlign(VAlign_Center)
+		.Padding(FMargin(2.0f, 0.0f, 2.0f, 0.0f))
 		[
 			TextBox
 		];
@@ -169,9 +171,109 @@ END_SLATE_FUNCTION_BUILD_OPTIMIZATION
 
 ////////////////////////////////////////////////////////////////////////////////////////////////////
 
-TSharedPtr<IToolTip> STableTreeViewCell::GetTooltip() const
+bool STableTreeViewCell::IsSelected() const
+{
+	return TableRow->IsItemSelected();
+}
+
+////////////////////////////////////////////////////////////////////////////////////////////////////
+
+FSlateColor STableTreeViewCell::GetIconColorAndOpacity() const
+{
+	bool bIsHoveredOrSelected = IsHovered() || IsSelected();
+
+	return bIsHoveredOrSelected ?
+		TableTreeNodePtr->GetColor() :
+		TableTreeNodePtr->GetColor().CopyWithNewOpacity(0.5f);
+}
+
+////////////////////////////////////////////////////////////////////////////////////////////////////
+
+FSlateColor STableTreeViewCell::GetDisplayNameColorAndOpacity() const
+{
+	bool bIsHoveredOrSelected = IsHovered() || IsSelected();
+
+	if (TableTreeNodePtr->IsFiltered())
+	{
+		return bIsHoveredOrSelected ?
+			TableTreeNodePtr->GetColor().CopyWithNewOpacity(0.5f) :
+			TableTreeNodePtr->GetColor().CopyWithNewOpacity(0.4f);
+	}
+	else
+	{
+		return bIsHoveredOrSelected ?
+			TableTreeNodePtr->GetColor() :
+			TableTreeNodePtr->GetColor().CopyWithNewOpacity(0.8f);
+	}
+}
+
+////////////////////////////////////////////////////////////////////////////////////////////////////
+
+FSlateColor STableTreeViewCell::GetExtraDisplayNameColorAndOpacity() const
+{
+	bool bIsHoveredOrSelected = IsHovered() || IsSelected();
+
+	if (TableTreeNodePtr->IsFiltered())
+	{
+		return bIsHoveredOrSelected ?
+			FLinearColor(0.3f, 0.3f, 0.3f, 0.5f) :
+			FLinearColor(0.3f, 0.3f, 0.3f, 0.4f);
+	}
+	else
+	{
+		return bIsHoveredOrSelected ?
+			FLinearColor(0.3f, 0.3f, 0.3f, 1.0f) :
+			FLinearColor(0.3f, 0.3f, 0.3f, 0.8f);
+	}
+}
+
+////////////////////////////////////////////////////////////////////////////////////////////////////
+
+FSlateColor STableTreeViewCell::GetNormalTextColorAndOpacity() const
+{
+	bool bIsHoveredOrSelected = IsHovered() || IsSelected();
+
+	if (TableTreeNodePtr->IsGroup())
+	{
+		if (TableTreeNodePtr->IsFiltered())
+		{
+			return bIsHoveredOrSelected ?
+				FStyleColors::ForegroundHover.GetSpecifiedColor().CopyWithNewOpacity(0.4f) :
+				FStyleColors::Foreground.GetSpecifiedColor().CopyWithNewOpacity(0.4f);
+		}
+		else
+		{
+			return bIsHoveredOrSelected ?
+				FStyleColors::ForegroundHover.GetSpecifiedColor().CopyWithNewOpacity(0.8f) :
+				FStyleColors::Foreground.GetSpecifiedColor().CopyWithNewOpacity(0.8f);
+		}
+	}
+	else
+	{
+		if (TableTreeNodePtr->IsFiltered())
+		{
+			return bIsHoveredOrSelected ?
+				FStyleColors::ForegroundHover.GetSpecifiedColor().CopyWithNewOpacity(0.5f) :
+				FStyleColors::Foreground.GetSpecifiedColor().CopyWithNewOpacity(0.5f);
+		}
+		else
+		{
+			return bIsHoveredOrSelected ?
+				FStyleColors::ForegroundHover :
+				FStyleColors::Foreground;
+		}
+	}
+}
+
+////////////////////////////////////////////////////////////////////////////////////////////////////
+
+FLinearColor STableTreeViewCell::GetShadowColorAndOpacity() const
 {
-	return ColumnPtr->GetValueFormatter()->GetCustomTooltip(*ColumnPtr, *TableTreeNodePtr);
+	const FLinearColor ShadowColor =
+		TableTreeNodePtr->IsFiltered() ?
+		FLinearColor(0.0f, 0.0f, 0.0f, 0.25f) :
+		FLinearColor(0.0f, 0.0f, 0.0f, 0.5f);
+	return ShadowColor;
 }
 
 ////////////////////////////////////////////////////////////////////////////////////////////////////
diff --git a/Engine/Source/Developer/TraceInsights/Private/Insights/Table/Widgets/STableTreeViewCell.h b/Engine/Source/Developer/TraceInsights/Private/Insights/Table/Widgets/STableTreeViewCell.h
--- a/Engine/Source/Developer/TraceInsights/Private/Insights/Table/Widgets/STableTreeViewCell.h
+++ b/Engine/Source/Developer/TraceInsights/Private/Insights/Table/Widgets/STableTreeViewCell.h
@@ -3,11 +3,13 @@
 #pragma once
 
 #include "CoreMinimal.h"
+#include "Styling/StyleColors.h"
 #include "Widgets/DeclarativeSyntaxSupport.h"
 #include "Widgets/SCompoundWidget.h"
 
 // Insights
 #include "Insights/Table/ViewModels/TableTreeNode.h"
+#include "Insights/Table/Widgets/STableTreeViewRow.h"
 
 class ITableRow;
 class IToolTip;
@@ -32,12 +34,12 @@ public:
 		SLATE_ARGUMENT(bool, IsNameColumn)
 	SLATE_END_ARGS()
 
-	void Construct(const FArguments& InArgs, const TSharedRef<ITableRow>& TableRow);
+	void Construct(const FArguments& InArgs, const TSharedRef<ITableRow>& InTableRow);
 
 protected:
-	TSharedRef<SWidget> GenerateWidgetForColumn(const FArguments& InArgs, const TSharedRef<ITableRow>& TableRow);
-	TSharedRef<SWidget> GenerateWidgetForNameColumn(const FArguments& InArgs, const TSharedRef<ITableRow>& TableRow);
-	TSharedRef<SWidget> GenerateWidgetForTableColumn(const FArguments& InArgs, const TSharedRef<ITableRow>& TableRow);
+	TSharedRef<SWidget> GenerateWidgetForColumn(const FArguments& InArgs);
+	TSharedRef<SWidget> GenerateWidgetForNameColumn(const FArguments& InArgs);
+	TSharedRef<SWidget> GenerateWidgetForTableColumn(const FArguments& InArgs);
 
 	/**
 	 * The system will use this event to notify a widget that the cursor has entered it. This event is NOT bubbled.
@@ -101,7 +103,12 @@ protected:
 		return IsHovered() ? EVisibility::Visible : EVisibility::Hidden;
 	}
 
-	TSharedPtr<class IToolTip> GetRowToolTip(const TSharedRef<ITableRow>& TableRow) const;
+	bool IsSelected() const;
+
+	const FSlateBrush* GetIcon() const
+	{
+		return TableTreeNodePtr->GetIcon();
+	}
 
 	FText GetDisplayName() const
 	{
@@ -120,47 +127,15 @@ protected:
 
 	FText GetValueAsText() const;
 
-	FSlateColor GetColorAndOpacity() const
-	{
-		const FLinearColor TextColor =
-			TableTreeNodePtr->IsFiltered() ?
-				FLinearColor(1.0f, 1.0f, 1.0f, 0.5f) :
-				FLinearColor(1.0f, 1.0f, 1.0f, 1.0f);
-		return TextColor;
-	}
-
-	FSlateColor GetExtraColorAndOpacity() const
-	{
-		const FLinearColor TextColor =
-			TableTreeNodePtr->IsFiltered() ?
-				FLinearColor(0.5f, 0.5f, 0.5f, 0.5f) :
-				FLinearColor(0.5f, 0.5f, 0.5f, 1.0f);
-		return TextColor;
-	}
-
-	FSlateColor GetStatsColorAndOpacity() const
-	{
-		const FLinearColor TextColor =
-			TableTreeNodePtr->IsGroup() ?
-				FLinearColor(0.7f, 0.7f, 0.7f, 1.0f) :
-			TableTreeNodePtr->IsFiltered() ?
-				FLinearColor(1.0f, 1.0f, 1.0f, 0.5f) :
-				FLinearColor(1.0f, 1.0f, 1.0f, 1.0f);
-		return TextColor;
-	}
-
-	FLinearColor GetShadowColorAndOpacity() const
-	{
-		const FLinearColor ShadowColor =
-			TableTreeNodePtr->IsFiltered() ?
-				FLinearColor(0.f, 0.f, 0.f, 0.25f) :
-				FLinearColor(0.0f, 0.0f, 0.0f, 0.5f);
-		return ShadowColor;
-	}
-
-	TSharedPtr<IToolTip> GetTooltip() const;
+	FSlateColor GetIconColorAndOpacity() const;
+	FSlateColor GetDisplayNameColorAndOpacity() const;
+	FSlateColor GetExtraDisplayNameColorAndOpacity() const;
+	FSlateColor GetNormalTextColorAndOpacity() const;
+	FLinearColor GetShadowColorAndOpacity() const;
 
 protected:
+	TSharedPtr<ITableRow> TableRow;
+
 	/** A shared pointer to the table view model. */
 	TSharedPtr<FTable> TablePtr;
 
diff --git a/Engine/Source/Developer/TraceInsights/Private/Insights/Table/Widgets/STableTreeViewTooltip.cpp b/Engine/Source/Developer/TraceInsights/Private/Insights/Table/Widgets/STableTreeViewTooltip.cpp
--- a/Engine/Source/Developer/TraceInsights/Private/Insights/Table/Widgets/STableTreeViewTooltip.cpp
+++ b/Engine/Source/Developer/TraceInsights/Private/Insights/Table/Widgets/STableTreeViewTooltip.cpp
@@ -23,6 +23,20 @@ namespace Insights
 
 ////////////////////////////////////////////////////////////////////////////////////////////////////
 
+static const FTextBlockStyle* GetTableTreeViewTooltipHeaderTextStyle()
+{
+	return &FInsightsStyle::Get().GetWidgetStyle<FTextBlockStyle>(TEXT("TreeTable.TooltipBold"));
+}
+
+////////////////////////////////////////////////////////////////////////////////////////////////////
+
+static const FTextBlockStyle* GetTableTreeViewTooltipNormalTextStyle()
+{
+	return &FInsightsStyle::Get().GetWidgetStyle<FTextBlockStyle>(TEXT("TreeTable.Tooltip"));
+}
+
+////////////////////////////////////////////////////////////////////////////////////////////////////
+
 BEGIN_SLATE_FUNCTION_BUILD_OPTIMIZATION
 
 TSharedPtr<SToolTip> STableTreeViewTooltip::GetTableTooltip(const FTable& Table)
@@ -38,7 +52,7 @@ TSharedPtr<SToolTip> STableTreeViewTooltip::GetTableTooltip(const FTable& Table)
 			[
 				SNew(STextBlock)
 				.Text(Table.GetDisplayName())
-				.TextStyle(FInsightsStyle::Get(), TEXT("TreeTable.TooltipBold"))
+				.TextStyle(GetTableTreeViewTooltipHeaderTextStyle())
 			]
 
 			+ SVerticalBox::Slot()
@@ -47,7 +61,7 @@ TSharedPtr<SToolTip> STableTreeViewTooltip::GetTableTooltip(const FTable& Table)
 			[
 				SNew(STextBlock)
 				.Text(Table.GetDescription())
-				.TextStyle(FInsightsStyle::Get(), TEXT("TreeTable.Tooltip"))
+				.TextStyle(GetTableTreeViewTooltipNormalTextStyle())
 			]
 		];
 
@@ -69,7 +83,7 @@ TSharedPtr<SToolTip> STableTreeViewTooltip::GetColumnTooltip(const FTableColumn&
 			[
 				SNew(STextBlock)
 				.Text(Column.GetTitleName())
-				.TextStyle(FInsightsStyle::Get(), TEXT("TreeTable.TooltipBold"))
+				.TextStyle(GetTableTreeViewTooltipHeaderTextStyle())
 			]
 
 			+ SVerticalBox::Slot()
@@ -78,7 +92,7 @@ TSharedPtr<SToolTip> STableTreeViewTooltip::GetColumnTooltip(const FTableColumn&
 			[
 				SNew(STextBlock)
 				.Text(Column.GetDescription())
-				.TextStyle(FInsightsStyle::Get(), TEXT("TreeTable.Tooltip"))
+				.TextStyle(GetTableTreeViewTooltipNormalTextStyle())
 			]
 		];
 
@@ -125,14 +139,14 @@ TSharedPtr<SToolTip> STableTreeViewTooltip::GetRowTooltip(const TSharedPtr<FTabl
 					[
 						SNew(STextBlock)
 						.Text(LOCTEXT("TT_Id", "Row:"))
-						.TextStyle(FInsightsStyle::Get(), TEXT("TreeTable.TooltipBold"))
+						.TextStyle(GetTableTreeViewTooltipHeaderTextStyle())
 					]
 					+ SGridPanel::Slot(1, 0)
 					.Padding(2.0f)
 					[
 						SNew(STextBlock)
 						.Text(FText::AsNumber(TreeNodePtr->GetRowIndex()))
-						.TextStyle(FInsightsStyle::Get(), TEXT("TreeTable.Tooltip"))
+						.TextStyle(GetTableTreeViewTooltipNormalTextStyle())
 					]
 
 					// Item Type: [Type]
@@ -141,14 +155,14 @@ TSharedPtr<SToolTip> STableTreeViewTooltip::GetRowTooltip(const TSharedPtr<FTabl
 					[
 						SNew(STextBlock)
 						.Text(LOCTEXT("TT_Type", "Item Type:"))
-						.TextStyle(FInsightsStyle::Get(), TEXT("TreeTable.TooltipBold"))
+						.TextStyle(GetTableTreeViewTooltipHeaderTextStyle())
 					]
 					+ SGridPanel::Slot(1, 1)
 					.Padding(2.0f)
 					[
 						SNew(STextBlock)
 						.Text(TreeNodePtr->IsGroup() ? LOCTEXT("TT_Type_Group", "Group Node") : LOCTEXT("TT_Type_TableRow", "Table Row"))
-						.TextStyle(FInsightsStyle::Get(), TEXT("TreeTable.Tooltip"))
+						.TextStyle(GetTableTreeViewTooltipNormalTextStyle())
 					]
 
 					// Item Name: [Name]
@@ -157,7 +171,7 @@ TSharedPtr<SToolTip> STableTreeViewTooltip::GetRowTooltip(const TSharedPtr<FTabl
 					[
 						SNew(STextBlock)
 						.Text(LOCTEXT("TT_Name", "Item Name:"))
-						.TextStyle(FInsightsStyle::Get(), TEXT("TreeTable.TooltipBold"))
+						.TextStyle(GetTableTreeViewTooltipHeaderTextStyle())
 					]
 					+ SGridPanel::Slot(1, 2)
 					.Padding(2.0f)
@@ -166,7 +180,7 @@ TSharedPtr<SToolTip> STableTreeViewTooltip::GetRowTooltip(const TSharedPtr<FTabl
 						.WrapTextAt(512.0f)
 						.WrappingPolicy(ETextWrappingPolicy::AllowPerCharacterWrapping)
 						.Text(FText::FromName(TreeNodePtr->GetName()))
-						.TextStyle(FInsightsStyle::Get(), TEXT("TreeTable.Tooltip"))
+						.TextStyle(GetTableTreeViewTooltipNormalTextStyle())
 					]
 				]
 
@@ -188,7 +202,7 @@ TSharedPtr<SToolTip> STableTreeViewTooltip::GetRowTooltip(const TSharedPtr<FTabl
 					.WrapTextAt(512.0f)
 					.WrappingPolicy(ETextWrappingPolicy::AllowPerCharacterWrapping)
 					.Text(NodeTooltip)
-					.TextStyle(FInsightsStyle::Get(), TEXT("TreeTable.Tooltip"))
+					.TextStyle(GetTableTreeViewTooltipNormalTextStyle())
 					.ColorAndOpacity(FLinearColor(0.5f, 0.5f, 0.5f, 1.0f))
 				]
 
@@ -245,7 +259,7 @@ void STableTreeViewTooltip::AddGridRow(TSharedPtr<SGridPanel> Grid, int32& Row,
 		[
 			SNew(STextBlock)
 			.Text(Name)
-			.TextStyle(FInsightsStyle::Get(), TEXT("TreeTable.TooltipBold"))
+			.TextStyle(GetTableTreeViewTooltipHeaderTextStyle())
 		];
 
 	Grid->AddSlot(1, Row)
@@ -253,7 +267,7 @@ void STableTreeViewTooltip::AddGridRow(TSharedPtr<SGridPanel> Grid, int32& Row,
 		[
 			SNew(STextBlock)
 			.Text(Value)
-			.TextStyle(FInsightsStyle::Get(), TEXT("TreeTable.Tooltip"))
+			.TextStyle(GetTableTreeViewTooltipNormalTextStyle())
 		];
 
 	Row++;
