Tasks Insights: fixed backward compatibility issue caused by removing audio and stats named threads from the engine

diff --git a/Engine/Source/Developer/TraceInsights/Private/Insights/ViewModels/ThreadTimingTrack.cpp b/Engine/Source/Developer/TraceInsights/Private/Insights/ViewModels/ThreadTimingTrack.cpp
--- a/Engine/Source/Developer/TraceInsights/Private/Insights/ViewModels/ThreadTimingTrack.cpp
+++ b/Engine/Source/Developer/TraceInsights/Private/Insights/ViewModels/ThreadTimingTrack.cpp
@@ -1176,42 +1176,35 @@ void FThreadTimingTrack::InitTooltip(FTooltipDrawState& InOutTooltip, const ITim
 				}
 
 				ENamedThreads::Type ThreadInfo = (ENamedThreads::Type)Task.ThreadToExecuteOn;
-				const TCHAR* NamedThreadsStr[] = { TEXT("RHI"), TEXT("Game"), TEXT("Rendering") };
 				ENamedThreads::Type ThreadIndex = ENamedThreads::GetThreadIndex(ThreadInfo);
 
+				auto GetTrackName = [this](uint32 InThreadId) -> FString
+				{
+					TSharedPtr<FCpuTimingTrack> Track = SharedState.GetCpuTrack(InThreadId);
+					return Track.IsValid() ? Track->GetName() : TEXT("Unknown");
+				};
+
 				if (ThreadIndex == ENamedThreads::AnyThread)
 				{
 					const TCHAR* TaskPri = ENamedThreads::GetTaskPriority(ThreadInfo) == ENamedThreads::NormalTaskPriority ? TEXT("Normal") : TEXT("High");
 
 					int32 ThreadPriIndex = ENamedThreads::GetThreadPriorityIndex(ThreadInfo);
 					const TCHAR* ThreadPriStrs[] = { TEXT("Normal"), TEXT("High"), TEXT("Low") };
-					const TCHAR* ThreadPri = ThreadPriStrs[ThreadPriIndex];
+					const TCHAR* ThreadPri = ensure(ThreadPriIndex >= 0 && ThreadPriIndex < 3) ? ThreadPriStrs[ThreadPriIndex] : TEXT("Unknown");
 
-					InOutTooltip.AddTextLine(FString::Printf(TEXT("%s Pri task on %s Pri worker"), TaskPri, ThreadPri), FLinearColor::Green);
+					InOutTooltip.AddTextLine(FString::Printf(TEXT("%s Pri task on %s Pri worker (%s)"), TaskPri, ThreadPri, *GetTrackName(Task.StartedThreadId)), FLinearColor::Green);
 				}
 				else
 				{
 					const TCHAR* QueueStr = ENamedThreads::GetQueueIndex(ThreadInfo) == ENamedThreads::MainQueue ? TEXT("Main") : TEXT("Local");
-					InOutTooltip.AddTextLine(FString::Printf(TEXT("%s (%s queue)"), NamedThreadsStr[ThreadIndex], QueueStr), FLinearColor::Green);
+					InOutTooltip.AddTextLine(FString::Printf(TEXT("%s (%s queue)"), *GetTrackName(Task.StartedThreadId), QueueStr), FLinearColor::Green);
 				}
 
-				{
-					TSharedPtr<FCpuTimingTrack> Track = SharedState.GetCpuTrack(Task.CreatedThreadId);
-					FString TrackName = Track.IsValid() ? Track->GetName() : TEXT("Unknown");
-					InOutTooltip.AddNameValueTextLine(TEXT("Created:"), FString::Printf(TEXT("%f on %s"), Task.CreatedTimestamp, *TrackName));
-				}
+				InOutTooltip.AddNameValueTextLine(TEXT("Created:"), FString::Printf(TEXT("%f on %s"), Task.CreatedTimestamp, *GetTrackName(Task.CreatedThreadId)));
+				
+				InOutTooltip.AddNameValueTextLine(TEXT("Launched:"), FString::Printf(TEXT("%f (+%s) on %s"), Task.LaunchedTimestamp, *TimeUtils::FormatTimeAuto(Task.LaunchedTimestamp - Task.CreatedTimestamp), *GetTrackName(Task.LaunchedThreadId)));
 
-				{
-					TSharedPtr<FCpuTimingTrack> Track = SharedState.GetCpuTrack(Task.LaunchedThreadId);
-					FString TrackName = Track.IsValid() ? Track->GetName() : TEXT("Unknown");
-					InOutTooltip.AddNameValueTextLine(TEXT("Launched:"), FString::Printf(TEXT("%f (+%s) on %s"), Task.LaunchedTimestamp, *TimeUtils::FormatTimeAuto(Task.LaunchedTimestamp - Task.CreatedTimestamp), *TrackName));
-				}
-
-				{
-					TSharedPtr<FCpuTimingTrack> Track = SharedState.GetCpuTrack(Task.ScheduledThreadId);
-					FString TrackName = Track.IsValid() ? Track->GetName() : TEXT("Unknown");
-					InOutTooltip.AddNameValueTextLine(TEXT("Scheduled:"), FString::Printf(TEXT("%f (+%s) on %s"), Task.ScheduledTimestamp, *TimeUtils::FormatTimeAuto(Task.ScheduledTimestamp - Task.LaunchedTimestamp), *TrackName));
-				}
+				InOutTooltip.AddNameValueTextLine(TEXT("Scheduled:"), FString::Printf(TEXT("%f (+%s) on %s"), Task.ScheduledTimestamp, *TimeUtils::FormatTimeAuto(Task.ScheduledTimestamp - Task.LaunchedTimestamp), *GetTrackName(Task.ScheduledThreadId)));
 
 				InOutTooltip.AddNameValueTextLine(TEXT("Started:"), FString::Printf(TEXT("%f (+%s)"), Task.StartedTimestamp, *TimeUtils::FormatTimeAuto(Task.StartedTimestamp - Task.ScheduledTimestamp)));
 				if (Task.FinishedTimestamp != TraceServices::FTaskInfo::InvalidTimestamp)
@@ -1220,15 +1213,11 @@ void FThreadTimingTrack::InitTooltip(FTooltipDrawState& InOutTooltip, const ITim
 
 					if (Task.CompletedTimestamp != TraceServices::FTaskInfo::InvalidTimestamp)
 					{
-						TSharedPtr<FCpuTimingTrack> CompletedTrack = SharedState.GetCpuTrack(Task.CompletedThreadId);
-						FString CompletedTrackName = CompletedTrack.IsValid() ? CompletedTrack->GetName() : TEXT("Unknown");
-						InOutTooltip.AddNameValueTextLine(TEXT("Completed:"), FString::Printf(TEXT("%f (+%s) on %s"), Task.CompletedTimestamp, *TimeUtils::FormatTimeAuto(Task.CompletedTimestamp - Task.FinishedTimestamp), *CompletedTrackName));
+						InOutTooltip.AddNameValueTextLine(TEXT("Completed:"), FString::Printf(TEXT("%f (+%s) on %s"), Task.CompletedTimestamp, *TimeUtils::FormatTimeAuto(Task.CompletedTimestamp - Task.FinishedTimestamp), *GetTrackName(Task.CompletedThreadId)));
 
 						if (Task.DestroyedTimestamp != TraceServices::FTaskInfo::InvalidTimestamp)
 						{
-							TSharedPtr<FCpuTimingTrack> DestroyedTrack = SharedState.GetCpuTrack(Task.DestroyedThreadId);
-							FString DestroyedTrackName = DestroyedTrack.IsValid() ? DestroyedTrack->GetName() : TEXT("Unknown");
-							InOutTooltip.AddNameValueTextLine(TEXT("Destroyed:"), FString::Printf(TEXT("%f (+%s) on %s"), Task.DestroyedTimestamp, *TimeUtils::FormatTimeAuto(Task.DestroyedTimestamp - Task.CompletedTimestamp), *DestroyedTrackName));
+							InOutTooltip.AddNameValueTextLine(TEXT("Destroyed:"), FString::Printf(TEXT("%f (+%s) on %s"), Task.DestroyedTimestamp, *TimeUtils::FormatTimeAuto(Task.DestroyedTimestamp - Task.CompletedTimestamp), *GetTrackName(Task.DestroyedThreadId)));
 						}
 					}
 				}
