Tasks Insights: added "launched" task table selection, modified "alive" to include any task that was alive at any moment of selected period

diff --git a/Engine/Source/Developer/TraceInsights/Private/Insights/TaskGraphProfiler/Widgets/STaskTableTreeView.cpp b/Engine/Source/Developer/TraceInsights/Private/Insights/TaskGraphProfiler/Widgets/STaskTableTreeView.cpp
--- a/Engine/Source/Developer/TraceInsights/Private/Insights/TaskGraphProfiler/Widgets/STaskTableTreeView.cpp
+++ b/Engine/Source/Developer/TraceInsights/Private/Insights/TaskGraphProfiler/Widgets/STaskTableTreeView.cpp
@@ -401,6 +401,10 @@ TSharedRef<SWidget> STaskTableTreeView::TasksSelectionOptions_OnGenerateWidget(T
 		{
 			return LOCTEXT("AliveTooltip", "Tasks that were created before the selection and destroyed after.");
 		}
+		case ETaskEnumerationOption::Launched:
+		{
+			return LOCTEXT("LaunchedTooltip", "Tasks that were launched during the selection.");
+		}
 		case ETaskEnumerationOption::Active:
 		{
 			return LOCTEXT("ActiveTooltip", "Tasks that were active (being executed) at any moment of the selection.");
@@ -472,6 +476,7 @@ const TArray<TSharedPtr<ETaskEnumerationOption>>* STaskTableTreeView::GetAvailab
 	if (AvailableTasksSelectionOptions.Num() == 0)
 	{
 		AvailableTasksSelectionOptions.Add(MakeShared<ETaskEnumerationOption>(ETaskEnumerationOption::Alive));
+		AvailableTasksSelectionOptions.Add(MakeShared<ETaskEnumerationOption>(ETaskEnumerationOption::Launched));
 		AvailableTasksSelectionOptions.Add(MakeShared<ETaskEnumerationOption>(ETaskEnumerationOption::Active));
 		AvailableTasksSelectionOptions.Add(MakeShared<ETaskEnumerationOption>(ETaskEnumerationOption::WaitingForPrerequisites));
 		AvailableTasksSelectionOptions.Add(MakeShared<ETaskEnumerationOption>(ETaskEnumerationOption::Queued));
@@ -580,6 +585,10 @@ FText STaskTableTreeView::TasksSelectionOptions_GetText(ETaskEnumerationOption I
 	{
 		return LOCTEXT("Alive", "Alive");
 	}
+	case ETaskEnumerationOption::Launched:
+	{
+		return LOCTEXT("Launched", "Launched");
+	}
 	case ETaskEnumerationOption::Active:
 	{
 		return LOCTEXT("Active", "Active");
diff --git a/Engine/Source/Developer/TraceServices/Private/Model/TasksProfiler.cpp b/Engine/Source/Developer/TraceServices/Private/Model/TasksProfiler.cpp
--- a/Engine/Source/Developer/TraceServices/Private/Model/TasksProfiler.cpp
+++ b/Engine/Source/Developer/TraceServices/Private/Model/TasksProfiler.cpp
@@ -606,8 +606,13 @@ namespace TraceServices
 			// Alive
 			[](double StartTime, double EndTime, const FTaskInfo& Task)
 			{	// created before `StartTime` and destroyed after `EndTime`
-				return (Task.CreatedTimestamp <= StartTime || Task.CreatedTimestamp == FTaskInfo::InvalidTimestamp) &&
-					(Task.DestroyedTimestamp >= EndTime || Task.DestroyedTimestamp == FTaskInfo::InvalidTimestamp);
+				return (Task.CreatedTimestamp <= EndTime || Task.CreatedTimestamp == FTaskInfo::InvalidTimestamp) &&
+					(Task.DestroyedTimestamp >= StartTime || Task.DestroyedTimestamp == FTaskInfo::InvalidTimestamp);
+			},
+			// Launched
+			[](double StartTime, double EndTime, const FTaskInfo& Task)
+			{	// launched between `StartTime` and `EndTime`
+				return Task.LaunchedTimestamp >= StartTime && Task.LaunchedTimestamp <= EndTime;
 			},
 			// Active
 			[](double StartTime, double EndTime, const FTaskInfo& Task)
diff --git a/Engine/Source/Developer/TraceServices/Public/TraceServices/Model/TasksProfiler.h b/Engine/Source/Developer/TraceServices/Public/TraceServices/Model/TasksProfiler.h
--- a/Engine/Source/Developer/TraceServices/Public/TraceServices/Model/TasksProfiler.h
+++ b/Engine/Source/Developer/TraceServices/Public/TraceServices/Model/TasksProfiler.h
@@ -87,6 +87,7 @@ struct FWaitingForTasks
 enum class ETaskEnumerationOption
 {
 	Alive,
+	Launched,
 	Active,
 	WaitingForPrerequisites,
 	Queued,
