[Insights] Fix Frame Stats setting persistence issues.

diff --git a/Engine/Source/Developer/TraceInsights/Private/Insights/Widgets/STimersView.cpp b/Engine/Source/Developer/TraceInsights/Private/Insights/Widgets/STimersView.cpp
--- a/Engine/Source/Developer/TraceInsights/Private/Insights/Widgets/STimersView.cpp
+++ b/Engine/Source/Developer/TraceInsights/Private/Insights/Widgets/STimersView.cpp
@@ -419,6 +419,9 @@ void STimersView::Construct(const FArguments& InArgs)
 	InitializeAndShowHeaderColumns();
 	LoadSettings();
 
+	Aggregator->SetFrameType(ModeFrameType);
+	SetTimingViewFrameType();
+
 	// Create the search filters: text based, type based etc.
 	TextFilter = MakeShared<FTimerNodeTextFilter>(FTimerNodeTextFilter::FItemToStringArray::CreateSP(this, &STimersView::HandleItemToStringArray));
 	Filters = MakeShared<FTimerNodeFilterCollection>();
@@ -1751,15 +1754,7 @@ void STimersView::Mode_OnSelectionChanged(TSharedPtr<ETraceFrameType> NewFrameTy
 	ModeFrameType = *NewFrameType;
 	LoadVisibleColumnsSettings();
 
-	TSharedPtr<STimingProfilerWindow> Window = FTimingProfilerManager::Get()->GetProfilerWindow();
-	if (Window.IsValid())
-	{
-		TSharedPtr<STimingView> TimingView = Window->GetTimingView();
-		if (TimingView.IsValid())
-		{
-			TimingView->SetFrameTypeToSnapTo(ModeFrameType);
-		}
-	}
+	SetTimingViewFrameType();
 
 	Aggregator->SetFrameType(ModeFrameType);
 	Aggregator->Cancel();
@@ -3185,6 +3180,21 @@ FReply STimersView::OnKeyDown(const FGeometry& MyGeometry, const FKeyEvent& InKe
 
 ////////////////////////////////////////////////////////////////////////////////////////////////////
 
+void STimersView::SetTimingViewFrameType()
+{
+	TSharedPtr<STimingProfilerWindow> Window = FTimingProfilerManager::Get()->GetProfilerWindow();
+	if (Window.IsValid())
+	{
+		TSharedPtr<STimingView> TimingView = Window->GetTimingView();
+		if (TimingView.IsValid())
+		{
+			TimingView->SelectTimeInterval(TimingView->GetSelectionStartTime(), TimingView->GetSelectionEndTime() - TimingView->GetSelectionStartTime());
+		}
+	}
+}
+
+////////////////////////////////////////////////////////////////////////////////////////////////////
+
 void STimersView::SaveVisibleColumnsSettings()
 {
 	FInsightsSettings& Settings = FInsightsManager::Get()->GetSettings();
diff --git a/Engine/Source/Developer/TraceInsights/Private/Insights/Widgets/STimersView.h b/Engine/Source/Developer/TraceInsights/Private/Insights/Widgets/STimersView.h
--- a/Engine/Source/Developer/TraceInsights/Private/Insights/Widgets/STimersView.h
+++ b/Engine/Source/Developer/TraceInsights/Private/Insights/Widgets/STimersView.h
@@ -89,6 +89,8 @@ public:
 
 	virtual FReply OnKeyDown(const FGeometry& MyGeometry, const FKeyEvent& InKeyEvent) override;
 
+	ETraceFrameType GetFrameTypeMode() { return ModeFrameType; }
+
 private:
 	void InitCommandList();
 
@@ -334,6 +336,8 @@ private:
 	void LoadSettings();
 	void LoadVisibleColumnsSettings();
 
+	void SetTimingViewFrameType();
+
 private:
 	/** Table view model. */
 	TSharedPtr<Insights::FTable> Table;
diff --git a/Engine/Source/Developer/TraceInsights/Private/Insights/Widgets/STimingView.cpp b/Engine/Source/Developer/TraceInsights/Private/Insights/Widgets/STimingView.cpp
--- a/Engine/Source/Developer/TraceInsights/Private/Insights/Widgets/STimingView.cpp
+++ b/Engine/Source/Developer/TraceInsights/Private/Insights/Widgets/STimingView.cpp
@@ -470,7 +470,6 @@ void STimingView::Reset(bool bIsFirstReset)
 	Tooltip.Reset();
 
 	LastSelectionType = ESelectionType::None;
-	FrameTypeToSnapTo = ETraceFrameType::TraceFrameType_Count;
 
 	//ThisGeometry
 
@@ -3443,7 +3442,7 @@ void STimingView::SelectTimeInterval(double IntervalStartTime, double IntervalDu
 	SelectionStartTime = IntervalStartTime;
 	SelectionEndTime = IntervalStartTime + IntervalDuration;
 
-	if (FrameTypeToSnapTo != ETraceFrameType::TraceFrameType_Count)
+	if (GetFrameTypeToSnapTo() != ETraceFrameType::TraceFrameType_Count)
 	{
 		SnapToFrameBound(SelectionStartTime, SelectionEndTime);
 	}
@@ -3466,6 +3465,7 @@ void STimingView::SnapToFrameBound(double& StartTime, double& EndTime)
 	TraceServices::FAnalysisSessionReadScope SessionReadScope(*Session.Get());
 	const TraceServices::IFrameProvider& FramesProvider = TraceServices::ReadFrameProvider(*Session.Get());
 
+	ETraceFrameType FrameTypeToSnapTo = GetFrameTypeToSnapTo();
 	uint32 FrameNum = FramesProvider.GetFrameNumberForTimestamp(FrameTypeToSnapTo, StartTime);
 	const TraceServices::FFrame* StartFrame = FramesProvider.GetFrame(FrameTypeToSnapTo, FrameNum);
 
@@ -3512,21 +3512,6 @@ void STimingView::SnapToFrameBound(double& StartTime, double& EndTime)
 
 ////////////////////////////////////////////////////////////////////////////////////////////////////
 
-void STimingView::SetFrameTypeToSnapTo(ETraceFrameType InFrameType)
-{
-	if (FrameTypeToSnapTo == InFrameType)
-	{
-		return;
-	}
-
-	FrameTypeToSnapTo = InFrameType;
-
-	SelectTimeInterval(SelectionStartTime, SelectionEndTime - SelectionStartTime);
-}
-
-////////////////////////////////////////////////////////////////////////////////////////////////////
-
-
 void STimingView::RaiseSelectionChanging()
 {
 	OnSelectionChangedDelegate.Broadcast(Insights::ETimeChangedFlags::Interactive, SelectionStartTime, SelectionEndTime);
@@ -5288,5 +5273,23 @@ void STimingView::EnumerateFilteredTracks(TSharedPtr<Insights::FFilterConfigurat
 
 ////////////////////////////////////////////////////////////////////////////////////////////////////
 
+ETraceFrameType STimingView::GetFrameTypeToSnapTo()
+{
+	TSharedPtr<STimingProfilerWindow> Window = FTimingProfilerManager::Get()->GetProfilerWindow();
+	if (Window.IsValid())
+	{
+		TSharedPtr<STimersView> TimersView = Window->GetTimersView();
+		if (TimersView.IsValid())
+		{
+			return TimersView->GetFrameTypeMode();
+		}
+	}
+
+	// TraceFrameType_Count is the Instance mode.
+	return ETraceFrameType::TraceFrameType_Count;
+}
+
+////////////////////////////////////////////////////////////////////////////////////////////////////
+
 #undef INSIGHTS_ACTIVATE_BENCHMARK
 #undef LOCTEXT_NAMESPACE
diff --git a/Engine/Source/Developer/TraceInsights/Private/Insights/Widgets/STimingView.h b/Engine/Source/Developer/TraceInsights/Private/Insights/Widgets/STimingView.h
--- a/Engine/Source/Developer/TraceInsights/Private/Insights/Widgets/STimingView.h
+++ b/Engine/Source/Developer/TraceInsights/Private/Insights/Widgets/STimingView.h
@@ -348,9 +348,6 @@ public:
 	void SnapToFrameBound(double& IntervalStartTime, double& IntervalDuration);
 	void SelectToTimeMarker(double InTimeMarker);
 
-	void SetFrameTypeToSnapTo(ETraceFrameType InFrameType);
-	ETraceFrameType GetFrameTypeToSnapTo() { return FrameTypeToSnapTo; }
-
 	//bool AreTimeMarkersVisible() { return MarkersTrack->IsVisible(); }
 	void SetTimeMarkersVisible(bool bOnOff);
 	//bool IsDrawOnlyBookmarksEnabled() { return MarkersTrack->IsBookmarksTrack(); }
@@ -515,6 +512,8 @@ protected:
 	typedef TFunctionRef<void(TSharedPtr<FBaseTimingTrack>& Track)> EnumerateFilteredTracksCallback;
 	void EnumerateFilteredTracks(TSharedPtr<Insights::FFilterConfigurator> FilterConfigurator, EnumerateFilteredTracksCallback Callback);
 
+	ETraceFrameType GetFrameTypeToSnapTo();
+
 protected:
 	/** The name of the view. */
 	FName ViewName;
@@ -681,8 +680,6 @@ protected:
 	};
 	ESelectionType LastSelectionType;
 
-	ETraceFrameType FrameTypeToSnapTo;
-
 	/** Throttle flag, allowing tracks to control whether Slate throttle should take place */
 	bool bPreventThrottling;
 
