[Insights] Removed background image from SSessionInfoWindow (and a leaked slate color brush).

diff --git a/Engine/Source/Developer/TraceInsights/Private/Insights/Widgets/SSessionInfoWindow.cpp b/Engine/Source/Developer/TraceInsights/Private/Insights/Widgets/SSessionInfoWindow.cpp
--- a/Engine/Source/Developer/TraceInsights/Private/Insights/Widgets/SSessionInfoWindow.cpp
+++ b/Engine/Source/Developer/TraceInsights/Private/Insights/Widgets/SSessionInfoWindow.cpp
@@ -8,7 +8,6 @@
 #include "TraceServices/Model/Diagnostics.h"
 #include "TraceServices/ModuleService.h"
 #include "Misc/Paths.h"
-#include "Widgets/Images/SImage.h"
 #include "Widgets/Input/SEditableTextBox.h"
 #include "Widgets/Input/SMultiLineEditableTextBox.h"
 #include "Widgets/Layout/SScrollBox.h"
@@ -255,26 +254,17 @@ void SSessionInfoWindow::AddSimpleInfoLine(TSharedPtr<SVerticalBox> InVerticalBo
 
 TSharedRef<SDockTab> SSessionInfoWindow::SpawnTab_SessionInfo(const FSpawnTabArgs& Args)
 {
-	TSharedPtr<SVerticalBox> VerticalBox;
-	TSharedPtr<SImage> Image;
-
 	TSharedPtr<SScrollBar> VScrollbar;
 	SAssignNew(VScrollbar, SScrollBar)
 	.Orientation(Orient_Vertical);
 
+	TSharedPtr<SVerticalBox> VerticalBox;
+
 	const TSharedRef<SDockTab> DockTab = SNew(SDockTab)
 		.TabRole(ETabRole::PanelTab)
 		[
 			SNew(SOverlay)
 
-			// Overlay slot for the Background image
-			+ SOverlay::Slot()
-			.HAlign(HAlign_Fill)
-			.VAlign(VAlign_Fill)
-			[
-				SAssignNew(Image, SImage)
-			]
-
 			// Overlay slot for the ScrollBox containing the data
 			+ SOverlay::Slot()
 			.HAlign(HAlign_Fill)
@@ -303,8 +293,6 @@ TSharedRef<SDockTab> SSessionInfoWindow::SpawnTab_SessionInfo(const FSpawnTabArg
 			]
 		];
 
-	Image->SetImage(new FSlateColorBrush(FLinearColor(0.015f, 0.015f, 0.015f, 1.0f)));
-
 	BeginSection(VerticalBox, LOCTEXT("SessionInfo_SectionText", "Session Info"));
 	AddInfoLine(VerticalBox, LOCTEXT("SessionName_HeaderText",	"Session Name"),		TAttribute<FText>(this, &SSessionInfoWindow::GetSessionNameText));
 	AddInfoLine(VerticalBox, LOCTEXT("Uri_HeaderText",			"URI"),					TAttribute<FText>(this, &SSessionInfoWindow::GetUriText));
