[Insights]
  - Refactoring on STableTreeView in preparation for a larger change re filtering (and async operations in general).
  - Moved SSessionTableTreeView class into a separate file.
  - Filtering: Fixed typo in FFilterConfigurator::OnChangesCommittedEvent. Also fixed LOCTEXT namespaces to avoid conflict with AssetManagerEditor plugin.
  - Added Initialize()/Shutdown() and UnregisterTabSpawner() to FFilterService API.
  - Fixed FInsightsManager's API that uses StoreClient to use the corresponding lock / critical section.

diff --git a/Engine/Source/Developer/TraceInsights/Private/Insights/CookProfiler/Widgets/SPackageTableTreeView.h b/Engine/Source/Developer/TraceInsights/Private/Insights/CookProfiler/Widgets/SPackageTableTreeView.h
--- a/Engine/Source/Developer/TraceInsights/Private/Insights/CookProfiler/Widgets/SPackageTableTreeView.h
+++ b/Engine/Source/Developer/TraceInsights/Private/Insights/CookProfiler/Widgets/SPackageTableTreeView.h
@@ -7,7 +7,7 @@
 
 // Insights
 #include "Insights/CookProfiler/ViewModels/PackageTable.h"
-#include "Insights/Table/Widgets/STableTreeView.h"
+#include "Insights/Table/Widgets/SSessionTableTreeView.h"
 
 class FMenuBuilder;
 
diff --git a/Engine/Source/Developer/TraceInsights/Private/Insights/InsightsManager.cpp b/Engine/Source/Developer/TraceInsights/Private/Insights/InsightsManager.cpp
--- a/Engine/Source/Developer/TraceInsights/Private/Insights/InsightsManager.cpp
+++ b/Engine/Source/Developer/TraceInsights/Private/Insights/InsightsManager.cpp
@@ -163,7 +163,7 @@ void FInsightsManager::Initialize(IUnrealInsightsModule& InsightsModule)
 
 	InsightsMenuBuilder = MakeShared<FInsightsMenuBuilder>();
 
-	Insights::FFilterService::CreateInstance();
+	Insights::FFilterService::Initialize();
 
 	FMessageLogModule& MessageLogModule = FModuleManager::LoadModuleChecked<FMessageLogModule>("MessageLog");
 	MessageLogModule.RegisterLogListing(GetLogListingName(), LOCTEXT("UnrealInsights", "Unreal Insights"));
@@ -204,6 +204,8 @@ void FInsightsManager::Shutdown()
 		}
 	}
 
+	Insights::FFilterService::Shutdown();
+
 	FInsightsManager::Instance.Reset();
 }
 
@@ -439,7 +441,7 @@ void FInsightsManager::OnSessionInfoTabClosed(TSharedRef<SDockTab> TabBeingClose
 
 ////////////////////////////////////////////////////////////////////////////////////////////////////
 
-FString FInsightsManager::GetStoreDir() 
+FString FInsightsManager::GetStoreDir()
 {
 	using namespace UE::Trace;
 	FScopeLock _(&StoreClientCriticalSection);
@@ -458,11 +460,11 @@ bool FInsightsManager::ConnectToStore(const TCHAR* Host, uint32 Port)
 	{
 		return false;
 	}
-	
+
 	LastStoreHost = Host;
 	LastStorePort = Port;
-	bCanChangeStoreSettings = LastStoreHost.Equals(TEXT("localhost"), ESearchCase::IgnoreCase)
-                              		|| LastStoreHost.Equals(TEXT("127.0.0.1"));
+	bCanChangeStoreSettings = LastStoreHost.Equals(TEXT("localhost"), ESearchCase::IgnoreCase) ||
+	                          LastStoreHost.Equals(TEXT("127.0.0.1"));
 
 	return true;
 }
@@ -471,6 +473,11 @@ bool FInsightsManager::ConnectToStore(const TCHAR* Host, uint32 Port)
 
 bool FInsightsManager::ReconnectToStore() const
 {
+	if (!StoreClient.IsValid())
+	{
+		return false;
+	}
+
 	return StoreClient->Reconnect(*LastStoreHost, LastStorePort);
 }
 
@@ -525,7 +532,7 @@ bool FInsightsManager::Tick(float DeltaTime)
 			RetryLoadLastLiveSessionTimer -= DeltaTime;
 		}
 	}
-	
+
 	AutoLoadLiveSession();
 
 	UpdateSessionDuration();
@@ -844,21 +851,30 @@ void FInsightsManager::AutoLoadLiveSession()
 		return;
 	}
 
-	const uint32 SessionCount = StoreClient->GetSessionCount();
-	for (uint32 SessionIndex = 0; SessionIndex < SessionCount; ++SessionIndex)
+	uint32 AutoLoadTraceId = 0;
 	{
-		const UE::Trace::FStoreClient::FSessionInfo* SessionInfo = StoreClient->GetSessionInfo(SessionIndex);
-		if (SessionInfo)
+		FScopeLock _(&StoreClientCriticalSection);
+		const uint32 SessionCount = StoreClient->GetSessionCount();
+		for (uint32 SessionIndex = 0; SessionIndex < SessionCount; ++SessionIndex)
 		{
-			const uint32 TraceId = SessionInfo->GetTraceId();
-			if (TraceId != CurrentTraceId && !AutoLoadedTraceIds.Contains(TraceId))
+			const UE::Trace::FStoreClient::FSessionInfo* SessionInfo = StoreClient->GetSessionInfo(SessionIndex);
+			if (SessionInfo)
 			{
-				AutoLoadedTraceIds.Add(TraceId);
-				LoadTrace(SessionInfo->GetTraceId());
-				break;
+				const uint32 TraceId = SessionInfo->GetTraceId();
+				if (TraceId != CurrentTraceId && !AutoLoadedTraceIds.Contains(TraceId))
+				{
+					AutoLoadTraceId = TraceId;
+					break;
+				}
 			}
 		}
 	}
+
+	if (AutoLoadTraceId != 0)
+	{
+		AutoLoadedTraceIds.Add(AutoLoadTraceId);
+		LoadTrace(AutoLoadTraceId);
+	}
 }
 
 ////////////////////////////////////////////////////////////////////////////////////////////////////
@@ -872,16 +888,25 @@ void FInsightsManager::LoadLastLiveSession(float RetryTime)
 		return;
 	}
 
-	const uint32 SessionCount = StoreClient->GetSessionCount();
-	if (SessionCount != 0)
+	uint32 LastLiveSessionTraceId = 0;
 	{
-		const UE::Trace::FStoreClient::FSessionInfo* SessionInfo = StoreClient->GetSessionInfo(SessionCount - 1);
-		if (SessionInfo)
+		FScopeLock _(&StoreClientCriticalSection);
+		const uint32 SessionCount = StoreClient->GetSessionCount();
+		if (SessionCount != 0)
 		{
-			LoadTrace(SessionInfo->GetTraceId());
+			const UE::Trace::FStoreClient::FSessionInfo* SessionInfo = StoreClient->GetSessionInfo(SessionCount - 1);
+			if (SessionInfo)
+			{
+				LastLiveSessionTraceId = SessionInfo->GetTraceId();
+			}
 		}
 	}
 
+	if (LastLiveSessionTraceId)
+	{
+		LoadTrace(LastLiveSessionTraceId);
+	}
+
 	if (Session == nullptr && RetryTime > 0)
 	{
 		RetryLoadLastLiveSessionTimer = RetryTime;
@@ -903,6 +928,8 @@ void FInsightsManager::LoadTrace(uint32 InTraceId, bool InAutoQuit)
 		return;
 	}
 
+	FScopeLock StoreClientLock(&StoreClientCriticalSection);
+
 	UE::Trace::FStoreClient::FTraceData TraceData = StoreClient->ReadTrace(InTraceId);
 	if (!TraceData)
 	{
@@ -927,6 +954,8 @@ void FInsightsManager::LoadTrace(uint32 InTraceId, bool InAutoQuit)
 		FPaths::NormalizeFilename(TraceName);
 	}
 
+	StoreClientLock.Unlock();
+
 	Session = AnalysisService->StartAnalysis(InTraceId, *TraceName, MoveTemp(TraceData));
 
 	if (Session)
@@ -1123,7 +1152,7 @@ void FInsightsManager::OnSessionAnalysisCompleted()
 			Log->Notify(NSLOCTEXT("TraceAnalysis", "Finished", "Analysis finished."));
 		}
 	}
-	
+
 	if (!SessionAnalysisCompletedCmd.IsEmpty())
 	{
 		FOutputDevice& Ar = *GLog;
diff --git a/Engine/Source/Developer/TraceInsights/Private/Insights/MemoryProfiler/Widgets/SMemAllocTableTreeView.h b/Engine/Source/Developer/TraceInsights/Private/Insights/MemoryProfiler/Widgets/SMemAllocTableTreeView.h
--- a/Engine/Source/Developer/TraceInsights/Private/Insights/MemoryProfiler/Widgets/SMemAllocTableTreeView.h
+++ b/Engine/Source/Developer/TraceInsights/Private/Insights/MemoryProfiler/Widgets/SMemAllocTableTreeView.h
@@ -8,7 +8,7 @@
 // Insights
 #include "Insights/Common/Stopwatch.h"
 #include "Insights/MemoryProfiler/ViewModels/MemAllocTable.h"
-#include "Insights/Table/Widgets/STableTreeView.h"
+#include "Insights/Table/Widgets/SSessionTableTreeView.h"
 
 namespace Insights
 {
diff --git a/Engine/Source/Developer/TraceInsights/Private/Insights/Table/Widgets/SSessionTableTreeView.cpp b/Engine/Source/Developer/TraceInsights/Private/Insights/Table/Widgets/SSessionTableTreeView.cpp
new file mode 100644
--- /dev/null
+++ b/Engine/Source/Developer/TraceInsights/Private/Insights/Table/Widgets/SSessionTableTreeView.cpp
@@ -0,0 +1,64 @@
+// Copyright Epic Games, Inc. All Rights Reserved.
+
+#include "SSessionTableTreeView.h"
+
+// Insights
+#include "Insights/InsightsManager.h"
+
+#define LOCTEXT_NAMESPACE "Insights::SSessionTableTreeView"
+
+namespace Insights
+{
+
+////////////////////////////////////////////////////////////////////////////////////////////////////
+// SSessionTableTreeView
+////////////////////////////////////////////////////////////////////////////////////////////////////
+
+SSessionTableTreeView::~SSessionTableTreeView()
+{
+	// Remove ourselves from the Insights manager.
+	if (FInsightsManager::Get().IsValid())
+	{
+		FInsightsManager::Get()->GetSessionChangedEvent().RemoveAll(this);
+	}
+
+	Session.Reset();
+}
+
+////////////////////////////////////////////////////////////////////////////////////////////////////
+
+void SSessionTableTreeView::ConstructWidget(TSharedPtr<FTable> InTablePtr)
+{
+	STableTreeView::ConstructWidget(InTablePtr);
+
+	// Register ourselves with the Insights manager.
+	FInsightsManager::Get()->GetSessionChangedEvent().AddSP(this, &SSessionTableTreeView::InsightsManager_OnSessionChanged);
+
+	// Update the Session (i.e. when analysis session was already started).
+	InsightsManager_OnSessionChanged();
+
+	CreateGroupings();
+	CreateSortings();
+}
+
+////////////////////////////////////////////////////////////////////////////////////////////////////
+
+void SSessionTableTreeView::InsightsManager_OnSessionChanged()
+{
+	TSharedPtr<const TraceServices::IAnalysisSession> NewSession = FInsightsManager::Get()->GetSession();
+	if (NewSession != Session)
+	{
+		Session = NewSession;
+		Reset();
+	}
+	else
+	{
+		UpdateTree();
+	}
+}
+
+////////////////////////////////////////////////////////////////////////////////////////////////////
+
+} // namespace Insights
+
+#undef LOCTEXT_NAMESPACE
diff --git a/Engine/Source/Developer/TraceInsights/Private/Insights/Table/Widgets/SSessionTableTreeView.h b/Engine/Source/Developer/TraceInsights/Private/Insights/Table/Widgets/SSessionTableTreeView.h
new file mode 100644
--- /dev/null
+++ b/Engine/Source/Developer/TraceInsights/Private/Insights/Table/Widgets/SSessionTableTreeView.h
@@ -0,0 +1,39 @@
+// Copyright Epic Games, Inc. All Rights Reserved.
+
+#pragma once
+
+#include "CoreMinimal.h"
+
+// Insights
+#include "Insights/Table/Widgets/STableTreeView.h"
+
+namespace TraceServices
+{
+	class IAnalysisSession;
+}
+
+namespace Insights
+{
+
+class FTable;
+
+////////////////////////////////////////////////////////////////////////////////////////////////////
+
+class SSessionTableTreeView : public STableTreeView
+{
+public:
+	virtual ~SSessionTableTreeView() override;
+
+protected:
+	virtual void ConstructWidget(TSharedPtr<FTable> InTablePtr) override;
+
+	/** Called when the analysis session has changed. */
+	void InsightsManager_OnSessionChanged();
+
+	/** The analysis session used to populate this widget. */
+	TSharedPtr<const TraceServices::IAnalysisSession> Session;
+};
+
+////////////////////////////////////////////////////////////////////////////////////////////////////
+
+} // namespace Insights
diff --git a/Engine/Source/Developer/TraceInsights/Private/Insights/Table/Widgets/STableTreeView.cpp b/Engine/Source/Developer/TraceInsights/Private/Insights/Table/Widgets/STableTreeView.cpp
--- a/Engine/Source/Developer/TraceInsights/Private/Insights/Table/Widgets/STableTreeView.cpp
+++ b/Engine/Source/Developer/TraceInsights/Private/Insights/Table/Widgets/STableTreeView.cpp
@@ -9,26 +9,19 @@
 #include "Framework/MultiBox/MultiBoxBuilder.h"
 #include "HAL/PlatformApplicationMisc.h"
 #include "HAL/PlatformFileManager.h"
-#include "Layout/WidgetPath.h"
 #include "Logging/MessageLog.h"
 #include "SlateOptMacros.h"
-#include "Widgets/Input/SCheckBox.h"
-#include "Widgets/Layout/SScrollBox.h"
 #include "Widgets/Input/SSearchBox.h"
 #include "Widgets/Images/SImage.h"
-#include "Widgets/Layout/SGridPanel.h"
-#include "Widgets/Layout/SSeparator.h"
 #include "Widgets/SToolTip.h"
 #include "Widgets/Views/STableViewBase.h"
 
 // Insights
-#include "Insights/InsightsManager.h"
 #include "Insights/InsightsStyle.h"
 #include "Insights/Log.h"
 #include "Insights/Table/ViewModels/Table.h"
 #include "Insights/Table/ViewModels/TableCellValueSorter.h"
 #include "Insights/Table/ViewModels/TreeNodeGrouping.h"
-
 #include "Insights/Table/Widgets/STableTreeViewTooltip.h"
 #include "Insights/Table/Widgets/STableTreeViewRow.h"
 #include "Insights/ViewModels/FilterConfigurator.h"
@@ -36,7 +29,7 @@
 
 #include <limits>
 
-#define LOCTEXT_NAMESPACE "STableTreeView"
+#define LOCTEXT_NAMESPACE "Insights::STableTreeView"
 
 namespace Insights
 {
@@ -108,13 +101,13 @@
 			FInputChord(EModifierKey::Control, EKeys::E));
 
 		UI_COMMAND(Command_ExportEntireTreeToFile,
-			"Export Entire Tree to File...",
+			"Export Entire Tree (+ Leaf Nodes) to File...",
 			"Exports the entire tree/table content to a file. It exports also the collapsed tree nodes, including the leaf nodes. Filtered out nodes are not exported.",
 			EUserInterfaceActionType::Button,
 			FInputChord(EModifierKey::Shift, EKeys::E));
 
-		UI_COMMAND(Command_ExportEntireTreeToFileNoLeafs,
-			"Export Entire Tree (No Leafs) to File...",
+		UI_COMMAND(Command_ExportEntireTreeToFileExceptLeaves,
+			"Export Entire Tree (- Leaf Nodes) to File...",
 			"Exports the entire tree/table content to a file, but not the leaf nodes. It exports the collapsed tree nodes. Filtered out nodes are not exported.",
 			EUserInterfaceActionType::Button,
 			FInputChord(EModifierKey::Alt, EKeys::E));
@@ -129,7 +122,7 @@
 	TSharedPtr<FUICommandInfo> Command_CollapseSubtree;
 	TSharedPtr<FUICommandInfo> Command_ExportToFile;
 	TSharedPtr<FUICommandInfo> Command_ExportEntireTreeToFile;
-	TSharedPtr<FUICommandInfo> Command_ExportEntireTreeToFileNoLeafs;
+	TSharedPtr<FUICommandInfo> Command_ExportEntireTreeToFileExceptLeaves;
 };
 
 ////////////////////////////////////////////////////////////////////////////////////////////////////
@@ -141,30 +134,9 @@
 ////////////////////////////////////////////////////////////////////////////////////////////////////
 
 STableTreeView::STableTreeView()
-	: Table()
-	, TreeView(nullptr)
-	, TreeViewHeaderRow(nullptr)
-	, ExternalScrollbar(nullptr)
-	, HoveredColumnId()
-	, HoveredNodePtr(nullptr)
-	, HighlightedNodeName()
-	, Root(MakeShared<FTableTreeNode>(RootNodeName, Table))
-	, TableTreeNodes()
-	, FilteredGroupNodes()
-	, ExpandedNodes()
-	, bExpansionSaved(false)
-	, SearchBox(nullptr)
-	, TextFilter(nullptr)
-	, Filters(nullptr)
-	, AvailableGroupings()
-	, CurrentGroupings()
-	, GroupingBreadcrumbTrail(nullptr)
-	, AvailableSorters()
-	, CurrentSorter(nullptr)
+	: Root(MakeShared<FTableTreeNode>(RootNodeName, Table))
 	, ColumnBeingSorted(GetDefaultColumnBeingSorted())
 	, ColumnSortMode(GetDefaultColumnSortMode())
-	, StatsStartTime(0.0)
-	, StatsEndTime(0.0)
 {
 }
 
@@ -225,7 +197,7 @@
 		FExecuteAction::CreateSP(this, &STableTreeView::ContextMenu_ExportToFile_Execute, true, true),
 		FCanExecuteAction::CreateSP(this, &STableTreeView::ContextMenu_ExportToFile_CanExecute));
 
-	CommandList->MapAction(FTableTreeViewCommands::Get().Command_ExportEntireTreeToFileNoLeafs,
+	CommandList->MapAction(FTableTreeViewCommands::Get().Command_ExportEntireTreeToFileExceptLeaves,
 		FExecuteAction::CreateSP(this, &STableTreeView::ContextMenu_ExportToFile_Execute, true, false),
 		FCanExecuteAction::CreateSP(this, &STableTreeView::ContextMenu_ExportToFile_CanExecute));
 }
@@ -278,6 +250,7 @@
 				.OnGetChildren(this, &STableTreeView::TreeView_OnGetChildren)
 				.OnGenerateRow(this, &STableTreeView::TreeView_OnGenerateRow)
 				.OnSelectionChanged(this, &STableTreeView::TreeView_OnSelectionChanged)
+				.OnExpansionChanged(this, &STableTreeView::TreeView_OnExpansionChanged)
 				.OnMouseButtonDoubleClick(this, &STableTreeView::TreeView_OnMouseButtonDoubleClick)
 				.OnContextMenuOpening(FOnContextMenuOpening::CreateSP(this, &STableTreeView::TreeView_GetMenuContent))
 				.ItemHeight(16.0f)
@@ -302,13 +275,13 @@
 			.Padding(FMargin(0.0f, 30.0f, 0.0f, 0.0f))
 			[
 				SNew(SBorder)
-				.BorderImage(FAppStyle::Get().GetBrush("PopupText.Background"))
-				[
-					SNew(SHorizontalBox)
-					.Visibility_Lambda([this]() -> EVisibility
+				.BorderImage(FAppStyle::Get().GetBrush("Brushes.Panel"))
+				.Visibility_Lambda([this]() -> EVisibility
 					{
 						return TreeViewBannerText.IsEmpty() ? EVisibility::Collapsed : EVisibility::Visible;
 					})
+				[
+					SNew(SHorizontalBox)
 
 					+ SHorizontalBox::Slot()
 					.VAlign(VAlign_Center)
@@ -351,21 +324,8 @@
 		WidgetContent
 	];
 
-	// Create the search filters: text based, type based etc.
-	TextFilter = MakeShared<FTableTreeNodeTextFilter>(FTableTreeNodeTextFilter::FItemToStringArray::CreateStatic(&STableTreeView::HandleItemToStringArray));
-	Filters = MakeShared<FTableTreeNodeFilterCollection>();
-	if (bRunInAsyncMode)
-	{
-		CurrentAsyncOpTextFilter = MakeShared<FTableTreeNodeTextFilter>(FTableTreeNodeTextFilter::FItemToStringArray::CreateStatic(&STableTreeView::HandleItemToStringArray));
-		Filters->Add(CurrentAsyncOpTextFilter);
-	}
-	else
-	{
-		Filters->Add(TextFilter);
-	}
-
+	InitHierarchyFiltering();
 	InitializeAndShowHeaderColumns();
-
 	InitCommandList();
 }
 END_SLATE_FUNCTION_BUILD_OPTIMIZATION
@@ -373,38 +333,5 @@
 ////////////////////////////////////////////////////////////////////////////////////////////////////
 
 BEGIN_SLATE_FUNCTION_BUILD_OPTIMIZATION
-TSharedRef<SWidget> STableTreeView::ConstructSearchBox()
-{
-	SAssignNew(SearchBox, SSearchBox)
-		.HintText(LOCTEXT("SearchBox_Hint", "Search"))
-		.OnTextChanged(this, &STableTreeView::SearchBox_OnTextChanged)
-		.IsEnabled(this, &STableTreeView::SearchBox_IsEnabled)
-		.ToolTipText(this, &STableTreeView::SearchBox_GetTooltipText);
-
-	return SearchBox.ToSharedRef();
-}
-END_SLATE_FUNCTION_BUILD_OPTIMIZATION
-
-////////////////////////////////////////////////////////////////////////////////////////////////////
-
-BEGIN_SLATE_FUNCTION_BUILD_OPTIMIZATION
-TSharedRef<SWidget> STableTreeView::ConstructAdvancedFiltersButton()
-{
-	return SNew(SButton)
-		.ContentPadding(FMargin(-2.0f, 2.0f, -2.0f, 2.0f))
-		//.Text(LOCTEXT("AdvancedFiltersBtn_Text", "Advanced Filters"))
-		.ToolTipText(this, &STableTreeView::AdvancedFilters_GetTooltipText)
-		.OnClicked(this, &STableTreeView::OnAdvancedFiltersClicked)
-		.IsEnabled(this, &STableTreeView::AdvancedFilters_ShouldBeEnabled)
-		[
-			SNew(SImage)
-			.Image(FInsightsStyle::GetBrush("Icons.ClassicFilterConfig"))
-		];
-}
-END_SLATE_FUNCTION_BUILD_OPTIMIZATION
-
-////////////////////////////////////////////////////////////////////////////////////////////////////
-
-BEGIN_SLATE_FUNCTION_BUILD_OPTIMIZATION
 TSharedRef<SWidget> STableTreeView::ConstructHierarchyBreadcrumbTrail()
 {
@@ -523,8 +450,8 @@
 
 TSharedPtr<SWidget> STableTreeView::TreeView_GetMenuContent()
 {
-	const TArray<FTableTreeNodePtr> SelectedNodes = TreeView->GetSelectedItems();
-	const int32 NumSelectedNodes = SelectedNodes.Num();
+	TArray<FTableTreeNodePtr> SelectedNodes;
+	const int32 NumSelectedNodes = TreeView->GetSelectedItems(SelectedNodes);
 	FTableTreeNodePtr SelectedNode = NumSelectedNodes ? SelectedNodes[0] : nullptr;
 
 	FText SelectionStr;
@@ -777,7 +704,7 @@
 
 		MenuBuilder.AddMenuEntry
 		(
-			FTableTreeViewCommands::Get().Command_ExportEntireTreeToFileNoLeafs,
+			FTableTreeViewCommands::Get().Command_ExportEntireTreeToFileExceptLeaves,
 			NAME_None,
 			TAttribute<FText>(),
 			TAttribute<FText>(),
@@ -967,7 +894,7 @@
 		{
 			// We wait for the TreeView to be refreshed before dispatching the tasks.
 			// This should make the TreeView release all of it's shared pointers to nodes to prevent
-			// the TreeView (MainThread) and the tasks from accesing the non-thread safe shared pointers at the same time.
+			// the TreeView (MainThread) and the tasks from accessing the non-thread safe shared pointers at the same time.
 			if (!TreeView->IsPendingRefresh())
 			{
 				DispatchEvent->DispatchSubsequents();
@@ -995,8 +922,8 @@
 		{
 			OnPreAsyncUpdate();
 
-			FGraphEventRef CompletedEvent = StartCreateGroupsTask();
-			CompletedEvent = StartSortTreeNodesTask(CompletedEvent);
+			FGraphEventRef CompletedEvent = StartGroupingTask();
+			CompletedEvent = StartSortingTask(CompletedEvent);
 			InProgressAsyncOperationEvent = StartApplyFiltersTask(CompletedEvent);
 		}
 		else
@@ -1013,7 +940,7 @@
 }
 
 ////////////////////////////////////////////////////////////////////////////////////////////////////
-// Filtering
+// Node Filtering
 ////////////////////////////////////////////////////////////////////////////////////////////////////
 
 void STableTreeView::OnFilteringChanged()
@@ -1039,149 +966,28 @@
 
 ////////////////////////////////////////////////////////////////////////////////////////////////////
 
-void STableTreeView::ApplyFiltering()
+FGraphEventRef STableTreeView::StartApplyFiltersTask(FGraphEventRef Prerequisite)
 {
-	FStopwatch Stopwatch;
-	Stopwatch.Start();
+	AddInProgressAsyncOperation(EAsyncOperationType::FilteringOp);
 
-	// Apply filter to all groups and its children.
-	if (FilterConfigurator_HasFilters())
+	CurrentAsyncOpTextFilter->SetRawFilterText(TextFilter->GetRawFilterText());
+
+	if (FilterConfigurator.IsValid() && CurrentAsyncOpFilterConfigurator)
 	{
-		ApplyAdvancedFiltersForNode(Root);
+		*CurrentAsyncOpFilterConfigurator = *FilterConfigurator;
+	}
+
+	FGraphEventArray Prerequisites;
+	if (Prerequisite.IsValid())
+	{
+		Prerequisites.Add(Prerequisite);
 	}
 	else
 	{
-		const bool bFilterIsEmpty = bRunInAsyncMode ? CurrentAsyncOpTextFilter->GetRawFilterText().IsEmpty() : TextFilter->GetRawFilterText().IsEmpty();
-		ApplyHierarchicalFilterForNode(Root, bFilterIsEmpty);
+		Prerequisites.Add(DispatchEvent);
 	}
 
-	FilteredGroupNodes.Reset();
-	const TArray<FBaseTreeNodePtr>& RootChildren = Root->GetFilteredChildren();
-	const int32 NumRootChildren = RootChildren.Num();
-	for (int32 Cx = 0; Cx < NumRootChildren; ++Cx)
-	{
-		// Add a child.
-		const FTableTreeNodePtr& ChildNodePtr = StaticCastSharedPtr<FTableTreeNode>(RootChildren[Cx]);
-		if (ChildNodePtr->IsGroup())
-		{
-			FilteredGroupNodes.Add(ChildNodePtr);
-		}
-	}
-
-	TreeViewBannerText = FText::GetEmpty();
-	if (Root->GetChildren().Num() == 0)
-	{
-		TreeViewBannerText = LOCTEXT("EmptyTreeMsg", "Tree is empty.");
-	}
-	else if (Root->GetFilteredChildren().Num() == 0)
-	{
-		TreeViewBannerText = LOCTEXT("AllNodesFiltered", "All tree nodes are filtered out.");
-	}
-
-	Stopwatch.Stop();
-	const double FilteringTime = Stopwatch.GetAccumulatedTime();
-	if (FilteringTime > 0.1)
-	{
-		UE_LOG(TraceInsights, Log, TEXT("[Tree - %s] Filtering completed in %.3fs."), *Table->GetDisplayName().ToString(), FilteringTime);
-	}
-
-	UpdateAggregatedValues(*Root);
-
-	// Cannot call TreeView functions from other threads than MainThread and SlateThread.
-	if (!bRunInAsyncMode)
-	{
-		// Only expand nodes if we have a text filter.
-		const bool bNonEmptyTextFilter = !TextFilter->GetRawFilterText().IsEmpty();
-		if (bNonEmptyTextFilter)
-		{
-			if (!bExpansionSaved)
-			{
-				ExpandedNodes.Empty();
-				TreeView->GetExpandedItems(ExpandedNodes);
-				bExpansionSaved = true;
-			}
-
-			for (int32 Fx = 0; Fx < FilteredGroupNodes.Num(); Fx++)
-			{
-				const FTableTreeNodePtr& GroupPtr = FilteredGroupNodes[Fx];
-				TreeView->SetItemExpansion(GroupPtr, GroupPtr->IsExpanded());
-			}
-		}
-		else
-		{
-			if (bExpansionSaved)
-			{
-				// Restore previously expanded nodes when the text filter is disabled.
-				TreeView->ClearExpandedItems();
-				for (auto It = ExpandedNodes.CreateConstIterator(); It; ++It)
-				{
-					TreeView->SetItemExpansion(*It, true);
-				}
-				bExpansionSaved = false;
-			}
-		}
-
-		// Request tree refresh
-		TreeView->RequestTreeRefresh();
-	}
+	return TGraphTask<FTableTreeViewFilterAsyncTask>::CreateTask(&Prerequisites).ConstructAndDispatchWhenReady(this);
 }
 
 ////////////////////////////////////////////////////////////////////////////////////////////////////
-
-bool STableTreeView::ApplyHierarchicalFilterForNode(FTableTreeNodePtr NodePtr, bool bFilterIsEmpty)
-{
-	bool bIsNodeVisible = bFilterIsEmpty || Filters->PassesAllFilters(NodePtr);
-
-	if (NodePtr->IsGroup())
-	{
-		// If a group node passes the filter, all child nodes will be shown.
-		if (bIsNodeVisible)
-		{
-			MakeSubtreeVisible(NodePtr, bFilterIsEmpty);
-			return true;
-		}
-
-		const TArray<FBaseTreeNodePtr>& GroupChildren = NodePtr->GetChildren();
-		const int32 NumChildren = GroupChildren.Num();
-
-		NodePtr->ClearFilteredChildren(bFilterIsEmpty ? NumChildren : 0);
-
-		int32 NumVisibleChildren = 0;
-		for (int32 Cx = 0; Cx < NumChildren; ++Cx)
-		{
-			if (AsyncOperationProgress.ShouldCancelAsyncOp())
-			{
-				break;
-			}
-
-			// Add a child.
-			const FTableTreeNodePtr& ChildNodePtr = StaticCastSharedPtr<FTableTreeNode>(GroupChildren[Cx]);
-			if (ApplyHierarchicalFilterForNode(ChildNodePtr, bFilterIsEmpty))
-			{
-				NodePtr->AddFilteredChild(ChildNodePtr);
-				NumVisibleChildren++;
-			}
-		}
-
-		const bool bIsGroupNodeVisible = bIsNodeVisible || NumVisibleChildren > 0;
-		if (!bFilterIsEmpty && bIsGroupNodeVisible)
-		{
-			if (bRunInAsyncMode)
-			{
-				NodesToExpand.Add(NodePtr);
-			}
-			else
-			{
-				TreeView->SetItemExpansion(NodePtr, true);
-			}
-		}
-		return bIsGroupNodeVisible;
-	}
-	else
-	{
-		NodePtr->SetIsFiltered(!bIsNodeVisible);
-		return bIsNodeVisible;
-	}
-}
-
-////////////////////////////////////////////////////////////////////////////////////////////////////
@@ -1236,6 +1042,349 @@
 
 ////////////////////////////////////////////////////////////////////////////////////////////////////
 
+bool STableTreeView::ApplyAdvancedFilters(const FTableTreeNodePtr& NodePtr)
+{
+	FFilterConfigurator* FilterConfiguratorToUse = bRunInAsyncMode ? CurrentAsyncOpFilterConfigurator : FilterConfigurator.Get();
+
+	if (FilterConfiguratorToUse == nullptr)
+	{
+		return true;
+	}
+
+	if (FilterConfiguratorToUse->GetRootNode()->GetChildren().Num() == 0)
+	{
+		return true;
+	}
+
+	for (const TSharedRef<FTableColumn>& Column : Table->GetColumns())
+	{
+		if (!Column->CanBeFiltered())
+		{
+			continue;
+		}
+
+		switch (Column->GetDataType())
+		{
+			case ETableCellDataType::Int64:
+			{
+				Context.SetFilterData<int64>(Column->GetIndex(), Column->GetValue(*NodePtr)->AsInt64());
+				break;
+			}
+			case ETableCellDataType::Double:
+			{
+				Context.SetFilterData<double>(Column->GetIndex(), Column->GetValue(*NodePtr)->AsDouble());
+				break;
+			}
+			case ETableCellDataType::CString:
+			case ETableCellDataType::Text:
+			case ETableCellDataType::Custom:
+			{
+				if (!Column->IsHierarchy())
+				{
+					// Converting string is heavy, check if the key is used for filtering before doing any conversion
+					if (FilterConfiguratorToUse->IsKeyUsed(Column->GetIndex()))
+					{
+						Context.SetFilterData<FString>(Column->GetIndex(), Column->GetValue(*NodePtr)->AsString());
+					}
+				}
+				break;
+			}
+		}
+	}
+
+	return ApplyCustomAdvancedFilters(NodePtr) && FilterConfiguratorToUse->ApplyFilters(Context);
+}
+
+////////////////////////////////////////////////////////////////////////////////////////////////////
+
+BEGIN_SLATE_FUNCTION_BUILD_OPTIMIZATION
+TSharedRef<SWidget> STableTreeView::ConstructAdvancedFiltersButton()
+{
+	return SNew(SButton)
+		.ContentPadding(FMargin(-2.0f, 2.0f, -2.0f, 2.0f))
+		.ToolTipText(this, &STableTreeView::AdvancedFilters_GetTooltipText)
+		.OnClicked(this, &STableTreeView::OnAdvancedFiltersClicked)
+		.IsEnabled(this, &STableTreeView::AdvancedFilters_ShouldBeEnabled)
+		[
+			SNew(SImage)
+			.Image(FInsightsStyle::GetBrush("Icons.ClassicFilterConfig"))
+		];
+}
+END_SLATE_FUNCTION_BUILD_OPTIMIZATION
+
+////////////////////////////////////////////////////////////////////////////////////////////////////
+
+FReply STableTreeView::OnAdvancedFiltersClicked()
+{
+	if (!FilterConfigurator.IsValid())
+	{
+		FilterConfigurator = MakeShared<FFilterConfigurator>();
+		TSharedPtr<TArray<TSharedPtr<struct FFilter>>>& AvailableFilters = FilterConfigurator->GetAvailableFilters();
+
+		for (const TSharedRef<FTableColumn>& Column : Table->GetColumns())
+		{
+			if (!Column->CanBeFiltered())
+			{
+				continue;
+			}
+
+			switch (Column->GetDataType())
+			{
+				case ETableCellDataType::Int64:
+				{
+					AvailableFilters->Add(MakeShared<FFilter>(Column->GetIndex(), Column->GetTitleName(), Column->GetDescription(), EFilterDataType::Int64, FFilterService::Get()->GetIntegerOperators()));
+					AvailableFilters->Last()->Converter = Column->GetValueConverter();
+					Context.AddFilterData<int64>(Column->GetIndex(), 0);
+					break;
+				}
+				case ETableCellDataType::Double:
+				{
+					AvailableFilters->Add(MakeShared<FFilter>(Column->GetIndex(), Column->GetTitleName(), Column->GetDescription(), EFilterDataType::Double, FFilterService::Get()->GetDoubleOperators()));
+					AvailableFilters->Last()->Converter = Column->GetValueConverter();
+					Context.AddFilterData<double>(Column->GetIndex(), 0.0);
+					break;
+				}
+				case ETableCellDataType::CString:
+				case ETableCellDataType::Text:
+				case ETableCellDataType::Custom:
+				{
+					if (!Column->IsHierarchy())
+					{
+						AvailableFilters->Add(MakeShared<FFilter>(Column->GetIndex(), Column->GetTitleName(), Column->GetDescription(), EFilterDataType::String, FFilterService::Get()->GetStringOperators()));
+						AvailableFilters->Last()->Converter = Column->GetValueConverter();
+						Context.AddFilterData<FString>(Column->GetIndex(), FString());
+					}
+					break;
+				}
+			}
+		}
+
+		AddCustomAdvancedFilters();
+
+		CurrentAsyncOpFilterConfigurator = new FFilterConfigurator(*FilterConfigurator);
+		OnFilterChangesCommittedHandle = FilterConfigurator->GetOnChangesCommittedEvent().AddSP(this, &STableTreeView::OnAdvancedFiltersChangesCommitted);
+	}
+
+	FFilterService::Get()->CreateFilterConfiguratorWidget(FilterConfigurator);
+
+	return FReply::Handled();
+}
+
+////////////////////////////////////////////////////////////////////////////////////////////////////
+
+void STableTreeView::OnAdvancedFiltersChangesCommitted()
+{
+	OnFilteringChanged();
+}
+
+////////////////////////////////////////////////////////////////////////////////////////////////////
+
+bool STableTreeView::AdvancedFilters_ShouldBeEnabled() const
+{
+	return TextFilter->GetRawFilterText().IsEmpty();
+}
+
+////////////////////////////////////////////////////////////////////////////////////////////////////
+
+FText STableTreeView::AdvancedFilters_GetTooltipText() const
+{
+	if (AdvancedFilters_ShouldBeEnabled())
+	{
+		return LOCTEXT("AdvancedFiltersBtn_ToolTip", "Opens the filter configurator window.");
+	}
+
+	return LOCTEXT("AdvancedFiltersBtn_Disabled_ToolTip", "Advanced filters cannot be added when filters are already applied using the search box.");
+}
+
+////////////////////////////////////////////////////////////////////////////////////////////////////
+
+bool STableTreeView::FilterConfigurator_HasFilters() const
+{
+	return FilterConfigurator.IsValid() && FilterConfigurator->GetRootNode()->GetChildren().Num() > 0;
+}
+
+////////////////////////////////////////////////////////////////////////////////////////////////////
+// Hierarchy Filtering
+////////////////////////////////////////////////////////////////////////////////////////////////////
+
+void STableTreeView::InitHierarchyFiltering()
+{
+	// The text filter editable on UI thread.
+	TextFilter = MakeShared<FTableTreeNodeTextFilter>(FTableTreeNodeTextFilter::FItemToStringArray::CreateStatic(&STableTreeView::HandleItemToStringArray));
+
+	if (bRunInAsyncMode)
+	{
+		// A copy of the TextFilter used to actually apply the filtering in an async task.
+		CurrentAsyncOpTextFilter = MakeShared<FTableTreeNodeTextFilter>(FTableTreeNodeTextFilter::FItemToStringArray::CreateStatic(&STableTreeView::HandleItemToStringArray));
+	}
+
+	Filters = MakeShared<FTableTreeNodeFilterCollection>();
+	if (bRunInAsyncMode)
+	{
+		Filters->Add(CurrentAsyncOpTextFilter);
+	}
+	else
+	{
+		Filters->Add(TextFilter);
+	}
+}
+
+////////////////////////////////////////////////////////////////////////////////////////////////////
+
+void STableTreeView::HandleItemToStringArray(const FTableTreeNodePtr& FTableTreeNodePtr, TArray<FString>& OutSearchStrings)
+{
+	OutSearchStrings.Add(FTableTreeNodePtr->GetName().ToString());
+}
+
+////////////////////////////////////////////////////////////////////////////////////////////////////
+
+void STableTreeView::ApplyFiltering()
+{
+	FStopwatch Stopwatch;
+	Stopwatch.Start();
+
+	// Apply filter to all groups and its children.
+	if (FilterConfigurator_HasFilters())
+	{
+		ApplyAdvancedFiltersForNode(Root);
+	}
+	else
+	{
+		const bool bFilterIsEmpty = bRunInAsyncMode ? CurrentAsyncOpTextFilter->GetRawFilterText().IsEmpty() : TextFilter->GetRawFilterText().IsEmpty();
+		ApplyHierarchicalFilterForNode(Root, bFilterIsEmpty);
+	}
+
+	// The Root node is always hidden. The tree shows the filtered children of the Root node.
+	const TArray<FBaseTreeNodePtr>& RootChildren = Root->GetFilteredChildren();
+	const int32 NumRootChildren = RootChildren.Num();
+	FilteredGroupNodes.Reset(NumRootChildren);
+	for (int32 Cx = 0; Cx < NumRootChildren; ++Cx)
+	{
+		const FTableTreeNodePtr& ChildNodePtr = StaticCastSharedPtr<FTableTreeNode>(RootChildren[Cx]);
+		if (ChildNodePtr->IsGroup())
+		{
+			FilteredGroupNodes.Add(ChildNodePtr);
+		}
+	}
+
+	TreeViewBannerText = FText::GetEmpty();
+	if (Root->GetChildren().Num() == 0)
+	{
+		TreeViewBannerText = LOCTEXT("EmptyTreeMsg", "Tree is empty.");
+	}
+	else if (Root->GetFilteredChildren().Num() == 0)
+	{
+		TreeViewBannerText = LOCTEXT("AllNodesFiltered", "All tree nodes are filtered out.");
+	}
+
+	Stopwatch.Stop();
+	const double FilteringTime = Stopwatch.GetAccumulatedTime();
+	if (FilteringTime > 0.1)
+	{
+		UE_LOG(TraceInsights, Log, TEXT("[Tree - %s] Filtering completed in %.3fs."), *Table->GetDisplayName().ToString(), FilteringTime);
+	}
+
+	UpdateAggregatedValues(*Root);
+
+	// Cannot call TreeView functions from other threads than MainThread and SlateThread.
+	if (!bRunInAsyncMode)
+	{
+		// Only expand nodes if we have a text filter.
+		const bool bNonEmptyTextFilter = !TextFilter->GetRawFilterText().IsEmpty();
+		if (bNonEmptyTextFilter)
+		{
+			if (!bExpansionSaved)
+			{
+				ExpandedNodes.Empty();
+				TreeView->GetExpandedItems(ExpandedNodes);
+				bExpansionSaved = true;
+			}
+
+			for (int32 Fx = 0; Fx < FilteredGroupNodes.Num(); Fx++)
+			{
+				const FTableTreeNodePtr& GroupPtr = FilteredGroupNodes[Fx];
+				TreeView->SetItemExpansion(GroupPtr, GroupPtr->IsExpanded());
+			}
+		}
+		else
+		{
+			if (bExpansionSaved)
+			{
+				// Restore previously expanded nodes when the text filter is disabled.
+				TreeView->ClearExpandedItems();
+				for (auto It = ExpandedNodes.CreateConstIterator(); It; ++It)
+				{
+					TreeView->SetItemExpansion(*It, true);
+				}
+				bExpansionSaved = false;
+			}
+		}
+
+		// Request tree refresh
+		TreeView->RequestTreeRefresh();
+	}
+}
+
+////////////////////////////////////////////////////////////////////////////////////////////////////
+
+bool STableTreeView::ApplyHierarchicalFilterForNode(FTableTreeNodePtr NodePtr, bool bFilterIsEmpty)
+{
+	bool bIsNodeVisible = bFilterIsEmpty || Filters->PassesAllFilters(NodePtr);
+
+	if (NodePtr->IsGroup())
+	{
+		// If a group node passes the filter, all child nodes will be shown.
+		if (bIsNodeVisible)
+		{
+			MakeSubtreeVisible(NodePtr, bFilterIsEmpty);
+			return true;
+		}
+
+		const TArray<FBaseTreeNodePtr>& GroupChildren = NodePtr->GetChildren();
+		const int32 NumChildren = GroupChildren.Num();
+
+		NodePtr->ClearFilteredChildren(bFilterIsEmpty ? NumChildren : 0);
+
+		int32 NumVisibleChildren = 0;
+		for (int32 Cx = 0; Cx < NumChildren; ++Cx)
+		{
+			if (AsyncOperationProgress.ShouldCancelAsyncOp())
+			{
+				break;
+			}
+
+			// Add a child.
+			const FTableTreeNodePtr& ChildNodePtr = StaticCastSharedPtr<FTableTreeNode>(GroupChildren[Cx]);
+			if (ApplyHierarchicalFilterForNode(ChildNodePtr, bFilterIsEmpty))
+			{
+				NodePtr->AddFilteredChild(ChildNodePtr);
+				NumVisibleChildren++;
+			}
+		}
+
+		const bool bIsGroupNodeVisible = bIsNodeVisible || NumVisibleChildren > 0;
+		if (!bFilterIsEmpty && bIsGroupNodeVisible)
+		{
+			if (bRunInAsyncMode)
+			{
+				NodesToExpand.Add(NodePtr);
+			}
+			else
+			{
+				TreeView->SetItemExpansion(NodePtr, true);
+			}
+		}
+		return bIsGroupNodeVisible;
+	}
+	else
+	{
+		NodePtr->SetIsFiltered(!bIsNodeVisible);
+		return bIsNodeVisible;
+	}
+}
+
+////////////////////////////////////////////////////////////////////////////////////////////////////
+
 bool STableTreeView::MakeSubtreeVisible(FTableTreeNodePtr NodePtr, bool bFilterIsEmpty)
 {
 	bool bPassesNonEmptyFilter = !bFilterIsEmpty && Filters->PassesAllFilters(NodePtr);
@@ -1281,12 +1430,50 @@
 }
 
 ////////////////////////////////////////////////////////////////////////////////////////////////////
+// SearchBox
+////////////////////////////////////////////////////////////////////////////////////////////////////
 
-void STableTreeView::HandleItemToStringArray(const FTableTreeNodePtr& FTableTreeNodePtr, TArray<FString>& OutSearchStrings)
+BEGIN_SLATE_FUNCTION_BUILD_OPTIMIZATION
+TSharedRef<SWidget> STableTreeView::ConstructSearchBox()
 {
-	OutSearchStrings.Add(FTableTreeNodePtr->GetName().ToString());
+	SAssignNew(SearchBox, SSearchBox)
+		.HintText(LOCTEXT("SearchBox_Hint", "Search"))
+		.OnTextChanged(this, &STableTreeView::SearchBox_OnTextChanged)
+		.IsEnabled(this, &STableTreeView::SearchBox_IsEnabled)
+		.ToolTipText(this, &STableTreeView::SearchBox_GetTooltipText);
+
+	return SearchBox.ToSharedRef();
+}
+END_SLATE_FUNCTION_BUILD_OPTIMIZATION
+
+////////////////////////////////////////////////////////////////////////////////////////////////////
+
+void STableTreeView::SearchBox_OnTextChanged(const FText& InFilterText)
+{
+	TextFilter->SetRawFilterText(InFilterText);
+	SearchBox->SetError(TextFilter->GetFilterErrorText());
+	OnFilteringChanged();
 }
 
 ////////////////////////////////////////////////////////////////////////////////////////////////////
+
+bool STableTreeView::SearchBox_IsEnabled() const
+{
+	return !FilterConfigurator_HasFilters();
+}
+
+////////////////////////////////////////////////////////////////////////////////////////////////////
+
+FText STableTreeView::SearchBox_GetTooltipText() const
+{
+	if (SearchBox_IsEnabled())
+	{
+		return LOCTEXT("SearchBox_ToolTip", "Type here to search the tree hierarchy by item or group name.");
+	}
+
+	return LOCTEXT("SearchBox_Disabled_ToolTip", "Searching the tree hierarchy is disabled when advanced filters are set.");
+}
+
+////////////////////////////////////////////////////////////////////////////////////////////////////
 // TreeView
 ////////////////////////////////////////////////////////////////////////////////////////////////////
@@ -1307,6 +1494,12 @@
 
 ////////////////////////////////////////////////////////////////////////////////////////////////////
 
+void STableTreeView::TreeView_OnExpansionChanged(FTableTreeNodePtr TreeNode, bool bShouldBeExpanded)
+{
+}
+
+////////////////////////////////////////////////////////////////////////////////////////////////////
+
 void STableTreeView::TreeView_OnGetChildren(FTableTreeNodePtr InParent, TArray<FTableTreeNodePtr>& OutChildren)
 {
 	const TArray<FBaseTreeNodePtr>& FilteredChildren = InParent->GetFilteredChildren();
@@ -1404,35 +1597,5 @@
 }
 
 ////////////////////////////////////////////////////////////////////////////////////////////////////
-// SearchBox
-////////////////////////////////////////////////////////////////////////////////////////////////////
-
-void STableTreeView::SearchBox_OnTextChanged(const FText& InFilterText)
-{
-	TextFilter->SetRawFilterText(InFilterText);
-	SearchBox->SetError(TextFilter->GetFilterErrorText());
-	OnFilteringChanged();
-}
-
-////////////////////////////////////////////////////////////////////////////////////////////////////
-
-bool STableTreeView::SearchBox_IsEnabled() const
-{
-	return !FilterConfigurator_HasFilters();
-}
-
-////////////////////////////////////////////////////////////////////////////////////////////////////
-
-FText STableTreeView::SearchBox_GetTooltipText() const
-{
-	if (SearchBox_IsEnabled())
-	{
-		return LOCTEXT("SearchBox_ToolTip", "Type here to search the tree hierarchy by item or group name.");
-	}
-
-	return LOCTEXT("SearchBox_Disabled_ToolTip", "Searching the tree hierarchy is disabled when advanced filters are set.");
-}
-
-////////////////////////////////////////////////////////////////////////////////////////////////////
 // Grouping
 ////////////////////////////////////////////////////////////////////////////////////////////////////
@@ -1445,8 +1608,8 @@
 		{
 			OnPreAsyncUpdate();
 
-			FGraphEventRef CompletedEvent = StartCreateGroupsTask();
-			CompletedEvent = StartSortTreeNodesTask(CompletedEvent);
+			FGraphEventRef CompletedEvent = StartGroupingTask();
+			CompletedEvent = StartSortingTask(CompletedEvent);
 			InProgressAsyncOperationEvent = StartApplyFiltersTask(CompletedEvent);
 		}
 		else
@@ -1464,6 +1627,28 @@
 
 ////////////////////////////////////////////////////////////////////////////////////////////////////
 
+FGraphEventRef STableTreeView::StartGroupingTask(FGraphEventRef Prerequisite)
+{
+	AddInProgressAsyncOperation(EAsyncOperationType::GroupingOp);
+
+	CurrentAsyncOpGroupings.Empty();
+	CurrentAsyncOpGroupings.Insert(CurrentGroupings, 0);
+
+	FGraphEventArray Prerequisites;
+	if (Prerequisite.IsValid())
+	{
+		Prerequisites.Add(Prerequisite);
+	}
+	else
+	{
+		Prerequisites.Add(DispatchEvent);
+	}
+
+	return TGraphTask<FTableTreeViewGroupAsyncTask>::CreateTask(&Prerequisites).ConstructAndDispatchWhenReady(this, &CurrentAsyncOpGroupings);
+}
+
+////////////////////////////////////////////////////////////////////////////////////////////////////
+
 void STableTreeView::ApplyGrouping()
 {
 	CreateGroups(CurrentGroupings);
@@ -1573,191 +1758,3 @@
 	}
 }
 
-////////////////////////////////////////////////////////////////////////////////////////////////////
-
-void STableTreeView::UpdateAggregatedValues(FTableTreeNode& GroupNode)
-{
-	FStopwatch Stopwatch;
-	Stopwatch.Start();
-
-	for (const TSharedRef<FTableColumn>& ColumnRef : Table->GetColumns())
-	{
-		FTableColumn& Column = ColumnRef.Get();
-		switch (Column.GetAggregation())
-		{
-			case ETableColumnAggregation::Sum:
-				if (Column.GetDataType() == ETableCellDataType::Float || Column.GetDataType() == ETableCellDataType::Double)
-				{
-					STableTreeView::UpdateAggregationRec<double>(Column, GroupNode, 0, true,
-						[](double InValue, const FTableCellValue& InTableCellValue)
-						{
-							return InValue + InTableCellValue.AsDouble();
-						});
-				}
-				else
-				{
-					STableTreeView::UpdateAggregationRec<int64>(Column, GroupNode, 0, true,
-						[](int64 InValue, const FTableCellValue& InTableCellValue)
-						{
-							return InValue + InTableCellValue.AsInt64();
-						});
-				}
-				break;
-
-			case ETableColumnAggregation::Min:
-				if (Column.GetDataType() == ETableCellDataType::Float || Column.GetDataType() == ETableCellDataType::Double)
-				{
-					STableTreeView::UpdateAggregationRec<double>(Column, GroupNode, std::numeric_limits<double>::max(), false,
-						[](double InValue, const FTableCellValue& InTableCellValue)
-						{
-							return FMath::Min(InValue, InTableCellValue.AsDouble());
-						});
-				}
-				else
-				{
-					STableTreeView::UpdateAggregationRec<int64>(Column, GroupNode, std::numeric_limits<int64>::max(), false,
-						[](int64 InValue, const FTableCellValue& InTableCellValue)
-						{
-							return FMath::Min(InValue, InTableCellValue.AsInt64());
-						});
-				}
-				break;
-
-			case ETableColumnAggregation::Max:
-				if (Column.GetDataType() == ETableCellDataType::Float || Column.GetDataType() == ETableCellDataType::Double)
-				{
-					STableTreeView::UpdateAggregationRec<double>(Column, GroupNode, std::numeric_limits<double>::lowest(), false,
-						[](double InValue, const FTableCellValue& InTableCellValue)
-						{
-							return FMath::Max(InValue, InTableCellValue.AsDouble());
-						});
-				}
-				else
-				{
-					STableTreeView::UpdateAggregationRec<int64>(Column, GroupNode, std::numeric_limits<int64>::min(), false,
-						[](int64 InValue, const FTableCellValue& InTableCellValue)
-						{
-							return FMath::Max(InValue, InTableCellValue.AsInt64());
-						});
-				}
-				break;
-
-			case ETableColumnAggregation::SameValue:
-				if (Column.GetDataType() == ETableCellDataType::CString)
-				{
-					STableTreeView::UpdateCStringSameValueAggregationRec(Column, GroupNode);
-				}
-				break;
-		}
-	}
-
-	Stopwatch.Stop();
-	const double AggregationTime = Stopwatch.GetAccumulatedTime();
-	if (AggregationTime > 0.1)
-	{
-		UE_LOG(TraceInsights, Log, TEXT("[Tree - %s] Aggregation completed in %.3fs."), *Table->GetDisplayName().ToString(), AggregationTime);
-	}
-}
-
-////////////////////////////////////////////////////////////////////////////////////////////////////
-
-void STableTreeView::UpdateCStringSameValueAggregationRec(FTableColumn& Column, FTableTreeNode& GroupNode)
-{
-	// Update child group nodes first.
-	for (FBaseTreeNodePtr NodePtr : GroupNode.GetChildren())
-	{
-		if (NodePtr->IsFiltered())
-		{
-			continue;
-		}
-		if (NodePtr->IsGroup())
-		{
-			FTableTreeNode& TableNode = *(FTableTreeNode*)NodePtr.Get();
-			UpdateCStringSameValueAggregationRec(Column, TableNode);
-		}
-	}
-
-	const TCHAR* AggregatedValue = nullptr;
-
-	// Find the first child node.
-	for (FBaseTreeNodePtr NodePtr : GroupNode.GetChildren())
-	{
-		if (NodePtr->IsFiltered())
-		{
-			continue;
-		}
-
-		if (!NodePtr->IsGroup())
-		{
-			const TOptional<FTableCellValue> NodeValue = Column.GetValue(*NodePtr);
-			if (NodeValue.IsSet() &&
-				NodeValue.GetValue().DataType == ETableCellDataType::CString)
-			{
-				AggregatedValue = NodeValue.GetValue().CString;
-			}
-		}
-		else
-		{
-			FTableTreeNode& TableNode = *(FTableTreeNode*)NodePtr.Get();
-			if (TableNode.HasAggregatedValue(Column.GetId()))
-			{
-				const FTableCellValue& ChildGroupAggregatedValue = TableNode.GetAggregatedValue(Column.GetId());
-				if (ChildGroupAggregatedValue.DataType == ETableCellDataType::CString)
-				{
-					AggregatedValue = ChildGroupAggregatedValue.CString;
-				}
-			}
-		}
-
-		break;
-	}
-
-	if (AggregatedValue != nullptr)
-	{
-		// Check if all other children have the same value as the first node.
-		for (FBaseTreeNodePtr NodePtr : GroupNode.GetChildren())
-		{
-			if (NodePtr->IsFiltered())
-			{
-				continue;
-			}
-
-			if (!NodePtr->IsGroup())
-			{
-				const TOptional<FTableCellValue> NodeValue = Column.GetValue(*NodePtr);
-				if (NodeValue.IsSet() &&
-					NodeValue.GetValue().DataType == ETableCellDataType::CString &&
-					AggregatedValue == NodeValue.GetValue().CString)
-				{
-					continue;
-				}
-			}
-			else
-			{
-				FTableTreeNode& TableNode = *(FTableTreeNode*)NodePtr.Get();
-				if (TableNode.HasAggregatedValue(Column.GetId()))
-				{
-					const FTableCellValue& ChildGroupAggregatedValue = TableNode.GetAggregatedValue(Column.GetId());
-					if (ChildGroupAggregatedValue.DataType == ETableCellDataType::CString &&
-						AggregatedValue == ChildGroupAggregatedValue.CString)
-					{
-						continue;
-					}
-				}
-			}
-
-			AggregatedValue = nullptr;
-			break;
-		}
-	}
-
-	if (AggregatedValue != nullptr)
-	{
-		GroupNode.AddAggregatedValue(Column.GetId(), FTableCellValue(AggregatedValue));
-	}
-	else
-	{
-		GroupNode.ResetAggregatedValues(Column.GetId());
-	}
-}
-
@@ -2333,3 +2330,193 @@
 	}
 }
 
+////////////////////////////////////////////////////////////////////////////////////////////////////
+// Aggregation
+////////////////////////////////////////////////////////////////////////////////////////////////////
+
+void STableTreeView::UpdateCStringSameValueAggregationRec(FTableColumn& Column, FTableTreeNode& GroupNode)
+{
+	// Update child group nodes first.
+	for (FBaseTreeNodePtr NodePtr : GroupNode.GetChildren())
+	{
+		if (NodePtr->IsFiltered())
+		{
+			continue;
+		}
+		if (NodePtr->IsGroup())
+		{
+			FTableTreeNode& TableNode = *(FTableTreeNode*)NodePtr.Get();
+			UpdateCStringSameValueAggregationRec(Column, TableNode);
+		}
+	}
+
+	const TCHAR* AggregatedValue = nullptr;
+
+	// Find the first child node.
+	for (FBaseTreeNodePtr NodePtr : GroupNode.GetChildren())
+	{
+		if (NodePtr->IsFiltered())
+		{
+			continue;
+		}
+
+		if (!NodePtr->IsGroup())
+		{
+			const TOptional<FTableCellValue> NodeValue = Column.GetValue(*NodePtr);
+			if (NodeValue.IsSet() &&
+				NodeValue.GetValue().DataType == ETableCellDataType::CString)
+			{
+				AggregatedValue = NodeValue.GetValue().CString;
+			}
+		}
+		else
+		{
+			FTableTreeNode& TableNode = *(FTableTreeNode*)NodePtr.Get();
+			if (TableNode.HasAggregatedValue(Column.GetId()))
+			{
+				const FTableCellValue& ChildGroupAggregatedValue = TableNode.GetAggregatedValue(Column.GetId());
+				if (ChildGroupAggregatedValue.DataType == ETableCellDataType::CString)
+				{
+					AggregatedValue = ChildGroupAggregatedValue.CString;
+				}
+			}
+		}
+
+		break;
+	}
+
+	if (AggregatedValue != nullptr)
+	{
+		// Check if all other children have the same value as the first node.
+		for (FBaseTreeNodePtr NodePtr : GroupNode.GetChildren())
+		{
+			if (NodePtr->IsFiltered())
+			{
+				continue;
+			}
+
+			if (!NodePtr->IsGroup())
+			{
+				const TOptional<FTableCellValue> NodeValue = Column.GetValue(*NodePtr);
+				if (NodeValue.IsSet() &&
+					NodeValue.GetValue().DataType == ETableCellDataType::CString &&
+					AggregatedValue == NodeValue.GetValue().CString)
+				{
+					continue;
+				}
+			}
+			else
+			{
+				FTableTreeNode& TableNode = *(FTableTreeNode*)NodePtr.Get();
+				if (TableNode.HasAggregatedValue(Column.GetId()))
+				{
+					const FTableCellValue& ChildGroupAggregatedValue = TableNode.GetAggregatedValue(Column.GetId());
+					if (ChildGroupAggregatedValue.DataType == ETableCellDataType::CString &&
+						AggregatedValue == ChildGroupAggregatedValue.CString)
+					{
+						continue;
+					}
+				}
+			}
+
+			AggregatedValue = nullptr;
+			break;
+		}
+	}
+
+	if (AggregatedValue != nullptr)
+	{
+		GroupNode.AddAggregatedValue(Column.GetId(), FTableCellValue(AggregatedValue));
+	}
+	else
+	{
+		GroupNode.ResetAggregatedValues(Column.GetId());
+	}
+}
+
+////////////////////////////////////////////////////////////////////////////////////////////////////
+
+void STableTreeView::UpdateAggregatedValues(FTableTreeNode& GroupNode)
+{
+	FStopwatch Stopwatch;
+	Stopwatch.Start();
+
+	for (const TSharedRef<FTableColumn>& ColumnRef : Table->GetColumns())
+	{
+		FTableColumn& Column = ColumnRef.Get();
+		switch (Column.GetAggregation())
+		{
+			case ETableColumnAggregation::Sum:
+				if (Column.GetDataType() == ETableCellDataType::Float || Column.GetDataType() == ETableCellDataType::Double)
+				{
+					STableTreeView::UpdateAggregationRec<double>(Column, GroupNode, 0, true,
+						[](double InValue, const FTableCellValue& InTableCellValue)
+						{
+							return InValue + InTableCellValue.AsDouble();
+						});
+				}
+				else
+				{
+					STableTreeView::UpdateAggregationRec<int64>(Column, GroupNode, 0, true,
+						[](int64 InValue, const FTableCellValue& InTableCellValue)
+						{
+							return InValue + InTableCellValue.AsInt64();
+						});
+				}
+				break;
+
+			case ETableColumnAggregation::Min:
+				if (Column.GetDataType() == ETableCellDataType::Float || Column.GetDataType() == ETableCellDataType::Double)
+				{
+					STableTreeView::UpdateAggregationRec<double>(Column, GroupNode, std::numeric_limits<double>::max(), false,
+						[](double InValue, const FTableCellValue& InTableCellValue)
+						{
+							return FMath::Min(InValue, InTableCellValue.AsDouble());
+						});
+				}
+				else
+				{
+					STableTreeView::UpdateAggregationRec<int64>(Column, GroupNode, std::numeric_limits<int64>::max(), false,
+						[](int64 InValue, const FTableCellValue& InTableCellValue)
+						{
+							return FMath::Min(InValue, InTableCellValue.AsInt64());
+						});
+				}
+				break;
+
+			case ETableColumnAggregation::Max:
+				if (Column.GetDataType() == ETableCellDataType::Float || Column.GetDataType() == ETableCellDataType::Double)
+				{
+					STableTreeView::UpdateAggregationRec<double>(Column, GroupNode, std::numeric_limits<double>::lowest(), false,
+						[](double InValue, const FTableCellValue& InTableCellValue)
+						{
+							return FMath::Max(InValue, InTableCellValue.AsDouble());
+						});
+				}
+				else
+				{
+					STableTreeView::UpdateAggregationRec<int64>(Column, GroupNode, std::numeric_limits<int64>::min(), false,
+						[](int64 InValue, const FTableCellValue& InTableCellValue)
+						{
+							return FMath::Max(InValue, InTableCellValue.AsInt64());
+						});
+				}
+				break;
+
+			case ETableColumnAggregation::SameValue:
+				if (Column.GetDataType() == ETableCellDataType::CString)
+				{
+					STableTreeView::UpdateCStringSameValueAggregationRec(Column, GroupNode);
+				}
+				break;
+		}
+	}
+
+	Stopwatch.Stop();
+	const double AggregationTime = Stopwatch.GetAccumulatedTime();
+	if (AggregationTime > 0.1)
+	{
+		UE_LOG(TraceInsights, Log, TEXT("[Tree - %s] Aggregation completed in %.3fs."), *Table->GetDisplayName().ToString(), AggregationTime);
+	}
+}
+
@@ -2337,8 +2524,6 @@
 // Sorting
 ////////////////////////////////////////////////////////////////////////////////////////////////////
 
-////////////////////////////////////////////////////////////////////////////////////////////////////
-
 const EColumnSortMode::Type STableTreeView::GetDefaultColumnSortMode()
 {
 	return EColumnSortMode::Descending;
@@ -2392,7 +2577,7 @@
 		{
 			OnPreAsyncUpdate();
 
-			FGraphEventRef CompletedEvent = StartSortTreeNodesTask();
+			FGraphEventRef CompletedEvent = StartSortingTask();
 			InProgressAsyncOperationEvent = StartApplyFiltersTask(CompletedEvent);
 		}
 		else
@@ -2409,6 +2594,33 @@
 
 ////////////////////////////////////////////////////////////////////////////////////////////////////
 
+FGraphEventRef STableTreeView::StartSortingTask(FGraphEventRef Prerequisite)
+{
+	if (!CurrentSorter.IsValid())
+	{
+		return Prerequisite;
+	}
+
+	AddInProgressAsyncOperation(EAsyncOperationType::SortingOp);
+
+	CurrentAsyncOpSorter = CurrentSorter.Get();
+	CurrentAsyncOpColumnSortMode = ColumnSortMode;
+
+	FGraphEventArray Prerequisites;
+	if (Prerequisite.IsValid())
+	{
+		Prerequisites.Add(Prerequisite);
+	}
+	else
+	{
+		Prerequisites.Add(DispatchEvent);
+	}
+
+	return TGraphTask<FTableTreeViewSortAsyncTask>::CreateTask(&Prerequisites).ConstructAndDispatchWhenReady(this, CurrentAsyncOpSorter, CurrentAsyncOpColumnSortMode);
+}
+
+////////////////////////////////////////////////////////////////////////////////////////////////////
+
 void STableTreeView::ApplySorting()
 {
 	if (CurrentSorter.IsValid())
@@ -2824,9 +3036,6 @@
 
 void STableTreeView::Reset()
 {
-	StatsStartTime = 0.0;
-	StatsEndTime = 0.0;
-
 	RebuildTree(true);
 }
 
@@ -2924,7 +3133,7 @@
 		ExpandedNodes.Empty();
 
 		// Expand each group node on the first few depths (if it doesn't have too many children).
-		SetExpandValueForChildGroups(Root.Get(), 1000, 4, true);
+		SetExpandValueForChildGroups(Root.Get(), MaxNodesToAutoExpand, MaxDepthToAutoExpand, true);
 
 		ClearInProgressAsyncOperations();
 		TreeView_Refresh();
@@ -2936,11 +3145,11 @@
 
 ////////////////////////////////////////////////////////////////////////////////////////////////////
 
-void STableTreeView::SetExpandValueForChildGroups(FBaseTreeNode* InRoot, int32 InMaxExpandedNodes, int32 InMaxDepthToExpand, bool InValue)
+void STableTreeView::SetExpandValueForChildGroups(FBaseTreeNode* InRoot, int32 InMaxNodesToExpand, int32 InMaxDepthToExpand, bool InValue)
 {
 	TArray<int32> NumNodesPerDepth;
 	NumNodesPerDepth.AddDefaulted(InMaxDepthToExpand + 1);
-	CountNumNodesPerDepthRec(InRoot, NumNodesPerDepth, 0, InMaxDepthToExpand, InMaxExpandedNodes);
+	CountNumNodesPerDepthRec(InRoot, NumNodesPerDepth, 0, InMaxDepthToExpand, InMaxNodesToExpand);
 
 	int32 MaxDepth = 0;
 	for (int32 Depth = 0; Depth <= InMaxDepthToExpand; ++Depth)
@@ -2949,7 +3158,7 @@
 		{
 			NumNodesPerDepth[Depth] += NumNodesPerDepth[Depth - 1];
 		}
-		if (NumNodesPerDepth[Depth] > InMaxExpandedNodes)
+		if (NumNodesPerDepth[Depth] > InMaxNodesToExpand)
 		{
 			break;
 		}
@@ -3000,77 +3209,3 @@
 }
 
 ////////////////////////////////////////////////////////////////////////////////////////////////////
-
-FGraphEventRef STableTreeView::StartSortTreeNodesTask(FGraphEventRef Prerequisite)
-{
-	if (!CurrentSorter.IsValid())
-	{
-		return Prerequisite;
-	}
-
-	AddInProgressAsyncOperation(EAsyncOperationType::SortingOp);
-
-	CurrentAsyncOpSorter = CurrentSorter.Get();
-	CurrentAsyncOpColumnSortMode = ColumnSortMode;
-
-	FGraphEventArray Prerequisites;
-	if (Prerequisite.IsValid())
-	{
-		Prerequisites.Add(Prerequisite);
-	}
-	else
-	{
-		Prerequisites.Add(DispatchEvent);
-	}
-
-	return TGraphTask<FTableTreeViewSortAsyncTask>::CreateTask(&Prerequisites).ConstructAndDispatchWhenReady(this, CurrentAsyncOpSorter, CurrentAsyncOpColumnSortMode);
-}
-
-////////////////////////////////////////////////////////////////////////////////////////////////////
-
-FGraphEventRef STableTreeView::StartCreateGroupsTask(FGraphEventRef Prerequisite)
-{
-	AddInProgressAsyncOperation(EAsyncOperationType::GroupingOp);
-
-	CurrentAsyncOpGroupings.Empty();
-	CurrentAsyncOpGroupings.Insert(CurrentGroupings, 0);
-
-	FGraphEventArray Prerequisites;
-	if (Prerequisite.IsValid())
-	{
-		Prerequisites.Add(Prerequisite);
-	}
-	else
-	{
-		Prerequisites.Add(DispatchEvent);
-	}
-
-	return TGraphTask<FTableTreeViewGroupAsyncTask>::CreateTask(&Prerequisites).ConstructAndDispatchWhenReady(this, &CurrentAsyncOpGroupings);
-}
-
-////////////////////////////////////////////////////////////////////////////////////////////////////
-
-FGraphEventRef STableTreeView::StartApplyFiltersTask(FGraphEventRef Prerequisite)
-{
-	AddInProgressAsyncOperation(EAsyncOperationType::FilteringOp);
-
-	CurrentAsyncOpTextFilter->SetRawFilterText(TextFilter->GetRawFilterText());
-	if (FilterConfigurator.IsValid() && CurrentAsyncOpFilterConfigurator)
-	{
-		*CurrentAsyncOpFilterConfigurator = *FilterConfigurator;
-	}
-
-	FGraphEventArray Prerequisites;
-	if (Prerequisite.IsValid())
-	{
-		Prerequisites.Add(Prerequisite);
-	}
-	else
-	{
-		Prerequisites.Add(DispatchEvent);
-	}
-
-	return TGraphTask<FTableTreeViewFilterAsyncTask>::CreateTask(&Prerequisites).ConstructAndDispatchWhenReady(this);
-}
-
-////////////////////////////////////////////////////////////////////////////////////////////////////
@@ -3134,8 +3269,8 @@
 	{
 		OnPreAsyncUpdate();
 
-		FGraphEventRef CompletedEvent = StartCreateGroupsTask();
-		CompletedEvent = StartSortTreeNodesTask(CompletedEvent);
+		FGraphEventRef CompletedEvent = StartGroupingTask();
+		CompletedEvent = StartSortingTask(CompletedEvent);
 		InProgressAsyncOperationEvent = StartApplyFiltersTask(CompletedEvent);
 
 		return;
@@ -3153,7 +3288,7 @@
 	{
 		OnPreAsyncUpdate();
 
-		FGraphEventRef CompletedEvent = StartSortTreeNodesTask();
+		FGraphEventRef CompletedEvent = StartSortingTask();
 		InProgressAsyncOperationEvent = StartApplyFiltersTask(CompletedEvent);
 
 		return;
@@ -3189,152 +3324,6 @@
 
 ////////////////////////////////////////////////////////////////////////////////////////////////////
 
-FReply STableTreeView::OnAdvancedFiltersClicked()
-{
-	if (!FilterConfigurator.IsValid())
-	{
-		FilterConfigurator = MakeShared<FFilterConfigurator>();
-		TSharedPtr<TArray<TSharedPtr<struct FFilter>>>& AvailableFilters = FilterConfigurator->GetAvailableFilters();
-
-		for (const TSharedRef<FTableColumn>& Column : Table->GetColumns())
-		{
-			if (!Column->CanBeFiltered())
-			{
-				continue;
-			}
-
-			switch (Column->GetDataType())
-			{
-				case ETableCellDataType::Int64:
-				{
-					AvailableFilters->Add(MakeShared<FFilter>(Column->GetIndex(), Column->GetTitleName(), Column->GetDescription(), EFilterDataType::Int64, FFilterService::Get()->GetIntegerOperators()));
-					AvailableFilters->Last()->Converter = Column->GetValueConverter();
-					Context.AddFilterData<int64>(Column->GetIndex(), 0);
-					break;
-				}
-				case ETableCellDataType::Double:
-				{
-					AvailableFilters->Add(MakeShared<FFilter>(Column->GetIndex(), Column->GetTitleName(), Column->GetDescription(), EFilterDataType::Double, FFilterService::Get()->GetDoubleOperators()));
-					AvailableFilters->Last()->Converter = Column->GetValueConverter();
-					Context.AddFilterData<double>(Column->GetIndex(), 0.0);
-					break;
-				}
-				case ETableCellDataType::CString:
-				case ETableCellDataType::Text:
-				case ETableCellDataType::Custom:
-				{
-					if (!Column->IsHierarchy())
-					{
-						AvailableFilters->Add(MakeShared<FFilter>(Column->GetIndex(), Column->GetTitleName(), Column->GetDescription(), EFilterDataType::String, FFilterService::Get()->GetStringOperators()));
-						AvailableFilters->Last()->Converter = Column->GetValueConverter();
-						Context.AddFilterData<FString>(Column->GetIndex(), FString());
-					}
-					break;
-				}
-			}
-		}
-
-		AddCustomAdvancedFilters();
-
-		CurrentAsyncOpFilterConfigurator = new FFilterConfigurator(*FilterConfigurator);
-		OnFilterChangesCommitedHandle = FilterConfigurator->GetOnChangesCommitedEvent().AddSP(this, &STableTreeView::OnAdvancedFiltersChangesCommited);
-	}
-
-	FFilterService::Get()->CreateFilterConfiguratorWidget(FilterConfigurator);
-
-	return FReply::Handled();
-}
-
-////////////////////////////////////////////////////////////////////////////////////////////////////
-
-bool STableTreeView::ApplyAdvancedFilters(const FTableTreeNodePtr& NodePtr)
-{
-	FFilterConfigurator* FilterConfiguratorToUse = bRunInAsyncMode ? CurrentAsyncOpFilterConfigurator : FilterConfigurator.Get();
-
-	if (FilterConfiguratorToUse == nullptr)
-	{
-		return true;
-	}
-
-	if (FilterConfiguratorToUse->GetRootNode()->GetChildren().Num() == 0)
-	{
-		return true;
-	}
-
-	for (const TSharedRef<FTableColumn>& Column : Table->GetColumns())
-	{
-		if (!Column->CanBeFiltered())
-		{
-			continue;
-		}
-
-		switch (Column->GetDataType())
-		{
-			case ETableCellDataType::Int64:
-			{
-				Context.SetFilterData<int64>(Column->GetIndex(), Column->GetValue(*NodePtr)->AsInt64());
-				break;
-			}
-			case ETableCellDataType::Double:
-			{
-				Context.SetFilterData<double>(Column->GetIndex(), Column->GetValue(*NodePtr)->AsDouble());
-				break;
-			}
-			case ETableCellDataType::CString:
-			case ETableCellDataType::Text:
-			case ETableCellDataType::Custom:
-			{
-				if (!Column->IsHierarchy())
-				{
-					// Converting string is heavy, check if the key is used for filtering before doing any conversion
-					if (FilterConfiguratorToUse->IsKeyUsed(Column->GetIndex()))
-					{
-						Context.SetFilterData<FString>(Column->GetIndex(), Column->GetValue(*NodePtr)->AsString());
-					}
-				}
-				break;
-			}
-		}
-	}
-
-	return ApplyCustomAdvancedFilters(NodePtr) && FilterConfiguratorToUse->ApplyFilters(Context);
-}
-
-////////////////////////////////////////////////////////////////////////////////////////////////////
-
-void STableTreeView::OnAdvancedFiltersChangesCommited()
-{
-	OnFilteringChanged();
-}
-
-////////////////////////////////////////////////////////////////////////////////////////////////////
-
-bool STableTreeView::AdvancedFilters_ShouldBeEnabled() const
-{
-	return TextFilter->GetRawFilterText().IsEmpty();
-}
-
-////////////////////////////////////////////////////////////////////////////////////////////////////
-
-FText STableTreeView::AdvancedFilters_GetTooltipText() const
-{
-	if (AdvancedFilters_ShouldBeEnabled())
-	{
-		return LOCTEXT("AdvancedFiltersBtn_ToolTip", "Opens the filter configurator window.");
-	}
-
-	return LOCTEXT("AdvancedFiltersBtn_Disabled_ToolTip", "Advanced filters cannot be added when filters are already applied using the search box.");
-}
-
-////////////////////////////////////////////////////////////////////////////////////////////////////
-
-bool STableTreeView::FilterConfigurator_HasFilters() const
-{
-	return FilterConfigurator.IsValid() && FilterConfigurator->GetRootNode()->GetChildren().Num() > 0;
-}
-
-////////////////////////////////////////////////////////////////////////////////////////////////////
-
 FReply STableTreeView::OnKeyDown(const FGeometry& MyGeometry, const FKeyEvent& InKeyEvent)
 {
 	return CommandList->ProcessCommandBindings(InKeyEvent) == true ? FReply::Handled() : FReply::Unhandled();
@@ -3453,14 +3442,19 @@
 
 void STableTreeView::ContextMenu_ExpandSubtree_Execute()
 {
-	const TArray<FTableTreeNodePtr> SelectedNodes = TreeView->GetSelectedItems();
+	TArray<FTableTreeNodePtr> SelectedNodes;
+	if (!TreeView->GetSelectedItems(SelectedNodes))
+	{
+		return;
+	}
+
 	for (const FTableTreeNodePtr& Node : SelectedNodes)
 	{
 		if (Node->IsGroup())
 		{
 			Node->SetExpansion(true);
 			TreeView->SetItemExpansion(Node, true);
-			SetExpandValueForChildGroups((FBaseTreeNode*)Node.Get(), MAX_NUMBER_OF_NODES_TO_EXPAND, MAX_DEPTH_TO_EXPAND, true);
+			SetExpandValueForChildGroups((FBaseTreeNode*)Node.Get(), MaxNodesToExpand, MaxDepthToExpand, true);
 		}
 	}
 
@@ -3471,7 +3465,12 @@
 
 bool STableTreeView::ContextMenu_CollapseSubtree_CanExecute() const
 {
-	const TArray<FTableTreeNodePtr> SelectedNodes = TreeView->GetSelectedItems();
+	TArray<FTableTreeNodePtr> SelectedNodes;
+	if (!TreeView->GetSelectedItems(SelectedNodes))
+	{
+		return false;
+	}
+
 	for (const FTableTreeNodePtr& Node : SelectedNodes)
 	{
 		if (Node->IsGroup() && Node->GetFilteredChildren().Num() > 0 && TreeView->IsItemExpanded(Node))
@@ -3487,14 +3486,19 @@
 
 void STableTreeView::ContextMenu_CollapseSubtree_Execute()
 {
-	const TArray<FTableTreeNodePtr> SelectedNodes = TreeView->GetSelectedItems();
+	TArray<FTableTreeNodePtr> SelectedNodes;
+	if (!TreeView->GetSelectedItems(SelectedNodes))
+	{
+		return;
+	}
+
 	for (const FTableTreeNodePtr& Node : SelectedNodes)
 	{
 		if (Node->IsGroup() && TreeView->IsItemExpanded(Node))
 		{
 			Node->SetExpansion(false);
 			TreeView->SetItemExpansion(Node, false);
-			SetExpandValueForChildGroups((FBaseTreeNode*)Node.Get(), MAX_NUMBER_OF_NODES_TO_EXPAND, MAX_DEPTH_TO_EXPAND, false);
+			SetExpandValueForChildGroups((FBaseTreeNode*)Node.Get(), MaxNodesToExpand, MaxDepthToExpand, false);
 		}
 	}
 
@@ -3512,7 +3516,12 @@
 
 void STableTreeView::ContextMenu_ExpandCriticalPath_Execute()
 {
-	const TArray<FTableTreeNodePtr> SelectedNodes = TreeView->GetSelectedItems();
+	TArray<FTableTreeNodePtr> SelectedNodes;
+	if (!TreeView->GetSelectedItems(SelectedNodes))
+	{
+		return;
+	}
+
 	for (const FTableTreeNodePtr& Node : SelectedNodes)
 	{
 		FTableTreeNodePtr CurrentNode = Node;
@@ -3689,7 +3698,7 @@
 void STableTreeView::ApplyColumnConfig(const TArrayView<FTableColumnConfig >& InColumnConfigSet)
 {
 	// TODO: Reorder columns as in the config set.
-	// Currenly we only apply visibility and column width.
+	// Currently we only apply visibility and column width.
 	for (const TSharedRef<FTableColumn>& ColumnRef : Table->GetColumns())
 	{
 		FTableColumn& Column = ColumnRef.Get();
@@ -3770,52 +3779,5 @@
 }
 
 ////////////////////////////////////////////////////////////////////////////////////////////////////
-// SSessionTableTreeView
-////////////////////////////////////////////////////////////////////////////////////////////////////
 
-SSessionTableTreeView::~SSessionTableTreeView()
-{
-	// Remove ourselves from the Insights manager.
-	if (FInsightsManager::Get().IsValid())
-	{
-		FInsightsManager::Get()->GetSessionChangedEvent().RemoveAll(this);
-	}
-
-	Session.Reset();
-}
-
-////////////////////////////////////////////////////////////////////////////////////////////////////
-
-void SSessionTableTreeView::ConstructWidget(TSharedPtr<FTable> InTablePtr)
-{
-	STableTreeView::ConstructWidget(InTablePtr);
-
-	// Register ourselves with the Insights manager.
-	FInsightsManager::Get()->GetSessionChangedEvent().AddSP(this, &SSessionTableTreeView::InsightsManager_OnSessionChanged);
-
-	// Update the Session (i.e. when analysis session was already started).
-	InsightsManager_OnSessionChanged();
-
-	CreateGroupings();
-	CreateSortings();
-}
-
-////////////////////////////////////////////////////////////////////////////////////////////////////
-
-void SSessionTableTreeView::InsightsManager_OnSessionChanged()
-{
-	TSharedPtr<const TraceServices::IAnalysisSession> NewSession = FInsightsManager::Get()->GetSession();
-	if (NewSession != Session)
-	{
-		Session = NewSession;
-		Reset();
-	}
-	else
-	{
-		UpdateTree();
-	}
-}
-
-////////////////////////////////////////////////////////////////////////////////////////////////////
-
 } // namespace Insights
diff --git a/Engine/Source/Developer/TraceInsights/Private/Insights/Table/Widgets/STableTreeView.h b/Engine/Source/Developer/TraceInsights/Private/Insights/Table/Widgets/STableTreeView.h
--- a/Engine/Source/Developer/TraceInsights/Private/Insights/Table/Widgets/STableTreeView.h
+++ b/Engine/Source/Developer/TraceInsights/Private/Insights/Table/Widgets/STableTreeView.h
@@ -31,11 +31,6 @@
 class FMenuBuilder;
 class FUICommandList;
 
-namespace TraceServices
-{
-	class IAnalysisSession;
-}
-
 namespace Insights
 {
 
@@ -181,8 +176,6 @@ protected:
 	void InitCommandList();
 
 	virtual void ConstructWidget(TSharedPtr<FTable> InTablePtr);
-	virtual TSharedRef<SWidget> ConstructSearchBox();
-	virtual TSharedRef<SWidget> ConstructAdvancedFiltersButton();
 	virtual TSharedRef<SWidget> ConstructHierarchyBreadcrumbTrail();
 	virtual TSharedPtr<SWidget> ConstructToolbar() { return nullptr; }
 	virtual TSharedPtr<SWidget> ConstructFooter() { return nullptr; }
@@ -245,7 +238,10 @@ protected:
 	void TreeView_OnGetChildren(FTableTreeNodePtr InParent, TArray<FTableTreeNodePtr>& OutChildren);
 
 	/** Called by STreeView when selection has changed. */
-	void TreeView_OnSelectionChanged(FTableTreeNodePtr SelectedItem, ESelectInfo::Type SelectInfo);
+	virtual void TreeView_OnSelectionChanged(FTableTreeNodePtr SelectedItem, ESelectInfo::Type SelectInfo);
+
+	/** Called by STreeView when a tree node is expanded or collapsed. */
+	virtual void TreeView_OnExpansionChanged(FTableTreeNodePtr TreeNode, bool bShouldBeExpanded);
 
 	/** Called by STreeView when a tree item is double clicked. */
 	virtual void TreeView_OnMouseButtonDoubleClick(FTableTreeNodePtr TreeNode);
@@ -268,20 +264,36 @@ protected:
 	// Filtering
 
 	void OnFilteringChanged();
+	FGraphEventRef StartApplyFiltersTask(FGraphEventRef Prerequisite = nullptr);
 
-	/** Populates the group and stat tree with items based on the current data. */
 	void ApplyFiltering();
 
 	bool ApplyAdvancedFiltersForNode(FTableTreeNodePtr NodePtr);
 
+	bool ApplyAdvancedFilters(const FTableTreeNodePtr& NodePtr);
+	bool virtual ApplyCustomAdvancedFilters(const FTableTreeNodePtr& NodePtr) { return true; };
+	virtual void AddCustomAdvancedFilters() {}
+
+	virtual TSharedRef<SWidget> ConstructAdvancedFiltersButton();
+
+	FReply OnAdvancedFiltersClicked();
+	void OnAdvancedFiltersChangesCommitted();
+
+	bool FilterConfigurator_HasFilters() const;
+	bool AdvancedFilters_ShouldBeEnabled() const;
+	FText AdvancedFilters_GetTooltipText() const;
+
+	void InitHierarchyFiltering();
+
 	bool ApplyHierarchicalFilterForNode(FTableTreeNodePtr NodePtr, bool bFilterIsEmpty);
 
 	/** Set all the nodes belonging to a subtree as visible. Returns true if the caller node should be expanded. */
 	bool MakeSubtreeVisible(FTableTreeNodePtr NodePtr, bool bFilterIsEmpty);
 
+	virtual TSharedRef<SWidget> ConstructSearchBox();
+
 	bool SearchBox_IsEnabled() const;
 	void SearchBox_OnTextChanged(const FText& InFilterText);
-
 	FText SearchBox_GetTooltipText() const;
 
 	////////////////////////////////////////////////////////////////////////////////////////////////////
@@ -292,7 +304,7 @@ protected:
 
 	void OnGroupingChanged();
 	void ApplyGrouping();
-
+	FGraphEventRef StartGroupingTask(FGraphEventRef Prerequisite = nullptr);
 	void CreateGroups(const TArray<TSharedPtr<FTreeNodeGrouping>>& Groupings);
 	void GroupNodesRec(const TArray<FTableTreeNodePtr>& Nodes, FTableTreeNode& ParentGroup, int32 GroupingDepth, const TArray<TSharedPtr<FTreeNodeGrouping>>& Groupings);
 
@@ -366,12 +378,12 @@ protected:
 	static const EColumnSortMode::Type GetDefaultColumnSortMode();
 	static const FName GetDefaultColumnBeingSorted();
 
-	void OnSortingChanged();
 	void CreateSortings();
-
 	void UpdateCurrentSortingByColumn();
 
+	void OnSortingChanged();
 	void ApplySorting();
+	FGraphEventRef StartSortingTask(FGraphEventRef Prerequisite = nullptr);
 	void SortTreeNodes(ITableCellValueSorter* InSorter, EColumnSortMode::Type InColumnSortMode);
 	void SortTreeNodesRec(FTableTreeNode& GroupNode, const ITableCellValueSorter& Sorter, EColumnSortMode::Type InColumnSortMode);
 
@@ -427,14 +439,12 @@ protected:
 	bool ContextMenu_HideAllColumns_CanExecute() const;
 	void ContextMenu_HideAllColumns_Execute();
 
+	////////////////////////////////////////////////////////////////////////////////////////////////////
 	//Async
+
 	virtual void OnPreAsyncUpdate();
 	virtual void OnPostAsyncUpdate();
 
-	FGraphEventRef StartSortTreeNodesTask(FGraphEventRef Prerequisite = nullptr);
-	FGraphEventRef StartCreateGroupsTask(FGraphEventRef Prerequisite = nullptr);
-	FGraphEventRef StartApplyFiltersTask(FGraphEventRef Prerequisite = nullptr);
-
 	void AddInProgressAsyncOperation(EAsyncOperationType InType) { EnumAddFlags(InProgressAsyncOperations, InType); }
 	bool HasInProgressAsyncOperation(EAsyncOperationType InType) const { return EnumHasAnyFlags(InProgressAsyncOperations, InType); }
 	void ClearInProgressAsyncOperations() { InProgressAsyncOperations = static_cast<EAsyncOperationType>(0); }
@@ -443,15 +453,6 @@ protected:
 
 	void CancelCurrentAsyncOp();
 
-	FReply OnAdvancedFiltersClicked();
-	bool AdvancedFilters_ShouldBeEnabled() const;
-	FText AdvancedFilters_GetTooltipText() const;
-	bool FilterConfigurator_HasFilters() const;
-	void OnAdvancedFiltersChangesCommited();
-	bool ApplyAdvancedFilters(const FTableTreeNodePtr& NodePtr);
-	bool virtual ApplyCustomAdvancedFilters(const FTableTreeNodePtr& NodePtr) { return true; };
-	virtual void AddCustomAdvancedFilters() {}
-
 	////////////////////////////////////////////////////////////////////////////////////////////////////
 
 	void SetExpandValueForChildGroups(FBaseTreeNode* InRoot, int32 InMaxExpandedNodes, int32 MaxDepthToExpand, bool InValue);
@@ -465,6 +466,9 @@ protected:
 
 	FText GetTreeViewBannerText() const { return TreeViewBannerText; }
 
+	////////////////////////////////////////////////////////////////////////////////////////////////////
+	// Presets
+
 	virtual void InitAvailableViewPresets() {};
 	const TArray<TSharedRef<ITableTreeViewPreset>>* GetAvailableViewPresets() const { return &AvailableViewPresets; }
 	FReply OnApplyViewPreset(const ITableTreeViewPreset* InPreset);
@@ -475,6 +479,8 @@ protected:
 	FText ViewPreset_GetSelectedText() const;
 	FText ViewPreset_GetSelectedToolTipText() const;
 
+	////////////////////////////////////////////////////////////////////////////////////////////////////
+
 	virtual void SearchForItem(TSharedPtr<FTableTaskCancellationToken> CancellationToken) {};
 
 	// Table data tasks should be tasks that operate read only operations on the data from the Table
@@ -486,7 +492,6 @@ protected:
 		Info->CancellationToken = MakeShared<FTableTaskCancellationToken>();
 		Info->Event = TGraphTask<T>::CreateTask().ConstructAndDispatchWhenReady(Info->CancellationToken, Forward<TArgs>(Args)...);
 		DataTaskInfos.Add(Info);
-
 		return Info;
 	}
 
@@ -495,8 +500,11 @@ protected:
 protected:
 	/** Table view model. */
 	TSharedPtr<FTable> Table;
+
+	TSharedPtr<FUICommandList> CommandList;
+
 	//////////////////////////////////////////////////
-	// Tree View, Columns
+	// Widget
 
 	/** The child STreeView widget. */
 	TSharedPtr<STreeView<FTableTreeNodePtr>> TreeView;
@@ -519,8 +527,6 @@ protected:
 	/** Name of the tree node that should be drawn as highlighted. */
 	FName HighlightedNodeName;
 
-	TSharedPtr<FUICommandList> CommandList;
-
 	//////////////////////////////////////////////////
 	// Tree Nodes
 
@@ -529,29 +535,48 @@ protected:
 	/** The root node of the tree. */
 	FTableTreeNodePtr Root;
 
-	/** Table (row) nodes. Each node corresponds to a table row. Index in this array corresponds to RowIndex in source table. */
+	/** Table row nodes. Each node corresponds to a table row. Index in this array corresponds to RowIndex in source table. */
 	TArray<FTableTreeNodePtr> TableTreeNodes;
 
-	/** A filtered array of group and nodes to be displayed in the tree widget. */
+	/** A filtered array of group nodes to be displayed in the tree widget. */
 	TArray<FTableTreeNodePtr> FilteredGroupNodes;
 
 	/** Currently expanded group nodes. */
 	TSet<FTableTreeNodePtr> ExpandedNodes;
 
 	/** If true, the expanded nodes have been saved before applying a text filter. */
-	bool bExpansionSaved;
+	bool bExpansionSaved = false;
+
+	static constexpr int32 MaxNodesToAutoExpand = 1000;
+	static constexpr int32 MaxDepthToAutoExpand = 4;
+	static constexpr int32 MaxNodesToExpand = 1000000;
+	static constexpr int32 MaxDepthToExpand = 100;
 
 	//////////////////////////////////////////////////
-	// Search box and filters
+	// Search box & the hierarchy filtering
 
 	/** The search box widget used to filter items displayed in the stats and groups tree. */
 	TSharedPtr<SSearchBox> SearchBox;
 
+	/** The filter collection. */
+	TSharedPtr<FTableTreeNodeFilterCollection> Filters;
+
 	/** The text based filter. */
 	TSharedPtr<FTableTreeNodeTextFilter> TextFilter;
 
-	/** The filter collection. */
-	TSharedPtr<FTableTreeNodeFilterCollection> Filters;
+	/** The text based filter actually used in the Hierarchy Filtering async task. */
+	TSharedPtr<FTableTreeNodeTextFilter> CurrentAsyncOpTextFilter;
+
+	//////////////////////////////////////////////////
+	// Node filtering
+
+	TSharedPtr<FFilterConfigurator> FilterConfigurator;
+	FDelegateHandle OnFilterChangesCommittedHandle;
+
+	/** The filter configurator actually used in the Hierarchy Filtering async task. */
+	FFilterConfigurator* CurrentAsyncOpFilterConfigurator = nullptr;
+
+	FFilterContext Context;
 
 	//////////////////////////////////////////////////
 	// Grouping
@@ -563,6 +588,9 @@ protected:
 
 	TSharedPtr<SBreadcrumbTrail<TSharedPtr<FTreeNodeGrouping>>> GroupingBreadcrumbTrail;
 
+	/** The groupings actually used in the Grouping async task. */
+	TArray<TSharedPtr<FTreeNodeGrouping>> CurrentAsyncOpGroupings;
+
 	//////////////////////////////////////////////////
 	// Sorting
 
@@ -576,10 +604,16 @@ protected:
 	FName ColumnBeingSorted;
 
 	/** How we sort the nodes? Ascending or Descending. */
-	EColumnSortMode::Type ColumnSortMode;
+	EColumnSortMode::Type ColumnSortMode = EColumnSortMode::None;
+
+	/** The sorter actually used in the Sorting async task. */
+	ITableCellValueSorter* CurrentAsyncOpSorter = nullptr;
+
+	/** The sort mode actually used in the Sorting async task. */
+	EColumnSortMode::Type CurrentAsyncOpColumnSortMode;
 
 	//////////////////////////////////////////////////
-	// Async
+	// Async Operations
 
 	bool bRunInAsyncMode = false;
 	bool bIsUpdateRunning = false;
@@ -591,12 +625,6 @@ protected:
 	EAsyncOperationType InProgressAsyncOperations = static_cast<EAsyncOperationType>(0);
 	TSharedPtr<class SAsyncOperationStatus> AsyncOperationStatus;
 	FStopwatch AsyncUpdateStopwatch;
-
-	TArray<TSharedPtr<FTreeNodeGrouping>> CurrentAsyncOpGroupings;
-	ITableCellValueSorter* CurrentAsyncOpSorter = nullptr;
-	EColumnSortMode::Type CurrentAsyncOpColumnSortMode;
-	TSharedPtr<FTableTreeNodeTextFilter> CurrentAsyncOpTextFilter;
-	FFilterConfigurator* CurrentAsyncOpFilterConfigurator = nullptr;
 	FAsyncOperationProgress AsyncOperationProgress;
 	FGraphEventRef DispatchEvent;
 	TArray<TSharedPtr<FTableTaskInfo>> DataTaskInfos;
@@ -604,22 +632,12 @@ protected:
 
 	//////////////////////////////////////////////////
 
-	TSharedPtr<FFilterConfigurator> FilterConfigurator;
-	FDelegateHandle OnFilterChangesCommitedHandle;
-	FFilterContext Context;
-
-	double StatsStartTime;
-	double StatsEndTime;
-
 	FText TreeViewBannerText;
 
 	TArray<TSharedRef<ITableTreeViewPreset>> AvailableViewPresets;
 	TSharedPtr<ITableTreeViewPreset> SelectedViewPreset;
 	TSharedPtr<SComboBox<TSharedRef<ITableTreeViewPreset>>> PresetComboBox;
 
-	static constexpr int32 MAX_NUMBER_OF_NODES_TO_EXPAND = 1000 * 1000;
-	static constexpr int32 MAX_DEPTH_TO_EXPAND = 100;
-
 	//////////////////////////////////////////////////
 	// Logging
 
@@ -629,23 +647,6 @@ protected:
 
 ////////////////////////////////////////////////////////////////////////////////////////////////////
 
-class SSessionTableTreeView : public STableTreeView
-{
-public:
-	virtual ~SSessionTableTreeView() override;
-
-protected:
-	virtual void ConstructWidget(TSharedPtr<FTable> InTablePtr) override;
-
-	/** Called when the analysis session has changed. */
-	void InsightsManager_OnSessionChanged();
-
-	/** The analysis session used to populate this widget. */
-	TSharedPtr<const TraceServices::IAnalysisSession> Session;
-};
-
-////////////////////////////////////////////////////////////////////////////////////////////////////
-
 class FTableTreeViewFilterAsyncTask
 {
 public:
@@ -744,7 +745,7 @@ public:
 
 	void DoTask(ENamedThreads::Type CurrentThread, const FGraphEventRef& MyCompletionGraphEvent)
 	{
-		// The role of this task is to keep the STableTreeView object alive until the task and it's prerequisits are completed and to destroy it on the game thread.
+		// The role of this task is to keep the STableTreeView object alive until the task and it's prerequisites are completed and to destroy it on the game thread.
 		FGraphEventRef Event = TableTreeViewPtr->AsyncCompleteTaskEvent;
 		if (TableTreeViewPtr.IsValid())
 		{
diff --git a/Engine/Source/Developer/TraceInsights/Private/Insights/TaskGraphProfiler/Widgets/STaskTableTreeView.h b/Engine/Source/Developer/TraceInsights/Private/Insights/TaskGraphProfiler/Widgets/STaskTableTreeView.h
--- a/Engine/Source/Developer/TraceInsights/Private/Insights/TaskGraphProfiler/Widgets/STaskTableTreeView.h
+++ b/Engine/Source/Developer/TraceInsights/Private/Insights/TaskGraphProfiler/Widgets/STaskTableTreeView.h
@@ -7,7 +7,7 @@
 
 // Insights
 #include "Insights/TaskGraphProfiler/ViewModels/TaskTable.h"
-#include "Insights/Table/Widgets/STableTreeView.h"
+#include "Insights/Table/Widgets/SSessionTableTreeView.h"
 
 class FMenuBuilder;
 
diff --git a/Engine/Source/Developer/TraceInsights/Private/Insights/ViewModels/FilterConfigurator.cpp b/Engine/Source/Developer/TraceInsights/Private/Insights/ViewModels/FilterConfigurator.cpp
--- a/Engine/Source/Developer/TraceInsights/Private/Insights/ViewModels/FilterConfigurator.cpp
+++ b/Engine/Source/Developer/TraceInsights/Private/Insights/ViewModels/FilterConfigurator.cpp
@@ -47,7 +47,7 @@ FFilterConfigurator& FFilterConfigurator::operator=(const FFilterConfigurator& O
 	ComputeUsedKeys();
 	RootNode->ProcessFilter();
 
-	OnChangesCommitedEvent.Broadcast();
+	OnChangesCommittedEvent.Broadcast();
 
 	return *this;
 }
diff --git a/Engine/Source/Developer/TraceInsights/Private/Insights/ViewModels/FilterConfigurator.h b/Engine/Source/Developer/TraceInsights/Private/Insights/ViewModels/FilterConfigurator.h
--- a/Engine/Source/Developer/TraceInsights/Private/Insights/ViewModels/FilterConfigurator.h
+++ b/Engine/Source/Developer/TraceInsights/Private/Insights/ViewModels/FilterConfigurator.h
@@ -3,9 +3,9 @@
 #pragma once
 
 #include "CoreMinimal.h"
-#include "Insights/ViewModels/IFilterExecutor.h"
-#include "Insights/ViewModels/Filters.h"
 #include "Insights/ViewModels/FilterConfiguratorNode.h"
+#include "Insights/ViewModels/Filters.h"
+#include "Insights/ViewModels/IFilterExecutor.h"
 
 namespace Insights
 {
@@ -21,11 +21,12 @@ public:
 	FFilterConfigurator& operator=(const FFilterConfigurator& Other);
 
 	bool operator==(const FFilterConfigurator& Other) const;
-
 	bool operator!=(const FFilterConfigurator& Other) const { return !(*this == Other); }
 
 	virtual ~FFilterConfigurator();
 
+	bool IsEmpty() const { return RootNode->GetChildren().Num() == 0; }
+
 	FFilterConfiguratorNodePtr GetRootNode() { return RootNode; }
 
 	virtual bool ApplyFilters(const FFilterContext& Context) const override;
@@ -46,15 +47,15 @@ private:
 	FOnDestroyedEvent OnDestroyedEvent;
 
 	////////////////////////////////////////////////////////////////////////////////////////////////////
-	// OnChangesCommitedEvent
+	// OnChangesCommittedEvent
 
 public:
 	/** The event to execute when the changes to the Filter Widget are saved by clicking on the OK Button. */
-	DECLARE_MULTICAST_DELEGATE(FOnChangesCommitedEvent);
-	FOnChangesCommitedEvent& GetOnChangesCommitedEvent() { return OnChangesCommitedEvent; }
+	DECLARE_MULTICAST_DELEGATE(FOnChangesCommittedEvent);
+	FOnChangesCommittedEvent& GetOnChangesCommittedEvent() { return OnChangesCommittedEvent; }
 
 private:
-	FOnChangesCommitedEvent OnChangesCommitedEvent;
+	FOnChangesCommittedEvent OnChangesCommittedEvent;
 
 	////////////////////////////////////////////////////////////////////////////////////////////////////
 
diff --git a/Engine/Source/Developer/TraceInsights/Private/Insights/ViewModels/FilterConfiguratorNode.cpp b/Engine/Source/Developer/TraceInsights/Private/Insights/ViewModels/FilterConfiguratorNode.cpp
--- a/Engine/Source/Developer/TraceInsights/Private/Insights/ViewModels/FilterConfiguratorNode.cpp
+++ b/Engine/Source/Developer/TraceInsights/Private/Insights/ViewModels/FilterConfiguratorNode.cpp
@@ -7,7 +7,7 @@
 // Insights
 #include "Insights/ViewModels/Filters.h"
 
-#define LOCTEXT_NAMESPACE "FFilterConfiguratorNode"
+#define LOCTEXT_NAMESPACE "Insights::FFilterConfiguratorNode"
 
 namespace Insights
 {
@@ -49,7 +49,7 @@ FFilterConfiguratorNode& FFilterConfiguratorNode::operator=(const FFilterConfigu
 	TextBoxValue = Other.TextBoxValue;
 	SetExpansion(Other.IsExpanded());
 
-	for (Insights::FBaseTreeNodePtr Child : Other.GetChildren())
+	for (FBaseTreeNodePtr Child : Other.GetChildren())
 	{
 		GetChildrenMutable().Add(MakeShared<FFilterConfiguratorNode>(*StaticCastSharedPtr<FFilterConfiguratorNode>(Child)));
 	}
@@ -91,7 +91,7 @@ const TArray<TSharedPtr<struct FFilterGroupOperator>>& FFilterConfiguratorNode::
 
 void FFilterConfiguratorNode::DeleteChildNode(FFilterConfiguratorNodePtr InNode)
 {
-	Insights::FBaseTreeNodePtr Node = StaticCastSharedPtr<Insights::FBaseTreeNode>(InNode);
+	FBaseTreeNodePtr Node = StaticCastSharedPtr<FBaseTreeNode>(InNode);
 	GetChildrenMutable().Remove(Node);
 }
 
@@ -99,7 +99,7 @@ void FFilterConfiguratorNode::DeleteChildNode(FFilterConfiguratorNodePtr InNode)
 
 void FFilterConfiguratorNode::SetGroupPtrForChildren()
 {
-	for (Insights::FBaseTreeNodePtr Child : GetChildrenMutable())
+	for (FBaseTreeNodePtr Child : GetChildrenMutable())
 	{
 		FFilterConfiguratorNodePtr CastedChild = StaticCastSharedPtr<FFilterConfiguratorNode>(Child);
 		CastedChild->SetGroupPtrForChildren();
@@ -129,7 +129,7 @@ void FFilterConfiguratorNode::SetSelectedFilter(TSharedPtr<struct FFilter> InSel
 		SetSelectedFilterOperator(SelectedFilter->GetSupportedOperators()->GetData()[0]);
 
 		AvailableFilterOperators->Empty();
-		Insights::SupportedOperatorsArrayPtr AvailableOperators = InSelectedFilter->SupportedOperators;
+		SupportedOperatorsArrayPtr AvailableOperators = InSelectedFilter->SupportedOperators;
 		for (auto& FilterOperator : *AvailableOperators)
 		{
 			AvailableFilterOperators->Add(FilterOperator);
@@ -143,8 +143,8 @@ void FFilterConfiguratorNode::ProcessFilter()
 {
 	if (IsGroup())
 	{
-		TArray<Insights::FBaseTreeNodePtr> ChildArray = GetChildrenMutable();
-		for (Insights::FBaseTreeNodePtr Child : ChildArray)
+		TArray<FBaseTreeNodePtr> ChildArray = GetChildrenMutable();
+		for (FBaseTreeNodePtr Child : ChildArray)
 		{
 			FFilterConfiguratorNodePtr CastedChild = StaticCastSharedPtr<FFilterConfiguratorNode>(Child);
 			CastedChild->ProcessFilter();
diff --git a/Engine/Source/Developer/TraceInsights/Private/Insights/ViewModels/FilterConfiguratorNode.h b/Engine/Source/Developer/TraceInsights/Private/Insights/ViewModels/FilterConfiguratorNode.h
--- a/Engine/Source/Developer/TraceInsights/Private/Insights/ViewModels/FilterConfiguratorNode.h
+++ b/Engine/Source/Developer/TraceInsights/Private/Insights/ViewModels/FilterConfiguratorNode.h
@@ -43,7 +43,7 @@ typedef TWeakPtr<class FFilterConfiguratorNode> FFilterConfiguratorNodeWeak;
 /**
  * Class used to store information about a filter node
  */
-class FFilterConfiguratorNode : public Insights::FBaseTreeNode
+class FFilterConfiguratorNode : public FBaseTreeNode
 {
 	INSIGHTS_DECLARE_RTTI(FFilterConfiguratorNode, FBaseTreeNode)
 
diff --git a/Engine/Source/Developer/TraceInsights/Private/Insights/ViewModels/Filters.cpp b/Engine/Source/Developer/TraceInsights/Private/Insights/ViewModels/Filters.cpp
--- a/Engine/Source/Developer/TraceInsights/Private/Insights/ViewModels/Filters.cpp
+++ b/Engine/Source/Developer/TraceInsights/Private/Insights/ViewModels/Filters.cpp
@@ -10,7 +10,7 @@
 #include "Insights/ViewModels/FilterConfigurator.h"
 #include "Insights/Widgets/SAdvancedFilter.h"
 
-#define LOCTEXT_NAMESPACE "SFilterService"
+#define LOCTEXT_NAMESPACE "Insights::Filters"
 
 namespace Insights
 {
@@ -18,6 +18,44 @@ namespace Insights
 INSIGHTS_IMPLEMENT_RTTI(FFilter)
 INSIGHTS_IMPLEMENT_RTTI(FFilterWithSuggestions)
 
+////////////////////////////////////////////////////////////////////////////////////////////////////
+// FFilterStorage
+////////////////////////////////////////////////////////////////////////////////////////////////////
+
+FFilterStorage::FFilterStorage()
+{
+	DoubleOperators = MakeShared<TArray<TSharedPtr<IFilterOperator>>>();
+	DoubleOperators->Add(StaticCastSharedRef<IFilterOperator>(MakeShared<FFilterOperator<double>>(EFilterOperator::Lt, TEXT("<"), std::less<>{})));
+	DoubleOperators->Add(StaticCastSharedRef<IFilterOperator>(MakeShared<FFilterOperator<double>>(EFilterOperator::Lte, TEXT("\u2264"), std::less_equal<>())));
+	DoubleOperators->Add(StaticCastSharedRef<IFilterOperator>(MakeShared<FFilterOperator<double>>(EFilterOperator::Eq, TEXT("="), std::equal_to<>())));
+	DoubleOperators->Add(StaticCastSharedRef<IFilterOperator>(MakeShared<FFilterOperator<double>>(EFilterOperator::Gt, TEXT(">"), std::greater<>())));
+	DoubleOperators->Add(StaticCastSharedRef<IFilterOperator>(MakeShared<FFilterOperator<double>>(EFilterOperator::Gte, TEXT("\u2265"), std::greater_equal<>())));
+
+	IntegerOperators = MakeShared<TArray<TSharedPtr<IFilterOperator>>>();
+	IntegerOperators->Add(StaticCastSharedRef<IFilterOperator>(MakeShared<FFilterOperator<int64>>(EFilterOperator::Lt, TEXT("<"), std::less<>{})));
+	IntegerOperators->Add(StaticCastSharedRef<IFilterOperator>(MakeShared<FFilterOperator<int64>>(EFilterOperator::Lte, TEXT("\u2264"), std::less_equal<>())));
+	IntegerOperators->Add(StaticCastSharedRef<IFilterOperator>(MakeShared<FFilterOperator<int64>>(EFilterOperator::Eq, TEXT("="), std::equal_to<>())));
+	IntegerOperators->Add(StaticCastSharedRef<IFilterOperator>(MakeShared<FFilterOperator<int64>>(EFilterOperator::Gt, TEXT(">"), std::greater<>())));
+	IntegerOperators->Add(StaticCastSharedRef<IFilterOperator>(MakeShared<FFilterOperator<int64>>(EFilterOperator::Gte, TEXT("\u2265"), std::greater_equal<>())));
+
+	StringOperators = MakeShared<TArray<TSharedPtr<IFilterOperator>>>();
+	StringOperators->Add(StaticCastSharedRef<IFilterOperator>(MakeShared<FFilterOperator<FString>>(EFilterOperator::Eq, TEXT("IS"), [](const FString& lhs, const FString& rhs) { return lhs.Equals(rhs); })));
+	StringOperators->Add(StaticCastSharedRef<IFilterOperator>(MakeShared<FFilterOperator<FString>>(EFilterOperator::NotEq, TEXT("IS NOT"), [](const FString& lhs, const FString& rhs) { return !lhs.Equals(rhs); })));
+	StringOperators->Add(StaticCastSharedRef<IFilterOperator>(MakeShared<FFilterOperator<FString>>(EFilterOperator::Contains, TEXT("CONTAINS"), [](const FString& lhs, const FString& rhs) { return lhs.Contains(rhs); })));
+	StringOperators->Add(StaticCastSharedRef<IFilterOperator>(MakeShared<FFilterOperator<FString>>(EFilterOperator::NotContains, TEXT("NOT CONTAINS"), [](const FString& lhs, const FString& rhs) { return !lhs.Contains(rhs); })));
+
+	FilterGroupOperators.Add(MakeShared<FFilterGroupOperator>(EFilterGroupOperator::And, LOCTEXT("AllOf", "All Of (AND)"), LOCTEXT("AllOfDesc", "All of the children must be true for the group to return true. Equivalent to an AND operation.")));
+	FilterGroupOperators.Add(MakeShared<FFilterGroupOperator>(EFilterGroupOperator::Or, LOCTEXT("AnyOf", "Any Of (OR)"), LOCTEXT("AnyOfDesc", "Any of the children must be true for the group to return true. Equivalent to an OR operation.")));
+}
+
+////////////////////////////////////////////////////////////////////////////////////////////////////
+
+FFilterStorage::~FFilterStorage()
+{
+}
+
+////////////////////////////////////////////////////////////////////////////////////////////////////
+// FFilterService
 ////////////////////////////////////////////////////////////////////////////////////////////////////
 
 TSharedPtr<FFilterService> FFilterService::Instance;
@@ -26,6 +64,20 @@ FName const FFilterService::FilterConfiguratorTabId(TEXT("FilterConfigurator"));
 
 ////////////////////////////////////////////////////////////////////////////////////////////////////
 
+void FFilterService::Initialize()
+{
+	Instance = MakeShared<FFilterService>();
+}
+
+////////////////////////////////////////////////////////////////////////////////////////////////////
+
+void FFilterService::Shutdown()
+{
+	Instance = nullptr;
+}
+
+////////////////////////////////////////////////////////////////////////////////////////////////////
+
 FFilterService::FFilterService()
 {
 	RegisterTabSpawner();
@@ -35,7 +87,7 @@ FFilterService::FFilterService()
 
 FFilterService::~FFilterService()
 {
-
+	UnregisterTabSpawner();
 }
 
 ////////////////////////////////////////////////////////////////////////////////////////////////////
@@ -66,6 +118,7 @@ TSharedRef<SDockTab> FFilterService::SpawnTab(const FSpawnTabArgs& Args)
 		OwnerWindow->Resize(FVector2D(600 * DPIScaleFactor, 400 * DPIScaleFactor));
 	}
 
+	check(PendingWidget.IsValid());
 	DockTab->SetContent(PendingWidget.ToSharedRef());
 	PendingWidget->SetParentTab(DockTab);
 
@@ -86,6 +139,13 @@ void FFilterService::RegisterTabSpawner()
 
 ////////////////////////////////////////////////////////////////////////////////////////////////////
 
+void FFilterService::UnregisterTabSpawner()
+{
+	FGlobalTabmanager::Get()->UnregisterNomadTabSpawner(FilterConfiguratorTabId);
+}
+
+////////////////////////////////////////////////////////////////////////////////////////////////////
+
 } // namespace Insights
 
 #undef LOCTEXT_NAMESPACE
\ No newline at end of file
diff --git a/Engine/Source/Developer/TraceInsights/Private/Insights/ViewModels/Filters.h b/Engine/Source/Developer/TraceInsights/Private/Insights/ViewModels/Filters.h
--- a/Engine/Source/Developer/TraceInsights/Private/Insights/ViewModels/Filters.h
+++ b/Engine/Source/Developer/TraceInsights/Private/Insights/ViewModels/Filters.h
@@ -9,8 +9,6 @@
 
 #include "Insights/Common/SimpleRtti.h"
 
-#define LOCTEXT_NAMESPACE "Filters"
-
 class FSpawnTabArgs;
 class SDockTab;
 class SWidget;
@@ -60,12 +58,14 @@ class FFilterOperator : public IFilterOperator
 public:
 	typedef TFunction<bool(T, T)> OperatorFunc;
 
+public:
 	FFilterOperator(EFilterOperator InKey, FString InName, OperatorFunc InFunc)
 		: Func(InFunc)
 		, Key(InKey)
 		, Name(InName)
 	{
 	}
+
 	virtual ~FFilterOperator()
 	{
 	}
@@ -132,6 +132,7 @@ public:
 		, SupportedOperators(InSupportedOperators)
 	{
 	}
+
 	virtual ~FFilter()
 	{
 	}
@@ -172,31 +173,8 @@ public:
 class FFilterStorage
 {
 public:
-	FFilterStorage()
-	{
-		DoubleOperators = MakeShared<TArray<TSharedPtr<IFilterOperator>>>();
-		DoubleOperators->Add(StaticCastSharedRef<IFilterOperator>(MakeShared<FFilterOperator<double>>(EFilterOperator::Lt, TEXT("<"), std::less<>{})));
-		DoubleOperators->Add(StaticCastSharedRef<IFilterOperator>(MakeShared<FFilterOperator<double>>(EFilterOperator::Lte, TEXT("\u2264"), std::less_equal<>())));
-		DoubleOperators->Add(StaticCastSharedRef<IFilterOperator>(MakeShared<FFilterOperator<double>>(EFilterOperator::Eq, TEXT("="), std::equal_to<>())));
-		DoubleOperators->Add(StaticCastSharedRef<IFilterOperator>(MakeShared<FFilterOperator<double>>(EFilterOperator::Gt, TEXT(">"), std::greater<>())));
-		DoubleOperators->Add(StaticCastSharedRef<IFilterOperator>(MakeShared<FFilterOperator<double>>(EFilterOperator::Gte, TEXT("\u2265"), std::greater_equal<>())));
-
-		IntegerOperators = MakeShared<TArray<TSharedPtr<IFilterOperator>>>();
-		IntegerOperators->Add(StaticCastSharedRef<IFilterOperator>(MakeShared<FFilterOperator<int64>>(EFilterOperator::Lt, TEXT("<"), std::less<>{})));
-		IntegerOperators->Add(StaticCastSharedRef<IFilterOperator>(MakeShared<FFilterOperator<int64>>(EFilterOperator::Lte, TEXT("\u2264"), std::less_equal<>())));
-		IntegerOperators->Add(StaticCastSharedRef<IFilterOperator>(MakeShared<FFilterOperator<int64>>(EFilterOperator::Eq, TEXT("="), std::equal_to<>())));
-		IntegerOperators->Add(StaticCastSharedRef<IFilterOperator>(MakeShared<FFilterOperator<int64>>(EFilterOperator::Gt, TEXT(">"), std::greater<>())));
-		IntegerOperators->Add(StaticCastSharedRef<IFilterOperator>(MakeShared<FFilterOperator<int64>>(EFilterOperator::Gte, TEXT("\u2265"), std::greater_equal<>())));
-
-		StringOperators = MakeShared<TArray<TSharedPtr<IFilterOperator>>>();
-		StringOperators->Add(StaticCastSharedRef<IFilterOperator>(MakeShared<FFilterOperator<FString>>(EFilterOperator::Eq, TEXT("IS"), [](const FString& lhs, const FString& rhs) { return lhs.Equals(rhs); })));
-		StringOperators->Add(StaticCastSharedRef<IFilterOperator>(MakeShared<FFilterOperator<FString>>(EFilterOperator::NotEq, TEXT("IS NOT"), [](const FString& lhs, const FString& rhs) { return !lhs.Equals(rhs); })));
-		StringOperators->Add(StaticCastSharedRef<IFilterOperator>(MakeShared<FFilterOperator<FString>>(EFilterOperator::Contains, TEXT("CONTAINS"), [](const FString& lhs, const FString& rhs) { return lhs.Contains(rhs); })));
-		StringOperators->Add(StaticCastSharedRef<IFilterOperator>(MakeShared<FFilterOperator<FString>>(EFilterOperator::NotContains, TEXT("NOT CONTAINS"), [](const FString& lhs, const FString& rhs) { return !lhs.Contains(rhs); })));
-
-		FilterGroupOperators.Add(MakeShared<FFilterGroupOperator>(EFilterGroupOperator::And, LOCTEXT("AllOf", "All Of (AND)"), LOCTEXT("AllOfDesc", "All of the children must be true for the group to return true. Equivalent to an AND operation.")));
-		FilterGroupOperators.Add(MakeShared<FFilterGroupOperator>(EFilterGroupOperator::Or, LOCTEXT("AnyOf", "Any Of (OR)"), LOCTEXT("AnyOfDesc", "Any of the children must be true for the group to return true. Equivalent to an OR operation.")));
-	}
+	FFilterStorage();
+	~FFilterStorage();
 
 	const TArray<TSharedPtr<FFilterGroupOperator>>& GetFilterGroupOperators() const { return FilterGroupOperators; }
 
@@ -220,10 +198,12 @@ public:
 	FFilterService();
 	~FFilterService();
 
-	static void CreateInstance() { Instance = MakeShared<FFilterService>(); }
+	static void Initialize();
+	static void Shutdown();
 	static TSharedPtr<FFilterService> Get() { return Instance; }
 
 	void RegisterTabSpawner();
+	void UnregisterTabSpawner();
 	TSharedRef<SDockTab> SpawnTab(const FSpawnTabArgs& Args);
 
 	const TArray<TSharedPtr<FFilterGroupOperator>>& GetFilterGroupOperators() const { return FilterStorage.GetFilterGroupOperators(); }
@@ -250,6 +230,7 @@ class FFilterContext
 public:
 	typedef TVariant<double, int64, FString> ContextData;
 
+public:
 	template<typename T>
 	void AddFilterData(int32 Key, const T& InData)
 	{
@@ -290,5 +271,3 @@ private:
 ////////////////////////////////////////////////////////////////////////////////////////////////////
 
 } // namespace Insights
-
-#undef LOCTEXT_NAMESPACE
diff --git a/Engine/Source/Developer/TraceInsights/Private/Insights/ViewModels/ThreadTimingTrack.cpp b/Engine/Source/Developer/TraceInsights/Private/Insights/ViewModels/ThreadTimingTrack.cpp
--- a/Engine/Source/Developer/TraceInsights/Private/Insights/ViewModels/ThreadTimingTrack.cpp
+++ b/Engine/Source/Developer/TraceInsights/Private/Insights/ViewModels/ThreadTimingTrack.cpp
@@ -781,7 +781,7 @@ FThreadTimingTrack::~FThreadTimingTrack()
 {
 	if (FilterConfigurator.IsValid())
 	{
-		FilterConfigurator->GetOnChangesCommitedEvent().Remove(OnFilterChangesCommitedHandle);
+		FilterConfigurator->GetOnChangesCommittedEvent().Remove(OnFilterChangesCommittedHandle);
 	}
 }
 
@@ -1758,13 +1758,13 @@ void FThreadTimingTrack::OnFilterTrackClicked()
 	}
 	else
 	{
-		FilterConfigurator->GetOnChangesCommitedEvent().Remove(OnFilterChangesCommitedHandle);
+		FilterConfigurator->GetOnChangesCommittedEvent().Remove(OnFilterChangesCommittedHandle);
 
 		// Make a copy, so it will not affect other tracks that shares same filter.
 		FilterConfigurator = MakeShared<FFilterConfigurator>(*FilterConfigurator);
 	}
 
-	OnFilterChangesCommitedHandle = FilterConfigurator->GetOnChangesCommitedEvent().AddLambda([this]()
+	OnFilterChangesCommittedHandle = FilterConfigurator->GetOnChangesCommittedEvent().AddLambda([this]()
 		{
 			this->SetDirtyFlag();
 		});
diff --git a/Engine/Source/Developer/TraceInsights/Private/Insights/ViewModels/ThreadTimingTrack.h b/Engine/Source/Developer/TraceInsights/Private/Insights/ViewModels/ThreadTimingTrack.h
--- a/Engine/Source/Developer/TraceInsights/Private/Insights/ViewModels/ThreadTimingTrack.h
+++ b/Engine/Source/Developer/TraceInsights/Private/Insights/ViewModels/ThreadTimingTrack.h
@@ -210,7 +210,7 @@ private:
 	FThreadTimingSharedState& SharedState;
 
 	TSharedPtr<Insights::FFilterConfigurator> FilterConfigurator;
-	FDelegateHandle OnFilterChangesCommitedHandle;
+	FDelegateHandle OnFilterChangesCommittedHandle;
 
 	// Search cache
 	mutable TTimingEventSearchCache<TraceServices::FTimingProfilerEvent> SearchCache;
diff --git a/Engine/Source/Developer/TraceInsights/Private/Insights/Widgets/SAdvancedFilter.cpp b/Engine/Source/Developer/TraceInsights/Private/Insights/Widgets/SAdvancedFilter.cpp
--- a/Engine/Source/Developer/TraceInsights/Private/Insights/Widgets/SAdvancedFilter.cpp
+++ b/Engine/Source/Developer/TraceInsights/Private/Insights/Widgets/SAdvancedFilter.cpp
@@ -17,7 +17,7 @@
 #include "Insights/Widgets/SFilterConfigurator.h"
 #include "Insights/ViewModels/FilterConfigurator.h"
 
-#define LOCTEXT_NAMESPACE "SAdvancedFilter"
+#define LOCTEXT_NAMESPACE "Insights::SAdvancedFilter"
 
 namespace Insights
 {
diff --git a/Engine/Source/Developer/TraceInsights/Private/Insights/Widgets/SFilterConfigurator.cpp b/Engine/Source/Developer/TraceInsights/Private/Insights/Widgets/SFilterConfigurator.cpp
--- a/Engine/Source/Developer/TraceInsights/Private/Insights/Widgets/SFilterConfigurator.cpp
+++ b/Engine/Source/Developer/TraceInsights/Private/Insights/Widgets/SFilterConfigurator.cpp
@@ -17,7 +17,7 @@
 #include "Insights/Widgets/SFilterConfiguratorRow.h"
 #include "Insights/ViewModels/FilterConfigurator.h"
 
-#define LOCTEXT_NAMESPACE "SFilterConfigurator"
+#define LOCTEXT_NAMESPACE "Insights::SFilterConfigurator"
 
 namespace Insights
 {
@@ -111,11 +111,11 @@ END_SLATE_FUNCTION_BUILD_OPTIMIZATION
 
 void SFilterConfigurator::TreeView_OnGetChildren(FFilterConfiguratorNodePtr InParent, TArray<FFilterConfiguratorNodePtr>& OutChildren)
 {
-	const TArray<Insights::FBaseTreeNodePtr>& Children = InParent->GetChildren();
+	const TArray<FBaseTreeNodePtr>& Children = InParent->GetChildren();
 	OutChildren.Reset(Children.Num());
-	for (const Insights::FBaseTreeNodePtr& Child : Children)
+	for (const FBaseTreeNodePtr& Child : Children)
 	{
-		OutChildren.Add(StaticCastSharedPtr<FFilterConfiguratorNode, Insights::FBaseTreeNode>(Child));
+		OutChildren.Add(StaticCastSharedPtr<FFilterConfiguratorNode, FBaseTreeNode>(Child));
 	}
 }
 
@@ -135,12 +135,14 @@ TSharedRef<ITableRow> SFilterConfigurator::TreeView_OnGenerateRow(FFilterConfigu
 void SFilterConfigurator::SetInitialExpansionRec(const FFilterConfiguratorNodePtr& Node, bool Value)
 {
 	TreeView->SetItemExpansion(Node, Value);
-	for (const Insights::FBaseTreeNodePtr& Child : Node->GetChildren())
+	for (const FBaseTreeNodePtr& Child : Node->GetChildren())
 	{
-		SetInitialExpansionRec(StaticCastSharedPtr<FFilterConfiguratorNode, Insights::FBaseTreeNode>(Child), Value);
+		SetInitialExpansionRec(StaticCastSharedPtr<FFilterConfiguratorNode, FBaseTreeNode>(Child), Value);
 	}
 }
 
+////////////////////////////////////////////////////////////////////////////////////////////////////
+
 } // namespace Insights
 
 #undef LOCTEXT_NAMESPACE
diff --git a/Engine/Source/Developer/TraceInsights/Private/Insights/Widgets/SFilterConfiguratorRow.cpp b/Engine/Source/Developer/TraceInsights/Private/Insights/Widgets/SFilterConfiguratorRow.cpp
--- a/Engine/Source/Developer/TraceInsights/Private/Insights/Widgets/SFilterConfiguratorRow.cpp
+++ b/Engine/Source/Developer/TraceInsights/Private/Insights/Widgets/SFilterConfiguratorRow.cpp
@@ -418,7 +418,7 @@ FReply SFilterConfiguratorRow::AddFilter_OnClicked()
 
 FReply SFilterConfiguratorRow::DeleteFilter_OnClicked()
 {
-	Insights::FBaseTreeNodeWeak GroupWeakPtr = FilterConfiguratorNodePtr->GetGroupPtr();
+	FBaseTreeNodeWeak GroupWeakPtr = FilterConfiguratorNodePtr->GetGroupPtr();
 	if (GroupWeakPtr.IsValid())
 	{
 		FFilterConfiguratorNodePtr GroupPtr = StaticCastSharedPtr<FFilterConfiguratorNode>(GroupWeakPtr.Pin());
@@ -449,7 +449,7 @@ FReply SFilterConfiguratorRow::AddGroup_OnClicked()
 
 FReply SFilterConfiguratorRow::DeleteGroup_OnClicked()
 {
-	Insights::FBaseTreeNodeWeak GroupWeakPtr = FilterConfiguratorNodePtr->GetGroupPtr();
+	FBaseTreeNodeWeak GroupWeakPtr = FilterConfiguratorNodePtr->GetGroupPtr();
 	if (GroupWeakPtr.IsValid())
 	{
 		FFilterConfiguratorNodePtr GroupPtr = StaticCastSharedPtr<FFilterConfiguratorNode>(GroupWeakPtr.Pin());
