[Insights] Refactored tree table view and filtering system in preparation to separate node filtering from hierarchy filtering.
  - Changed tooltip to appear only when hovering the icon  of a tree item (i.e. not when hovering the name).
  - Fixed tooltip to close automaticaly when context menu is open.

diff --git a/Engine/Source/Developer/TraceInsights/Private/Insights/ContextSwitches/ViewModels/ContextSwitchesSharedState.cpp b/Engine/Source/Developer/TraceInsights/Private/Insights/ContextSwitches/ViewModels/ContextSwitchesSharedState.cpp
--- a/Engine/Source/Developer/TraceInsights/Private/Insights/ContextSwitches/ViewModels/ContextSwitchesSharedState.cpp
+++ b/Engine/Source/Developer/TraceInsights/Private/Insights/ContextSwitches/ViewModels/ContextSwitchesSharedState.cpp
@@ -247,22 +247,25 @@ void FContextSwitchesSharedState::BuildSubMenu(FMenuBuilder& InMenuBuilder)
 
 ////////////////////////////////////////////////////////////////////////////////////////////////////
 
-void FContextSwitchesSharedState::AddQuickFindFilters(TSharedPtr<class FFilterConfigurator> FilterConfigurator)
+void FContextSwitchesSharedState::AddQuickFindFilters(TSharedPtr<FFilterConfigurator> FilterConfigurator)
 {
-	TSharedPtr<TArray<TSharedPtr<struct FFilter>>>& AvailableFilters = FilterConfigurator->GetAvailableFilters();
-
 	TSharedPtr<TArray<TSharedPtr<IFilterOperator>>> CoreEventNameFilterOperators = MakeShared<TArray<TSharedPtr<IFilterOperator>>>();
 	CoreEventNameFilterOperators->Add(StaticCastSharedRef<IFilterOperator>(MakeShared<FFilterOperator<int64>>(EFilterOperator::Eq, TEXT("Is"), [](int64 lhs, int64 rhs) { return lhs == rhs; })));
 
-	TSharedPtr<FFilterWithSuggestions> TimerNameFilter = MakeShared<FFilterWithSuggestions>(static_cast<int32>(EFilterField::CoreEventName), LOCTEXT("CoreEventName", "Core Event Name"), LOCTEXT("CoreEventName", "Core Event Name"), EFilterDataType::StringInt64Pair, CoreEventNameFilterOperators);
-	TimerNameFilter->Callback = [this](const FString& Text, TArray<FString>& OutSuggestions)
+	TSharedRef<FFilterWithSuggestions> TimerNameFilter = MakeShared<FFilterWithSuggestions>(
+		static_cast<int32>(EFilterField::CoreEventName),
+		LOCTEXT("CoreEventName", "Core Event Name"),
+		LOCTEXT("CoreEventName", "Core Event Name"),
+		EFilterDataType::StringInt64Pair,
+		MakeShared<FCoreEventNameFilterValueConverter>(),
+		CoreEventNameFilterOperators);
+
+	TimerNameFilter->SetCallback([this](const FString& Text, TArray<FString>& OutSuggestions)
 	{
 		this->PopulateCoreEventNameSuggestionList(Text, OutSuggestions);
-	};
-
-	TimerNameFilter->Converter = MakeShared<FCoreEventNameFilterValueConverter>();
+	});
 
-	AvailableFilters->Add(TimerNameFilter);
+	FilterConfigurator->Add(TimerNameFilter);
 }
 
 ////////////////////////////////////////////////////////////////////////////////////////////////////
diff --git a/Engine/Source/Developer/TraceInsights/Private/Insights/ContextSwitches/ViewModels/CpuCoreTimingTrack.cpp b/Engine/Source/Developer/TraceInsights/Private/Insights/ContextSwitches/ViewModels/CpuCoreTimingTrack.cpp
--- a/Engine/Source/Developer/TraceInsights/Private/Insights/ContextSwitches/ViewModels/CpuCoreTimingTrack.cpp
+++ b/Engine/Source/Developer/TraceInsights/Private/Insights/ContextSwitches/ViewModels/CpuCoreTimingTrack.cpp
@@ -427,7 +427,7 @@ bool FCpuCoreTimingTrack::HasCustomFilter() const
 	}
 
 	TSharedPtr<FFilterConfigurator> FilterConfigurator = TimingView->GetFilterConfigurator();
-	return FilterConfigurator.IsValid() && FilterConfigurator->GetRootNode().IsValid() && FilterConfigurator->GetRootNode()->GetChildren().Num() > 0;
+	return FilterConfigurator.IsValid() && !FilterConfigurator->IsEmpty();
 }
 
 ////////////////////////////////////////////////////////////////////////////////////////////////////
diff --git a/Engine/Source/Developer/TraceInsights/Private/Insights/MemoryProfiler/ViewModels/MemAllocTable.cpp b/Engine/Source/Developer/TraceInsights/Private/Insights/MemoryProfiler/ViewModels/MemAllocTable.cpp
--- a/Engine/Source/Developer/TraceInsights/Private/Insights/MemoryProfiler/ViewModels/MemAllocTable.cpp
+++ b/Engine/Source/Developer/TraceInsights/Private/Insights/MemoryProfiler/ViewModels/MemAllocTable.cpp
@@ -923,15 +923,13 @@ void FMemAllocTable::AddDefaultColumns()
 		Column.SetTitleName(LOCTEXT("FunctionColumnTitle", "Top Function"));
 		Column.SetDescription(LOCTEXT("FunctionColumnDesc", "Resolved top function from the callstack of allocation"));
 
-		Column.SetFlags(ETableColumnFlags::ShouldBeVisible | ETableColumnFlags::CanBeHidden | ETableColumnFlags::CanBeFiltered);
+		Column.SetFlags(ETableColumnFlags::ShouldBeVisible | ETableColumnFlags::CanBeHidden | ETableColumnFlags::CanBeFiltered | ETableColumnFlags::IsDynamic);
 
 		Column.SetHorizontalAlignment(HAlign_Left);
 		Column.SetInitialWidth(550.0f);
 
 		Column.SetDataType(ETableCellDataType::Text);
 
-		Column.SetIsDynamic(true);
-
 		class FFunctionValueGetter : public FTableCellValueGetter
 		{
 		public:
@@ -1026,15 +1024,13 @@ void FMemAllocTable::AddDefaultColumns()
 		Column.SetTitleName(LOCTEXT("SourceFileColumnTitle", "Top Source File"));
 		Column.SetDescription(LOCTEXT("SourceFileColumnDesc", "Source file of the top function from the callstack of allocation"));
 
-		Column.SetFlags(ETableColumnFlags::CanBeHidden | ETableColumnFlags::CanBeFiltered);
+		Column.SetFlags(ETableColumnFlags::CanBeHidden | ETableColumnFlags::CanBeFiltered | ETableColumnFlags::IsDynamic);
 
 		Column.SetHorizontalAlignment(HAlign_Left);
 		Column.SetInitialWidth(550.0f);
 
 		Column.SetDataType(ETableCellDataType::Text);
 
-		Column.SetIsDynamic(true);
-
 		class FSourceValueGetter : public FTableCellValueGetter
 		{
 		public:
diff --git a/Engine/Source/Developer/TraceInsights/Private/Insights/MemoryProfiler/Widgets/SMemAllocTableTreeView.cpp b/Engine/Source/Developer/TraceInsights/Private/Insights/MemoryProfiler/Widgets/SMemAllocTableTreeView.cpp
--- a/Engine/Source/Developer/TraceInsights/Private/Insights/MemoryProfiler/Widgets/SMemAllocTableTreeView.cpp
+++ b/Engine/Source/Developer/TraceInsights/Private/Insights/MemoryProfiler/Widgets/SMemAllocTableTreeView.cpp
@@ -43,6 +43,7 @@
 #include "Insights/Table/ViewModels/TableColumn.h"
 #include "Insights/TimingProfilerCommon.h"
 #include "Insights/ViewModels/FilterConfigurator.h"
+#include "Insights/ViewModels/Filters.h"
 
 #include <limits>
 #include <memory>
@@ -52,8 +53,8 @@
 namespace Insights
 {
 
-const int SMemAllocTableTreeView::FullCallStackIndex = 0x0000FFFFF;
-const int SMemAllocTableTreeView::LLMFilterIndex = 0x0000FFFFE;
+const int32 SMemAllocTableTreeView::FullCallStackIndex = 0x0000FFFFF;
+const int32 SMemAllocTableTreeView::LLMFilterIndex = 0x0000FFFFE;
 
 ////////////////////////////////////////////////////////////////////////////////////////////////////
 
@@ -1263,41 +1264,56 @@ void SMemAllocTableTreeView::UpdateQueryInfo()
 
 ////////////////////////////////////////////////////////////////////////////////////////////////////
 
-bool SMemAllocTableTreeView::ApplyCustomAdvancedFilters(const FTableTreeNodePtr& NodePtr)
+void SMemAllocTableTreeView::UpdateFilterContext(const FFilterConfigurator& InFilterConfigurator, const FTableTreeNode& InNode) const
 {
-	// Super heavy to compute, validate that the filter has a use for this key before computing it
-	if (FilterConfigurator && FilterConfigurator->IsKeyUsed(FullCallStackIndex))
-	{
-		FMemAllocNodePtr MemNodePtr = StaticCastSharedPtr<FMemAllocNode>(NodePtr);
-		Context.SetFilterData<FString>(FullCallStackIndex, MemNodePtr->GetFullCallstack().ToString());
-	}
+	STableTreeView::UpdateFilterContext(InFilterConfigurator, InNode);
 
-	if (FilterConfigurator && FilterConfigurator->IsKeyUsed(LLMFilterIndex))
+	if (InNode.Is<FMemAllocNode>())
 	{
-		FMemAllocNodePtr MemNodePtr = StaticCastSharedPtr<FMemAllocNode>(NodePtr);
-		Context.SetFilterData<FString>(LLMFilterIndex, MemNodePtr->GetMemAlloc()->GetTag());
-	}
+		const FMemAllocNode& MemNode = InNode.As<FMemAllocNode>();
 
-	return true;
+		// GetFullCallstack is super heavy to compute. Validate that the filter has a use for this key before computing it.
+		if (InFilterConfigurator.IsKeyUsed(FullCallStackIndex))
+		{
+			FilterContext.SetFilterData<FString>(FullCallStackIndex, MemNode.GetFullCallstack().ToString());
+		}
+
+		if (InFilterConfigurator.IsKeyUsed(LLMFilterIndex))
+		{
+			FilterContext.SetFilterData<FString>(LLMFilterIndex, MemNode.GetMemAlloc()->GetTag());
+		}
+	}
 }
 
 ////////////////////////////////////////////////////////////////////////////////////////////////////
 
-void SMemAllocTableTreeView::AddCustomAdvancedFilters()
+void SMemAllocTableTreeView::InitFilterConfigurator(FFilterConfigurator& InOutFilterConfigurator)
 {
-	TSharedPtr<TArray<TSharedPtr<struct FFilter>>>& AvailableFilters = FilterConfigurator->GetAvailableFilters();
-
-	AvailableFilters->Add(MakeShared<FFilter>(FullCallStackIndex, LOCTEXT("FullCallstack", "Full Callstack"), LOCTEXT("SearchFullCallstack", "Search in all the callstack frames"), EFilterDataType::String, FFilterService::Get()->GetStringOperators()));
-	Context.AddFilterData<FString>(FullCallStackIndex, FString());
-
-	TSharedPtr<FFilterWithSuggestions> LLMCategoryFilter = MakeShared<FFilterWithSuggestions>(static_cast<int32>(LLMFilterIndex), LOCTEXT("LLMTag", "LLM Tag"), LOCTEXT("LLMTag", "LLM Tag"), EFilterDataType::String, FFilterService::Get()->GetStringOperators());
-	Context.AddFilterData<FString>(LLMFilterIndex, FString());
-	LLMCategoryFilter->Callback = [this](const FString& Text, TArray<FString>& OutSuggestions)
+	STableTreeView::InitFilterConfigurator(InOutFilterConfigurator);
+
+	TSharedRef<FFilter> FullCallStackFilter = MakeShared<FFilter>(
+		FullCallStackIndex,
+		LOCTEXT("FullCallstack", "Full Callstack"),
+		LOCTEXT("SearchFullCallstack", "Search in all the callstack frames"),
+		EFilterDataType::String,
+		nullptr,
+		FFilterService::Get()->GetStringOperators());
+	FilterContext.AddFilterData<FString>(FullCallStackIndex, FString());
+	InOutFilterConfigurator.Add(FullCallStackFilter);
+
+	TSharedRef<FFilterWithSuggestions> LLMTagFilter = MakeShared<FFilterWithSuggestions>(
+		LLMFilterIndex,
+		LOCTEXT("LLMTag", "LLM Tag"),
+		LOCTEXT("LLMTag", "LLM Tag"),
+		EFilterDataType::String,
+		nullptr,
+		FFilterService::Get()->GetStringOperators());
+	FilterContext.AddFilterData<FString>(LLMFilterIndex, FString());
+	LLMTagFilter->SetCallback([this](const FString& Text, TArray<FString>& OutSuggestions)
 	{
 		this->PopulateLLMTagSuggestionList(Text, OutSuggestions);
-	};
-
-	AvailableFilters->Add(LLMCategoryFilter);
+	});
+	InOutFilterConfigurator.Add(LLMTagFilter);
 }
 
 ////////////////////////////////////////////////////////////////////////////////////////////////////
diff --git a/Engine/Source/Developer/TraceInsights/Private/Insights/MemoryProfiler/Widgets/SMemAllocTableTreeView.h b/Engine/Source/Developer/TraceInsights/Private/Insights/MemoryProfiler/Widgets/SMemAllocTableTreeView.h
--- a/Engine/Source/Developer/TraceInsights/Private/Insights/MemoryProfiler/Widgets/SMemAllocTableTreeView.h
+++ b/Engine/Source/Developer/TraceInsights/Private/Insights/MemoryProfiler/Widgets/SMemAllocTableTreeView.h
@@ -119,8 +119,8 @@ private:
 	FText GetQueryInfoTooltip() const;
 
 	void UpdateQueryInfo();
-	bool virtual ApplyCustomAdvancedFilters(const FTableTreeNodePtr& NodePtr) override;
-	virtual void AddCustomAdvancedFilters() override;
+	virtual void UpdateFilterContext(const FFilterConfigurator& InFilterConfigurator, const FTableTreeNode& InNode) const override;
+	virtual void InitFilterConfigurator(FFilterConfigurator& InOutFilterConfigurator) override;
 
 	TSharedRef<SWidget> ConstructFunctionToggleButton();
 	void CallstackGroupingByFunction_OnCheckStateChanged(ECheckBoxState NewRadioState);
@@ -130,8 +130,8 @@ private:
 	void PopulateLLMTagSuggestionList(const FString& Text, TArray<FString>& OutSuggestions);
 
 private:
-	const static int FullCallStackIndex;
-	const static int LLMFilterIndex;
+	const static int32 FullCallStackIndex;
+	const static int32 LLMFilterIndex;
 	int32 TabIndex = -1;
 	TSharedPtr<FMemoryRuleSpec> Rule = nullptr;
 	double TimeMarkers[4];
diff --git a/Engine/Source/Developer/TraceInsights/Private/Insights/MemoryProfiler/Widgets/SMemTagTreeView.cpp b/Engine/Source/Developer/TraceInsights/Private/Insights/MemoryProfiler/Widgets/SMemTagTreeView.cpp
--- a/Engine/Source/Developer/TraceInsights/Private/Insights/MemoryProfiler/Widgets/SMemTagTreeView.cpp
+++ b/Engine/Source/Developer/TraceInsights/Private/Insights/MemoryProfiler/Widgets/SMemTagTreeView.cpp
@@ -1483,19 +1483,13 @@ void SMemTagTreeView::SortTreeNodes()
 
 void SMemTagTreeView::SortTreeNodesRec(FMemTagNode& Node, const Insights::ITableCellValueSorter& Sorter)
 {
-	if (ColumnSortMode == EColumnSortMode::Type::Descending)
-	{
-		Node.SortChildrenDescending(Sorter);
-	}
-	else // if (ColumnSortMode == EColumnSortMode::Type::Ascending)
-	{
-		Node.SortChildrenAscending(Sorter);
-	}
+	Insights::ESortMode SortMode = (ColumnSortMode == EColumnSortMode::Type::Descending) ? Insights::ESortMode::Descending : Insights::ESortMode::Ascending;
+	Node.SortChildren(Sorter, SortMode);
 
 	for (Insights::FBaseTreeNodePtr ChildPtr : Node.GetChildren())
 	{
 		//if (ChildPtr->IsGroup())
-		if (ChildPtr->GetChildren().Num() > 0)
+		if (ChildPtr->GetChildrenCount() > 0)
 		{
 			SortTreeNodesRec(*StaticCastSharedPtr<FMemTagNode>(ChildPtr), Sorter);
 		}
diff --git a/Engine/Source/Developer/TraceInsights/Private/Insights/NetworkingProfiler/ViewModels/NetStatsCounterNode.h b/Engine/Source/Developer/TraceInsights/Private/Insights/NetworkingProfiler/ViewModels/NetStatsCounterNode.h
--- a/Engine/Source/Developer/TraceInsights/Private/Insights/NetworkingProfiler/ViewModels/NetStatsCounterNode.h
+++ b/Engine/Source/Developer/TraceInsights/Private/Insights/NetworkingProfiler/ViewModels/NetStatsCounterNode.h
@@ -100,17 +100,6 @@ public:
 	const FText GetTextForAggregatedStatsMax(bool bForTooltip = false) const;
 	const FText GetTextForAggregatedStatsAverage(bool bForTooltip = false) const;
 
-	/** Sorts children using the specified class instance. */
-	template<typename TSortingClass>
-	void SortChildren(const TSortingClass& Instance)
-	{
-		auto Projection = [](Insights::FBaseTreeNodePtr Node) -> FNetStatsCounterNodePtr
-		{
-			return StaticCastSharedPtr<FNetStatsCounterNode, Insights::FBaseTreeNode>(Node);
-		};
-		Algo::SortBy(GetChildrenMutable(), Projection, Instance);
-	}
-
 private:
 	const FText FormatAggregatedStatsValue(uint32 Value) const;
 
diff --git a/Engine/Source/Developer/TraceInsights/Private/Insights/NetworkingProfiler/Widgets/SNetStatsCountersView.cpp b/Engine/Source/Developer/TraceInsights/Private/Insights/NetworkingProfiler/Widgets/SNetStatsCountersView.cpp
--- a/Engine/Source/Developer/TraceInsights/Private/Insights/NetworkingProfiler/Widgets/SNetStatsCountersView.cpp
+++ b/Engine/Source/Developer/TraceInsights/Private/Insights/NetworkingProfiler/Widgets/SNetStatsCountersView.cpp
@@ -1076,18 +1076,12 @@ void SNetStatsCountersView::SortTreeNodes()
 
 void SNetStatsCountersView::SortTreeNodesRec(FNetStatsCounterNode& Node, const Insights::ITableCellValueSorter& Sorter)
 {
-	if (ColumnSortMode == EColumnSortMode::Type::Descending)
-	{
-		Node.SortChildrenDescending(Sorter);
-	}
-	else // if (ColumnSortMode == EColumnSortMode::Type::Ascending)
-	{
-		Node.SortChildrenAscending(Sorter);
-	}
+	Insights::ESortMode SortMode = (ColumnSortMode == EColumnSortMode::Type::Descending) ? Insights::ESortMode::Descending : Insights::ESortMode::Ascending;
+	Node.SortChildren(Sorter, SortMode);
 
 	for (Insights::FBaseTreeNodePtr ChildPtr : Node.GetChildren())
 	{
-		if (ChildPtr->GetChildren().Num() > 0)
+		if (ChildPtr->GetChildrenCount() > 0)
 		{
 			SortTreeNodesRec(*StaticCastSharedPtr<FNetStatsCounterNode>(ChildPtr), Sorter);
 		}
diff --git a/Engine/Source/Developer/TraceInsights/Private/Insights/NetworkingProfiler/Widgets/SNetStatsView.cpp b/Engine/Source/Developer/TraceInsights/Private/Insights/NetworkingProfiler/Widgets/SNetStatsView.cpp
--- a/Engine/Source/Developer/TraceInsights/Private/Insights/NetworkingProfiler/Widgets/SNetStatsView.cpp
+++ b/Engine/Source/Developer/TraceInsights/Private/Insights/NetworkingProfiler/Widgets/SNetStatsView.cpp
@@ -1083,18 +1083,12 @@ void SNetStatsView::SortTreeNodes()
 
 void SNetStatsView::SortTreeNodesRec(FNetEventNode& Node, const Insights::ITableCellValueSorter& Sorter)
 {
-	if (ColumnSortMode == EColumnSortMode::Type::Descending)
-	{
-		Node.SortChildrenDescending(Sorter);
-	}
-	else // if (ColumnSortMode == EColumnSortMode::Type::Ascending)
-	{
-		Node.SortChildrenAscending(Sorter);
-	}
+	Insights::ESortMode SortMode = (ColumnSortMode == EColumnSortMode::Type::Descending) ? Insights::ESortMode::Descending : Insights::ESortMode::Ascending;
+	Node.SortChildren(Sorter, SortMode);
 
 	for (Insights::FBaseTreeNodePtr ChildPtr : Node.GetChildren())
 	{
-		if (ChildPtr->GetChildren().Num() > 0)
+		if (ChildPtr->GetChildrenCount() > 0)
 		{
 			SortTreeNodesRec(*StaticCastSharedPtr<FNetEventNode>(ChildPtr), Sorter);
 		}
diff --git a/Engine/Source/Developer/TraceInsights/Private/Insights/Table/ViewModels/BaseTreeNode.cpp b/Engine/Source/Developer/TraceInsights/Private/Insights/Table/ViewModels/BaseTreeNode.cpp
--- a/Engine/Source/Developer/TraceInsights/Private/Insights/Table/ViewModels/BaseTreeNode.cpp
+++ b/Engine/Source/Developer/TraceInsights/Private/Insights/Table/ViewModels/BaseTreeNode.cpp
@@ -84,16 +84,16 @@ FLinearColor FBaseTreeNode::GetDefaultColor(bool bIsGroupNode)
 
 ////////////////////////////////////////////////////////////////////////////////////////////////////
 
-void FBaseTreeNode::SortChildrenAscending(const ITableCellValueSorter& Sorter)
+void FBaseTreeNode::SortChildren(const ITableCellValueSorter& Sorter, ESortMode SortMode)
 {
-	Sorter.Sort(GroupData->Children, ESortMode::Ascending);
+	Sorter.Sort(GroupData->Children, SortMode);
 }
 
 ////////////////////////////////////////////////////////////////////////////////////////////////////
 
-void FBaseTreeNode::SortChildrenDescending(const ITableCellValueSorter& Sorter)
+void FBaseTreeNode::SortFilteredChildren(const ITableCellValueSorter& Sorter, ESortMode SortMode)
 {
-	Sorter.Sort(GroupData->Children, ESortMode::Descending);
+	Sorter.Sort(GroupData->FilteredChildren, SortMode);
 }
 
 ////////////////////////////////////////////////////////////////////////////////////////////////////
diff --git a/Engine/Source/Developer/TraceInsights/Private/Insights/Table/ViewModels/BaseTreeNode.h b/Engine/Source/Developer/TraceInsights/Private/Insights/Table/ViewModels/BaseTreeNode.h
--- a/Engine/Source/Developer/TraceInsights/Private/Insights/Table/ViewModels/BaseTreeNode.h
+++ b/Engine/Source/Developer/TraceInsights/Private/Insights/Table/ViewModels/BaseTreeNode.h
@@ -4,6 +4,7 @@
 
 #include "CoreMinimal.h"
 #include "Math/Color.h"
+#include "Templates/Function.h"
 
 #include "Insights/Common/SimpleRtti.h"
 
@@ -13,6 +14,7 @@ namespace Insights
 {
 
 class ITableCellValueSorter;
+enum class ESortMode;
 
 ////////////////////////////////////////////////////////////////////////////////////////////////////
 
@@ -83,7 +85,7 @@ public:
 	}
 
 	/**
-	 * @return a name of this node.
+	 * @returns a name of this node.
 	 */
 	const FName& GetName() const
 	{
@@ -91,19 +93,22 @@ public:
 	}
 
 	/**
-	 * @return a name of this node to display in tree view.
+	 * @returns a name of this node to display in tree view.
 	 */
 	virtual const FText GetDisplayName() const;
 
 	/**
-	 * @return a name suffix for this node to display in tree view. Ex.: group nodes may include additional info, like the number of visible / total children.
+	 * @returns a name suffix for this node to display in tree view. Ex.: group nodes may include additional info, like the number of visible / total children.
 	 */
 	virtual const FText GetExtraDisplayName() const;
 
+	/**
+	 * @returns true if the node has an extra display name (a name suffix to display in tree view).
+	 */
 	virtual bool HasExtraDisplayName() const;
 
 	/**
-	 * @return a descriptive text for this node to display in a tooltip.
+	 * @returns a descriptive text for this node to display in a tooltip.
 	 */
 	virtual const FText GetTooltip() const
 	{
@@ -122,12 +127,12 @@ public:
 	}
 
 	/**
-	 * @return the default icon for a group/leaf node.
+	 * @returns the default icon for a group/leaf node.
 	 */
 	static const FSlateBrush* GetDefaultIcon(bool bIsGroupNode);
 
 	/**
-	 * @return a brush icon for this node.
+	 * @returns a brush icon for this node.
 	 */
 	virtual const FSlateBrush* GetIcon() const
 	{
@@ -135,12 +140,12 @@ public:
 	}
 
 	/**
-	 * @return the default color tint for a group/leaf node.
+	 * @returns the default color tint for a group/leaf node.
 	 */
 	static FLinearColor GetDefaultColor(bool bIsGroupNode);
 
 	/**
-	 * @return the color tint to be used for the icon and the name text of this node.
+	 * @returns the color tint to be used for the icon and the name text of this node.
 	 */
 	virtual FLinearColor GetColor() const
 	{
@@ -148,7 +153,7 @@ public:
 	}
 
 	/**
-	 * @return a pointer to a data context for this node.
+	 * @returns a pointer to a data context for this node. Only available for group nodes.
 	 */
 	virtual void* GetContext() const
 	{
@@ -156,7 +161,7 @@ public:
 	}
 
 	/**
-	 * Sets a pointer to a data context for this node.
+	 * Sets a pointer to a data context for this node. Only available for group nodes.
 	 */
 	virtual void SetContext(void* InContext)
 	{
@@ -167,7 +172,7 @@ public:
 	}
 
 	/**
-	 * @return true, if this node is a group node.
+	 * @returns true if this node is a group node.
 	 */
 	bool IsGroup() const
 	{
@@ -186,23 +191,67 @@ public:
 	}
 
 	/**
-	 * @return a const reference to the child nodes of this group.
+	 * @returns the number of children nodes.
+	 */
+	int32 GetChildrenCount() const
+	{
+		return GroupData->Children.Num();
+	}
+
+	/**
+	 * @returns the number of filtered children nodes.
+	 */
+	int32 GetFilteredChildrenCount() const
+	{
+		return GroupData->FilteredChildren.Num();
+	}
+
+	/**
+	 * Enumerates the children nodes.
+	 */
+	void EnumerateChildren(TFunction<bool(const FBaseTreeNodePtr&)> Callback) const
+	{
+		for (const FBaseTreeNodePtr& ChildNode : GroupData->Children)
+		{
+			if (!Callback(ChildNode))
+			{
+				break;
+			}
+		}
+	}
+
+	/**
+	 * Enumerates the filtered children nodes.
+	 */
+	void EnumerateFilteredChildren(TFunction<bool(const FBaseTreeNodePtr&)> Callback) const
+	{
+		for (const FBaseTreeNodePtr& ChildNode : GroupData->FilteredChildren)
+		{
+			if (!Callback(ChildNode))
+			{
+				break;
+			}
+		}
+	}
+
+	/**
+	 * @returns a const reference to the children nodes of this group.
 	 */
-	FORCEINLINE_DEBUGGABLE const TArray<FBaseTreeNodePtr>& GetChildren() const
+	const TArray<FBaseTreeNodePtr>& GetChildren() const
 	{
 		return GroupData->Children;
 	}
 
 	/**
-	 * @return a const reference to the child nodes that should be visible to the UI based on filtering.
+	 * @returns a const reference to the children nodes that should be visible to the UI based on filtering.
 	 */
-	FORCEINLINE_DEBUGGABLE const TArray<FBaseTreeNodePtr>& GetFilteredChildren() const
+	const TArray<FBaseTreeNodePtr>& GetFilteredChildren() const
 	{
 		return GroupData->FilteredChildren;
 	}
 
 	/**
-	 * @return a weak reference to the group of this node, may be invalid.
+	 * @returns a weak reference to the group of this node, may be invalid.
 	 */
 	FBaseTreeNodeWeak GetGroupPtr() const
 	{
@@ -210,11 +259,11 @@ public:
 	}
 
 	/**
-	 * @return a name of the group node that this node belongs to.
+	 * @returns a weak reference to the group of this node, may be invalid.
 	 */
-	const FName& GetGroupName() const
+	FBaseTreeNodePtr GetParentNode() const
 	{
-		return GroupPtr.Pin()->GetName();
+		return GroupPtr.Pin();
 	}
 
 	virtual bool IsFiltered() const
@@ -222,25 +271,31 @@ public:
 		return false;
 	}
 
-	void SortChildrenAscending(const ITableCellValueSorter& Sorter);
-	void SortChildrenDescending(const ITableCellValueSorter& Sorter);
+	void SortChildren(const ITableCellValueSorter& Sorter, ESortMode SortMode);
+	void SortFilteredChildren(const ITableCellValueSorter& Sorter, ESortMode SortMode);
 
-	template <class PREDICATE_CLASS>
-	void SortChildren(const PREDICATE_CLASS& Predicate)
+	template <typename PredicateType>
+	void SortChildren(PredicateType Predicate)
 	{
 		GroupData->Children.Sort(Predicate);
 	}
 
-	/** Adds specified child to the children and sets group for it. */
-	FORCEINLINE_DEBUGGABLE void AddChildAndSetGroupPtr(const FBaseTreeNodePtr& ChildPtr)
+	template <typename PredicateType>
+	void SortFilteredChildren(PredicateType Predicate)
 	{
-		ChildPtr->GroupPtr = AsShared();
+		GroupData->FilteredChildren.Sort(Predicate);
+	}
+
+	/** Adds specified node to the children nodes. Also sets the current node as the group of the specified node. */
+	void AddChildAndSetGroupPtr(const FBaseTreeNodePtr& ChildPtr)
+	{
+		ChildPtr->GroupPtr = AsWeak();
 		GroupData->Children.Add(ChildPtr);
 	}
 
 	void SetGroupPtrForAllChildren()
 	{
-		FBaseTreeNodeWeak ThisNode = AsShared();
+		FBaseTreeNodeWeak ThisNode = AsWeak();
 		for (FBaseTreeNodePtr& NodePtr : GroupData->Children)
 		{
 			NodePtr->GroupPtr = ThisNode;
@@ -256,40 +311,65 @@ public:
 	}
 
 	/** Clears children. */
-	FORCEINLINE_DEBUGGABLE void ClearChildren(int32 NewSize = 0)
+	void ClearChildren(int32 NewSize = 0)
 	{
 		ResetGroupPtrForAllChildren();
 		GroupData->Children.Reset(NewSize);
 	}
 
-	FORCEINLINE_DEBUGGABLE void SwapChildren(TArray<FBaseTreeNodePtr>& NewChildren)
+	void SwapChildren(TArray<FBaseTreeNodePtr>& NewChildren)
 	{
 		ResetGroupPtrForAllChildren();
 		Swap(NewChildren, GroupData->Children);
 		SetGroupPtrForAllChildren();
 	}
 
-	FORCEINLINE_DEBUGGABLE void SwapChildrenFast(TArray<FBaseTreeNodePtr>& NewChildren)
+	void SwapChildrenFast(TArray<FBaseTreeNodePtr>& NewChildren)
 	{
 		Swap(NewChildren, GroupData->Children);
 	}
 
-	/** Adds specified child to the filtered children. */
-	FORCEINLINE_DEBUGGABLE void AddFilteredChild(const FBaseTreeNodePtr& ChildPtr)
+	/** Adds specified child to the filtered children nodes. */
+	void AddFilteredChild(const FBaseTreeNodePtr& ChildPtr)
 	{
 		GroupData->FilteredChildren.Add(ChildPtr);
 	}
 
-	/** Clears filtered children. */
+	/** Clears the filtered children nodes. */
 	void ClearFilteredChildren(int32 NewSize = 0)
 	{
 		GroupData->FilteredChildren.Reset(NewSize);
 	}
 
+	/**
+	 * Resets the filtered children nodes.
+	 * The filtered array of children nodes will be initialized with the unfiltered array of children.
+	 */
+	void ResetFilteredChildren()
+	{
+		const int32 Size = GroupData->Children.Num();
+		GroupData->FilteredChildren.Reset();
+		GroupData->FilteredChildren.Append(GroupData->Children);
+	}
+
+	/**
+	 * Resets the filtered children for this node and also recursively for all children nodes.
+	 * The filtered array of children nodes will be initialized with the unfiltered array of children.
+	 */
+	void ResetFilteredChildrenRec()
+	{
+		this->ResetFilteredChildren();
+		for (FBaseTreeNodePtr Node : GroupData->Children)
+		{
+			Node->ResetFilteredChildrenRec();
+		}
+	}
+
 	bool IsExpanded() const
 	{
 		return GroupData->bIsExpanded;
 	}
+
 	void SetExpansion(bool bOnOff)
 	{
 		if (IsGroup())
@@ -300,9 +380,9 @@ public:
 
 protected:
 	/**
-	 * @return a reference to the child nodes of this group.
+	 * @returns a reference to the children nodes of this group.
 	 */
-	FORCEINLINE_DEBUGGABLE TArray<FBaseTreeNodePtr>& GetChildrenMutable()
+	TArray<FBaseTreeNodePtr>& GetChildrenMutable()
 	{
 		return GroupData->Children;
 	}
diff --git a/Engine/Source/Developer/TraceInsights/Private/Insights/Table/ViewModels/TableColumn.h b/Engine/Source/Developer/TraceInsights/Private/Insights/Table/ViewModels/TableColumn.h
--- a/Engine/Source/Developer/TraceInsights/Private/Insights/Table/ViewModels/TableColumn.h
+++ b/Engine/Source/Developer/TraceInsights/Private/Insights/Table/ViewModels/TableColumn.h
@@ -23,12 +23,12 @@ class ITableCellValueSorter;
 
 enum class ETableColumnFlags : uint32
 {
-	None = 0,
-
-	ShouldBeVisible = (1 << 0),
-	CanBeHidden     = (1 << 1),
-	CanBeFiltered   = (1 << 2),
-	IsHierarchy     = (1 << 3),
+	None            = 0,
+	ShouldBeVisible = (1 << 0), // Column should be initially visible.
+	CanBeHidden     = (1 << 1), // Column can be hidden.
+	CanBeFiltered   = (1 << 2), // Column can be used for filtering.
+	IsHierarchy     = (1 << 3), // Column is the hierarchy (name) column, in a tree view.
+	IsDynamic       = (1 << 4), // Column is dynamic. The column's value can not be cached.
 };
 ENUM_CLASS_FLAGS(ETableColumnFlags);
 
@@ -40,8 +40,6 @@ enum class ETableColumnAggregation : uint32
 	Sum,
 	Min,
 	Max,
-	//Average,
-	//Median,
 	SameValue
 };
 
@@ -64,7 +62,6 @@ public:
 		, TitleName()
 		, Description()
 		, bIsVisible(false)
-		, bIsDynamic(false)
 		, Flags(ETableColumnFlags::None)
 		, HorizontalAlignment(HAlign_Left)
 		, InitialWidth(60.0f)
@@ -115,12 +112,15 @@ public:
 	/** Whether this column can be hidden. */
 	bool CanBeHidden() const { return EnumHasAnyFlags(Flags, ETableColumnFlags::CanBeHidden); }
 
-	/** Whether this column can be used for filtering displayed results. */
+	/** Whether this column can be used for filtering. */
 	bool CanBeFiltered() const { return EnumHasAnyFlags(Flags, ETableColumnFlags::CanBeFiltered); }
 
-	/** Whether this column is the hierarcy (name) column, in a tree view. */
+	/** Whether this column is the hierarchy (name) column, in a tree view. */
 	bool IsHierarchy() const { return EnumHasAnyFlags(Flags, ETableColumnFlags::IsHierarchy); }
 
+	/** Whether this column is dynamic or not. If a column is dynamic, the column's value can not be cached. */
+	bool IsDynamic() const { return EnumHasAnyFlags(Flags, ETableColumnFlags::IsDynamic); }
+
 	void SetFlags(ETableColumnFlags InFlags) { Flags = InFlags; }
 
 	//////////////////////////////////////////////////
@@ -209,9 +209,6 @@ public:
 
 	//////////////////////////////////////////////////
 
-	void SetIsDynamic(bool InValue) { bIsDynamic = InValue; }
-	bool IsDynamic() const { return bIsDynamic; }
-
 private:
 	/** Id of the column. */
 	FName Id;
@@ -231,9 +228,6 @@ private:
 	/** Is this column visible? */
 	bool bIsVisible;
 
-	/** If it is not dynamic, the column's value can be cached. */
-	bool bIsDynamic;
-
 	/** Other on/off switches. */
 	ETableColumnFlags Flags;
 
@@ -254,7 +248,7 @@ private:
 	/** Custom getter for values identified by this column. */
 	TSharedRef<ITableCellValueGetter> ValueGetter;
 
-	/** Custom formater for values displayed by this column. */
+	/** Custom formatter for values displayed by this column. */
 	TSharedRef<ITableCellValueFormatter> ValueFormatter;
 
 	/** Custom sorter for values displayed by this column. */
@@ -263,7 +257,7 @@ private:
 	/** Initial sorting mode. */
 	EColumnSortMode::Type InitialSortMode;
 
-	/** Used to convert in a custom way from string to column data type in FilterConfigurator */
+	/** Used in FilterConfigurator to convert from string to column's data type. */
 	TSharedPtr<IFilterValueConverter> ValueConverter;
 
 	/* Parent table. Only one table instance can own this column. */
diff --git a/Engine/Source/Developer/TraceInsights/Private/Insights/Table/ViewModels/TableTreeNode.h b/Engine/Source/Developer/TraceInsights/Private/Insights/Table/ViewModels/TableTreeNode.h
--- a/Engine/Source/Developer/TraceInsights/Private/Insights/Table/ViewModels/TableTreeNode.h
+++ b/Engine/Source/Developer/TraceInsights/Private/Insights/Table/ViewModels/TableTreeNode.h
@@ -12,6 +12,7 @@ namespace Insights
 {
 
 class FTable;
+class STableTreeView;
 
 struct FTableRowId
 {
@@ -77,9 +78,11 @@ public:
 	FTableRowId GetRowId() const { return RowId; }
 	int32 GetRowIndex() const { return RowId.RowIndex; }
 
+	//////////////////////////////////////////////////
+	// Aggregation
+
 	void InitAggregatedValues()
 	{
-		ensure(IsGroup());
 		if (!AggregatedValues)
 		{
 			AggregatedValues = new TMap<FName, FTableCellValue>();
@@ -95,17 +98,53 @@ public:
 		}
 	}
 
-	void ResetAggregatedValues() { CleanupAggregatedValues(); }
-	void ResetAggregatedValues(const FName& ColumnId) { if (AggregatedValues) { AggregatedValues->Remove(ColumnId); } }
-	bool HasAggregatedValue(const FName& ColumnId) const { return AggregatedValues && AggregatedValues->Contains(ColumnId); }
-	const FTableCellValue* FindAggregatedValue(const FName& ColumnId) const { return AggregatedValues ? AggregatedValues->Find(ColumnId) : nullptr; }
-	const FTableCellValue& GetAggregatedValue(const FName& ColumnId) const { return AggregatedValues->FindChecked(ColumnId); }
-	void AddAggregatedValue(const FName& ColumnId, const FTableCellValue& Value) { InitAggregatedValues(); AggregatedValues->Add(ColumnId, Value); }
-	void SetAggregatedValue(const FName& ColumnId, const FTableCellValue& Value) { InitAggregatedValues(); (*AggregatedValues)[ColumnId] = Value; }
+	void ResetAggregatedValues()
+	{
+		CleanupAggregatedValues();
+	}
+
+	void ResetAggregatedValue(const FName& ColumnId)
+	{
+		if (AggregatedValues)
+		{
+			AggregatedValues->Remove(ColumnId);
+		}
+	}
+
+	bool HasAggregatedValue(const FName& ColumnId) const
+	{
+		return AggregatedValues && AggregatedValues->Contains(ColumnId);
+	}
 
-	virtual bool IsFiltered() const override { return bIsFiltered; }
+	const FTableCellValue* FindAggregatedValue(const FName& ColumnId) const
+	{
+		return AggregatedValues ? AggregatedValues->Find(ColumnId) : nullptr;
+	}
+
+	const FTableCellValue& GetAggregatedValue(const FName& ColumnId) const
+	{
+		return AggregatedValues->FindChecked(ColumnId);
+	}
+
+	void AddAggregatedValue(const FName& ColumnId, const FTableCellValue& Value)
+	{
+		InitAggregatedValues();
+		AggregatedValues->Add(ColumnId, Value);
+	}
+
+	void SetAggregatedValue(const FName& ColumnId, const FTableCellValue& Value)
+	{
+		InitAggregatedValues();
+		(*AggregatedValues)[ColumnId] = Value;
+	}
+
+	//////////////////////////////////////////////////
+
+	virtual bool IsFiltered() const { return bIsFiltered; }
 	virtual void SetIsFiltered(bool InValue) { bIsFiltered = InValue; }
 
+	virtual bool OnLazyCreateChildren(TSharedPtr<STableTreeView> InTableTreeView) { return false; }
+
 protected:
 	TWeakPtr<FTable> ParentTable;
 	FTableRowId RowId;
diff --git a/Engine/Source/Developer/TraceInsights/Private/Insights/Table/Widgets/STableTreeView.cpp b/Engine/Source/Developer/TraceInsights/Private/Insights/Table/Widgets/STableTreeView.cpp
--- a/Engine/Source/Developer/TraceInsights/Private/Insights/Table/Widgets/STableTreeView.cpp
+++ b/Engine/Source/Developer/TraceInsights/Private/Insights/Table/Widgets/STableTreeView.cpp
@@ -146,7 +146,7 @@ STableTreeView::~STableTreeView()
 {
 	if (bRunInAsyncMode)
 	{
-		checkf(bIsCloseScheduled, TEXT("TableTreeView running in async mode was closed but OnClose() was not called. This can lead to a crash. Call OnClose() from the owner tab/window."))
+		ensureMsgf(bIsCloseScheduled, TEXT("TableTreeView running in async mode was closed but OnClose() was not called. This can lead to a crash. Call OnClose() from the owner tab/window."));
 	}
 
 	if (CurrentAsyncOpFilterConfigurator)
@@ -324,6 +324,7 @@ void STableTreeView::ConstructWidget(TSharedPtr<FTable> InTablePtr)
 		WidgetContent
 	];
 
+	InitNodeFiltering();
 	InitHierarchyFiltering();
 	InitializeAndShowHeaderColumns();
 	InitCommandList();
@@ -385,7 +386,7 @@ void STableTreeView::ConstructHeaderArea(TSharedRef<SVerticalBox> InWidgetConten
 		.Padding(4.0f, 0.0f, 0.0f, 0.0f)
 		.VAlign(VAlign_Center)
 		[
-			ConstructAdvancedFiltersButton()
+			ConstructFilterConfiguratorButton()
 		]
 	];
 
@@ -450,6 +451,8 @@ END_SLATE_FUNCTION_BUILD_OPTIMIZATION
 
 TSharedPtr<SWidget> STableTreeView::TreeView_GetMenuContent()
 {
+	FSlateApplication::Get().CloseToolTip();
+
 	TArray<FTableTreeNodePtr> SelectedNodes;
 	const int32 NumSelectedNodes = TreeView->GetSelectedItems(SelectedNodes);
 	FTableTreeNodePtr SelectedNode = NumSelectedNodes ? SelectedNodes[0] : nullptr;
@@ -916,52 +919,62 @@ void STableTreeView::Tick(const FGeometry& AllottedGeometry, const double InCurr
 
 void STableTreeView::UpdateTree()
 {
-	if (bRunInAsyncMode)
-	{
-		if (!bIsUpdateRunning)
-		{
-			OnPreAsyncUpdate();
-
-			FGraphEventRef CompletedEvent = StartGroupingTask();
-			CompletedEvent = StartSortingTask(CompletedEvent);
-			InProgressAsyncOperationEvent = StartApplyFiltersTask(CompletedEvent);
-		}
-		else
-		{
-			CancelCurrentAsyncOp();
-		}
-	}
-	else
-	{
-		ApplyGrouping();
-		ApplySorting();
-		ApplyFiltering();
-	}
+	OnGroupingChanged();
 }
 
 ////////////////////////////////////////////////////////////////////////////////////////////////////
 // Node Filtering
 ////////////////////////////////////////////////////////////////////////////////////////////////////
 
-void STableTreeView::OnFilteringChanged()
+void STableTreeView::InitNodeFiltering()
 {
-	if (bRunInAsyncMode)
-	{
-		if (!bIsUpdateRunning)
-		{
-			OnPreAsyncUpdate();
+	// Filter Configurator is created when the Filter Configurator button is pressed.
+	//FilterConfigurator = nullptr;
+	//CurrentAsyncOpFilterConfigurator = nullptr;
+	//OnFilterChangesCommittedHandle = nullptr;
+}
 
-			InProgressAsyncOperationEvent = StartApplyFiltersTask();
-		}
-		else
-		{
-			CancelCurrentAsyncOp();
-		}
+////////////////////////////////////////////////////////////////////////////////////////////////////
+
+void STableTreeView::OnNodeFilteringChanged()
+{
+	if (!bRunInAsyncMode)
+	{
+		ApplyNodeFiltering();
+	}
+	else if (!bIsUpdateRunning)
+	{
+		ScheduleNodeFilteringAsyncOperation();
 	}
 	else
 	{
-		ApplyFiltering();
+		CancelCurrentAsyncOp();
+	}
+}
+
+////////////////////////////////////////////////////////////////////////////////////////////////////
+
+bool STableTreeView::ScheduleNodeFilteringAsyncOperationIfNeeded()
+{
+	if (HasInProgressAsyncOperation(EAsyncOperationType::FilteringOp) ||
+		(FilterConfigurator.IsValid() &&
+		 CurrentAsyncOpFilterConfigurator &&
+		 *FilterConfigurator != *CurrentAsyncOpFilterConfigurator))
+	{
+		ScheduleNodeFilteringAsyncOperation();
+		return true;
 	}
+
+	return false;
+}
+
+////////////////////////////////////////////////////////////////////////////////////////////////////
+
+void STableTreeView::ScheduleNodeFilteringAsyncOperation()
+{
+	OnPreAsyncUpdate();
+
+	InProgressAsyncOperationEvent = StartApplyFiltersTask();
 }
 
 ////////////////////////////////////////////////////////////////////////////////////////////////////
@@ -992,6 +1005,13 @@ FGraphEventRef STableTreeView::StartApplyFiltersTask(FGraphEventRef Prerequisite
 
 ////////////////////////////////////////////////////////////////////////////////////////////////////
 
+void STableTreeView::ApplyNodeFiltering()
+{
+	ApplyHierarchyFiltering();
+}
+
+////////////////////////////////////////////////////////////////////////////////////////////////////
+
 bool STableTreeView::ApplyAdvancedFiltersForNode(FTableTreeNodePtr NodePtr)
 {
 	if (NodePtr->IsGroup())
@@ -1034,7 +1054,7 @@ bool STableTreeView::ApplyAdvancedFiltersForNode(FTableTreeNodePtr NodePtr)
 	}
 	else
 	{
-		bool bIsNodeVisible = ApplyAdvancedFilters(NodePtr);
+		bool bIsNodeVisible = FilterNode(NodePtr);
 		NodePtr->SetIsFiltered(!bIsNodeVisible);
 		return bIsNodeVisible;
 	}
@@ -1042,20 +1062,23 @@ bool STableTreeView::ApplyAdvancedFiltersForNode(FTableTreeNodePtr NodePtr)
 
 ////////////////////////////////////////////////////////////////////////////////////////////////////
 
-bool STableTreeView::ApplyAdvancedFilters(const FTableTreeNodePtr& NodePtr)
+bool STableTreeView::FilterNode(const FTableTreeNodePtr& NodePtr) const
 {
 	FFilterConfigurator* FilterConfiguratorToUse = bRunInAsyncMode ? CurrentAsyncOpFilterConfigurator : FilterConfigurator.Get();
 
-	if (FilterConfiguratorToUse == nullptr)
+	if (!FilterConfiguratorToUse || FilterConfiguratorToUse->IsEmpty())
 	{
 		return true;
 	}
 
-	if (FilterConfiguratorToUse->GetRootNode()->GetChildren().Num() == 0)
-	{
-		return true;
-	}
+	UpdateFilterContext(*FilterConfiguratorToUse, *NodePtr);
+	return FilterConfiguratorToUse->ApplyFilters(FilterContext);
+}
+
+////////////////////////////////////////////////////////////////////////////////////////////////////
 
+void STableTreeView::InitFilterConfigurator(FFilterConfigurator& InOutFilterConfigurator)
+{
 	for (const TSharedRef<FTableColumn>& Column : Table->GetColumns())
 	{
 		if (!Column->CanBeFiltered())
@@ -1067,130 +1090,129 @@ bool STableTreeView::ApplyAdvancedFilters(const FTableTreeNodePtr& NodePtr)
 		{
 			case ETableCellDataType::Int64:
 			{
-				Context.SetFilterData<int64>(Column->GetIndex(), Column->GetValue(*NodePtr)->AsInt64());
+				TSharedRef<FFilter> Filter = MakeShared<FFilter>(
+					Column->GetIndex(),
+					Column->GetTitleName(),
+					Column->GetDescription(),
+					EFilterDataType::Int64,
+					Column->GetValueConverter(),
+					FFilterService::Get()->GetIntegerOperators()
+				);
+				InOutFilterConfigurator.Add(Filter);
+				FilterContext.AddFilterData<int64>(Column->GetIndex(), 0);
 				break;
 			}
+
 			case ETableCellDataType::Double:
 			{
-				Context.SetFilterData<double>(Column->GetIndex(), Column->GetValue(*NodePtr)->AsDouble());
+				TSharedRef<FFilter> Filter = MakeShared<FFilter>(
+					Column->GetIndex(),
+					Column->GetTitleName(),
+					Column->GetDescription(),
+					EFilterDataType::Double,
+					Column->GetValueConverter(),
+					FFilterService::Get()->GetDoubleOperators());
+				InOutFilterConfigurator.Add(Filter);
+				FilterContext.AddFilterData<double>(Column->GetIndex(), 0.0);
 				break;
 			}
+
 			case ETableCellDataType::CString:
 			case ETableCellDataType::Text:
 			case ETableCellDataType::Custom:
 			{
 				if (!Column->IsHierarchy())
 				{
-					// Converting string is heavy, check if the key is used for filtering before doing any conversion
-					if (FilterConfiguratorToUse->IsKeyUsed(Column->GetIndex()))
-					{
-						Context.SetFilterData<FString>(Column->GetIndex(), Column->GetValue(*NodePtr)->AsString());
-					}
+					TSharedRef<FFilter> Filter = MakeShared<FFilter>(
+						Column->GetIndex(),
+						Column->GetTitleName(),
+						Column->GetDescription(),
+						EFilterDataType::String,
+						Column->GetValueConverter(),
+						FFilterService::Get()->GetStringOperators());
+					InOutFilterConfigurator.Add(Filter);
+					FilterContext.AddFilterData<FString>(Column->GetIndex(), FString());
 				}
 				break;
 			}
 		}
 	}
-
-	return ApplyCustomAdvancedFilters(NodePtr) && FilterConfiguratorToUse->ApplyFilters(Context);
 }
 
 ////////////////////////////////////////////////////////////////////////////////////////////////////
 
-BEGIN_SLATE_FUNCTION_BUILD_OPTIMIZATION
-TSharedRef<SWidget> STableTreeView::ConstructAdvancedFiltersButton()
+void STableTreeView::UpdateFilterContext(const FFilterConfigurator& InFilterConfigurator, const FTableTreeNode& InNode) const
 {
-	return SNew(SButton)
-		.ContentPadding(FMargin(-2.0f, 2.0f, -2.0f, 2.0f))
-		.ToolTipText(this, &STableTreeView::AdvancedFilters_GetTooltipText)
-		.OnClicked(this, &STableTreeView::OnAdvancedFiltersClicked)
-		.IsEnabled(this, &STableTreeView::AdvancedFilters_ShouldBeEnabled)
-		[
-			SNew(SImage)
-			.Image(FInsightsStyle::GetBrush("Icons.ClassicFilterConfig"))
-		];
-}
-END_SLATE_FUNCTION_BUILD_OPTIMIZATION
-
-////////////////////////////////////////////////////////////////////////////////////////////////////
-
-FReply STableTreeView::OnAdvancedFiltersClicked()
-{
-	if (!FilterConfigurator.IsValid())
+	for (const TSharedRef<FTableColumn>& Column : Table->GetColumns())
 	{
-		FilterConfigurator = MakeShared<FFilterConfigurator>();
-		TSharedPtr<TArray<TSharedPtr<struct FFilter>>>& AvailableFilters = FilterConfigurator->GetAvailableFilters();
+		if (!Column->CanBeFiltered())
+		{
+			continue;
+		}
 
-		for (const TSharedRef<FTableColumn>& Column : Table->GetColumns())
+		switch (Column->GetDataType())
 		{
-			if (!Column->CanBeFiltered())
+			case ETableCellDataType::Int64:
 			{
-				continue;
+				FilterContext.SetFilterData<int64>(Column->GetIndex(), Column->GetValue(InNode)->AsInt64());
+				break;
 			}
 
-			switch (Column->GetDataType())
+			case ETableCellDataType::Double:
 			{
-				case ETableCellDataType::Int64:
-				{
-					AvailableFilters->Add(MakeShared<FFilter>(Column->GetIndex(), Column->GetTitleName(), Column->GetDescription(), EFilterDataType::Int64, FFilterService::Get()->GetIntegerOperators()));
-					AvailableFilters->Last()->Converter = Column->GetValueConverter();
-					Context.AddFilterData<int64>(Column->GetIndex(), 0);
-					break;
-				}
-				case ETableCellDataType::Double:
-				{
-					AvailableFilters->Add(MakeShared<FFilter>(Column->GetIndex(), Column->GetTitleName(), Column->GetDescription(), EFilterDataType::Double, FFilterService::Get()->GetDoubleOperators()));
-					AvailableFilters->Last()->Converter = Column->GetValueConverter();
-					Context.AddFilterData<double>(Column->GetIndex(), 0.0);
-					break;
-				}
-				case ETableCellDataType::CString:
-				case ETableCellDataType::Text:
-				case ETableCellDataType::Custom:
+				FilterContext.SetFilterData<double>(Column->GetIndex(), Column->GetValue(InNode)->AsDouble());
+				break;
+			}
+
+			case ETableCellDataType::CString:
+			case ETableCellDataType::Text:
+			case ETableCellDataType::Custom:
+			{
+				if (!Column->IsHierarchy())
 				{
-					if (!Column->IsHierarchy())
+					// Converting string is heavy, check if the key is used for filtering before doing any conversion
+					if (InFilterConfigurator.IsKeyUsed(Column->GetIndex()))
 					{
-						AvailableFilters->Add(MakeShared<FFilter>(Column->GetIndex(), Column->GetTitleName(), Column->GetDescription(), EFilterDataType::String, FFilterService::Get()->GetStringOperators()));
-						AvailableFilters->Last()->Converter = Column->GetValueConverter();
-						Context.AddFilterData<FString>(Column->GetIndex(), FString());
+						FilterContext.SetFilterData<FString>(Column->GetIndex(), Column->GetValue(InNode)->AsString());
 					}
-					break;
 				}
+				break;
 			}
 		}
-
-		AddCustomAdvancedFilters();
-
-		CurrentAsyncOpFilterConfigurator = new FFilterConfigurator(*FilterConfigurator);
-		OnFilterChangesCommittedHandle = FilterConfigurator->GetOnChangesCommittedEvent().AddSP(this, &STableTreeView::OnAdvancedFiltersChangesCommitted);
 	}
-
-	FFilterService::Get()->CreateFilterConfiguratorWidget(FilterConfigurator);
-
-	return FReply::Handled();
 }
 
 ////////////////////////////////////////////////////////////////////////////////////////////////////
 
-void STableTreeView::OnAdvancedFiltersChangesCommitted()
+BEGIN_SLATE_FUNCTION_BUILD_OPTIMIZATION
+TSharedRef<SWidget> STableTreeView::ConstructFilterConfiguratorButton()
 {
-	OnFilteringChanged();
+	return SNew(SButton)
+		.ContentPadding(FMargin(-2.0f, 2.0f, -2.0f, 2.0f))
+		.IsEnabled(this, &STableTreeView::FilterConfigurator_IsEnabled)
+		.ToolTipText(this, &STableTreeView::FilterConfigurator_GetTooltipText)
+		.OnClicked(this, &STableTreeView::FilterConfigurator_OnClicked)
+		[
+			SNew(SImage)
+			.Image(FInsightsStyle::GetBrush("Icons.ClassicFilterConfig"))
+		];
 }
+END_SLATE_FUNCTION_BUILD_OPTIMIZATION
 
 ////////////////////////////////////////////////////////////////////////////////////////////////////
 
-bool STableTreeView::AdvancedFilters_ShouldBeEnabled() const
+bool STableTreeView::FilterConfigurator_IsEnabled() const
 {
 	return TextFilter->GetRawFilterText().IsEmpty();
 }
 
 ////////////////////////////////////////////////////////////////////////////////////////////////////
 
-FText STableTreeView::AdvancedFilters_GetTooltipText() const
+FText STableTreeView::FilterConfigurator_GetTooltipText() const
 {
-	if (AdvancedFilters_ShouldBeEnabled())
+	if (FilterConfigurator_IsEnabled())
 	{
-		return LOCTEXT("AdvancedFiltersBtn_ToolTip", "Opens the filter configurator window.");
+		return LOCTEXT("FilterConfiguratorBtn_ToolTip", "Opens the filter configurator window.");
 	}
 
 	return LOCTEXT("AdvancedFiltersBtn_Disabled_ToolTip", "Advanced filters cannot be added when filters are already applied using the search box.");
@@ -1200,7 +1222,24 @@ FText STableTreeView::AdvancedFilters_GetTooltipText() const
 
 bool STableTreeView::FilterConfigurator_HasFilters() const
 {
-	return FilterConfigurator.IsValid() && FilterConfigurator->GetRootNode()->GetChildren().Num() > 0;
+	return FilterConfigurator.IsValid() && !FilterConfigurator->IsEmpty();
+}
+
+////////////////////////////////////////////////////////////////////////////////////////////////////
+
+FReply STableTreeView::FilterConfigurator_OnClicked()
+{
+	if (!FilterConfigurator.IsValid())
+	{
+		FilterConfigurator = MakeShared<FFilterConfigurator>();
+		InitFilterConfigurator(*FilterConfigurator);
+		OnFilterChangesCommittedHandle = FilterConfigurator->GetOnChangesCommittedEvent().AddSP(this, &STableTreeView::OnNodeFilteringChanged);
+		CurrentAsyncOpFilterConfigurator = new FFilterConfigurator(*FilterConfigurator);
+	}
+
+	FFilterService::Get()->CreateFilterConfiguratorWidget(FilterConfigurator);
+
+	return FReply::Handled();
 }
 
 ////////////////////////////////////////////////////////////////////////////////////////////////////
@@ -1238,7 +1277,48 @@ void STableTreeView::HandleItemToStringArray(const FTableTreeNodePtr& FTableTree
 
 ////////////////////////////////////////////////////////////////////////////////////////////////////
 
-void STableTreeView::ApplyFiltering()
+void STableTreeView::OnHierarchyFilteringChanged()
+{
+	if (!bRunInAsyncMode)
+	{
+		ApplyHierarchyFiltering();
+	}
+	else if (!bIsUpdateRunning)
+	{
+		ScheduleHierarchyFilteringAsyncOperation();
+	}
+	else
+	{
+		CancelCurrentAsyncOp();
+	}
+}
+
+////////////////////////////////////////////////////////////////////////////////////////////////////
+
+bool STableTreeView::ScheduleHierarchyFilteringAsyncOperationIfNeeded()
+{
+	if (HasInProgressAsyncOperation(EAsyncOperationType::FilteringOp) ||
+		TextFilter->GetRawFilterText().CompareTo(CurrentAsyncOpTextFilter->GetRawFilterText()) != 0)
+	{
+		ScheduleHierarchyFilteringAsyncOperation();
+		return true;
+	}
+
+	return false;
+}
+
+////////////////////////////////////////////////////////////////////////////////////////////////////
+
+void STableTreeView::ScheduleHierarchyFilteringAsyncOperation()
+{
+	OnPreAsyncUpdate();
+
+	InProgressAsyncOperationEvent = StartApplyFiltersTask();
+}
+
+////////////////////////////////////////////////////////////////////////////////////////////////////
+
+void STableTreeView::ApplyHierarchyFiltering()
 {
 	FStopwatch Stopwatch;
 	Stopwatch.Start();
@@ -1250,8 +1330,15 @@ void STableTreeView::ApplyFiltering()
 	}
 	else
 	{
-		const bool bFilterIsEmpty = bRunInAsyncMode ? CurrentAsyncOpTextFilter->GetRawFilterText().IsEmpty() : TextFilter->GetRawFilterText().IsEmpty();
-		ApplyHierarchicalFilterForNode(Root, bFilterIsEmpty);
+		const bool bIsEmptyFilter = bRunInAsyncMode ? CurrentAsyncOpTextFilter->GetRawFilterText().IsEmpty() : TextFilter->GetRawFilterText().IsEmpty();
+		if (bIsEmptyFilter)
+		{
+			ApplyEmptyHierarchyFilteringRec(Root);
+		}
+		else
+		{
+			ApplyHierarchyFilteringRec(Root);
+		}
 	}
 
 	// The Root node is always hidden. The tree shows the filtered children of the Root node.
@@ -1268,11 +1355,11 @@ void STableTreeView::ApplyFiltering()
 	}
 
 	TreeViewBannerText = FText::GetEmpty();
-	if (Root->GetChildren().Num() == 0)
+	if (Root->GetChildrenCount() == 0)
 	{
 		TreeViewBannerText = LOCTEXT("EmptyTreeMsg", "Tree is empty.");
 	}
-	else if (Root->GetFilteredChildren().Num() == 0)
+	else if (Root->GetFilteredChildrenCount() == 0)
 	{
 		TreeViewBannerText = LOCTEXT("AllNodesFiltered", "All tree nodes are filtered out.");
 	}
@@ -1327,23 +1414,38 @@ void STableTreeView::ApplyFiltering()
 
 ////////////////////////////////////////////////////////////////////////////////////////////////////
 
-bool STableTreeView::ApplyHierarchicalFilterForNode(FTableTreeNodePtr NodePtr, bool bFilterIsEmpty)
+void STableTreeView::ApplyEmptyHierarchyFilteringRec(FTableTreeNodePtr NodePtr)
+{
+	if (NodePtr->IsGroup())
+	{
+		// If a group node passes the filter, all child nodes will be shown.
+		MakeSubtreeVisible(NodePtr, true);
+	}
+	else
+	{
+		NodePtr->SetIsFiltered(false);
+	}
+}
+
+////////////////////////////////////////////////////////////////////////////////////////////////////
+
+bool STableTreeView::ApplyHierarchyFilteringRec(FTableTreeNodePtr NodePtr)
 {
-	bool bIsNodeVisible = bFilterIsEmpty || Filters->PassesAllFilters(NodePtr);
+	bool bIsNodeVisible = Filters->PassesAllFilters(NodePtr);
 
 	if (NodePtr->IsGroup())
 	{
 		// If a group node passes the filter, all child nodes will be shown.
 		if (bIsNodeVisible)
 		{
-			MakeSubtreeVisible(NodePtr, bFilterIsEmpty);
+			MakeSubtreeVisible(NodePtr, false);
 			return true;
 		}
 
 		const TArray<FBaseTreeNodePtr>& GroupChildren = NodePtr->GetChildren();
 		const int32 NumChildren = GroupChildren.Num();
 
-		NodePtr->ClearFilteredChildren(bFilterIsEmpty ? NumChildren : 0);
+		NodePtr->ClearFilteredChildren(0);
 
 		int32 NumVisibleChildren = 0;
 		for (int32 Cx = 0; Cx < NumChildren; ++Cx)
@@ -1355,7 +1457,7 @@ bool STableTreeView::ApplyHierarchicalFilterForNode(FTableTreeNodePtr NodePtr, b
 
 			// Add a child.
 			const FTableTreeNodePtr& ChildNodePtr = StaticCastSharedPtr<FTableTreeNode>(GroupChildren[Cx]);
-			if (ApplyHierarchicalFilterForNode(ChildNodePtr, bFilterIsEmpty))
+			if (ApplyHierarchyFilteringRec(ChildNodePtr))
 			{
 				NodePtr->AddFilteredChild(ChildNodePtr);
 				NumVisibleChildren++;
@@ -1363,7 +1465,7 @@ bool STableTreeView::ApplyHierarchicalFilterForNode(FTableTreeNodePtr NodePtr, b
 		}
 
 		const bool bIsGroupNodeVisible = bIsNodeVisible || NumVisibleChildren > 0;
-		if (!bFilterIsEmpty && bIsGroupNodeVisible)
+		if (bIsGroupNodeVisible)
 		{
 			if (bRunInAsyncMode)
 			{
@@ -1452,7 +1554,7 @@ void STableTreeView::SearchBox_OnTextChanged(const FText& InFilterText)
 {
 	TextFilter->SetRawFilterText(InFilterText);
 	SearchBox->SetError(TextFilter->GetFilterErrorText());
-	OnFilteringChanged();
+	OnHierarchyFilteringChanged();
 }
 
 ////////////////////////////////////////////////////////////////////////////////////////////////////
@@ -1602,27 +1704,64 @@ FName STableTreeView::TableRow_GetHighlightedNodeName() const
 
 void STableTreeView::OnGroupingChanged()
 {
-	if (bRunInAsyncMode)
+	if (!bRunInAsyncMode)
 	{
-		if (!bIsUpdateRunning)
-		{
-			OnPreAsyncUpdate();
+		ApplyGrouping();
+		ApplySorting();
+		ApplyHierarchyFiltering();
+	}
+	else if (!bIsUpdateRunning)
+	{
+		ScheduleGroupingAsyncOperation();
+	}
+	else
+	{
+		CancelCurrentAsyncOp();
+	}
+}
 
-			FGraphEventRef CompletedEvent = StartGroupingTask();
-			CompletedEvent = StartSortingTask(CompletedEvent);
-			InProgressAsyncOperationEvent = StartApplyFiltersTask(CompletedEvent);
-		}
-		else
+////////////////////////////////////////////////////////////////////////////////////////////////////
+
+bool STableTreeView::ScheduleGroupingAsyncOperationIfNeeded()
+{
+	bool bGroupingHasChanged = HasInProgressAsyncOperation(EAsyncOperationType::GroupingOp);
+
+	if (!bGroupingHasChanged &&
+		CurrentGroupings.Num() != CurrentAsyncOpGroupings.Num())
+	{
+		bGroupingHasChanged = true;
+	}
+
+	if (!bGroupingHasChanged)
+	{
+		for (int32 Index = 0; Index < CurrentGroupings.Num(); ++Index)
 		{
-			CancelCurrentAsyncOp();
+			if (CurrentGroupings[Index] != CurrentAsyncOpGroupings[Index])
+			{
+				bGroupingHasChanged = true;
+				break;
+			}
 		}
 	}
-	else
+
+	if (bGroupingHasChanged)
 	{
-		ApplyGrouping();
-		ApplySorting();
-		ApplyFiltering();
+		ScheduleGroupingAsyncOperation();
+		return true;
 	}
+
+	return false;
+}
+
+////////////////////////////////////////////////////////////////////////////////////////////////////
+
+void STableTreeView::ScheduleGroupingAsyncOperation()
+{
+	OnPreAsyncUpdate();
+
+	FGraphEventRef CompletedEvent = StartGroupingTask();
+	CompletedEvent = StartSortingTask(CompletedEvent);
+	InProgressAsyncOperationEvent = StartApplyFiltersTask(CompletedEvent);
 }
 
 ////////////////////////////////////////////////////////////////////////////////////////////////////
@@ -2430,7 +2569,7 @@ void STableTreeView::UpdateCStringSameValueAggregationRec(FTableColumn& Column,
 	}
 	else
 	{
-		GroupNode.ResetAggregatedValues(Column.GetId());
+		GroupNode.ResetAggregatedValue(Column.GetId());
 	}
 }
 
@@ -2571,29 +2710,62 @@ void STableTreeView::UpdateCurrentSortingByColumn()
 
 void STableTreeView::OnSortingChanged()
 {
-	if (bRunInAsyncMode)
+	if (!bRunInAsyncMode)
 	{
-		if (!bIsUpdateRunning)
-		{
-			OnPreAsyncUpdate();
-
-			FGraphEventRef CompletedEvent = StartSortingTask();
-			InProgressAsyncOperationEvent = StartApplyFiltersTask(CompletedEvent);
-		}
-		else
-		{
-			CancelCurrentAsyncOp();
-		}
+		ApplySorting();
+		ApplyHierarchyFiltering();
+	}
+	else if (!bIsUpdateRunning)
+	{
+		ScheduleSortingAsyncOperation();
 	}
 	else
 	{
-		ApplySorting();
-		ApplyFiltering();
+		CancelCurrentAsyncOp();
 	}
 }
 
 ////////////////////////////////////////////////////////////////////////////////////////////////////
 
+bool STableTreeView::ScheduleSortingAsyncOperationIfNeeded()
+{
+	bool bSortingHasChanged = HasInProgressAsyncOperation(EAsyncOperationType::SortingOp);
+
+	if (!bSortingHasChanged &&
+		((CurrentSorter.IsValid() && CurrentAsyncOpSorter == nullptr) ||
+		(!CurrentSorter.IsValid() && CurrentAsyncOpSorter != nullptr)))
+	{
+		bSortingHasChanged = true;
+	}
+
+	if (!bSortingHasChanged &&
+		CurrentSorter.IsValid() &&
+		(CurrentSorter.Get() != CurrentAsyncOpSorter || ColumnSortMode != CurrentAsyncOpColumnSortMode))
+	{
+		bSortingHasChanged = true;
+	}
+
+	if (bSortingHasChanged)
+	{
+		ScheduleSortingAsyncOperation();
+		return true;
+	}
+
+	return false;
+}
+
+////////////////////////////////////////////////////////////////////////////////////////////////////
+
+void STableTreeView::ScheduleSortingAsyncOperation()
+{
+	OnPreAsyncUpdate();
+
+	FGraphEventRef CompletedEvent = StartSortingTask();
+	InProgressAsyncOperationEvent = StartApplyFiltersTask(CompletedEvent);
+}
+
+////////////////////////////////////////////////////////////////////////////////////////////////////
+
 FGraphEventRef STableTreeView::StartSortingTask(FGraphEventRef Prerequisite)
 {
 	if (!CurrentSorter.IsValid())
@@ -2664,14 +2836,8 @@ void STableTreeView::SortTreeNodes(ITableCellValueSorter* InSorter, EColumnSortM
 
 void STableTreeView::SortTreeNodesRec(FTableTreeNode& GroupNode, const ITableCellValueSorter& Sorter, EColumnSortMode::Type InColumnSortMode)
 {
-	if (InColumnSortMode == EColumnSortMode::Type::Descending)
-	{
-		GroupNode.SortChildrenDescending(Sorter);
-	}
-	else // if (ColumnSortMode == EColumnSortMode::Type::Ascending)
-	{
-		GroupNode.SortChildrenAscending(Sorter);
-	}
+	const ESortMode SortMode = (InColumnSortMode == EColumnSortMode::Type::Descending) ? ESortMode::Descending : ESortMode::Ascending;
+	GroupNode.SortChildren(Sorter, SortMode);
 
 	for (FBaseTreeNodePtr ChildPtr : GroupNode.GetChildren())
 	{
@@ -3175,7 +3341,7 @@ void STableTreeView::SetExpandValueForChildGroups(FBaseTreeNode* InRoot, int32 I
 
 void STableTreeView::CountNumNodesPerDepthRec(FBaseTreeNode* InRoot, TArray<int32>& InOutNumNodesPerDepth, int32 InDepth, int32 InMaxDepth, int32 InMaxNodes) const
 {
-	InOutNumNodesPerDepth[InDepth] += InRoot->GetChildren().Num();
+	InOutNumNodesPerDepth[InDepth] += InRoot->GetChildrenCount();
 
 	if (InDepth < InMaxDepth && InOutNumNodesPerDepth[InDepth] < InMaxNodes)
 	{
@@ -3249,65 +3415,23 @@ double STableTreeView::GetAllOperationsDuration()
 
 void STableTreeView::StartPendingAsyncOperations()
 {
-	// Check if grouping settings have changed. If they did, a full refresh (Grouping, Sorting and Filtering) is scheduled.
-	bool bGroupingsHaveChanged = HasInProgressAsyncOperation(EAsyncOperationType::GroupingOp);
-	bGroupingsHaveChanged |= CurrentGroupings.Num() != CurrentAsyncOpGroupings.Num();
-
-	if (!bGroupingsHaveChanged)
-	{
-		for (int Index = 0; Index < CurrentGroupings.Num(); ++Index)
-		{
-			if (CurrentGroupings[Index] != CurrentAsyncOpGroupings[Index])
-			{
-				bGroupingsHaveChanged = true;
-				break;
-			}
-		}
-	}
-
-	if (bGroupingsHaveChanged)
+	if (ScheduleGroupingAsyncOperationIfNeeded())
 	{
-		OnPreAsyncUpdate();
-
-		FGraphEventRef CompletedEvent = StartGroupingTask();
-		CompletedEvent = StartSortingTask(CompletedEvent);
-		InProgressAsyncOperationEvent = StartApplyFiltersTask(CompletedEvent);
-
 		return;
 	}
 
-	// Check if sorting settings have changed. If they did, a Sorting and Filtering Refresh is scheduled.
-	bool bSortingHasChanged = HasInProgressAsyncOperation(EAsyncOperationType::SortingOp);
-	bSortingHasChanged |= (CurrentSorter.IsValid() && CurrentAsyncOpSorter == nullptr) || (!CurrentSorter.IsValid() && CurrentAsyncOpSorter != nullptr);
-	if (!bSortingHasChanged && CurrentSorter.IsValid())
-	{
-		bSortingHasChanged = CurrentSorter.Get() != CurrentAsyncOpSorter || ColumnSortMode != CurrentAsyncOpColumnSortMode;
-	}
-
-	if (bSortingHasChanged)
+	if (ScheduleSortingAsyncOperationIfNeeded())
 	{
-		OnPreAsyncUpdate();
-
-		FGraphEventRef CompletedEvent = StartSortingTask();
-		InProgressAsyncOperationEvent = StartApplyFiltersTask(CompletedEvent);
-
 		return;
 	}
 
-	// Check if the text filter has changed. If it has, schedule a new Filtering Refresh.
-	bool bFiltersHaveChanged = HasInProgressAsyncOperation(EAsyncOperationType::FilteringOp);
-	bFiltersHaveChanged |= TextFilter->GetRawFilterText().CompareTo(CurrentAsyncOpTextFilter->GetRawFilterText()) != 0;
-	if (FilterConfigurator.IsValid() && CurrentAsyncOpFilterConfigurator && !bFiltersHaveChanged)
+	if (ScheduleHierarchyFilteringAsyncOperationIfNeeded())
 	{
-		bFiltersHaveChanged |= *FilterConfigurator != *CurrentAsyncOpFilterConfigurator;
+		return;
 	}
 
-	if (bFiltersHaveChanged)
+	if (ScheduleNodeFilteringAsyncOperationIfNeeded())
 	{
-		OnPreAsyncUpdate();
-
-		InProgressAsyncOperationEvent = StartApplyFiltersTask();
-
 		return;
 	}
 }
diff --git a/Engine/Source/Developer/TraceInsights/Private/Insights/Table/Widgets/STableTreeView.h b/Engine/Source/Developer/TraceInsights/Private/Insights/Table/Widgets/STableTreeView.h
--- a/Engine/Source/Developer/TraceInsights/Private/Insights/Table/Widgets/STableTreeView.h
+++ b/Engine/Source/Developer/TraceInsights/Private/Insights/Table/Widgets/STableTreeView.h
@@ -126,6 +126,8 @@ public:
 	 */
 	void Construct(const FArguments& InArgs, TSharedPtr<FTable> InTablePtr);
 
+	TSharedPtr<STreeView<FTableTreeNodePtr>> GetInnerTreeView() const { return TreeView; }
+
 	TSharedPtr<FTable>& GetTable() { return Table; }
 	const TSharedPtr<FTable>& GetTable() const { return Table; }
 
@@ -184,15 +186,6 @@ protected:
 
 	void UpdateTree();
 
-	/**
-	 * Populates OutSearchStrings with the strings that should be used in searching.
-	 *
-	 * @param GroupOrStatNodePtr - the group and stat node to get a text description from.
-	 * @param OutSearchStrings   - an array of strings to use in searching.
-	 *
-	 */
-	static void HandleItemToStringArray(const FTableTreeNodePtr& GroupOrStatNodePtr, TArray<FString>& OutSearchStrings);
-
 	////////////////////////////////////////////////////////////////////////////////////////////////////
 	// Tree View - Context Menu
 
@@ -261,31 +254,50 @@ protected:
 	FName TableRow_GetHighlightedNodeName() const;
 
 	////////////////////////////////////////////////////////////////////////////////////////////////////
-	// Filtering
+	// Node Filtering
 
-	void OnFilteringChanged();
+	void InitNodeFiltering();
+
+	void OnNodeFilteringChanged();
+	bool ScheduleNodeFilteringAsyncOperationIfNeeded();
+	void ScheduleNodeFilteringAsyncOperation();
 	FGraphEventRef StartApplyFiltersTask(FGraphEventRef Prerequisite = nullptr);
 
-	void ApplyFiltering();
+	void ApplyNodeFiltering();
 
 	bool ApplyAdvancedFiltersForNode(FTableTreeNodePtr NodePtr);
 
-	bool ApplyAdvancedFilters(const FTableTreeNodePtr& NodePtr);
-	bool virtual ApplyCustomAdvancedFilters(const FTableTreeNodePtr& NodePtr) { return true; };
-	virtual void AddCustomAdvancedFilters() {}
+	virtual bool FilterNode(const FTableTreeNodePtr& NodePtr) const;
 
-	virtual TSharedRef<SWidget> ConstructAdvancedFiltersButton();
+	virtual void InitFilterConfigurator(FFilterConfigurator& InOutFilterConfigurator);
+	virtual void UpdateFilterContext(const FFilterConfigurator& InFilterConfigurator, const FTableTreeNode& InNode) const;
 
-	FReply OnAdvancedFiltersClicked();
-	void OnAdvancedFiltersChangesCommitted();
+	virtual TSharedRef<SWidget> ConstructFilterConfiguratorButton();
 
+	bool FilterConfigurator_IsEnabled() const;
+	FText FilterConfigurator_GetTooltipText() const;
 	bool FilterConfigurator_HasFilters() const;
-	bool AdvancedFilters_ShouldBeEnabled() const;
-	FText AdvancedFilters_GetTooltipText() const;
+	FReply FilterConfigurator_OnClicked();
+
+	////////////////////////////////////////////////////////////////////////////////////////////////////
+	// Hierarchy Filtering
 
 	void InitHierarchyFiltering();
 
-	bool ApplyHierarchicalFilterForNode(FTableTreeNodePtr NodePtr, bool bFilterIsEmpty);
+	/**
+	 * Populates OutSearchStrings with the strings that should be used in searching.
+	 *
+	 * @param GroupOrStatNodePtr - the group and stat node to get a text description from.
+	 * @param OutSearchStrings   - an array of strings to use in searching.
+	 */
+	static void HandleItemToStringArray(const FTableTreeNodePtr& GroupOrStatNodePtr, TArray<FString>& OutSearchStrings);
+
+	void OnHierarchyFilteringChanged();
+	bool ScheduleHierarchyFilteringAsyncOperationIfNeeded();
+	void ScheduleHierarchyFilteringAsyncOperation();
+	void ApplyHierarchyFiltering();
+	void ApplyEmptyHierarchyFilteringRec(FTableTreeNodePtr NodePtr);
+	bool ApplyHierarchyFilteringRec(FTableTreeNodePtr NodePtr);
 
 	/** Set all the nodes belonging to a subtree as visible. Returns true if the caller node should be expanded. */
 	bool MakeSubtreeVisible(FTableTreeNodePtr NodePtr, bool bFilterIsEmpty);
@@ -303,8 +315,10 @@ protected:
 	virtual void InternalCreateGroupings();
 
 	void OnGroupingChanged();
-	void ApplyGrouping();
+	bool ScheduleGroupingAsyncOperationIfNeeded();
+	void ScheduleGroupingAsyncOperation();
 	FGraphEventRef StartGroupingTask(FGraphEventRef Prerequisite = nullptr);
+	void ApplyGrouping();
 	void CreateGroups(const TArray<TSharedPtr<FTreeNodeGrouping>>& Groupings);
 	void GroupNodesRec(const TArray<FTableTreeNodePtr>& Nodes, FTableTreeNode& ParentGroup, int32 GroupingDepth, const TArray<TSharedPtr<FTreeNodeGrouping>>& Groupings);
 
@@ -355,7 +369,7 @@ protected:
 			else
 			{
 				FTableTreeNode& TableNode = *(FTableTreeNode*)NodePtr.Get();
-				TableNode.ResetAggregatedValues(Column.GetId());
+				TableNode.ResetAggregatedValue(Column.GetId());
 				UpdateAggregationRec(Column, TableNode, InitialAggregatedValue, bSetInitialValue, ValueGetterFunc);
 				if (TableNode.HasAggregatedValue(Column.GetId()))
 				{
@@ -382,8 +396,10 @@ protected:
 	void UpdateCurrentSortingByColumn();
 
 	void OnSortingChanged();
-	void ApplySorting();
+	bool ScheduleSortingAsyncOperationIfNeeded();
+	void ScheduleSortingAsyncOperation();
 	FGraphEventRef StartSortingTask(FGraphEventRef Prerequisite = nullptr);
+	void ApplySorting();
 	void SortTreeNodes(ITableCellValueSorter* InSorter, EColumnSortMode::Type InColumnSortMode);
 	void SortTreeNodesRec(FTableTreeNode& GroupNode, const ITableCellValueSorter& Sorter, EColumnSortMode::Type InColumnSortMode);
 
@@ -576,7 +592,7 @@ protected:
 	/** The filter configurator actually used in the Hierarchy Filtering async task. */
 	FFilterConfigurator* CurrentAsyncOpFilterConfigurator = nullptr;
 
-	FFilterContext Context;
+	mutable FFilterContext FilterContext;
 
 	//////////////////////////////////////////////////
 	// Grouping
@@ -663,7 +679,7 @@ public:
 	{
 		if (TableTreeViewPtr)
 		{
-			TableTreeViewPtr->ApplyFiltering();
+			TableTreeViewPtr->ApplyHierarchyFiltering();
 		}
 	}
 
diff --git a/Engine/Source/Developer/TraceInsights/Private/Insights/Table/Widgets/STableTreeViewCell.cpp b/Engine/Source/Developer/TraceInsights/Private/Insights/Table/Widgets/STableTreeViewCell.cpp
--- a/Engine/Source/Developer/TraceInsights/Private/Insights/Table/Widgets/STableTreeViewCell.cpp
+++ b/Engine/Source/Developer/TraceInsights/Private/Insights/Table/Widgets/STableTreeViewCell.cpp
@@ -104,7 +104,6 @@ TSharedRef<SWidget> STableTreeViewCell::GenerateWidgetForNameColumn(const FArgum
 			.TextStyle(FInsightsStyle::Get(), TEXT("TreeTable.NameText"))
 			.ColorAndOpacity(this, &STableTreeViewCell::GetDisplayNameColorAndOpacity)
 			.ShadowColorAndOpacity(this, &STableTreeViewCell::GetShadowColorAndOpacity)
-			.ToolTip(RowToolTip)
 		]
 
 		// Name Suffix
@@ -120,7 +119,6 @@ TSharedRef<SWidget> STableTreeViewCell::GenerateWidgetForNameColumn(const FArgum
 			.TextStyle(FInsightsStyle::Get(), TEXT("TreeTable.NameText"))
 			.ColorAndOpacity(this, &STableTreeViewCell::GetExtraDisplayNameColorAndOpacity)
 			.ShadowColorAndOpacity(this, &STableTreeViewCell::GetShadowColorAndOpacity)
-			.ToolTip(RowToolTip)
 		]
 	;
 }
diff --git a/Engine/Source/Developer/TraceInsights/Private/Insights/Table/Widgets/STableTreeViewRow.cpp b/Engine/Source/Developer/TraceInsights/Private/Insights/Table/Widgets/STableTreeViewRow.cpp
--- a/Engine/Source/Developer/TraceInsights/Private/Insights/Table/Widgets/STableTreeViewRow.cpp
+++ b/Engine/Source/Developer/TraceInsights/Private/Insights/Table/Widgets/STableTreeViewRow.cpp
@@ -14,7 +14,7 @@
 #include "Insights/Table/Widgets/STableTreeViewCell.h"
 #include "Insights/Table/Widgets/STableTreeViewTooltip.h"
 
-#define LOCTEXT_NAMESPACE "STableTreeView"
+#define LOCTEXT_NAMESPACE "Insights::STableTreeView"
 
 namespace Insights
 {
diff --git a/Engine/Source/Developer/TraceInsights/Private/Insights/Table/Widgets/STableTreeViewTooltip.cpp b/Engine/Source/Developer/TraceInsights/Private/Insights/Table/Widgets/STableTreeViewTooltip.cpp
--- a/Engine/Source/Developer/TraceInsights/Private/Insights/Table/Widgets/STableTreeViewTooltip.cpp
+++ b/Engine/Source/Developer/TraceInsights/Private/Insights/Table/Widgets/STableTreeViewTooltip.cpp
@@ -16,7 +16,7 @@
 #include "Insights/Table/ViewModels/TableColumn.h"
 #include "Insights/Table/ViewModels/TableTreeNode.h"
 
-#define LOCTEXT_NAMESPACE "STableTreeView"
+#define LOCTEXT_NAMESPACE "Insights::STableTreeView"
 
 namespace Insights
 {
diff --git a/Engine/Source/Developer/TraceInsights/Private/Insights/TaskGraphProfiler/Widgets/STaskTableTreeView.cpp b/Engine/Source/Developer/TraceInsights/Private/Insights/TaskGraphProfiler/Widgets/STaskTableTreeView.cpp
--- a/Engine/Source/Developer/TraceInsights/Private/Insights/TaskGraphProfiler/Widgets/STaskTableTreeView.cpp
+++ b/Engine/Source/Developer/TraceInsights/Private/Insights/TaskGraphProfiler/Widgets/STaskTableTreeView.cpp
@@ -250,7 +250,7 @@ void STaskTableTreeView::ConstructHeaderArea(TSharedRef<SVerticalBox> InWidgetCo
 		.Padding(4.0f, 0.0f, 0.0f, 0.0f)
 		.VAlign(VAlign_Center)
 		[
-			ConstructAdvancedFiltersButton()
+			ConstructFilterConfiguratorButton()
 		]
 	];
 
diff --git a/Engine/Source/Developer/TraceInsights/Private/Insights/ViewModels/FilterConfigurator.h b/Engine/Source/Developer/TraceInsights/Private/Insights/ViewModels/FilterConfigurator.h
--- a/Engine/Source/Developer/TraceInsights/Private/Insights/ViewModels/FilterConfigurator.h
+++ b/Engine/Source/Developer/TraceInsights/Private/Insights/ViewModels/FilterConfigurator.h
@@ -25,15 +25,16 @@ public:
 
 	virtual ~FFilterConfigurator();
 
-	bool IsEmpty() const { return RootNode->GetChildren().Num() == 0; }
-
-	FFilterConfiguratorNodePtr GetRootNode() { return RootNode; }
+	bool IsEmpty() const { return RootNode->GetChildrenCount() == 0; }
 
 	virtual bool ApplyFilters(const FFilterContext& Context) const override;
 
 	bool IsKeyUsed(int32 Key) const;
 
-	TSharedPtr<TArray<TSharedPtr<struct FFilter>>>& GetAvailableFilters() { return AvailableFilters; }
+	FFilterConfiguratorNodePtr GetRootNode() { return RootNode; }
+	TSharedPtr<TArray<TSharedPtr<FFilter>>>& GetAvailableFilters() { return AvailableFilters; }
+
+	void Add(TSharedPtr<FFilter> InFilter) { AvailableFilters->Add(InFilter); }
 
 	////////////////////////////////////////////////////////////////////////////////////////////////////
 	// OnDestroyedEvent
@@ -65,7 +66,7 @@ private:
 private:
 	FFilterConfiguratorNodePtr RootNode;
 
-	TSharedPtr<TArray<TSharedPtr<struct FFilter>>> AvailableFilters;
+	TSharedPtr<TArray<TSharedPtr<FFilter>>> AvailableFilters;
 
 	TSet<int32> KeysUsed;
 };
diff --git a/Engine/Source/Developer/TraceInsights/Private/Insights/ViewModels/FilterConfiguratorNode.cpp b/Engine/Source/Developer/TraceInsights/Private/Insights/ViewModels/FilterConfiguratorNode.cpp
--- a/Engine/Source/Developer/TraceInsights/Private/Insights/ViewModels/FilterConfiguratorNode.cpp
+++ b/Engine/Source/Developer/TraceInsights/Private/Insights/ViewModels/FilterConfiguratorNode.cpp
@@ -67,11 +67,12 @@ bool FFilterConfiguratorNode::operator==(const FFilterConfiguratorNode& Other) c
 	bIsEqual &= SelectedFilterOperator.Get() == Other.SelectedFilterOperator.Get();
 	bIsEqual &= SelectedFilterGroupOperator.Get() == Other.SelectedFilterGroupOperator.Get();
 	bIsEqual &= TextBoxValue == Other.TextBoxValue;
-	bIsEqual &= GetChildren().Num() == Other.GetChildren().Num();
+	bIsEqual &= GetChildrenCount() == Other.GetChildrenCount();
 
 	if (bIsEqual)
 	{
-		for (int32 Index = 0; Index < GetChildren().Num(); ++Index)
+		const int32 Count = GetChildrenCount();
+		for (int32 Index = 0; Index < Count; ++Index)
 		{
 			bIsEqual &= *StaticCastSharedPtr<FFilterConfiguratorNode>(GetChildren()[Index]) == *StaticCastSharedPtr<FFilterConfiguratorNode>(Other.GetChildren()[Index]);
 		}
@@ -82,7 +83,7 @@ bool FFilterConfiguratorNode::operator==(const FFilterConfiguratorNode& Other) c
 
 ////////////////////////////////////////////////////////////////////////////////////////////////////
 
-const TArray<TSharedPtr<struct FFilterGroupOperator>>& FFilterConfiguratorNode::GetFilterGroupOperators()
+const TArray<TSharedPtr<FFilterGroupOperator>>& FFilterConfiguratorNode::GetFilterGroupOperators()
 {
 	return FFilterService::Get()->GetFilterGroupOperators();
 }
@@ -109,7 +110,7 @@ void FFilterConfiguratorNode::SetGroupPtrForChildren()
 
 ////////////////////////////////////////////////////////////////////////////////////////////////////
 
-void FFilterConfiguratorNode::SetAvailableFilters(TSharedPtr<TArray<TSharedPtr<struct FFilter>>> InAvailableFilters)
+void FFilterConfiguratorNode::SetAvailableFilters(TSharedPtr<TArray<TSharedPtr<FFilter>>> InAvailableFilters)
 {
 	AvailableFilters = InAvailableFilters;
 
@@ -121,7 +122,7 @@ void FFilterConfiguratorNode::SetAvailableFilters(TSharedPtr<TArray<TSharedPtr<s
 
 ////////////////////////////////////////////////////////////////////////////////////////////////////
 
-void FFilterConfiguratorNode::SetSelectedFilter(TSharedPtr<struct FFilter> InSelectedFilter)
+void FFilterConfiguratorNode::SetSelectedFilter(TSharedPtr<FFilter> InSelectedFilter)
 {
 	SelectedFilter = InSelectedFilter;
 	if (SelectedFilter.IsValid() && SelectedFilter->GetSupportedOperators()->Num() > 0)
@@ -129,7 +130,7 @@ void FFilterConfiguratorNode::SetSelectedFilter(TSharedPtr<struct FFilter> InSel
 		SetSelectedFilterOperator(SelectedFilter->GetSupportedOperators()->GetData()[0]);
 
 		AvailableFilterOperators->Empty();
-		SupportedOperatorsArrayPtr AvailableOperators = InSelectedFilter->SupportedOperators;
+		SupportedOperatorsArrayPtr AvailableOperators = InSelectedFilter->GetSupportedOperators();
 		for (auto& FilterOperator : *AvailableOperators)
 		{
 			AvailableFilterOperators->Add(FilterOperator);
@@ -152,15 +153,15 @@ void FFilterConfiguratorNode::ProcessFilter()
 	}
 	else
 	{
-		switch (SelectedFilter->DataType)
+		switch (SelectedFilter->GetDataType())
 		{
 		case EFilterDataType::Double:
 		{
-			if (SelectedFilter->Converter.IsValid())
+			if (SelectedFilter->GetConverter().IsValid())
 			{
 				double Value = 0.0;
 				FText Errors;
-				bool Result = SelectedFilter->Converter->Convert(TextBoxValue, Value, Errors);
+				bool Result = SelectedFilter->GetConverter()->Convert(TextBoxValue, Value, Errors);
 				FilterValue.Set<double>(Result ? Value : 0.0);
 			}
 			else
@@ -171,11 +172,11 @@ void FFilterConfiguratorNode::ProcessFilter()
 		}
 		case EFilterDataType::Int64:
 		{
-			if (SelectedFilter->Converter.IsValid())
+			if (SelectedFilter->GetConverter().IsValid())
 			{
 				int64 Value = 0;
 				FText Errors;
-				bool Result = SelectedFilter->Converter->Convert(TextBoxValue, Value, Errors);
+				bool Result = SelectedFilter->GetConverter()->Convert(TextBoxValue, Value, Errors);
 				FilterValue.Set<int64>(Result ? Value : 0);
 			}
 			else
@@ -198,10 +199,10 @@ void FFilterConfiguratorNode::ProcessFilter()
 		}
 		case EFilterDataType::StringInt64Pair:
 		{
-			checkf(SelectedFilter->Converter.IsValid(), TEXT("StringToInt64Pair filters must have a converter set"));
+			checkf(SelectedFilter->GetConverter().IsValid(), TEXT("StringToInt64Pair filters must have a converter set"));
 			int64 Value = 0;
 			FText Errors;
-			bool Result = SelectedFilter->Converter->Convert(TextBoxValue, Value, Errors);
+			bool Result = SelectedFilter->GetConverter()->Convert(TextBoxValue, Value, Errors);
 			FilterValue.Set<int64>(Result ? Value : -1);
 		}
 		}
@@ -220,7 +221,7 @@ void FFilterConfiguratorNode::GetUsedKeys(TSet<int32>& KeysUsed) const
 	}
 	else
 	{
-		KeysUsed.Add(SelectedFilter->Key);
+		KeysUsed.Add(SelectedFilter->GetKey());
 	}
 }
 
@@ -231,7 +232,7 @@ bool FFilterConfiguratorNode::ApplyFilters(const FFilterContext& Context) const
 	bool Ret = true;
 	if (IsGroup())
 	{
-		switch (SelectedFilterGroupOperator->Type)
+		switch (SelectedFilterGroupOperator->GetType())
 		{
 		case EFilterGroupOperator::And:
 		{
@@ -260,21 +261,21 @@ bool FFilterConfiguratorNode::ApplyFilters(const FFilterContext& Context) const
 	}
 	else
 	{
-		if (!Context.HasFilterData(SelectedFilter->Key))
+		if (!Context.HasFilterData(SelectedFilter->GetKey()))
 		{
 			// If data is not set for this filter return the value specified in the Context.
 			return Context.GetReturnValueForUnsetFilters();
 		}
 
-		switch (SelectedFilter->DataType)
+		switch (SelectedFilter->GetDataType())
 		{
 		case EFilterDataType::Double:
 		{
 			FFilterOperator<double>* Operator = (FFilterOperator<double>*) SelectedFilterOperator.Get();
 			double Value = 0.0;
-			Context.GetFilterData<double>(SelectedFilter->Key, Value);
+			Context.GetFilterData<double>(SelectedFilter->GetKey(), Value);
 
-			Ret = Operator->Func(Value, FilterValue.Get<double>());
+			Ret = Operator->Apply(Value, FilterValue.Get<double>());
 			break;
 		}
 		case EFilterDataType::Int64:
@@ -282,18 +283,18 @@ bool FFilterConfiguratorNode::ApplyFilters(const FFilterContext& Context) const
 		{
 			FFilterOperator<int64>* Operator = (FFilterOperator<int64>*) SelectedFilterOperator.Get();
 			int64 Value = 0;
-			Context.GetFilterData<int64>(SelectedFilter->Key, Value);
+			Context.GetFilterData<int64>(SelectedFilter->GetKey(), Value);
 
-			Ret = Operator->Func(Value, FilterValue.Get<int64>());
+			Ret = Operator->Apply(Value, FilterValue.Get<int64>());
 			break;
 		}
 		case EFilterDataType::String:
 		{
 			FFilterOperator<FString>* Operator = (FFilterOperator<FString>*) SelectedFilterOperator.Get();
 			FString Value;
-			Context.GetFilterData<FString>(SelectedFilter->Key, Value);
+			Context.GetFilterData<FString>(SelectedFilter->GetKey(), Value);
 
-			Ret = Operator->Func(Value, FilterValue.Get<FString>());
+			Ret = Operator->Apply(Value, FilterValue.Get<FString>());
 			break;
 		}
 		default:
diff --git a/Engine/Source/Developer/TraceInsights/Private/Insights/ViewModels/FilterConfiguratorNode.h b/Engine/Source/Developer/TraceInsights/Private/Insights/ViewModels/FilterConfiguratorNode.h
--- a/Engine/Source/Developer/TraceInsights/Private/Insights/ViewModels/FilterConfiguratorNode.h
+++ b/Engine/Source/Developer/TraceInsights/Private/Insights/ViewModels/FilterConfiguratorNode.h
@@ -57,18 +57,18 @@ public:
 
 	virtual ~FFilterConfiguratorNode() {}
 
-	void SetAvailableFilters(TSharedPtr<TArray<TSharedPtr<struct FFilter>>> InAvailableFilters);
-	TSharedPtr<TArray<TSharedPtr<struct FFilter>>> GetAvailableFilters() { return AvailableFilters; }
+	void SetAvailableFilters(TSharedPtr<TArray<TSharedPtr<FFilter>>> InAvailableFilters);
+	TSharedPtr<TArray<TSharedPtr<FFilter>>> GetAvailableFilters() { return AvailableFilters; }
 
-	void SetSelectedFilter(TSharedPtr<struct FFilter> InSelectedFilter);
-	TSharedPtr<struct FFilter> GetSelectedFilter() const { return SelectedFilter; }
+	void SetSelectedFilter(TSharedPtr<FFilter> InSelectedFilter);
+	TSharedPtr<FFilter> GetSelectedFilter() const { return SelectedFilter; }
 
-	void SetSelectedFilterOperator(TSharedPtr<class IFilterOperator> InSelectedFilterOperator) { SelectedFilterOperator = InSelectedFilterOperator; }
+	void SetSelectedFilterOperator(TSharedPtr<IFilterOperator> InSelectedFilterOperator) { SelectedFilterOperator = InSelectedFilterOperator; }
 	TSharedPtr<IFilterOperator> GetSelectedFilterOperator() const { return SelectedFilterOperator;	}
-	const TArray<TSharedPtr<struct FFilterGroupOperator>>& GetFilterGroupOperators();
+	const TArray<TSharedPtr<FFilterGroupOperator>>& GetFilterGroupOperators();
 
-	void SetSelectedFilterGroupOperator(TSharedPtr<struct FFilterGroupOperator> InSelectedFilterGroupOperator) { SelectedFilterGroupOperator = InSelectedFilterGroupOperator; }
-	TSharedPtr<struct FFilterGroupOperator> GetSelectedFilterGroupOperator() const { return SelectedFilterGroupOperator; }
+	void SetSelectedFilterGroupOperator(TSharedPtr<FFilterGroupOperator> InSelectedFilterGroupOperator) { SelectedFilterGroupOperator = InSelectedFilterGroupOperator; }
+	TSharedPtr<FFilterGroupOperator> GetSelectedFilterGroupOperator() const { return SelectedFilterGroupOperator; }
 
 	TSharedPtr<TArray<TSharedPtr<IFilterOperator>>> GetAvailableFilterOperators() const {	return AvailableFilterOperators;	}
 
@@ -88,13 +88,13 @@ public:
 private:
 	FFilterConfiguratorNode& operator=(const FFilterConfiguratorNode& Other);
 
-	TSharedPtr<TArray<TSharedPtr<struct FFilter>>> AvailableFilters;
+	TSharedPtr<TArray<TSharedPtr<FFilter>>> AvailableFilters;
 
-	TSharedPtr<struct FFilter> SelectedFilter;
+	TSharedPtr<FFilter> SelectedFilter;
 
-	TSharedPtr<class IFilterOperator> SelectedFilterOperator;
+	TSharedPtr<IFilterOperator> SelectedFilterOperator;
 
-	TSharedPtr<struct FFilterGroupOperator> SelectedFilterGroupOperator;
+	TSharedPtr<FFilterGroupOperator> SelectedFilterGroupOperator;
 
 	TSharedPtr<TArray<TSharedPtr<IFilterOperator>>> AvailableFilterOperators;
 
diff --git a/Engine/Source/Developer/TraceInsights/Private/Insights/ViewModels/Filters.h b/Engine/Source/Developer/TraceInsights/Private/Insights/ViewModels/Filters.h
--- a/Engine/Source/Developer/TraceInsights/Private/Insights/ViewModels/Filters.h
+++ b/Engine/Source/Developer/TraceInsights/Private/Insights/ViewModels/Filters.h
@@ -45,9 +45,8 @@ enum class EFilterOperator : uint8
 class IFilterOperator
 {
 public:
-
-	virtual EFilterOperator GetKey() = 0;
-	virtual FString GetName() = 0;
+	virtual EFilterOperator GetKey() const = 0;
+	virtual FString GetName() const = 0;
 };
 
 ////////////////////////////////////////////////////////////////////////////////////////////////////
@@ -56,13 +55,13 @@ template<typename T>
 class FFilterOperator : public IFilterOperator
 {
 public:
-	typedef TFunction<bool(T, T)> OperatorFunc;
+	typedef TFunction<bool(T, T)> FOperatorFunc;
 
 public:
-	FFilterOperator(EFilterOperator InKey, FString InName, OperatorFunc InFunc)
-		: Func(InFunc)
-		, Key(InKey)
+	FFilterOperator(EFilterOperator InKey, FString InName, FOperatorFunc InFunc)
+		: Key(InKey)
 		, Name(InName)
+		, Func(InFunc)
 	{
 	}
 
@@ -70,14 +69,15 @@ public:
 	{
 	}
 
-	virtual EFilterOperator GetKey() override { return Key; }
-	virtual FString GetName() override { return Name; };
+	virtual EFilterOperator GetKey() const override { return Key; }
+	virtual FString GetName() const override { return Name; };
 
-	OperatorFunc Func;
+	bool Apply(T InValueA, T InValueB) const { return Func(InValueA, InValueB); }
 
 private:
 	EFilterOperator Key;
 	FString Name;
+	FOperatorFunc Func;
 };
 
 ////////////////////////////////////////////////////////////////////////////////////////////////////
@@ -90,8 +90,9 @@ enum class EFilterGroupOperator
 
 ////////////////////////////////////////////////////////////////////////////////////////////////////
 
-struct FFilterGroupOperator
+class FFilterGroupOperator
 {
+public:
 	FFilterGroupOperator(EFilterGroupOperator InType, FText InName, FText InDesc)
 		: Type(InType)
 		, Name(InName)
@@ -99,6 +100,11 @@ struct FFilterGroupOperator
 	{
 	}
 
+	EFilterGroupOperator GetType() const { return Type; }
+	const FText& GetName() const { return Name; }
+	const FText& GetDesc() const { return Desc; }
+
+private:
 	EFilterGroupOperator Type;
 	FText Name;
 	FText Desc;
@@ -119,16 +125,17 @@ public:
 
 typedef TSharedPtr<const TArray<TSharedPtr<IFilterOperator>>> SupportedOperatorsArrayPtr;
 
-struct FFilter
+class FFilter
 {
 	INSIGHTS_DECLARE_RTTI_BASE(FFilter)
 
 public:
-	FFilter(int32 InKey, FText InName, FText InDesc, EFilterDataType InDataType, SupportedOperatorsArrayPtr InSupportedOperators)
+	FFilter(int32 InKey, FText InName, FText InDesc, EFilterDataType InDataType, TSharedPtr<IFilterValueConverter> InConverter, SupportedOperatorsArrayPtr InSupportedOperators)
 		: Key(InKey)
 		, Name(InName)
 		, Desc(InDesc)
 		, DataType(InDataType)
+		, Converter(InConverter)
 		, SupportedOperators(InSupportedOperators)
 	{
 	}
@@ -137,8 +144,14 @@ public:
 	{
 	}
 
+	int32 GetKey() const { return Key; }
+	const FText& GetName() const { return Name; }
+	const FText& GetDesc() const { return Desc; }
+	EFilterDataType GetDataType() const { return DataType; }
+	const TSharedPtr<IFilterValueConverter>& GetConverter() const { return Converter; }
 	SupportedOperatorsArrayPtr GetSupportedOperators() const { return SupportedOperators; }
 
+private:
 	int32 Key;
 	FText Name;
 	FText Desc;
@@ -149,15 +162,16 @@ public:
 
 ////////////////////////////////////////////////////////////////////////////////////////////////////
 
-struct FFilterWithSuggestions : FFilter
+class FFilterWithSuggestions : public FFilter
 {
-	typedef TFunction<void(const FString& /*Text*/, TArray<FString>& OutSuggestions)> GetSuggestionsCallback;
-
 	INSIGHTS_DECLARE_RTTI(FFilterWithSuggestions, FFilter)
 
 public:
-	FFilterWithSuggestions(int32 InKey, FText InName, FText InDesc, EFilterDataType InDataType, SupportedOperatorsArrayPtr InSupportedOperators)
-		: FFilter(InKey, InName, InDesc, InDataType, InSupportedOperators)
+	typedef TFunction<void(const FString& /*Text*/, TArray<FString>& OutSuggestions)> FGetSuggestionsCallback;
+
+public:
+	FFilterWithSuggestions(int32 InKey, FText InName, FText InDesc, EFilterDataType InDataType, TSharedPtr<IFilterValueConverter> InConverter, SupportedOperatorsArrayPtr InSupportedOperators)
+		: FFilter(InKey, InName, InDesc, InDataType, InConverter, InSupportedOperators)
 	{
 	}
 
@@ -165,7 +179,12 @@ public:
 	{
 	}
 
-	GetSuggestionsCallback Callback;
+	const FGetSuggestionsCallback& GetCallback() const { return Callback; }
+	void SetCallback(FGetSuggestionsCallback InCallback) { Callback = InCallback; }
+	void GetSuggestions(const FString& Text, TArray<FString>& OutSuggestions) { Callback(Text, OutSuggestions); }
+
+private:
+	FGetSuggestionsCallback Callback;
 };
 
 ////////////////////////////////////////////////////////////////////////////////////////////////////
diff --git a/Engine/Source/Developer/TraceInsights/Private/Insights/ViewModels/StatsNode.h b/Engine/Source/Developer/TraceInsights/Private/Insights/ViewModels/StatsNode.h
--- a/Engine/Source/Developer/TraceInsights/Private/Insights/ViewModels/StatsNode.h
+++ b/Engine/Source/Developer/TraceInsights/Private/Insights/ViewModels/StatsNode.h
@@ -246,17 +246,6 @@ public:
 	const FText GetTextForAggregatedStatsUpperQuartile(bool bForTooltip = false) const;
 	const FText GetTextForAggregatedStatsDiff(bool bForTooltip = false) const;
 
-	/** Sorts children using the specified class instance. */
-	template<typename TSortingClass>
-	void SortChildren(const TSortingClass& Instance)
-	{
-		auto Projection = [](Insights::FBaseTreeNodePtr Node) -> FStatsNodePtr
-		{
-			return StaticCastSharedPtr<FStatsNode, Insights::FBaseTreeNode>(Node);
-		};
-		Algo::SortBy(GetChildrenMutable(), Projection, Instance);
-	}
-
 private:
 	/** The counter id provided by the analyzer. */
 	uint32 CounterId;
diff --git a/Engine/Source/Developer/TraceInsights/Private/Insights/ViewModels/ThreadTimingTrack.cpp b/Engine/Source/Developer/TraceInsights/Private/Insights/ViewModels/ThreadTimingTrack.cpp
--- a/Engine/Source/Developer/TraceInsights/Private/Insights/ViewModels/ThreadTimingTrack.cpp
+++ b/Engine/Source/Developer/TraceInsights/Private/Insights/ViewModels/ThreadTimingTrack.cpp
@@ -1201,7 +1201,7 @@ void FThreadTimingTrack::InitTooltip(FTooltipDrawState& InOutTooltip, const ITim
 				}
 
 				InOutTooltip.AddNameValueTextLine(TEXT("Created:"), FString::Printf(TEXT("%f on %s"), Task.CreatedTimestamp, *GetTrackName(Task.CreatedThreadId)));
-				
+
 				InOutTooltip.AddNameValueTextLine(TEXT("Launched:"), FString::Printf(TEXT("%f (+%s) on %s"), Task.LaunchedTimestamp, *TimeUtils::FormatTimeAuto(Task.LaunchedTimestamp - Task.CreatedTimestamp), *GetTrackName(Task.LaunchedThreadId)));
 
 				InOutTooltip.AddNameValueTextLine(TEXT("Scheduled:"), FString::Printf(TEXT("%f (+%s) on %s"), Task.ScheduledTimestamp, *TimeUtils::FormatTimeAuto(Task.ScheduledTimestamp - Task.LaunchedTimestamp), *GetTrackName(Task.ScheduledThreadId)));
@@ -1749,12 +1749,38 @@ void FThreadTimingTrack::OnFilterTrackClicked()
 	if (!FilterConfigurator.IsValid())
 	{
 		FilterConfigurator = MakeShared<FFilterConfigurator>();
-		TSharedPtr<TArray<TSharedPtr<struct FFilter>>>& AvailableFilters = FilterConfigurator->GetAvailableFilters();
 
-		AvailableFilters->Add(MakeShared<FFilter>(static_cast<int32>(EFilterField::StartTime), LOCTEXT("StartTime", "Start Time"), LOCTEXT("StartTime", "Start Time"), EFilterDataType::Double, FFilterService::Get()->GetDoubleOperators()));
-		AvailableFilters->Add(MakeShared<FFilter>(static_cast<int32>(EFilterField::EndTime), LOCTEXT("EndTime", "End Time"), LOCTEXT("EndTime", "End Time"), EFilterDataType::Double, FFilterService::Get()->GetDoubleOperators()));
-		AvailableFilters->Add(MakeShared<FFilter>(static_cast<int32>(EFilterField::Duration), LOCTEXT("Duration", "Duration"), LOCTEXT("Duration", "Duration"), EFilterDataType::Double, FFilterService::Get()->GetDoubleOperators()));
-		AvailableFilters->Add(MakeShared<FFilter>(static_cast<int32>(EFilterField::TimerId), LOCTEXT("TimerId", "Timer Id"), LOCTEXT("TimerId", "Timer Id"), EFilterDataType::Int64, FFilterService::Get()->GetIntegerOperators()));
+		FilterConfigurator->Add(MakeShared<FFilter>(
+			static_cast<int32>(EFilterField::StartTime),
+			LOCTEXT("StartTime", "Start Time"),
+			LOCTEXT("StartTime", "Start Time"),
+			EFilterDataType::Double,
+			nullptr,
+			FFilterService::Get()->GetDoubleOperators()));
+
+		FilterConfigurator->Add(MakeShared<FFilter>(
+			static_cast<int32>(EFilterField::EndTime),
+			LOCTEXT("EndTime", "End Time"),
+			LOCTEXT("EndTime", "End Time"),
+			EFilterDataType::Double,
+			nullptr,
+			FFilterService::Get()->GetDoubleOperators()));
+
+		FilterConfigurator->Add(MakeShared<FFilter>(
+			static_cast<int32>(EFilterField::Duration),
+			LOCTEXT("Duration", "Duration"),
+			LOCTEXT("Duration", "Duration"),
+			EFilterDataType::Double,
+			nullptr,
+			FFilterService::Get()->GetDoubleOperators()));
+
+		FilterConfigurator->Add(MakeShared<FFilter>(
+			static_cast<int32>(EFilterField::TimerId),
+			LOCTEXT("TimerId", "Timer Id"),
+			LOCTEXT("TimerId", "Timer Id"),
+			EFilterDataType::Int64,
+			nullptr,
+			FFilterService::Get()->GetIntegerOperators()));
 	}
 	else
 	{
@@ -1776,16 +1802,7 @@ void FThreadTimingTrack::OnFilterTrackClicked()
 
 bool FThreadTimingTrack::HasCustomFilter() const
 {
-	if (!FilterConfigurator.IsValid())
-	{
-		return false;
-	}
-	if (FilterConfigurator->GetRootNode().IsValid() && FilterConfigurator->GetRootNode()->GetChildren().Num() > 0)
-	{
-		return true;
-	}
-
-	return false;
+	return FilterConfigurator.IsValid() && !FilterConfigurator->IsEmpty();
 }
 
 ////////////////////////////////////////////////////////////////////////////////////////////////////
diff --git a/Engine/Source/Developer/TraceInsights/Private/Insights/Widgets/SAdvancedFilter.h b/Engine/Source/Developer/TraceInsights/Private/Insights/Widgets/SAdvancedFilter.h
--- a/Engine/Source/Developer/TraceInsights/Private/Insights/Widgets/SAdvancedFilter.h
+++ b/Engine/Source/Developer/TraceInsights/Private/Insights/Widgets/SAdvancedFilter.h
@@ -16,6 +16,7 @@ class SDockTab;
 namespace Insights
 {
 
+class FFilterConfigurator;
 class SFilterConfigurator;
 
 ////////////////////////////////////////////////////////////////////////////////////////////////////
@@ -38,7 +39,7 @@ public:
 	 * Construct this widget
 	 * @param InArgs - The declaration data for this widget
 	 */
-	void Construct(const FArguments& InArgs, TSharedPtr<class FFilterConfigurator> InFilterConfiguratorViewModel);
+	void Construct(const FArguments& InArgs, TSharedPtr<FFilterConfigurator> InFilterConfiguratorViewModel);
 
 	void Reset();
 
@@ -57,9 +58,9 @@ private:
 private:
 	TSharedPtr<SFilterConfigurator> FilterConfigurator;
 
-	TWeakPtr<class FFilterConfigurator> OriginalFilterConfiguratorViewModel;
+	TWeakPtr<FFilterConfigurator> OriginalFilterConfiguratorViewModel;
 
-	TSharedPtr<class FFilterConfigurator> FilterConfiguratorViewModel;
+	TSharedPtr<FFilterConfigurator> FilterConfiguratorViewModel;
 
 	TWeakPtr<SDockTab> ParentTab;
 
diff --git a/Engine/Source/Developer/TraceInsights/Private/Insights/Widgets/SFilterConfigurator.h b/Engine/Source/Developer/TraceInsights/Private/Insights/Widgets/SFilterConfigurator.h
--- a/Engine/Source/Developer/TraceInsights/Private/Insights/Widgets/SFilterConfigurator.h
+++ b/Engine/Source/Developer/TraceInsights/Private/Insights/Widgets/SFilterConfigurator.h
@@ -16,6 +16,8 @@ class SDockTab;
 namespace Insights
 {
 
+class FFilterConfigurator;
+
 ////////////////////////////////////////////////////////////////////////////////////////////////////
 
 /**
@@ -37,7 +39,7 @@ public:
 	 * Construct this widget
 	 * @param InArgs - The declaration data for this widget
 	 */
-	void Construct(const FArguments& InArgs, TSharedPtr<class FFilterConfigurator> InFilterConfiguratorViewModel);
+	void Construct(const FArguments& InArgs, TSharedPtr<FFilterConfigurator> InFilterConfiguratorViewModel);
 
 	void Reset();
 
@@ -62,7 +64,7 @@ private:
 	void SetInitialExpansionRec(const FFilterConfiguratorNodePtr& Node, bool Value);
 
 private:
-	TSharedPtr<class FFilterConfigurator> FilterConfiguratorViewModel;
+	TSharedPtr<FFilterConfigurator> FilterConfiguratorViewModel;
 
 	/** The tree widget which holds the filters. */
 	TSharedPtr<STreeView<FFilterConfiguratorNodePtr>> TreeView;
diff --git a/Engine/Source/Developer/TraceInsights/Private/Insights/Widgets/SFilterConfiguratorRow.cpp b/Engine/Source/Developer/TraceInsights/Private/Insights/Widgets/SFilterConfiguratorRow.cpp
--- a/Engine/Source/Developer/TraceInsights/Private/Insights/Widgets/SFilterConfiguratorRow.cpp
+++ b/Engine/Source/Developer/TraceInsights/Private/Insights/Widgets/SFilterConfiguratorRow.cpp
@@ -19,7 +19,7 @@
 #include "Insights/InsightsStyle.h"
 #include "Insights/ViewModels/Filters.h"
 
-#define LOCTEXT_NAMESPACE "SFilterConfiguratorRow"
+#define LOCTEXT_NAMESPACE "Insights::SFilterConfiguratorRow"
 
 namespace Insights
 {
@@ -289,7 +289,7 @@ TSharedRef<SWidget> SFilterConfiguratorRow::AvailableFilters_OnGenerateWidget(TS
 		.AutoWidth()
 		[
 			SNew(STextBlock)
-			.Text(InFilter->Name)
+			.Text(InFilter->GetName())
 			.Margin(2.0f)
 		];
 
@@ -317,12 +317,12 @@ TSharedRef<SWidget> SFilterConfiguratorRow::AvailableFilterOperators_OnGenerateW
 TSharedRef<SWidget> SFilterConfiguratorRow::FilterGroupOperators_OnGenerateWidget(TSharedPtr<FFilterGroupOperator> InFilterGroupOperator)
 {
 	TSharedRef<SHorizontalBox> Widget = SNew(SHorizontalBox);
-	Widget->SetToolTipText(InFilterGroupOperator->Desc);
+	Widget->SetToolTipText(InFilterGroupOperator->GetDesc());
 	Widget->AddSlot()
 		.AutoWidth()
 		[
 			SNew(STextBlock)
-			.Text(InFilterGroupOperator->Name)
+			.Text(InFilterGroupOperator->GetName())
 			.Margin(2.0f)
 		];
 
@@ -333,9 +333,9 @@ END_SLATE_FUNCTION_BUILD_OPTIMIZATION
 
 ////////////////////////////////////////////////////////////////////////////////////////////////////
 
-const TArray<TSharedPtr<struct FFilter>>* SFilterConfiguratorRow::GetAvailableFilters()
+const TArray<TSharedPtr<FFilter>>* SFilterConfiguratorRow::GetAvailableFilters()
 {
-	return &*FilterConfiguratorNodePtr->GetAvailableFilters();
+	return FilterConfiguratorNodePtr->GetAvailableFilters().Get();
 }
 
 ///////////////////////////////////////////////////////////////////////////////////////////////////
@@ -354,7 +354,7 @@ void SFilterConfiguratorRow::AvailableFilters_OnSelectionChanged(TSharedPtr<FFil
 
 FText SFilterConfiguratorRow::AvailableFilters_GetSelectionText() const
 {
-	return FilterConfiguratorNodePtr->GetSelectedFilter()->Name;
+	return FilterConfiguratorNodePtr->GetSelectedFilter()->GetName();
 }
 
 ///////////////////////////////////////////////////////////////////////////////////////////////////
@@ -396,7 +396,7 @@ void SFilterConfiguratorRow::FilterGroupOperators_OnSelectionChanged(TSharedPtr<
 
 FText SFilterConfiguratorRow::FilterGroupOperators_GetSelectionText() const
 {
-	return FilterConfiguratorNodePtr->GetSelectedFilterGroupOperator()->Name;
+	return FilterConfiguratorNodePtr->GetSelectedFilterGroupOperator()->GetName();
 }
 
 ///////////////////////////////////////////////////////////////////////////////////////////////////
@@ -478,7 +478,7 @@ void SFilterConfiguratorRow::OnTextBoxValueCommitted(const FText& InNewText, ETe
 
 FText SFilterConfiguratorRow::GetTextBoxTooltipText() const
 {
-	TSharedPtr<IFilterValueConverter> Converter = FilterConfiguratorNodePtr->GetSelectedFilter()->Converter;
+	const TSharedPtr<IFilterValueConverter>& Converter = FilterConfiguratorNodePtr->GetSelectedFilter()->GetConverter();
 	if (Converter.IsValid())
 	{
 		return Converter->GetTooltipText();
@@ -491,7 +491,7 @@ FText SFilterConfiguratorRow::GetTextBoxTooltipText() const
 
 FText SFilterConfiguratorRow::GetTextBoxHintText() const
 {
-	TSharedPtr<IFilterValueConverter> Converter = FilterConfiguratorNodePtr->GetSelectedFilter()->Converter;
+	const TSharedPtr<IFilterValueConverter>& Converter = FilterConfiguratorNodePtr->GetSelectedFilter()->GetConverter();
 	if (Converter.IsValid())
 	{
 		return Converter->GetHintText();
@@ -504,15 +504,15 @@ FText SFilterConfiguratorRow::GetTextBoxHintText() const
 
 bool SFilterConfiguratorRow::TextBox_OnVerifyTextChanged(const FText& InText, FText& OutErrorMessage)
 {
-	TSharedPtr<IFilterValueConverter> Converter = FilterConfiguratorNodePtr->GetSelectedFilter()->Converter;
+	const TSharedPtr<IFilterValueConverter>& Converter = FilterConfiguratorNodePtr->GetSelectedFilter()->GetConverter();
 	if (Converter.IsValid())
 	{
-		if (FilterConfiguratorNodePtr->GetSelectedFilter()->DataType == EFilterDataType::Int64)
+		if (FilterConfiguratorNodePtr->GetSelectedFilter()->GetDataType() == EFilterDataType::Int64)
 		{
 			int64 Value;
 			return Converter->Convert(InText.ToString(), Value, OutErrorMessage);
 		}		
-		else if (FilterConfiguratorNodePtr->GetSelectedFilter()->DataType == EFilterDataType::Double)
+		else if (FilterConfiguratorNodePtr->GetSelectedFilter()->GetDataType() == EFilterDataType::Double)
 		{
 			double Value;
 			return Converter->Convert(InText.ToString(), Value, OutErrorMessage);
@@ -531,7 +531,7 @@ void SFilterConfiguratorRow::SuggestionTextBox_GetSuggestions(const FString& Tex
 	if (Filter->Is<FFilterWithSuggestions>())
 	{
 		TSharedPtr<FFilterWithSuggestions> FilterWithSuggestions = StaticCastSharedPtr<FFilterWithSuggestions>(Filter);
-		FilterWithSuggestions->Callback(Text, Suggestions);
+		FilterWithSuggestions->GetSuggestions(Text, Suggestions);
 	}
 }
 
diff --git a/Engine/Source/Developer/TraceInsights/Private/Insights/Widgets/SFilterConfiguratorRow.h b/Engine/Source/Developer/TraceInsights/Private/Insights/Widgets/SFilterConfiguratorRow.h
--- a/Engine/Source/Developer/TraceInsights/Private/Insights/Widgets/SFilterConfiguratorRow.h
+++ b/Engine/Source/Developer/TraceInsights/Private/Insights/Widgets/SFilterConfiguratorRow.h
@@ -13,6 +13,10 @@
 namespace Insights
 {
 
+class FFilter;
+class FFilterGroupOperator;
+class IFilterOperator;
+
 /** Widget that represents a table row in the Filter's tree control. Generates widgets for each column on demand. */
 class SFilterConfiguratorRow : public SMultiColumnTableRow<FFilterConfiguratorNodePtr>
 {
@@ -30,27 +34,27 @@ public:
 	virtual TSharedRef<SWidget> GenerateWidgetForColumn(const FName& InColumnName) override;
 
 private:
-	const TArray<TSharedPtr<struct FFilter>>* GetAvailableFilters();
+	const TArray<TSharedPtr<FFilter>>* GetAvailableFilters();
 
-	TSharedRef<SWidget> AvailableFilters_OnGenerateWidget(TSharedPtr<struct FFilter> InFilter);
+	TSharedRef<SWidget> AvailableFilters_OnGenerateWidget(TSharedPtr<FFilter> InFilter);
 
-	void AvailableFilters_OnSelectionChanged(TSharedPtr<struct FFilter> InFilter, ESelectInfo::Type SelectInfo);
+	void AvailableFilters_OnSelectionChanged(TSharedPtr<FFilter> InFilter, ESelectInfo::Type SelectInfo);
 
 	FText AvailableFilters_GetSelectionText() const;
 
-	const TArray<TSharedPtr<class IFilterOperator>>* GetAvailableFilterOperators();
+	const TArray<TSharedPtr<IFilterOperator>>* GetAvailableFilterOperators();
 
-	TSharedRef<SWidget> AvailableFilterOperators_OnGenerateWidget(TSharedPtr<class IFilterOperator> InFilter);
+	TSharedRef<SWidget> AvailableFilterOperators_OnGenerateWidget(TSharedPtr<IFilterOperator> InFilter);
 
-	void AvailableFilterOperators_OnSelectionChanged(TSharedPtr<class IFilterOperator> InFilter, ESelectInfo::Type SelectInfo);
+	void AvailableFilterOperators_OnSelectionChanged(TSharedPtr<IFilterOperator> InFilter, ESelectInfo::Type SelectInfo);
 
 	FText AvailableFilterOperators_GetSelectionText() const;
 
-	const TArray<TSharedPtr<struct FFilterGroupOperator>>* GetFilterGroupOperators();
+	const TArray<TSharedPtr<FFilterGroupOperator>>* GetFilterGroupOperators();
 
-	TSharedRef<SWidget> FilterGroupOperators_OnGenerateWidget(TSharedPtr<struct FFilterGroupOperator> InFilter);
+	TSharedRef<SWidget> FilterGroupOperators_OnGenerateWidget(TSharedPtr<FFilterGroupOperator> InFilter);
 
-	void FilterGroupOperators_OnSelectionChanged(TSharedPtr<struct FFilterGroupOperator> InFilter, ESelectInfo::Type SelectInfo);
+	void FilterGroupOperators_OnSelectionChanged(TSharedPtr<FFilterGroupOperator> InFilter, ESelectInfo::Type SelectInfo);
 
 	FText FilterGroupOperators_GetSelectionText() const;
 
@@ -77,11 +81,11 @@ private:
 private:
 	FFilterConfiguratorNodePtr FilterConfiguratorNodePtr;
 
-	TSharedPtr<SComboBox<TSharedPtr<struct FFilter>>> FilterTypeComboBox;
+	TSharedPtr<SComboBox<TSharedPtr<FFilter>>> FilterTypeComboBox;
 
-	TSharedPtr<SComboBox<TSharedPtr<class IFilterOperator>>> FilterOperatorComboBox;
+	TSharedPtr<SComboBox<TSharedPtr<IFilterOperator>>> FilterOperatorComboBox;
 
-	TSharedPtr<SComboBox<TSharedPtr<struct FFilterGroupOperator>>> FilterGroupOperatorComboBox;
+	TSharedPtr<SComboBox<TSharedPtr<FFilterGroupOperator>>> FilterGroupOperatorComboBox;
 
 	FString SuggestionTextBoxValue;
 };
diff --git a/Engine/Source/Developer/TraceInsights/Private/Insights/Widgets/SStatsView.cpp b/Engine/Source/Developer/TraceInsights/Private/Insights/Widgets/SStatsView.cpp
--- a/Engine/Source/Developer/TraceInsights/Private/Insights/Widgets/SStatsView.cpp
+++ b/Engine/Source/Developer/TraceInsights/Private/Insights/Widgets/SStatsView.cpp
@@ -1539,19 +1539,13 @@ void SStatsView::SortTreeNodes()
 
 void SStatsView::SortTreeNodesRec(FStatsNode& Node, const Insights::ITableCellValueSorter& Sorter)
 {
-	if (ColumnSortMode == EColumnSortMode::Type::Descending)
-	{
-		Node.SortChildrenDescending(Sorter);
-	}
-	else // if (ColumnSortMode == EColumnSortMode::Type::Ascending)
-	{
-		Node.SortChildrenAscending(Sorter);
-	}
+	Insights::ESortMode SortMode = (ColumnSortMode == EColumnSortMode::Type::Descending) ? Insights::ESortMode::Descending : Insights::ESortMode::Ascending;
+	Node.SortChildren(Sorter, SortMode);
 
 #if 0 // Current groupings creates only one level.
 	for (Insights::FBaseTreeNodePtr ChildPtr : Node.GetChildren())
 	{
-		if (ChildPtr->GetChildren().Num() > 0)
+		if (ChildPtr->GetChildrenCount() > 0)
 		{
 			SortTreeNodesRec(*StaticCastSharedPtr<FTimerNode>(ChildPtr), Sorter);
 		}
diff --git a/Engine/Source/Developer/TraceInsights/Private/Insights/Widgets/STimerTreeView.cpp b/Engine/Source/Developer/TraceInsights/Private/Insights/Widgets/STimerTreeView.cpp
--- a/Engine/Source/Developer/TraceInsights/Private/Insights/Widgets/STimerTreeView.cpp
+++ b/Engine/Source/Developer/TraceInsights/Private/Insights/Widgets/STimerTreeView.cpp
@@ -539,7 +539,7 @@ void STimerTreeView::TreeView_OnGetChildren(FTimerNodePtr InParent, TArray<FTime
 void STimerTreeView::TreeView_OnMouseButtonDoubleClick(FTimerNodePtr TimerNodePtr)
 {
 	//if (TimerNodePtr->IsGroup())
-	if (TimerNodePtr->GetChildren().Num() > 0)
+	if (TimerNodePtr->GetChildrenCount() > 0)
 	{
 		const bool bIsGroupExpanded = TreeView->IsItemExpanded(TimerNodePtr);
 		TreeView->SetItemExpansion(TimerNodePtr, !bIsGroupExpanded);
@@ -686,19 +686,13 @@ void STimerTreeView::SortTreeNodes()
 
 void STimerTreeView::SortTreeNodesRec(FTimerNode& Node, const Insights::ITableCellValueSorter& Sorter)
 {
-	if (ColumnSortMode == EColumnSortMode::Type::Descending)
-	{
-		Node.SortChildrenDescending(Sorter);
-	}
-	else // if (ColumnSortMode == EColumnSortMode::Type::Ascending)
-	{
-		Node.SortChildrenAscending(Sorter);
-	}
+	Insights::ESortMode SortMode = (ColumnSortMode == EColumnSortMode::Type::Descending) ? Insights::ESortMode::Descending : Insights::ESortMode::Ascending;
+	Node.SortChildren(Sorter, SortMode);
 
 	for (Insights::FBaseTreeNodePtr ChildPtr : Node.GetChildren())
 	{
 		//if (ChildPtr->IsGroup())
-		if (ChildPtr->GetChildren().Num() > 0)
+		if (ChildPtr->GetChildrenCount() > 0)
 		{
 			SortTreeNodesRec(*StaticCastSharedPtr<FTimerNode>(ChildPtr), Sorter);
 		}
diff --git a/Engine/Source/Developer/TraceInsights/Private/Insights/Widgets/STimersView.cpp b/Engine/Source/Developer/TraceInsights/Private/Insights/Widgets/STimersView.cpp
--- a/Engine/Source/Developer/TraceInsights/Private/Insights/Widgets/STimersView.cpp
+++ b/Engine/Source/Developer/TraceInsights/Private/Insights/Widgets/STimersView.cpp
@@ -1900,19 +1900,13 @@ void STimersView::SortTreeNodes()
 
 void STimersView::SortTreeNodesRec(FTimerNode& Node, const Insights::ITableCellValueSorter& Sorter)
 {
-	if (ColumnSortMode == EColumnSortMode::Type::Descending)
-	{
-		Node.SortChildrenDescending(Sorter);
-	}
-	else // if (ColumnSortMode == EColumnSortMode::Type::Ascending)
-	{
-		Node.SortChildrenAscending(Sorter);
-	}
+	Insights::ESortMode SortMode = (ColumnSortMode == EColumnSortMode::Type::Descending) ? Insights::ESortMode::Descending : Insights::ESortMode::Ascending;
+	Node.SortChildren(Sorter, SortMode);
 
 #if 0 // Current groupings creates only one level.
 	for (Insights::FBaseTreeNodePtr ChildPtr : Node.GetChildren())
 	{
-		if (ChildPtr->GetChildren().Num() > 0)
+		if (ChildPtr->GetChildrenCount() > 0)
 		{
 			SortTreeNodesRec(*StaticCastSharedPtr<FTimerNode>(ChildPtr), Sorter);
 		}
diff --git a/Engine/Source/Developer/TraceInsights/Private/Insights/Widgets/STimingView.cpp b/Engine/Source/Developer/TraceInsights/Private/Insights/Widgets/STimingView.cpp
--- a/Engine/Source/Developer/TraceInsights/Private/Insights/Widgets/STimingView.cpp
+++ b/Engine/Source/Developer/TraceInsights/Private/Insights/Widgets/STimingView.cpp
@@ -2912,7 +2912,7 @@ void STimingView::ShowContextMenu(const FPointerEvent& MouseEvent)
 		{
 			double RangeStart = HoveredEvent->GetStartTime();
 			double RangeDuration = HoveredEvent->GetDuration();
-			
+
 			MenuBuilder.AddMenuEntry(
 				LOCTEXT("ContextMenu_SelectEventTimeRange", "Select Time Range of Event"),
 				FText(),
@@ -4894,40 +4894,66 @@ void STimingView::QuickFind_Execute()
 	if (!QuickFindVm.IsValid())
 	{
 		TSharedPtr<FFilterConfigurator> NewFilterConfigurator = MakeShared<FFilterConfigurator>();
-		TSharedPtr<TArray<TSharedPtr<struct FFilter>>>& AvailableFilters = NewFilterConfigurator->GetAvailableFilters();
-		TSharedPtr<FFilter> CurrentFilter;
-
-		CurrentFilter = AvailableFilters->Add_GetRef(MakeShared<FFilter>(static_cast<int32>(EFilterField::StartTime), LOCTEXT("StartTime", "Start Time"), LOCTEXT("StartTime", "Start Time"), EFilterDataType::Double, FFilterService::Get()->GetDoubleOperators()));
-		CurrentFilter->Converter = MakeShared<FTimeFilterValueConverter>();
-
-		CurrentFilter = AvailableFilters->Add_GetRef(MakeShared<FFilter>(static_cast<int32>(EFilterField::EndTime), LOCTEXT("EndTime", "End Time"), LOCTEXT("EndTime", "End Time"), EFilterDataType::Double, FFilterService::Get()->GetDoubleOperators()));
-		CurrentFilter->Converter = MakeShared<FTimeFilterValueConverter>();
 
-		CurrentFilter = AvailableFilters->Add_GetRef(MakeShared<FFilter>(static_cast<int32>(EFilterField::Duration), LOCTEXT("Duration", "Duration"), LOCTEXT("Duration", "Duration"), EFilterDataType::Double, FFilterService::Get()->GetDoubleOperators()));
-		CurrentFilter->Converter = MakeShared<FTimeFilterValueConverter>();
-
-		TSharedPtr<FFilterWithSuggestions> TrackFilter = MakeShared<FFilterWithSuggestions>(static_cast<int32>(EFilterField::TrackName), LOCTEXT("Track", "Track"), LOCTEXT("Track", "Track"), EFilterDataType::String, FFilterService::Get()->GetStringOperators());
-		TrackFilter->Callback = [this](const FString& Text, TArray<FString>& OutSuggestions)
+		NewFilterConfigurator->Add(MakeShared<FFilter>(
+			static_cast<int32>(EFilterField::StartTime),
+			LOCTEXT("StartTime", "Start Time"),
+			LOCTEXT("StartTime", "Start Time"),
+			EFilterDataType::Double,
+			MakeShared<FTimeFilterValueConverter>(),
+			FFilterService::Get()->GetDoubleOperators()));
+
+		NewFilterConfigurator->Add(MakeShared<FFilter>(
+			static_cast<int32>(EFilterField::EndTime),
+			LOCTEXT("EndTime", "End Time"),
+			LOCTEXT("EndTime", "End Time"),
+			EFilterDataType::Double,
+			MakeShared<FTimeFilterValueConverter>(),
+			FFilterService::Get()->GetDoubleOperators()));
+
+		NewFilterConfigurator->Add(MakeShared<FFilter>(
+			static_cast<int32>(EFilterField::Duration),
+			LOCTEXT("Duration", "Duration"),
+			LOCTEXT("Duration", "Duration"),
+			EFilterDataType::Double,
+			MakeShared<FTimeFilterValueConverter>(),
+			FFilterService::Get()->GetDoubleOperators()));
+
+		TSharedRef<FFilterWithSuggestions> TrackFilter = MakeShared<FFilterWithSuggestions>(
+			static_cast<int32>(EFilterField::TrackName),
+			LOCTEXT("Track", "Track"),
+			LOCTEXT("Track", "Track"),
+			EFilterDataType::String,
+			nullptr,
+			FFilterService::Get()->GetStringOperators());
+		TrackFilter->SetCallback([this](const FString& Text, TArray<FString>& OutSuggestions)
 		{
 			this->PopulateTrackSuggestionList(Text, OutSuggestions);
-		};
+		});
+		NewFilterConfigurator->Add(TrackFilter);
 
-		AvailableFilters->Add(TrackFilter);
-
-		AvailableFilters->Add(MakeShared<FFilter>(static_cast<int32>(EFilterField::TimerId), LOCTEXT("TimerId", "Timer Id"), LOCTEXT("TimerId", "Timer Id"), EFilterDataType::Int64, FFilterService::Get()->GetIntegerOperators()));
+		NewFilterConfigurator->Add(MakeShared<FFilter>(
+			static_cast<int32>(EFilterField::TimerId),
+			LOCTEXT("TimerId", "Timer Id"),
+			LOCTEXT("TimerId", "Timer Id"),
+			EFilterDataType::Int64,
+			nullptr,
+			FFilterService::Get()->GetIntegerOperators()));
 
 		TSharedPtr<TArray<TSharedPtr<IFilterOperator>>> EventNameFilterOperators = MakeShared<TArray<TSharedPtr<IFilterOperator>>>();
 		EventNameFilterOperators->Add(StaticCastSharedRef<IFilterOperator>(MakeShared<FFilterOperator<int64>>(EFilterOperator::Eq, TEXT("Is"), [](int64 lhs, int64 rhs) { return lhs == rhs; })));
-
-		TSharedPtr<FFilterWithSuggestions> TimerNameFilter = MakeShared<FFilterWithSuggestions>(static_cast<int32>(EFilterField::TimerName), LOCTEXT("TimerName", "Timer Name"), LOCTEXT("TimerName", "Timer Name"), EFilterDataType::StringInt64Pair, EventNameFilterOperators);
-		TimerNameFilter->Callback = [this](const FString& Text, TArray<FString>& OutSuggestions)
+		TSharedRef<FFilterWithSuggestions> TimerNameFilter = MakeShared<FFilterWithSuggestions>(
+			static_cast<int32>(EFilterField::TimerName),
+			LOCTEXT("TimerName", "Timer Name"),
+			LOCTEXT("TimerName", "Timer Name"),
+			EFilterDataType::StringInt64Pair,
+			MakeShared<FEventNameFilterValueConverter>(),
+			EventNameFilterOperators);
+		TimerNameFilter->SetCallback([this](const FString& Text, TArray<FString>& OutSuggestions)
 		{
 			this->PopulateTimerNameSuggestionList(Text, OutSuggestions);
-		};
-
-		TimerNameFilter->Converter = MakeShared<FEventNameFilterValueConverter>();
-
-		AvailableFilters->Add(TimerNameFilter);
+		});
+		NewFilterConfigurator->Add(TimerNameFilter);
 
 		for (Insights::ITimingViewExtender* Extender : GetExtenders())
 		{
