[Insights]
  - Separated tree view node filtering by hierarchy filtering. The order of async operations was changed. Node filtering (i.e. advanced filtering using the Filter Configurator) is applied first. Then grouping + aggregation is applied. Then sorting is applied. Finally, the hierarchy filtering (i.e. text search) is applied. This enables text search (hierarchy filtering) and advanced filtering (node filtering) to be used in the same time. Also, the aggregation is now applied only once. These enables significant performance improvements for the tree operations, especially when node filtering is applied.
  - Added support for tree group nodes to lazy create children.
  - Refactored tree table aggregation to be able to apply it to a single node or to an entire sub-tree.

diff --git a/Engine/Source/Developer/TraceInsights/Private/Insights/CookProfiler/Widgets/SPackageTableTreeView.cpp b/Engine/Source/Developer/TraceInsights/Private/Insights/CookProfiler/Widgets/SPackageTableTreeView.cpp
--- a/Engine/Source/Developer/TraceInsights/Private/Insights/CookProfiler/Widgets/SPackageTableTreeView.cpp
+++ b/Engine/Source/Developer/TraceInsights/Private/Insights/CookProfiler/Widgets/SPackageTableTreeView.cpp
@@ -127,7 +127,7 @@ void SPackageTableTreeView::RebuildTree(bool bResync)
 	TSharedPtr<FPackageTable> PackageTable = GetPackageTable();
 	TArray<FPackageEntry>& Packages = PackageTable->GetPackageEntries();
 	Packages.Empty();
-	TableTreeNodes.Empty();
+	TableRowNodes.Empty();
 
 	const TraceServices::ICookProfilerProvider* CookProvider = TraceServices::ReadCookProfilerProvider(*Session.Get());
 
@@ -137,9 +137,9 @@ void SPackageTableTreeView::RebuildTree(bool bResync)
 
 		const uint32 NumPackages = CookProvider->GetNumPackages();
 		Packages.Reserve(NumPackages);
-		TableTreeNodes.Reserve(NumPackages);
+		TableRowNodes.Reserve(NumPackages);
 
-		TArray<FTableTreeNodePtr>* Nodes = &TableTreeNodes;
+		TArray<FTableTreeNodePtr>* Nodes = &TableRowNodes;
 		CookProvider->EnumeratePackages(0, 0, [&Packages, &PackageTable, Nodes](const TraceServices::FPackageData& Package)
 			{
 				Packages.Emplace(Package);
diff --git a/Engine/Source/Developer/TraceInsights/Private/Insights/MemoryProfiler/Widgets/SMemAllocTableTreeView.cpp b/Engine/Source/Developer/TraceInsights/Private/Insights/MemoryProfiler/Widgets/SMemAllocTableTreeView.cpp
--- a/Engine/Source/Developer/TraceInsights/Private/Insights/MemoryProfiler/Widgets/SMemAllocTableTreeView.cpp
+++ b/Engine/Source/Developer/TraceInsights/Private/Insights/MemoryProfiler/Widgets/SMemAllocTableTreeView.cpp
@@ -114,10 +114,10 @@ void SMemAllocTableTreeView::RebuildTree(bool bResync)
 
 	if (bResync)
 	{
-		TableTreeNodes.Empty();
+		TableRowNodes.Empty();
 	}
 
-	const int32 PreviousNodeCount = TableTreeNodes.Num();
+	const int32 PreviousNodeCount = TableRowNodes.Num();
 
 	TSharedPtr<Insights::FMemAllocTable> MemAllocTable = GetMemAllocTable();
 
@@ -132,20 +132,20 @@ void SMemAllocTableTreeView::RebuildTree(bool bResync)
 			TArray<FMemoryAlloc>& Allocs = MemAllocTable->GetAllocs();
 
 			const int32 TotalAllocCount = Allocs.Num();
-			if (TotalAllocCount != TableTreeNodes.Num())
+			if (TotalAllocCount != TableRowNodes.Num())
 			{
-				UE_LOG(MemoryProfiler, Log, TEXT("[MemAlloc] Creating nodes (%d nodes --> %d allocs)..."), TableTreeNodes.Num(), TotalAllocCount);
+				UE_LOG(MemoryProfiler, Log, TEXT("[MemAlloc] Creating nodes (%d nodes --> %d allocs)..."), TableRowNodes.Num(), TotalAllocCount);
 
-				if (TableTreeNodes.Num() > TotalAllocCount)
+				if (TableRowNodes.Num() > TotalAllocCount)
 				{
-					TableTreeNodes.Empty();
+					TableRowNodes.Empty();
 				}
-				TableTreeNodes.Reserve(TotalAllocCount);
+				TableRowNodes.Reserve(TotalAllocCount);
 
 				uint32 HeapAllocCount = 0;
 				const FName BaseNodeName(TEXT("alloc"));
 				const FName BaseHeapName(TEXT("heap"));
-				for (int32 AllocIndex = TableTreeNodes.Num(); AllocIndex < TotalAllocCount; ++AllocIndex)
+				for (int32 AllocIndex = TableRowNodes.Num(); AllocIndex < TotalAllocCount; ++AllocIndex)
 				{
 					const FMemoryAlloc* Alloc = MemAllocTable->GetMemAlloc(AllocIndex);
 
@@ -160,9 +160,9 @@ void SMemAllocTableTreeView::RebuildTree(bool bResync)
 
 					FName NodeName(Alloc->bIsHeap ? BaseHeapName : BaseNodeName, static_cast<int32>(Alloc->GetStartEventIndex() + 1));
 					FMemAllocNodePtr NodePtr = MakeShared<FMemAllocNode>(NodeName, MemAllocTable, AllocIndex);
-					TableTreeNodes.Add(NodePtr);
+					TableRowNodes.Add(NodePtr);
 				}
-				ensure(TableTreeNodes.Num() == (bIncludeHeapAllocs ? TotalAllocCount : TotalAllocCount - HeapAllocCount));
+				ensure(TableRowNodes.Num() == (bIncludeHeapAllocs ? TotalAllocCount : TotalAllocCount - HeapAllocCount));
 				UpdateQueryInfo();
 			}
 		}
@@ -170,7 +170,7 @@ void SMemAllocTableTreeView::RebuildTree(bool bResync)
 
 	SyncStopwatch.Stop();
 
-	if (bResync || TableTreeNodes.Num() != PreviousNodeCount)
+	if (bResync || TableRowNodes.Num() != PreviousNodeCount)
 	{
 		// Save selection.
 		TArray<FTableTreeNodePtr> SelectedItems;
@@ -203,7 +203,7 @@ void SMemAllocTableTreeView::RebuildTree(bool bResync)
 	{
 		const double SyncTime = SyncStopwatch.GetAccumulatedTime();
 		UE_LOG(MemoryProfiler, Log, TEXT("[MemAlloc] Tree view rebuilt in %.4fs (sync: %.4fs + update: %.4fs) --> %d nodes (%d added)"),
-			TotalTime, SyncTime, TotalTime - SyncTime, TableTreeNodes.Num(), TableTreeNodes.Num() - PreviousNodeCount);
+			TotalTime, SyncTime, TotalTime - SyncTime, TableRowNodes.Num(), TableRowNodes.Num() - PreviousNodeCount);
 	}
 }
 
@@ -227,7 +227,7 @@ void SMemAllocTableTreeView::OnQueryInvalidated()
 
 void SMemAllocTableTreeView::ResetAndStartQuery()
 {
-	TableTreeNodes.Reset();
+	TableRowNodes.Reset();
 
 	TSharedPtr<Insights::FMemAllocTable> MemAllocTable = GetMemAllocTable();
 	if (MemAllocTable)
@@ -1257,7 +1257,7 @@ void SMemAllocTableTreeView::UpdateQueryInfo()
 			check(false);
 		}
 
-		QueryInfo = FText::Format(LOCTEXT("QueryInfoFmt", "{0} ({1}) : {2} allocs"), Rule->GetVerboseName(), TimeMarkersText, FText::AsNumber(TableTreeNodes.Num()));
+		QueryInfo = FText::Format(LOCTEXT("QueryInfoFmt", "{0} ({1}) : {2} allocs"), Rule->GetVerboseName(), TimeMarkersText, FText::AsNumber(TableRowNodes.Num()));
 		QueryInfoTooltip = Rule->GetDescription();
 	}
 }
@@ -1840,7 +1840,7 @@ void SMemAllocTableTreeView::ExportMemorySnapshot() const
 
 	// 2. Iterate over TreeNodes
 	TStringBuilder<2048> Buffer;
-	for (const TSharedPtr<FTableTreeNode>& Node : TableTreeNodes)
+	for (const TSharedPtr<FTableTreeNode>& Node : TableRowNodes)
 	{
 		// Export only leaves
 		if (Node->IsGroup()) continue;
@@ -1897,7 +1897,7 @@ void SMemAllocTableTreeView::ExportMemorySnapshot() const
 
 bool SMemAllocTableTreeView::IsExportMemorySnapshotAvailable() const
 {
-	return !TableTreeNodes.IsEmpty();
+	return !TableRowNodes.IsEmpty();
 }
 
 ////////////////////////////////////////////////////////////////////////////////////////////////////
diff --git a/Engine/Source/Developer/TraceInsights/Private/Insights/Table/ViewModels/TableTreeNode.h b/Engine/Source/Developer/TraceInsights/Private/Insights/Table/ViewModels/TableTreeNode.h
--- a/Engine/Source/Developer/TraceInsights/Private/Insights/Table/ViewModels/TableTreeNode.h
+++ b/Engine/Source/Developer/TraceInsights/Private/Insights/Table/ViewModels/TableTreeNode.h
@@ -126,16 +126,10 @@ public:
 		return AggregatedValues->FindChecked(ColumnId);
 	}
 
-	void AddAggregatedValue(const FName& ColumnId, const FTableCellValue& Value)
-	{
-		InitAggregatedValues();
-		AggregatedValues->Add(ColumnId, Value);
-	}
-
 	void SetAggregatedValue(const FName& ColumnId, const FTableCellValue& Value)
 	{
 		InitAggregatedValues();
-		(*AggregatedValues)[ColumnId] = Value;
+		AggregatedValues->Add(ColumnId, Value);
 	}
 
 	//////////////////////////////////////////////////
diff --git a/Engine/Source/Developer/TraceInsights/Private/Insights/Table/Widgets/STableTreeView.cpp b/Engine/Source/Developer/TraceInsights/Private/Insights/Table/Widgets/STableTreeView.cpp
--- a/Engine/Source/Developer/TraceInsights/Private/Insights/Table/Widgets/STableTreeView.cpp
+++ b/Engine/Source/Developer/TraceInsights/Private/Insights/Table/Widgets/STableTreeView.cpp
@@ -135,6 +135,7 @@ const FName STableTreeView::RootNodeName(TEXT("Root"));
 
 STableTreeView::STableTreeView()
 	: Root(MakeShared<FTableTreeNode>(RootNodeName, Table))
+	, FilteredNodesPtr(&TableRowNodes)
 	, ColumnBeingSorted(GetDefaultColumnBeingSorted())
 	, ColumnSortMode(GetDefaultColumnSortMode())
 {
@@ -919,7 +920,7 @@ void STableTreeView::Tick(const FGeometry& AllottedGeometry, const double InCurr
 
 void STableTreeView::UpdateTree()
 {
-	OnGroupingChanged();
+	OnNodeFilteringChanged();
 }
 
 ////////////////////////////////////////////////////////////////////////////////////////////////////
@@ -941,6 +942,9 @@ void STableTreeView::OnNodeFilteringChanged()
 	if (!bRunInAsyncMode)
 	{
 		ApplyNodeFiltering();
+		ApplyGrouping();
+		ApplySorting();
+		ApplyHierarchyFiltering();
 	}
 	else if (!bIsUpdateRunning)
 	{
@@ -956,7 +960,7 @@ void STableTreeView::OnNodeFilteringChanged()
 
 bool STableTreeView::ScheduleNodeFilteringAsyncOperationIfNeeded()
 {
-	if (HasInProgressAsyncOperation(EAsyncOperationType::FilteringOp) ||
+	if (HasInProgressAsyncOperation(EAsyncOperationType::NodeFiltering) ||
 		(FilterConfigurator.IsValid() &&
 		 CurrentAsyncOpFilterConfigurator &&
 		 *FilterConfigurator != *CurrentAsyncOpFilterConfigurator))
@@ -974,16 +978,17 @@ void STableTreeView::ScheduleNodeFilteringAsyncOperation()
 {
 	OnPreAsyncUpdate();
 
-	InProgressAsyncOperationEvent = StartApplyFiltersTask();
+	FGraphEventRef CompletedEvent = StartNodeFilteringTask();
+	CompletedEvent = StartGroupingTask(CompletedEvent);
+	CompletedEvent = StartSortingTask(CompletedEvent);
+	InProgressAsyncOperationEvent = StartHierarchyFilteringTask(CompletedEvent);
 }
 
 ////////////////////////////////////////////////////////////////////////////////////////////////////
 
-FGraphEventRef STableTreeView::StartApplyFiltersTask(FGraphEventRef Prerequisite)
+FGraphEventRef STableTreeView::StartNodeFilteringTask(FGraphEventRef Prerequisite)
 {
-	AddInProgressAsyncOperation(EAsyncOperationType::FilteringOp);
-
-	CurrentAsyncOpTextFilter->SetRawFilterText(TextFilter->GetRawFilterText());
+	AddInProgressAsyncOperation(EAsyncOperationType::NodeFiltering);
 
 	if (FilterConfigurator.IsValid() && CurrentAsyncOpFilterConfigurator)
 	{
@@ -1000,79 +1005,66 @@ FGraphEventRef STableTreeView::StartApplyFiltersTask(FGraphEventRef Prerequisite
 		Prerequisites.Add(DispatchEvent);
 	}
 
-	return TGraphTask<FTableTreeViewFilterAsyncTask>::CreateTask(&Prerequisites).ConstructAndDispatchWhenReady(this);
+	return TGraphTask<FTableTreeViewNodeFilteringAsyncTask>::CreateTask(&Prerequisites).ConstructAndDispatchWhenReady(this);
 }
 
 ////////////////////////////////////////////////////////////////////////////////////////////////////
 
+// Note: This function is called from a task thread!
 void STableTreeView::ApplyNodeFiltering()
 {
-	ApplyHierarchyFiltering();
-}
+	FStopwatch Stopwatch;
+	Stopwatch.Start();
 
-////////////////////////////////////////////////////////////////////////////////////////////////////
+	FFilterConfigurator* FilterConfiguratorToUse = bRunInAsyncMode ? CurrentAsyncOpFilterConfigurator : FilterConfigurator.Get();
 
-bool STableTreeView::ApplyAdvancedFiltersForNode(FTableTreeNodePtr NodePtr)
-{
-	if (NodePtr->IsGroup())
+	// TableRowNodes --> FilteredNodes / FilteredNodesPtr
+	FilteredNodes.Reset();
+	if (FilterConfiguratorToUse && !FilterConfiguratorToUse->IsEmpty())
 	{
-		const TArray<FBaseTreeNodePtr>& GroupChildren = NodePtr->GetChildren();
-		const int32 NumChildren = GroupChildren.Num();
-
-		NodePtr->ClearFilteredChildren();
-
-		int32 NumVisibleChildren = 0;
-		for (int32 Cx = 0; Cx < NumChildren; ++Cx)
+		for (const FTableTreeNodePtr& NodePtr : TableRowNodes)
 		{
 			if (AsyncOperationProgress.ShouldCancelAsyncOp())
 			{
 				break;
 			}
 
-			// Add a child.
-			const FTableTreeNodePtr& ChildNodePtr = StaticCastSharedPtr<FTableTreeNode>(GroupChildren[Cx]);
-			if (ApplyAdvancedFiltersForNode(ChildNodePtr))
+			if (FilterNode(*FilterConfiguratorToUse, *NodePtr))
 			{
-				NodePtr->AddFilteredChild(ChildNodePtr);
-				NumVisibleChildren++;
+				FilteredNodes.Add(NodePtr);
 			}
 		}
-
-		const bool bIsGroupNodeVisible = NumVisibleChildren > 0;
-
-		if (bIsGroupNodeVisible)
-		{
-			// Add a group.
-			NodePtr->SetExpansion(true);
-		}
-		else
-		{
-			NodePtr->SetExpansion(false);
-		}
-
-		return bIsGroupNodeVisible;
+		FilteredNodesPtr = &FilteredNodes;
 	}
 	else
 	{
-		bool bIsNodeVisible = FilterNode(NodePtr);
-		NodePtr->SetIsFiltered(!bIsNodeVisible);
-		return bIsNodeVisible;
+		FilteredNodesPtr = &TableRowNodes;
 	}
-}
 
-////////////////////////////////////////////////////////////////////////////////////////////////////
-
-bool STableTreeView::FilterNode(const FTableTreeNodePtr& NodePtr) const
-{
-	FFilterConfigurator* FilterConfiguratorToUse = bRunInAsyncMode ? CurrentAsyncOpFilterConfigurator : FilterConfigurator.Get();
+	TreeViewBannerText = FText::GetEmpty();
+	if (TableRowNodes.Num() == 0)
+	{
+		TreeViewBannerText = LOCTEXT("EmptyTreeMsg", "Tree is empty.");
+	}
+	else if (FilteredNodesPtr->Num() == 0)
+	{
+		TreeViewBannerText = LOCTEXT("AllNodesFiltered", "All tree nodes are filtered out.");
+	}
 
-	if (!FilterConfiguratorToUse || FilterConfiguratorToUse->IsEmpty())
+	Stopwatch.Stop();
+	const double FilteringTime = Stopwatch.GetAccumulatedTime();
+	if (FilteringTime > 0.1)
 	{
-		return true;
+		UE_LOG(TraceInsights, Log, TEXT("[Tree - %s] Node filtering completed in %.3fs."), *Table->GetDisplayName().ToString(), FilteringTime);
 	}
+}
+
+////////////////////////////////////////////////////////////////////////////////////////////////////
 
-	UpdateFilterContext(*FilterConfiguratorToUse, *NodePtr);
-	return FilterConfiguratorToUse->ApplyFilters(FilterContext);
+bool STableTreeView::FilterNode(const FFilterConfigurator& InFilterConfigurator, const FTableTreeNode& InNode) const
+{
+	UpdateFilterContext(InFilterConfigurator, InNode);
+	return InFilterConfigurator.ApplyFilters(FilterContext);
 }
 
 ////////////////////////////////////////////////////////////////////////////////////////////////////
@@ -1189,8 +1181,7 @@ TSharedRef<SWidget> STableTreeView::ConstructFilterConfiguratorButton()
 {
 	return SNew(SButton)
 		.ContentPadding(FMargin(-2.0f, 2.0f, -2.0f, 2.0f))
-		.IsEnabled(this, &STableTreeView::FilterConfigurator_IsEnabled)
-		.ToolTipText(this, &STableTreeView::FilterConfigurator_GetTooltipText)
+		.ToolTipText(LOCTEXT("FilterConfiguratorBtn_ToolTip", "Opens the filter configurator window."))
 		.OnClicked(this, &STableTreeView::FilterConfigurator_OnClicked)
 		[
 			SNew(SImage)
@@ -1201,32 +1192,6 @@ END_SLATE_FUNCTION_BUILD_OPTIMIZATION
 
 ////////////////////////////////////////////////////////////////////////////////////////////////////
 
-bool STableTreeView::FilterConfigurator_IsEnabled() const
-{
-	return TextFilter->GetRawFilterText().IsEmpty();
-}
-
-////////////////////////////////////////////////////////////////////////////////////////////////////
-
-FText STableTreeView::FilterConfigurator_GetTooltipText() const
-{
-	if (FilterConfigurator_IsEnabled())
-	{
-		return LOCTEXT("FilterConfiguratorBtn_ToolTip", "Opens the filter configurator window.");
-	}
-
-	return LOCTEXT("AdvancedFiltersBtn_Disabled_ToolTip", "Advanced filters cannot be added when filters are already applied using the search box.");
-}
-
-////////////////////////////////////////////////////////////////////////////////////////////////////
-
-bool STableTreeView::FilterConfigurator_HasFilters() const
-{
-	return FilterConfigurator.IsValid() && !FilterConfigurator->IsEmpty();
-}
-
-////////////////////////////////////////////////////////////////////////////////////////////////////
-
 FReply STableTreeView::FilterConfigurator_OnClicked()
 {
 	if (!FilterConfigurator.IsValid())
@@ -1297,7 +1262,7 @@ void STableTreeView::OnHierarchyFilteringChanged()
 
 bool STableTreeView::ScheduleHierarchyFilteringAsyncOperationIfNeeded()
 {
-	if (HasInProgressAsyncOperation(EAsyncOperationType::FilteringOp) ||
+	if (HasInProgressAsyncOperation(EAsyncOperationType::HierarchyFiltering) ||
 		TextFilter->GetRawFilterText().CompareTo(CurrentAsyncOpTextFilter->GetRawFilterText()) != 0)
 	{
 		ScheduleHierarchyFilteringAsyncOperation();
@@ -1313,32 +1278,46 @@ void STableTreeView::ScheduleHierarchyFilteringAsyncOperation()
 {
 	OnPreAsyncUpdate();
 
-	InProgressAsyncOperationEvent = StartApplyFiltersTask();
+	InProgressAsyncOperationEvent = StartHierarchyFilteringTask();
 }
 
 ////////////////////////////////////////////////////////////////////////////////////////////////////
 
+FGraphEventRef STableTreeView::StartHierarchyFilteringTask(FGraphEventRef Prerequisite)
+{
+	AddInProgressAsyncOperation(EAsyncOperationType::HierarchyFiltering);
+
+	CurrentAsyncOpTextFilter->SetRawFilterText(TextFilter->GetRawFilterText());
+
+	FGraphEventArray Prerequisites;
+	if (Prerequisite.IsValid())
+	{
+		Prerequisites.Add(Prerequisite);
+	}
+	else
+	{
+		Prerequisites.Add(DispatchEvent);
+	}
+
+	return TGraphTask<FTableTreeViewHierarchyFilteringAsyncTask>::CreateTask(&Prerequisites).ConstructAndDispatchWhenReady(this);
+}
+
+////////////////////////////////////////////////////////////////////////////////////////////////////
+
+// Note: This function is called from a task thread!
 void STableTreeView::ApplyHierarchyFiltering()
 {
 	FStopwatch Stopwatch;
 	Stopwatch.Start();
 
-	// Apply filter to all groups and its children.
-	if (FilterConfigurator_HasFilters())
+	const bool bIsEmptyFilter = bRunInAsyncMode ? CurrentAsyncOpTextFilter->GetRawFilterText().IsEmpty() : TextFilter->GetRawFilterText().IsEmpty();
+	if (bIsEmptyFilter)
 	{
-		ApplyAdvancedFiltersForNode(Root);
+		ApplyEmptyHierarchyFilteringRec(Root);
 	}
 	else
 	{
-		const bool bIsEmptyFilter = bRunInAsyncMode ? CurrentAsyncOpTextFilter->GetRawFilterText().IsEmpty() : TextFilter->GetRawFilterText().IsEmpty();
-		if (bIsEmptyFilter)
-		{
-			ApplyEmptyHierarchyFilteringRec(Root);
-		}
-		else
-		{
-			ApplyHierarchyFilteringRec(Root);
-		}
+		ApplyHierarchyFilteringRec(Root);
 	}
 
 	// The Root node is always hidden. The tree shows the filtered children of the Root node.
@@ -1354,25 +1333,6 @@ void STableTreeView::ApplyHierarchyFiltering()
 		}
 	}
 
-	TreeViewBannerText = FText::GetEmpty();
-	if (Root->GetChildrenCount() == 0)
-	{
-		TreeViewBannerText = LOCTEXT("EmptyTreeMsg", "Tree is empty.");
-	}
-	else if (Root->GetFilteredChildrenCount() == 0)
-	{
-		TreeViewBannerText = LOCTEXT("AllNodesFiltered", "All tree nodes are filtered out.");
-	}
-
-	Stopwatch.Stop();
-	const double FilteringTime = Stopwatch.GetAccumulatedTime();
-	if (FilteringTime > 0.1)
-	{
-		UE_LOG(TraceInsights, Log, TEXT("[Tree - %s] Filtering completed in %.3fs."), *Table->GetDisplayName().ToString(), FilteringTime);
-	}
-
-	UpdateAggregatedValues(*Root);
-
 	// Cannot call TreeView functions from other threads than MainThread and SlateThread.
 	if (!bRunInAsyncMode)
 	{
@@ -1410,6 +1370,13 @@ void STableTreeView::ApplyHierarchyFiltering()
 		// Request tree refresh
 		TreeView->RequestTreeRefresh();
 	}
+
+	Stopwatch.Stop();
+	const double FilteringTime = Stopwatch.GetAccumulatedTime();
+	if (FilteringTime > 0.1)
+	{
+		UE_LOG(TraceInsights, Log, TEXT("[Tree - %s] Hierarchy filtering completed in %.3fs."), *Table->GetDisplayName().ToString(), FilteringTime);
+	}
 }
 
 ////////////////////////////////////////////////////////////////////////////////////////////////////
@@ -1540,9 +1507,8 @@ TSharedRef<SWidget> STableTreeView::ConstructSearchBox()
 {
 	SAssignNew(SearchBox, SSearchBox)
 		.HintText(LOCTEXT("SearchBox_Hint", "Search"))
-		.OnTextChanged(this, &STableTreeView::SearchBox_OnTextChanged)
-		.IsEnabled(this, &STableTreeView::SearchBox_IsEnabled)
-		.ToolTipText(this, &STableTreeView::SearchBox_GetTooltipText);
+		.ToolTipText(LOCTEXT("SearchBox_ToolTip", "Type here to search the tree hierarchy by item or group name."))
+		.OnTextChanged(this, &STableTreeView::SearchBox_OnTextChanged);
 
 	return SearchBox.ToSharedRef();
 }
@@ -1557,25 +1523,6 @@ void STableTreeView::SearchBox_OnTextChanged(const FText& InFilterText)
 	OnHierarchyFilteringChanged();
 }
 
-////////////////////////////////////////////////////////////////////////////////////////////////////
-
-bool STableTreeView::SearchBox_IsEnabled() const
-{
-	return !FilterConfigurator_HasFilters();
-}
-
-////////////////////////////////////////////////////////////////////////////////////////////////////
-
-FText STableTreeView::SearchBox_GetTooltipText() const
-{
-	if (SearchBox_IsEnabled())
-	{
-		return LOCTEXT("SearchBox_ToolTip", "Type here to search the tree hierarchy by item or group name.");
-	}
-
-	return LOCTEXT("SearchBox_Disabled_ToolTip", "Searching the tree hierarchy is disabled when advanced filters are set.");
-}
-
 ////////////////////////////////////////////////////////////////////////////////////////////////////
 // TreeView
 ////////////////////////////////////////////////////////////////////////////////////////////////////
@@ -1604,6 +1551,30 @@ void STableTreeView::TreeView_OnExpansionChanged(FTableTreeNodePtr TreeNode, boo
 
 void STableTreeView::TreeView_OnGetChildren(FTableTreeNodePtr InParent, TArray<FTableTreeNodePtr>& OutChildren)
 {
+	if (InParent->OnLazyCreateChildren(SharedThis(this)))
+	{
+		// Node filtering does not apply to lazy expanded nodes.
+		// Grouping is ignored for lazy expanded nodes.
+
+		// Update aggregation.
+		UpdateAggregatedValuesRec(*InParent);
+		FTableTreeNodePtr Node = StaticCastSharedPtr<FTableTreeNode>(InParent->GetParentNode());
+		while (Node)
+		{
+			UpdateAggregatedValuesSingleNode(*Node);
+			Node = StaticCastSharedPtr<FTableTreeNode>(Node->GetParentNode());
+		}
+
+		// Update sorting.
+		if (CurrentSorter.IsValid())
+		{
+			SortTreeNodesRec(*InParent, *CurrentSorter, ColumnSortMode);
+		}
+
+		// The hierarchy filtering is ignored for lazy expanded nodes.
+		InParent->ResetFilteredChildrenRec();
+	}
+
 	const TArray<FBaseTreeNodePtr>& FilteredChildren = InParent->GetFilteredChildren();
 	for (const FBaseTreeNodePtr& NodePtr : FilteredChildren)
 	{
@@ -1724,7 +1695,7 @@ void STableTreeView::OnGroupingChanged()
 
 bool STableTreeView::ScheduleGroupingAsyncOperationIfNeeded()
 {
-	bool bGroupingHasChanged = HasInProgressAsyncOperation(EAsyncOperationType::GroupingOp);
+	bool bGroupingHasChanged = HasInProgressAsyncOperation(EAsyncOperationType::Grouping);
 
 	if (!bGroupingHasChanged &&
 		CurrentGroupings.Num() != CurrentAsyncOpGroupings.Num())
@@ -1761,17 +1732,16 @@ void STableTreeView::ScheduleGroupingAsyncOperation()
 
 	FGraphEventRef CompletedEvent = StartGroupingTask();
 	CompletedEvent = StartSortingTask(CompletedEvent);
-	InProgressAsyncOperationEvent = StartApplyFiltersTask(CompletedEvent);
+	InProgressAsyncOperationEvent = StartHierarchyFilteringTask(CompletedEvent);
 }
 
 ////////////////////////////////////////////////////////////////////////////////////////////////////
 
 FGraphEventRef STableTreeView::StartGroupingTask(FGraphEventRef Prerequisite)
 {
-	AddInProgressAsyncOperation(EAsyncOperationType::GroupingOp);
+	AddInProgressAsyncOperation(EAsyncOperationType::Grouping);
 
-	CurrentAsyncOpGroupings.Empty();
-	CurrentAsyncOpGroupings.Insert(CurrentGroupings, 0);
+	CurrentAsyncOpGroupings = CurrentGroupings;
 
 	FGraphEventArray Prerequisites;
 	if (Prerequisite.IsValid())
@@ -1783,7 +1753,7 @@ FGraphEventRef STableTreeView::StartGroupingTask(FGraphEventRef Prerequisite)
 		Prerequisites.Add(DispatchEvent);
 	}
 
-	return TGraphTask<FTableTreeViewGroupAsyncTask>::CreateTask(&Prerequisites).ConstructAndDispatchWhenReady(this, &CurrentAsyncOpGroupings);
+	return TGraphTask<FTableTreeViewGroupingAsyncTask>::CreateTask(&Prerequisites).ConstructAndDispatchWhenReady(this, &CurrentAsyncOpGroupings);
 }
 
 ////////////////////////////////////////////////////////////////////////////////////////////////////
@@ -1795,10 +1765,10 @@ void STableTreeView::ApplyGrouping()
 
 ////////////////////////////////////////////////////////////////////////////////////////////////////
 
-// Note: This function is also called from async thread, with CurrentAsyncOpGroupings as parameter.
+// Note: This function is called from a task thread!
 void STableTreeView::CreateGroups(const TArray<TSharedPtr<FTreeNodeGrouping>>& Groupings)
 {
-	if (TableTreeNodes.Num() == 0)
+	if (FilteredNodesPtr->Num() == 0)
 	{
 		Root->ClearChildren();
 		return;
@@ -1809,7 +1779,7 @@ void STableTreeView::CreateGroups(const TArray<TSharedPtr<FTreeNodeGrouping>>& G
 
 	if (Groupings.Num() > 0)
 	{
-		GroupNodesRec(TableTreeNodes, *Root, 0, Groupings);
+		GroupNodesRec(*FilteredNodesPtr, *Root, 0, Groupings);
 	}
 
 	Stopwatch.Stop();
@@ -1819,7 +1789,7 @@ void STableTreeView::CreateGroups(const TArray<TSharedPtr<FTreeNodeGrouping>>& G
 		UE_LOG(TraceInsights, Log, TEXT("[Tree - %s] Grouping completed in %.3fs."), *Table->GetDisplayName().ToString(), GroupingTime);
 	}
 
-	UpdateAggregatedValues(*Root);
+	UpdateAggregatedValuesRec(*Root);
 }
 
 ////////////////////////////////////////////////////////////////////////////////////////////////////
@@ -2473,35 +2443,18 @@ bool STableTreeView::GroupingCrumbMenu_Add_CanExecute(const TSharedPtr<FTreeNode
 // Aggregation
 ////////////////////////////////////////////////////////////////////////////////////////////////////
 
-void STableTreeView::UpdateCStringSameValueAggregationRec(FTableColumn& Column, FTableTreeNode& GroupNode)
+void STableTreeView::UpdateCStringSameValueAggregationSingleNode(const FTableColumn& InColumn, FTableTreeNode& InOutGroupNode)
 {
-	// Update child group nodes first.
-	for (FBaseTreeNodePtr NodePtr : GroupNode.GetChildren())
-	{
-		if (NodePtr->IsFiltered())
-		{
-			continue;
-		}
-		if (NodePtr->IsGroup())
-		{
-			FTableTreeNode& TableNode = *(FTableTreeNode*)NodePtr.Get();
-			UpdateCStringSameValueAggregationRec(Column, TableNode);
-		}
-	}
+	check(InOutGroupNode.IsGroup());
 
 	const TCHAR* AggregatedValue = nullptr;
 
 	// Find the first child node.
-	for (FBaseTreeNodePtr NodePtr : GroupNode.GetChildren())
+	for (FBaseTreeNodePtr NodePtr : InOutGroupNode.GetChildren())
 	{
-		if (NodePtr->IsFiltered())
-		{
-			continue;
-		}
-
 		if (!NodePtr->IsGroup())
 		{
-			const TOptional<FTableCellValue> NodeValue = Column.GetValue(*NodePtr);
+			const TOptional<FTableCellValue> NodeValue = InColumn.GetValue(*NodePtr);
 			if (NodeValue.IsSet() &&
 				NodeValue.GetValue().DataType == ETableCellDataType::CString)
 			{
@@ -2511,9 +2464,9 @@ void STableTreeView::UpdateCStringSameValueAggregationRec(FTableColumn& Column,
 		else
 		{
 			FTableTreeNode& TableNode = *(FTableTreeNode*)NodePtr.Get();
-			if (TableNode.HasAggregatedValue(Column.GetId()))
+			if (TableNode.HasAggregatedValue(InColumn.GetId()))
 			{
-				const FTableCellValue& ChildGroupAggregatedValue = TableNode.GetAggregatedValue(Column.GetId());
+				const FTableCellValue& ChildGroupAggregatedValue = TableNode.GetAggregatedValue(InColumn.GetId());
 				if (ChildGroupAggregatedValue.DataType == ETableCellDataType::CString)
 				{
 					AggregatedValue = ChildGroupAggregatedValue.CString;
@@ -2527,16 +2480,11 @@ void STableTreeView::UpdateCStringSameValueAggregationRec(FTableColumn& Column,
 	if (AggregatedValue != nullptr)
 	{
 		// Check if all other children have the same value as the first node.
-		for (FBaseTreeNodePtr NodePtr : GroupNode.GetChildren())
+		for (FBaseTreeNodePtr NodePtr : InOutGroupNode.GetChildren())
 		{
-			if (NodePtr->IsFiltered())
-			{
-				continue;
-			}
-
 			if (!NodePtr->IsGroup())
 			{
-				const TOptional<FTableCellValue> NodeValue = Column.GetValue(*NodePtr);
+				const TOptional<FTableCellValue> NodeValue = InColumn.GetValue(*NodePtr);
 				if (NodeValue.IsSet() &&
 					NodeValue.GetValue().DataType == ETableCellDataType::CString &&
 					AggregatedValue == NodeValue.GetValue().CString)
@@ -2547,9 +2495,9 @@ void STableTreeView::UpdateCStringSameValueAggregationRec(FTableColumn& Column,
 			else
 			{
 				FTableTreeNode& TableNode = *(FTableTreeNode*)NodePtr.Get();
-				if (TableNode.HasAggregatedValue(Column.GetId()))
+				if (TableNode.HasAggregatedValue(InColumn.GetId()))
 				{
-					const FTableCellValue& ChildGroupAggregatedValue = TableNode.GetAggregatedValue(Column.GetId());
+					const FTableCellValue& ChildGroupAggregatedValue = TableNode.GetAggregatedValue(InColumn.GetId());
 					if (ChildGroupAggregatedValue.DataType == ETableCellDataType::CString &&
 						AggregatedValue == ChildGroupAggregatedValue.CString)
 					{
@@ -2565,22 +2513,82 @@ void STableTreeView::UpdateCStringSameValueAggregationRec(FTableColumn& Column,
 
 	if (AggregatedValue != nullptr)
 	{
-		GroupNode.AddAggregatedValue(Column.GetId(), FTableCellValue(AggregatedValue));
+		InOutGroupNode.SetAggregatedValue(InColumn.GetId(), FTableCellValue(AggregatedValue));
 	}
 	else
 	{
-		GroupNode.ResetAggregatedValue(Column.GetId());
+		InOutGroupNode.ResetAggregatedValue(InColumn.GetId());
+	}
+}
+
+////////////////////////////////////////////////////////////////////////////////////////////////////
+
+void STableTreeView::UpdateCStringSameValueAggregationRec(const FTableColumn& InColumn, FTableTreeNode& InOutGroupNode)
+{
+	check(InOutGroupNode.IsGroup());
+
+	// Update child group nodes first.
+	for (FBaseTreeNodePtr NodePtr : InOutGroupNode.GetChildren())
+	{
+		if (NodePtr->IsGroup())
+		{
+			FTableTreeNode& TableNode = *(FTableTreeNode*)NodePtr.Get();
+			STableTreeView::UpdateCStringSameValueAggregationRec(InColumn, TableNode);
+		}
 	}
+
+	STableTreeView::UpdateCStringSameValueAggregationSingleNode(InColumn, InOutGroupNode);
 }
 
 ////////////////////////////////////////////////////////////////////////////////////////////////////
 
-void STableTreeView::UpdateAggregatedValues(FTableTreeNode& GroupNode)
+template<typename T, bool bSetInitialValue, bool bIsRercursive>
+void STableTreeView::UpdateAggregation(const FTableColumn& InColumn, FTableTreeNode& InOutGroupNode, const T InitialAggregatedValue, TFunctionRef<T(T, const FTableCellValue&)> ValueGetterFunc)
+{
+	T AggregatedValue = InitialAggregatedValue;
+
+	for (FBaseTreeNodePtr NodePtr : InOutGroupNode.GetChildren())
+	{
+		if (NodePtr->IsGroup())
+		{
+			FTableTreeNode& TableNode = *(FTableTreeNode*)NodePtr.Get();
+			if (bIsRercursive)
+			{
+				TableNode.ResetAggregatedValue(InColumn.GetId());
+				STableTreeView::UpdateAggregation<T, bSetInitialValue, true>(InColumn, TableNode, InitialAggregatedValue, ValueGetterFunc);
+			}
+			if (TableNode.HasAggregatedValue(InColumn.GetId()))
+			{
+				AggregatedValue = ValueGetterFunc(AggregatedValue, TableNode.GetAggregatedValue(InColumn.GetId()));
+			}
+		}
+		else
+		{
+			const TOptional<FTableCellValue> NodeValue = InColumn.GetValue(*NodePtr);
+			if (NodeValue.IsSet())
+			{
+				AggregatedValue = ValueGetterFunc(AggregatedValue, NodeValue.GetValue());
+			}
+		}
+	}
+
+	if (bSetInitialValue || AggregatedValue != InitialAggregatedValue)
+	{
+		InOutGroupNode.SetAggregatedValue(InColumn.GetId(), FTableCellValue(AggregatedValue));
+	}
+}
+
+////////////////////////////////////////////////////////////////////////////////////////////////////
+
+template<bool bIsRecursive>
+void STableTreeView::UpdateAggregatedValues(TSharedPtr<FTable> InTable, FTableTreeNode& InOutGroupNode)
 {
 	FStopwatch Stopwatch;
 	Stopwatch.Start();
 
-	for (const TSharedRef<FTableColumn>& ColumnRef : Table->GetColumns())
+	check(InOutGroupNode.IsGroup());
+
+	for (const TSharedRef<FTableColumn>& ColumnRef : InTable->GetColumns())
 	{
 		FTableColumn& Column = ColumnRef.Get();
 		switch (Column.GetAggregation())
@@ -2588,7 +2596,7 @@ void STableTreeView::UpdateAggregatedValues(FTableTreeNode& GroupNode)
 			case ETableColumnAggregation::Sum:
 				if (Column.GetDataType() == ETableCellDataType::Float || Column.GetDataType() == ETableCellDataType::Double)
 				{
-					STableTreeView::UpdateAggregationRec<double>(Column, GroupNode, 0, true,
+					STableTreeView::UpdateAggregation<double, true, bIsRecursive>(Column, InOutGroupNode, 0.0,
 						[](double InValue, const FTableCellValue& InTableCellValue)
 						{
 							return InValue + InTableCellValue.AsDouble();
@@ -2596,7 +2604,7 @@ void STableTreeView::UpdateAggregatedValues(FTableTreeNode& GroupNode)
 				}
 				else
 				{
-					STableTreeView::UpdateAggregationRec<int64>(Column, GroupNode, 0, true,
+					STableTreeView::UpdateAggregation<int64, true, bIsRecursive>(Column, InOutGroupNode, (int64)0,
 						[](int64 InValue, const FTableCellValue& InTableCellValue)
 						{
 							return InValue + InTableCellValue.AsInt64();
@@ -2607,7 +2615,7 @@ void STableTreeView::UpdateAggregatedValues(FTableTreeNode& GroupNode)
 			case ETableColumnAggregation::Min:
 				if (Column.GetDataType() == ETableCellDataType::Float || Column.GetDataType() == ETableCellDataType::Double)
 				{
-					STableTreeView::UpdateAggregationRec<double>(Column, GroupNode, std::numeric_limits<double>::max(), false,
+					STableTreeView::UpdateAggregation<double, false, bIsRecursive>(Column, InOutGroupNode, std::numeric_limits<double>::max(),
 						[](double InValue, const FTableCellValue& InTableCellValue)
 						{
 							return FMath::Min(InValue, InTableCellValue.AsDouble());
@@ -2615,7 +2623,7 @@ void STableTreeView::UpdateAggregatedValues(FTableTreeNode& GroupNode)
 				}
 				else
 				{
-					STableTreeView::UpdateAggregationRec<int64>(Column, GroupNode, std::numeric_limits<int64>::max(), false,
+					STableTreeView::UpdateAggregation<int64, false, bIsRecursive>(Column, InOutGroupNode, std::numeric_limits<int64>::max(),
 						[](int64 InValue, const FTableCellValue& InTableCellValue)
 						{
 							return FMath::Min(InValue, InTableCellValue.AsInt64());
@@ -2626,7 +2634,7 @@ void STableTreeView::UpdateAggregatedValues(FTableTreeNode& GroupNode)
 			case ETableColumnAggregation::Max:
 				if (Column.GetDataType() == ETableCellDataType::Float || Column.GetDataType() == ETableCellDataType::Double)
 				{
-					STableTreeView::UpdateAggregationRec<double>(Column, GroupNode, std::numeric_limits<double>::lowest(), false,
+					STableTreeView::UpdateAggregation<double, false, bIsRecursive>(Column, InOutGroupNode, std::numeric_limits<double>::lowest(),
 						[](double InValue, const FTableCellValue& InTableCellValue)
 						{
 							return FMath::Max(InValue, InTableCellValue.AsDouble());
@@ -2634,7 +2642,7 @@ void STableTreeView::UpdateAggregatedValues(FTableTreeNode& GroupNode)
 				}
 				else
 				{
-					STableTreeView::UpdateAggregationRec<int64>(Column, GroupNode, std::numeric_limits<int64>::min(), false,
+					STableTreeView::UpdateAggregation<int64, false, bIsRecursive>(Column, InOutGroupNode, std::numeric_limits<int64>::min(),
 						[](int64 InValue, const FTableCellValue& InTableCellValue)
 						{
 							return FMath::Max(InValue, InTableCellValue.AsInt64());
@@ -2645,7 +2653,14 @@ void STableTreeView::UpdateAggregatedValues(FTableTreeNode& GroupNode)
 			case ETableColumnAggregation::SameValue:
 				if (Column.GetDataType() == ETableCellDataType::CString)
 				{
-					STableTreeView::UpdateCStringSameValueAggregationRec(Column, GroupNode);
+					if (bIsRecursive)
+					{
+						STableTreeView::UpdateCStringSameValueAggregationRec(Column, InOutGroupNode);
+					}
+					else
+					{
+						STableTreeView::UpdateCStringSameValueAggregationSingleNode(Column, InOutGroupNode);
+					}
 				}
 				break;
 		}
@@ -2655,10 +2670,24 @@ void STableTreeView::UpdateAggregatedValues(FTableTreeNode& GroupNode)
 	const double AggregationTime = Stopwatch.GetAccumulatedTime();
 	if (AggregationTime > 0.1)
 	{
-		UE_LOG(TraceInsights, Log, TEXT("[Tree - %s] Aggregation completed in %.3fs."), *Table->GetDisplayName().ToString(), AggregationTime);
+		UE_LOG(TraceInsights, Log, TEXT("[Tree - %s] Aggregation completed in %.3fs."), *InTable->GetDisplayName().ToString(), AggregationTime);
 	}
 }
 
+////////////////////////////////////////////////////////////////////////////////////////////////////
+
+void STableTreeView::UpdateAggregatedValuesSingleNode(FTableTreeNode& InOutGroupNode)
+{
+	STableTreeView::UpdateAggregatedValues<false>(Table, InOutGroupNode);
+}
+
+////////////////////////////////////////////////////////////////////////////////////////////////////
+
+void STableTreeView::UpdateAggregatedValuesRec(FTableTreeNode& InOutGroupNode)
+{
+	STableTreeView::UpdateAggregatedValues<true>(Table, InOutGroupNode);
+}
+
 ////////////////////////////////////////////////////////////////////////////////////////////////////
 // Sorting
 ////////////////////////////////////////////////////////////////////////////////////////////////////
@@ -2729,18 +2758,17 @@ void STableTreeView::OnSortingChanged()
 
 bool STableTreeView::ScheduleSortingAsyncOperationIfNeeded()
 {
-	bool bSortingHasChanged = HasInProgressAsyncOperation(EAsyncOperationType::SortingOp);
+	bool bSortingHasChanged = HasInProgressAsyncOperation(EAsyncOperationType::Sorting);
 
 	if (!bSortingHasChanged &&
-		((CurrentSorter.IsValid() && CurrentAsyncOpSorter == nullptr) ||
-		(!CurrentSorter.IsValid() && CurrentAsyncOpSorter != nullptr)))
+		CurrentSorter.Get() != CurrentAsyncOpSorter)
 	{
 		bSortingHasChanged = true;
 	}
 
 	if (!bSortingHasChanged &&
 		CurrentSorter.IsValid() &&
-		(CurrentSorter.Get() != CurrentAsyncOpSorter || ColumnSortMode != CurrentAsyncOpColumnSortMode))
+		ColumnSortMode != CurrentAsyncOpColumnSortMode)
 	{
 		bSortingHasChanged = true;
 	}
@@ -2761,7 +2789,7 @@ void STableTreeView::ScheduleSortingAsyncOperation()
 	OnPreAsyncUpdate();
 
 	FGraphEventRef CompletedEvent = StartSortingTask();
-	InProgressAsyncOperationEvent = StartApplyFiltersTask(CompletedEvent);
+	InProgressAsyncOperationEvent = StartHierarchyFilteringTask(CompletedEvent);
 }
 
 ////////////////////////////////////////////////////////////////////////////////////////////////////
@@ -2773,7 +2801,7 @@ FGraphEventRef STableTreeView::StartSortingTask(FGraphEventRef Prerequisite)
 		return Prerequisite;
 	}
 
-	AddInProgressAsyncOperation(EAsyncOperationType::SortingOp);
+	AddInProgressAsyncOperation(EAsyncOperationType::Sorting);
 
 	CurrentAsyncOpSorter = CurrentSorter.Get();
 	CurrentAsyncOpColumnSortMode = ColumnSortMode;
@@ -2788,7 +2816,7 @@ FGraphEventRef STableTreeView::StartSortingTask(FGraphEventRef Prerequisite)
 		Prerequisites.Add(DispatchEvent);
 	}
 
-	return TGraphTask<FTableTreeViewSortAsyncTask>::CreateTask(&Prerequisites).ConstructAndDispatchWhenReady(this, CurrentAsyncOpSorter, CurrentAsyncOpColumnSortMode);
+	return TGraphTask<FTableTreeViewSortingAsyncTask>::CreateTask(&Prerequisites).ConstructAndDispatchWhenReady(this, CurrentAsyncOpSorter, CurrentAsyncOpColumnSortMode);
 }
 
 ////////////////////////////////////////////////////////////////////////////////////////////////////
@@ -2803,6 +2831,7 @@ void STableTreeView::ApplySorting()
 
 ////////////////////////////////////////////////////////////////////////////////////////////////////
 
+// Note: This function is called from a task thread!
 void STableTreeView::SortTreeNodes(ITableCellValueSorter* InSorter, EColumnSortMode::Type InColumnSortMode)
 {
 	FStopwatch Stopwatch;
@@ -3216,16 +3245,16 @@ void STableTreeView::RebuildTree(bool bResync)
 
 FTableTreeNodePtr STableTreeView::GetNodeByTableRowIndex(int32 RowIndex) const
 {
-	return (RowIndex >= 0 && RowIndex < TableTreeNodes.Num()) ? TableTreeNodes[RowIndex] : nullptr;
+	return (RowIndex >= 0 && RowIndex < TableRowNodes.Num()) ? TableRowNodes[RowIndex] : nullptr;
 }
 
 ////////////////////////////////////////////////////////////////////////////////////////////////////
 
 void STableTreeView::SelectNodeByTableRowIndex(int32 RowIndex)
 {
-	if (RowIndex >= 0 && RowIndex < TableTreeNodes.Num())
+	if (RowIndex >= 0 && RowIndex < TableRowNodes.Num())
 	{
-		FTableTreeNodePtr NodePtr = TableTreeNodes[RowIndex];
+		FTableTreeNodePtr NodePtr = TableRowNodes[RowIndex];
 		ensure(NodePtr);
 
 		if (!NodePtr)
@@ -3283,7 +3312,7 @@ void STableTreeView::OnPostAsyncUpdate()
 		TreeView->ClearExpandedItems();
 
 		// Grouping can result in old group nodes no longer existing in the tree, so we don't keep the expanded list.
-		if (!HasInProgressAsyncOperation(EAsyncOperationType::GroupingOp))
+		if (!HasInProgressAsyncOperation(EAsyncOperationType::Grouping))
 		{
 			for (auto It = ExpandedNodes.CreateConstIterator(); It; ++It)
 			{
@@ -3415,22 +3444,22 @@ double STableTreeView::GetAllOperationsDuration()
 
 void STableTreeView::StartPendingAsyncOperations()
 {
-	if (ScheduleGroupingAsyncOperationIfNeeded())
+	if (ScheduleNodeFilteringAsyncOperationIfNeeded())
 	{
 		return;
 	}
 
-	if (ScheduleSortingAsyncOperationIfNeeded())
+	if (ScheduleGroupingAsyncOperationIfNeeded())
 	{
 		return;
 	}
 
-	if (ScheduleHierarchyFilteringAsyncOperationIfNeeded())
+	if (ScheduleSortingAsyncOperationIfNeeded())
 	{
 		return;
 	}
 
-	if (ScheduleNodeFilteringAsyncOperationIfNeeded())
+	if (ScheduleHierarchyFilteringAsyncOperationIfNeeded())
 	{
 		return;
 	}
diff --git a/Engine/Source/Developer/TraceInsights/Private/Insights/Table/Widgets/STableTreeView.h b/Engine/Source/Developer/TraceInsights/Private/Insights/Table/Widgets/STableTreeView.h
--- a/Engine/Source/Developer/TraceInsights/Private/Insights/Table/Widgets/STableTreeView.h
+++ b/Engine/Source/Developer/TraceInsights/Private/Insights/Table/Widgets/STableTreeView.h
@@ -52,9 +52,11 @@ typedef TTextFilter<const FTableTreeNodePtr&> FTableTreeNodeTextFilter;
 
 enum class EAsyncOperationType : uint32
 {
-	FilteringOp = 1,
-	SortingOp = 1 << 1,
-	GroupingOp = 1 << 2,
+	NodeFiltering       = 1 << 0,
+	Grouping            = 1 << 1,
+	Aggregation         = 1 << 2,
+	Sorting             = 1 << 3,
+	HierarchyFiltering  = 1 << 4,
 };
 
 ENUM_CLASS_FLAGS(EAsyncOperationType);
@@ -103,9 +105,10 @@ struct FTableTaskInfo
  */
 class STableTreeView : public SCompoundWidget, public IAsyncOperationStatusProvider
 {
-	friend class FTableTreeViewFilterAsyncTask;
-	friend class FTableTreeViewSortAsyncTask;
-	friend class FTableTreeViewGroupAsyncTask;
+	friend class FTableTreeViewNodeFilteringAsyncTask;
+	friend class FTableTreeViewSortingAsyncTask;
+	friend class FTableTreeViewGroupingAsyncTask;
+	friend class FTableTreeViewHierarchyFilteringAsyncTask;
 	friend class FTableTreeViewAsyncCompleteTask;
 	friend class FSearchForItemToSelectTask;
 	friend class FSelectNodeByTableRowIndexTask;
@@ -254,33 +257,27 @@ protected:
 	FName TableRow_GetHighlightedNodeName() const;
 
 	////////////////////////////////////////////////////////////////////////////////////////////////////
-	// Node Filtering
+	// Node Filtering (TableRowNodes --> FilteredNodes)
 
 	void InitNodeFiltering();
 
 	void OnNodeFilteringChanged();
 	bool ScheduleNodeFilteringAsyncOperationIfNeeded();
 	void ScheduleNodeFilteringAsyncOperation();
-	FGraphEventRef StartApplyFiltersTask(FGraphEventRef Prerequisite = nullptr);
+	FGraphEventRef StartNodeFilteringTask(FGraphEventRef Prerequisite = nullptr);
 
 	void ApplyNodeFiltering();
 
-	bool ApplyAdvancedFiltersForNode(FTableTreeNodePtr NodePtr);
-
-	virtual bool FilterNode(const FTableTreeNodePtr& NodePtr) const;
+	virtual bool FilterNode(const FFilterConfigurator& InFilterConfigurator, const FTableTreeNode& InNode) const;
 
 	virtual void InitFilterConfigurator(FFilterConfigurator& InOutFilterConfigurator);
 	virtual void UpdateFilterContext(const FFilterConfigurator& InFilterConfigurator, const FTableTreeNode& InNode) const;
 
 	virtual TSharedRef<SWidget> ConstructFilterConfiguratorButton();
-
-	bool FilterConfigurator_IsEnabled() const;
-	FText FilterConfigurator_GetTooltipText() const;
-	bool FilterConfigurator_HasFilters() const;
 	FReply FilterConfigurator_OnClicked();
 
 	////////////////////////////////////////////////////////////////////////////////////////////////////
-	// Hierarchy Filtering
+	// Hierarchy Filtering (Root->Children hierarchy --> Root->FilteredChildren hierarchy)
 
 	void InitHierarchyFiltering();
 
@@ -295,6 +292,8 @@ protected:
 	void OnHierarchyFilteringChanged();
 	bool ScheduleHierarchyFilteringAsyncOperationIfNeeded();
 	void ScheduleHierarchyFilteringAsyncOperation();
+	FGraphEventRef StartHierarchyFilteringTask(FGraphEventRef Prerequisite = nullptr);
+
 	void ApplyHierarchyFiltering();
 	void ApplyEmptyHierarchyFilteringRec(FTableTreeNodePtr NodePtr);
 	bool ApplyHierarchyFilteringRec(FTableTreeNodePtr NodePtr);
@@ -303,13 +302,10 @@ protected:
 	bool MakeSubtreeVisible(FTableTreeNodePtr NodePtr, bool bFilterIsEmpty);
 
 	virtual TSharedRef<SWidget> ConstructSearchBox();
-
-	bool SearchBox_IsEnabled() const;
 	void SearchBox_OnTextChanged(const FText& InFilterText);
-	FText SearchBox_GetTooltipText() const;
 
 	////////////////////////////////////////////////////////////////////////////////////////////////////
-	// Grouping
+	// Grouping (FilteredNodes --> Root->Children hierarchy)
 
 	void CreateGroupings();
 	virtual void InternalCreateGroupings();
@@ -342,52 +338,22 @@ protected:
 	bool GroupingCrumbMenu_Add_CanExecute(const TSharedPtr<FTreeNodeGrouping> Grouping, const TSharedPtr<FTreeNodeGrouping> AfterGrouping) const;
 
 	////////////////////////////////////////////////////////////////////////////////////////////////////
-	// Aggregation
+	// Aggregation (Root->Children hierarchy)
 
-	void UpdateAggregatedValues(FTableTreeNode& GroupNode);
+	static void UpdateCStringSameValueAggregationSingleNode(const FTableColumn& InColumn, FTableTreeNode& GroupNode);
+	static void UpdateCStringSameValueAggregationRec(const FTableColumn& InColumn, FTableTreeNode& GroupNode);
 
-	template<typename T>
-	static void UpdateAggregationRec(FTableColumn& Column, FTableTreeNode& GroupNode, T InitialAggregatedValue, bool bSetInitialValue, TFunctionRef<T(T, const FTableCellValue&)> ValueGetterFunc)
-	{
-		T AggregatedValue = InitialAggregatedValue;
+	template<typename T, bool bSetInitialValue, bool bIsRercursive>
+	static void UpdateAggregation(const FTableColumn& InColumn, FTableTreeNode& InOutGroupNode, const T InitialAggregatedValue, TFunctionRef<T(T, const FTableCellValue&)> ValueGetterFunc);
 
-		for (FBaseTreeNodePtr NodePtr : GroupNode.GetChildren())
-		{
-			if (NodePtr->IsFiltered())
-			{
-				continue;
-			}
-
-			if (!NodePtr->IsGroup())
-			{
-				const TOptional<FTableCellValue> NodeValue = Column.GetValue(*NodePtr);
-				if (NodeValue.IsSet())
-				{
-					AggregatedValue = ValueGetterFunc(AggregatedValue, NodeValue.GetValue());
-				}
-			}
-			else
-			{
-				FTableTreeNode& TableNode = *(FTableTreeNode*)NodePtr.Get();
-				TableNode.ResetAggregatedValue(Column.GetId());
-				UpdateAggregationRec(Column, TableNode, InitialAggregatedValue, bSetInitialValue, ValueGetterFunc);
-				if (TableNode.HasAggregatedValue(Column.GetId()))
-				{
-					AggregatedValue = ValueGetterFunc(AggregatedValue, TableNode.GetAggregatedValue(Column.GetId()));
-				}
-			}
-		}
+	template<bool bIsRercursive>
+	static void UpdateAggregatedValues(TSharedPtr<FTable> InTable, FTableTreeNode& InOutGroupNode);
 
-		if (bSetInitialValue || InitialAggregatedValue != AggregatedValue)
-		{
-			GroupNode.AddAggregatedValue(Column.GetId(), FTableCellValue(AggregatedValue));
-		}
-	}
-
-	static void UpdateCStringSameValueAggregationRec(FTableColumn& Column, FTableTreeNode& GroupNode);
+	void UpdateAggregatedValuesSingleNode(FTableTreeNode& GroupNode);
+	void UpdateAggregatedValuesRec(FTableTreeNode& GroupNode);
 
 	////////////////////////////////////////////////////////////////////////////////////////////////////
-	// Sorting
+	// Sorting (Root->Children hierarchy)
 
 	static const EColumnSortMode::Type GetDefaultColumnSortMode();
 	static const FName GetDefaultColumnBeingSorted();
@@ -552,7 +518,13 @@ protected:
 	FTableTreeNodePtr Root;
 
 	/** Table row nodes. Each node corresponds to a table row. Index in this array corresponds to RowIndex in source table. */
-	TArray<FTableTreeNodePtr> TableTreeNodes;
+	TArray<FTableTreeNodePtr> TableRowNodes;
+
+	/** Filtered table nodes. These are the nodes from the TableRowNodes array, after applying the node filtering. */
+	TArray<FTableTreeNodePtr> FilteredNodes;
+
+	/** A pointer to the filtered table nodes. If node filtering is empty, this points directly to TableRowNodes. */
+	TArray<FTableTreeNodePtr>* FilteredNodesPtr = nullptr;
 
 	/** A filtered array of group nodes to be displayed in the tree widget. */
 	TArray<FTableTreeNodePtr> FilteredGroupNodes;
@@ -663,15 +635,41 @@ protected:
 
 ////////////////////////////////////////////////////////////////////////////////////////////////////
 
-class FTableTreeViewFilterAsyncTask
+class FTableTreeViewNodeFilteringAsyncTask
+{
+public:
+	FTableTreeViewNodeFilteringAsyncTask(STableTreeView* InPtr)
+	{
+		TableTreeViewPtr = InPtr;
+	}
+
+	FORCEINLINE TStatId GetStatId() const { RETURN_QUICK_DECLARE_CYCLE_STAT(FTableTreeViewNodeFilteringAsyncTask, STATGROUP_TaskGraphTasks); }
+	ENamedThreads::Type GetDesiredThread() { return ENamedThreads::Type::AnyThread; }
+	static ESubsequentsMode::Type GetSubsequentsMode() { return ESubsequentsMode::TrackSubsequents; }
+
+	void DoTask(ENamedThreads::Type CurrentThread, const FGraphEventRef& MyCompletionGraphEvent)
+	{
+		if (TableTreeViewPtr)
+		{
+			TableTreeViewPtr->ApplyNodeFiltering();
+		}
+	}
+
+private:
+	STableTreeView* TableTreeViewPtr = nullptr;
+};
+
+////////////////////////////////////////////////////////////////////////////////////////////////////
+
+class FTableTreeViewHierarchyFilteringAsyncTask
 {
 public:
-	FTableTreeViewFilterAsyncTask(STableTreeView* InPtr)
+	FTableTreeViewHierarchyFilteringAsyncTask(STableTreeView* InPtr)
 	{
 		TableTreeViewPtr = InPtr;
 	}
 
-	FORCEINLINE TStatId GetStatId() const { RETURN_QUICK_DECLARE_CYCLE_STAT(FTableTreeViewFilterAsyncTask, STATGROUP_TaskGraphTasks); }
+	FORCEINLINE TStatId GetStatId() const { RETURN_QUICK_DECLARE_CYCLE_STAT(FTableTreeViewHierarchyFilteringAsyncTask, STATGROUP_TaskGraphTasks); }
 	ENamedThreads::Type GetDesiredThread() { return ENamedThreads::Type::AnyThread; }
 	static ESubsequentsMode::Type GetSubsequentsMode() { return ESubsequentsMode::TrackSubsequents; }
 
@@ -689,17 +687,17 @@ private:
 
 ////////////////////////////////////////////////////////////////////////////////////////////////////
 
-class FTableTreeViewSortAsyncTask
+class FTableTreeViewSortingAsyncTask
 {
 public:
-	FTableTreeViewSortAsyncTask(STableTreeView* InPtr, ITableCellValueSorter* InSorter, EColumnSortMode::Type InColumnSortMode)
+	FTableTreeViewSortingAsyncTask(STableTreeView* InPtr, ITableCellValueSorter* InSorter, EColumnSortMode::Type InColumnSortMode)
 	{
 		TableTreeViewPtr = InPtr;
 		Sorter = InSorter;
 		ColumnSortMode = InColumnSortMode;
 	}
 
-	FORCEINLINE TStatId GetStatId() const { RETURN_QUICK_DECLARE_CYCLE_STAT(FTableTreeViewSortAsyncTask, STATGROUP_TaskGraphTasks); }
+	FORCEINLINE TStatId GetStatId() const { RETURN_QUICK_DECLARE_CYCLE_STAT(FTableTreeViewSortingAsyncTask, STATGROUP_TaskGraphTasks); }
 	ENamedThreads::Type GetDesiredThread() { return ENamedThreads::Type::AnyThread; }
 	static ESubsequentsMode::Type GetSubsequentsMode() { return ESubsequentsMode::TrackSubsequents; }
 
@@ -719,16 +717,16 @@ private:
 
 ////////////////////////////////////////////////////////////////////////////////////////////////////
 
-class FTableTreeViewGroupAsyncTask
+class FTableTreeViewGroupingAsyncTask
 {
 public:
-	FTableTreeViewGroupAsyncTask(STableTreeView* InPtr, TArray<TSharedPtr<FTreeNodeGrouping>>* InGroupings)
+	FTableTreeViewGroupingAsyncTask(STableTreeView* InPtr, TArray<TSharedPtr<FTreeNodeGrouping>>* InGroupings)
 	{
 		TableTreeViewPtr = InPtr;
 		Groupings = InGroupings;
 	}
 
-	FORCEINLINE TStatId GetStatId() const { RETURN_QUICK_DECLARE_CYCLE_STAT(FTableTreeViewGroupAsyncTask, STATGROUP_TaskGraphTasks); }
+	FORCEINLINE TStatId GetStatId() const { RETURN_QUICK_DECLARE_CYCLE_STAT(FTableTreeViewGroupingAsyncTask, STATGROUP_TaskGraphTasks); }
 	ENamedThreads::Type GetDesiredThread() { return ENamedThreads::Type::AnyThread; }
 	static ESubsequentsMode::Type GetSubsequentsMode() { return ESubsequentsMode::TrackSubsequents; }
 
diff --git a/Engine/Source/Developer/TraceInsights/Private/Insights/Table/Widgets/STableTreeViewCell.cpp b/Engine/Source/Developer/TraceInsights/Private/Insights/Table/Widgets/STableTreeViewCell.cpp
--- a/Engine/Source/Developer/TraceInsights/Private/Insights/Table/Widgets/STableTreeViewCell.cpp
+++ b/Engine/Source/Developer/TraceInsights/Private/Insights/Table/Widgets/STableTreeViewCell.cpp
@@ -18,7 +18,7 @@
 #include "Insights/Table/ViewModels/TableColumn.h"
 #include "Insights/Table/Widgets/STableTreeViewRow.h"
 
-#define LOCTEXT_NAMESPACE "STableTreeView"
+#define LOCTEXT_NAMESPACE "Insights::STableTreeView"
 
 namespace Insights
 {
diff --git a/Engine/Source/Developer/TraceInsights/Private/Insights/Table/Widgets/SUntypedTableTreeView.cpp b/Engine/Source/Developer/TraceInsights/Private/Insights/Table/Widgets/SUntypedTableTreeView.cpp
--- a/Engine/Source/Developer/TraceInsights/Private/Insights/Table/Widgets/SUntypedTableTreeView.cpp
+++ b/Engine/Source/Developer/TraceInsights/Private/Insights/Table/Widgets/SUntypedTableTreeView.cpp
@@ -70,10 +70,10 @@ void SUntypedTableTreeView::RebuildTree(bool bResync)
 
 	if (bResync)
 	{
-		TableTreeNodes.Empty();
+		TableRowNodes.Empty();
 	}
 
-	const int32 PreviousNodeCount = TableTreeNodes.Num();
+	const int32 PreviousNodeCount = TableRowNodes.Num();
 
 	//check(Table->Is<Insights::FUntypedTable>());
 	TSharedPtr<Insights::FUntypedTable> UntypedTable = StaticCastSharedPtr<Insights::FUntypedTable>(Table);
@@ -84,9 +84,9 @@ void SUntypedTableTreeView::RebuildTree(bool bResync)
 	if (SourceTable.IsValid() && TableReader.IsValid())
 	{
 		const int32 TotalRowCount = static_cast<int32>(SourceTable->GetRowCount());
-		if (TotalRowCount != TableTreeNodes.Num())
+		if (TotalRowCount != TableRowNodes.Num())
 		{
-			TableTreeNodes.Empty(TotalRowCount);
+			TableRowNodes.Empty(TotalRowCount);
 			FName BaseNodeName(TEXT("row"));
 			for (int32 RowIndex = 0; RowIndex < TotalRowCount; ++RowIndex)
 			{
@@ -94,15 +94,15 @@ void SUntypedTableTreeView::RebuildTree(bool bResync)
 				FName NodeName(BaseNodeName, RowIndex + 1);
 				FTableTreeNodePtr NodePtr = MakeShared<FTableTreeNode>(NodeName, Table, RowIndex);
 				NodePtr->SetDefaultSortOrder(RowIndex + 1);
-				TableTreeNodes.Add(NodePtr);
+				TableRowNodes.Add(NodePtr);
 			}
-			ensure(TableTreeNodes.Num() == TotalRowCount);
+			ensure(TableRowNodes.Num() == TotalRowCount);
 		}
 	}
 
 	SyncStopwatch.Stop();
 
-	if (bResync || TableTreeNodes.Num() != PreviousNodeCount)
+	if (bResync || TableRowNodes.Num() != PreviousNodeCount)
 	{
 		// Save selection.
 		TArray<FTableTreeNodePtr> SelectedItems;
@@ -135,7 +135,7 @@ void SUntypedTableTreeView::RebuildTree(bool bResync)
 	{
 		const double SyncTime = SyncStopwatch.GetAccumulatedTime();
 		UE_LOG(TraceInsights, Log, TEXT("[Table] Tree view rebuilt in %.4fs (sync: %.4fs + update: %.4fs) --> %d rows (%d added)"),
-			TotalTime, SyncTime, TotalTime - SyncTime, TableTreeNodes.Num(), TableTreeNodes.Num() - PreviousNodeCount);
+			TotalTime, SyncTime, TotalTime - SyncTime, TableRowNodes.Num(), TableRowNodes.Num() - PreviousNodeCount);
 	}
 }
 
diff --git a/Engine/Source/Developer/TraceInsights/Private/Insights/TaskGraphProfiler/Widgets/STaskTableTreeView.cpp b/Engine/Source/Developer/TraceInsights/Private/Insights/TaskGraphProfiler/Widgets/STaskTableTreeView.cpp
--- a/Engine/Source/Developer/TraceInsights/Private/Insights/TaskGraphProfiler/Widgets/STaskTableTreeView.cpp
+++ b/Engine/Source/Developer/TraceInsights/Private/Insights/TaskGraphProfiler/Widgets/STaskTableTreeView.cpp
@@ -174,7 +174,7 @@ void STaskTableTreeView::RebuildTree(bool bResync)
 		TSharedPtr<FTaskTable> TaskTable = GetTaskTable();
 		TArray<FTaskEntry>& Tasks = TaskTable->GetTaskEntries();
 		Tasks.Empty();
-		TableTreeNodes.Empty();
+		TableRowNodes.Empty();
 
 		if (QueryStartTime < QueryEndTime && Session.IsValid())
 		{
@@ -186,7 +186,7 @@ void STaskTableTreeView::RebuildTree(bool bResync)
 			{
 				FName BaseNodeName(TEXT("task"));
 
-				TArray<FTableTreeNodePtr>* Nodes = &TableTreeNodes;
+				TArray<FTableTreeNodePtr>* Nodes = &TableRowNodes;
 
 				TasksProvider->EnumerateTasks(QueryStartTime, QueryEndTime, SelectedTasksSelectionOption, [&Tasks, &TaskTable, &BaseNodeName, Nodes](const TraceServices::FTaskInfo& TaskInfo)
 				{
