[Insights]
  - Improved FBaseTreeNode API.
    * Made the array of filtered nodes optional (allocated only when needed).
    * Removed the Tooltip and Context from FGroupNodeData (there are not needed anymore as these info can be added in a FTableTreeNode derived class when needed).
    * These saves memory used by large tree views (for a tree view of 3M allocations grouped by callstack the saving is ~25%).
    * Updated the banner text displayed when tree is empty or filtered.
  - Memory Insights: Added new FCallstackFrameGroupNode for the group nodes used in by Callstack grouping. Also fixed the special tooltip for these nodes to be formatted on multiple lines.

diff --git a/Engine/Source/Developer/TraceInsights/Private/Insights/MemoryProfiler/ViewModels/CallstackFormatting.h b/Engine/Source/Developer/TraceInsights/Private/Insights/MemoryProfiler/ViewModels/CallstackFormatting.h
--- a/Engine/Source/Developer/TraceInsights/Private/Insights/MemoryProfiler/ViewModels/CallstackFormatting.h
+++ b/Engine/Source/Developer/TraceInsights/Private/Insights/MemoryProfiler/ViewModels/CallstackFormatting.h
@@ -10,10 +10,11 @@
 
 enum class EStackFrameFormatFlags : uint8
 {
-	Module  = 1 << 0,
-	Symbol  = 1 << 1,
-	File    = 1 << 2,
-	Line    = 1 << 3,
+	Module    = 1 << 0, // include module name
+	Symbol    = 1 << 1, // include symbol name
+	File      = 1 << 2, // include source file name
+	Line      = 1 << 3, // include source line number
+	Multiline = 1 << 4, // allow formatting on multiple lines
 
 	ModuleAndSymbol         = Module + Symbol,
 	ModuleSymbolFileAndLine = Module + Symbol + File + Line,
@@ -47,21 +48,23 @@ inline void FormatStackFrame(const TraceServices::FStackFrame& Frame, FStringBui
 			if (EnumHasAnyFlags(FormatFlags, EStackFrameFormatFlags::Module))
 			{
 				OutString.Append(Frame.Symbol->Module);
-				if (FormatFlags != EStackFrameFormatFlags::Module)
+				if (FormatFlags != EStackFrameFormatFlags::Module &&
+					FormatFlags != (EStackFrameFormatFlags::Module | EStackFrameFormatFlags::Multiline))
 				{
-					OutString.AppendChar(TEXT('!'));
+					OutString.AppendChar(EnumHasAnyFlags(FormatFlags, EStackFrameFormatFlags::Multiline) ? TEXT('\n') : TEXT('!'));
 				}
 			}
 			if (EnumHasAnyFlags(FormatFlags, EStackFrameFormatFlags::Symbol))
 			{
 				OutString.Append(Frame.Symbol->Name);
 			}
+			if (EnumHasAnyFlags(FormatFlags, EStackFrameFormatFlags::Symbol) &&
+				EnumHasAnyFlags(FormatFlags, EStackFrameFormatFlags::File | EStackFrameFormatFlags::Line))
+			{
+				OutString.AppendChar(EnumHasAnyFlags(FormatFlags, EStackFrameFormatFlags::Multiline) ? TEXT('\n') : TEXT(' '));
+			}
 			if (EnumHasAnyFlags(FormatFlags, EStackFrameFormatFlags::File))
 			{
-				if (EnumHasAnyFlags(FormatFlags, EStackFrameFormatFlags::Symbol))
-				{
-					OutString.AppendChar(TEXT(' '));
-				}
 				OutString.Append(Frame.Symbol->File);
 			}
 			if (EnumHasAnyFlags(FormatFlags, EStackFrameFormatFlags::Line))
@@ -76,9 +79,10 @@ inline void FormatStackFrame(const TraceServices::FStackFrame& Frame, FStringBui
 			if (EnumHasAnyFlags(FormatFlags, EStackFrameFormatFlags::Module))
 			{
 				OutString.Append(Frame.Symbol->Module);
-				if (FormatFlags != EStackFrameFormatFlags::Module)
+				if (FormatFlags != EStackFrameFormatFlags::Module &&
+					FormatFlags != (EStackFrameFormatFlags::Module | EStackFrameFormatFlags::Multiline))
 				{
-					OutString.AppendChar(TEXT('!'));
+					OutString.AppendChar(EnumHasAnyFlags(FormatFlags, EStackFrameFormatFlags::Multiline) ? TEXT('\n') : TEXT('!'));
 				}
 			}
 			if (EnumHasAnyFlags(FormatFlags, EStackFrameFormatFlags::Symbol))
@@ -87,7 +91,7 @@ inline void FormatStackFrame(const TraceServices::FStackFrame& Frame, FStringBui
 			}
 			if (EnumHasAnyFlags(FormatFlags, EStackFrameFormatFlags::ModuleAndSymbol))
 			{
-				OutString.AppendChar(TEXT(' '));
+				OutString.AppendChar(EnumHasAnyFlags(FormatFlags, EStackFrameFormatFlags::Multiline) ? TEXT('\n') : TEXT(' '));
 			}
 			OutString.Appendf(TEXT("(%s)"), QueryResultToString((Result)));
 			break;
diff --git a/Engine/Source/Developer/TraceInsights/Private/Insights/MemoryProfiler/ViewModels/MemAllocGroupingByCallstack.cpp b/Engine/Source/Developer/TraceInsights/Private/Insights/MemoryProfiler/ViewModels/MemAllocGroupingByCallstack.cpp
--- a/Engine/Source/Developer/TraceInsights/Private/Insights/MemoryProfiler/ViewModels/MemAllocGroupingByCallstack.cpp
+++ b/Engine/Source/Developer/TraceInsights/Private/Insights/MemoryProfiler/ViewModels/MemAllocGroupingByCallstack.cpp
@@ -5,6 +5,7 @@
 
 // Insights
 #include "Insights/Common/AsyncOperationProgress.h"
+#include "Insights/InsightsStyle.h"
 #include "Insights/MemoryProfiler/ViewModels/CallstackFormatting.h"
 #include "Insights/MemoryProfiler/ViewModels/MemAllocNode.h"
 
@@ -13,6 +14,39 @@
 namespace Insights
 {
 
+////////////////////////////////////////////////////////////////////////////////////////////////////
+// FCallstackFrameGroupNode
+////////////////////////////////////////////////////////////////////////////////////////////////////
+
+INSIGHTS_IMPLEMENT_RTTI(FCallstackFrameGroupNode)
+
+////////////////////////////////////////////////////////////////////////////////////////////////////
+
+const FText FCallstackFrameGroupNode::GetTooltipText() const
+{
+	if (StackFrame)
+	{
+		TStringBuilder<1024> String;
+		FormatStackFrame(*StackFrame, String, EStackFrameFormatFlags::ModuleSymbolFileAndLine | EStackFrameFormatFlags::Multiline);
+		return FText::FromString(FString(String));
+	}
+	return FText::GetEmpty();
+}
+
+////////////////////////////////////////////////////////////////////////////////////////////////////
+
+const FSlateBrush* FCallstackFrameGroupNode::GetIcon() const
+{
+	return FInsightsStyle::GetBrush("Icons.Group.TreeItem");
+}
+
+////////////////////////////////////////////////////////////////////////////////////////////////////
+
+FLinearColor FCallstackFrameGroupNode::GetColor() const
+{
+	return FLinearColor(0.5f, 0.75f, 1.0f, 1.0f);
+}
+
 ////////////////////////////////////////////////////////////////////////////////////////////////////
 // FMemAllocGroupingByCallstack
 ////////////////////////////////////////////////////////////////////////////////////////////////////
@@ -66,7 +100,7 @@ void FMemAllocGroupingByCallstack::GroupNodes(const TArray<FTableTreeNodePtr>& N
 
 		if (NodePtr->IsGroup())
 		{
-			ParentGroup.AddChildAndSetGroupPtr(NodePtr);
+			ParentGroup.AddChildAndSetParent(NodePtr);
 			continue;
 		}
 
@@ -150,7 +184,7 @@ void FMemAllocGroupingByCallstack::GroupNodes(const TArray<FTableTreeNodePtr>& N
 
 		if (GroupPtr != Root)
 		{
-			GroupPtr->Node->AddChildAndSetGroupPtr(NodePtr);
+			GroupPtr->Node->AddChildAndSetParent(NodePtr);
 		}
 		else if (NumFrames == 0)
 		{
@@ -158,7 +192,7 @@ void FMemAllocGroupingByCallstack::GroupNodes(const TArray<FTableTreeNodePtr>& N
 			{
 				EmptyCallstackGroupPtr = CreateEmptyCallstackGroup(InParentTable, ParentGroup);
 			}
-			EmptyCallstackGroupPtr->AddChildAndSetGroupPtr(NodePtr);
+			EmptyCallstackGroupPtr->AddChildAndSetParent(NodePtr);
 		}
 		else
 		{
@@ -166,7 +200,7 @@ void FMemAllocGroupingByCallstack::GroupNodes(const TArray<FTableTreeNodePtr>& N
 			{
 				UnsetGroupPtr = CreateUnsetGroup(InParentTable, ParentGroup);
 			}
-			UnsetGroupPtr->AddChildAndSetGroupPtr(NodePtr);
+			UnsetGroupPtr->AddChildAndSetParent(NodePtr);
 		}
 	}
 
@@ -202,15 +236,6 @@ FName FMemAllocGroupingByCallstack::GetGroupName(const TraceServices::FStackFram
 
 ////////////////////////////////////////////////////////////////////////////////////////////////////
 
-FText FMemAllocGroupingByCallstack::GetGroupTooltip(const TraceServices::FStackFrame* Frame) const
-{
-	TStringBuilder<1024> String;
-	FormatStackFrame(*Frame, String, EStackFrameFormatFlags::ModuleSymbolFileAndLine);
-	return FText::FromString(FString(String));
-}
-
-////////////////////////////////////////////////////////////////////////////////////////////////////
-
 FMemAllocGroupingByCallstack::FCallstackGroup* FMemAllocGroupingByCallstack::CreateGroup(
 	TArray<FCallstackGroup*>& InOutAllCallstackGroup,
 	FCallstackGroup* InParentGroup,
@@ -224,16 +249,11 @@ FMemAllocGroupingByCallstack::FCallstackGroup* FMemAllocGroupingByCallstack::Cre
 
 	InOutAllCallstackGroup.Add(NewGroupPtr);
 
-	const FSlateBrush* IconBrush = FBaseTreeNode::GetDefaultIcon(true);
-	FLinearColor GroupNodeColor(0.5f, 0.75f, 1.0f, 1.0f);
-
-	FTableTreeNodePtr NodePtr = MakeShared<FCustomTableTreeNode>(InGroupName, InParentTable, IconBrush, GroupNodeColor);
+	TSharedPtr<FCallstackFrameGroupNode> NodePtr = MakeShared<FCallstackFrameGroupNode>(InGroupName, InParentTable, InFrame);
 	NodePtr->SetExpansion(false);
-	InParentGroup->Node->AddChildAndSetGroupPtr(NodePtr);
 
+	InParentGroup->Node->AddChildAndSetParent(NodePtr);
 	NewGroupPtr->Node = NodePtr.Get();
-	NewGroupPtr->Node->SetTooltip(GetGroupTooltip(InFrame));
-	NewGroupPtr->Node->SetContext((void*)InFrame);
 
 	return NewGroupPtr;
 }
@@ -245,7 +265,7 @@ FTableTreeNode* FMemAllocGroupingByCallstack::CreateUnsetGroup(TWeakPtr<FTable>
 	static FName NotAvailableName(TEXT("N/A"));
 	FTableTreeNodePtr NodePtr = MakeShared<FTableTreeNode>(NotAvailableName, ParentTable);
 	NodePtr->SetExpansion(false);
-	Parent.AddChildAndSetGroupPtr(NodePtr);
+	Parent.AddChildAndSetParent(NodePtr);
 	return NodePtr.Get();
 }
 
@@ -256,7 +276,7 @@ FTableTreeNode* FMemAllocGroupingByCallstack::CreateEmptyCallstackGroup(TWeakPtr
 	static FName NotAvailableName(GetEmptyCallstackString());
 	FTableTreeNodePtr NodePtr = MakeShared<FTableTreeNode>(NotAvailableName, ParentTable);
 	NodePtr->SetExpansion(false);
-	Parent.AddChildAndSetGroupPtr(NodePtr);
+	Parent.AddChildAndSetParent(NodePtr);
 	return NodePtr.Get();
 }
 
diff --git a/Engine/Source/Developer/TraceInsights/Private/Insights/MemoryProfiler/ViewModels/MemAllocGroupingByCallstack.h b/Engine/Source/Developer/TraceInsights/Private/Insights/MemoryProfiler/ViewModels/MemAllocGroupingByCallstack.h
--- a/Engine/Source/Developer/TraceInsights/Private/Insights/MemoryProfiler/ViewModels/MemAllocGroupingByCallstack.h
+++ b/Engine/Source/Developer/TraceInsights/Private/Insights/MemoryProfiler/ViewModels/MemAllocGroupingByCallstack.h
@@ -20,6 +20,40 @@ class IAsyncOperationProgress;
 
 ////////////////////////////////////////////////////////////////////////////////////////////////////
 
+class FCallstackFrameGroupNode : public FTableTreeNode
+{
+	INSIGHTS_DECLARE_RTTI(FCallstackFrameGroupNode, FTableTreeNode)
+
+public:
+	/** Initialization constructor for the group node. */
+	explicit FCallstackFrameGroupNode(const FName InName, TWeakPtr<FTable> InParentTable, const TraceServices::FStackFrame* InStackFrame)
+		: FTableTreeNode(InName, InParentTable)
+		, StackFrame(InStackFrame)
+	{
+	}
+
+	virtual ~FCallstackFrameGroupNode()
+	{
+	}
+
+	virtual const FText GetTooltipText() const override;
+	virtual const FSlateBrush* GetIcon() const override;
+	virtual FLinearColor GetColor() const override;
+
+	/**
+	 * @returns the stack frame.
+	 */
+	const TraceServices::FStackFrame* GetStackFrame() const
+	{
+		return StackFrame;
+	}
+
+private:
+	const TraceServices::FStackFrame* StackFrame = nullptr;
+};
+
+////////////////////////////////////////////////////////////////////////////////////////////////////
+
 class FMemAllocGroupingByCallstack : public FTreeNodeGrouping
 {
 	INSIGHTS_DECLARE_RTTI(FMemAllocGroupingByCallstack, FTreeNodeGrouping)
@@ -48,7 +82,6 @@ public:
 
 private:
 	FName GetGroupName(const TraceServices::FStackFrame* Frame) const;
-	FText GetGroupTooltip(const TraceServices::FStackFrame* Frame) const;
 
 	FCallstackGroup* CreateGroup(TArray<FCallstackGroup*>& InOutAllCallstackGroup, FCallstackGroup* InParentGroup, const FName InGroupName, TWeakPtr<FTable> InParentTable, const TraceServices::FStackFrame* InFrame) const;
 
diff --git a/Engine/Source/Developer/TraceInsights/Private/Insights/MemoryProfiler/ViewModels/MemAllocGroupingByHeap.cpp b/Engine/Source/Developer/TraceInsights/Private/Insights/MemoryProfiler/ViewModels/MemAllocGroupingByHeap.cpp
--- a/Engine/Source/Developer/TraceInsights/Private/Insights/MemoryProfiler/ViewModels/MemAllocGroupingByHeap.cpp
+++ b/Engine/Source/Developer/TraceInsights/Private/Insights/MemoryProfiler/ViewModels/MemAllocGroupingByHeap.cpp
@@ -45,7 +45,7 @@ FTableTreeNodePtr MakeGroupNodeHierarchy(const TraceServices::IAllocationsProvid
 	for (TraceServices::IAllocationsProvider::FHeapSpec* ChildSpec : Spec.Children)
 	{
 		auto ChildNode = MakeGroupNodeHierarchy(*ChildSpec, InParentTable, NodeTable);
-		Node->AddChildAndSetGroupPtr(ChildNode);
+		Node->AddChildAndSetParent(ChildNode);
 	}
 	return Node;
 }
@@ -65,7 +65,7 @@ void FMemAllocGroupingByHeap::GroupNodes(const TArray<FTableTreeNodePtr>& Nodes,
 		TraceServices::FProviderReadScopeLock _(AllocProvider);
 		AllocProvider.EnumerateRootHeaps([&](HeapId Id, const TraceServices::IAllocationsProvider::FHeapSpec& Spec)
 		{
-			ParentGroup.AddChildAndSetGroupPtr(MakeGroupNodeHierarchy(Spec, InParentTable, HeapNodes));
+			ParentGroup.AddChildAndSetParent(MakeGroupNodeHierarchy(Spec, InParentTable, HeapNodes));
 		});
 	}
 
@@ -79,7 +79,7 @@ void FMemAllocGroupingByHeap::GroupNodes(const TArray<FTableTreeNodePtr>& Nodes,
 
 		if (NodePtr->IsGroup())
 		{
-			ParentGroup.AddChildAndSetGroupPtr(NodePtr);
+			ParentGroup.AddChildAndSetParent(NodePtr);
 			continue;
 		}
 
@@ -93,7 +93,7 @@ void FMemAllocGroupingByHeap::GroupNodes(const TArray<FTableTreeNodePtr>& Nodes,
 			const uint8 HeapId = static_cast<uint8>(Alloc->GetRootHeap());
 			if (FTableTreeNodePtr GroupPtr = HeapNodes[HeapId])
 			{
-				GroupPtr->AddChildAndSetGroupPtr(NodePtr);
+				GroupPtr->AddChildAndSetParent(NodePtr);
 			}
 		}
 	}
diff --git a/Engine/Source/Developer/TraceInsights/Private/Insights/MemoryProfiler/ViewModels/MemAllocGroupingBySize.cpp b/Engine/Source/Developer/TraceInsights/Private/Insights/MemoryProfiler/ViewModels/MemAllocGroupingBySize.cpp
--- a/Engine/Source/Developer/TraceInsights/Private/Insights/MemoryProfiler/ViewModels/MemAllocGroupingBySize.cpp
+++ b/Engine/Source/Developer/TraceInsights/Private/Insights/MemoryProfiler/ViewModels/MemAllocGroupingBySize.cpp
@@ -92,7 +92,7 @@ void FMemAllocGroupingBySize::GroupNodes(const TArray<FTableTreeNodePtr>& Nodes,
 
 		if (NodePtr->IsGroup())
 		{
-			ParentGroup.AddChildAndSetGroupPtr(NodePtr);
+			ParentGroup.AddChildAndSetParent(NodePtr);
 			continue;
 		}
 
@@ -120,14 +120,14 @@ void FMemAllocGroupingBySize::GroupNodes(const TArray<FTableTreeNodePtr>& Nodes,
 				const FThreshold& Threshold = Thresholds[ThresholdIndex];
 				GroupPtr = MakeShared<FTableTreeNode>(FName(Threshold.Name.ToString(), 0), InParentTable);
 				GroupPtr->SetExpansion(false);
-				ParentGroup.AddChildAndSetGroupPtr(GroupPtr);
+				ParentGroup.AddChildAndSetParent(GroupPtr);
 				GroupMap[ThresholdIndex] = GroupPtr;
 			}
 		}
 
 		if (GroupPtr != nullptr)
 		{
-			GroupPtr->AddChildAndSetGroupPtr(NodePtr);
+			GroupPtr->AddChildAndSetParent(NodePtr);
 		}
 		else
 		{
@@ -135,9 +135,9 @@ void FMemAllocGroupingBySize::GroupNodes(const TArray<FTableTreeNodePtr>& Nodes,
 			{
 				UnsetGroupPtr = MakeShared<FTableTreeNode>(NotAvailableName, InParentTable);
 				UnsetGroupPtr->SetExpansion(false);
-				ParentGroup.AddChildAndSetGroupPtr(UnsetGroupPtr);
+				ParentGroup.AddChildAndSetParent(UnsetGroupPtr);
 			}
-			UnsetGroupPtr->AddChildAndSetGroupPtr(NodePtr);
+			UnsetGroupPtr->AddChildAndSetParent(NodePtr);
 		}
 	}
 }
diff --git a/Engine/Source/Developer/TraceInsights/Private/Insights/MemoryProfiler/ViewModels/MemAllocGroupingByTag.cpp b/Engine/Source/Developer/TraceInsights/Private/Insights/MemoryProfiler/ViewModels/MemAllocGroupingByTag.cpp
--- a/Engine/Source/Developer/TraceInsights/Private/Insights/MemoryProfiler/ViewModels/MemAllocGroupingByTag.cpp
+++ b/Engine/Source/Developer/TraceInsights/Private/Insights/MemoryProfiler/ViewModels/MemAllocGroupingByTag.cpp
@@ -199,7 +199,7 @@ void FMemAllocGroupingByTag::GroupNodes(const TArray<FTableTreeNodePtr>& Nodes,
 
 		if (NodePtr->IsGroup())
 		{
-			ParentGroup.AddChildAndSetGroupPtr(NodePtr);
+			ParentGroup.AddChildAndSetParent(NodePtr);
 			continue;
 		}
 
@@ -213,12 +213,12 @@ void FMemAllocGroupingByTag::GroupNodes(const TArray<FTableTreeNodePtr>& Nodes,
 			{
 				TagNode = UnknownTagNode;
 			}
-			TagNode->TreeNode->AddChildAndSetGroupPtr(NodePtr);
+			TagNode->TreeNode->AddChildAndSetParent(NodePtr);
 			TagNode->AllocCount++;
 		}
 		else
 		{
-			ParentGroup.AddChildAndSetGroupPtr(NodePtr);
+			ParentGroup.AddChildAndSetParent(NodePtr);
 		}
 	}
 
@@ -250,12 +250,12 @@ void FMemAllocGroupingByTag::GroupNodes(const TArray<FTableTreeNodePtr>& Nodes,
 			check(TagNode->TreeNode);
 			if (TagNode->Parent->TreeNode)
 			{
-				TagNode->Parent->TreeNode->AddChildAndSetGroupPtr(TagNode->TreeNode);
+				TagNode->Parent->TreeNode->AddChildAndSetParent(TagNode->TreeNode);
 			}
 			else
 			{
 				check(TagNode->Parent == &Root);
-				ParentGroup.AddChildAndSetGroupPtr(TagNode->TreeNode);
+				ParentGroup.AddChildAndSetParent(TagNode->TreeNode);
 			}
 		}
 	}
diff --git a/Engine/Source/Developer/TraceInsights/Private/Insights/MemoryProfiler/Widgets/SMemAllocTableTreeView.cpp b/Engine/Source/Developer/TraceInsights/Private/Insights/MemoryProfiler/Widgets/SMemAllocTableTreeView.cpp
--- a/Engine/Source/Developer/TraceInsights/Private/Insights/MemoryProfiler/Widgets/SMemAllocTableTreeView.cpp
+++ b/Engine/Source/Developer/TraceInsights/Private/Insights/MemoryProfiler/Widgets/SMemAllocTableTreeView.cpp
@@ -1904,27 +1904,30 @@ bool SMemAllocTableTreeView::IsExportMemorySnapshotAvailable() const
 
 bool SMemAllocTableTreeView::CanOpenCallstackFrameSourceFileInIDE() const
 {
-	if (TreeView->GetNumItemsSelected() != 1)
+	if (TreeView->GetNumItemsSelected() == 1)
 	{
-		return false;
+		FTableTreeNodePtr TreeNode = TreeView->GetSelectedItems()[0];
+		if (TreeNode.IsValid() && TreeNode->Is<FCallstackFrameGroupNode>())
+		{
+			const FCallstackFrameGroupNode& CallstackFrameNode = TreeNode->As<FCallstackFrameGroupNode>();
+			return CallstackFrameNode.GetStackFrame() != nullptr;
+		}
 	}
-
-	FTableTreeNodePtr TreeNode = TreeView->GetSelectedItems()[0];
-	return TreeNode.IsValid() && TreeNode->IsGroup() && TreeNode->GetContext();
+	return false;
 }
 
 ////////////////////////////////////////////////////////////////////////////////////////////////////
 
 void SMemAllocTableTreeView::OpenCallstackFrameSourceFileInIDE()
 {
-	if (TreeView->GetNumItemsSelected() > 0)
+	if (TreeView->GetNumItemsSelected() == 1)
 	{
 		FTableTreeNodePtr TreeNode = TreeView->GetSelectedItems()[0];
-		if (TreeNode.IsValid() && TreeNode->IsGroup() && TreeNode->GetContext())
+		if (TreeNode.IsValid() && TreeNode->Is<FCallstackFrameGroupNode>())
 		{
-			const TraceServices::FStackFrame* Frame = (const TraceServices::FStackFrame*)TreeNode->GetContext();
-
-			if (Frame->Symbol && Frame->Symbol->File)
+			const FCallstackFrameGroupNode& CallstackFrameNode = TreeNode->As<FCallstackFrameGroupNode>();
+			const TraceServices::FStackFrame* Frame = CallstackFrameNode.GetStackFrame();
+			if (Frame && Frame->Symbol && Frame->Symbol->File)
 			{
 				OpenSourceFileInIDE(Frame->Symbol->File, Frame->Symbol->Line);
 			}
@@ -1936,13 +1939,14 @@ void SMemAllocTableTreeView::OpenCallstackFrameSourceFileInIDE()
 
 FText SMemAllocTableTreeView::GetSelectedCallstackFrameFileName() const
 {
-	if (TreeView->GetNumItemsSelected() > 0)
+	if (TreeView->GetNumItemsSelected() == 1)
 	{
 		FTableTreeNodePtr TreeNode = TreeView->GetSelectedItems()[0];
-		if (TreeNode.IsValid() && TreeNode->IsGroup() && TreeNode->GetContext())
+		if (TreeNode.IsValid() && TreeNode->Is<FCallstackFrameGroupNode>())
 		{
-			const TraceServices::FStackFrame* Frame = (const TraceServices::FStackFrame*)TreeNode->GetContext();
-			if (Frame->Symbol && Frame->Symbol->File)
+			const FCallstackFrameGroupNode& CallstackFrameNode = TreeNode->As<FCallstackFrameGroupNode>();
+			const TraceServices::FStackFrame* Frame = CallstackFrameNode.GetStackFrame();
+			if (Frame && Frame->Symbol && Frame->Symbol->File)
 			{
 				FString SourceFileAndLine = FString::Printf(TEXT("%s(%d)"), Frame->Symbol->File, Frame->Symbol->Line);
 				return FText::FromString(SourceFileAndLine);
diff --git a/Engine/Source/Developer/TraceInsights/Private/Insights/MemoryProfiler/Widgets/SMemTagTreeView.cpp b/Engine/Source/Developer/TraceInsights/Private/Insights/MemoryProfiler/Widgets/SMemTagTreeView.cpp
--- a/Engine/Source/Developer/TraceInsights/Private/Insights/MemoryProfiler/Widgets/SMemTagTreeView.cpp
+++ b/Engine/Source/Developer/TraceInsights/Private/Insights/MemoryProfiler/Widgets/SMemTagTreeView.cpp
@@ -1261,7 +1261,7 @@ void SMemTagTreeView::CreateGroups()
 
 		for (const FMemTagNodePtr& NodePtr : MemTagNodes)
 		{
-			GroupPtr->AddChildAndSetGroupPtr(NodePtr);
+			GroupPtr->AddChildAndSetParent(NodePtr);
 		}
 		TreeView->SetItemExpansion(GroupPtr, true);
 	}
@@ -1278,7 +1278,7 @@ void SMemTagTreeView::CreateGroups()
 				const FName GroupName = *MemTagNodeTypeHelper::ToText(NodeType).ToString();
 				GroupPtr = GroupNodeSet.Add(NodeType, MakeShared<FMemTagNode>(GroupName));
 			}
-			GroupPtr->AddChildAndSetGroupPtr(NodePtr);
+			GroupPtr->AddChildAndSetParent(NodePtr);
 			TreeView->SetItemExpansion(GroupPtr, true);
 		}
 		GroupNodeSet.KeySort([](const EMemTagNodeType& A, const EMemTagNodeType& B) { return A < B; }); // sort groups by type
@@ -1298,7 +1298,7 @@ void SMemTagTreeView::CreateGroups()
 				const FName GroupName(FirstLetterStr);
 				GroupPtr = GroupNodeSet.Add(FirstLetter, MakeShared<FMemTagNode>(GroupName));
 			}
-			GroupPtr->AddChildAndSetGroupPtr(NodePtr);
+			GroupPtr->AddChildAndSetParent(NodePtr);
 		}
 		GroupNodeSet.KeySort([](const TCHAR& A, const TCHAR& B) { return A < B; }); // sort groups alphabetically
 		GroupNodeSet.GenerateValueArray(GroupNodes);
@@ -1316,7 +1316,7 @@ void SMemTagTreeView::CreateGroups()
 				const FName GroupName = *NodePtr->GetTrackerText().ToString();
 				GroupPtr = GroupNodeSet.Add(TrackerId, MakeShared<FMemTagNode>(GroupName));
 			}
-			GroupPtr->AddChildAndSetGroupPtr(NodePtr);
+			GroupPtr->AddChildAndSetParent(NodePtr);
 			TreeView->SetItemExpansion(GroupPtr, true);
 		}
 		GroupNodeSet.GenerateValueArray(GroupNodes);
@@ -1333,7 +1333,7 @@ void SMemTagTreeView::CreateGroups()
 			{
 				GroupPtr = GroupNodeSet.Add(GroupName, MakeShared<FMemTagNode>(GroupName));
 			}
-			GroupPtr->AddChildAndSetGroupPtr(NodePtr);
+			GroupPtr->AddChildAndSetParent(NodePtr);
 			TreeView->SetItemExpansion(GroupPtr, true);
 		}
 		GroupNodeSet.GenerateValueArray(GroupNodes);
diff --git a/Engine/Source/Developer/TraceInsights/Private/Insights/NetworkingProfiler/Widgets/SNetStatsCountersView.cpp b/Engine/Source/Developer/TraceInsights/Private/Insights/NetworkingProfiler/Widgets/SNetStatsCountersView.cpp
--- a/Engine/Source/Developer/TraceInsights/Private/Insights/NetworkingProfiler/Widgets/SNetStatsCountersView.cpp
+++ b/Engine/Source/Developer/TraceInsights/Private/Insights/NetworkingProfiler/Widgets/SNetStatsCountersView.cpp
@@ -912,7 +912,7 @@ void SNetStatsCountersView::CreateGroups()
 
 		for (const FNetStatsCounterNodePtr& NodePtr : NetStatsCounterNodes)
 		{
-			GroupPtr->AddChildAndSetGroupPtr(NodePtr);
+			GroupPtr->AddChildAndSetParent(NodePtr);
 		}
 		TreeView->SetItemExpansion(GroupPtr, true);
 	}
@@ -929,7 +929,7 @@ void SNetStatsCountersView::CreateGroups()
 				const FName GroupName = *NetStatsCounterNodeTypeHelper::ToText(NodeType).ToString();
 				GroupPtr = GroupNodeSet.Add(NodeType, MakeShared<FNetStatsCounterNode>(GroupName));
 			}
-			GroupPtr->AddChildAndSetGroupPtr(NodePtr);
+			GroupPtr->AddChildAndSetParent(NodePtr);
 			TreeView->SetItemExpansion(GroupPtr, true);
 		}
 		GroupNodeSet.KeySort([](const ENetStatsCounterNodeType& A, const ENetStatsCounterNodeType& B) { return A < B; }); // sort groups by type
@@ -949,7 +949,7 @@ void SNetStatsCountersView::CreateGroups()
 				const FName GroupName(FirstLetterStr);
 				GroupPtr = GroupNodeSet.Add(FirstLetter, MakeShared<FNetStatsCounterNode>(GroupName));
 			}
-			GroupPtr->AddChildAndSetGroupPtr(NodePtr);
+			GroupPtr->AddChildAndSetParent(NodePtr);
 		}
 		GroupNodeSet.KeySort([](const TCHAR& A, const TCHAR& B) { return A < B; }); // sort groups alphabetically
 		GroupNodeSet.GenerateValueArray(GroupNodes);
diff --git a/Engine/Source/Developer/TraceInsights/Private/Insights/NetworkingProfiler/Widgets/SNetStatsView.cpp b/Engine/Source/Developer/TraceInsights/Private/Insights/NetworkingProfiler/Widgets/SNetStatsView.cpp
--- a/Engine/Source/Developer/TraceInsights/Private/Insights/NetworkingProfiler/Widgets/SNetStatsView.cpp
+++ b/Engine/Source/Developer/TraceInsights/Private/Insights/NetworkingProfiler/Widgets/SNetStatsView.cpp
@@ -899,7 +899,7 @@ void SNetStatsView::CreateGroups()
 
 		for (const FNetEventNodePtr& NodePtr : NetEventNodes)
 		{
-			GroupPtr->AddChildAndSetGroupPtr(NodePtr);
+			GroupPtr->AddChildAndSetParent(NodePtr);
 		}
 		TreeView->SetItemExpansion(GroupPtr, true);
 	}
@@ -916,7 +916,7 @@ void SNetStatsView::CreateGroups()
 				const FName GroupName = *NetEventNodeTypeHelper::ToText(NodeType).ToString();
 				GroupPtr = GroupNodeSet.Add(NodeType, MakeShared<FNetEventNode>(GroupName));
 			}
-			GroupPtr->AddChildAndSetGroupPtr(NodePtr);
+			GroupPtr->AddChildAndSetParent(NodePtr);
 			TreeView->SetItemExpansion(GroupPtr, true);
 		}
 		GroupNodeSet.KeySort([](const ENetEventNodeType& A, const ENetEventNodeType& B) { return A < B; }); // sort groups by type
@@ -936,7 +936,7 @@ void SNetStatsView::CreateGroups()
 				const FName GroupName(FirstLetterStr);
 				GroupPtr = GroupNodeSet.Add(FirstLetter, MakeShared<FNetEventNode>(GroupName));
 			}
-			GroupPtr->AddChildAndSetGroupPtr(NodePtr);
+			GroupPtr->AddChildAndSetParent(NodePtr);
 		}
 		GroupNodeSet.KeySort([](const TCHAR& A, const TCHAR& B) { return A < B; }); // sort groups alphabetically
 		GroupNodeSet.GenerateValueArray(GroupNodes);
@@ -954,7 +954,7 @@ void SNetStatsView::CreateGroups()
 				const FName GroupName(*FString::Printf(TEXT("Level %d"), Level));
 				GroupPtr = GroupNodeSet.Add(Level, MakeShared<FNetEventNode>(GroupName));
 			}
-			GroupPtr->AddChildAndSetGroupPtr(NodePtr);
+			GroupPtr->AddChildAndSetParent(NodePtr);
 			TreeView->SetItemExpansion(GroupPtr, true);
 		}
 		GroupNodeSet.KeySort([](const uint32& A, const uint32& B) { return A < B; }); // sort groups by level
diff --git a/Engine/Source/Developer/TraceInsights/Private/Insights/Table/ViewModels/BaseTreeNode.cpp b/Engine/Source/Developer/TraceInsights/Private/Insights/Table/ViewModels/BaseTreeNode.cpp
--- a/Engine/Source/Developer/TraceInsights/Private/Insights/Table/ViewModels/BaseTreeNode.cpp
+++ b/Engine/Source/Developer/TraceInsights/Private/Insights/Table/ViewModels/BaseTreeNode.cpp
@@ -31,8 +31,8 @@ const FText FBaseTreeNode::GetExtraDisplayName() const
 {
 	if (IsGroup())
 	{
-		const int32 NumChildren = GroupData->Children.Num();
-		const int32 NumFilteredChildren = GroupData->FilteredChildren.Num();
+		const int32 NumChildren = GetChildrenCount();
+		const int32 NumFilteredChildren = GetFilteredChildrenCount();
 
 		if (NumFilteredChildren == NumChildren)
 		{
@@ -93,7 +93,7 @@ void FBaseTreeNode::SortChildren(const ITableCellValueSorter& Sorter, ESortMode
 
 void FBaseTreeNode::SortFilteredChildren(const ITableCellValueSorter& Sorter, ESortMode SortMode)
 {
-	Sorter.Sort(GroupData->FilteredChildren, SortMode);
+	Sorter.Sort(*GroupData->FilteredChildrenPtr, SortMode);
 }
 
 ////////////////////////////////////////////////////////////////////////////////////////////////////
diff --git a/Engine/Source/Developer/TraceInsights/Private/Insights/Table/ViewModels/BaseTreeNode.h b/Engine/Source/Developer/TraceInsights/Private/Insights/Table/ViewModels/BaseTreeNode.h
--- a/Engine/Source/Developer/TraceInsights/Private/Insights/Table/ViewModels/BaseTreeNode.h
+++ b/Engine/Source/Developer/TraceInsights/Private/Insights/Table/ViewModels/BaseTreeNode.h
@@ -39,22 +39,25 @@ class FBaseTreeNode : public TSharedFromThis<FBaseTreeNode>
 	INSIGHTS_DECLARE_RTTI_BASE(FBaseTreeNode)
 
 protected:
+	enum class EGroupNodeFlags : uint32
+	{
+		None = 0,
+		IsExpanded = 1 << 0,
+	};
+
 	struct FGroupNodeData
 	{
 		/** Children of the group node. */
 		TArray<FBaseTreeNodePtr> Children;
 
-		/** Filtered children of the group node. */
-		TArray<FBaseTreeNodePtr> FilteredChildren;
+		/**
+		 * A pointer to the filtered children of the group node.
+		 * If filter is empty (no filtering), this points directly to Children array.
+		 */
+		TArray<FBaseTreeNodePtr>* FilteredChildrenPtr = &Children;
 
-		/** Descriptive text for the group node to display in a tooltip. */
-		FText Tooltip;
-
-		/** A generic pointer to a data context for the group node. Can be nullptr. */
-		void* Context = nullptr;
-
-		/** Whether the group node should be expanded or not. */
-		bool bIsExpanded = false;
+		/** Various flags specific to a group node. */
+		EGroupNodeFlags Flags = EGroupNodeFlags::None;
 	};
 
 public:
@@ -62,26 +65,14 @@ public:
 	explicit FBaseTreeNode(const FName InName, bool bInIsGroup)
 		: DefaultSortOrder(0)
 		, Name(InName)
+		, Parent(nullptr)
 		, GroupData(bInIsGroup ? new FGroupNodeData() : &DefaultGroupData)
 	{
 	}
 
 	virtual ~FBaseTreeNode()
 	{
-		if (GroupData != &DefaultGroupData)
-		{
-			delete GroupData;
-			GroupData = nullptr;
-		}
-	}
-
-	uint32 GetDefaultSortOrder() const
-	{
-		return DefaultSortOrder;
-	}
-	void SetDefaultSortOrder(uint32 Order)
-	{
-		DefaultSortOrder = Order;
+		RemoveGroupData();
 	}
 
 	/**
@@ -110,20 +101,9 @@ public:
 	/**
 	 * @returns a descriptive text for this node to display in a tooltip.
 	 */
-	virtual const FText GetTooltip() const
+	virtual const FText GetTooltipText() const
 	{
-		return GroupData->Tooltip;
-	}
-
-	/**
-	 * Sets a descriptive text for this node to display in a tooltip.
-	 */
-	virtual void SetTooltip(const FText& InTooltip)
-	{
-		if (IsGroup())
-		{
-			GroupData->Tooltip = InTooltip;
-		}
+		return FText::GetEmpty();
 	}
 
 	/**
@@ -152,25 +132,28 @@ public:
 		return GetDefaultColor(IsGroup());
 	}
 
+	//////////////////////////////////////////////////
+	// Parent Node
+
 	/**
-	 * @returns a pointer to a data context for this node. Only available for group nodes.
+	 * @returns a weak reference to the parent group of this node; may be invalid.
 	 */
-	virtual void* GetContext() const
+	FBaseTreeNodeWeak GetParentWeak() const
 	{
-		return GroupData->Context;
+		return Parent;
 	}
 
 	/**
-	 * Sets a pointer to a data context for this node. Only available for group nodes.
+	 * @returns a reference to the parent group of this node; may be invalid.
 	 */
-	virtual void SetContext(void* InContext)
+	FBaseTreeNodePtr GetParent() const
 	{
-		if (IsGroup())
-		{
-			GroupData->Context = InContext;
-		}
+		return Parent.Pin();
 	}
 
+	//////////////////////////////////////////////////
+	// Group Node
+
 	/**
 	 * @returns true if this node is a group node.
 	 */
@@ -190,20 +173,29 @@ public:
 		}
 	}
 
-	/**
-	 * @returns the number of children nodes.
-	 */
-	int32 GetChildrenCount() const
+	void RemoveGroupData()
 	{
-		return GroupData->Children.Num();
+		if (GroupData != &DefaultGroupData)
+		{
+			if (GroupData->FilteredChildrenPtr != &GroupData->Children)
+			{
+				delete GroupData->FilteredChildrenPtr;
+			}
+
+			delete GroupData;
+			GroupData = &DefaultGroupData;
+		}
 	}
 
+	//////////////////////////////////////////////////
+	// Children Nodes
+
 	/**
-	 * @returns the number of filtered children nodes.
+	 * @returns the number of children nodes.
 	 */
-	int32 GetFilteredChildrenCount() const
+	int32 GetChildrenCount() const
 	{
-		return GroupData->FilteredChildren.Num();
+		return GroupData->Children.Num();
 	}
 
 	/**
@@ -220,20 +212,6 @@ public:
 		}
 	}
 
-	/**
-	 * Enumerates the filtered children nodes.
-	 */
-	void EnumerateFilteredChildren(TFunction<bool(const FBaseTreeNodePtr&)> Callback) const
-	{
-		for (const FBaseTreeNodePtr& ChildNode : GroupData->FilteredChildren)
-		{
-			if (!Callback(ChildNode))
-			{
-				break;
-			}
-		}
-	}
-
 	/**
 	 * @returns a const reference to the children nodes of this group.
 	 */
@@ -242,86 +220,124 @@ public:
 		return GroupData->Children;
 	}
 
-	/**
-	 * @returns a const reference to the children nodes that should be visible to the UI based on filtering.
-	 */
-	const TArray<FBaseTreeNodePtr>& GetFilteredChildren() const
+	void SortChildren(const ITableCellValueSorter& Sorter, ESortMode SortMode);
+
+	template <typename PredicateType>
+	void SortChildren(PredicateType Predicate)
 	{
-		return GroupData->FilteredChildren;
+		GroupData->Children.Sort(Predicate);
 	}
 
-	/**
-	 * @returns a weak reference to the group of this node, may be invalid.
-	 */
-	FBaseTreeNodeWeak GetGroupPtr() const
+	/** Adds specified node to the children nodes. */
+	void AddChild(const FBaseTreeNodePtr& ChildPtr)
 	{
-		return GroupPtr;
+		GroupData->Children.Add(ChildPtr);
 	}
 
-	/**
-	 * @returns a weak reference to the group of this node, may be invalid.
-	 */
-	FBaseTreeNodePtr GetParentNode() const
+	/** Adds specified node to the children nodes. Also sets the current node as the parent group of the specified node. */
+	void AddChildAndSetParent(const FBaseTreeNodePtr& ChildPtr)
 	{
-		return GroupPtr.Pin();
+		ChildPtr->Parent = AsWeak();
+		GroupData->Children.Add(ChildPtr);
 	}
 
-	virtual bool IsFiltered() const
+	/** Removes this node from the parent's children nodes (preserving the order). Also resets the parent group. */
+	void RemoveFromParent()
 	{
-		return false;
+		FBaseTreeNodePtr ParentNodePtr = Parent.Pin();
+		if (ParentNodePtr.IsValid())
+		{
+			ParentNodePtr->GroupData->Children.RemoveSingle(AsShared());
+		}
+		Parent.Reset();
 	}
 
-	void SortChildren(const ITableCellValueSorter& Sorter, ESortMode SortMode);
-	void SortFilteredChildren(const ITableCellValueSorter& Sorter, ESortMode SortMode);
+	/** Removes this node from the parent's children nodes (not preserving the order). Also resets the parent group. */
+	void RemoveFromParentSwap()
+	{
+		FBaseTreeNodePtr ParentNodePtr = Parent.Pin();
+		if (ParentNodePtr.IsValid())
+		{
+			ParentNodePtr->GroupData->Children.RemoveSwap(AsShared());
+		}
+		Parent.Reset();
+	}
 
-	template <typename PredicateType>
-	void SortChildren(PredicateType Predicate)
+	/** Removes the specified node from the children nodes (preserving the order). If successful, resets the parent group of the specified node. */
+	int32 RemoveChild(const FBaseTreeNodePtr& ChildPtr)
 	{
-		GroupData->Children.Sort(Predicate);
+		if (GroupData->Children.RemoveSingle(ChildPtr))
+		{
+			ChildPtr->Parent.Reset();
+			return 1;
+		}
+		return 0;
 	}
 
-	template <typename PredicateType>
-	void SortFilteredChildren(PredicateType Predicate)
+	/** Removes the specified node from the children nodes (not preserving the order). If successful, resets the parent group of the specified node. */
+	int32 RemoveChildSwap(const FBaseTreeNodePtr& ChildPtr)
 	{
-		GroupData->FilteredChildren.Sort(Predicate);
+		if (GroupData->Children.RemoveSwap(ChildPtr))
+		{
+			ChildPtr->Parent.Reset();
+			return 1;
+		}
+		return 0;
 	}
 
-	/** Adds specified node to the children nodes. Also sets the current node as the group of the specified node. */
-	void AddChildAndSetGroupPtr(const FBaseTreeNodePtr& ChildPtr)
+	void SetParent(const FBaseTreeNodeWeak& InParentWeak)
 	{
-		ChildPtr->GroupPtr = AsWeak();
-		GroupData->Children.Add(ChildPtr);
+		Parent = InParentWeak;
+	}
+
+	void SetParentForChildren()
+	{
+		FBaseTreeNodeWeak ThisNode = AsWeak();
+		for (FBaseTreeNodePtr& NodePtr : GroupData->Children)
+		{
+			NodePtr->Parent = ThisNode;
+		}
 	}
 
-	void SetGroupPtrForAllChildren()
+	void SetParentForChildrenRec()
 	{
 		FBaseTreeNodeWeak ThisNode = AsWeak();
 		for (FBaseTreeNodePtr& NodePtr : GroupData->Children)
 		{
-			NodePtr->GroupPtr = ThisNode;
+			NodePtr->Parent = ThisNode;
+			NodePtr->SetParentForChildrenRec();
+		}
+	}
+
+	void ResetParentForChildren()
+	{
+		for (FBaseTreeNodePtr& NodePtr : GroupData->Children)
+		{
+			NodePtr->Parent.Reset();
 		}
 	}
 
-	void ResetGroupPtrForAllChildren()
+	void ResetParentForChildrenRec()
 	{
 		for (FBaseTreeNodePtr& NodePtr : GroupData->Children)
 		{
-			NodePtr->GroupPtr = nullptr;
+			NodePtr->ResetParentForChildrenRec();
+			NodePtr->Parent.Reset();
 		}
 	}
 
 	/** Clears children. */
 	void ClearChildren(int32 NewSize = 0)
 	{
-		ResetGroupPtrForAllChildren();
+		ResetParentForChildren();
 		GroupData->Children.Reset(NewSize);
 	}
 
 	void SwapChildren(TArray<FBaseTreeNodePtr>& NewChildren)
 	{
-		ResetGroupPtrForAllChildren();
+		ResetParentForChildren();
 		Swap(NewChildren, GroupData->Children);
-		SetGroupPtrForAllChildren();
+		SetParentForChildren();
 	}
 
 	void SwapChildrenFast(TArray<FBaseTreeNodePtr>& NewChildren)
@@ -329,16 +345,86 @@ public:
 		Swap(NewChildren, GroupData->Children);
 	}
 
-	/** Adds specified child to the filtered children nodes. */
-	void AddFilteredChild(const FBaseTreeNodePtr& ChildPtr)
+	//////////////////////////////////////////////////
+	// Filtered Children Nodes
+
+	virtual bool IsFiltered() const
+	{
+		return false;
+	}
+
+	/**
+	 * @returns the number of filtered children nodes.
+	 */
+	int32 GetFilteredChildrenCount() const
+	{
+		return GroupData->FilteredChildrenPtr->Num();
+	}
+
+	/**
+	 * @returns a const reference to the filtered child node at specified index.
+	 */
+	const FBaseTreeNodePtr& GetFilteredChildNode(int32 Index) const
+	{
+		return (*GroupData->FilteredChildrenPtr)[Index];
+	}
+
+	/**
+	 * @returns a reference to the filtered child node at specified index.
+	 */
+	FBaseTreeNodePtr& GetFilteredChildNode(int32 Index)
+	{
+		return (*GroupData->FilteredChildrenPtr)[Index];
+	}
+
+	/**
+	 * Enumerates the filtered children nodes.
+	 */
+	void EnumerateFilteredChildren(TFunction<bool(const FBaseTreeNodePtr&)> Callback) const
+	{
+		for (const FBaseTreeNodePtr& ChildNode : *GroupData->FilteredChildrenPtr)
+		{
+			if (!Callback(ChildNode))
+			{
+				break;
+			}
+		}
+	}
+
+	/**
+	 * @returns a const reference to the children nodes that should be visible to the UI based on filtering.
+	 */
+	const TArray<FBaseTreeNodePtr>& GetFilteredChildren() const
+	{
+		return *GroupData->FilteredChildrenPtr;
+	}
+
+	void SortFilteredChildren(const ITableCellValueSorter& Sorter, ESortMode SortMode);
+
+	template <typename PredicateType>
+	void SortFilteredChildren(PredicateType Predicate)
 	{
-		GroupData->FilteredChildren.Add(ChildPtr);
+		GroupData->FilteredChildrenPtr->Sort(Predicate);
 	}
 
-	/** Clears the filtered children nodes. */
+	/** Clears (removes all) the filtered children nodes. */
 	void ClearFilteredChildren(int32 NewSize = 0)
 	{
-		GroupData->FilteredChildren.Reset(NewSize);
+		if (GroupData->FilteredChildrenPtr == &GroupData->Children)
+		{
+			GroupData->FilteredChildrenPtr = new TArray<FBaseTreeNodePtr>();
+		}
+		GroupData->FilteredChildrenPtr->Reset(NewSize);
+	}
+
+	/** Adds specified child to the filtered children nodes. */
+	void AddFilteredChild(const FBaseTreeNodePtr& ChildPtr)
+	{
+		if (GroupData->FilteredChildrenPtr == &GroupData->Children)
+		{
+			GroupData->FilteredChildrenPtr = new TArray<FBaseTreeNodePtr>();
+		}
+		GroupData->FilteredChildrenPtr->Add(ChildPtr);
 	}
 
 	/**
@@ -347,9 +433,11 @@ public:
 	 */
 	void ResetFilteredChildren()
 	{
-		const int32 Size = GroupData->Children.Num();
-		GroupData->FilteredChildren.Reset();
-		GroupData->FilteredChildren.Append(GroupData->Children);
+		if (GroupData->FilteredChildrenPtr != &GroupData->Children)
+		{
+			delete GroupData->FilteredChildrenPtr;
+			GroupData->FilteredChildrenPtr = &GroupData->Children;
+		}
 	}
 
 	/**
@@ -365,19 +453,38 @@ public:
 		}
 	}
 
+	//////////////////////////////////////////////////
+
 	bool IsExpanded() const
 	{
-		return GroupData->bIsExpanded;
+		return ((uint32)GroupData->Flags & (uint32)EGroupNodeFlags::IsExpanded) != 0;
 	}
 
 	void SetExpansion(bool bOnOff)
 	{
 		if (IsGroup())
 		{
-			GroupData->bIsExpanded = bOnOff;
+			if (bOnOff)
+			{
+				GroupData->Flags = (EGroupNodeFlags)((uint32)GroupData->Flags | (uint32)EGroupNodeFlags::IsExpanded);
+			}
+			else
+			{
+				GroupData->Flags = (EGroupNodeFlags)((uint32)GroupData->Flags & ~(uint32)EGroupNodeFlags::IsExpanded);
+			}
 		}
 	}
 
+	uint32 GetDefaultSortOrder() const
+	{
+		return DefaultSortOrder;
+	}
+
+	void SetDefaultSortOrder(uint32 Order)
+	{
+		DefaultSortOrder = Order;
+	}
+
 protected:
 	/**
 	 * @returns a reference to the children nodes of this group.
@@ -387,11 +494,6 @@ protected:
 		return GroupData->Children;
 	}
 
-	void SetGroupPtr(FBaseTreeNodePtr InGroupPtr)
-	{
-		GroupPtr = InGroupPtr;
-	}
-
 private:
 	/** The default sort order. Index used to optimize sorting. */
 	int32 DefaultSortOrder;
@@ -399,8 +501,8 @@ private:
 	/** The name of this node. */
 	const FName Name;
 
-	/** A weak pointer to the group/parent of this node. */
-	FBaseTreeNodeWeak GroupPtr;
+	/** A weak pointer to the parent group of this node. */
+	FBaseTreeNodeWeak Parent;
 
 	/** The struct containing properties of a group node. It is allocated only for group nodes. */
 	FGroupNodeData* GroupData;
diff --git a/Engine/Source/Developer/TraceInsights/Private/Insights/Table/ViewModels/TreeNodeGrouping.cpp b/Engine/Source/Developer/TraceInsights/Private/Insights/Table/ViewModels/TreeNodeGrouping.cpp
--- a/Engine/Source/Developer/TraceInsights/Private/Insights/Table/ViewModels/TreeNodeGrouping.cpp
+++ b/Engine/Source/Developer/TraceInsights/Private/Insights/Table/ViewModels/TreeNodeGrouping.cpp
@@ -48,7 +48,7 @@ void FTreeNodeGrouping::GroupNodes(const TArray<FTableTreeNodePtr>& Nodes, FTabl
 
 		if (NodePtr->IsGroup())
 		{
-			ParentGroup.AddChildAndSetGroupPtr(NodePtr);
+			ParentGroup.AddChildAndSetParent(NodePtr);
 			continue;
 		}
 
@@ -60,7 +60,7 @@ void FTreeNodeGrouping::GroupNodes(const TArray<FTableTreeNodePtr>& Nodes, FTabl
 		{
 			GroupPtr = MakeShared<FTableTreeNode>(GroupInfo.Name, InParentTable);
 			GroupPtr->SetExpansion(GroupInfo.IsExpanded);
-			ParentGroup.AddChildAndSetGroupPtr(GroupPtr);
+			ParentGroup.AddChildAndSetParent(GroupPtr);
 			GroupMap.Add(GroupInfo.Name, GroupPtr);
 		}
 		else
@@ -68,7 +68,7 @@ void FTreeNodeGrouping::GroupNodes(const TArray<FTableTreeNodePtr>& Nodes, FTabl
 			GroupPtr = *GroupPtrPtr;
 		}
 
-		GroupPtr->AddChildAndSetGroupPtr(NodePtr);
+		GroupPtr->AddChildAndSetParent(NodePtr);
 	}
 }
 
@@ -94,7 +94,7 @@ void FTreeNodeGroupingFlat::GroupNodes(const TArray<FTableTreeNodePtr>& Nodes, F
 
 	FTableTreeNodePtr GroupPtr = MakeShared<FTableTreeNode>(FName(TEXT("All")), InParentTable);
 	GroupPtr->SetExpansion(true);
-	ParentGroup.AddChildAndSetGroupPtr(GroupPtr);
+	ParentGroup.AddChildAndSetParent(GroupPtr);
 
 	GroupPtr->ClearChildren(Nodes.Num());
 	for (FTableTreeNodePtr NodePtr : Nodes)
@@ -104,7 +104,7 @@ void FTreeNodeGroupingFlat::GroupNodes(const TArray<FTableTreeNodePtr>& Nodes, F
 			return;
 		}
 
-		GroupPtr->AddChildAndSetGroupPtr(NodePtr);
+		GroupPtr->AddChildAndSetParent(NodePtr);
 	}
 }
 
@@ -180,7 +180,7 @@ void FTreeNodeGroupingByUniqueValueCString::GroupNodes(const TArray<FTableTreeNo
 
 		if (NodePtr->IsGroup())
 		{
-			ParentGroup.AddChildAndSetGroupPtr(NodePtr);
+			ParentGroup.AddChildAndSetParent(NodePtr);
 			continue;
 		}
 
@@ -202,7 +202,7 @@ void FTreeNodeGroupingByUniqueValueCString::GroupNodes(const TArray<FTableTreeNo
 					{
 						GroupPtr = MakeShared<FTableTreeNode>(GroupName, InParentTable);
 						GroupPtr->SetExpansion(false);
-						ParentGroup.AddChildAndSetGroupPtr(GroupPtr);
+						ParentGroup.AddChildAndSetParent(GroupPtr);
 						GroupNameMap.Add(GroupName, GroupPtr);
 					}
 					else
@@ -222,7 +222,7 @@ void FTreeNodeGroupingByUniqueValueCString::GroupNodes(const TArray<FTableTreeNo
 				{
 					EmptyGroupPtr = MakeShared<FTableTreeNode>(FName(TEXT("N/A")), InParentTable);
 					EmptyGroupPtr->SetExpansion(false);
-					ParentGroup.AddChildAndSetGroupPtr(EmptyGroupPtr);
+					ParentGroup.AddChildAndSetParent(EmptyGroupPtr);
 				}
 				GroupPtr = EmptyGroupPtr;
 			}
@@ -233,12 +233,12 @@ void FTreeNodeGroupingByUniqueValueCString::GroupNodes(const TArray<FTableTreeNo
 			{
 				UnsetGroupPtr = MakeShared<FTableTreeNode>(FName(TEXT("<unset>")), InParentTable);
 				UnsetGroupPtr->SetExpansion(false);
-				ParentGroup.AddChildAndSetGroupPtr(UnsetGroupPtr);
+				ParentGroup.AddChildAndSetParent(UnsetGroupPtr);
 			}
 			GroupPtr = UnsetGroupPtr;
 		}
 
-		GroupPtr->AddChildAndSetGroupPtr(NodePtr);
+		GroupPtr->AddChildAndSetParent(NodePtr);
 	}
 }
 
@@ -343,7 +343,7 @@ void FTreeNodeGroupingByPathBreakdown::GroupNodes(
 
 		if (NodePtr->IsGroup())
 		{
-			ParentGroup.AddChildAndSetGroupPtr(NodePtr);
+			ParentGroup.AddChildAndSetParent(NodePtr);
 			continue;
 		}
 
@@ -382,7 +382,7 @@ void FTreeNodeGroupingByPathBreakdown::GroupNodes(
 						{
 							FTableTreeNodePtr GroupPtr = MakeShared<FTableTreeNode>(DirName, InParentTable);
 							GroupPtr->SetExpansion(false);
-							Directory->Node->AddChildAndSetGroupPtr(GroupPtr);
+							Directory->Node->AddChildAndSetParent(GroupPtr);
 							Directory = new FDirectory(DirName, Directory, GroupPtr.Get());
 							Directory->Parent->ChildrenMap.Add(DirName, Directory);
 							Directories.Add(Directory);
@@ -403,7 +403,7 @@ void FTreeNodeGroupingByPathBreakdown::GroupNodes(
 
 		if (Directory != Root)
 		{
-			Directory->Node->AddChildAndSetGroupPtr(NodePtr);
+			Directory->Node->AddChildAndSetParent(NodePtr);
 		}
 		else
 		{
@@ -412,10 +412,10 @@ void FTreeNodeGroupingByPathBreakdown::GroupNodes(
 				static FName NotAvailableName(TEXT("N/A"));
 				FTableTreeNodePtr NotAvailableNode = MakeShared<FTableTreeNode>(NotAvailableName, InParentTable);
 				NotAvailableNode->SetExpansion(false);
-				ParentGroup.AddChildAndSetGroupPtr(NotAvailableNode);
+				ParentGroup.AddChildAndSetParent(NotAvailableNode);
 				UnsetGroupPtr = NotAvailableNode.Get();
 			}
-			UnsetGroupPtr->AddChildAndSetGroupPtr(NodePtr);
+			UnsetGroupPtr->AddChildAndSetParent(NodePtr);
 		}
 	}
 
diff --git a/Engine/Source/Developer/TraceInsights/Private/Insights/Table/ViewModels/TreeNodeGrouping.h b/Engine/Source/Developer/TraceInsights/Private/Insights/Table/ViewModels/TreeNodeGrouping.h
--- a/Engine/Source/Developer/TraceInsights/Private/Insights/Table/ViewModels/TreeNodeGrouping.h
+++ b/Engine/Source/Developer/TraceInsights/Private/Insights/Table/ViewModels/TreeNodeGrouping.h
@@ -163,7 +163,7 @@ void TTreeNodeGroupingByUniqueValue<Type>::GroupNodes(const TArray<FTableTreeNod
 
 		if (NodePtr->IsGroup())
 		{
-			ParentGroup.AddChildAndSetGroupPtr(NodePtr);
+			ParentGroup.AddChildAndSetParent(NodePtr);
 			continue;
 		}
 
@@ -181,7 +181,7 @@ void TTreeNodeGroupingByUniqueValue<Type>::GroupNodes(const TArray<FTableTreeNod
 				const FName GroupName = GetGroupName(Column, *NodePtr);
 				GroupPtr = MakeShared<FTableTreeNode>(GroupName, InParentTable);
 				GroupPtr->SetExpansion(false);
-				ParentGroup.AddChildAndSetGroupPtr(GroupPtr);
+				ParentGroup.AddChildAndSetParent(GroupPtr);
 				GroupMap.Add(Value, GroupPtr);
 			}
 			else
@@ -195,12 +195,12 @@ void TTreeNodeGroupingByUniqueValue<Type>::GroupNodes(const TArray<FTableTreeNod
 			{
 				UnsetGroupPtr = MakeShared<FTableTreeNode>(FName(TEXT("<unset>")), InParentTable);
 				UnsetGroupPtr->SetExpansion(false);
-				ParentGroup.AddChildAndSetGroupPtr(UnsetGroupPtr);
+				ParentGroup.AddChildAndSetParent(UnsetGroupPtr);
 			}
 			GroupPtr = UnsetGroupPtr;
 		}
 
-		GroupPtr->AddChildAndSetGroupPtr(NodePtr);
+		GroupPtr->AddChildAndSetParent(NodePtr);
 	}
 }
 
diff --git a/Engine/Source/Developer/TraceInsights/Private/Insights/Table/Widgets/STableTreeView.cpp b/Engine/Source/Developer/TraceInsights/Private/Insights/Table/Widgets/STableTreeView.cpp
--- a/Engine/Source/Developer/TraceInsights/Private/Insights/Table/Widgets/STableTreeView.cpp
+++ b/Engine/Source/Developer/TraceInsights/Private/Insights/Table/Widgets/STableTreeView.cpp
@@ -11,6 +11,7 @@
 #include "HAL/PlatformFileManager.h"
 #include "Logging/MessageLog.h"
 #include "SlateOptMacros.h"
+#include "Styling/StyleColors.h"
 #include "Widgets/Input/SSearchBox.h"
 #include "Widgets/Images/SImage.h"
 #include "Widgets/SToolTip.h"
@@ -300,7 +301,7 @@ void STableTreeView::ConstructWidget(TSharedPtr<FTable> InTablePtr)
 					[
 						SNew(STextBlock)
 						.Text(this, &STableTreeView::GetTreeViewBannerText)
-						.ColorAndOpacity(FLinearColor(1.0f, 1.0f, 1.0f, 1.0f))
+						.ColorAndOpacity(FSlateColor(EStyleColor::Foreground))
 					]
 				]
 			]
@@ -1041,16 +1042,6 @@ void STableTreeView::ApplyNodeFiltering()
 		FilteredNodesPtr = &TableRowNodes;
 	}
 
-	TreeViewBannerText = FText::GetEmpty();
-	if (TableRowNodes.Num() == 0)
-	{
-		TreeViewBannerText = LOCTEXT("EmptyTreeMsg", "Tree is empty.");
-	}
-	else if (FilteredNodesPtr->Num() == 0)
-	{
-		TreeViewBannerText = LOCTEXT("AllNodesFiltered", "All tree nodes are filtered out.");
-	}
-
 	Stopwatch.Stop();
 	const double FilteringTime = Stopwatch.GetAccumulatedTime();
 	if (FilteringTime > 0.1)
@@ -1371,6 +1362,8 @@ void STableTreeView::ApplyHierarchyFiltering()
 		TreeView->RequestTreeRefresh();
 	}
 
+	UpdateBannerText();
+
 	Stopwatch.Stop();
 	const double FilteringTime = Stopwatch.GetAccumulatedTime();
 	if (FilteringTime > 0.1)
@@ -1558,11 +1551,11 @@ void STableTreeView::TreeView_OnGetChildren(FTableTreeNodePtr InParent, TArray<F
 
 		// Update aggregation.
 		UpdateAggregatedValuesRec(*InParent);
-		FTableTreeNodePtr Node = StaticCastSharedPtr<FTableTreeNode>(InParent->GetParentNode());
+		FTableTreeNodePtr Node = StaticCastSharedPtr<FTableTreeNode>(InParent->GetParent());
 		while (Node)
 		{
 			UpdateAggregatedValuesSingleNode(*Node);
-			Node = StaticCastSharedPtr<FTableTreeNode>(Node->GetParentNode());
+			Node = StaticCastSharedPtr<FTableTreeNode>(Node->GetParent());
 		}
 
 		// Update sorting.
@@ -3252,21 +3245,15 @@ FTableTreeNodePtr STableTreeView::GetNodeByTableRowIndex(int32 RowIndex) const
 
 void STableTreeView::SelectNodeByTableRowIndex(int32 RowIndex)
 {
-	if (RowIndex >= 0 && RowIndex < TableRowNodes.Num())
+	FTableTreeNodePtr NodePtr = GetNodeByTableRowIndex(RowIndex);
+	if (NodePtr.IsValid())
 	{
-		FTableTreeNodePtr NodePtr = TableRowNodes[RowIndex];
-		ensure(NodePtr);
-
-		if (!NodePtr)
+		// Expand all parent nodes.
+		FBaseTreeNodePtr ParentNode = NodePtr->GetParent();
+		while (ParentNode.IsValid())
 		{
-			return;
-		}
-
-		FBaseTreeNodePtr GroupNode = NodePtr->GetGroupPtr().Pin();
-		while (GroupNode.IsValid())
-		{
-			TreeView->SetItemExpansion(StaticCastSharedPtr<FTableTreeNode>(GroupNode), true);
-			GroupNode = GroupNode->GetGroupPtr().Pin();
+			TreeView->SetItemExpansion(StaticCastSharedPtr<FTableTreeNode>(ParentNode), true);
+			ParentNode = ParentNode->GetParent();
 
 		}
 
@@ -3823,6 +3810,25 @@ void STableTreeView::ExportToFileRec(const FBaseTreeNodePtr& InGroupNode, TArray
 
 ////////////////////////////////////////////////////////////////////////////////////////////////////
 
+void STableTreeView::UpdateBannerText()
+{
+	TreeViewBannerText = FText::GetEmpty();
+	if (TableRowNodes.Num() == 0)
+	{
+		TreeViewBannerText = LOCTEXT("EmptyTable", "This table is empty.");
+	}
+	else if (FilteredNodesPtr->Num() == 0)
+	{
+		TreeViewBannerText = LOCTEXT("AllNodesFilteredOut", "All tree nodes are filtered out. Check the filter configurator.");
+	}
+	else if (Root->GetFilteredChildrenCount() == 0)
+	{
+		TreeViewBannerText = LOCTEXT("HierarchyFilteringZeroResults", "No tree node is matching the current text search.");
+	}
+}
+
+////////////////////////////////////////////////////////////////////////////////////////////////////
+
 FReply STableTreeView::OnApplyViewPreset(const ITableTreeViewPreset* InPreset)
 {
 	ApplyViewPreset(*InPreset);
diff --git a/Engine/Source/Developer/TraceInsights/Private/Insights/Table/Widgets/STableTreeView.h b/Engine/Source/Developer/TraceInsights/Private/Insights/Table/Widgets/STableTreeView.h
--- a/Engine/Source/Developer/TraceInsights/Private/Insights/Table/Widgets/STableTreeView.h
+++ b/Engine/Source/Developer/TraceInsights/Private/Insights/Table/Widgets/STableTreeView.h
@@ -447,6 +447,7 @@ protected:
 	void ExportToFileRec(const FBaseTreeNodePtr& InGroupNode, TArray<FBaseTreeNodePtr>& InNodes, bool bInExportCollapsed, bool InExportLeafs, WriteToFileCallback Callback);
 
 	FText GetTreeViewBannerText() const { return TreeViewBannerText; }
+	virtual void UpdateBannerText();
 
 	////////////////////////////////////////////////////////////////////////////////////////////////////
 	// Presets
diff --git a/Engine/Source/Developer/TraceInsights/Private/Insights/Table/Widgets/STableTreeViewTooltip.cpp b/Engine/Source/Developer/TraceInsights/Private/Insights/Table/Widgets/STableTreeViewTooltip.cpp
--- a/Engine/Source/Developer/TraceInsights/Private/Insights/Table/Widgets/STableTreeViewTooltip.cpp
+++ b/Engine/Source/Developer/TraceInsights/Private/Insights/Table/Widgets/STableTreeViewTooltip.cpp
@@ -106,7 +106,7 @@ TSharedPtr<SToolTip> STableTreeViewTooltip::GetRowTooltip(const TSharedPtr<FTabl
 	TSharedPtr<SGridPanel> GridPanel;
 	TSharedPtr<SHorizontalBox> HBox;
 
-	const FText NodeTooltip = TreeNodePtr->GetTooltip();
+	const FText NodeTooltip = TreeNodePtr->GetTooltipText();
 	const EVisibility NodeTooltipVisibility = NodeTooltip.IsEmpty() ? EVisibility::Collapsed : EVisibility::Visible;
 
 	TSharedPtr<SToolTip> TableCellTooltip =
@@ -199,7 +199,7 @@ TSharedPtr<SToolTip> STableTreeViewTooltip::GetRowTooltip(const TSharedPtr<FTabl
 				[
 					SNew(STextBlock)
 					.Visibility(NodeTooltipVisibility)
-					.WrapTextAt(512.0f)
+					.WrapTextAt(1024.0f)
 					.WrappingPolicy(ETextWrappingPolicy::AllowPerCharacterWrapping)
 					.Text(NodeTooltip)
 					.TextStyle(GetTableTreeViewTooltipNormalTextStyle())
diff --git a/Engine/Source/Developer/TraceInsights/Private/Insights/ViewModels/FilterConfigurator.cpp b/Engine/Source/Developer/TraceInsights/Private/Insights/ViewModels/FilterConfigurator.cpp
--- a/Engine/Source/Developer/TraceInsights/Private/Insights/ViewModels/FilterConfigurator.cpp
+++ b/Engine/Source/Developer/TraceInsights/Private/Insights/ViewModels/FilterConfigurator.cpp
@@ -39,8 +39,8 @@ FFilterConfigurator& FFilterConfigurator::operator=(const FFilterConfigurator& O
 		return *this;
 	}
 
-	RootNode = MakeShared<FFilterConfiguratorNode>(*Other.RootNode);
-	RootNode->SetGroupPtrForChildren();
+	RootNode = FFilterConfiguratorNode::DeepCopy(*Other.RootNode);
+
 	AvailableFilters = Other.AvailableFilters;
 	OnDestroyedEvent = Other.OnDestroyedEvent;
 
diff --git a/Engine/Source/Developer/TraceInsights/Private/Insights/ViewModels/FilterConfiguratorNode.cpp b/Engine/Source/Developer/TraceInsights/Private/Insights/ViewModels/FilterConfiguratorNode.cpp
--- a/Engine/Source/Developer/TraceInsights/Private/Insights/ViewModels/FilterConfiguratorNode.cpp
+++ b/Engine/Source/Developer/TraceInsights/Private/Insights/ViewModels/FilterConfiguratorNode.cpp
@@ -29,32 +29,33 @@ FFilterConfiguratorNode::FFilterConfiguratorNode(const FName InName, bool bInIsG
 
 ////////////////////////////////////////////////////////////////////////////////////////////////////
 
-FFilterConfiguratorNode::FFilterConfiguratorNode(const FFilterConfiguratorNode& Other)
-	: FBaseTreeNode(Other.GetName(), Other.IsGroup())
+TSharedRef<FFilterConfiguratorNode> FFilterConfiguratorNode::DeepCopy(const FFilterConfiguratorNode& InSourceNode)
 {
-	*this = Other;
-}
+	TSharedRef<FFilterConfiguratorNode> NodeCopyPtr = MakeShared<FFilterConfiguratorNode>(InSourceNode.GetName(), InSourceNode.IsGroup());
+	FFilterConfiguratorNode& NodeCopy = *NodeCopyPtr;
 
-////////////////////////////////////////////////////////////////////////////////////////////////////
+	NodeCopy.AvailableFilters = InSourceNode.AvailableFilters;
+	NodeCopy.AvailableFilterOperators = InSourceNode.AvailableFilterOperators;
+	NodeCopy.SelectedFilter = InSourceNode.SelectedFilter;
+	NodeCopy.SelectedFilterOperator = InSourceNode.SelectedFilterOperator;
+	NodeCopy.SelectedFilterGroupOperator = InSourceNode.SelectedFilterGroupOperator;
+	NodeCopy.TextBoxValue = InSourceNode.TextBoxValue;
 
-FFilterConfiguratorNode& FFilterConfiguratorNode::operator=(const FFilterConfiguratorNode& Other)
-{
-	GetChildrenMutable().Empty();
+	NodeCopy.SetExpansion(InSourceNode.IsExpanded());
 
-	AvailableFilters = Other.AvailableFilters;
-	AvailableFilterOperators = Other.AvailableFilterOperators;
-	SelectedFilter = Other.SelectedFilter;
-	SelectedFilterOperator = Other.SelectedFilterOperator;
-	SelectedFilterGroupOperator = Other.SelectedFilterGroupOperator;
-	TextBoxValue = Other.TextBoxValue;
-	SetExpansion(Other.IsExpanded());
-
-	for (FBaseTreeNodePtr Child : Other.GetChildren())
+	if (InSourceNode.IsGroup())
 	{
-		GetChildrenMutable().Add(MakeShared<FFilterConfiguratorNode>(*StaticCastSharedPtr<FFilterConfiguratorNode>(Child)));
+		check(NodeCopy.IsGroup());
+		NodeCopy.ClearChildren(InSourceNode.GetChildrenCount());
+		for (FBaseTreeNodePtr Child : InSourceNode.GetChildren())
+		{
+			check(Child.IsValid() && Child->Is<FFilterConfiguratorNode>());
+			TSharedPtr<FFilterConfiguratorNode> ChildCopy = FFilterConfiguratorNode::DeepCopy(Child->As<FFilterConfiguratorNode>());
+			NodeCopy.AddChildAndSetParent(ChildCopy);
+		}
 	}
 
-	return *this;
+	return NodeCopyPtr;
 }
 
 ////////////////////////////////////////////////////////////////////////////////////////////////////
@@ -90,26 +91,6 @@ const TArray<TSharedPtr<FFilterGroupOperator>>& FFilterConfiguratorNode::GetFilt
 
 ////////////////////////////////////////////////////////////////////////////////////////////////////
 
-void FFilterConfiguratorNode::DeleteChildNode(FFilterConfiguratorNodePtr InNode)
-{
-	FBaseTreeNodePtr Node = StaticCastSharedPtr<FBaseTreeNode>(InNode);
-	GetChildrenMutable().Remove(Node);
-}
-
-////////////////////////////////////////////////////////////////////////////////////////////////////
-
-void FFilterConfiguratorNode::SetGroupPtrForChildren()
-{
-	for (FBaseTreeNodePtr Child : GetChildrenMutable())
-	{
-		FFilterConfiguratorNodePtr CastedChild = StaticCastSharedPtr<FFilterConfiguratorNode>(Child);
-		CastedChild->SetGroupPtrForChildren();
-		CastedChild->SetGroupPtr(SharedThis(this));
-	}
-}
-
-////////////////////////////////////////////////////////////////////////////////////////////////////
-
 void FFilterConfiguratorNode::SetAvailableFilters(TSharedPtr<TArray<TSharedPtr<FFilter>>> InAvailableFilters)
 {
 	AvailableFilters = InAvailableFilters;
diff --git a/Engine/Source/Developer/TraceInsights/Private/Insights/ViewModels/FilterConfiguratorNode.h b/Engine/Source/Developer/TraceInsights/Private/Insights/ViewModels/FilterConfiguratorNode.h
--- a/Engine/Source/Developer/TraceInsights/Private/Insights/ViewModels/FilterConfiguratorNode.h
+++ b/Engine/Source/Developer/TraceInsights/Private/Insights/ViewModels/FilterConfiguratorNode.h
@@ -51,7 +51,7 @@ public:
 	/** Initialization constructor for the filter configurator node. */
 	FFilterConfiguratorNode(const FName InName, bool bInIsGroup);
 
-	FFilterConfiguratorNode(const FFilterConfiguratorNode& Other);
+	static TSharedRef<FFilterConfiguratorNode> DeepCopy(const FFilterConfiguratorNode& Node);
 
 	bool operator==(const FFilterConfiguratorNode& Other) const;
 
@@ -72,13 +72,9 @@ public:
 
 	TSharedPtr<TArray<TSharedPtr<IFilterOperator>>> GetAvailableFilterOperators() const {	return AvailableFilterOperators;	}
 
-	void DeleteChildNode(FFilterConfiguratorNodePtr InNode);
-
 	const FString& GetTextBoxValue() { return TextBoxValue; }
 	void SetTextBoxValue(const FString& InValue) { TextBoxValue = InValue; }
 
-	void SetGroupPtrForChildren();
-
 	bool ApplyFilters(const class FFilterContext& Context) const;
 
 	void GetUsedKeys(TSet<int32>& GetUsedKeys) const;
diff --git a/Engine/Source/Developer/TraceInsights/Private/Insights/Widgets/SFilterConfiguratorRow.cpp b/Engine/Source/Developer/TraceInsights/Private/Insights/Widgets/SFilterConfiguratorRow.cpp
--- a/Engine/Source/Developer/TraceInsights/Private/Insights/Widgets/SFilterConfiguratorRow.cpp
+++ b/Engine/Source/Developer/TraceInsights/Private/Insights/Widgets/SFilterConfiguratorRow.cpp
@@ -239,7 +239,7 @@ TSharedRef<SWidget> SFilterConfiguratorRow::GenerateWidgetForColumn(const FName&
 			];
 
 		// Do not show Delete button for the Root node
-		EVisibility DeleteVisibility = FilterConfiguratorNodePtr->GetGroupPtr().IsValid() ? EVisibility::Visible : EVisibility::Hidden;
+		EVisibility DeleteVisibility = FilterConfiguratorNodePtr->GetParent().IsValid() ? EVisibility::Visible : EVisibility::Hidden;
 		RightBox->AddSlot()
 			.VAlign(VAlign_Center)
 			.HAlign(HAlign_Right)
@@ -406,7 +406,7 @@ FReply SFilterConfiguratorRow::AddFilter_OnClicked()
 	FFilterConfiguratorNodePtr ChildNode = MakeShared<FFilterConfiguratorNode>(TEXT(""), false);
 	ChildNode->SetAvailableFilters(FilterConfiguratorNodePtr->GetAvailableFilters());
 
-	FilterConfiguratorNodePtr->AddChildAndSetGroupPtr(ChildNode);
+	FilterConfiguratorNodePtr->AddChildAndSetParent(ChildNode);
 	FilterConfiguratorNodePtr->SetExpansion(true);
 	OwnerTablePtr.Pin()->Private_SetItemExpansion(FilterConfiguratorNodePtr, true);
 
@@ -418,12 +418,7 @@ FReply SFilterConfiguratorRow::AddFilter_OnClicked()
 
 FReply SFilterConfiguratorRow::DeleteFilter_OnClicked()
 {
-	FBaseTreeNodeWeak GroupWeakPtr = FilterConfiguratorNodePtr->GetGroupPtr();
-	if (GroupWeakPtr.IsValid())
-	{
-		FFilterConfiguratorNodePtr GroupPtr = StaticCastSharedPtr<FFilterConfiguratorNode>(GroupWeakPtr.Pin());
-		GroupPtr->DeleteChildNode(FilterConfiguratorNodePtr);
-	}
+	FilterConfiguratorNodePtr->RemoveFromParent();
 
 	StaticCastSharedPtr<STreeView<FFilterConfiguratorNodePtr>>(OwnerTablePtr.Pin())->RequestTreeRefresh();
 	return FReply::Handled();
@@ -437,7 +432,7 @@ FReply SFilterConfiguratorRow::AddGroup_OnClicked()
 	ChildNode->SetAvailableFilters(FilterConfiguratorNodePtr->GetAvailableFilters());
 	ChildNode->SetExpansion(true);
 
-	FilterConfiguratorNodePtr->AddChildAndSetGroupPtr(ChildNode);
+	FilterConfiguratorNodePtr->AddChildAndSetParent(ChildNode);
 	FilterConfiguratorNodePtr->SetExpansion(true);
 	OwnerTablePtr.Pin()->Private_SetItemExpansion(FilterConfiguratorNodePtr, true);
 
@@ -449,12 +444,7 @@ FReply SFilterConfiguratorRow::AddGroup_OnClicked()
 
 FReply SFilterConfiguratorRow::DeleteGroup_OnClicked()
 {
-	FBaseTreeNodeWeak GroupWeakPtr = FilterConfiguratorNodePtr->GetGroupPtr();
-	if (GroupWeakPtr.IsValid())
-	{
-		FFilterConfiguratorNodePtr GroupPtr = StaticCastSharedPtr<FFilterConfiguratorNode>(GroupWeakPtr.Pin());
-		GroupPtr->DeleteChildNode(FilterConfiguratorNodePtr);
-	}
+	FilterConfiguratorNodePtr->RemoveFromParent();
 
 	StaticCastSharedPtr<STreeView<FFilterConfiguratorNodePtr>>(OwnerTablePtr.Pin())->RequestTreeRefresh();
 	return FReply::Handled();
diff --git a/Engine/Source/Developer/TraceInsights/Private/Insights/Widgets/SStatsView.cpp b/Engine/Source/Developer/TraceInsights/Private/Insights/Widgets/SStatsView.cpp
--- a/Engine/Source/Developer/TraceInsights/Private/Insights/Widgets/SStatsView.cpp
+++ b/Engine/Source/Developer/TraceInsights/Private/Insights/Widgets/SStatsView.cpp
@@ -1278,7 +1278,7 @@ void SStatsView::CreateGroups()
 
 		for (const FStatsNodePtr& NodePtr : StatsNodes)
 		{
-			GroupPtr->AddChildAndSetGroupPtr(NodePtr);
+			GroupPtr->AddChildAndSetParent(NodePtr);
 		}
 		TreeView->SetItemExpansion(GroupPtr, true);
 	}
@@ -1294,7 +1294,7 @@ void SStatsView::CreateGroups()
 			{
 				GroupPtr = GroupNodeSet.Add(GroupName, MakeShared<FStatsNode>(GroupName));
 			}
-			GroupPtr->AddChildAndSetGroupPtr(NodePtr);
+			GroupPtr->AddChildAndSetParent(NodePtr);
 			TreeView->SetItemExpansion(GroupPtr, true);
 		}
 		GroupNodeSet.KeySort([](const FName& A, const FName& B) { return A.Compare(B) < 0; }); // sort groups by name
@@ -1313,7 +1313,7 @@ void SStatsView::CreateGroups()
 				const FName GroupName = *StatsNodeTypeHelper::ToText(NodeType).ToString();
 				GroupPtr = GroupNodeSet.Add(NodeType, MakeShared<FStatsNode>(GroupName));
 			}
-			GroupPtr->AddChildAndSetGroupPtr(NodePtr);
+			GroupPtr->AddChildAndSetParent(NodePtr);
 			TreeView->SetItemExpansion(GroupPtr, true);
 		}
 		GroupNodeSet.KeySort([](const EStatsNodeType& A, const EStatsNodeType& B) { return A < B; }); // sort groups by type
@@ -1332,7 +1332,7 @@ void SStatsView::CreateGroups()
 				const FName GroupName = *StatsNodeDataTypeHelper::ToText(DataType).ToString();
 				GroupPtr = GroupNodeSet.Add(DataType, MakeShared<FStatsNode>(GroupName));
 			}
-			GroupPtr->AddChildAndSetGroupPtr(NodePtr);
+			GroupPtr->AddChildAndSetParent(NodePtr);
 			TreeView->SetItemExpansion(GroupPtr, true);
 		}
 		GroupNodeSet.KeySort([](const EStatsNodeDataType& A, const EStatsNodeDataType& B) { return A < B; }); // sort groups by data type
@@ -1352,7 +1352,7 @@ void SStatsView::CreateGroups()
 				const FName GroupName(FirstLetterStr);
 				GroupPtr = GroupNodeSet.Add(FirstLetter, MakeShared<FStatsNode>(GroupName));
 			}
-			GroupPtr->AddChildAndSetGroupPtr(NodePtr);
+			GroupPtr->AddChildAndSetParent(NodePtr);
 		}
 		GroupNodeSet.KeySort([](const TCHAR& A, const TCHAR& B) { return A < B; }); // sort groups alphabetically
 		GroupNodeSet.GenerateValueArray(GroupNodes);
@@ -1388,7 +1388,7 @@ void SStatsView::CreateGroups()
 				                            FName(FString::Printf(TEXT("Count >= %s"), Orders[MaxOrder - 1]));
 				GroupPtr = GroupNodeSet.Add(Order, MakeShared<FStatsNode>(GroupName));
 			}
-			GroupPtr->AddChildAndSetGroupPtr(NodePtr);
+			GroupPtr->AddChildAndSetParent(NodePtr);
 		}
 		GroupNodeSet.KeySort([](const uint32& A, const uint32& B) { return A > B; }); // sort groups by order
 		GroupNodeSet.GenerateValueArray(GroupNodes);
diff --git a/Engine/Source/Developer/TraceInsights/Private/Insights/Widgets/STimerTreeView.cpp b/Engine/Source/Developer/TraceInsights/Private/Insights/Widgets/STimerTreeView.cpp
--- a/Engine/Source/Developer/TraceInsights/Private/Insights/Widgets/STimerTreeView.cpp
+++ b/Engine/Source/Developer/TraceInsights/Private/Insights/Widgets/STimerTreeView.cpp
@@ -1044,7 +1044,7 @@ FTimerNodePtr STimerTreeView::CreateTimerNodeRec(const TraceServices::FTimingPro
 			FTimerNodePtr ChildTimerNodePtr = CreateTimerNodeRec(*ChildNodePtr);
 			if (ChildTimerNodePtr)
 			{
-				TimerNodePtr->AddChildAndSetGroupPtr(ChildTimerNodePtr);
+				TimerNodePtr->AddChildAndSetParent(ChildTimerNodePtr);
 			}
 		}
 	}
diff --git a/Engine/Source/Developer/TraceInsights/Private/Insights/Widgets/STimersView.cpp b/Engine/Source/Developer/TraceInsights/Private/Insights/Widgets/STimersView.cpp
--- a/Engine/Source/Developer/TraceInsights/Private/Insights/Widgets/STimersView.cpp
+++ b/Engine/Source/Developer/TraceInsights/Private/Insights/Widgets/STimersView.cpp
@@ -1542,7 +1542,7 @@ void STimersView::CreateGroups()
 
 		for (const FTimerNodePtr& NodePtr : TimerNodes)
 		{
-			GroupPtr->AddChildAndSetGroupPtr(NodePtr);
+			GroupPtr->AddChildAndSetParent(NodePtr);
 		}
 		TreeView->SetItemExpansion(GroupPtr, true);
 	}
@@ -1558,7 +1558,7 @@ void STimersView::CreateGroups()
 			{
 				GroupPtr = GroupNodeSet.Add(GroupName, MakeShared<FTimerNode>(GroupName));
 			}
-			GroupPtr->AddChildAndSetGroupPtr(NodePtr);
+			GroupPtr->AddChildAndSetParent(NodePtr);
 			TreeView->SetItemExpansion(GroupPtr, true);
 		}
 		GroupNodeSet.KeySort([](const FName& A, const FName& B) { return A.Compare(B) < 0; }); // sort groups by name
@@ -1577,7 +1577,7 @@ void STimersView::CreateGroups()
 				const FName GroupName = *TimerNodeTypeHelper::ToText(NodeType).ToString();
 				GroupPtr = GroupNodeSet.Add(NodeType, MakeShared<FTimerNode>(GroupName));
 			}
-			GroupPtr->AddChildAndSetGroupPtr(NodePtr);
+			GroupPtr->AddChildAndSetParent(NodePtr);
 			TreeView->SetItemExpansion(GroupPtr, true);
 		}
 		GroupNodeSet.KeySort([](const ETimerNodeType& A, const ETimerNodeType& B) { return A < B; }); // sort groups by type
@@ -1597,7 +1597,7 @@ void STimersView::CreateGroups()
 				const FName GroupName(FirstLetterStr);
 				GroupPtr = GroupNodeSet.Add(FirstLetter, MakeShared<FTimerNode>(GroupName));
 			}
-			GroupPtr->AddChildAndSetGroupPtr(NodePtr);
+			GroupPtr->AddChildAndSetParent(NodePtr);
 		}
 		GroupNodeSet.KeySort([](const TCHAR& A, const TCHAR& B) { return A < B; }); // sort groups alphabetically
 		GroupNodeSet.GenerateValueArray(GroupNodes);
@@ -1633,7 +1633,7 @@ void STimersView::CreateGroups()
 				                            FName(FString::Printf(TEXT("Count >= %s"), Orders[MaxOrder - 1]));
 				GroupPtr = GroupNodeSet.Add(Order, MakeShared<FTimerNode>(GroupName));
 			}
-			GroupPtr->AddChildAndSetGroupPtr(NodePtr);
+			GroupPtr->AddChildAndSetParent(NodePtr);
 		}
 		GroupNodeSet.KeySort([](const uint32& A, const uint32& B) { return A > B; }); // sort groups by order
 		GroupNodeSet.GenerateValueArray(GroupNodes);
