[Insights]
  - Session Info: Added "Project Name" info, if available.
  - Session Info: Added "Other Metadata" info, if available.
  - Session Info: Added "TraceId" as an optional metadata.

diff --git a/Engine/Source/Developer/TraceInsights/Private/Insights/Widgets/SSessionInfoWindow.cpp b/Engine/Source/Developer/TraceInsights/Private/Insights/Widgets/SSessionInfoWindow.cpp
--- a/Engine/Source/Developer/TraceInsights/Private/Insights/Widgets/SSessionInfoWindow.cpp
+++ b/Engine/Source/Developer/TraceInsights/Private/Insights/Widgets/SSessionInfoWindow.cpp
@@ -203,13 +203,18 @@ TSharedRef<SWidget> SSessionInfoWindow::CreateTextBox(const TAttribute<FText>& I
 
 ////////////////////////////////////////////////////////////////////////////////////////////////////
 
-void SSessionInfoWindow::AddInfoLine(TSharedPtr<SVerticalBox> InVerticalBox, const FText& InHeader, const TAttribute<FText>& InValue, bool bMultiLine) const
+void SSessionInfoWindow::AddInfoLine(TSharedPtr<SVerticalBox> InVerticalBox,
+									 const FText& InHeader,
+									 FText(SSessionInfoWindow::* InGetTextMethodPtr)() const,
+									 EVisibility(SSessionInfoWindow::* InVisibilityMethodPtr)() const,
+									 bool bMultiLine) const
 {
 	InVerticalBox->AddSlot()
 		.AutoHeight()
-		.Padding(16.0f, 4.0f, 16.0f, 4.0f)
+		.Padding(16.0f, 0.0f, 16.0f, 0.0f)
 		[
 			SNew(SHorizontalBox)
+			.Visibility(this, InVisibilityMethodPtr)
 
 			+ SHorizontalBox::Slot()
 			.HAlign(HAlign_Right)
@@ -218,7 +223,7 @@ void SSessionInfoWindow::AddInfoLine(TSharedPtr<SVerticalBox> InVerticalBox, con
 			.Padding(0.0f)
 			[
 				SNew(SBox)
-				.Padding(FMargin(0.0f, 4.0f, 16.0f, 4.0f))
+				.Padding(FMargin(0.0f, 8.0f, 16.0f, 4.0f))
 				.HAlign(HAlign_Right)
 				.MinDesiredWidth(160.0f)
 				[
@@ -231,9 +236,9 @@ void SSessionInfoWindow::AddInfoLine(TSharedPtr<SVerticalBox> InVerticalBox, con
 			.HAlign(HAlign_Fill)
 			.VAlign(VAlign_Top)
 			.FillWidth(1.0f)
-			.Padding(0.0f)
+			.Padding(FMargin(0.0f, 4.0f, 0.0f, 4.0f))
 			[
-				CreateTextBox(InValue, bMultiLine)
+				CreateTextBox(TAttribute<FText>(this, InGetTextMethodPtr), bMultiLine)
 			]
 		];
 }
@@ -294,17 +299,20 @@ TSharedRef<SDockTab> SSessionInfoWindow::SpawnTab_SessionInfo(const FSpawnTabArg
 		];
 
 	BeginSection(VerticalBox, LOCTEXT("SessionInfo_SectionText", "Session Info"));
-	AddInfoLine(VerticalBox, LOCTEXT("SessionName_HeaderText",	"Session Name"),		TAttribute<FText>(this, &SSessionInfoWindow::GetSessionNameText));
-	AddInfoLine(VerticalBox, LOCTEXT("Uri_HeaderText",			"URI"),					TAttribute<FText>(this, &SSessionInfoWindow::GetUriText));
-	//AddInfoLine(VerticalBox, LOCTEXT("FileSize_HeaderText",		"File Size"),			TAttribute<FText>(this, &SSessionInfoWindow::GetFileSizeText));
-	AddInfoLine(VerticalBox, LOCTEXT("Platform_HeaderText",		"Platform"),			TAttribute<FText>(this, &SSessionInfoWindow::GetPlatformText));
-	AddInfoLine(VerticalBox, LOCTEXT("AppName_HeaderText",		"Application Name"),	TAttribute<FText>(this, &SSessionInfoWindow::GetAppNameText));
-	AddInfoLine(VerticalBox, LOCTEXT("Branch_HeaderText",		"Branch"),				TAttribute<FText>(this, &SSessionInfoWindow::GetBranchText));
-	AddInfoLine(VerticalBox, LOCTEXT("BuildVersion_HeaderText",	"Build Version"),		TAttribute<FText>(this, &SSessionInfoWindow::GetBuildVersionText));
-	AddInfoLine(VerticalBox, LOCTEXT("Changelist_HeaderText",	"Changelist"),			TAttribute<FText>(this, &SSessionInfoWindow::GetChangelistText));
-	AddInfoLine(VerticalBox, LOCTEXT("BuildConfig_HeaderText",	"Build Config"),		TAttribute<FText>(this, &SSessionInfoWindow::GetBuildConfigText));
-	AddInfoLine(VerticalBox, LOCTEXT("BuildTarget_HeaderText",	"Build Target"),		TAttribute<FText>(this, &SSessionInfoWindow::GetBuildTargetText));
-	AddInfoLine(VerticalBox, LOCTEXT("CommandLine_HeaderText",	"Command Line"),		TAttribute<FText>(this, &SSessionInfoWindow::GetCommandLineText), true);
+
+	AddInfoLine(VerticalBox, LOCTEXT("SessionName_HeaderText",		"Session Name"),		&SSessionInfoWindow::GetSessionNameText,	&SSessionInfoWindow::IsAlwaysVisible);
+	AddInfoLine(VerticalBox, LOCTEXT("Uri_HeaderText",				"URI"),					&SSessionInfoWindow::GetUriText,			&SSessionInfoWindow::IsAlwaysVisible);
+	AddInfoLine(VerticalBox, LOCTEXT("Platform_HeaderText",			"Platform"),			&SSessionInfoWindow::GetPlatformText,		&SSessionInfoWindow::IsVisiblePlatformText);
+	AddInfoLine(VerticalBox, LOCTEXT("AppName_HeaderText",			"Application Name"),	&SSessionInfoWindow::GetAppNameText,		&SSessionInfoWindow::IsVisibleAppNameText);
+	AddInfoLine(VerticalBox, LOCTEXT("ProjectName_HeaderText",		"Project Name"),		&SSessionInfoWindow::GetProjectNameText,	&SSessionInfoWindow::IsVisibleProjectNameText);
+	AddInfoLine(VerticalBox, LOCTEXT("Branch_HeaderText",			"Branch"),				&SSessionInfoWindow::GetBranchText,			&SSessionInfoWindow::IsVisibleBranchText);
+	AddInfoLine(VerticalBox, LOCTEXT("BuildVersion_HeaderText",		"Build Version"),		&SSessionInfoWindow::GetBuildVersionText,	&SSessionInfoWindow::IsVisibleBuildVersionText);
+	AddInfoLine(VerticalBox, LOCTEXT("Changelist_HeaderText",		"Changelist"),			&SSessionInfoWindow::GetChangelistText,		&SSessionInfoWindow::IsVisibleChangelistText);
+	AddInfoLine(VerticalBox, LOCTEXT("BuildConfig_HeaderText",		"Build Config"),		&SSessionInfoWindow::GetBuildConfigText,	&SSessionInfoWindow::IsVisibleBuildConfigText);
+	AddInfoLine(VerticalBox, LOCTEXT("BuildTarget_HeaderText",		"Build Target"),		&SSessionInfoWindow::GetBuildTargetText,	&SSessionInfoWindow::IsVisibleBuildTargetText);
+	AddInfoLine(VerticalBox, LOCTEXT("CommandLine_HeaderText",		"Command Line"),		&SSessionInfoWindow::GetCommandLineText,	&SSessionInfoWindow::IsVisibleCommandLineText, true);
+	AddInfoLine(VerticalBox, LOCTEXT("OtherMetadata_HeaderText",	"Other Metadata"),		&SSessionInfoWindow::GetOtherMetadataText,	&SSessionInfoWindow::IsVisibleOtherMetadataText, true);
+
 	EndSection(VerticalBox);
 
 	BeginSection(VerticalBox, LOCTEXT("AnalysisStatus_SectionText", "Analysis Status"));
@@ -367,12 +375,14 @@ void SSessionInfoWindow::Tick(const FGeometry& AllottedGeometry, const double In
 
 		PlatformText = FText::GetEmpty();
 		AppNameText = FText::GetEmpty();
+		ProjectNameText = FText::GetEmpty();
 		BranchText = FText::GetEmpty();
 		BuildVersionText = FText::GetEmpty();
 		ChangelistText = FText::GetEmpty();
 		BuildConfigurationTypeText = FText::GetEmpty();
 		BuildTargetTypeText = FText::GetEmpty();
 		CommandLineText = FText::GetEmpty();
+		OtherMetadataText = FText::GetEmpty();
 	}
 
 	// If we already have the session info data, we no longer poll for it.
@@ -387,6 +397,7 @@ void SSessionInfoWindow::Tick(const FGeometry& AllottedGeometry, const double In
 			TraceServices::FSessionInfo SessionInfo = DiagnosticsProvider->GetSessionInfo();
 			PlatformText = FText::FromString(SessionInfo.Platform);
 			AppNameText = FText::FromString(SessionInfo.AppName);
+			ProjectNameText = FText::FromString(SessionInfo.ProjectName);
 			BranchText = FText::FromString(SessionInfo.Branch);
 			BuildVersionText = FText::FromString(SessionInfo.BuildVersion);
 			ChangelistText = FText::AsNumber(SessionInfo.Changelist, &FNumberFormattingOptions::DefaultNoGrouping());
@@ -394,6 +405,76 @@ void SSessionInfoWindow::Tick(const FGeometry& AllottedGeometry, const double In
 			BuildTargetTypeText = FText::FromString(LexToString(SessionInfo.TargetType));
 			CommandLineText = FText::FromString(SessionInfo.CommandLine);
 
+			TStringBuilder<1024> OtherMetadata;
+
+			if (Session->GetTraceId() != 0)
+			{
+				OtherMetadata.Appendf(TEXT("TraceId=0x%X"), Session->GetTraceId());
+			}
+
+			Session->EnumerateMetadata([&OtherMetadata](const TraceServices::FTraceSessionMetadata& Metadata)
+				{
+					static FName ExcludedMetadata[]
+					{
+						FName("Platform"),
+						FName("AppName"),
+						FName("ProjectName"),
+						FName("Branch"),
+						FName("BuildVersion"),
+						FName("Changelist"),
+						FName("ConfigurationType"),
+						FName("TargetType"),
+						FName("CommandLine"),
+					};
+
+					for (int32 Index = 0; Index < UE_ARRAY_COUNT(ExcludedMetadata); ++Index)
+					{
+						if (Metadata.Name == ExcludedMetadata[Index])
+						{
+							return;
+						}
+					}
+
+					switch (Metadata.Type)
+					{
+						case TraceServices::FTraceSessionMetadata::EType::Int64:
+						{
+							if (OtherMetadata.Len() > 0)
+							{
+								OtherMetadata.Append(TEXT("\n"));
+							}
+							OtherMetadata.Append(Metadata.Name.GetPlainNameString());
+							OtherMetadata.Append(TEXT("="));
+							OtherMetadata.Appendf(TEXT("%lld"), Metadata.Int64Value);
+							break;
+						}
+						case TraceServices::FTraceSessionMetadata::EType::Double:
+						{
+							if (OtherMetadata.Len() > 0)
+							{
+								OtherMetadata.Append(TEXT("\n"));
+							}
+							OtherMetadata.Append(Metadata.Name.GetPlainNameString());
+							OtherMetadata.Append(TEXT("="));
+							OtherMetadata.Appendf(TEXT("%f"), Metadata.DoubleValue);
+							break;
+						}
+						case TraceServices::FTraceSessionMetadata::EType::String:
+						{
+							if (OtherMetadata.Len() > 0)
+							{
+								OtherMetadata.Append(TEXT("\n"));
+							}
+							OtherMetadata.Append(Metadata.Name.GetPlainNameString());
+							OtherMetadata.Append(TEXT("="));
+							OtherMetadata.Append(Metadata.StringValue);
+							break;
+						}
+					}
+				});
+
+			OtherMetadataText = FText::FromString(FString(OtherMetadata.ToView()));
+
 			bIsSessionInfoSet = true;
 		}
 	}
diff --git a/Engine/Source/Developer/TraceInsights/Private/Insights/Widgets/SSessionInfoWindow.h b/Engine/Source/Developer/TraceInsights/Private/Insights/Widgets/SSessionInfoWindow.h
--- a/Engine/Source/Developer/TraceInsights/Private/Insights/Widgets/SSessionInfoWindow.h
+++ b/Engine/Source/Developer/TraceInsights/Private/Insights/Widgets/SSessionInfoWindow.h
@@ -122,7 +122,11 @@ private:
 	void BeginSection(TSharedPtr<SVerticalBox> InVerticalBox, const FText& InSectionName) const;
 	void EndSection(TSharedPtr<SVerticalBox> InVerticalBox) const;
 	TSharedRef<SWidget> CreateTextBox(const TAttribute<FText>& InText, bool bMultiLine) const;
-	void AddInfoLine(TSharedPtr<SVerticalBox> InVerticalBox, const FText& InHeader, const TAttribute<FText>& InValue, bool bMultiLine = false) const;
+	void AddInfoLine(TSharedPtr<SVerticalBox> InVerticalBox,
+					 const FText& InHeader,
+					 FText(SSessionInfoWindow::* InGetTextMethodPtr)() const,
+					 EVisibility(SSessionInfoWindow::* InVisibilityMethodPtr)() const,
+					 bool bMultiLine = false) const;
 	void AddSimpleInfoLine(TSharedPtr<SVerticalBox> InVerticalBox, const TAttribute<FText>& InValue, bool bMultiLine = false) const;
 
 	TSharedRef<SDockTab> SpawnTab_SessionInfo(const FSpawnTabArgs& Args);
@@ -133,12 +137,27 @@ private:
 
 	FText GetPlatformText() const { return PlatformText; }
 	FText GetAppNameText() const { return AppNameText; }
+	FText GetProjectNameText() const { return ProjectNameText; }
 	FText GetBranchText() const { return BranchText; }
 	FText GetBuildVersionText() const { return BuildVersionText; }
 	FText GetChangelistText() const { return ChangelistText; }
 	FText GetBuildConfigText() const { return BuildConfigurationTypeText; }
 	FText GetBuildTargetText() const { return BuildTargetTypeText; }
 	FText GetCommandLineText() const { return CommandLineText; }
+	FText GetOtherMetadataText() const { return OtherMetadataText; }
+
+	EVisibility IsAlwaysVisible() const { return EVisibility::Visible; }
+
+	EVisibility IsVisiblePlatformText() const { return PlatformText.IsEmpty() ? EVisibility::Collapsed : EVisibility::Visible; }
+	EVisibility IsVisibleAppNameText() const { return AppNameText.IsEmpty() ? EVisibility::Collapsed : EVisibility::Visible; }
+	EVisibility IsVisibleProjectNameText() const { return ProjectNameText.IsEmpty() ? EVisibility::Collapsed : EVisibility::Visible; }
+	EVisibility IsVisibleBranchText() const { return BranchText.IsEmpty() ? EVisibility::Collapsed : EVisibility::Visible; }
+	EVisibility IsVisibleBuildVersionText() const { return BuildVersionText.IsEmpty() ? EVisibility::Collapsed : EVisibility::Visible; }
+	EVisibility IsVisibleChangelistText() const { return ChangelistText.IsEmpty() ? EVisibility::Collapsed : EVisibility::Visible; }
+	EVisibility IsVisibleBuildConfigText() const { return BuildConfigurationTypeText.IsEmpty() ? EVisibility::Collapsed : EVisibility::Visible; }
+	EVisibility IsVisibleBuildTargetText() const { return BuildTargetTypeText.IsEmpty() ? EVisibility::Collapsed : EVisibility::Visible; }
+	EVisibility IsVisibleCommandLineText() const { return CommandLineText.IsEmpty() ? EVisibility::Collapsed : EVisibility::Visible; }
+	EVisibility IsVisibleOtherMetadataText() const { return OtherMetadataText.IsEmpty() ? EVisibility::Collapsed : EVisibility::Visible; }
 
 	FText GetFileSizeText() const;
 	FText GetStatusText() const;
@@ -161,12 +180,14 @@ private:
 
 	FText PlatformText;
 	FText AppNameText;
+	FText ProjectNameText;
 	FText BranchText;
 	FText BuildVersionText;
 	FText ChangelistText;
 	FText BuildConfigurationTypeText;
 	FText BuildTargetTypeText;
 	FText CommandLineText;
+	FText OtherMetadataText;
 
 	TSharedPtr<const TraceServices::IAnalysisSession> AnalysisSession;
 	bool bIsSessionInfoSet = false;
diff --git a/Engine/Source/Developer/TraceServices/Private/AnalysisService.cpp b/Engine/Source/Developer/TraceServices/Private/AnalysisService.cpp
--- a/Engine/Source/Developer/TraceServices/Private/AnalysisService.cpp
+++ b/Engine/Source/Developer/TraceServices/Private/AnalysisService.cpp
@@ -131,6 +131,42 @@ void FAnalysisSession::Wait() const
 	Processor.Wait();
 }
 
+void FAnalysisSession::EnumerateMetadata(TFunctionRef<void(const FTraceSessionMetadata& Metadata)> Callback) const
+{
+	Lock.ReadAccessCheck();
+	for (const auto& KV : Metadata)
+	{
+		Callback(KV.Value);
+	}
+}
+
+void FAnalysisSession::AddMetadata(FName InName, int64 InValue)
+{
+	Lock.WriteAccessCheck();
+	FTraceSessionMetadata& Value = Metadata.Add(InName);
+	Value.Name = InName;
+	Value.Type = FTraceSessionMetadata::EType::Int64;
+	Value.Int64Value = InValue;
+}
+
+void FAnalysisSession::AddMetadata(FName InName, double InValue)
+{
+	Lock.WriteAccessCheck();
+	FTraceSessionMetadata& Value = Metadata.Add(InName);
+	Value.Name = InName;
+	Value.Type = FTraceSessionMetadata::EType::Double;
+	Value.DoubleValue = InValue;
+}
+
+void FAnalysisSession::AddMetadata(FName InName, FString InValue)
+{
+	Lock.WriteAccessCheck();
+	FTraceSessionMetadata& Value = Metadata.Add(InName);
+	Value.Name = InName;
+	Value.Type = FTraceSessionMetadata::EType::String;
+	Value.StringValue = InValue;
+}
+
 FMessageLog* FAnalysisSession::GetLog() const
 {
 	return Processor.GetLog();
diff --git a/Engine/Source/Developer/TraceServices/Private/AnalysisServicePrivate.h b/Engine/Source/Developer/TraceServices/Private/AnalysisServicePrivate.h
--- a/Engine/Source/Developer/TraceServices/Private/AnalysisServicePrivate.h
+++ b/Engine/Source/Developer/TraceServices/Private/AnalysisServicePrivate.h
@@ -37,15 +37,24 @@ class FAnalysisSession
 public:
 	FAnalysisSession(uint32 TraceId, const TCHAR* SessionName, TUniquePtr<UE::Trace::IInDataStream>&& InDataStream);
 	virtual ~FAnalysisSession();
+
 	void Start();
 	virtual void Stop(bool bAndWait) const override;
 	virtual void Wait() const override;
+	virtual bool IsAnalysisComplete() const override { return !Processor.IsActive(); }
 
 	virtual const TCHAR* GetName() const override { return *Name; }
 	virtual uint32 GetTraceId() const override { return TraceId; }
-	virtual bool IsAnalysisComplete() const override { return !Processor.IsActive(); }
+
 	virtual double GetDurationSeconds() const override { Lock.ReadAccessCheck(); return DurationSeconds; }
 	virtual void UpdateDurationSeconds(double Duration) override { Lock.WriteAccessCheck(); DurationSeconds = FMath::Max(Duration, DurationSeconds); }
+
+	virtual uint32 GetMetadataCount() const override { Lock.ReadAccessCheck(); return Metadata.Num(); }
+	virtual void EnumerateMetadata(TFunctionRef<void(const FTraceSessionMetadata& Metadata)> Callback) const override;
+	virtual void AddMetadata(FName InName, int64 InValue) override;
+	virtual void AddMetadata(FName InName, double InValue) override;
+	virtual void AddMetadata(FName InName, FString InValue) override;
+
 	virtual FMessageLog* GetLog() const override;
 
 	virtual ILinearAllocator& GetLinearAllocator() override { return Allocator; }
@@ -72,11 +81,12 @@ private:
 	virtual const IProvider* ReadProviderPrivate(const FName& Name) const override;
 	virtual IEditableProvider* EditProviderPrivate(const FName& Name) override;
 
+private:
 	mutable FAnalysisSessionLock Lock;
-
 	FString Name;
-	uint32 TraceId;
+	uint32 TraceId = 0;
 	double DurationSeconds = 0.0;
+	TMap<FName, FTraceSessionMetadata> Metadata;
 	FSlabAllocator Allocator;
 	FStringStore StringStore;
 	FAnalysisCache Cache;
diff --git a/Engine/Source/Developer/TraceServices/Private/Analyzers/DiagnosticsAnalysis.cpp b/Engine/Source/Developer/TraceServices/Private/Analyzers/DiagnosticsAnalysis.cpp
--- a/Engine/Source/Developer/TraceServices/Private/Analyzers/DiagnosticsAnalysis.cpp
+++ b/Engine/Source/Developer/TraceServices/Private/Analyzers/DiagnosticsAnalysis.cpp
@@ -30,19 +30,11 @@ void FDiagnosticsAnalyzer::OnAnalysisBegin(const FOnAnalysisContext& Context)
 
 bool FDiagnosticsAnalyzer::OnEvent(uint16 RouteId, EStyle Style, const FOnEventContext& Context)
 {
-	if (!Provider)
-	{
-		return false;
-	}
-
 	LLM_SCOPE_BYNAME(TEXT("Insights/FDiagnosticsAnalyzer"));
 
-	FAnalysisSessionEditScope _(Session);
-
 	const auto& EventData = Context.EventData;
-	switch (RouteId)
-	{
-	case RouteId_Session:
+
+	if (RouteId == RouteId_Session)
 	{
 		const uint8* Attachment = EventData.GetAttachment();
 		if (Attachment == nullptr)
@@ -51,6 +43,7 @@ bool FDiagnosticsAnalyzer::OnEvent(uint16 RouteId, EStyle Style, const FOnEventC
 		}
 
 		FSessionInfo SessionInfo;
+
 		uint8 AppNameOffset = EventData.GetValue<uint8>("AppNameOffset");
 		uint8 CommandLineOffset = EventData.GetValue<uint8>("CommandLineOffset");
 
@@ -64,36 +57,87 @@ bool FDiagnosticsAnalyzer::OnEvent(uint16 RouteId, EStyle Style, const FOnEventC
 		int32 CommandLineLength = EventData.GetAttachmentSize() - CommandLineOffset;
 		SessionInfo.CommandLine = FString(CommandLineLength, (const ANSICHAR*)Attachment);
 
-		SessionInfo.ConfigurationType = (EBuildConfiguration) EventData.GetValue<uint8>("ConfigurationType");
-		SessionInfo.TargetType = (EBuildTargetType) EventData.GetValue<uint8>("TargetType");
+		SessionInfo.ConfigurationType = (EBuildConfiguration)EventData.GetValue<uint8>("ConfigurationType");
 
-		Provider->SetSessionInfo(SessionInfo);
+		SessionInfo.TargetType = (EBuildTargetType)EventData.GetValue<uint8>("TargetType");
+
+		FAnalysisSessionEditScope _(Session);
+
+		if (Provider)
+		{
+			Provider->SetSessionInfo(SessionInfo);
+		}
+
+		Session.AddMetadata(FName("Platform"), SessionInfo.Platform);
+		Session.AddMetadata(FName("AppName"), SessionInfo.AppName);
+		Session.AddMetadata(FName("CommandLine"), SessionInfo.CommandLine);
+		UpdateSessionMetadata(EventData);
 
 		return false;
 	}
-	break;
-	case RouteId_Session2:
+
+	if (RouteId == RouteId_Session2)
 	{
 		FSessionInfo SessionInfo;
 
 		EventData.GetString("Platform", SessionInfo.Platform);
 		EventData.GetString("AppName", SessionInfo.AppName);
-		EventData.GetString("ProjectName", SessionInfo.ProjectName);	
+		EventData.GetString("ProjectName", SessionInfo.ProjectName);
 		EventData.GetString("CommandLine", SessionInfo.CommandLine);
 		EventData.GetString("Branch", SessionInfo.Branch);
 		EventData.GetString("BuildVersion", SessionInfo.BuildVersion);
 		SessionInfo.Changelist = EventData.GetValue<uint32>("Changelist", 0);
-		SessionInfo.ConfigurationType = (EBuildConfiguration) EventData.GetValue<uint8>("ConfigurationType");
-		SessionInfo.TargetType = (EBuildTargetType) EventData.GetValue<uint8>("TargetType");
+		SessionInfo.ConfigurationType = (EBuildConfiguration)EventData.GetValue<uint8>("ConfigurationType");
+		SessionInfo.TargetType = (EBuildTargetType)EventData.GetValue<uint8>("TargetType");
 
+		FAnalysisSessionEditScope _(Session);
 		Provider->SetSessionInfo(SessionInfo);
+		UpdateSessionMetadata(EventData);
 
 		return false;
-	};
-	break;
 	}
 
 	return true;
 }
 
+void FDiagnosticsAnalyzer::UpdateSessionMetadata(const UE::Trace::IAnalyzer::FEventData& EventData)
+{
+	const UE::Trace::IAnalyzer::FEventTypeInfo& TypeInfo = EventData.GetTypeInfo();
+	const uint32 FieldCount = TypeInfo.GetFieldCount();
+	for (uint32 FieldIndex = 0; FieldIndex < FieldCount; ++FieldIndex)
+	{
+		const UE::Trace::IAnalyzer::FEventFieldInfo* FieldInfo = TypeInfo.GetFieldInfo(FieldIndex);
+		if (!FieldInfo)
+		{
+			continue;
+		}
+		switch (FieldInfo->GetType())
+		{
+			case UE::Trace::IAnalyzer::FEventFieldInfo::EType::Integer:
+			{
+				FName FieldName(FieldInfo->GetName());
+				int64 Value = EventData.GetValue<int64>(FieldInfo->GetName());
+				Session.AddMetadata(FieldName, Value);
+				break;
+			}
+			case UE::Trace::IAnalyzer::FEventFieldInfo::EType::Float:
+			{
+				FName FieldName(FieldInfo->GetName());
+				double Value = EventData.GetValue<double>(FieldInfo->GetName());
+				Session.AddMetadata(FieldName, Value);
+				break;
+			}
+			case UE::Trace::IAnalyzer::FEventFieldInfo::EType::AnsiString:
+			case UE::Trace::IAnalyzer::FEventFieldInfo::EType::WideString:
+			{
+				FName FieldName(FieldInfo->GetName());
+				FString Value;
+				EventData.GetString(FieldInfo->GetName(), Value);
+				Session.AddMetadata(FieldName, Value);
+				break;
+			}
+		}
+	}
+}
+
 } // namespace TraceServices
diff --git a/Engine/Source/Developer/TraceServices/Private/Analyzers/DiagnosticsAnalysis.h b/Engine/Source/Developer/TraceServices/Private/Analyzers/DiagnosticsAnalysis.h
--- a/Engine/Source/Developer/TraceServices/Private/Analyzers/DiagnosticsAnalysis.h
+++ b/Engine/Source/Developer/TraceServices/Private/Analyzers/DiagnosticsAnalysis.h
@@ -18,15 +18,19 @@ class FDiagnosticsAnalyzer
 public:
 	FDiagnosticsAnalyzer(IAnalysisSession& Session, FDiagnosticsProvider* InProvider);
 	~FDiagnosticsAnalyzer();
+
 	virtual void OnAnalysisBegin(const FOnAnalysisContext& Context) override;
 	virtual bool OnEvent(uint16 RouteId, EStyle Style, const FOnEventContext& Context) override;
 
 private:
+	void UpdateSessionMetadata(const UE::Trace::IAnalyzer::FEventData& EventData);
+
 	enum : uint16
 	{
 		RouteId_Session,
 		RouteId_Session2,
 	};
+
 	IAnalysisSession& Session;
 	FDiagnosticsProvider* Provider;
 };
diff --git a/Engine/Source/Developer/TraceServices/Private/Model/DiagnosticsPrivate.h b/Engine/Source/Developer/TraceServices/Private/Model/DiagnosticsPrivate.h
--- a/Engine/Source/Developer/TraceServices/Private/Model/DiagnosticsPrivate.h
+++ b/Engine/Source/Developer/TraceServices/Private/Model/DiagnosticsPrivate.h
@@ -1,4 +1,5 @@
 // Copyright Epic Games, Inc. All Rights Reserved.
+
 #pragma once
 
 #include "TraceServices/Model/Diagnostics.h"
@@ -12,10 +13,9 @@ public:
 	explicit FDiagnosticsProvider(IAnalysisSession& Session);
 	virtual ~FDiagnosticsProvider() {}
 
-public:
-	void SetSessionInfo(const FSessionInfo& InSessionInfo);
-	virtual const FSessionInfo& GetSessionInfo() const override;
 	virtual bool IsSessionInfoAvailable() const override;
+	virtual const FSessionInfo& GetSessionInfo() const override;
+	void SetSessionInfo(const FSessionInfo& InSessionInfo);
 
 private:
 	IAnalysisSession& Session;
diff --git a/Engine/Source/Developer/TraceServices/Public/TraceServices/Model/AnalysisSession.h b/Engine/Source/Developer/TraceServices/Public/TraceServices/Model/AnalysisSession.h
--- a/Engine/Source/Developer/TraceServices/Public/TraceServices/Model/AnalysisSession.h
+++ b/Engine/Source/Developer/TraceServices/Public/TraceServices/Model/AnalysisSession.h
@@ -44,6 +44,25 @@ public:
 	virtual void EditAccessCheck() const { unimplemented(); }
 };
 
+struct FTraceSessionMetadata
+{
+	enum class EType
+	{
+		Int64,
+		Double,
+		String
+	};
+
+	FName Name;
+	EType Type;
+	union
+	{
+		int64 Int64Value;
+		double DoubleValue;
+	};
+	FString StringValue;
+};
+
 class IAnalysisSession
 {
 public:
@@ -54,6 +73,7 @@ public:
 	virtual const TCHAR* GetName() const = 0;
 	virtual uint32 GetTraceId() const = 0;
 	virtual bool IsAnalysisComplete() const = 0;
+
 	virtual double GetDurationSeconds() const = 0;
 	/**
 	 * Update the internal estimation of the session duration with a new timestamp.
@@ -61,6 +81,13 @@ public:
 	 * estimated.
 	 */
 	virtual void UpdateDurationSeconds(double Duration) = 0;
+
+	virtual uint32 GetMetadataCount() const = 0;
+	virtual void EnumerateMetadata(TFunctionRef<void(const FTraceSessionMetadata& Metadata)> Callback) const = 0;
+	virtual void AddMetadata(FName InName, int64 InValue) = 0;
+	virtual void AddMetadata(FName InName, double InValue) = 0;
+	virtual void AddMetadata(FName InName, FString InValue) = 0;
+
 	virtual FMessageLog* GetLog() const = 0;
 
 	virtual void BeginRead() const = 0;
