[Insights] Frame stats specific column tooltips for TimersView.

diff --git a/Engine/Source/Developer/TraceInsights/Private/Insights/ViewModels/TimersViewColumnFactory.cpp b/Engine/Source/Developer/TraceInsights/Private/Insights/ViewModels/TimersViewColumnFactory.cpp
--- a/Engine/Source/Developer/TraceInsights/Private/Insights/ViewModels/TimersViewColumnFactory.cpp
+++ b/Engine/Source/Developer/TraceInsights/Private/Insights/ViewModels/TimersViewColumnFactory.cpp
@@ -81,12 +81,16 @@ TSharedRef<Insights::FTableColumn> FTimersViewColumnFactory::CreateNameColumn()
 {
 	using namespace Insights;
 
-	TSharedRef<FTableColumn> ColumnRef = MakeShared<FTableColumn>(FTimersViewColumns::NameColumnID);
-	FTableColumn& Column = *ColumnRef;
+	TSharedRef<FTimersTableColumn> ColumnRef = MakeShared<FTimersTableColumn>(FTimersViewColumns::NameColumnID);
+	FTimersTableColumn& Column = *ColumnRef;
 
 	Column.SetShortName(LOCTEXT("Name_ColumnName", "Name"));
 	Column.SetTitleName(LOCTEXT("Name_ColumnTitle", "Timer or Group Name"));
-	Column.SetDescription(LOCTEXT("Name_ColumnDesc", "Name of timer or group"));
+
+	FText Description = LOCTEXT("Name_ColumnDesc", "Name of timer or group");
+	Column.SetDescription(ETraceFrameType::TraceFrameType_Count, Description);
+	Column.SetDescription(ETraceFrameType::TraceFrameType_Game, Description);
+	Column.SetDescription(ETraceFrameType::TraceFrameType_Rendering, Description);
 
 	Column.SetFlags(ETableColumnFlags::ShouldBeVisible |
 					ETableColumnFlags::CanBeFiltered |
@@ -116,12 +120,16 @@ TSharedRef<Insights::FTableColumn> FTimersViewColumnFactory::CreateMetaGroupName
 {
 	using namespace Insights;
 
-	TSharedRef<FTableColumn> ColumnRef = MakeShared<FTableColumn>(FTimersViewColumns::MetaGroupNameColumnID);
-	FTableColumn& Column = *ColumnRef;
+	TSharedRef<FTimersTableColumn> ColumnRef = MakeShared<FTimersTableColumn>(FTimersViewColumns::MetaGroupNameColumnID);
+	FTimersTableColumn& Column = *ColumnRef;
 
 	Column.SetShortName(LOCTEXT("MetaGroupName_ColumnName", "Meta Group"));
 	Column.SetTitleName(LOCTEXT("MetaGroupName_ColumnTitle", "Meta Group Name"));
-	Column.SetDescription(LOCTEXT("MetaGroupName_ColumnDesc", "Name of the meta group"));
+
+	FText Description = LOCTEXT("MetaGroupName_ColumnDesc", "Name of the meta group");
+	Column.SetDescription(ETraceFrameType::TraceFrameType_Count, Description);
+	Column.SetDescription(ETraceFrameType::TraceFrameType_Game, Description);
+	Column.SetDescription(ETraceFrameType::TraceFrameType_Rendering, Description);
 
 	Column.SetFlags(ETableColumnFlags::CanBeHidden |
 					//ETableColumnFlags::ShouldBeVisible |
@@ -161,12 +169,16 @@ TSharedRef<Insights::FTableColumn> FTimersViewColumnFactory::CreateTypeColumn()
 {
 	using namespace Insights;
 
-	TSharedRef<FTableColumn> ColumnRef = MakeShared<FTableColumn>(FTimersViewColumns::TypeColumnID);
-	FTableColumn& Column = *ColumnRef;
+	TSharedRef<FTimersTableColumn> ColumnRef = MakeShared<FTimersTableColumn>(FTimersViewColumns::TypeColumnID);
+	FTimersTableColumn& Column = *ColumnRef;
 
 	Column.SetShortName(LOCTEXT("Type_ColumnName", "Type"));
 	Column.SetTitleName(LOCTEXT("Type_ColumnTitle", "Type"));
-	Column.SetDescription(LOCTEXT("Type_ColumnDesc", "Type of timer or group"));
+
+	FText Description = LOCTEXT("Type_ColumnDesc", "Type of timer or group");
+	Column.SetDescription(ETraceFrameType::TraceFrameType_Count, Description);
+	Column.SetDescription(ETraceFrameType::TraceFrameType_Game, Description);
+	Column.SetDescription(ETraceFrameType::TraceFrameType_Rendering, Description);
 
 	Column.SetFlags(ETableColumnFlags::CanBeHidden |
 					//ETableColumnFlags::ShouldBeVisible |
@@ -207,12 +219,16 @@ TSharedRef<Insights::FTableColumn> FTimersViewColumnFactory::CreateInstanceCount
 {
 	using namespace Insights;
 
-	TSharedRef<FTableColumn> ColumnRef = MakeShared<FTableColumn>(FTimersViewColumns::InstanceCountColumnID);
-	FTableColumn& Column = *ColumnRef;
+	TSharedRef<FTimersTableColumn> ColumnRef = MakeShared<FTimersTableColumn>(FTimersViewColumns::InstanceCountColumnID);
+	FTimersTableColumn& Column = *ColumnRef;
 
 	Column.SetShortName(LOCTEXT("InstanceCount_ColumnName", "Count"));
 	Column.SetTitleName(LOCTEXT("InstanceCount_ColumnTitle", "Instance Count"));
-	Column.SetDescription(LOCTEXT("InstanceCount_ColumnDesc", "Number of selected timer's instances"));
+
+	FText Description = LOCTEXT("InstanceCount_ColumnDesc", "Number of selected timer's instances");
+	Column.SetDescription(ETraceFrameType::TraceFrameType_Count, Description);
+	Column.SetDescription(ETraceFrameType::TraceFrameType_Game, Description);
+	Column.SetDescription(ETraceFrameType::TraceFrameType_Rendering, Description);
 
 	Column.SetFlags(ETableColumnFlags::CanBeHidden |
 					ETableColumnFlags::ShouldBeVisible |
@@ -256,12 +272,16 @@ TSharedRef<Insights::FTableColumn> FTimersViewColumnFactory::CreateTotalInclusiv
 {
 	using namespace Insights;
 
-	TSharedRef<FTableColumn> ColumnRef = MakeShared<FTableColumn>(FTimersViewColumns::TotalInclusiveTimeColumnID);
-	FTableColumn& Column = *ColumnRef;
+	TSharedRef<FTimersTableColumn> ColumnRef = MakeShared<FTimersTableColumn>(FTimersViewColumns::TotalInclusiveTimeColumnID);
+	FTimersTableColumn& Column = *ColumnRef;
 
 	Column.SetShortName(LOCTEXT("TotalInclusiveTime_ColumnName", "Incl"));
 	Column.SetTitleName(LOCTEXT("TotalInclusiveTime_ColumnTitle", "Total Inclusive Time"));
-	Column.SetDescription(LOCTEXT("TotalInclusiveTime_ColumnDesc", "Total inclusive duration of selected timer's instances"));
+
+	FText Description = LOCTEXT("TotalInclusiveTime_ColumnDesc", "Total inclusive duration of selected timer's instances");
+	Column.SetDescription(ETraceFrameType::TraceFrameType_Count, Description);
+	Column.SetDescription(ETraceFrameType::TraceFrameType_Game, Description);
+	Column.SetDescription(ETraceFrameType::TraceFrameType_Rendering, Description);
 
 	Column.SetFlags(ETableColumnFlags::CanBeHidden |
 					ETableColumnFlags::ShouldBeVisible |
@@ -304,12 +324,15 @@ TSharedRef<Insights::FTableColumn> FTimersViewColumnFactory::CreateMaxInclusiveT
 {
 	using namespace Insights;
 
-	TSharedRef<FTableColumn> ColumnRef = MakeShared<FTableColumn>(FTimersViewColumns::MaxInclusiveTimeColumnID);
-	FTableColumn& Column = *ColumnRef;
+	TSharedRef<FTimersTableColumn> ColumnRef = MakeShared<FTimersTableColumn>(FTimersViewColumns::MaxInclusiveTimeColumnID);
+	FTimersTableColumn& Column = *ColumnRef;
 
 	Column.SetShortName(LOCTEXT("MaxInclusiveTime_ColumnName", "I.Max"));
 	Column.SetTitleName(LOCTEXT("MaxInclusiveTime_ColumnTitle", "Max Inclusive Time (ms)"));
-	Column.SetDescription(LOCTEXT("MaxInclusiveTime_ColumnDesc", "Maximum inclusive duration of selected timer's instances, in milliseconds"));
+
+	Column.SetDescription(ETraceFrameType::TraceFrameType_Count, LOCTEXT("MaxInclusiveTime_ColumnDesc", "Maximum inclusive duration of selected timer's instances, in milliseconds"));
+	Column.SetDescription(ETraceFrameType::TraceFrameType_Game, LOCTEXT("MaxInclusiveTime_GameFrameColumnDesc", "Game Frame Maximum Inclusive Duration.\nInclusive duration is computed for a single frame as the sum of inclusive duration of all instances of the timer in the respective frame.\nThe maximum is selected from these per-frame inclusive durations. Unit is miliseconds."));
+	Column.SetDescription(ETraceFrameType::TraceFrameType_Rendering, LOCTEXT("MaxInclusiveTime_RenderingFrameColumnDesc", "Rendering Frame Maximum Inclusive Duration.\nInclusive duration is computed for a single frame as the sum of inclusive duration of all instances of the timer in the respective frame.\nThe maximum is selected from these per-frame inclusive durations. Unit is miliseconds."));
 
 	Column.SetFlags(ETableColumnFlags::CanBeHidden |
 					//ETableColumnFlags::ShouldBeVisible |
@@ -350,12 +373,15 @@ TSharedRef<Insights::FTableColumn> FTimersViewColumnFactory::CreateAverageInclus
 {
 	using namespace Insights;
 
-	TSharedRef<FTableColumn> ColumnRef = MakeShared<FTableColumn>(FTimersViewColumns::AverageInclusiveTimeColumnID);
-	FTableColumn& Column = *ColumnRef;
+	TSharedRef<FTimersTableColumn> ColumnRef = MakeShared<FTimersTableColumn>(FTimersViewColumns::AverageInclusiveTimeColumnID);
+	FTimersTableColumn& Column = *ColumnRef;
 
 	Column.SetShortName(LOCTEXT("AvgInclusiveTime_ColumnName", "I.Avg"));
 	Column.SetTitleName(LOCTEXT("AvgInclusiveTime_ColumnTitle", "Average Inclusive Time (ms)"));
-	Column.SetDescription(LOCTEXT("AvgInclusiveTime_ColumnDesc", "Average inclusive duration of selected timer's instances, in milliseconds"));
+
+	Column.SetDescription(ETraceFrameType::TraceFrameType_Count, LOCTEXT("AvgInclusiveTime_ColumnDesc", "Average inclusive duration of selected timer's instances, in milliseconds"));
+	Column.SetDescription(ETraceFrameType::TraceFrameType_Game, LOCTEXT("AvgInclusiveTime_GameFrameColumnDesc", "Game Frame Average Inclusive Duration.\nInclusive duration is computed for a single frame as the sum of inclusive duration of all instances of the timer in the respective frame.\nThe average is applied for these per-frame inclusive durations. Unit is miliseconds."));
+	Column.SetDescription(ETraceFrameType::TraceFrameType_Rendering, LOCTEXT("AvgInclusiveTime_RenderingFrameColumnDesc", "Rendering Frame Average Inclusive Duration.\nInclusive duration is computed for a single frame as the sum of inclusive duration of all instances of the timer in the respective frame.\nThe average is applied for these per-frame inclusive durations. Unit is miliseconds."));
 
 	Column.SetFlags(ETableColumnFlags::CanBeHidden |
 					//ETableColumnFlags::ShouldBeVisible |
@@ -396,12 +422,15 @@ TSharedRef<Insights::FTableColumn> FTimersViewColumnFactory::CreateMedianInclusi
 {
 	using namespace Insights;
 
-	TSharedRef<FTableColumn> ColumnRef = MakeShared<FTableColumn>(FTimersViewColumns::MedianInclusiveTimeColumnID);
-	FTableColumn& Column = *ColumnRef;
+	TSharedRef<FTimersTableColumn> ColumnRef = MakeShared<FTimersTableColumn>(FTimersViewColumns::MedianInclusiveTimeColumnID);
+	FTimersTableColumn& Column = *ColumnRef;
 
 	Column.SetShortName(LOCTEXT("MedInclusiveTime_ColumnName", "I.Med"));
 	Column.SetTitleName(LOCTEXT("MedInclusiveTime_ColumnTitle", "Median Inclusive Time (ms)"));
-	Column.SetDescription(LOCTEXT("MedInclusiveTime_ColumnDesc", "Median inclusive duration of selected timer's instances, in milliseconds"));
+
+	Column.SetDescription(ETraceFrameType::TraceFrameType_Count, LOCTEXT("MedInclusiveTime_ColumnDesc", "Median inclusive duration of selected timer's instances, in milliseconds"));
+	Column.SetDescription(ETraceFrameType::TraceFrameType_Game, LOCTEXT("MedInclusiveTime_GameFrameColumnDesc", "Game Frame Median Inclusive Duration.\nInclusive duration is computed for a single frame as the sum of inclusive duration of all instances of the timer in the respective frame.\nThe median is aproximated these per-frame inclusive durations. Unit is miliseconds."));
+	Column.SetDescription(ETraceFrameType::TraceFrameType_Rendering, LOCTEXT("MedInclusiveTime_RenderingFrameColumnDesc", "Rendering Frame Median Inclusive Duration.\nInclusive duration is computed for a single frame as the sum of inclusive duration of all instances of the timer in the respective frame.The median is aproximated these per-frame inclusive durations.\nUnit is miliseconds."));
 
 	Column.SetFlags(ETableColumnFlags::CanBeHidden |
 					//ETableColumnFlags::ShouldBeVisible |
@@ -442,12 +471,15 @@ TSharedRef<Insights::FTableColumn> FTimersViewColumnFactory::CreateMinInclusiveT
 {
 	using namespace Insights;
 
-	TSharedRef<FTableColumn> ColumnRef = MakeShared<FTableColumn>(FTimersViewColumns::MinInclusiveTimeColumnID);
-	FTableColumn& Column = *ColumnRef;
+	TSharedRef<FTimersTableColumn> ColumnRef = MakeShared<FTimersTableColumn>(FTimersViewColumns::MinInclusiveTimeColumnID);
+	FTimersTableColumn& Column = *ColumnRef;
 
 	Column.SetShortName(LOCTEXT("MinInclusiveTime_ColumnName", "I.Min"));
 	Column.SetTitleName(LOCTEXT("MinInclusiveTime_ColumnTitle", "Min Inclusive Time (ms)"));
-	Column.SetDescription(LOCTEXT("MinInclusiveTime_ColumnDesc", "Minimum inclusive duration of selected timer's instances, in milliseconds"));
+
+	Column.SetDescription(ETraceFrameType::TraceFrameType_Count, LOCTEXT("MinInclusiveTime_ColumnDesc", "Minimum inclusive duration of selected timer's instances, in milliseconds"));
+	Column.SetDescription(ETraceFrameType::TraceFrameType_Game, LOCTEXT("MinInclusiveTime_GameFrameColumnDesc", "Game Frame Minimum Inclusive Duration.\nInclusive duration is computed for a single frame as the sum of inclusive duration of all instances of the timer in the respective frame.\nThe minimum is selectedfrom these per-frame inclusive durations. Unit is miliseconds."));
+	Column.SetDescription(ETraceFrameType::TraceFrameType_Rendering, LOCTEXT("MinInclusiveTime_RenderingFrameColumnDesc", "Rendering Frame Minimum Inclusive Duration.\nInclusive duration is computed for a single frame as the sum of inclusive duration of all instances of the timer in the respective frame.The minimum is selectedfrom these per-frame inclusive durations.\nUnit is miliseconds."));
 
 	Column.SetFlags(ETableColumnFlags::CanBeHidden |
 					//ETableColumnFlags::ShouldBeVisible |
@@ -490,12 +522,16 @@ TSharedRef<Insights::FTableColumn> FTimersViewColumnFactory::CreateTotalExclusiv
 {
 	using namespace Insights;
 
-	TSharedRef<FTableColumn> ColumnRef = MakeShared<FTableColumn>(FTimersViewColumns::TotalExclusiveTimeColumnID);
-	FTableColumn& Column = *ColumnRef;
+	TSharedRef<FTimersTableColumn> ColumnRef = MakeShared<FTimersTableColumn>(FTimersViewColumns::TotalExclusiveTimeColumnID);
+	FTimersTableColumn& Column = *ColumnRef;
 
 	Column.SetShortName(LOCTEXT("TotalExclusiveTime_ColumnName", "Excl"));
 	Column.SetTitleName(LOCTEXT("TotalExclusiveTime_ColumnTitle", "Total Exclusive Time"));
-	Column.SetDescription(LOCTEXT("TotalExclusiveTime_ColumnDesc", "Total exclusive duration of selected timer's instances"));
+
+	FText Description = LOCTEXT("TotalExclusiveTime_ColumnDesc", "Total exclusive duration of selected timer's instances");
+	Column.SetDescription(ETraceFrameType::TraceFrameType_Count, Description);
+	Column.SetDescription(ETraceFrameType::TraceFrameType_Game, Description);
+	Column.SetDescription(ETraceFrameType::TraceFrameType_Rendering, Description);
 
 	Column.SetFlags(ETableColumnFlags::CanBeHidden |
 					ETableColumnFlags::ShouldBeVisible |
@@ -538,13 +574,16 @@ TSharedRef<Insights::FTableColumn> FTimersViewColumnFactory::CreateMaxExclusiveT
 {
 	using namespace Insights;
 
-	TSharedRef<FTableColumn> ColumnRef = MakeShared<FTableColumn>(FTimersViewColumns::MaxExclusiveTimeColumnID);
-	FTableColumn& Column = *ColumnRef;
+	TSharedRef<FTimersTableColumn> ColumnRef = MakeShared<FTimersTableColumn>(FTimersViewColumns::MaxExclusiveTimeColumnID);
+	FTimersTableColumn& Column = *ColumnRef;
 
 	Column.SetShortName(LOCTEXT("MaxExclusiveTime_ColumnName", "E.Max"));
 	Column.SetTitleName(LOCTEXT("MaxExclusiveTime_ColumnTitle", "Max Exclusive Time (ms)"));
-	Column.SetDescription(LOCTEXT("MaxExclusiveTime_ColumnDesc", "Maximum exclusive duration of selected timer's instances, in milliseconds"));
 
+	Column.SetDescription(ETraceFrameType::TraceFrameType_Count, LOCTEXT("MaxExclusiveTime_ColumnDesc", "Maximum exclusive duration of selected timer's instances, in milliseconds"));
+	Column.SetDescription(ETraceFrameType::TraceFrameType_Game, LOCTEXT("MaxExclusiveTime_GameFrameColumnDesc", "Game Frame Maximum Exclusive Duration.\nExclusive duration is computed for a single frame as the sum of exclusive duration of all instances of the timer in the respective frame.\nThe maximum is selected from these per-frame Exclusive durations. Unit is miliseconds."));
+	Column.SetDescription(ETraceFrameType::TraceFrameType_Rendering, LOCTEXT("MaxExclusiveTime_RenderingFrameColumnDesc", "Rendering Frame Maximum Exclusive Duration.\nExclusive duration is computed for a single frame as the sum of exclusive duration of all instances of the timer in the respective frame. The maximum is selected from these per-frame Exclusive durations.\nUnit is miliseconds."));
+	
 	Column.SetFlags(ETableColumnFlags::CanBeHidden |
 					//ETableColumnFlags::ShouldBeVisible |
 					ETableColumnFlags::CanBeFiltered);
@@ -584,12 +623,15 @@ TSharedRef<Insights::FTableColumn> FTimersViewColumnFactory::CreateAverageExclus
 {
 	using namespace Insights;
 
-	TSharedRef<FTableColumn> ColumnRef = MakeShared<FTableColumn>(FTimersViewColumns::AverageExclusiveTimeColumnID);
-	FTableColumn& Column = *ColumnRef;
+	TSharedRef<FTimersTableColumn> ColumnRef = MakeShared<FTimersTableColumn>(FTimersViewColumns::AverageExclusiveTimeColumnID);
+	FTimersTableColumn& Column = *ColumnRef;
 
 	Column.SetShortName(LOCTEXT("AvgExclusiveTime_ColumnName", "E.Avg"));
 	Column.SetTitleName(LOCTEXT("AvgExclusiveTime_ColumnTitle", "Average Exclusive Time (ms)"));
-	Column.SetDescription(LOCTEXT("AvgExclusiveTime_ColumnDesc", "Average exclusive duration of selected timer's instances, in milliseconds"));
+
+	Column.SetDescription(ETraceFrameType::TraceFrameType_Count, LOCTEXT("AvgExclusiveTime_ColumnDesc", "Average exclusive duration of selected timer's instances, in milliseconds"));
+	Column.SetDescription(ETraceFrameType::TraceFrameType_Game, LOCTEXT("AvgExclusiveTime_GameFrameColumnDesc", "Game Frame Average Exclusive Duration.\nExclusive duration is computed for a single frame as the sum of exclusive duration of all instances of the timer in the respective frame.\nThe average is applied for these per-frame Exclusive durations. Unit is miliseconds."));
+	Column.SetDescription(ETraceFrameType::TraceFrameType_Rendering, LOCTEXT("AvgExclusiveTime_RenderingFrameColumnDesc", "Rendering Frame Average Exclusive Duration.\nExclusive duration is computed for a single frame as the sum of exclusive duration of all instances of the timer in the respective frame. The average is applied for these per-frame Exclusive durations.\nUnit is miliseconds."));
 
 	Column.SetFlags(ETableColumnFlags::CanBeHidden |
 					//ETableColumnFlags::ShouldBeVisible |
@@ -630,12 +672,15 @@ TSharedRef<Insights::FTableColumn> FTimersViewColumnFactory::CreateMedianExclusi
 {
 	using namespace Insights;
 
-	TSharedRef<FTableColumn> ColumnRef = MakeShared<FTableColumn>(FTimersViewColumns::MedianExclusiveTimeColumnID);
-	FTableColumn& Column = *ColumnRef;
+	TSharedRef<FTimersTableColumn> ColumnRef = MakeShared<FTimersTableColumn>(FTimersViewColumns::MedianExclusiveTimeColumnID);
+	FTimersTableColumn& Column = *ColumnRef;
 
 	Column.SetShortName(LOCTEXT("MedExclusiveTime_ColumnName", "E.Med"));
 	Column.SetTitleName(LOCTEXT("MedExclusiveTime_ColumnTitle", "Median Exclusive Time (ms)"));
-	Column.SetDescription(LOCTEXT("MedExclusiveTime_ColumnDesc", "Median exclusive duration of selected timer's instances, in milliseconds"));
+
+	Column.SetDescription(ETraceFrameType::TraceFrameType_Count, LOCTEXT("MedExclusiveTime_ColumnDesc", "Median exclusive duration of selected timer's instances, in milliseconds"));
+	Column.SetDescription(ETraceFrameType::TraceFrameType_Game, LOCTEXT("MedExclusiveTime_GameFrameColumnDesc", "Game Frame Median Exclusive Duration.\nExclusive duration is computed for a single frame as the sum of exclusive duration of all instances of the timer in the respective frame.\nThe median is aproximated these per-frame Exclusive durations. Unit is miliseconds."));
+	Column.SetDescription(ETraceFrameType::TraceFrameType_Rendering, LOCTEXT("MedExclusiveTime_RenderingFrameColumnDesc", "Rendering Frame Median Exclusive Duration.\nExclusive duration is computed for a single frame as the sum of exclusive duration of all instances of the timer in the respective frame.The median is aproximated these per-frame Exclusive durations.\nUnit is miliseconds."));
 
 	Column.SetFlags(ETableColumnFlags::CanBeHidden |
 					//ETableColumnFlags::ShouldBeVisible |
@@ -676,12 +721,15 @@ TSharedRef<Insights::FTableColumn> FTimersViewColumnFactory::CreateMinExclusiveT
 {
 	using namespace Insights;
 
-	TSharedRef<FTableColumn> ColumnRef = MakeShared<FTableColumn>(FTimersViewColumns::MinExclusiveTimeColumnID);
-	FTableColumn& Column = *ColumnRef;
+	TSharedRef<FTimersTableColumn> ColumnRef = MakeShared<FTimersTableColumn>(FTimersViewColumns::MinExclusiveTimeColumnID);
+	FTimersTableColumn& Column = *ColumnRef;
 
 	Column.SetShortName(LOCTEXT("MinExclusiveTime_ColumnName", "E.Min"));
 	Column.SetTitleName(LOCTEXT("MinExclusiveTime_ColumnTitle", "Min Exclusive Time (ms)"));
-	Column.SetDescription(LOCTEXT("MinExclusiveTime_ColumnDesc", "Minimum exclusive duration of selected timer's instances, in milliseconds"));
+
+	Column.SetDescription(ETraceFrameType::TraceFrameType_Count, LOCTEXT("MinExclusiveTime_ColumnDesc", "Minimum exclusive duration of selected timer's instances, in milliseconds"));
+	Column.SetDescription(ETraceFrameType::TraceFrameType_Game, LOCTEXT("MinExclusiveTime_GameFrameColumnDesc", "Game Frame Minimum Exclusive Duration.\nExclusive duration is computed for a single frame as the sum of exclusive duration of all instances of the timer in the respective frame.\nThe minimum is selected from these per-frame Exclusive durations. Unit is miliseconds."));
+	Column.SetDescription(ETraceFrameType::TraceFrameType_Rendering, LOCTEXT("MinExclusiveTime_RenderingFrameColumnDesc", "Rendering Frame Minimum Exclusive Duration.\nExclusive duration is computed for a single frame as the sum of exclusive duration of all instances of the timer in the respective frame.The minimum is selectedfrom these per-frame Exclusive durations.\nUnit is miliseconds."));
 
 	Column.SetFlags(ETableColumnFlags::CanBeHidden |
 					//ETableColumnFlags::ShouldBeVisible |
diff --git a/Engine/Source/Developer/TraceInsights/Private/Insights/ViewModels/TimersViewColumnFactory.h b/Engine/Source/Developer/TraceInsights/Private/Insights/ViewModels/TimersViewColumnFactory.h
--- a/Engine/Source/Developer/TraceInsights/Private/Insights/ViewModels/TimersViewColumnFactory.h
+++ b/Engine/Source/Developer/TraceInsights/Private/Insights/ViewModels/TimersViewColumnFactory.h
@@ -4,6 +4,8 @@
 
 #include "CoreMinimal.h"
 
+#include "Insights/Table/ViewModels/TableColumn.h"
+
 ////////////////////////////////////////////////////////////////////////////////////////////////////
 
 struct FTimersViewColumns
@@ -37,10 +39,51 @@ struct FTimersViewColumns
 
 ////////////////////////////////////////////////////////////////////////////////////////////////////
 
-namespace Insights
+class FTimersTableColumn : public Insights::FTableColumn
 {
-	class FTableColumn;
-}
+public:
+	FTimersTableColumn(const FName InId)
+		: Insights::FTableColumn(InId)
+	{}
+
+	FText GetDescription(ETraceFrameType InAggreagationMode) const
+	{
+		switch (InAggreagationMode)
+		{
+		case TraceFrameType_Game:
+			return GameFrame_Description;
+			break;
+		case TraceFrameType_Rendering:
+			return RenderingFrame_Description;
+			break;
+		default:
+			return FTableColumn::GetDescription();
+		}
+	}
+
+	void SetDescription(ETraceFrameType InAggreagationMode, FText InDescription)
+	{
+		switch (InAggreagationMode)
+		{
+		case TraceFrameType_Game:
+			GameFrame_Description = InDescription;
+			break;
+		case TraceFrameType_Rendering:
+			RenderingFrame_Description = InDescription;
+			break;
+		case TraceFrameType_Count:
+			Insights::FTableColumn::SetDescription(InDescription);
+			break;
+		default:
+			ensure(0);
+		}
+	}
+
+private:
+
+	FText GameFrame_Description;
+	FText RenderingFrame_Description;
+};
 
 ////////////////////////////////////////////////////////////////////////////////////////////////////
 
diff --git a/Engine/Source/Developer/TraceInsights/Private/Insights/Widgets/STimersViewTooltip.cpp b/Engine/Source/Developer/TraceInsights/Private/Insights/Widgets/STimersViewTooltip.cpp
--- a/Engine/Source/Developer/TraceInsights/Private/Insights/Widgets/STimersViewTooltip.cpp
+++ b/Engine/Source/Developer/TraceInsights/Private/Insights/Widgets/STimersViewTooltip.cpp
@@ -15,8 +15,12 @@
 #include "Insights/InsightsStyle.h"
 #include "Insights/Table/ViewModels/Table.h"
 #include "Insights/Table/ViewModels/TableColumn.h"
+#include "Insights/TimingProfilerManager.h"
 #include "Insights/ViewModels/TimerNode.h"
 #include "Insights/ViewModels/TimerNodeHelper.h"
+#include "Insights/ViewModels/TimersViewColumnFactory.h"
+#include "Insights/Widgets/STimersView.h"
+#include "Insights/Widgets/STimingProfilerWindow.h"
 
 ////////////////////////////////////////////////////////////////////////////////////////////////////
 
@@ -59,6 +63,41 @@ TSharedPtr<SToolTip> STimersViewTooltip::GetTableTooltip(const Insights::FTable&
 
 TSharedPtr<SToolTip> STimersViewTooltip::GetColumnTooltip(const Insights::FTableColumn& Column)
 {
+	const FTimersTableColumn& TimersColumn = static_cast<const FTimersTableColumn&>(Column);
+	FText InstanceDescription = TimersColumn.GetDescription(ETraceFrameType::TraceFrameType_Count);
+	FText GameFrameDescription = TimersColumn.GetDescription(ETraceFrameType::TraceFrameType_Game);
+	FText RenderingDescription = TimersColumn.GetDescription(ETraceFrameType::TraceFrameType_Rendering);
+
+	auto GetDescriptionLamda = [InstanceDescription, GameFrameDescription, RenderingDescription]()
+	{
+		ETraceFrameType FrameType = ETraceFrameType::TraceFrameType_Count;
+		TSharedPtr<STimingProfilerWindow> Window = FTimingProfilerManager::Get()->GetProfilerWindow();
+		if (Window.IsValid())
+		{
+			TSharedPtr<STimersView> TimersView = Window->GetTimersView();
+			if (TimersView.IsValid())
+			{
+				FrameType = TimersView->GetFrameTypeMode();
+			}
+		}
+
+		switch (FrameType)
+		{
+		case TraceFrameType_Count:
+			return InstanceDescription;
+		case TraceFrameType_Game:
+			return GameFrameDescription;
+			break;
+		case TraceFrameType_Rendering:
+			return RenderingDescription;
+			break;
+		default:
+			ensure(0);
+		}
+
+		return InstanceDescription;
+	};
+
 	TSharedPtr<SToolTip> ColumnTooltip =
 		SNew(SToolTip)
 		[
@@ -78,7 +117,7 @@ TSharedPtr<SToolTip> STimersViewTooltip::GetColumnTooltip(const Insights::FTable
 			.Padding(2.0f)
 			[
 				SNew(STextBlock)
-				.Text(Column.GetDescription())
+				.Text(TAttribute<FText>::CreateLambda(GetDescriptionLamda))
 				.TextStyle(FInsightsStyle::Get(), TEXT("TreeTable.Tooltip"))
 			]
 		];
