[Insights] Reorganize plot submenu in the Timers Panel.

diff --git a/Engine/Source/Developer/TraceInsights/Private/Insights/Widgets/STimersView.cpp b/Engine/Source/Developer/TraceInsights/Private/Insights/Widgets/STimersView.cpp
--- a/Engine/Source/Developer/TraceInsights/Private/Insights/Widgets/STimersView.cpp
+++ b/Engine/Source/Developer/TraceInsights/Private/Insights/Widgets/STimersView.cpp
@@ -786,7 +786,7 @@
 		return TimingView.IsValid() && NumSelectedNodes == 1 && SelectedNode.IsValid() && SelectedNode->GetType() != ETimerNodeType::Group;
 	};
 
-	MenuBuilder.BeginSection("MainGraphTrack", LOCTEXT("Plot_Series_MainGraphTrack_Section", "Main Graph Track"));
+	MenuBuilder.BeginSection("Instance", LOCTEXT("Plot_Series_Instance_Section", "Instance"));
 
 	// Add/remove series to/from graph track
 	{
@@ -822,6 +822,10 @@
 		}
 	}
 
+	MenuBuilder.EndSection();
+	
+	MenuBuilder.BeginSection("Game Frame", LOCTEXT("Plot_Series_GameFrame_Section", "Game Frame"));
+
 	// Add/remove game frame stats series to/from graph track
 	{
 		FUIAction Action_ToggleFrameStatsTimerInGraphTrack;
@@ -856,22 +860,22 @@
 		}
 	}
 
-	// Add/remove rendering frame stats series to/from graph track
+	// Add/remove game frame stats series to/from frame track
 	{
-		FUIAction Action_ToggleFrameStatsTimerInGraphTrack;
-		Action_ToggleFrameStatsTimerInGraphTrack.CanExecuteAction = FCanExecuteAction::CreateLambda(CanExecute);
-		Action_ToggleFrameStatsTimerInGraphTrack.ExecuteAction = FExecuteAction::CreateSP(this, &STimersView::ToggleTimingViewMainGraphEventFrameStatsSeries, SelectedNode, ETraceFrameType::TraceFrameType_Rendering);
+		FUIAction Action_ToggleFrameStatsTimerInFrameTrack;
+		Action_ToggleFrameStatsTimerInFrameTrack.CanExecuteAction = FCanExecuteAction::CreateLambda(CanExecute);
+		Action_ToggleFrameStatsTimerInFrameTrack.ExecuteAction = FExecuteAction::CreateSP(this, &STimersView::ToggleFrameTrackSeries, SelectedNode, ETraceFrameType::TraceFrameType_Game);
 
 		if (SelectedNode.IsValid() &&
 			SelectedNode->GetType() != ETimerNodeType::Group &&
-			IsFrameStatsSeriesInTimingViewMainGraph(SelectedNode, ETraceFrameType::TraceFrameType_Rendering))
+			IsSeriesInFrameTrack(SelectedNode, ETraceFrameType::TraceFrameType_Game))
 		{
 			MenuBuilder.AddMenuEntry
 			(
-				LOCTEXT("ContextMenu_RemoveRenderingFrameStatsFromGraphTrack", "Remove rendering frame stats series from graph track"),
-				LOCTEXT("ContextMenu_RemoveRenderingFrameStatsFromGraphTrack_Desc", "Remove the rendering frame stats series for this timer."),
+				LOCTEXT("ContextMenu_RemoveGameFrameStatsSeriesFromFrameTrack", "Remove game frame stats series from frame track"),
+				LOCTEXT("ContextMenu_RemoveGameFrameStatsSeriesFromFrameTrack_Desc", "Remove the game frame stats series for this timer from the frame track."),
 				FSlateIcon(FInsightsStyle::GetStyleSetName(), "Icons.RemoveGraphSeries"),
-				Action_ToggleFrameStatsTimerInGraphTrack,
+				Action_ToggleFrameStatsTimerInFrameTrack,
 				NAME_None,
 				EUserInterfaceActionType::Button
 			);
@@ -880,10 +884,10 @@
 		{
 			MenuBuilder.AddMenuEntry
 			(
-				LOCTEXT("ContextMenu_AddRenderingFrameStatsSeriesToGraphTrack", "Add rendering frame stats series to graph track"),
-				LOCTEXT("ContextMenu_AddRenderingFrameStatsSeriesToGraphTrack_Desc", "Adds a rendering frame stats series for this timer. Each data entry is computed as the sum of all instances of this timer in a rendering frame."),
+				LOCTEXT("ContextMenu_AddGameFrameStatsSeriesToFrameTrack", "Add game frame stats series to the frame track"),
+				LOCTEXT("ContextMenu_AddGameFrameStatsSeriesToFrameTrack_Desc", "Adds a game frame stats series for this timer to the frame track. Each data entry is computed as the sum of all instances of this timer in a game frame."),
 				FSlateIcon(FInsightsStyle::GetStyleSetName(), "Icons.AddGraphSeries"),
-				Action_ToggleFrameStatsTimerInGraphTrack,
+				Action_ToggleFrameStatsTimerInFrameTrack,
 				NAME_None,
 				EUserInterfaceActionType::Button
 			);
@@ -891,24 +895,25 @@
 	}
 
 	MenuBuilder.EndSection();
-	MenuBuilder.BeginSection("FrameTrack", LOCTEXT("Plot_Series_FrameTrack_Section", "Frame Track"));
 
-	// Add/remove game frame stats series to/from frame track
+	MenuBuilder.BeginSection("Rendering Frame", LOCTEXT("Plot_Series_RenderingFrame_Section", "Rendering frame"));
+
+	// Add/remove rendering frame stats series to/from graph track
 	{
-		FUIAction Action_ToggleFrameStatsTimerInFrameTrack;
-		Action_ToggleFrameStatsTimerInFrameTrack.CanExecuteAction = FCanExecuteAction::CreateLambda(CanExecute);
-		Action_ToggleFrameStatsTimerInFrameTrack.ExecuteAction = FExecuteAction::CreateSP(this, &STimersView::ToggleFrameTrackSeries, SelectedNode, ETraceFrameType::TraceFrameType_Game);
+		FUIAction Action_ToggleFrameStatsTimerInGraphTrack;
+		Action_ToggleFrameStatsTimerInGraphTrack.CanExecuteAction = FCanExecuteAction::CreateLambda(CanExecute);
+		Action_ToggleFrameStatsTimerInGraphTrack.ExecuteAction = FExecuteAction::CreateSP(this, &STimersView::ToggleTimingViewMainGraphEventFrameStatsSeries, SelectedNode, ETraceFrameType::TraceFrameType_Rendering);
 
 		if (SelectedNode.IsValid() &&
 			SelectedNode->GetType() != ETimerNodeType::Group &&
-			IsSeriesInFrameTrack(SelectedNode, ETraceFrameType::TraceFrameType_Game))
+			IsFrameStatsSeriesInTimingViewMainGraph(SelectedNode, ETraceFrameType::TraceFrameType_Rendering))
 		{
 			MenuBuilder.AddMenuEntry
 			(
-				LOCTEXT("ContextMenu_RemoveGameFrameStatsSeriesFromFrameTrack", "Remove game frame stats series from frame track"),
-				LOCTEXT("ContextMenu_RemoveGameFrameStatsSeriesFromFrameTrack_Desc", "Remove the game frame stats series for this timer from the frame track."),
+				LOCTEXT("ContextMenu_RemoveRenderingFrameStatsFromGraphTrack", "Remove rendering frame stats series from graph track"),
+				LOCTEXT("ContextMenu_RemoveRenderingFrameStatsFromGraphTrack_Desc", "Remove the rendering frame stats series for this timer."),
 				FSlateIcon(FInsightsStyle::GetStyleSetName(), "Icons.RemoveGraphSeries"),
-				Action_ToggleFrameStatsTimerInFrameTrack,
+				Action_ToggleFrameStatsTimerInGraphTrack,
 				NAME_None,
 				EUserInterfaceActionType::Button
 			);
@@ -917,15 +922,16 @@
 		{
 			MenuBuilder.AddMenuEntry
 			(
-				LOCTEXT("ContextMenu_AddGameFrameStatsSeriesToFrameTrack", "Add game frame stats series to the frame track"),
-				LOCTEXT("ContextMenu_AddGameFrameStatsSeriesToFrameTrack_Desc", "Adds a game frame stats series for this timer to the frame track. Each data entry is computed as the sum of all instances of this timer in a game frame."),
+				LOCTEXT("ContextMenu_AddRenderingFrameStatsSeriesToGraphTrack", "Add rendering frame stats series to graph track"),
+				LOCTEXT("ContextMenu_AddRenderingFrameStatsSeriesToGraphTrack_Desc", "Adds a rendering frame stats series for this timer. Each data entry is computed as the sum of all instances of this timer in a rendering frame."),
 				FSlateIcon(FInsightsStyle::GetStyleSetName(), "Icons.AddGraphSeries"),
-				Action_ToggleFrameStatsTimerInFrameTrack,
+				Action_ToggleFrameStatsTimerInGraphTrack,
 				NAME_None,
 				EUserInterfaceActionType::Button
 			);
 		}
 	}
+
 	// Add/remove rendering frame stats series to/from frame track
 	{
 		FUIAction Action_ToggleFrameStatsTimerInFrameTrack;
