[Insights] Double clicking on a timer now plots the graph correspoding to the currently selected frame type.

diff --git a/Engine/Source/Developer/TraceInsights/Private/Insights/Widgets/STimersView.cpp b/Engine/Source/Developer/TraceInsights/Private/Insights/Widgets/STimersView.cpp
--- a/Engine/Source/Developer/TraceInsights/Private/Insights/Widgets/STimersView.cpp
+++ b/Engine/Source/Developer/TraceInsights/Private/Insights/Widgets/STimersView.cpp
@@ -1429,7 +1429,25 @@ void STimersView::TreeView_OnMouseButtonDoubleClick(FTimerNodePtr NodePtr)
 		}
 		else
 		{
-			ToggleTimingViewMainGraphEventSeries(NodePtr);
+			switch (ModeFrameType)
+			{
+				case ETraceFrameType::TraceFrameType_Count:
+				{
+					// Instance graph
+					ToggleTimingViewMainGraphEventSeries(NodePtr);
+					break;
+				}
+				case ETraceFrameType::TraceFrameType_Game:
+				case ETraceFrameType::TraceFrameType_Rendering:
+				{
+					ToggleTimingViewMainGraphEventFrameStatsSeries(NodePtr, ModeFrameType);
+					break;
+				}
+				default:
+				{
+					ensure(0);
+				}
+			}
 		}
 	}
 }
