Allow swapping of watch and store directory.

Previously setting the store directory to the same directory as a watch directory was disabled. With this change the selected directory when set as store it is also automatically removed as watch directory.

diff --git a/Engine/Source/Developer/TraceAnalysis/Private/Store/StoreClient.cpp b/Engine/Source/Developer/TraceAnalysis/Private/Store/StoreClient.cpp
--- a/Engine/Source/Developer/TraceAnalysis/Private/Store/StoreClient.cpp
+++ b/Engine/Source/Developer/TraceAnalysis/Private/Store/StoreClient.cpp
@@ -397,15 +397,18 @@ bool FStoreCborClient::SetStoreDirectories(const TCHAR* StoreDir, const TArray<F
 	{
 		Builder.AddString("StoreDir", TCHAR_TO_ANSI(StoreDir));
 	}
+	TArray<FString> WatchDirs;
+	if (!RemoveWatchDir.IsEmpty())
+	{
+		Algo::Transform(RemoveWatchDir, WatchDirs, [](const FString& In) { return TEXT("-") + In; });
+	}
 	if (!AddWatchDirs.IsEmpty())
 	{
-		Builder.AddStringArray("Additionalwatchdirs", AddWatchDirs);
+		WatchDirs.Append(AddWatchDirs);
 	}
-	if (!RemoveWatchDir.IsEmpty())
+	if (!WatchDirs.IsEmpty())
 	{
-		TArray<FString> RemoveDirs;
-		Algo::Transform(RemoveWatchDir, RemoveDirs, [](const FString& In) { return TEXT("-") + In; });
-		Builder.AddStringArray("Additionalwatchdirs", RemoveDirs);
+		Builder.AddStringArray("Additionalwatchdirs", WatchDirs);
 	}
 	return Communicate(Builder.Done());
 }
diff --git a/Engine/Source/Developer/TraceInsights/Private/Insights/Widgets/SStartPageWindow.cpp b/Engine/Source/Developer/TraceInsights/Private/Insights/Widgets/SStartPageWindow.cpp
--- a/Engine/Source/Developer/TraceInsights/Private/Insights/Widgets/SStartPageWindow.cpp
+++ b/Engine/Source/Developer/TraceInsights/Private/Insights/Widgets/SStartPageWindow.cpp
@@ -1502,12 +1502,19 @@ FReply STraceDirectoryItem::OnModifyStore()
 		const bool bIsWatchDir = Window->WatchDirectoriesModel.FindByPredicate([&](const auto& Directory){ return FPathViews::Equals(SelectedDirectory, Directory->Path);}) != nullptr;
 		const bool bIsCurrentStoreDir = FPathViews::Equals(SelectedDirectory, CurrentStoreDirectory);
 
-		if (bHasSelected && !bIsWatchDir && !bIsCurrentStoreDir)
+		if (bHasSelected && !bIsCurrentStoreDir)
 		{
 			FPaths::MakePlatformFilename(SelectedDirectory);
 			FPaths::MakePlatformFilename(CurrentStoreDirectory);
-			TArray<FString> AddWatchDirs = { CurrentStoreDirectory };
-			if (!FInsightsManager::Get()->GetStoreClient()->SetStoreDirectories(*SelectedDirectory, AddWatchDirs, TArray<FString>()))
+			const TArray<FString> AddWatchDirs = { CurrentStoreDirectory };
+			TArray<FString> RemoveWatchDirs;
+			if (bIsWatchDir)
+			{
+				// If we are selecting a watch dir as new store dir, make
+				// sure we remove it as watch directory.
+				RemoveWatchDirs.Emplace(SelectedDirectory);
+			}
+			if (!FInsightsManager::Get()->GetStoreClient()->SetStoreDirectories(*SelectedDirectory, AddWatchDirs, RemoveWatchDirs))
 			{
 				FMessageLog(FInsightsManager::Get()->GetLogListingName()).Error(LOCTEXT("StoreCommunicationFail", "Failed to change settings on the store service."));
 			}
