Try to update symbol search path until module provider is created

Previously the symbol search path string was update only whenever the panel was reset, which doesn't always correspond to when the module provider was created, leaving the string blank until the panel was recreated.

diff --git a/Engine/Source/Developer/TraceInsights/Private/Insights/MemoryProfiler/Widgets/SMemInvestigationView.cpp b/Engine/Source/Developer/TraceInsights/Private/Insights/MemoryProfiler/Widgets/SMemInvestigationView.cpp
--- a/Engine/Source/Developer/TraceInsights/Private/Insights/MemoryProfiler/Widgets/SMemInvestigationView.cpp
+++ b/Engine/Source/Developer/TraceInsights/Private/Insights/MemoryProfiler/Widgets/SMemInvestigationView.cpp
@@ -211,7 +211,6 @@ TSharedRef<SWidget> SMemInvestigationView::ConstructInvestigationWidgetArea()
 		.HAlign(HAlign_Fill)
 		[
 			SAssignNew(SymbolPathsTextBlock, STextBlock)
-			.Text(GetSymbolPathsText())
 			.ColorAndOpacity(FLinearColor(0.3f, 0.3f, 0.3f, 1.0f))
 			.AutoWrapText(true)
 		]
@@ -367,13 +366,14 @@ void SMemInvestigationView::InsightsManager_OnSessionChanged()
 
 void SMemInvestigationView::Reset()
 {
-	SymbolPathsTextBlock->SetText(GetSymbolPathsText());
+	SymbolPathsTextBlock->SetText(FText::GetEmpty());
 }
 
 ////////////////////////////////////////////////////////////////////////////////////////////////////
 
 void SMemInvestigationView::Tick(const FGeometry& AllottedGeometry, const double InCurrentTime, const float InDeltaTime)
 {
+	UpdateSymbolPathsText();
 }
 
 ////////////////////////////////////////////////////////////////////////////////////////////////////
@@ -474,16 +474,16 @@ void SMemInvestigationView::QueryTarget_OnSelectionChanged(TSharedPtr<Insights::
 }
 
 ////////////////////////////////////////////////////////////////////////////////////////////////////
-FText SMemInvestigationView::GetSymbolPathsText() const
+
+void SMemInvestigationView::UpdateSymbolPathsText() const
 {
-	if (Session)
+	if (SymbolPathsTextBlock->GetText().IsEmpty() && Session)
 	{
 		if (const TraceServices::IModuleProvider* ModuleProvider = ReadModuleProvider(*Session.Get()))
 		{
-			return FSymbolSearchPathsHelper::GetLocalizedSymbolSearchPathsText(ModuleProvider);
+			SymbolPathsTextBlock->SetText(FSymbolSearchPathsHelper::GetLocalizedSymbolSearchPathsText(ModuleProvider));
 		}
 	}
-	return FText();
 }
 
 ////////////////////////////////////////////////////////////////////////////////////////////////////
diff --git a/Engine/Source/Developer/TraceInsights/Private/Insights/MemoryProfiler/Widgets/SMemInvestigationView.h b/Engine/Source/Developer/TraceInsights/Private/Insights/MemoryProfiler/Widgets/SMemInvestigationView.h
--- a/Engine/Source/Developer/TraceInsights/Private/Insights/MemoryProfiler/Widgets/SMemInvestigationView.h
+++ b/Engine/Source/Developer/TraceInsights/Private/Insights/MemoryProfiler/Widgets/SMemInvestigationView.h
@@ -61,7 +61,7 @@ public:
 	void QueryTarget_OnSelectionChanged(TSharedPtr<Insights::FQueryTargetWindowSpec> InRule, ESelectInfo::Type SelectInfo);
 
 private:
-	FText GetSymbolPathsText() const;
+	void UpdateSymbolPathsText() const;
 	TSharedRef<SWidget> ConstructInvestigationWidgetArea();
 	TSharedRef<SWidget> ConstructTimeMarkerWidget(uint32 TimeMarkerIndex);
 
