[Insights]
  - MemAllocs: Improved/fixed the status bar text (query info, query status, number of allocs + added number of filtered allocs).
  - MemAllocs: Added status bar text info for selected items (including the total alloc count and total mem size of selected items).
  - Tree Views: Fixed name of Copy / Copy Value and Copy Tooltip context menu actions.
  - Tree Views: Fixed aggregation for group nodes with custom get value.

diff --git a/Engine/Plugins/Editor/AssetManagerEditor/Source/AssetManagerEditor/Private/Insights/Table/Widgets/STableTreeView.cpp b/Engine/Plugins/Editor/AssetManagerEditor/Source/AssetManagerEditor/Private/Insights/Table/Widgets/STableTreeView.cpp
--- a/Engine/Plugins/Editor/AssetManagerEditor/Source/AssetManagerEditor/Private/Insights/Table/Widgets/STableTreeView.cpp
+++ b/Engine/Plugins/Editor/AssetManagerEditor/Source/AssetManagerEditor/Private/Insights/Table/Widgets/STableTreeView.cpp
@@ -62,20 +62,20 @@ public:
 	virtual void RegisterCommands() override
 	{
 		UI_COMMAND(Command_CopyToClipboard,
-			"Copy To Clipboard",
-			"Copies the selection to clipboard.",
+			"Copy",
+			"Copies the selected table rows to clipboard.",
 			EUserInterfaceActionType::Button,
 			FInputChord(EModifierKey::Control, EKeys::C));
 
 		UI_COMMAND(Command_CopyColumnToClipboard,
-			"Copy Column Value To Clipboard",
-			"Copies the value of the hovered column to clipboard.",
+			"Copy Value",
+			"Copies the value of the hovered cell to clipboard.",
 			EUserInterfaceActionType::Button,
 			FInputChord(EModifierKey::Control | EModifierKey::Shift, EKeys::C));
 
 		UI_COMMAND(Command_CopyColumnTooltipToClipboard,
-			"Copy Column Tooltip To Clipboard",
-			"Copies the value of the hovered column's tooltip to clipboard.",
+			"Copy Tooltip",
+			"Copies the tooltip of the hovered cell to clipboard.",
 			EUserInterfaceActionType::Button,
 			FInputChord(EModifierKey::Control | EModifierKey::Alt, EKeys::C));
 
diff --git a/Engine/Source/Developer/TraceInsights/Private/Insights/MemoryProfiler/Widgets/SMemAllocTableTreeView.cpp b/Engine/Source/Developer/TraceInsights/Private/Insights/MemoryProfiler/Widgets/SMemAllocTableTreeView.cpp
--- a/Engine/Source/Developer/TraceInsights/Private/Insights/MemoryProfiler/Widgets/SMemAllocTableTreeView.cpp
+++ b/Engine/Source/Developer/TraceInsights/Private/Insights/MemoryProfiler/Widgets/SMemAllocTableTreeView.cpp
@@ -10,6 +10,7 @@
 #include "Modules/ModuleManager.h"
 #include "SlateOptMacros.h"
 #include "Styling/AppStyle.h"
+#include "Styling/StyleColors.h"
 #include "Widgets/Images/SImage.h"
 #include "Widgets/Input/SCheckBox.h"
 #include "Widgets/SToolTip.h"
@@ -163,7 +164,6 @@ void SMemAllocTableTreeView::RebuildTree(bool bResync)
 					TableRowNodes.Add(NodePtr);
 				}
 				ensure(TableRowNodes.Num() == (bIncludeHeapAllocs ? TotalAllocCount : TotalAllocCount - HeapAllocCount));
-				UpdateQueryInfo();
 			}
 		}
 	}
@@ -221,6 +221,8 @@ void SMemAllocTableTreeView::OnQueryInvalidated()
 	{
 		ResetAndStartQuery();
 	}
+
+	UpdateQueryInfo();
 }
 
 ////////////////////////////////////////////////////////////////////////////////////////////////////
@@ -236,8 +238,6 @@ void SMemAllocTableTreeView::ResetAndStartQuery()
 		Allocs.Reset(10 * 1024 * 1024);
 	}
 
-	UpdateQueryInfo();
-
 	StartQuery();
 
 	RebuildTree(true);
@@ -408,7 +408,7 @@ void SMemAllocTableTreeView::UpdateQuery(TraceServices::IAllocationsProvider::EQ
 					Alloc.Asset = nullptr;
 					Alloc.ClassName = nullptr;
 					Alloc.Package = nullptr;
-					
+
 					const uint32 MetadataId = Allocation->GetMetadataId();
 					if (MetadataId != TraceServices::IMetadataProvider::InvalidMetadataId && MetadataProvider && DefinitionProvider)
 					{
@@ -1103,8 +1103,9 @@ TSharedPtr<SWidget> SMemAllocTableTreeView::ConstructFooter()
 		SNew(SHorizontalBox)
 
 		+ SHorizontalBox::Slot()
+		.AutoWidth()
 		.HAlign(HAlign_Left)
-		.Padding(2.0f)
+		.Padding(2.0f, 2.0f, 0.0f, 2.0f)
 		[
 			SNew(STextBlock)
 			.Text(this, &SMemAllocTableTreeView::GetQueryInfo)
@@ -1112,8 +1113,37 @@ TSharedPtr<SWidget> SMemAllocTableTreeView::ConstructFooter()
 		]
 
 		+ SHorizontalBox::Slot()
+		.AutoWidth()
+		.HAlign(HAlign_Left)
+		.Padding(0.0f, 2.0f, 0.0f, 2.0f)
+		[
+			SNew(STextBlock)
+			.Text(LOCTEXT("FooterSeparator", " : "))
+			.ColorAndOpacity(FSlateColor(EStyleColor::White25))
+		]
+
+		+ SHorizontalBox::Slot()
+		.AutoWidth()
+		.HAlign(HAlign_Left)
+		.Padding(0.0f, 2.0f, 8.0f, 2.0f)
+		[
+			SNew(STextBlock)
+			.Text(this, &SMemAllocTableTreeView::GetFooterLeftText)
+		]
+
+		+ SHorizontalBox::Slot()
+		.FillWidth(1.0f)
+		.HAlign(HAlign_Center)
+		.Padding(2.0f, 2.0f, 2.0f, 2.0f)
+		[
+			SNew(STextBlock)
+			.Text(this, &SMemAllocTableTreeView::GetFooterCenterText)
+		]
+
+		+ SHorizontalBox::Slot()
+		.AutoWidth()
 		.HAlign(HAlign_Right)
-		.Padding(2.0f)
+		.Padding(8.0f, 2.0f, 2.0f, 2.0f)
 		[
 			SNew(STextBlock)
 			.Text(this, &SMemAllocTableTreeView::GetSymbolResolutionStatus)
@@ -1175,7 +1205,80 @@ FText SMemAllocTableTreeView::GetQueryInfo() const
 
 FText SMemAllocTableTreeView::GetQueryInfoTooltip() const
 {
-	return QueryInfoTooltip;
+	return Rule.IsValid() ? Rule->GetDescription() : FText::GetEmpty();
+}
+
+////////////////////////////////////////////////////////////////////////////////////////////////////
+
+FText SMemAllocTableTreeView::GetFooterLeftText() const
+{
+	if (!Rule.IsValid())
+	{
+		return FText::GetEmpty();
+	}
+
+	if (Query != 0)
+	{
+		return LOCTEXT("FooterLeftTextRunningQuery", "running query...");
+	}
+
+	if (FilteredNodesPtr->Num() == TableRowNodes.Num())
+	{
+		return FText::Format(LOCTEXT("FooterLeftTextFmt1", "{0} {0}|plural(one=alloc,other=allocs)"), FText::AsNumber(TableRowNodes.Num()));
+	}
+	else
+	{
+		return FText::Format(LOCTEXT("FooterLeftTextFmt2", "{0} / {1} {1}|plural(one=alloc,other=allocs)"), FilteredNodesPtr->Num(), FText::AsNumber(TableRowNodes.Num()));
+	}
+}
+
+////////////////////////////////////////////////////////////////////////////////////////////////////
+
+FText SMemAllocTableTreeView::GetFooterCenterText() const
+{
+	return SelectionStatsText;
+}
+
+////////////////////////////////////////////////////////////////////////////////////////////////////
+
+void SMemAllocTableTreeView::TreeView_OnSelectionChanged(FTableTreeNodePtr SelectedItem, ESelectInfo::Type SelectInfo)
+{
+	TArray<FTableTreeNodePtr> SelectedNodes;
+	const int32 NumSelectedNodes = TreeView->GetSelectedItems(SelectedNodes);
+
+	if (NumSelectedNodes > 0)
+	{
+		int64 TotalAllocCount = 0;
+		int64 TotalAllocSize = 0;
+
+		TSharedRef<FTableColumn> CountColumn = Table->FindColumnChecked(FMemAllocTableColumns::CountColumnId);
+		TSharedRef<FTableColumn> SizeColumn = Table->FindColumnChecked(FMemAllocTableColumns::SizeColumnId);
+
+		for (const Insights::FTableTreeNodePtr& Node : SelectedNodes)
+		{
+			TOptional<FTableCellValue> CountValue = CountColumn->GetValue(*Node.Get());
+			if (CountValue.IsSet())
+			{
+				TotalAllocCount += CountValue.GetValue().AsInt64();
+			}
+
+			TOptional<FTableCellValue> SizeValue = SizeColumn->GetValue(*Node.Get());
+			if (SizeValue.IsSet())
+			{
+				TotalAllocSize += SizeValue.GetValue().AsInt64();
+			}
+		}
+
+		FNumberFormattingOptions FormattingOptionsMem;
+		FormattingOptionsMem.MaximumFractionalDigits = 2;
+
+		SelectionStatsText = FText::Format(LOCTEXT("SelectionStatsFmt", "{0} selected {0}|plural(one=item,other=items) ({1} {1}|plural(one=alloc,other=allocs), {2})"),
+			FText::AsNumber(NumSelectedNodes), FText::AsNumber(TotalAllocCount), FText::AsMemory(TotalAllocSize, &FormattingOptionsMem));
+	}
+	else
+	{
+		SelectionStatsText = FText::GetEmpty();
+	}
 }
 
 ////////////////////////////////////////////////////////////////////////////////////////////////////
@@ -1257,8 +1360,7 @@ void SMemAllocTableTreeView::UpdateQueryInfo()
 			check(false);
 		}
 
-		QueryInfo = FText::Format(LOCTEXT("QueryInfoFmt", "{0} ({1}) : {2} allocs"), Rule->GetVerboseName(), TimeMarkersText, FText::AsNumber(TableRowNodes.Num()));
-		QueryInfoTooltip = Rule->GetDescription();
+		QueryInfo = FText::Format(LOCTEXT("QueryInfoFmt", "{0} ({1})"), Rule->GetVerboseName(), TimeMarkersText);
 	}
 }
 
@@ -1507,7 +1609,7 @@ void SMemAllocTableTreeView::ExtendMenu(FMenuBuilder& MenuBuilder)
 	}
 
 	{
-		const FText ItemLabel = LOCTEXT("ContextMenu_Export_SubMenu", "Export Snapshot");
+		const FText ItemLabel = LOCTEXT("ContextMenu_Export_SubMenu", "Export Snapshot...");
 		const FText ItemToolTip = LOCTEXT("ContextMenu_Export_Desc_SubMenu", "Export memory snapshot to construct diff later.");
 
 		FUIAction Action_ExportSnapshot
diff --git a/Engine/Source/Developer/TraceInsights/Private/Insights/MemoryProfiler/Widgets/SMemAllocTableTreeView.h b/Engine/Source/Developer/TraceInsights/Private/Insights/MemoryProfiler/Widgets/SMemAllocTableTreeView.h
--- a/Engine/Source/Developer/TraceInsights/Private/Insights/MemoryProfiler/Widgets/SMemAllocTableTreeView.h
+++ b/Engine/Source/Developer/TraceInsights/Private/Insights/MemoryProfiler/Widgets/SMemAllocTableTreeView.h
@@ -115,10 +115,15 @@ private:
 
 	FText GetSymbolResolutionStatus() const;
 	FText GetSymbolResolutionTooltip() const;
+	
+	void UpdateQueryInfo();
 	FText GetQueryInfo() const;
 	FText GetQueryInfoTooltip() const;
+	FText GetFooterLeftText() const;
+	FText GetFooterCenterText() const;
+
+	virtual void TreeView_OnSelectionChanged(Insights::FTableTreeNodePtr SelectedItem, ESelectInfo::Type SelectInfo) override;
 
-	void UpdateQueryInfo();
 	virtual void UpdateFilterContext(const FFilterConfigurator& InFilterConfigurator, const FTableTreeNode& InNode) const override;
 	virtual void InitFilterConfigurator(FFilterConfigurator& InOutFilterConfigurator) override;
 
@@ -139,6 +144,7 @@ private:
 	TraceServices::IAllocationsProvider::FQueryHandle Query = 0;
 	FText QueryInfo;
 	FText QueryInfoTooltip;
+	FText SelectionStatsText;
 	FStopwatch QueryStopwatch;
 	bool bHasPendingQueryReset = false;
 	bool bIsCallstackGroupingByFunction = true;
diff --git a/Engine/Source/Developer/TraceInsights/Private/Insights/Table/ViewModels/TableTreeNode.h b/Engine/Source/Developer/TraceInsights/Private/Insights/Table/ViewModels/TableTreeNode.h
--- a/Engine/Source/Developer/TraceInsights/Private/Insights/Table/ViewModels/TableTreeNode.h
+++ b/Engine/Source/Developer/TraceInsights/Private/Insights/Table/ViewModels/TableTreeNode.h
@@ -69,6 +69,15 @@ public:
 	{
 	}
 
+	/** Initialization constructor for a table record node. */
+	explicit FTableTreeNode(const FName InName, TWeakPtr<FTable> InParentTable, int32 InRowIndex, bool IsGroup)
+		: FBaseTreeNode(InName, IsGroup)
+		, ParentTable(InParentTable)
+		, RowId(InRowIndex)
+		, AggregatedValues(nullptr)
+	{
+	}
+
 	virtual ~FTableTreeNode()
 	{
 		CleanupAggregatedValues();
@@ -157,8 +166,8 @@ class FCustomTableTreeNode : public FTableTreeNode
 
 public:
 	/** Initialization constructor for a table record node. */
-	explicit FCustomTableTreeNode(const FName InName, TWeakPtr<FTable> InParentTable, int32 InRowIndex, const FSlateBrush* InIconBrush, FLinearColor InColor)
-		: FTableTreeNode(InName, InParentTable)
+	explicit FCustomTableTreeNode(const FName InName, TWeakPtr<FTable> InParentTable, int32 InRowIndex, const FSlateBrush* InIconBrush, FLinearColor InColor, bool IsGroup)
+		: FTableTreeNode(InName, InParentTable, InRowIndex, IsGroup)
 		, IconBrush(InIconBrush)
 		, Color(InColor)
 	{
diff --git a/Engine/Source/Developer/TraceInsights/Private/Insights/Table/Widgets/STableTreeView.cpp b/Engine/Source/Developer/TraceInsights/Private/Insights/Table/Widgets/STableTreeView.cpp
--- a/Engine/Source/Developer/TraceInsights/Private/Insights/Table/Widgets/STableTreeView.cpp
+++ b/Engine/Source/Developer/TraceInsights/Private/Insights/Table/Widgets/STableTreeView.cpp
@@ -60,20 +60,20 @@ public:
 	virtual void RegisterCommands() override
 	{
 		UI_COMMAND(Command_CopyToClipboard,
-			"Copy To Clipboard",
-			"Copies the selection to clipboard.",
+			"Copy",
+			"Copies the selected table rows to clipboard.",
 			EUserInterfaceActionType::Button,
 			FInputChord(EModifierKey::Control, EKeys::C));
 
 		UI_COMMAND(Command_CopyColumnToClipboard,
-			"Copy Column Value To Clipboard",
-			"Copies the value of the hovered column to clipboard.",
+			"Copy Value",
+			"Copies the value of the hovered cell to clipboard.",
 			EUserInterfaceActionType::Button,
 			FInputChord(EModifierKey::Control | EModifierKey::Shift, EKeys::C));
 
 		UI_COMMAND(Command_CopyColumnTooltipToClipboard,
-			"Copy Column Tooltip To Clipboard",
-			"Copies the value of the hovered column's tooltip to clipboard.",
+			"Copy Tooltip",
+			"Copies the tooltip of the hovered cell to clipboard.",
 			EUserInterfaceActionType::Button,
 			FInputChord(EModifierKey::Control | EModifierKey::Alt, EKeys::C));
 
@@ -1128,7 +1128,7 @@ void STableTreeView::UpdateFilterContext(const FFilterConfigurator& InFilterConf
 {
 	for (const TSharedRef<FTableColumn>& Column : Table->GetColumns())
 	{
-		if (!Column->CanBeFiltered())
+		if (!Column->CanBeFiltered() || !InFilterConfigurator.IsKeyUsed(Column->GetIndex()))
 		{
 			continue;
 		}
@@ -1153,11 +1153,7 @@ void STableTreeView::UpdateFilterContext(const FFilterConfigurator& InFilterConf
 			{
 				if (!Column->IsHierarchy())
 				{
-					// Converting string is heavy, check if the key is used for filtering before doing any conversion
-					if (InFilterConfigurator.IsKeyUsed(Column->GetIndex()))
-					{
-						FilterContext.SetFilterData<FString>(Column->GetIndex(), Column->GetValue(InNode)->AsString());
-					}
+					FilterContext.SetFilterData<FString>(Column->GetIndex(), Column->GetValue(InNode)->AsString());
 				}
 				break;
 			}
@@ -1171,12 +1167,13 @@ BEGIN_SLATE_FUNCTION_BUILD_OPTIMIZATION
 TSharedRef<SWidget> STableTreeView::ConstructFilterConfiguratorButton()
 {
 	return SNew(SButton)
-		.ContentPadding(FMargin(-2.0f, 2.0f, -2.0f, 2.0f))
+		.ButtonStyle(FAppStyle::Get(), "SimpleButton")
 		.ToolTipText(LOCTEXT("FilterConfiguratorBtn_ToolTip", "Opens the filter configurator window."))
 		.OnClicked(this, &STableTreeView::FilterConfigurator_OnClicked)
 		[
 			SNew(SImage)
 			.Image(FInsightsStyle::GetBrush("Icons.ClassicFilterConfig"))
+			.ColorAndOpacity(FSlateColor::UseForeground())
 		];
 }
 END_SLATE_FUNCTION_BUILD_OPTIMIZATION
@@ -2443,31 +2440,15 @@ void STableTreeView::UpdateCStringSameValueAggregationSingleNode(const FTableCol
 	const TCHAR* AggregatedValue = nullptr;
 
 	// Find the first child node.
-	for (FBaseTreeNodePtr NodePtr : InOutGroupNode.GetChildren())
+	if (InOutGroupNode.GetChildrenCount() > 0)
 	{
-		if (!NodePtr->IsGroup())
-		{
-			const TOptional<FTableCellValue> NodeValue = InColumn.GetValue(*NodePtr);
-			if (NodeValue.IsSet() &&
-				NodeValue.GetValue().DataType == ETableCellDataType::CString)
-			{
-				AggregatedValue = NodeValue.GetValue().CString;
-			}
-		}
-		else
+		FBaseTreeNodePtr NodePtr = InOutGroupNode.GetChildren()[0];
+		const TOptional<FTableCellValue> NodeValue = InColumn.GetValue(*NodePtr);
+		if (NodeValue.IsSet() &&
+			NodeValue.GetValue().DataType == ETableCellDataType::CString)
 		{
-			FTableTreeNode& TableNode = *(FTableTreeNode*)NodePtr.Get();
-			if (TableNode.HasAggregatedValue(InColumn.GetId()))
-			{
-				const FTableCellValue& ChildGroupAggregatedValue = TableNode.GetAggregatedValue(InColumn.GetId());
-				if (ChildGroupAggregatedValue.DataType == ETableCellDataType::CString)
-				{
-					AggregatedValue = ChildGroupAggregatedValue.CString;
-				}
-			}
+			AggregatedValue = NodeValue.GetValue().CString;
 		}
-
-		break;
 	}
 
 	if (AggregatedValue != nullptr)
@@ -2475,28 +2456,12 @@ void STableTreeView::UpdateCStringSameValueAggregationSingleNode(const FTableCol
 		// Check if all other children have the same value as the first node.
 		for (FBaseTreeNodePtr NodePtr : InOutGroupNode.GetChildren())
 		{
-			if (!NodePtr->IsGroup())
+			const TOptional<FTableCellValue> NodeValue = InColumn.GetValue(*NodePtr);
+			if (NodeValue.IsSet() &&
+				NodeValue.GetValue().DataType == ETableCellDataType::CString &&
+				AggregatedValue == NodeValue.GetValue().CString)
 			{
-				const TOptional<FTableCellValue> NodeValue = InColumn.GetValue(*NodePtr);
-				if (NodeValue.IsSet() &&
-					NodeValue.GetValue().DataType == ETableCellDataType::CString &&
-					AggregatedValue == NodeValue.GetValue().CString)
-				{
-					continue;
-				}
-			}
-			else
-			{
-				FTableTreeNode& TableNode = *(FTableTreeNode*)NodePtr.Get();
-				if (TableNode.HasAggregatedValue(InColumn.GetId()))
-				{
-					const FTableCellValue& ChildGroupAggregatedValue = TableNode.GetAggregatedValue(InColumn.GetId());
-					if (ChildGroupAggregatedValue.DataType == ETableCellDataType::CString &&
-						AggregatedValue == ChildGroupAggregatedValue.CString)
-					{
-						continue;
-					}
-				}
+				continue;
 			}
 
 			AggregatedValue = nullptr;
@@ -2550,18 +2515,11 @@ void STableTreeView::UpdateAggregation(const FTableColumn& InColumn, FTableTreeN
 				TableNode.ResetAggregatedValue(InColumn.GetId());
 				STableTreeView::UpdateAggregation<T, bSetInitialValue, true>(InColumn, TableNode, InitialAggregatedValue, ValueGetterFunc);
 			}
-			if (TableNode.HasAggregatedValue(InColumn.GetId()))
-			{
-				AggregatedValue = ValueGetterFunc(AggregatedValue, TableNode.GetAggregatedValue(InColumn.GetId()));
-			}
 		}
-		else
+		const TOptional<FTableCellValue> NodeValue = InColumn.GetValue(*NodePtr);
+		if (NodeValue.IsSet())
 		{
-			const TOptional<FTableCellValue> NodeValue = InColumn.GetValue(*NodePtr);
-			if (NodeValue.IsSet())
-			{
-				AggregatedValue = ValueGetterFunc(AggregatedValue, NodeValue.GetValue());
-			}
+			AggregatedValue = ValueGetterFunc(AggregatedValue, NodeValue.GetValue());
 		}
 	}
 
