Timer graphs now gather data from visible tracks only and auto update when track visiblity changes

diff --git a/Engine/Source/Developer/TraceInsights/Private/Insights/ContextSwitches/ViewModels/ContextSwitchesSharedState.cpp b/Engine/Source/Developer/TraceInsights/Private/Insights/ContextSwitches/ViewModels/ContextSwitchesSharedState.cpp
--- a/Engine/Source/Developer/TraceInsights/Private/Insights/ContextSwitches/ViewModels/ContextSwitchesSharedState.cpp
+++ b/Engine/Source/Developer/TraceInsights/Private/Insights/ContextSwitches/ViewModels/ContextSwitchesSharedState.cpp
@@ -656,7 +656,7 @@ void FContextSwitchesSharedState::Command_DockCpuThreadTrackToBottom_Execute()
 					if (!ThreadTimingTrack->IsVisible())
 					{
 						ThreadTimingTrack->Show();
-						TimingView->OnTrackVisibilityChanged();
+						TimingView->HandleTrackVisibilityChanged();
 					}
 				}
 			}
@@ -723,7 +723,7 @@ void FContextSwitchesSharedState::Command_DockCpuCoreTrackToTop_Execute()
 					if (!CpuCoreTimingTrack->IsVisible())
 					{
 						CpuCoreTimingTrack->Show();
-						TimingView->OnTrackVisibilityChanged();
+						TimingView->HandleTrackVisibilityChanged();
 					}
 				}
 			}
diff --git a/Engine/Source/Developer/TraceInsights/Private/Insights/MemoryProfiler/ViewModels/MemorySharedState.cpp b/Engine/Source/Developer/TraceInsights/Private/Insights/MemoryProfiler/ViewModels/MemorySharedState.cpp
--- a/Engine/Source/Developer/TraceInsights/Private/Insights/MemoryProfiler/ViewModels/MemorySharedState.cpp
+++ b/Engine/Source/Developer/TraceInsights/Private/Insights/MemoryProfiler/ViewModels/MemorySharedState.cpp
@@ -420,7 +420,7 @@ void FMemorySharedState::SetAllMemoryTracksToggle(bool bOnOff)
 
 	if (TimingView)
 	{
-		TimingView->OnTrackVisibilityChanged();
+		TimingView->HandleTrackVisibilityChanged();
 	}
 }
 
@@ -482,7 +482,7 @@ int32 FMemorySharedState::RemoveMemoryGraphTrack(TSharedPtr<FMemoryGraphTrack> G
 		if (GraphTrack->GetSeries().Num() == 0)
 		{
 			GraphTrack->Hide();
-			TimingView->OnTrackVisibilityChanged();
+			TimingView->HandleTrackVisibilityChanged();
 		}
 		return -1;
 	}
@@ -585,7 +585,7 @@ TSharedPtr<FMemoryGraphTrack> FMemorySharedState::CreateMemTagGraphTrack(Insight
 		}
 
 		MainGraphTrack->Show();
-		TimingView->OnTrackVisibilityChanged();
+		TimingView->HandleTrackVisibilityChanged();
 	}
 
 	TSharedPtr<FMemoryGraphTrack> GraphTrack = GetMemTagGraphTrack(InMemTrackerId, InMemTagId);
@@ -627,7 +627,7 @@ TSharedPtr<FMemoryGraphTrack> FMemorySharedState::CreateMemTagGraphTrack(Insight
 	else
 	{
 		GraphTrack->Show();
-		TimingView->OnTrackVisibilityChanged();
+		TimingView->HandleTrackVisibilityChanged();
 	}
 
 	return GraphTrack;
@@ -655,7 +655,7 @@ int32 FMemorySharedState::RemoveMemTagGraphTrack(Insights::FMemoryTrackerId InMe
 				if (GraphTrack == MainGraphTrack)
 				{
 					GraphTrack->Hide();
-					TimingView->OnTrackVisibilityChanged();
+					TimingView->HandleTrackVisibilityChanged();
 				}
 				else
 				{
@@ -691,7 +691,7 @@ int32 FMemorySharedState::RemoveAllMemTagGraphTracks()
 			if (GraphTrack == MainGraphTrack)
 			{
 				GraphTrack->Hide();
-				TimingView->OnTrackVisibilityChanged();
+				TimingView->HandleTrackVisibilityChanged();
 			}
 			else
 			{
@@ -731,7 +731,7 @@ TSharedPtr<FMemoryGraphSeries> FMemorySharedState::ToggleMemTagGraphSeries(TShar
 		// Remove existing series.
 		InGraphTrack->RemoveMemTagSeries(InMemTrackerId, InMemTagId);
 		InGraphTrack->SetDirtyFlag();
-		TimingView->OnTrackVisibilityChanged();
+		TimingView->HandleTrackVisibilityChanged();
 
 		if (TagPtr)
 		{
@@ -753,7 +753,7 @@ TSharedPtr<FMemoryGraphSeries> FMemorySharedState::ToggleMemTagGraphSeries(TShar
 
 		InGraphTrack->SetDirtyFlag();
 		InGraphTrack->Show();
-		TimingView->OnTrackVisibilityChanged();
+		TimingView->HandleTrackVisibilityChanged();
 
 		return Series;
 	}
diff --git a/Engine/Source/Developer/TraceInsights/Private/Insights/ViewModels/FileActivityTimingTrack.cpp b/Engine/Source/Developer/TraceInsights/Private/Insights/ViewModels/FileActivityTimingTrack.cpp
--- a/Engine/Source/Developer/TraceInsights/Private/Insights/ViewModels/FileActivityTimingTrack.cpp
+++ b/Engine/Source/Developer/TraceInsights/Private/Insights/ViewModels/FileActivityTimingTrack.cpp
@@ -519,7 +519,7 @@ void FFileActivitySharedState::SetAllIoTracksToggle(bool bOnOff)
 
 	if (TimingView)
 	{
-		TimingView->OnTrackVisibilityChanged();
+		TimingView->HandleTrackVisibilityChanged();
 	}
 
 	if (bShowHideAllIoTracks)
@@ -546,7 +546,7 @@ void FFileActivitySharedState::ShowHideIoOverviewTrack()
 
 	if (TimingView)
 	{
-		TimingView->OnTrackVisibilityChanged();
+		TimingView->HandleTrackVisibilityChanged();
 	}
 
 	const bool bIsOverviewTrackVisible = IsIoOverviewTrackVisible();
@@ -581,7 +581,7 @@ void FFileActivitySharedState::ShowHideIoActivityTrack()
 
 	if (TimingView)
 	{
-		TimingView->OnTrackVisibilityChanged();
+		TimingView->HandleTrackVisibilityChanged();
 	}
 
 	const bool bIsOverviewTrackVisible = IsIoOverviewTrackVisible();
diff --git a/Engine/Source/Developer/TraceInsights/Private/Insights/ViewModels/FrameStatsHelper.cpp b/Engine/Source/Developer/TraceInsights/Private/Insights/ViewModels/FrameStatsHelper.cpp
--- a/Engine/Source/Developer/TraceInsights/Private/Insights/ViewModels/FrameStatsHelper.cpp
+++ b/Engine/Source/Developer/TraceInsights/Private/Insights/ViewModels/FrameStatsHelper.cpp
@@ -9,6 +9,18 @@
 namespace Insights
 {
 
+void FFrameStatsHelper::ComputeFrameStatsForTimer(TArray<FFrameStatsCachedEvent>& FrameStatsEvents, uint32 TimerId, const TSet<uint32>& Timelines)
+{
+	TSharedPtr<const TraceServices::IAnalysisSession> Session = FInsightsManager::Get()->GetSession();
+	if (Session.IsValid())
+	{
+		for (uint32 TimelineIndex : Timelines)
+		{
+			ProcessTimeline(FrameStatsEvents, TimerId, TimelineIndex);
+		}
+	}
+}
+
 void FFrameStatsHelper::ComputeFrameStatsForTimer(TArray<FFrameStatsCachedEvent>& FrameStatsEvents, uint32 TimerId)
 {
 	TSharedPtr<const TraceServices::IAnalysisSession> Session = FInsightsManager::Get()->GetSession();
@@ -16,6 +28,22 @@ void FFrameStatsHelper::ComputeFrameStatsForTimer(TArray<FFrameStatsCachedEvent>
 	{
 		Session->ReadAccessCheck();
 
+		const TraceServices::ITimingProfilerProvider& TimingProfilerProvider = *TraceServices::ReadTimingProfilerProvider(*Session.Get());
+
+		for (uint32 TimelineIndex = 0; TimelineIndex < TimingProfilerProvider.GetTimelineCount(); ++TimelineIndex)
+		{
+			ProcessTimeline(FrameStatsEvents, TimerId, TimelineIndex);
+		}
+	}
+}
+
+void FFrameStatsHelper::ProcessTimeline(TArray<FFrameStatsCachedEvent>& FrameStatsEvents, uint32 TimerId, uint32 TimelineIndex)
+{
+	TSharedPtr<const TraceServices::IAnalysisSession> Session = FInsightsManager::Get()->GetSession();
+	if (Session.IsValid())
+	{
+		Session->ReadAccessCheck();
+
 		const double SessionDuration = Session->GetDurationSeconds();
 
 		const TraceServices::ITimingProfilerProvider& TimingProfilerProvider = *TraceServices::ReadTimingProfilerProvider(*Session.Get());
@@ -23,68 +51,64 @@ void FFrameStatsHelper::ComputeFrameStatsForTimer(TArray<FFrameStatsCachedEvent>
 		const TraceServices::ITimingProfilerTimerReader* TimerReader;
 		TimingProfilerProvider.ReadTimers([&TimerReader](const TraceServices::ITimingProfilerTimerReader& Out) { TimerReader = &Out; });
 
-		const uint32 TimelineCount = TimingProfilerProvider.GetTimelineCount();
-		for (uint32 TimelineIndex = 0; TimelineIndex < TimelineCount; ++TimelineIndex)
-		{
-			TimingProfilerProvider.ReadTimeline(TimelineIndex,
-				[SessionDuration, &FrameStatsEvents, TimerReader, TimerId](const TraceServices::ITimingProfilerProvider::Timeline& Timeline)
+		TimingProfilerProvider.ReadTimeline(TimelineIndex,
+			[SessionDuration, &FrameStatsEvents, TimerReader, TimerId](const TraceServices::ITimingProfilerProvider::Timeline& Timeline)
+			{
+				TraceServices::ITimeline<TraceServices::FTimingProfilerEvent>::EnumerateAsyncParams Params;
+				Params.IntervalStart = 0;
+				Params.IntervalEnd = SessionDuration;
+				Params.Resolution = 0.0;
+				Params.SetupCallback = [](uint32 NumTasks) {};
+				Params.Callback = [TimerReader, &FrameStatsEvents, SessionDuration, TimerId](double StartTime, double EndTime, uint32 Depth, const TraceServices::FTimingProfilerEvent& Event, uint32 TaskIndex)
 				{
-					TraceServices::ITimeline<TraceServices::FTimingProfilerEvent>::EnumerateAsyncParams Params;
-					Params.IntervalStart = 0;
-					Params.IntervalEnd = SessionDuration;
-					Params.Resolution = 0.0;
-					Params.SetupCallback = [](uint32 NumTasks) {};
-					Params.Callback = [TimerReader, &FrameStatsEvents, SessionDuration, TimerId](double StartTime, double EndTime, uint32 Depth, const TraceServices::FTimingProfilerEvent& Event, uint32 TaskIndex)
+					const TraceServices::FTimingProfilerTimer* Timer = TimerReader->GetTimer(Event.TimerIndex);
+					if (ensure(Timer != nullptr))
 					{
-						const TraceServices::FTimingProfilerTimer* Timer = TimerReader->GetTimer(Event.TimerIndex);
-						if (ensure(Timer != nullptr))
+						if (Timer->Id == TimerId)
 						{
-							if (Timer->Id == TimerId)
+							int32 Index = Algo::UpperBoundBy(FrameStatsEvents, StartTime, &FFrameStatsCachedEvent::FrameStartTime);
+							if (Index > 0)
+							{
+								--Index;
+							}
+
+							// This can can happen when the event is between frames.
+							if (StartTime > FrameStatsEvents[Index].FrameEndTime)
 							{
-								int32 Index = Algo::UpperBoundBy(FrameStatsEvents, StartTime, &FFrameStatsCachedEvent::FrameStartTime);
-								if (Index > 0)
+								Index++;
+								if (Index >= FrameStatsEvents.Num())
 								{
-									--Index;
+									return TraceServices::EEventEnumerate::Continue;
 								}
+							}
+
+							do
+							{
+								FFrameStatsCachedEvent& Entry = FrameStatsEvents[Index];
 
-								// This can can happen when the event is between frames.
-								if (StartTime > FrameStatsEvents[Index].FrameEndTime)
+								if (EndTime < Entry.FrameStartTime)
 								{
-									Index++;
-									if (Index >= FrameStatsEvents.Num())
-									{
-										return TraceServices::EEventEnumerate::Continue;
-									}
+									return TraceServices::EEventEnumerate::Continue;
 								}
 
-								do
+								if (StartTime < Entry.FrameStartTime)
 								{
-									FFrameStatsCachedEvent& Entry = FrameStatsEvents[Index];
-
-									if (EndTime < Entry.FrameStartTime)
-									{
-										return TraceServices::EEventEnumerate::Continue;
-									}
-
-									if (StartTime < Entry.FrameStartTime)
-									{
-										StartTime = Entry.FrameStartTime;
-									}
+									StartTime = Entry.FrameStartTime;
+								}
 
-									const double Duration = FMath::Min(EndTime, Entry.FrameEndTime) - StartTime;
-									ensure(Duration >= 0.0f);
-									for (double Value = Entry.Duration.load(); !Entry.Duration.compare_exchange_strong(Value, Value + Duration););
+								const double Duration = FMath::Min(EndTime, Entry.FrameEndTime) - StartTime;
+								ensure(Duration >= 0.0f);
+								for (double Value = Entry.Duration.load(); !Entry.Duration.compare_exchange_strong(Value, Value + Duration););
 
-									Index++;
-								} while (Index < FrameStatsEvents.Num());
-							}
+								Index++;
+							} while (Index < FrameStatsEvents.Num());
 						}
-						return TraceServices::EEventEnumerate::Continue;
-					};
+					}
+					return TraceServices::EEventEnumerate::Continue;
+				};
 
-					Timeline.EnumerateEventsDownSampledAsync(Params);
-				});
-		}
+				Timeline.EnumerateEventsDownSampledAsync(Params);
+			});
 	}
 }
 
diff --git a/Engine/Source/Developer/TraceInsights/Private/Insights/ViewModels/FrameStatsHelper.h b/Engine/Source/Developer/TraceInsights/Private/Insights/ViewModels/FrameStatsHelper.h
--- a/Engine/Source/Developer/TraceInsights/Private/Insights/ViewModels/FrameStatsHelper.h
+++ b/Engine/Source/Developer/TraceInsights/Private/Insights/ViewModels/FrameStatsHelper.h
@@ -34,7 +34,11 @@ struct FFrameStatsCachedEvent
 class FFrameStatsHelper
 {
 public:
+	static void ComputeFrameStatsForTimer(TArray<FFrameStatsCachedEvent>& FrameStatsEvents, uint32 TimerId, const TSet<uint32>& Timelines);
 	static void ComputeFrameStatsForTimer(TArray<FFrameStatsCachedEvent>& FrameStatsEvents, uint32 TimerId);
+
+private:
+	static void ProcessTimeline(TArray<FFrameStatsCachedEvent>& FrameStatsEvents, uint32 TimerId, uint32 TimelineIndex);
 };
 
 // namespace Insights
diff --git a/Engine/Source/Developer/TraceInsights/Private/Insights/ViewModels/FrameTimingTrack.cpp b/Engine/Source/Developer/TraceInsights/Private/Insights/ViewModels/FrameTimingTrack.cpp
--- a/Engine/Source/Developer/TraceInsights/Private/Insights/ViewModels/FrameTimingTrack.cpp
+++ b/Engine/Source/Developer/TraceInsights/Private/Insights/ViewModels/FrameTimingTrack.cpp
@@ -196,7 +196,7 @@ void FFrameSharedState::SetAllFrameTracksToggle(bool bOnOff)
 
 	if (TimingView)
 	{
-		TimingView->OnTrackVisibilityChanged();
+		TimingView->HandleTrackVisibilityChanged();
 	}
 }
 
diff --git a/Engine/Source/Developer/TraceInsights/Private/Insights/ViewModels/LoadingTimingTrack.cpp b/Engine/Source/Developer/TraceInsights/Private/Insights/ViewModels/LoadingTimingTrack.cpp
--- a/Engine/Source/Developer/TraceInsights/Private/Insights/ViewModels/LoadingTimingTrack.cpp
+++ b/Engine/Source/Developer/TraceInsights/Private/Insights/ViewModels/LoadingTimingTrack.cpp
@@ -201,7 +201,7 @@ void FLoadingSharedState::SetAllLoadingTracksToggle(bool bOnOff)
 		Track.SetVisibilityFlag(bShowHideAllLoadingTracks);
 	}
 
-	TimingView->OnTrackVisibilityChanged();
+	TimingView->HandleTrackVisibilityChanged();
 }
 
 ////////////////////////////////////////////////////////////////////////////////////////////////////
diff --git a/Engine/Source/Developer/TraceInsights/Private/Insights/ViewModels/RegionsTimingTrack.cpp b/Engine/Source/Developer/TraceInsights/Private/Insights/ViewModels/RegionsTimingTrack.cpp
--- a/Engine/Source/Developer/TraceInsights/Private/Insights/ViewModels/RegionsTimingTrack.cpp
+++ b/Engine/Source/Developer/TraceInsights/Private/Insights/ViewModels/RegionsTimingTrack.cpp
@@ -94,7 +94,7 @@ void FTimingRegionsSharedState::ShowHideRegionsTrack()
 
 	if (TimingView)
 	{
-		TimingView->OnTrackVisibilityChanged();
+		TimingView->HandleTrackVisibilityChanged();
 	}
 
 	if (bShowHideRegionsTrack)
diff --git a/Engine/Source/Developer/TraceInsights/Private/Insights/ViewModels/ThreadTimingTrack.cpp b/Engine/Source/Developer/TraceInsights/Private/Insights/ViewModels/ThreadTimingTrack.cpp
--- a/Engine/Source/Developer/TraceInsights/Private/Insights/ViewModels/ThreadTimingTrack.cpp
+++ b/Engine/Source/Developer/TraceInsights/Private/Insights/ViewModels/ThreadTimingTrack.cpp
@@ -413,6 +413,21 @@ void FThreadTimingSharedState::GetVisibleCpuThreads(TSet<uint32>& OutSet) const
 
 ////////////////////////////////////////////////////////////////////////////////////////////////////
 
+void FThreadTimingSharedState::GetVisibleTimelineIndexes(TSet<uint32>& OutSet) const
+{
+	OutSet.Reset();
+	for (const auto& KV : CpuTracks)
+	{
+		const FCpuTimingTrack& Track = *KV.Value;
+		if (Track.IsVisible())
+		{
+			OutSet.Add(Track.GetTimelineIndex());
+		}
+	}
+}
+
+////////////////////////////////////////////////////////////////////////////////////////////////////
+
 void FThreadTimingSharedState::OnBeginSession(Insights::ITimingViewSession& InSession)
 {
 	if (&InSession != TimingView)
@@ -712,7 +727,7 @@ void FThreadTimingSharedState::SetAllCpuTracksToggle(bool bOnOff)
 		KV.Value.bIsVisible = bShowHideAllCpuTracks;
 	}
 
-	TimingView->OnTrackVisibilityChanged();
+	TimingView->HandleTrackVisibilityChanged();
 }
 
 ////////////////////////////////////////////////////////////////////////////////////////////////////
@@ -731,7 +746,7 @@ void FThreadTimingSharedState::SetAllGpuTracksToggle(bool bOnOff)
 	}
 	if (GpuTrack.IsValid() || Gpu2Track.IsValid())
 	{
-		TimingView->OnTrackVisibilityChanged();
+		TimingView->HandleTrackVisibilityChanged();
 	}
 }
 
@@ -765,7 +780,7 @@ void FThreadTimingSharedState::ToggleTrackVisibilityByGroup_Execute(const TCHAR*
 			}
 		}
 
-		TimingView->OnTrackVisibilityChanged();
+		TimingView->HandleTrackVisibilityChanged();
 	}
 }
 
diff --git a/Engine/Source/Developer/TraceInsights/Private/Insights/ViewModels/ThreadTimingTrack.h b/Engine/Source/Developer/TraceInsights/Private/Insights/ViewModels/ThreadTimingTrack.h
--- a/Engine/Source/Developer/TraceInsights/Private/Insights/ViewModels/ThreadTimingTrack.h
+++ b/Engine/Source/Developer/TraceInsights/Private/Insights/ViewModels/ThreadTimingTrack.h
@@ -67,6 +67,7 @@ public:
 	bool IsCpuTrackVisible(uint32 InThreadId) const;
 
 	void GetVisibleCpuThreads(TSet<uint32>& OutSet) const;
+	void GetVisibleTimelineIndexes(TSet<uint32>& OutSet) const;
 
 	//////////////////////////////////////////////////
 	// ITimingViewExtender interface
diff --git a/Engine/Source/Developer/TraceInsights/Private/Insights/ViewModels/TimingGraphTrack.cpp b/Engine/Source/Developer/TraceInsights/Private/Insights/ViewModels/TimingGraphTrack.cpp
--- a/Engine/Source/Developer/TraceInsights/Private/Insights/ViewModels/TimingGraphTrack.cpp
+++ b/Engine/Source/Developer/TraceInsights/Private/Insights/ViewModels/TimingGraphTrack.cpp
@@ -12,11 +12,15 @@
 #include "Insights/Common/PaintUtils.h"
 #include "Insights/Common/TimeUtils.h"
 #include "Insights/InsightsManager.h"
+#include "Insights/TimingProfilerManager.h"
 #include "Insights/ViewModels/AxisViewportDouble.h"
 #include "Insights/ViewModels/FrameStatsHelper.h"
 #include "Insights/ViewModels/GraphTrackBuilder.h"
 #include "Insights/ViewModels/ITimingViewDrawHelper.h"
 #include "Insights/ViewModels/TimingTrackViewport.h"
+#include "Insights/ViewModels/ThreadTimingTrack.h"
+#include "Insights/Widgets/STimingProfilerWindow.h"
+#include "Insights/Widgets/STimingView.h"
 
 #define LOCTEXT_NAMESPACE "GraphTrack"
 
@@ -137,6 +141,18 @@ FTimingGraphTrack::FTimingGraphTrack()
 
 FTimingGraphTrack::~FTimingGraphTrack()
 {
+	if (OnTrackVisibilityChangedHandle.IsValid())
+	{
+		TSharedPtr<class STimingProfilerWindow> TimingWindow = FTimingProfilerManager::Get()->GetProfilerWindow();
+		if (TimingWindow.IsValid())
+		{
+			TSharedPtr<STimingView> TimingView = TimingWindow->GetTimingView();
+			if (TimingView.IsValid())
+			{
+				TimingView->OnTrackVisibilityChanged().Remove(OnTrackVisibilityChangedHandle);
+			}
+		}
+	}
 }
 
 ////////////////////////////////////////////////////////////////////////////////////////////////////
@@ -144,6 +160,21 @@ FTimingGraphTrack::~FTimingGraphTrack()
 void FTimingGraphTrack::Update(const ITimingTrackUpdateContext& Context)
 {
 	FGraphTrack::Update(Context);
+	if (!OnTrackVisibilityChangedHandle.IsValid())
+	{
+		TSharedPtr<class STimingProfilerWindow> TimingWindow = FTimingProfilerManager::Get()->GetProfilerWindow();
+		if (TimingWindow.IsValid())
+		{
+			TSharedPtr<STimingView> TimingView = TimingWindow->GetTimingView();
+			if (TimingView.IsValid())
+			{
+				OnTrackVisibilityChangedHandle = TimingView->OnTrackVisibilityChanged().AddLambda([this]()
+					{
+						this->SetDirtyFlag();
+					});
+			}
+		}
+	}
 
 	const bool bIsEntireGraphTrackDirty = IsDirty() || Context.GetViewport().IsHorizontalViewportDirty();
 	bool bNeedsUpdate = bIsEntireGraphTrackDirty;
@@ -319,8 +350,10 @@ void FTimingGraphTrack::UpdateTimerSeries(FTimingGraphSeries& Series, const FTim
 	{
 		TraceServices::FAnalysisSessionReadScope SessionReadScope(*Session.Get());
 
+		TSet<uint32> Timelines;
+		GetVisibleTimelineIndexes(Timelines);
 		const double SessionDuration = Session->GetDurationSeconds();
-		if (Series.CachedSessionDuration != SessionDuration)
+		if (Series.CachedSessionDuration != SessionDuration || Series.CachedTimelinesNum != Timelines.Num())
 		{
 			Series.CachedSessionDuration = SessionDuration;
 
@@ -331,7 +364,11 @@ void FTimingGraphTrack::UpdateTimerSeries(FTimingGraphSeries& Series, const FTim
 
 			const uint32 TimelineCount = TimingProfilerProvider.GetTimelineCount();
 			uint32 NumTimelinesContainingEvent = 0;
-			for (uint32 TimelineIndex = 0; TimelineIndex < TimelineCount; ++TimelineIndex)
+
+			Series.CachedTimelinesNum = Timelines.Num();
+			Series.CachedEvents.Empty();
+
+			for (uint32 TimelineIndex : Timelines)
 			{
 				TimingProfilerProvider.ReadTimeline(TimelineIndex,
 					[SessionDuration, &Series, TimerReader, &Viewport, &NumTimelinesContainingEvent](const TraceServices::ITimingProfilerProvider::Timeline& Timeline)
@@ -466,10 +503,13 @@ void FTimingGraphTrack::UpdateFrameStatsTimerSeries(FTimingGraphSeries& Series,
 	{
 		TraceServices::FAnalysisSessionReadScope SessionReadScope(*Session.Get());
 
+		TSet<uint32> VisibleTimelines;
+		GetVisibleTimelineIndexes(VisibleTimelines);
 		const double SessionDuration = Session->GetDurationSeconds();
-		if (Series.CachedSessionDuration != SessionDuration)
+		if (Series.CachedSessionDuration != SessionDuration || Series.CachedTimelinesNum != VisibleTimelines.Num())
 		{
 			Series.CachedSessionDuration = SessionDuration;
+			Series.CachedTimelinesNum = VisibleTimelines.Num();
 
 			const TraceServices::IFrameProvider& FramesProvider = ReadFrameProvider(*Session.Get());
 
@@ -489,7 +529,7 @@ void FTimingGraphTrack::UpdateFrameStatsTimerSeries(FTimingGraphSeries& Series,
 					Series.FrameStatsCachedEvents.Add(Event);
 				});
 
-			Insights::FFrameStatsHelper::ComputeFrameStatsForTimer(Series.FrameStatsCachedEvents, Series.TimerId);
+			Insights::FFrameStatsHelper::ComputeFrameStatsForTimer(Series.FrameStatsCachedEvents, Series.TimerId, VisibleTimelines);
 		}
 
 		int32 StartIndex = Algo::UpperBoundBy(Series.FrameStatsCachedEvents, Viewport.GetStartTime(), &Insights::FFrameStatsCachedEvent::FrameStartTime);
@@ -657,6 +697,26 @@ void FTimingGraphTrack::UpdateStatsCounterSeries(FTimingGraphSeries& Series, con
 	}
 }
 
+////////////////////////////////////////////////////////////////////////////////////////////////////
+
+void FTimingGraphTrack::GetVisibleTimelineIndexes(TSet<uint32>& TimelineIndexes)
+{
+	TSharedPtr<class STimingProfilerWindow> TimingWindow = FTimingProfilerManager::Get()->GetProfilerWindow();
+	if (!TimingWindow.IsValid())
+	{
+		return;
+	}
+
+	TSharedPtr<STimingView> TimingView = TimingWindow->GetTimingView();
+	if (!TimingView.IsValid())
+	{
+		return;
+	}
+
+	TSharedPtr<FThreadTimingSharedState> ThreadSharedState = TimingView->GetThreadTimingSharedState();
+	ThreadSharedState->GetVisibleTimelineIndexes(TimelineIndexes);
+}
+
 ////////////////////////////////////////////////////////////////////////////////////////////////////
 ////////////////////////////////////////////////////////////////////////////////////////////////////
 
diff --git a/Engine/Source/Developer/TraceInsights/Private/Insights/ViewModels/TimingGraphTrack.h b/Engine/Source/Developer/TraceInsights/Private/Insights/ViewModels/TimingGraphTrack.h
--- a/Engine/Source/Developer/TraceInsights/Private/Insights/ViewModels/TimingGraphTrack.h
+++ b/Engine/Source/Developer/TraceInsights/Private/Insights/ViewModels/TimingGraphTrack.h
@@ -55,6 +55,7 @@ public:
 	double CachedSessionDuration;
 	TArray<FSimpleTimingEvent> CachedEvents; // used by Timer series
 	TArray<Insights::FFrameStatsCachedEvent> FrameStatsCachedEvents; // used by Frame Stats Timer series
+	uint32 CachedTimelinesNum = 0; // the number of timelines used to gather the data
 
 	bool bIsTime; // the unit for values is [second]
 	bool bIsMemory; // the unit for value is [byte]
@@ -93,7 +94,12 @@ protected:
 	void UpdateFrameStatsTimerSeries(FTimingGraphSeries& Series, const FTimingTrackViewport& Viewport);
 	void UpdateStatsCounterSeries(FTimingGraphSeries& Series, const FTimingTrackViewport& Viewport);
 
+	void GetVisibleTimelineIndexes(TSet<uint32>& TimelineIndexes);
+
 	virtual void DrawVerticalAxisGrid(const ITimingTrackDrawContext& Context) const override;
+
+private:
+	FDelegateHandle OnTrackVisibilityChangedHandle;
 };
 
 ////////////////////////////////////////////////////////////////////////////////////////////////////
diff --git a/Engine/Source/Developer/TraceInsights/Private/Insights/Widgets/SFrameTrack.cpp b/Engine/Source/Developer/TraceInsights/Private/Insights/Widgets/SFrameTrack.cpp
--- a/Engine/Source/Developer/TraceInsights/Private/Insights/Widgets/SFrameTrack.cpp
+++ b/Engine/Source/Developer/TraceInsights/Private/Insights/Widgets/SFrameTrack.cpp
@@ -25,6 +25,7 @@
 #include "Insights/TimingProfilerManager.h"
 #include "Insights/ViewModels/FrameStatsHelper.h"
 #include "Insights/ViewModels/FrameTrackHelper.h"
+#include "Insights/ViewModels/ThreadTimingTrack.h"
 #include "Insights/Widgets/STimingProfilerWindow.h"
 #include "Insights/Widgets/STimingView.h"
 
@@ -45,6 +46,21 @@ SFrameTrack::SFrameTrack()
 
 SFrameTrack::~SFrameTrack()
 {
+	if (OnTrackVisibilityChangedHandle.IsValid())
+	{
+		TSharedPtr<class STimingProfilerWindow> TimingWindow = FTimingProfilerManager::Get()->GetProfilerWindow();
+		if (TimingWindow.IsValid())
+		{
+			TSharedPtr<STimingView> TimingView = TimingWindow->GetTimingView();
+			if (TimingView.IsValid())
+			{
+				if (TimingView.Get() == RegisteredTimingView)
+				{
+					TimingView->OnTrackVisibilityChanged().Remove(OnTrackVisibilityChangedHandle);
+				}
+			}
+		}
+	}
 }
 
 ////////////////////////////////////////////////////////////////////////////////////////////////////
@@ -136,6 +152,27 @@ void SFrameTrack::Construct(const FArguments& InArgs)
 
 void SFrameTrack::Tick(const FGeometry& AllottedGeometry, const double InCurrentTime, const float InDeltaTime)
 {
+	TSharedPtr<class STimingProfilerWindow> TimingWindow = FTimingProfilerManager::Get()->GetProfilerWindow();
+	if (TimingWindow.IsValid())
+	{
+		TSharedPtr<STimingView> TimingView = TimingWindow->GetTimingView();
+		if (TimingView.IsValid())
+		{
+			if (!OnTrackVisibilityChangedHandle.IsValid() || TimingView.Get() != RegisteredTimingView)
+			{
+				RegisteredTimingView = TimingView.Get();
+				OnTrackVisibilityChangedHandle = TimingView->OnTrackVisibilityChanged().AddLambda([this]()
+					{
+						// If there are more series than the default frame series.
+						if (this->AllSeries.Num() > ETraceFrameType::TraceFrameType_Count)
+						{
+							this->bIsStateDirty = true;
+						}
+					});
+			}
+		}
+	}
+
 	if (ThisGeometry != AllottedGeometry || bIsViewportDirty)
 	{
 		bIsViewportDirty = false;
@@ -336,7 +373,31 @@ void SFrameTrack::UpdateState()
 
 			FFrameTrackSeriesBuilder Builder(*Series, Viewport);
 
-			Insights::FFrameStatsHelper::ComputeFrameStatsForTimer(Frames, TimerSeries->TimerId);
+			bool bTimingViewExists = false;
+			TSet<uint32> Timelines;
+			TSharedPtr<class STimingProfilerWindow> TimingWindow = FTimingProfilerManager::Get()->GetProfilerWindow();
+
+			// Attemp to compute only from visible timelines.
+			if (TimingWindow.IsValid())
+			{
+				TSharedPtr<STimingView> TimingView = TimingWindow->GetTimingView();
+				if (TimingView.IsValid())
+				{
+					TSharedPtr<FThreadTimingSharedState> ThreadSharedState = TimingView->GetThreadTimingSharedState();
+					if (ThreadSharedState.IsValid())
+					{
+						ThreadSharedState->GetVisibleTimelineIndexes(Timelines);
+						Insights::FFrameStatsHelper::ComputeFrameStatsForTimer(Frames, TimerSeries->TimerId, Timelines);
+						bTimingViewExists = true;
+					}
+				}
+			}
+
+			if (!bTimingViewExists)
+			{
+				// Compute the stats for all timelines. 
+				Insights::FFrameStatsHelper::ComputeFrameStatsForTimer(Frames, TimerSeries->TimerId);
+			}
 
 			uint64 CurrentIndex = StartIndex;
 			for (Insights::FFrameStatsCachedEvent& Event : Frames)
diff --git a/Engine/Source/Developer/TraceInsights/Private/Insights/Widgets/SFrameTrack.h b/Engine/Source/Developer/TraceInsights/Private/Insights/Widgets/SFrameTrack.h
--- a/Engine/Source/Developer/TraceInsights/Private/Insights/Widgets/SFrameTrack.h
+++ b/Engine/Source/Developer/TraceInsights/Private/Insights/Widgets/SFrameTrack.h
@@ -16,6 +16,7 @@
 #include "Insights/ViewModels/FrameTrackViewport.h"
 
 class SScrollBar;
+class STimingView;
 
 ////////////////////////////////////////////////////////////////////////////////////////////////////
 
@@ -240,6 +241,9 @@ protected:
 	/** Cursor type. */
 	ECursorType CursorType;
 
+	STimingView* RegisteredTimingView = nullptr; // For pointer comparison only, do not dereferentiate.
+	FDelegateHandle OnTrackVisibilityChangedHandle;
+
 	// Debug stats
 	int32 NumUpdatedFrames;
 	TFixedCircularBuffer<uint64, 32> UpdateDurationHistory;
diff --git a/Engine/Source/Developer/TraceInsights/Private/Insights/Widgets/STimingView.cpp b/Engine/Source/Developer/TraceInsights/Private/Insights/Widgets/STimingView.cpp
--- a/Engine/Source/Developer/TraceInsights/Private/Insights/Widgets/STimingView.cpp
+++ b/Engine/Source/Developer/TraceInsights/Private/Insights/Widgets/STimingView.cpp
@@ -1824,7 +1824,7 @@ void STimingView::HideAllScrollableTracks()
 	{
 		Track->Hide();
 	}
-	OnTrackVisibilityChanged();
+	HandleTrackVisibilityChanged();
 }
 
 ////////////////////////////////////////////////////////////////////////////////////////////////////
@@ -3114,7 +3114,7 @@ void STimingView::ChangeTrackLocation(TSharedRef<FBaseTimingTrack> Track, ETimin
 			break;
 		}
 
-		OnTrackVisibilityChanged();
+		HandleTrackVisibilityChanged();
 	}
 }
 
@@ -4872,7 +4872,7 @@ void STimingView::ToggleTrackVisibility_Execute(uint64 InTrackId)
 	if (TrackPtrPtr)
 	{
 		(*TrackPtrPtr)->ToggleVisibility();
-		OnTrackVisibilityChanged();
+		HandleTrackVisibilityChanged();
 	}
 }
 
@@ -5022,7 +5022,7 @@ void STimingView::CloseQuickFindTab()
 
 ////////////////////////////////////////////////////////////////////////////////////////////////////
 
-void STimingView::OnTrackVisibilityChanged()
+void STimingView::HandleTrackVisibilityChanged()
 {
 	if (HoveredTrack.IsValid())
 	{
@@ -5046,7 +5046,7 @@ void STimingView::OnTrackVisibilityChanged()
 	}
 	Tooltip.SetDesiredOpacity(0.0f);
 
-	//TODO: TrackVisibilityChangedEvent.Broadcast();
+	OnTrackVisibilityChangedDelegate.Broadcast();
 	FTimingProfilerManager::Get()->OnThreadFilterChanged();
 }
 
diff --git a/Engine/Source/Developer/TraceInsights/Private/Insights/Widgets/STimingView.h b/Engine/Source/Developer/TraceInsights/Private/Insights/Widgets/STimingView.h
--- a/Engine/Source/Developer/TraceInsights/Private/Insights/Widgets/STimingView.h
+++ b/Engine/Source/Developer/TraceInsights/Private/Insights/Widgets/STimingView.h
@@ -272,6 +272,7 @@ public:
 	virtual Insights::FHoveredEventChangedDelegate& OnHoveredEventChanged() override { return OnHoveredEventChangedDelegate; }
 	virtual Insights::FSelectedTrackChangedDelegate& OnSelectedTrackChanged() override { return OnSelectedTrackChangedDelegate; }
 	virtual Insights::FSelectedEventChangedDelegate& OnSelectedEventChanged() override { return OnSelectedEventChangedDelegate; }
+	virtual Insights::FTrackVisibilityChangedDelegate& OnTrackVisibilityChanged() override { return OnTrackVisibilityChangedDelegate; }
 
 	virtual void ResetSelectedEvent() override
 	{
@@ -317,7 +318,7 @@ public:
 
 	void HideAllScrollableTracks();
 
-	void OnTrackVisibilityChanged();
+	void HandleTrackVisibilityChanged();
 
 	bool IsGpuTrackVisible() const;
 	bool IsCpuTrackVisible(uint32 InThreadId) const;
@@ -717,6 +718,7 @@ protected:
 	Insights::FHoveredEventChangedDelegate OnHoveredEventChangedDelegate;
 	Insights::FSelectedTrackChangedDelegate OnSelectedTrackChangedDelegate;
 	Insights::FSelectedEventChangedDelegate OnSelectedEventChangedDelegate;
+	Insights::FTrackVisibilityChangedDelegate OnTrackVisibilityChangedDelegate;
 
 	TSharedPtr<FUICommandList> CommandList;
 
diff --git a/Engine/Source/Developer/TraceInsights/Private/Insights/Widgets/STimingViewTrackList.cpp b/Engine/Source/Developer/TraceInsights/Private/Insights/Widgets/STimingViewTrackList.cpp
--- a/Engine/Source/Developer/TraceInsights/Private/Insights/Widgets/STimingViewTrackList.cpp
+++ b/Engine/Source/Developer/TraceInsights/Private/Insights/Widgets/STimingViewTrackList.cpp
@@ -62,7 +62,7 @@ public:
 				{
 					Track.Pin()->SetVisibilityFlag(InCheckBoxState == ECheckBoxState::Checked);
 					Track.Pin()->SetDirtyFlag();
-					TimingView.Pin()->OnTrackVisibilityChanged();
+					TimingView.Pin()->HandleTrackVisibilityChanged();
 				})
 				.Content()
 				[
@@ -165,7 +165,7 @@ void STimingViewTrackList::Construct(const FArguments& InArgs, const TSharedRef<
 						Track->SetVisibilityFlag(bVisible);
 						Track->SetDirtyFlag();
 					}
-					TimingView.Pin()->OnTrackVisibilityChanged();
+					TimingView.Pin()->HandleTrackVisibilityChanged();
 				})
 			]
 			+ SHorizontalBox::Slot()
diff --git a/Engine/Source/Developer/TraceInsights/Public/Insights/ITimingViewSession.h b/Engine/Source/Developer/TraceInsights/Public/Insights/ITimingViewSession.h
--- a/Engine/Source/Developer/TraceInsights/Public/Insights/ITimingViewSession.h
+++ b/Engine/Source/Developer/TraceInsights/Public/Insights/ITimingViewSession.h
@@ -58,6 +58,9 @@ DECLARE_MULTICAST_DELEGATE_OneParam(FSelectedTrackChangedDelegate, const TShared
 /** The delegate to be invoked when the selected timing event has changed */
 DECLARE_MULTICAST_DELEGATE_OneParam(FSelectedEventChangedDelegate, const TSharedPtr<const ITimingEvent> /*InEvent*/);
 
+/** The delegate to be invoked when a track visibility has changed */
+DECLARE_MULTICAST_DELEGATE(FTrackVisibilityChangedDelegate);
+
 ////////////////////////////////////////////////////////////////////////////////////////////////////
 
 /** Hosts a number of timing view visualizers, represents a session of the timing view. */
@@ -127,6 +130,9 @@ public:
 	/** Gets the delegate to be invoked when the selected timing event has changed. */
 	virtual FSelectedEventChangedDelegate& OnSelectedEventChanged() = 0;
 
+	/** Gets the delegate to be invoked when the track visibility has changed. */
+	virtual FTrackVisibilityChangedDelegate& OnTrackVisibilityChanged() = 0;
+
 	//////////////////////////////////////////////////
 
 	/** Resets the selected event back to empty. */
