[Insights] Fix crash when closing a table view that was never in focus. Changed the implementation of the TableTreeView so it now waits on async tasks when it is closed.

diff --git a/Engine/Source/Developer/TraceInsights/Private/Insights/InsightsManager.cpp b/Engine/Source/Developer/TraceInsights/Private/Insights/InsightsManager.cpp
--- a/Engine/Source/Developer/TraceInsights/Private/Insights/InsightsManager.cpp
+++ b/Engine/Source/Developer/TraceInsights/Private/Insights/InsightsManager.cpp
@@ -1282,52 +1282,4 @@ bool FInsightsManager::HandleResponseFileCmd(const TCHAR* ResponseFile, FOutputD
 
 ////////////////////////////////////////////////////////////////////////////////////////////////////
 
-void FInsightsManager::AddInProgressAsyncOp(FGraphEventRef Event, const FString& Name)
-{
-	check(IsInGameThread());
-	if (Event.IsValid() && !Event->IsComplete())
-	{
-		InProgressAsyncTasks.AddTail(FAsyncTaskData(Event, Name));
-	}
-}
-
-////////////////////////////////////////////////////////////////////////////////////////////////////
-
-void FInsightsManager::RemoveInProgressAsyncOp(FGraphEventRef Event)
-{
-	check(IsInGameThread());
-	if (Event.IsValid())
-	{
-		auto Element = InProgressAsyncTasks.GetHead();
-		while (Element && Element->GetValue().GraphEvent != Event)
-		{
-			Element = Element->GetNextNode();
-		}
-		if (Element)
-		{
-			InProgressAsyncTasks.RemoveNode(Element);
-		}
-	}
-}
-
-////////////////////////////////////////////////////////////////////////////////////////////////////
-
-void FInsightsManager::WaitOnInProgressAsyncOps()
-{
-	check(IsInGameThread());
-	while(!InProgressAsyncTasks.IsEmpty())
-	{
-		auto Entry = InProgressAsyncTasks.GetHead();
-		UE_LOG(TraceInsights, Log, TEXT("Waiting on task: %s."), *(Entry->GetValue().Name));
-		FGraphEventRef Event = Entry->GetValue().GraphEvent;
-		while (!Event->IsComplete())
-		{
-			FTaskGraphInterface::Get().ProcessThreadUntilIdle(ENamedThreads::GameThread);
-			FPlatformProcess::Sleep(0.1f);
-		}
-	}
-}
-
-////////////////////////////////////////////////////////////////////////////////////////////////////
-
 #undef LOCTEXT_NAMESPACE
diff --git a/Engine/Source/Developer/TraceInsights/Private/Insights/InsightsManager.h b/Engine/Source/Developer/TraceInsights/Private/Insights/InsightsManager.h
--- a/Engine/Source/Developer/TraceInsights/Private/Insights/InsightsManager.h
+++ b/Engine/Source/Developer/TraceInsights/Private/Insights/InsightsManager.h
@@ -328,10 +328,6 @@ public:
 
 	void ScheduleCommand(const FString& InCmd);
 
-	void AddInProgressAsyncOp(FGraphEventRef Event, const FString& Name);
-	void RemoveInProgressAsyncOp(FGraphEventRef Event);
-	void WaitOnInProgressAsyncOps();
-
 	////////////////////////////////////////////////////////////////////////////////////////////////////
 	// SessionChangedEvent
 
diff --git a/Engine/Source/Developer/TraceInsights/Private/Insights/Table/Widgets/STableTreeView.cpp b/Engine/Source/Developer/TraceInsights/Private/Insights/Table/Widgets/STableTreeView.cpp
--- a/Engine/Source/Developer/TraceInsights/Private/Insights/Table/Widgets/STableTreeView.cpp
+++ b/Engine/Source/Developer/TraceInsights/Private/Insights/Table/Widgets/STableTreeView.cpp
@@ -146,11 +146,14 @@ STableTreeView::STableTreeView()
 
 STableTreeView::~STableTreeView()
 {
-	if (bRunInAsyncMode)
+	if (bRunInAsyncMode && !bIsCloseScheduled)
 	{
-		ensureMsgf(bIsCloseScheduled, TEXT("TableTreeView running in async mode was closed but OnClose() was not called. This can lead to a crash. Call OnClose() from the owner tab/window."));
+		UE_LOG(TraceInsights, Log, TEXT("TableTreeView running in async mode was closed but OnClose() was not called. Call OnClose() from the owner tab/window."));
 	}
 
+	// Backup call to OnClose() in case it was not called from the owner.
+	OnClose();
+
 	if (CurrentAsyncOpFilterConfigurator)
 	{
 		delete CurrentAsyncOpFilterConfigurator;
@@ -3359,14 +3362,11 @@ void STableTreeView::OnClose()
 
 	bIsCloseScheduled = true;
 
-	if (bIsUpdateRunning && InProgressAsyncOperationEvent.IsValid() && !InProgressAsyncOperationEvent->IsComplete())
+	if (bIsUpdateRunning)
 	{
 		CancelCurrentAsyncOp();
-
-		FGraphEventArray Prerequisites;
-		Prerequisites.Add(InProgressAsyncOperationEvent);
-		AsyncCompleteTaskEvent = TGraphTask<FTableTreeViewAsyncCompleteTask>::CreateTask(&Prerequisites).ConstructAndDispatchWhenReady(SharedThis(this));
-		FInsightsManager::Get()->AddInProgressAsyncOp(AsyncCompleteTaskEvent, TEXT("FTableTreeViewAsyncCompleteTask"));
+		check(InProgressAsyncOperationEvent.IsValid());
+		InProgressAsyncOperationEvent->Wait(ENamedThreads::GameThread);
 	}
 }
 
@@ -3416,6 +3416,11 @@ void STableTreeView::CancelCurrentAsyncOp()
 {
 	if (bIsUpdateRunning)
 	{
+		if (DispatchEvent.IsValid() && !DispatchEvent->IsComplete())
+		{
+			DispatchEvent->DispatchSubsequents();
+		}
+
 		AsyncOperationProgress.CancelCurrentAsyncOp();
 	}
 }
diff --git a/Engine/Source/Developer/TraceInsights/Private/Insights/Table/Widgets/STableTreeView.h b/Engine/Source/Developer/TraceInsights/Private/Insights/Table/Widgets/STableTreeView.h
--- a/Engine/Source/Developer/TraceInsights/Private/Insights/Table/Widgets/STableTreeView.h
+++ b/Engine/Source/Developer/TraceInsights/Private/Insights/Table/Widgets/STableTreeView.h
@@ -746,35 +746,6 @@ private:
 
 ////////////////////////////////////////////////////////////////////////////////////////////////////
 
-class FTableTreeViewAsyncCompleteTask
-{
-public:
-	FTableTreeViewAsyncCompleteTask(TSharedPtr<STableTreeView> InPtr)
-	{
-		TableTreeViewPtr = InPtr;
-	}
-
-	FORCEINLINE TStatId GetStatId() const { RETURN_QUICK_DECLARE_CYCLE_STAT(FTableTreeViewAsyncCompleteTask, STATGROUP_TaskGraphTasks); }
-	ENamedThreads::Type GetDesiredThread() { return ENamedThreads::Type::GameThread; }
-	static ESubsequentsMode::Type GetSubsequentsMode() { return ESubsequentsMode::TrackSubsequents; }
-
-	void DoTask(ENamedThreads::Type CurrentThread, const FGraphEventRef& MyCompletionGraphEvent)
-	{
-		// The role of this task is to keep the STableTreeView object alive until the task and it's prerequisites are completed and to destroy it on the game thread.
-		FGraphEventRef Event = TableTreeViewPtr->AsyncCompleteTaskEvent;
-		if (TableTreeViewPtr.IsValid())
-		{
-			TableTreeViewPtr.Reset();
-		}
-		FInsightsManager::Get()->RemoveInProgressAsyncOp(Event);
-	}
-
-private:
-	TSharedPtr<STableTreeView> TableTreeViewPtr;
-};
-
-////////////////////////////////////////////////////////////////////////////////////////////////////
-
 class FSearchForItemToSelectTask
 {
 public:
diff --git a/Engine/Source/Developer/TraceInsights/Private/Insights/TraceInsightsModule.cpp b/Engine/Source/Developer/TraceInsights/Private/Insights/TraceInsightsModule.cpp
--- a/Engine/Source/Developer/TraceInsights/Private/Insights/TraceInsightsModule.cpp
+++ b/Engine/Source/Developer/TraceInsights/Private/Insights/TraceInsightsModule.cpp
@@ -86,8 +86,6 @@ void FTraceInsightsModule::ShutdownModule()
 {
 	LLM_SCOPE_BYTAG(Insights);
 
-	FInsightsManager::Get()->WaitOnInProgressAsyncOps();
-
 #if !WITH_EDITOR
 	if (FModuleManager::Get().IsModuleLoaded("SourceCodeAccess"))
 	{
