[Insights] The Graph Track and the Frame Track update timer series when a new Thread Timing Track is added or removed.

diff --git a/Engine/Source/Developer/TraceInsights/Private/Insights/ViewModels/TimingGraphTrack.cpp b/Engine/Source/Developer/TraceInsights/Private/Insights/ViewModels/TimingGraphTrack.cpp
--- a/Engine/Source/Developer/TraceInsights/Private/Insights/ViewModels/TimingGraphTrack.cpp
+++ b/Engine/Source/Developer/TraceInsights/Private/Insights/ViewModels/TimingGraphTrack.cpp
@@ -150,6 +150,8 @@ FTimingGraphTrack::~FTimingGraphTrack()
 			if (TimingView.IsValid())
 			{
 				TimingView->OnTrackVisibilityChanged().Remove(OnTrackVisibilityChangedHandle);
+				TimingView->OnTrackAdded().Remove(OnTrackAddedHandle);
+				TimingView->OnTrackRemoved().Remove(OnTrackRemovedHandle);
 			}
 		}
 	}
@@ -168,6 +170,21 @@ void FTimingGraphTrack::Update(const ITimingTrackUpdateContext& Context)
 			TSharedPtr<STimingView> TimingView = TimingWindow->GetTimingView();
 			if (TimingView.IsValid())
 			{
+				auto OnTrackAddedRemovedLamda = [this](const TSharedPtr<const FBaseTimingTrack> Track)
+				{
+					if (Track->Is<FThreadTimingTrack>())
+					{
+						// If there are more series than the default frame series.
+						if (this->AllSeries.Num() > ETraceFrameType::TraceFrameType_Count)
+						{
+							this->SetDirtyFlag();
+						}
+					}
+				};
+
+				OnTrackAddedHandle = TimingView->OnTrackAdded().AddLambda(OnTrackAddedRemovedLamda);
+				OnTrackRemovedHandle = TimingView->OnTrackRemoved().AddLambda(OnTrackAddedRemovedLamda);
+
 				OnTrackVisibilityChangedHandle = TimingView->OnTrackVisibilityChanged().AddLambda([this]()
 					{
 						this->SetDirtyFlag();
diff --git a/Engine/Source/Developer/TraceInsights/Private/Insights/ViewModels/TimingGraphTrack.h b/Engine/Source/Developer/TraceInsights/Private/Insights/ViewModels/TimingGraphTrack.h
--- a/Engine/Source/Developer/TraceInsights/Private/Insights/ViewModels/TimingGraphTrack.h
+++ b/Engine/Source/Developer/TraceInsights/Private/Insights/ViewModels/TimingGraphTrack.h
@@ -100,6 +100,8 @@ protected:
 
 private:
 	FDelegateHandle OnTrackVisibilityChangedHandle;
+	FDelegateHandle OnTrackAddedHandle;
+	FDelegateHandle OnTrackRemovedHandle;
 };
 
 ////////////////////////////////////////////////////////////////////////////////////////////////////
diff --git a/Engine/Source/Developer/TraceInsights/Private/Insights/Widgets/SFrameTrack.cpp b/Engine/Source/Developer/TraceInsights/Private/Insights/Widgets/SFrameTrack.cpp
--- a/Engine/Source/Developer/TraceInsights/Private/Insights/Widgets/SFrameTrack.cpp
+++ b/Engine/Source/Developer/TraceInsights/Private/Insights/Widgets/SFrameTrack.cpp
@@ -57,6 +57,8 @@ SFrameTrack::~SFrameTrack()
 				if (TimingView.Get() == RegisteredTimingView)
 				{
 					TimingView->OnTrackVisibilityChanged().Remove(OnTrackVisibilityChangedHandle);
+					TimingView->OnTrackAdded().Remove(OnTrackAddedHandle);
+					TimingView->OnTrackRemoved().Remove(OnTrackRemovedHandle);
 				}
 			}
 		}
@@ -161,13 +163,29 @@ void SFrameTrack::Tick(const FGeometry& AllottedGeometry, const double InCurrent
 			if (!OnTrackVisibilityChangedHandle.IsValid() || TimingView.Get() != RegisteredTimingView)
 			{
 				RegisteredTimingView = TimingView.Get();
-				OnTrackVisibilityChangedHandle = TimingView->OnTrackVisibilityChanged().AddLambda([this]()
+				this->bIsStateDirty = true;
+
+				auto OnTrackAddedRemovedLamda = [this](const TSharedPtr<const FBaseTimingTrack> Track)
+				{
+					if (Track->Is<FThreadTimingTrack>())
 					{
 						// If there are more series than the default frame series.
 						if (this->AllSeries.Num() > ETraceFrameType::TraceFrameType_Count)
 						{
 							this->bIsStateDirty = true;
 						}
+					}
+				};
+
+				OnTrackAddedHandle = TimingView->OnTrackAdded().AddLambda(OnTrackAddedRemovedLamda);
+				OnTrackRemovedHandle = TimingView->OnTrackRemoved().AddLambda(OnTrackAddedRemovedLamda);
+
+				OnTrackVisibilityChangedHandle = TimingView->OnTrackVisibilityChanged().AddLambda([this]()
+					{
+						if (this->AllSeries.Num() > ETraceFrameType::TraceFrameType_Count)
+						{
+							this->bIsStateDirty = true;
+						}
 					});
 			}
 		}
diff --git a/Engine/Source/Developer/TraceInsights/Private/Insights/Widgets/SFrameTrack.h b/Engine/Source/Developer/TraceInsights/Private/Insights/Widgets/SFrameTrack.h
--- a/Engine/Source/Developer/TraceInsights/Private/Insights/Widgets/SFrameTrack.h
+++ b/Engine/Source/Developer/TraceInsights/Private/Insights/Widgets/SFrameTrack.h
@@ -243,6 +243,8 @@ protected:
 
 	STimingView* RegisteredTimingView = nullptr; // For pointer comparison only, do not dereferentiate.
 	FDelegateHandle OnTrackVisibilityChangedHandle;
+	FDelegateHandle OnTrackAddedHandle;
+	FDelegateHandle OnTrackRemovedHandle;
 
 	// Debug stats
 	int32 NumUpdatedFrames;
diff --git a/Engine/Source/Developer/TraceInsights/Private/Insights/Widgets/STimingView.cpp b/Engine/Source/Developer/TraceInsights/Private/Insights/Widgets/STimingView.cpp
--- a/Engine/Source/Developer/TraceInsights/Private/Insights/Widgets/STimingView.cpp
+++ b/Engine/Source/Developer/TraceInsights/Private/Insights/Widgets/STimingView.cpp
@@ -1775,6 +1775,8 @@ void STimingView::AddTrack(TSharedPtr<FBaseTimingTrack> Track, ETimingTrackLocat
 	{
 		InvalidateScrollableTracksOrder();
 	}
+
+	OnTrackAddedDelegate.Broadcast(Track);
 }
 
 ////////////////////////////////////////////////////////////////////////////////////////////////////
@@ -1803,6 +1805,8 @@ bool STimingView::RemoveTrack(TSharedPtr<FBaseTimingTrack> Track)
 			InvalidateScrollableTracksOrder();
 		}
 
+		OnTrackRemovedDelegate.Broadcast(Track);
+
 #if 0
 		UE_LOG(TimingProfiler, Log, TEXT("Removed %s Track (%d) : %s (\"%s\")"),
 			LocationName,
diff --git a/Engine/Source/Developer/TraceInsights/Private/Insights/Widgets/STimingView.h b/Engine/Source/Developer/TraceInsights/Private/Insights/Widgets/STimingView.h
--- a/Engine/Source/Developer/TraceInsights/Private/Insights/Widgets/STimingView.h
+++ b/Engine/Source/Developer/TraceInsights/Private/Insights/Widgets/STimingView.h
@@ -273,6 +273,8 @@ public:
 	virtual Insights::FSelectedTrackChangedDelegate& OnSelectedTrackChanged() override { return OnSelectedTrackChangedDelegate; }
 	virtual Insights::FSelectedEventChangedDelegate& OnSelectedEventChanged() override { return OnSelectedEventChangedDelegate; }
 	virtual Insights::FTrackVisibilityChangedDelegate& OnTrackVisibilityChanged() override { return OnTrackVisibilityChangedDelegate; }
+	virtual Insights::FTrackAddedDelegate& OnTrackAdded() override { return OnTrackAddedDelegate; }
+	virtual Insights::FTrackRemovedDelegate& OnTrackRemoved() override { return OnTrackRemovedDelegate; }
 
 	virtual void ResetSelectedEvent() override
 	{
@@ -719,6 +721,8 @@ protected:
 	Insights::FSelectedTrackChangedDelegate OnSelectedTrackChangedDelegate;
 	Insights::FSelectedEventChangedDelegate OnSelectedEventChangedDelegate;
 	Insights::FTrackVisibilityChangedDelegate OnTrackVisibilityChangedDelegate;
+	Insights::FTrackAddedDelegate OnTrackAddedDelegate;
+	Insights::FTrackRemovedDelegate OnTrackRemovedDelegate;
 
 	TSharedPtr<FUICommandList> CommandList;
 
diff --git a/Engine/Source/Developer/TraceInsights/Public/Insights/ITimingViewSession.h b/Engine/Source/Developer/TraceInsights/Public/Insights/ITimingViewSession.h
--- a/Engine/Source/Developer/TraceInsights/Public/Insights/ITimingViewSession.h
+++ b/Engine/Source/Developer/TraceInsights/Public/Insights/ITimingViewSession.h
@@ -61,6 +61,12 @@ DECLARE_MULTICAST_DELEGATE_OneParam(FSelectedEventChangedDelegate, const TShared
 /** The delegate to be invoked when a track visibility has changed */
 DECLARE_MULTICAST_DELEGATE(FTrackVisibilityChangedDelegate);
 
+/** The delegate to be invoked when a track is added. */
+DECLARE_MULTICAST_DELEGATE_OneParam(FTrackAddedDelegate, const TSharedPtr<const FBaseTimingTrack> /*Track*/);
+
+/** The delegate to be invoked when a track is removed. */
+DECLARE_MULTICAST_DELEGATE_OneParam(FTrackRemovedDelegate, const TSharedPtr<const FBaseTimingTrack> /*Track*/);
+
 ////////////////////////////////////////////////////////////////////////////////////////////////////
 
 /** Hosts a number of timing view visualizers, represents a session of the timing view. */
@@ -133,6 +139,12 @@ public:
 	/** Gets the delegate to be invoked when the track visibility has changed. */
 	virtual FTrackVisibilityChangedDelegate& OnTrackVisibilityChanged() = 0;
 
+	/** Gets the delegate to be invoked when a new track is added. */
+	virtual Insights::FTrackAddedDelegate& OnTrackAdded() = 0;
+
+	/** Gets the delegate to be invoked when a track is removed. */
+	virtual Insights::FTrackRemovedDelegate& OnTrackRemoved() = 0;
+
 	//////////////////////////////////////////////////
 
 	/** Resets the selected event back to empty. */
