Network insights
- Modified filtering to also include sub-events when filtering packet content events.
- Added display of delta time of the selected range of packets.
- Filtering now supports using filter by NetId and EventType in combination, which is useful to highlight specific events for a specific instance.

diff --git a/Engine/Source/Developer/TraceInsights/Private/Insights/NetworkingProfiler/ViewModels/PacketViewDrawHelper.cpp b/Engine/Source/Developer/TraceInsights/Private/Insights/NetworkingProfiler/ViewModels/PacketViewDrawHelper.cpp
--- a/Engine/Source/Developer/TraceInsights/Private/Insights/NetworkingProfiler/ViewModels/PacketViewDrawHelper.cpp
+++ b/Engine/Source/Developer/TraceInsights/Private/Insights/NetworkingProfiler/ViewModels/PacketViewDrawHelper.cpp
@@ -302,7 +302,7 @@ void FPacketViewDrawHelper::DrawSampleHighlight(const FNetworkPacketAggregatedSa
 
 ////////////////////////////////////////////////////////////////////////////////////////////////////
 
-void FPacketViewDrawHelper::DrawSelection(int32 StartPacketIndex, int32 EndPacketIndex) const
+void FPacketViewDrawHelper::DrawSelection(int32 StartPacketIndex, int32 EndPacketIndex, double SelectionTimeSpan) const
 {
 	const float SampleW = Viewport.GetSampleWidth();
 	const int32 PacketsPerSample = Viewport.GetNumPacketsPerSample();
@@ -338,7 +338,7 @@ void FPacketViewDrawHelper::DrawSelection(int32 StartPacketIndex, int32 EndPacke
 
 		if (X1 <= MaxX && X2 >= MinX)
 		{
-			const FString Text = FString::Printf(TEXT("%d packets"), PacketCount);
+			const FString Text = FString::Printf(TEXT("%d packets (%.3f s)"), PacketCount, (float)SelectionTimeSpan);
 			FDrawHelpers::DrawSelection(DrawContext, MinX, MaxX, X1, X2, Y, H, 30.0f, Text, WhiteBrush, SelectionFont);
 		}
 	}
diff --git a/Engine/Source/Developer/TraceInsights/Private/Insights/NetworkingProfiler/ViewModels/PacketViewDrawHelper.h b/Engine/Source/Developer/TraceInsights/Private/Insights/NetworkingProfiler/ViewModels/PacketViewDrawHelper.h
--- a/Engine/Source/Developer/TraceInsights/Private/Insights/NetworkingProfiler/ViewModels/PacketViewDrawHelper.h
+++ b/Engine/Source/Developer/TraceInsights/Private/Insights/NetworkingProfiler/ViewModels/PacketViewDrawHelper.h
@@ -199,7 +199,7 @@ public:
 	void DrawBackground() const;
 	void DrawCached(const FNetworkPacketSeries& Series) const;
 	void DrawSampleHighlight(const FNetworkPacketAggregatedSample& Sample, EHighlightMode Mode) const;
-	void DrawSelection(int32 StartPacketIndex, int32 EndPacketIndex) const;
+	void DrawSelection(int32 StartPacketIndex, int32 EndPacketIndex, double SelectionTimeSpan) const;
 
 	static FLinearColor GetColorByStatus(TraceServices::ENetProfilerDeliveryStatus Status);
 
diff --git a/Engine/Source/Developer/TraceInsights/Private/Insights/NetworkingProfiler/Widgets/SPacketContentView.cpp b/Engine/Source/Developer/TraceInsights/Private/Insights/NetworkingProfiler/Widgets/SPacketContentView.cpp
--- a/Engine/Source/Developer/TraceInsights/Private/Insights/NetworkingProfiler/Widgets/SPacketContentView.cpp
+++ b/Engine/Source/Developer/TraceInsights/Private/Insights/NetworkingProfiler/Widgets/SPacketContentView.cpp
@@ -1047,7 +1047,10 @@ void SPacketContentView::UpdateState(float FontScale)
 				// Count all events in packet, including split data
 				const uint32 StartPos = 0U;
 				const uint32 EndPos = ~0U;
-				NetProfilerProvider->EnumeratePacketContentEventsByPosition(ConnectionIndex, ConnectionMode, PacketIndex, StartPos, EndPos, [this, &Builder, &FilteredDrawStateBuilder, NetProfilerProvider](const TraceServices::FNetProfilerContentEvent& Event)
+				uint32 EndNetIdMatchPos = ~0U;
+				uint32 EndEventTypeMatchPos = ~0U;
+
+				NetProfilerProvider->EnumeratePacketContentEventsByPosition(ConnectionIndex, ConnectionMode, PacketIndex, StartPos, EndPos, [this, &Builder, &FilteredDrawStateBuilder, NetProfilerProvider, &EndNetIdMatchPos, &EndEventTypeMatchPos](const TraceServices::FNetProfilerContentEvent& Event)
 				{
 					const TCHAR* Name = nullptr;
 
@@ -1069,8 +1072,33 @@ void SPacketContentView::UpdateState(float FontScale)
 
 					Builder.AddEvent(Event, Name, NetId);
 
-					if ((!bFilterByEventType || FilterEventTypeIndex == Event.EventTypeIndex) &&
-						(!bFilterByNetId || (Event.ObjectInstanceIndex != 0 && FilterNetId == NetId)))
+					// Include events and sub-events matching event type
+					if (bFilterByEventType)
+					{
+						if (Event.EndPos > EndEventTypeMatchPos)
+						{
+							EndEventTypeMatchPos = ~0U;
+						}
+						if (EndEventTypeMatchPos == ~0U && FilterEventTypeIndex == Event.EventTypeIndex)
+						{
+							EndEventTypeMatchPos = Event.EndPos;
+						}
+					}
+
+					// Include events and sub-events matching net id
+					if (bFilterByNetId)
+					{
+						if (Event.EndPos > EndNetIdMatchPos)
+						{
+							EndNetIdMatchPos = ~0U;
+						}
+						if (EndNetIdMatchPos == ~0U && (Event.ObjectInstanceIndex != 0 && FilterNetId == NetId))
+						{
+							EndNetIdMatchPos = Event.EndPos;
+						}
+					}
+
+					if ((!bFilterByNetId || EndNetIdMatchPos != ~0U) && (!bFilterByEventType || EndEventTypeMatchPos != ~0U))
 					{
 						FilteredDrawStateBuilder.AddEvent(Event, Name, NetId);
 					}
diff --git a/Engine/Source/Developer/TraceInsights/Private/Insights/NetworkingProfiler/Widgets/SPacketView.cpp b/Engine/Source/Developer/TraceInsights/Private/Insights/NetworkingProfiler/Widgets/SPacketView.cpp
--- a/Engine/Source/Developer/TraceInsights/Private/Insights/NetworkingProfiler/Widgets/SPacketView.cpp
+++ b/Engine/Source/Developer/TraceInsights/Private/Insights/NetworkingProfiler/Widgets/SPacketView.cpp
@@ -90,6 +90,7 @@ void SPacketView::Reset()
 	SelectionStartPacketIndex = 0;
 	SelectionEndPacketIndex = 0;
 	LastSelectedPacketIndex = 0;
+	SelectedTimeSpan = 0.0;
 
 	SelectedSample.Reset();
 	HoveredSample.Reset();
@@ -352,22 +353,35 @@ void SPacketView::UpdateState()
 							if ((!SamplePtr->bAtLeastOnePacketMatchesFilter || Filter.AggregationMode != TraceServices::ENetProfilerAggregationMode::None) && (Filter.bByNetId || Filter.bByEventType))
 							{
 								bool bFilterMatch = false;
+								bool bOldEventMatchesFilter = false;
 								uint32 FilterMatchAggregatedEventSizeInBits = 0U;
 								uint32 FilterMatchMaxEventSizeBits = 0U;
 
 								// Filter all events in packet, including split data
 								const uint32 StartPos = 0;
 								const uint32 EndPos = ~0U;
-								NetProfilerProvider->EnumeratePacketContentEventsByPosition(ConnectionIndex, ConnectionMode, PacketIndex - 1, StartPos, EndPos, [this, &bFilterMatch, &Filter, NetProfilerProvider, &FilterMatchAggregatedEventSizeInBits,&FilterMatchMaxEventSizeBits,  &FilterMatchEventTypeIndex](const TraceServices::FNetProfilerContentEvent& Event)
+								uint32 EndNetIdMatchPos = ~0U;
+								uint32 EndEventTypeMatchPos = ~0U;
+
+								NetProfilerProvider->EnumeratePacketContentEventsByPosition(ConnectionIndex, ConnectionMode, PacketIndex - 1, StartPos, EndPos, [this, &bFilterMatch, &bOldEventMatchesFilter, &Filter, NetProfilerProvider, &FilterMatchAggregatedEventSizeInBits,&FilterMatchMaxEventSizeBits,  &FilterMatchEventTypeIndex, &EndNetIdMatchPos, &EndEventTypeMatchPos](const TraceServices::FNetProfilerContentEvent& Event)
 								{
-									bool bEventMatchesFilter = true;
 									if (!bFilterMatch || (Filter.AggregationMode != TraceServices::ENetProfilerAggregationMode::None))
 									{
-										if (Filter.bByEventType && Filter.EventTypeIndex != Event.EventTypeIndex)
+										// Include events and sub-events matching event type
+										if (Filter.bByEventType)
 										{
-											bEventMatchesFilter = false;
+											if (Event.EndPos > EndEventTypeMatchPos)
+											{
+												EndEventTypeMatchPos = ~0U;
+											}
+											if (EndEventTypeMatchPos == ~0U && Filter.EventTypeIndex == Event.EventTypeIndex)
+											{
+												EndEventTypeMatchPos = Event.EndPos;
+											}
 										}
-										if (bEventMatchesFilter && Filter.bByNetId)
+
+										// Include events and sub-events matching net id
+										if (Filter.bByNetId)
 										{
 											uint64 NetId = uint64(-1);
 											if (Event.ObjectInstanceIndex != 0)
@@ -377,13 +391,20 @@ void SPacketView::UpdateState()
 													NetId = ObjectInstance.NetObjectId;
 												});
 											}
-											if (Filter.NetId != NetId)
+
+											if (Event.EndPos > EndNetIdMatchPos)
+											{
+												EndNetIdMatchPos = ~0U;
+											}
+											if (EndNetIdMatchPos == ~0U && (Event.ObjectInstanceIndex != 0 && Filter.NetId == NetId))
 											{
-												bEventMatchesFilter = false;
+												EndNetIdMatchPos = Event.EndPos;
 											}
 										}
 
-										if (bEventMatchesFilter)
+										// Check if all conditions are fulfilled but only aggregate stats for top-level event.
+										const bool bEventMatchesFilter = (!Filter.bByNetId || EndNetIdMatchPos != ~0U) && (!Filter.bByEventType || EndEventTypeMatchPos != ~0U);
+										if (bEventMatchesFilter && !bOldEventMatchesFilter)
 										{
 											const uint32 EventSize = static_cast<uint32>(Event.EndPos - Event.StartPos);
 											FilterMatchAggregatedEventSizeInBits += EventSize;
@@ -395,6 +416,7 @@ void SPacketView::UpdateState()
 												bFilterMatch = true;
 											}
 										}
+										bOldEventMatchesFilter = bEventMatchesFilter;
 									}
 
 								});
@@ -501,6 +523,47 @@ FNetworkPacketSampleRef SPacketView::GetSample(const int32 InPacketIndex)
 	return SampleRef;
 }
 
+void SPacketView::UpdateSelectedTimeSpan()
+{
+	SelectedTimeSpan = 0.0;
+
+	if (SelectionEndPacketIndex == SelectionStartPacketIndex + 1)
+	{		
+		return;
+	}
+
+	TSharedPtr<const TraceServices::IAnalysisSession> Session = FInsightsManager::Get()->GetSession();
+	if (Session.IsValid())
+	{
+		TraceServices::FAnalysisSessionReadScope SessionReadScope(*Session.Get());
+		const TraceServices::INetProfilerProvider* NetProfilerProvider = TraceServices::ReadNetProfilerProvider(*Session.Get());
+		if (NetProfilerProvider && IsConnectionValid(*NetProfilerProvider, GameInstanceIndex, ConnectionIndex, ConnectionMode))
+		{
+			double StartTimeStamp = 0.0f;
+			double EndTimeStamp = 0.0f;
+
+			const int32 LastPacketIndex = SelectionEndPacketIndex - 1;
+
+			const uint32 NumPackets = NetProfilerProvider->GetPacketCount(ConnectionIndex, ConnectionMode);
+			if (SelectionStartPacketIndex >= 0 && SelectionStartPacketIndex < static_cast<int32>(NumPackets) && 
+				LastPacketIndex >= 0 && LastPacketIndex < static_cast<int32>(NumPackets))
+			{
+				NetProfilerProvider->EnumeratePackets(ConnectionIndex, ConnectionMode, SelectionStartPacketIndex, SelectionStartPacketIndex, [&StartTimeStamp](const TraceServices::FNetProfilerPacket& Packet)
+				{
+					StartTimeStamp = Packet.TimeStamp;
+				});
+				NetProfilerProvider->EnumeratePackets(ConnectionIndex, ConnectionMode, LastPacketIndex, LastPacketIndex, [&EndTimeStamp](const TraceServices::FNetProfilerPacket& Packet)
+				{
+					EndTimeStamp = Packet.TimeStamp;
+				});
+			
+				SelectedTimeSpan = FMath::Max(0.0, EndTimeStamp - StartTimeStamp);
+			}
+		}
+	}
+
+}
+
 ////////////////////////////////////////////////////////////////////////////////////////////////////
 
 FNetworkPacketSampleRef SPacketView::GetSampleAtMousePosition(double X, double Y)
@@ -638,6 +701,9 @@ void SPacketView::OnSelectionChanged()
 	TSharedPtr<SNetworkingProfilerWindow> ProfilerWindow = ProfilerWindowWeakPtr.Pin();
 	if (ProfilerWindow.IsValid())
 	{
+		// Update selected time range
+		UpdateSelectedTimeSpan();
+
 		if (SelectedSample.IsValid())
 		{
 			const uint32 BitSize = SelectedSample.Sample->LargestPacket.TotalSizeInBytes * 8;
@@ -723,7 +789,8 @@ int32 SPacketView::OnPaint(const FPaintArgs& Args, const FGeometry& AllottedGeom
 		}
 		if (SelectionEndPacketIndex > SelectionStartPacketIndex + 1)
 		{
-			Helper.DrawSelection(SelectionStartPacketIndex, SelectionEndPacketIndex);
+			
+			Helper.DrawSelection(SelectionStartPacketIndex, SelectionEndPacketIndex, SelectedTimeSpan);
 		}
 
 		// Draw the vertical axis grid.
diff --git a/Engine/Source/Developer/TraceInsights/Private/Insights/NetworkingProfiler/Widgets/SPacketView.h b/Engine/Source/Developer/TraceInsights/Private/Insights/NetworkingProfiler/Widgets/SPacketView.h
--- a/Engine/Source/Developer/TraceInsights/Private/Insights/NetworkingProfiler/Widgets/SPacketView.h
+++ b/Engine/Source/Developer/TraceInsights/Private/Insights/NetworkingProfiler/Widgets/SPacketView.h
@@ -178,6 +178,8 @@ private:
 	void ZoomHorizontally(const float Delta, const float X);
 
 private:
+	void UpdateSelectedTimeSpan();
+
 	TWeakPtr<SNetworkingProfilerWindow> ProfilerWindowWeakPtr;
 
 	uint32 GameInstanceIndex;
@@ -229,6 +231,7 @@ private:
 	int32 SelectionStartPacketIndex;
 	int32 SelectionEndPacketIndex;
 	int32 LastSelectedPacketIndex;
+	double SelectedTimeSpan;
 
 	FNetworkPacketSampleRef SelectedSample;
 	FNetworkPacketSampleRef HoveredSample;
