[Insights] Added Automation tab to Session Browser + Automated C28843555 "Automate Test Case: Unreal Insights - Automatic renaming and deleting of symbol cache files"

diff --git a/Engine/Source/Developer/TraceInsights/Private/Insights/Tests/UITests.cpp b/Engine/Source/Developer/TraceInsights/Private/Insights/Tests/UITests.cpp
--- a/Engine/Source/Developer/TraceInsights/Private/Insights/Tests/UITests.cpp
+++ b/Engine/Source/Developer/TraceInsights/Private/Insights/Tests/UITests.cpp
@@ -2,6 +2,8 @@
 
 #include "UITests.h"
 
+#include "HAL/PlatformFileManager.h"
+
 #include "Insights/Common/Stopwatch.h"
 
 #include "TraceServices/Model/TimingProfiler.h"
@@ -10,6 +12,7 @@
 #include "Insights/TimingProfilerManager.h"
 #include "Insights/ViewModels/TimeFilterValueConverter.h"
 #include "Insights/Widgets/STimingProfilerWindow.h"
+#include "Insights/Widgets/SStartPageWindow.h"
 
 ////////////////////////////////////////////////////////////////////////////////////////////////////
 
@@ -42,6 +45,104 @@ bool FHideAndShowAllTimingViewTabs::RunTest(const FString& Parameters)
 	return true;
 }
 
+////////////////////////////////////////////////////////////////////////////////////////////////////
+
+// TestRail: C28843555
+void FAutomaticRenamingAndDeletingOfSymbolCacheFilesInsightsTest::Define()
+{
+	BeforeEach([this]() {
+		if (IAutomationDriverModule::Get().IsEnabled())
+		{
+			IAutomationDriverModule::Get().Disable();
+		}
+
+		IAutomationDriverModule::Get().Enable();
+
+		Driver = IAutomationDriverModule::Get().CreateDriver();
+
+		});
+	Describe("Make sure the user can rename and delete trace with automatic cache deletion", [this]()
+		{
+
+			It("Copy, rename and delete", EAsyncExecution::ThreadPool, [this]()
+				{
+					FDriverElementRef MinimizeWindowButton = Driver->FindElements(By::Id("launcher-minimizeWindowButton"))->GetElements()[0];
+					MinimizeWindowButton->Click();
+
+					IPlatformFile& PlatformFile = FPlatformFileManager::Get().GetPlatformFile();
+					TSharedPtr<FInsightsManager> InsightsManager = FInsightsManager::Get();
+					TestTrue("Insigts manager should not be null", InsightsManager.IsValid());
+					InsightsManager->GetTraceStoreWindow()->SetDeleteTraceConfirmationWindowVisibility(false);
+
+					FString StoreDir = InsightsManager->GetStoreDir();
+					FString ProjectDir = FPaths::ProjectDir();
+
+					const FString SourceTestTracePath = ProjectDir + TEXT("SourceAssets/UTraces/Test.utrace");
+					const FString SourceTestCachePath = ProjectDir + TEXT("SourceAssets/UCaches/Test.ucache");
+
+					FString StoreTestTracePath = StoreDir / TEXT("Test.utrace");
+					FString StoreTestCachePath = StoreDir / TEXT("Test.ucache");
+
+					TestTrue("Trace in project exists", PlatformFile.FileExists(*SourceTestTracePath));
+					TestTrue("Cache in project exists", PlatformFile.FileExists(*SourceTestCachePath));
+
+					TestFalse("Trace in store should not exist before copy", PlatformFile.FileExists(*StoreTestTracePath));
+					TestFalse("Cache in store should not exist before copy", PlatformFile.FileExists(*StoreTestCachePath));
+
+					// Copy trace
+					// Here we just check that button can be clicked. Unable to copy and paste via Automation Driver 
+					FDriverElementRef ExploreTraceStoreDirButton = Driver->FindElement(By::Id("ExploreTraceStoreDirButton"));
+					TestTrue("Explore Trace Store Dir Button clicked", ExploreTraceStoreDirButton->IsInteractable());
+
+					PlatformFile.CopyFile(*StoreTestTracePath, *SourceTestTracePath);
+					PlatformFile.CopyFile(*StoreTestCachePath, *SourceTestCachePath);
+
+					TestTrue("Trace copied", PlatformFile.FileExists(*StoreTestTracePath));
+					TestTrue("Cache copied", PlatformFile.FileExists(*StoreTestCachePath));
+
+					StoreTestTracePath = StoreDir / TEXT("TestUcacheRenaming.utrace");
+					StoreTestCachePath = StoreDir / TEXT("TestUcacheRenaming.ucache");
+
+					TestFalse("Renamed trace should not exist before renaming", PlatformFile.FileExists(*StoreTestTracePath));
+					TestFalse("Renamed cache should not exist before renaming", PlatformFile.FileExists(*StoreTestCachePath));
+
+					// Rename 
+					auto TraceWaiter = [Driver = Driver](void) -> bool {
+						return Driver->FindElements(By::Id("TraceList"))->GetElements()[0]->GetText().ToString() == TEXT("Test");
+					};
+					Driver->Wait(Until::Condition(TraceWaiter, FWaitTimeout::InSeconds(3)));
+					FDriverElementRef TraceElement = Driver->FindElements(By::Id("TraceList"))->GetElements()[0];
+
+					FDriverSequenceRef Sequence = Driver->CreateSequence();
+					Sequence->Actions()
+						.Click(TraceElement)
+						.Type(EKeys::F2)
+						.Type(TEXT("UcacheRenaming"))
+						.Type(EKeys::Enter);
+
+					TestTrue("Trace renamed", Sequence->Perform());
+
+					TestTrue("Renamed trace should exist", PlatformFile.FileExists(*StoreTestTracePath));
+					TestTrue("Renamed cache should exist", PlatformFile.FileExists(*StoreTestCachePath));
+
+					// Delete
+					FDriverElementRef OpenTraceButton = Driver->FindElement(By::Id("OpenTraceButton"));
+					Driver->Wait(Until::ElementIsInteractable(OpenTraceButton, FWaitTimeout::InSeconds(5)));
+
+					TraceElement = Driver->FindElements(By::Id("TraceList"))->GetElements()[0];
+					TraceElement->Type(EKeys::Delete);
+
+					TestFalse("Renamed trace should be deleted", PlatformFile.FileExists(*StoreTestTracePath));
+					TestFalse("Renamed cache should be deleted", PlatformFile.FileExists(*StoreTestCachePath));
+				});
+		});
+
+	AfterEach([this]() {
+		Driver.Reset();
+		IAutomationDriverModule::Get().Disable();
+	});
+}
+
 #endif // !WITH_EDITOR
 
 ////////////////////////////////////////////////////////////////////////////////////////////////////
diff --git a/Engine/Source/Developer/TraceInsights/Private/Insights/Tests/UITests.h b/Engine/Source/Developer/TraceInsights/Private/Insights/Tests/UITests.h
--- a/Engine/Source/Developer/TraceInsights/Private/Insights/Tests/UITests.h
+++ b/Engine/Source/Developer/TraceInsights/Private/Insights/Tests/UITests.h
@@ -4,6 +4,15 @@
 
 #include "CoreMinimal.h"
 #include "Logging/LogMacros.h"
+
+//Driver
+#include "IAutomationDriver.h"
+#include "IAutomationDriverModule.h"
+#include "DriverConfiguration.h"
+#include "IDriverElement.h"
+#include "IDriverSequence.h"
+#include "LocateBy.h"
+
 #include "Misc/AutomationTest.h"
 
 DECLARE_LOG_CATEGORY_EXTERN(UITests, Log, All);
@@ -11,6 +20,9 @@ DECLARE_LOG_CATEGORY_EXTERN(UITests, Log, All);
 #if !WITH_EDITOR
 
 IMPLEMENT_SIMPLE_AUTOMATION_TEST(FHideAndShowAllTimingViewTabs, "Insights.HideAndShowAllTimingViewTabs", EAutomationTestFlags::ApplicationContextMask | EAutomationTestFlags::ProductFilter)
+BEGIN_DEFINE_SPEC(FAutomaticRenamingAndDeletingOfSymbolCacheFilesInsightsTest, "Insights.SessionBrowser.AutomaticRenamingAndDeletingOfSymbolCacheFiles", EAutomationTestFlags::ApplicationContextMask | EAutomationTestFlags::ProductFilter)
+FAutomationDriverPtr Driver;
+END_DEFINE_SPEC(FAutomaticRenamingAndDeletingOfSymbolCacheFilesInsightsTest)
 
 #endif
 
diff --git a/Engine/Source/Developer/TraceInsights/Private/Insights/TraceInsightsModule.cpp b/Engine/Source/Developer/TraceInsights/Private/Insights/TraceInsightsModule.cpp
--- a/Engine/Source/Developer/TraceInsights/Private/Insights/TraceInsightsModule.cpp
+++ b/Engine/Source/Developer/TraceInsights/Private/Insights/TraceInsightsModule.cpp
@@ -606,6 +606,24 @@ void FTraceInsightsModule::ScheduleCommand(const FString& InCmd)
 
 ////////////////////////////////////////////////////////////////////////////////////////////////////
 
+void FTraceInsightsModule::RunAutomationTest(const FString& InCmd)
+{
+#if !UE_BUILD_SHIPPING && !WITH_EDITOR
+	FString ActualCmd = InCmd;
+	ActualCmd.TrimCharInline(TEXT('\"'), nullptr);
+	ActualCmd.TrimCharInline(TEXT('\''), nullptr);
+
+	if (ActualCmd.StartsWith(TEXT("Automation RunTests")))
+	{
+		FInsightsTestRunner::Get()->ScheduleCommand(ActualCmd);
+		FInsightsTestRunner::Get()->RunTests();
+		return;
+	}
+#endif
+}
+
+////////////////////////////////////////////////////////////////////////////////////////////////////
+
 bool FTraceInsightsModule::Exec(const TCHAR* Cmd, FOutputDevice& Ar)
 {
 	for (TSharedRef<IInsightsComponent>& Component : Components)
diff --git a/Engine/Source/Developer/TraceInsights/Private/Insights/TraceInsightsModule.h b/Engine/Source/Developer/TraceInsights/Private/Insights/TraceInsightsModule.h
--- a/Engine/Source/Developer/TraceInsights/Private/Insights/TraceInsightsModule.h
+++ b/Engine/Source/Developer/TraceInsights/Private/Insights/TraceInsightsModule.h
@@ -76,6 +76,7 @@ public:
 
 	virtual void InitializeTesting(bool InInitAutomationModules, bool InAutoQuit) override;
 	virtual void ScheduleCommand(const FString& InCmd) override;
+	virtual void RunAutomationTest(const FString& InCmd) override;
 	virtual bool Exec(const TCHAR* Cmd, FOutputDevice& Ar) override;
 
 protected:
diff --git a/Engine/Source/Developer/TraceInsights/Private/Insights/Widgets/SStartPageWindow.cpp b/Engine/Source/Developer/TraceInsights/Private/Insights/Widgets/SStartPageWindow.cpp
--- a/Engine/Source/Developer/TraceInsights/Private/Insights/Widgets/SStartPageWindow.cpp
+++ b/Engine/Source/Developer/TraceInsights/Private/Insights/Widgets/SStartPageWindow.cpp
@@ -53,6 +53,9 @@
 #include "Insights/Widgets/SInsightsSettings.h"
 #include "Insights/Widgets/SLazyToolTip.h"
 
+// Driver
+#include "Framework/MetaData/DriverMetaData.h"
+
 ////////////////////////////////////////////////////////////////////////////////////////////////////
 
 #define LOCTEXT_NAMESPACE "STraceStoreWindow"
@@ -123,6 +126,7 @@ public:
 					.HighlightColor(FLinearColor(0.75f, 0.75f, 0.75f, 1.0f))
 					.HighlightShape(FInsightsStyle::Get().GetBrush("DarkGreenBrush"))
 					.ToolTip(STraceListRow::GetTraceTooltip())
+			     	.AddMetaData(FDriverMetaData::Id("TraceList"))
 				]
 
 				+ SOverlay::Slot()
@@ -217,6 +221,7 @@ public:
 					SNew(STextBlock)
 					.Text(this, &STraceListRow::GetTraceStatus)
 					.ToolTip(STraceListRow::GetTraceTooltip())
+					.AddMetaData(FDriverMetaData::Id("TraceStatusColumnList"))
 					.ColorAndOpacity(FStyleColors::AccentRed)
 				];
 		}
@@ -1604,6 +1609,7 @@ TSharedRef<SWidget> STraceStoreWindow::ConstructLoadPanel()
 		.IsEnabled(this, &STraceStoreWindow::Open_IsEnabled)
 		.OnClicked(this, &STraceStoreWindow::Open_OnClicked)
 		.ToolTipText(LOCTEXT("OpenButtonTooltip", "Start analysis for selected trace session."))
+		.AddMetaData(FDriverMetaData::Id("OpenTraceButton"))
 		.ContentPadding(FMargin(0.0f, 0.0f, 0.0f, 0.0f))
 		.Content()
 		[
@@ -1721,6 +1727,7 @@ TSharedRef<SWidget> STraceStoreWindow::ConstructTraceStoreDirectoryPanel()
 				.ButtonStyle(&FAppStyle::Get().GetWidgetStyle<FButtonStyle>("SimpleButton"))
 				.ToolTipText(LOCTEXT("ExploreTraceStoreDirButtonToolTip", "Explore the Trace Store Directory"))
 				.OnClicked(this, &STraceStoreWindow::ExploreTraceStoreDirectory_OnClicked)
+				.AddMetaData(FDriverMetaData::Id("ExploreTraceStoreDirButton"))
 				.IsEnabled_Static(&STraceStoreWindow::CanChangeStoreSettings)
 				[
 					SNew(SImage)
@@ -2011,14 +2018,17 @@ void STraceStoreWindow::DeleteTraceFile()
 
 	FString TraceFile = SelectedTrace->Uri.ToString();
 
-	//TODO: Make a custom OkCancel modal dialog. See FSlateApplication::Get().AddModalWindow(..).
-	FText Title = LOCTEXT("ConfirmToDeleteTraceFile_Title", "Unreal Insights");
-	FText ConfirmMessage = FText::Format(LOCTEXT("ConfirmToDeleteTraceFile", "You are about to delete the utrace file:\n{0}\n\nPress Ok to continue."),
-		FText::FromString(TraceFile));
-	EAppReturnType::Type OkToDelete = FMessageDialog::Open(EAppMsgType::OkCancel, ConfirmMessage, Title);
-	if (OkToDelete == EAppReturnType::Cancel)
+	if (bIsDeleteTraceConfirmWindowVisible)
 	{
-		return;
+		//TODO: Make a custom OkCancel modal dialog. See FSlateApplication::Get().AddModalWindow(..).
+		FText Title = LOCTEXT("ConfirmToDeleteTraceFile_Title", "Unreal Insights");
+		FText ConfirmMessage = FText::Format(LOCTEXT("ConfirmToDeleteTraceFile", "You are about to delete the utrace file:\n{0}\n\nPress Ok to continue."),
+			FText::FromString(TraceFile));
+		EAppReturnType::Type OkToDelete = FMessageDialog::Open(EAppMsgType::OkCancel, ConfirmMessage, Title);
+		if (OkToDelete == EAppReturnType::Cancel)
+		{
+			return;
+		}
 	}
 
 	FString TraceName = SelectedTrace->Name.ToString();
@@ -2295,6 +2305,13 @@ void STraceStoreWindow::RefreshTraceList()
 
 ////////////////////////////////////////////////////////////////////////////////////////////////////
 
+void STraceStoreWindow::SetDeleteTraceConfirmationWindowVisibility(bool bIsVisibile)
+{
+	bIsDeleteTraceConfirmWindowVisible = bIsVisibile;
+}
+
+////////////////////////////////////////////////////////////////////////////////////////////////////
+
 void STraceStoreWindow::UpdateTrace(FTraceViewModel& InOutTrace, const Insights::FStoreBrowserTraceInfo& InSourceTrace)
 {
 	InOutTrace.TraceId = InSourceTrace.TraceId;
@@ -2792,6 +2809,21 @@ TSharedRef<SWidget> STraceStoreWindow::MakeTraceListMenu()
 
 	MenuBuilder.BeginSection("DebugOptions", LOCTEXT("TraceListMenu_Section_DebugOptions", "Debug Options"));
 
+	if (FGlobalTabmanager::Get()->HasTabSpawner(FInsightsManagerTabs::AutomationWindowTabId))
+	{
+		MenuBuilder.AddMenuEntry(
+			LOCTEXT("Automation", "Automation"),
+			TAttribute<FText>(), // no tooltip
+			FSlateIcon(FInsightsStyle::GetStyleSetName(), "Icons.TestAutomation"),
+			FUIAction(FExecuteAction::CreateLambda([this]()
+				{
+					FGlobalTabmanager::Get()->TryInvokeTab(FInsightsManagerTabs::AutomationWindowTabId);
+				})),
+			NAME_None,
+					EUserInterfaceActionType::Button
+					);
+	}
+
 	// Enable Automation Tests Option.
 	{
 		FUIAction ToogleAutomationTestsAction;
@@ -2805,7 +2837,7 @@ TSharedRef<SWidget> STraceStoreWindow::MakeTraceListMenu()
 			});
 
 		MenuBuilder.AddMenuEntry(
-			LOCTEXT("EnableAutomatedTesting", "Enable Automation Testing"),
+			LOCTEXT("EnableAutomatedTesting", "Enable Session Automation Testing"),
 			LOCTEXT("EnableAutomatedTestingDesc", "Activates the automatic test system for new sessions opened from this window."),
 			FSlateIcon(FInsightsStyle::GetStyleSetName(), "Icons.TestAutomation"),
 			ToogleAutomationTestsAction,
@@ -4073,4 +4105,4 @@ END_SLATE_FUNCTION_BUILD_OPTIMIZATION
 
 ////////////////////////////////////////////////////////////////////////////////////////////////////
 
-#undef LOCTEXT_NAMESPACE
+#undef LOCTEXT_NAMESPACE
\ No newline at end of file
diff --git a/Engine/Source/Developer/TraceInsights/Private/Insights/Widgets/SStartPageWindow.h b/Engine/Source/Developer/TraceInsights/Private/Insights/Widgets/SStartPageWindow.h
--- a/Engine/Source/Developer/TraceInsights/Private/Insights/Widgets/SStartPageWindow.h
+++ b/Engine/Source/Developer/TraceInsights/Private/Insights/Widgets/SStartPageWindow.h
@@ -278,6 +278,7 @@ public:
 
 	void OnFilterChanged();
 	const TArray<TSharedPtr<FTraceViewModel>>& GetAllAvailableTraces() const;
+	void SetDeleteTraceConfirmationWindowVisibility(bool bIsVisibile);
 
 private:
 	TSharedRef<SWidget> ConstructFiltersToolbar();
@@ -491,6 +492,10 @@ private:
 	void ShowSuccessMessage(FText& InMessage);
 	void ShowFailMessage(FText& InMessage);
 
+public:
+	/** Parameter that controls the visibility of the confirmation window in case the trace is deleted. */
+	bool bIsDeleteTraceConfirmWindowVisible = true;
+
 private:
 	/** Widget for the non-intrusive notifications. */
 	TSharedPtr<SNotificationList> NotificationList;
diff --git a/Engine/Source/Developer/TraceInsights/Public/Insights/IUnrealInsightsModule.h b/Engine/Source/Developer/TraceInsights/Public/Insights/IUnrealInsightsModule.h
--- a/Engine/Source/Developer/TraceInsights/Public/Insights/IUnrealInsightsModule.h
+++ b/Engine/Source/Developer/TraceInsights/Public/Insights/IUnrealInsightsModule.h
@@ -277,6 +277,11 @@ public:
 	*/
 	virtual void ScheduleCommand(const FString& InCmd) = 0;
 
+	/**
+	* Called to run automation test in Insights.
+	*/
+	virtual void RunAutomationTest(const FString& InCmd) = 0;
+
 	/**
 	* Called to initialize testing in stand alone Insights.
 	 * @param InInitAutomationModules If true Insights will initialize the modules required for running automation tests.
diff --git a/Engine/Source/Developer/TraceInsights/TraceInsights.Build.cs b/Engine/Source/Developer/TraceInsights/TraceInsights.Build.cs
--- a/Engine/Source/Developer/TraceInsights/TraceInsights.Build.cs
+++ b/Engine/Source/Developer/TraceInsights/TraceInsights.Build.cs
@@ -11,6 +11,7 @@ public class TraceInsights : ModuleRules
 			new string[] {
 				"AppFramework", // for SColorPicker
 				"ApplicationCore",
+				"AutomationDriver",
 				"Cbor",
 				"Core",
 				"CoreUObject",
@@ -26,7 +27,7 @@ public class TraceInsights : ModuleRules
 				"TraceServices",
 				"WorkspaceMenuStructure",
 				"XmlParser",
-			}
+		}
 		);
 
 		//Modules required for running automation in stand alone Insights
diff --git a/Engine/Source/Programs/UnrealInsights/Private/UserInterfaceCommand.cpp b/Engine/Source/Programs/UnrealInsights/Private/UserInterfaceCommand.cpp
--- a/Engine/Source/Programs/UnrealInsights/Private/UserInterfaceCommand.cpp
+++ b/Engine/Source/Programs/UnrealInsights/Private/UserInterfaceCommand.cpp
@@ -300,19 +300,18 @@ void FUserInterfaceCommand::InitializeSlateApplication(bool bOpenTraceFile, cons
 		bUseCustomStoreAddress = true;
 	}
 
+	// This parameter will cause the application to close when analysis fails to start or completes successfully.
+	const bool bAutoQuit = FParse::Param(FCommandLine::Get(), TEXT("AutoQuit"));
+
 	const bool bInitializeTesting = FParse::Param(FCommandLine::Get(), TEXT("InsightsTest"));
+	if (bInitializeTesting)
+	{
+		const bool bInitAutomationModules = true;
+		TraceInsightsModule.InitializeTesting(bInitAutomationModules, bAutoQuit);
+	}
 
 	if (bUseTraceId || bOpenTraceFile) // viewer mode
 	{
-		// This parameter will cause the application to close when analysis fails to start or completes successfully.
-		const bool bAutoQuit = FParse::Param(FCommandLine::Get(), TEXT("AutoQuit"));
-
-		if (bInitializeTesting)
-		{
-			const bool bInitAutomationModules = true;
-			TraceInsightsModule.InitializeTesting(bInitAutomationModules, bAutoQuit);
-		}
-
 		TCHAR Cmd[1024];
 		bool bExecuteCommand = false;
 		if (FParse::Value(FCommandLine::Get(), TEXT("-ExecOnAnalysisCompleteCmd="), Cmd, 1024, false))
@@ -342,6 +341,14 @@ void FUserInterfaceCommand::InitializeSlateApplication(bool bOpenTraceFile, cons
 	}
 	else // browser mode
 	{
+
+		TCHAR Cmd[1024];
+		bool bExecuteCommand = false;
+		if (FParse::Value(FCommandLine::Get(), TEXT("-ExecBrowserAutomationTest="), Cmd, 1024, false))
+		{
+			bExecuteCommand = true;
+		}
+
 		if (bUseCustomStoreAddress)
 		{
 			TraceInsightsModule.ConnectToStore(StoreHost, StorePort);
@@ -356,6 +363,11 @@ void FUserInterfaceCommand::InitializeSlateApplication(bool bOpenTraceFile, cons
 		Params.bInitializeTesting = bInitializeTesting;
 		Params.bStartProcessWithStompMalloc = FParse::Param(FCommandLine::Get(), TEXT("stompmalloc"));
 		TraceInsightsModule.CreateSessionBrowser(Params);
+
+		if (bExecuteCommand)
+		{
+			TraceInsightsModule.RunAutomationTest(Cmd);
+		}
 	}
 
 	delete[] StoreHost;
