[Insights] Fix bugs affecting aggregation and selection of timing events caused by infinite events.

diff --git a/Engine/Source/Developer/TraceInsights/Private/Insights/ViewModels/ThreadTimingTrack.cpp b/Engine/Source/Developer/TraceInsights/Private/Insights/ViewModels/ThreadTimingTrack.cpp
--- a/Engine/Source/Developer/TraceInsights/Private/Insights/ViewModels/ThreadTimingTrack.cpp
+++ b/Engine/Source/Developer/TraceInsights/Private/Insights/ViewModels/ThreadTimingTrack.cpp
@@ -1420,7 +1420,7 @@ const TSharedPtr<const ITimingEvent> FThreadTimingTrack::GetEvent(float InPosX,
 		{
 			TraceServices::FAnalysisSessionReadScope SessionReadScope(*Session.Get());
 
-			if (TraceServices::ReadTimingProfilerProvider(*Session.Get()))
+			if (EventTime <= Session->GetDurationSeconds() && TraceServices::ReadTimingProfilerProvider(*Session.Get()))
 			{
 				const TraceServices::ITimingProfilerProvider& TimingProfilerProvider = *TraceServices::ReadTimingProfilerProvider(*Session.Get());
 
diff --git a/Engine/Source/Developer/TraceInsights/Private/Insights/ViewModels/TimerButterflyAggregation.cpp b/Engine/Source/Developer/TraceInsights/Private/Insights/ViewModels/TimerButterflyAggregation.cpp
--- a/Engine/Source/Developer/TraceInsights/Private/Insights/ViewModels/TimerButterflyAggregation.cpp
+++ b/Engine/Source/Developer/TraceInsights/Private/Insights/ViewModels/TimerButterflyAggregation.cpp
@@ -55,12 +55,26 @@ void FTimerButterflyAggregationWorker::DoWork(TSharedPtr<TraceServices::FCancell
 		TraceServices::FAnalysisSessionReadScope SessionReadScope(*Session.Get());
 		const TraceServices::ITimingProfilerProvider& TimingProfilerProvider = *TraceServices::ReadTimingProfilerProvider(*Session.Get());
 
-		auto ThreadFilter = [this](uint32 ThreadId)
+		if (StartTime <= Session->GetDurationSeconds())
 		{
-			return CpuThreads.Contains(ThreadId);
-		};
+			auto ThreadFilter = [this](uint32 ThreadId)
+			{
+				return CpuThreads.Contains(ThreadId);
+			};
+
+			double ClampedEndTime = FMath::Min(Session->GetDurationSeconds(), EndTime);
+			ResultButterfly.Reset(TimingProfilerProvider.CreateButterfly(StartTime, ClampedEndTime, ThreadFilter, bIncludeGpuThread));
+		}
+		else
+		{
+			// Exclude all threads to receive empty results in case infinite events are in the session.
+			auto ThreadFilter = [this](uint32 ThreadId)
+			{
+				return false;
+			};
 
-		ResultButterfly.Reset(TimingProfilerProvider.CreateButterfly(StartTime, EndTime, ThreadFilter, bIncludeGpuThread));
+			ResultButterfly.Reset(TimingProfilerProvider.CreateButterfly(StartTime, EndTime, ThreadFilter, false));
+		}
 	}
 }
 
diff --git a/Engine/Source/Developer/TraceServices/Private/Model/TimingProfiler.cpp b/Engine/Source/Developer/TraceServices/Private/Model/TimingProfiler.cpp
--- a/Engine/Source/Developer/TraceServices/Private/Model/TimingProfiler.cpp
+++ b/Engine/Source/Developer/TraceServices/Private/Model/TimingProfiler.cpp
@@ -291,7 +291,12 @@ ITable<FTimingProfilerAggregatedStats>* FTimingProfilerProvider::CreateAggregati
 	TMap<const FTimingProfilerTimer*, FAggregatedTimingStats> Aggregation;
 	if (Params.FrameType == ETraceFrameType::TraceFrameType_Count)
 	{
-		FTimelineStatistics::CreateAggregation(IncludedTimelines, BucketMappingFunc, Params.IntervalStart, Params.IntervalEnd, Params.CancellationToken, Aggregation);
+		if (Params.IntervalStart <= Session.GetDurationSeconds())
+		{
+			// Do not allow inf for the end time.
+			double EndTime = FMath::Min(Params.IntervalEnd, Session.GetDurationSeconds());
+			FTimelineStatistics::CreateAggregation(IncludedTimelines, BucketMappingFunc, Params.IntervalStart, EndTime, Params.CancellationToken, Aggregation);
+		}
 	}
 	else
 	{
