[Insights] Resize the SSuggestionBox used for filters to prevent excesize resizing.

diff --git a/Engine/Source/Developer/TraceInsights/Private/Insights/Widgets/SFilterConfiguratorRow.cpp b/Engine/Source/Developer/TraceInsights/Private/Insights/Widgets/SFilterConfiguratorRow.cpp
--- a/Engine/Source/Developer/TraceInsights/Private/Insights/Widgets/SFilterConfiguratorRow.cpp
+++ b/Engine/Source/Developer/TraceInsights/Private/Insights/Widgets/SFilterConfiguratorRow.cpp
@@ -99,7 +99,7 @@ TSharedRef<SWidget> SFilterConfiguratorRow::GenerateWidgetForColumn(const FName&
 				.Padding(FMargin(4.0f, 0.0f, 4.0f, 0.0f))
 				[
 					SNew(SSuggestionTextBox)
-					.MinDesiredWidth(50.0f)
+					.MinDesiredWidth(300.0f)
 					.ForegroundColor(FSlateColor::UseForeground())
 					.OnTextCommitted(this, &SFilterConfiguratorRow::SuggestionTextBox_OnValueCommitted)
 					.OnTextChanged(this, &SFilterConfiguratorRow::SuggestionTextBox_OnValueChanged)
