[Insights] Automated C35547680 "Unreal Insights - UI - Memory Insights - LLM - XML Reports".
Changed test naming to keep Automation Spec Framework code style.
Moved "CopyRenameDeleteTrace" and "MemoryInsightsXMLReportsUpload" in one Spec.
Added closing tab via FSlate in BeforeEach and pop-up in AfterEach.
Moved Test.utrace and Test.ucache file to "Content" folder.

diff --git a/Engine/Source/Developer/TraceInsights/Private/Insights/Tests/InsightsTestUtils.cpp b/Engine/Source/Developer/TraceInsights/Private/Insights/Tests/InsightsTestUtils.cpp
--- a/Engine/Source/Developer/TraceInsights/Private/Insights/Tests/InsightsTestUtils.cpp
+++ b/Engine/Source/Developer/TraceInsights/Private/Insights/Tests/InsightsTestUtils.cpp
@@ -12,6 +12,8 @@
 #include "Insights/Common/Stopwatch.h"
 #include "Insights/IUnrealInsightsModule.h"
 
+#include "Misc/FileHelper.h"
+
 ////////////////////////////////////////////////////////////////////////////////////////////////////
 
 FInsightsTestUtils::FInsightsTestUtils(FAutomationTestBase* InTest) :
@@ -72,3 +74,22 @@ bool FInsightsTestUtils::AnalyzeTrace(const TCHAR* Path) const
 }
 
 ////////////////////////////////////////////////////////////////////////////////////////////////////
+
+bool FInsightsTestUtils::FileContainsString(const FString& PathToFile, const FString& ExpectedString, const float Timeout) const
+{
+	float StartTime = FPlatformTime::Seconds();
+	while ((FPlatformTime::Seconds() - StartTime) < Timeout)
+	{
+		FString LogFileContents;
+		if (FFileHelper::LoadFileToString(LogFileContents, *PathToFile))
+		{
+			if (LogFileContents.Contains(ExpectedString))
+			{
+				return true;
+			}
+		}
+		FPlatformProcess::Sleep(0.1f);
+	}
+
+	return false;
+}
diff --git a/Engine/Source/Developer/TraceInsights/Private/Insights/Tests/UITests.cpp b/Engine/Source/Developer/TraceInsights/Private/Insights/Tests/UITests.cpp
--- a/Engine/Source/Developer/TraceInsights/Private/Insights/Tests/UITests.cpp
+++ b/Engine/Source/Developer/TraceInsights/Private/Insights/Tests/UITests.cpp
@@ -13,6 +13,11 @@
 #include "Insights/ViewModels/TimeFilterValueConverter.h"
 #include "Insights/Widgets/STimingProfilerWindow.h"
 #include "Insights/Widgets/SStartPageWindow.h"
+#include "Insights/MemoryProfiler/Widgets/SMemoryProfilerWindow.h"
+#include "Insights/MemoryProfiler/MemoryProfilerManager.h"
+#include "Insights/MemoryProfiler/ViewModels/MemorySharedState.h"
+#include "Insights/Widgets/STimingView.h"
+#include "Insights/Tests/InsightsTestUtils.h"
 
 ////////////////////////////////////////////////////////////////////////////////////////////////////
 
@@ -47,28 +52,30 @@ bool FHideAndShowAllTimingViewTabs::RunTest(const FString& Parameters)
 
 ////////////////////////////////////////////////////////////////////////////////////////////////////
 
-// TestRail: C28843555
-void FAutomaticRenamingAndDeletingOfSymbolCacheFilesInsightsTest::Define()
+void FAutomationDriverUnrealInsightsSessionBrowserTest::Define()
 {
 	BeforeEach([this]() {
+		AutomationWindow = FSlateApplication::Get().GetActiveTopLevelWindow();
+		const FString AutomationWindowName = TEXT("Automation");
+		if (AutomationWindow->GetTitle().ToString().Contains(AutomationWindowName))
+		{
+			AutomationWindow->Minimize();
+		}
+
 		if (IAutomationDriverModule::Get().IsEnabled())
 		{
 			IAutomationDriverModule::Get().Disable();
 		}
-
 		IAutomationDriverModule::Get().Enable();
 
-		Driver = IAutomationDriverModule::Get().CreateDriver();
-
+		Driver = IAutomationDriverModule::Get().CreateDriver();		
 		});
-	Describe("Make sure the user can rename and delete trace with automatic cache deletion", [this]()
-		{
 
-			It("Copy, rename and delete", EAsyncExecution::ThreadPool, [this]()
+	// TestRail: C28843555
+	Describe("CopyRenameDeleteTrace", [this]()
+		{
+			It("should verify that user copy, rename and delete traces", EAsyncExecution::ThreadPool, [this]()
 				{
-					FDriverElementRef MinimizeWindowButton = Driver->FindElements(By::Id("launcher-minimizeWindowButton"))->GetElements()[0];
-					MinimizeWindowButton->Click();
-
 					IPlatformFile& PlatformFile = FPlatformFileManager::Get().GetPlatformFile();
 					TSharedPtr<FInsightsManager> InsightsManager = FInsightsManager::Get();
 					TestTrue("Insigts manager should not be null", InsightsManager.IsValid());
@@ -77,8 +84,8 @@ void FAutomaticRenamingAndDeletingOfSymbolCacheFilesInsightsTest::Define()
 					FString StoreDir = InsightsManager->GetStoreDir();
 					FString ProjectDir = FPaths::ProjectDir();
 
-					const FString SourceTestTracePath = ProjectDir + TEXT("SourceAssets/UTraces/Test.utrace");
-					const FString SourceTestCachePath = ProjectDir + TEXT("SourceAssets/UCaches/Test.ucache");
+					const FString SourceTestTracePath = FPaths::RootDir() / TEXT("EngineTest/SourceAssets/Utrace/Test.utrace");
+					const FString SourceTestCachePath = FPaths::RootDir() / TEXT("EngineTest/SourceAssets/Utrace/Test.ucache");
 
 					FString StoreTestTracePath = StoreDir / TEXT("Test.utrace");
 					FString StoreTestCachePath = StoreDir / TEXT("Test.ucache");
@@ -122,8 +129,8 @@ void FAutomaticRenamingAndDeletingOfSymbolCacheFilesInsightsTest::Define()
 
 					TestTrue("Trace renamed", Sequence->Perform());
 
-					TestTrue("Renamed trace should exist", PlatformFile.FileExists(*StoreTestTracePath));
-					TestTrue("Renamed cache should exist", PlatformFile.FileExists(*StoreTestCachePath));
+					TestTrue("Renamed trace should exists", PlatformFile.FileExists(*StoreTestTracePath));
+					TestTrue("Renamed cache should exists", PlatformFile.FileExists(*StoreTestCachePath));
 
 					// Delete
 					FDriverElementRef OpenTraceButton = Driver->FindElement(By::Id("OpenTraceButton"));
@@ -137,12 +144,116 @@ void FAutomaticRenamingAndDeletingOfSymbolCacheFilesInsightsTest::Define()
 				});
 		});
 
+	// TestRail: C35547680
+	Describe("MemoryInsights.XMLReportsUpload", [this]()
+		{
+			It("should verify that user can upload xml reports in Memory Insights tab", EAsyncExecution::ThreadPool, FTimespan::FromSeconds(120), [this]()
+				{
+					FInsightsTestUtils Utils(this);
+					IPlatformFile& PlatformFile = FPlatformFileManager::Get().GetPlatformFile();
+					TSharedPtr<FInsightsManager> InsightsManager = FInsightsManager::Get();
+					if (!InsightsManager.IsValid())
+					{
+						AddError("Insigts manager should not be null");
+						return;
+					}
+
+					// Start tracing editor instance, not Lyra. There is no difference between them in this test.
+					FString UEPath = FPlatformProcess::GenerateApplicationPath("UnrealEditor", EBuildConfiguration::Development);
+					FString Parameters = TEXT("-trace=Bookmark,Memory -tracehost=127.0.0.1");
+					constexpr bool bLaunchDetached = true;
+					constexpr bool bLaunchHidden = false;
+					constexpr bool bLaunchReallyHidden = false;
+					uint32 ProcessID = 0;
+					const int32 PriorityModifier = 0;
+					const TCHAR* OptionalWorkingDirectory = nullptr;
+					void* PipeWriteChild = nullptr;
+					void* PipeReadChild = nullptr;
+					FProcHandle EditorHandle = FPlatformProcess::CreateProc(*UEPath, *Parameters, bLaunchDetached, bLaunchHidden, bLaunchReallyHidden, &ProcessID, PriorityModifier, OptionalWorkingDirectory, PipeWriteChild, PipeReadChild);
+					if (!EditorHandle.IsValid())
+					{
+						AddError("Lyra should be started");
+						return;
+					}
+
+					// Verify that LIVE trace appeared
+					auto TraceWaiter = [Driver = Driver](void) -> bool {
+						return Driver->FindElements(By::Id("TraceStatusColumnList"))->GetElements()[0]->GetText().ToString() == TEXT("LIVE");
+					};
+					if (!Driver->Wait(Until::Condition(TraceWaiter, FWaitTimeout::InSeconds(10))))
+					{
+						AddError("Live trace should appear");
+						FPlatformProcess::TerminateProc(EditorHandle);
+						return;
+					}
+
+					FDriverElementRef TraceElement = Driver->FindElements(By::Id("TraceList"))->GetElements()[0];
+					const FString TraceName = TraceElement->GetText().ToString();
+
+					const FString StoreDir = InsightsManager->GetStoreDir();
+					const FString ProjectDir = FPaths::ProjectDir();
+					const FString StoreTracePath = StoreDir / FString::Printf(TEXT("%s.utrace"), *TraceName);
+					const FString LogDirPath = ProjectDir / TEXT("TestResults");
+					const FString LogPath = ProjectDir / TEXT("TestResults/Log.txt");
+					const FString SuccessTestResult = TEXT("Test Completed. Result={Success}");
+
+					// Test live trace
+					FString TraceParameters = FString::Printf(TEXT("-InsightsTest -ABSLOG=\"%s\" -AutoQuit -ExecOnAnalysisCompleteCmd=\"Automation RunTests Insights.UploadMemoryInsightsLLMXMLReportsTrace\" -OpenTraceFile=%s"), *LogPath, *StoreTracePath);
+					InsightsManager->OpenUnrealInsights(*TraceParameters);
+					bool bLineFound = Utils.FileContainsString(LogPath, SuccessTestResult, 60.0f);
+					TestTrue("Test for live trace should pass", bLineFound);
+
+					check(IFileManager::Get().DeleteDirectory(*LogDirPath, false, true));
+					FPlatformProcess::TerminateProc(EditorHandle);
+
+					// Test stopped trace 
+					InsightsManager->OpenUnrealInsights(*TraceParameters);
+					bLineFound = Utils.FileContainsString(LogPath, SuccessTestResult, 60.0f);
+					TestTrue("Test for stopped trace should pass", bLineFound);
+
+					check(IFileManager::Get().DeleteDirectory(*LogDirPath, false, true));
+				});
+		});
 	AfterEach([this]() {
 		Driver.Reset();
 		IAutomationDriverModule::Get().Disable();
+		AutomationWindow->Restore();
 	});
 }
 
+////////////////////////////////////////////////////////////////////////////////////////////////////
+
+// TestRail: C35547680
+bool FMemoryInsightsUploadLLMXMLReportsTraceTest::RunTest(const FString& Parameters)
+{
+	const FString ReportGraphsXMLPath = FPaths::RootDir() / TEXT("EngineTest/SourceAssets/Utrace/ReportGraphs.xml");
+	const FString LLMReportTypesXMLPath = FPaths::RootDir() / TEXT("EngineTest/SourceAssets/Utrace/LLMReportTypes.xml");
+
+	FMemorySharedState* SharedState = FMemoryProfilerManager::Get()->GetSharedState();
+	if (!SharedState)
+	{
+		AddError("ProfilerWindow should be valid. Please, run this test throught Insights Session automation tab");
+		return false;
+	}
+
+	SharedState->RemoveAllMemTagGraphTracks();
+	const int DefaultTracksAmount = SharedState->GetTimingView()->GetAllTracks().Num();
+	SharedState->RemoveAllMemTagGraphTracks();
+	AddExpectedError("Failed to load Report");
+	SharedState->CreateTracksFromReport(ReportGraphsXMLPath);
+	SharedState = FMemoryProfilerManager::Get()->GetSharedState();
+	const int AfterReportGraphsUploadTrackAmount = SharedState->GetTimingView()->GetAllTracks().Num();
+	TestTrue("Tracks amount should be default ", DefaultTracksAmount == AfterReportGraphsUploadTrackAmount);
+
+	SharedState->RemoveAllMemTagGraphTracks();
+	SharedState->CreateTracksFromReport(LLMReportTypesXMLPath);
+	SharedState = FMemoryProfilerManager::Get()->GetSharedState();
+	const int AfterLLMReportTypesUploadTrackAmount = SharedState->GetTimingView()->GetAllTracks().Num();
+	TestTrue("Tracks should not be default", DefaultTracksAmount != AfterLLMReportTypesUploadTrackAmount);
+
+	return true;
+}
+
 #endif // !WITH_EDITOR
 
 ////////////////////////////////////////////////////////////////////////////////////////////////////
diff --git a/Engine/Source/Developer/TraceInsights/Private/Insights/Tests/UITests.h b/Engine/Source/Developer/TraceInsights/Private/Insights/Tests/UITests.h
--- a/Engine/Source/Developer/TraceInsights/Private/Insights/Tests/UITests.h
+++ b/Engine/Source/Developer/TraceInsights/Private/Insights/Tests/UITests.h
@@ -5,14 +5,7 @@
 #include "CoreMinimal.h"
 #include "Logging/LogMacros.h"
 
-//Driver
-#include "IAutomationDriver.h"
-#include "IAutomationDriverModule.h"
-#include "DriverConfiguration.h"
-#include "IDriverElement.h"
-#include "IDriverSequence.h"
-#include "LocateBy.h"
-
+#include "AutomationDriverCommon.h"
 #include "Misc/AutomationTest.h"
 
 DECLARE_LOG_CATEGORY_EXTERN(UITests, Log, All);
@@ -20,10 +13,11 @@ DECLARE_LOG_CATEGORY_EXTERN(UITests, Log, All);
 #if !WITH_EDITOR
 
 IMPLEMENT_SIMPLE_AUTOMATION_TEST(FHideAndShowAllTimingViewTabs, "Insights.HideAndShowAllTimingViewTabs", EAutomationTestFlags::ApplicationContextMask | EAutomationTestFlags::ProductFilter)
-BEGIN_DEFINE_SPEC(FAutomaticRenamingAndDeletingOfSymbolCacheFilesInsightsTest, "Insights.SessionBrowser.AutomaticRenamingAndDeletingOfSymbolCacheFiles", EAutomationTestFlags::ApplicationContextMask | EAutomationTestFlags::ProductFilter)
+BEGIN_DEFINE_SPEC(FAutomationDriverUnrealInsightsSessionBrowserTest, "Insights.SessionBrowser", EAutomationTestFlags::ApplicationContextMask | EAutomationTestFlags::ProductFilter)
 FAutomationDriverPtr Driver;
-END_DEFINE_SPEC(FAutomaticRenamingAndDeletingOfSymbolCacheFilesInsightsTest)
-
+TSharedPtr<SWindow> AutomationWindow;
+END_DEFINE_SPEC(FAutomationDriverUnrealInsightsSessionBrowserTest)
+IMPLEMENT_SIMPLE_AUTOMATION_TEST(FMemoryInsightsUploadLLMXMLReportsTraceTest, "Insights.MemoryInsights.UploadMemoryInsightsLLMXMLReportsTrace", EAutomationTestFlags::ApplicationContextMask | EAutomationTestFlags::ProductFilter)
 #endif
 
 IMPLEMENT_SIMPLE_AUTOMATION_TEST(FMemoryFilterValueConverterTest, "Insights.FMemoryFilterValueConverterTest", EAutomationTestFlags::ApplicationContextMask | EAutomationTestFlags::ProductFilter)
diff --git a/Engine/Source/Developer/TraceInsights/Public/Insights/Tests/InsightsTestUtils.h b/Engine/Source/Developer/TraceInsights/Public/Insights/Tests/InsightsTestUtils.h
--- a/Engine/Source/Developer/TraceInsights/Public/Insights/Tests/InsightsTestUtils.h
+++ b/Engine/Source/Developer/TraceInsights/Public/Insights/Tests/InsightsTestUtils.h
@@ -12,6 +12,7 @@ public:
 	FInsightsTestUtils(FAutomationTestBase* Test);
 
 	bool AnalyzeTrace(const TCHAR* Path) const;
+	bool FileContainsString(const FString& PathToFile, const FString& ExpectedString, const float Timeout) const;
 
 private:
 	FAutomationTestBase* Test;
