[Insights] Fixed Compile errors in UITests.cpp

diff --git a/Engine/Source/Developer/TraceInsights/Private/Insights/Tests/UITests.cpp b/Engine/Source/Developer/TraceInsights/Private/Insights/Tests/UITests.cpp
--- a/Engine/Source/Developer/TraceInsights/Private/Insights/Tests/UITests.cpp
+++ b/Engine/Source/Developer/TraceInsights/Private/Insights/Tests/UITests.cpp
@@ -18,6 +18,7 @@
 #include "Insights/MemoryProfiler/ViewModels/MemorySharedState.h"
 #include "Insights/Widgets/STimingView.h"
 #include "Insights/Tests/InsightsTestUtils.h"
+#include "Insights/TaskGraphProfiler/ViewModels/TaskTimingTrack.h"
 
 ////////////////////////////////////////////////////////////////////////////////////////////////////
 
