[Insights] Added forgotten change for CL 26634054

diff --git a/Engine/Source/Developer/TraceInsights/Private/Insights/Widgets/STimingView.h b/Engine/Source/Developer/TraceInsights/Private/Insights/Widgets/STimingView.h
--- a/Engine/Source/Developer/TraceInsights/Private/Insights/Widgets/STimingView.h
+++ b/Engine/Source/Developer/TraceInsights/Private/Insights/Widgets/STimingView.h
@@ -387,6 +387,8 @@ public:
 
 	TSharedPtr<Insights::FFilterConfigurator> GetFilterConfigurator() { return FilterConfigurator; }
 
+	TMap<uint64, TSharedPtr<FBaseTimingTrack>>& GetAllTracks() { return AllTracks; }
+
 protected:
 	virtual FVector2D ComputeDesiredSize(float) const override
 	{
