[Insights] Automated C27304222 "Unreal Insights - Commands - Export - Windows".

diff --git a/Engine/Source/Developer/TraceInsights/Private/Insights/Tests/MiscTests.cpp b/Engine/Source/Developer/TraceInsights/Private/Insights/Tests/MiscTests.cpp
--- a/Engine/Source/Developer/TraceInsights/Private/Insights/Tests/MiscTests.cpp
+++ b/Engine/Source/Developer/TraceInsights/Private/Insights/Tests/MiscTests.cpp
@@ -3,7 +3,11 @@
 #include "Misc/AutomationTest.h"
 //awful hack since SlabAllocator is private
 #include "Developer/TraceServices/Private/Common/SlabAllocator.h"
+#include "HAL/PlatformFileManager.h"
 #include "Common/PagedArray.h"
+#include "Insights/InsightsManager.h"
+#include "Insights/InsightsCommands.h"
+#include "Insights/Tests/InsightsTestUtils.h"
 
 IMPLEMENT_SIMPLE_AUTOMATION_TEST(FPagedArrayFilteringTest, "Insights.FPagedArrayFilteringTest", EAutomationTestFlags::ApplicationContextMask | EAutomationTestFlags::SmokeFilter)
 
@@ -80,4 +84,138 @@ bool FPagedArrayFilteringTest::RunTest(const FString& Parameters)
 	TestEqual(TEXT("Overlaps for interval inside to after element-range"), Result, {5,6});
 
 	return !HasAnyErrors();
-}
\ No newline at end of file
+}
+
+#if !WITH_EDITOR
+
+IMPLEMENT_SIMPLE_AUTOMATION_TEST(FСommandsExportWindowsTest, "Insights.CommandsExport(Windows)", EAutomationTestFlags::ApplicationContextMask | EAutomationTestFlags::ProductFilter)
+
+// TestRail: C27304222
+bool FСommandsExportWindowsTest::RunTest(const FString& Parameters)
+{
+	TSharedPtr<FInsightsManager> InsightsManager = FInsightsManager::Get();
+	IPlatformFile& PlatformFile = FPlatformFileManager::Get().GetPlatformFile();
+	FInsightsTestUtils Utils(this);
+
+	const FString StoreDir = InsightsManager->GetStoreDir();
+	const FString SourceTracePath = FPaths::RootDir() / TEXT("EngineTest/SourceAssets/Utrace/Editor_2.utrace");
+	const FString SourceExportPath = FPaths::RootDir() / TEXT("EngineTest/SourceAssets/Rsp/export.rsp");
+	const FString StoreTracePath = StoreDir / TEXT("Editor_2.utrace");
+	const FString LogResultExportPath = TEXT("/TestResults/export.rsp");
+	const FString CmdLogPath = TEXT("/TestResults/Logs/cmd.log");
+	const FString ExportThreadsTask = TEXT("TimingInsights.ExportThreads /TestResults/SingleCommand/Threads.csv");
+	const FString ExportTimersTask = TEXT("TimingInsights.ExportTimers /TestResults/SingleCommand/Timers.csv");
+	const FString ExportTimingEventsTask = TEXT("TimingInsights.ExportTimingEvents /TestResults/SingleCommand/TimingEvents.TXT");
+	const FString ExportExportTimingEventsBorderedTask = TEXT("TimingInsights.ExportTimingEvents /TestResults/SingleCommand/TimingEventsNonDefault.CSV  -columns=ThreadId,ThreadName,TimerId,TimerName,StartTime,EndTime,Duration,Depth -threads=GameThread -timers=* -startTime=10 -endTime=20");
+	
+	TestTrue("Trace in project exists", PlatformFile.FileExists(*SourceTracePath));
+	TestTrue("Export in project exists", PlatformFile.FileExists(*SourceExportPath));
+	TestFalse("Trace in store should not exists before copy", PlatformFile.FileExists(*StoreTracePath));
+	
+	PlatformFile.CopyFile(*StoreTracePath, *SourceTracePath);
+	if (!PlatformFile.FileExists(*StoreTracePath))
+	{
+		AddError("Trace in store should exists after copy");
+		return false;
+	}
+
+    // ExportThreads
+	FString InsightsParameters = FString::Printf(TEXT("-OpenTraceFile=%s -ABSLOG=\"%s\" -AutoQuit -NoUI  -ExecOnAnalysisCompleteCmd=\"%s\" -log"), *StoreTracePath, *CmdLogPath, *ExportThreadsTask);
+	InsightsManager->OpenUnrealInsights(*InsightsParameters);
+	
+	FString ExpectedThreadsElements[5] = { TEXT("Id,Name,Group"), TEXT("1196447025,GPU1,GPU"), TEXT("1196447026,GPU2,GPU"), TEXT("2,GameThread,"), TEXT("88,RTHeartBeat 0,Render") };
+	FString ExpectedResult = TEXT("Exported 107 threads to file");
+	bool bLineFound = Utils.FileContainsString(CmdLogPath, ExpectedResult, 15.0f);
+	TestTrue(FString::Printf(TEXT("Line %s should exists in file"), *ExpectedResult), bLineFound);
+
+	for (int i = 0; i < 5; i++)
+	{
+		bLineFound = Utils.FileContainsString(TEXT("/TestResults/SingleCommand/Threads.csv"), ExpectedThreadsElements[i], 15.0f);
+		TestTrue(FString::Printf(TEXT("Line %s should exists in file"), *ExpectedThreadsElements[i]), bLineFound);
+	}
+
+	// ExportTimers
+	InsightsParameters = FString::Printf(TEXT("-OpenTraceFile=%s -ABSLOG=\"%s\" -AutoQuit -NoUI  -ExecOnAnalysisCompleteCmd=\"%s\" -log"), *StoreTracePath, *CmdLogPath, *ExportTimersTask);
+	InsightsManager->OpenUnrealInsights(*InsightsParameters);
+
+	ExpectedResult = TEXT("Exported 1942 timers to file");
+	bLineFound = Utils.FileContainsString(CmdLogPath, ExpectedResult, 15.0f);
+	TestTrue(FString::Printf(TEXT("Line %s should exists in file"), *ExpectedResult), bLineFound);
+
+	bLineFound = Utils.FileContainsString(TEXT("/TestResults/SingleCommand/Timers.csv"), TEXT("Id,Type,Name,File,Line"), 15.0f);
+	TestTrue(FString::Printf(TEXT("Line %s should exists in file"), TEXT("Id,Type,Name,File,Line")), bLineFound);
+
+	bLineFound = Utils.FileContainsString(TEXT("/TestResults/SingleCommand/Timers.csv"), TEXT("15,CPU,LoadModule_StorageServerClient"), 15.0f);
+	TestTrue(FString::Printf(TEXT("Line %s should exists in file"), TEXT("15,CPU,LoadModule_StorageServerClient")), bLineFound);
+
+	// UI verification cannot be executed in that case. 
+	
+	// ExportTimingEvents
+	InsightsParameters = FString::Printf(TEXT("-OpenTraceFile=%s -ABSLOG=\"%s\" -AutoQuit -NoUI  -ExecOnAnalysisCompleteCmd=\"%s\" -log"), *StoreTracePath, *CmdLogPath, *ExportTimingEventsTask);
+	InsightsManager->OpenUnrealInsights(*InsightsParameters);
+
+	FString ExpectedTimingEventsElements[3] = { TEXT("ThreadId	TimerId	StartTime	EndTime	Depth"), TEXT("1196447025	1916	46.510678	46.510955	0"), TEXT("1196447025	1916	46.51073	46.510955	1") };
+	ExpectedResult = TEXT("Exported 2381149 timing events to file");
+	bLineFound = Utils.FileContainsString(CmdLogPath, ExpectedResult, 15.0f);
+	TestTrue(FString::Printf(TEXT("Line %s should exists in file"), *ExpectedResult), bLineFound);
+
+	for (int i = 0; i < 3; i++)
+	{
+		bLineFound = Utils.FileContainsString(TEXT("/TestResults/SingleCommand/TimingEvents.TXT"), ExpectedTimingEventsElements[i], 15.0f);
+		TestTrue(FString::Printf(TEXT("Line %s should exists in file"), *ExpectedTimingEventsElements[i]), bLineFound);
+	}
+
+	// ExportTimingEventsNonDefault
+	InsightsParameters = FString::Printf(TEXT("-OpenTraceFile=%s -ABSLOG=\"%s\" -AutoQuit -NoUI  -ExecOnAnalysisCompleteCmd=\"%s\" -log"), *StoreTracePath, *CmdLogPath, *ExportExportTimingEventsBorderedTask);
+	InsightsManager->OpenUnrealInsights(*InsightsParameters);
+
+	FString ExpectedTimingEventsBorderedElements[3] = { TEXT("ThreadId,ThreadName,TimerId,TimerName,StartTime,EndTime,Duration,Depth"), 
+		TEXT("2,GameThread,0,FEngineLoop::PreInitPreStartupScreen,0.3376052,41.5370876,41.199482400,0"), 
+		TEXT("2,GameThread,182,InitDerivedData,1.7132565,40.2586897,38.545433200,1") };
+	ExpectedResult = TEXT("Exported 2 timing events to file");
+
+	bLineFound = Utils.FileContainsString(CmdLogPath, ExpectedResult, 15.0f);
+	TestTrue(FString::Printf(TEXT("Line %s should exists in file"), *ExpectedResult), bLineFound);
+
+	for (int i = 0; i < 3; i++)
+	{
+		bLineFound = Utils.FileContainsString(TEXT("/TestResults/SingleCommand/TimingEventsNonDefault.csv"), ExpectedTimingEventsBorderedElements[i], 15.0f);
+		TestTrue(FString::Printf(TEXT("Line %s should exists in file"), *ExpectedTimingEventsBorderedElements[i]), bLineFound);
+	}
+
+	// Export.rsp
+	PlatformFile.CopyFile(*LogResultExportPath, *SourceExportPath);
+	TestTrue("Rsp in log directory should exists after copy", PlatformFile.FileExists(*LogResultExportPath));
+ 
+	InsightsParameters = FString::Printf(TEXT("-OpenTraceFile=%s -ABSLOG=\"%s\" -AutoQuit -NoUI  -ExecOnAnalysisCompleteCmd=\"@=/TestResults/export.rsp\" -log"), *StoreTracePath, *CmdLogPath);
+	InsightsManager->OpenUnrealInsights(*InsightsParameters);
+
+	FString ExpectedExportElements[14] = { TEXT("(\"/TestResults/RSPtest/CSV/Threads_rsp.csv\")"), 
+		TEXT("(\"/TestResults/RSPtest/TSV/Threads_rsp.tsv\")"),
+		TEXT("(\"/TestResults/RSPtest/TXT/Threads_rsp.txt\")"),
+		TEXT("(\"/TestResults/RSPtest/CSV/Timers_rsp.csv\")"),
+		TEXT("(\"/TestResults/RSPtest/TSV/Timers_rsp.tsv\")"),
+		TEXT("(\"/TestResults/RSPtest/TXT/Timers_rsp.txt\")"),
+		TEXT("(\"/TestResults/RSPtest/CSV/TimingEvents_rsp.csv\")"),
+		TEXT("(\"/TestResults/RSPtest/TSV/TimingEvents_rsp.tsv\")"),
+		TEXT("(\"/TestResults/RSPtest/TXT/TimingEvents_rsp.txt\")"), 
+		TEXT("Exported 107 threads to file"),
+		TEXT("Exported 1942 threads to file"),
+		TEXT("Exported 538 threads to file"),
+		TEXT("Exported 30 threads to file"),
+		TEXT("Exported 1535 threads to file"),
+	};
+
+	for (int i = 0; i < 9; i++)
+	{
+		bLineFound = Utils.FileContainsString(CmdLogPath, ExpectedExportElements[i], 15.0f);
+		TestTrue(FString::Printf(TEXT("Line %s should exists in file"), *ExpectedExportElements[i]), bLineFound);
+	}
+
+	IFileManager::Get().Delete(*StoreTracePath);
+	check(IFileManager::Get().DeleteDirectory(TEXT("/TestResults"), false, true));
+
+	return true;
+}
+
+#endif // !WITH_EDITOR
