[Insights] Added some clean up for C35547680 and C28843555.
Made TraceWaiter more stable and independend of trace location.
Added test exist of TraceWaiter return false for C28843555.

diff --git a/Engine/Source/Developer/TraceInsights/Private/Insights/Tests/UITests.cpp b/Engine/Source/Developer/TraceInsights/Private/Insights/Tests/UITests.cpp
--- a/Engine/Source/Developer/TraceInsights/Private/Insights/Tests/UITests.cpp
+++ b/Engine/Source/Developer/TraceInsights/Private/Insights/Tests/UITests.cpp
@@ -82,8 +82,8 @@ void FAutomationDriverUnrealInsightsSessionBrowserTest::Define()
 					TestTrue("Insigts manager should not be null", InsightsManager.IsValid());
 					InsightsManager->GetTraceStoreWindow()->SetDeleteTraceConfirmationWindowVisibility(false);
 
-					FString StoreDir = InsightsManager->GetStoreDir();
-					FString ProjectDir = FPaths::ProjectDir();
+					const FString StoreDir = InsightsManager->GetStoreDir();
+					const FString ProjectDir = FPaths::ProjectDir();
 
 					const FString SourceTestTracePath = FPaths::RootDir() / TEXT("EngineTest/SourceAssets/Utrace/Test.utrace");
 					const FString SourceTestCachePath = FPaths::RootDir() / TEXT("EngineTest/SourceAssets/Utrace/Test.ucache");
@@ -115,11 +115,27 @@ void FAutomationDriverUnrealInsightsSessionBrowserTest::Define()
 					TestFalse("Renamed cache should not exist before renaming", PlatformFile.FileExists(*StoreTestCachePath));
 
 					// Rename 
-					auto TraceWaiter = [Driver = Driver](void) -> bool {
-						return Driver->FindElements(By::Id("TraceList"))->GetElements()[0]->GetText().ToString() == TEXT("Test");
+					int Index = 0;
+					auto TraceWaiter = [Driver = Driver, &Index](void) -> bool
+					{
+						auto Elements = Driver->FindElements(By::Id("TraceList"))->GetElements();
+						for (int i = 0; i < Elements.Num(); ++i) {
+							if (Elements[i]->GetText().ToString() == TEXT("Test")) {
+								Index = i;
+								return true;
+							}
+						}
+						return false;
 					};
-					Driver->Wait(Until::Condition(TraceWaiter, FWaitTimeout::InSeconds(3)));
-					FDriverElementRef TraceElement = Driver->FindElements(By::Id("TraceList"))->GetElements()[0];
+
+					bool bTestTraceExists = Driver->Wait(Until::Condition(TraceWaiter, FWaitTimeout::InSeconds(10)));
+					if (!bTestTraceExists)
+					{
+						AddError("Trace should exists in Session Browser");
+						return;
+					}
+
+					FDriverElementRef TraceElement = Driver->FindElements(By::Id("TraceList"))->GetElements()[Index];
 
 					FDriverSequenceRef Sequence = Driver->CreateSequence();
 					Sequence->Actions()
@@ -135,7 +151,7 @@ void FAutomationDriverUnrealInsightsSessionBrowserTest::Define()
 
 					// Delete
 					FDriverElementRef OpenTraceButton = Driver->FindElement(By::Id("OpenTraceButton"));
-					Driver->Wait(Until::ElementIsInteractable(OpenTraceButton, FWaitTimeout::InSeconds(5)));
+					Driver->Wait(Until::ElementIsInteractable(OpenTraceButton, FWaitTimeout::InSeconds(10)));
 
 					TraceElement = Driver->FindElements(By::Id("TraceList"))->GetElements()[0];
 					TraceElement->Type(EKeys::Delete);
@@ -143,6 +159,15 @@ void FAutomationDriverUnrealInsightsSessionBrowserTest::Define()
 					TestFalse("Renamed trace should be deleted", PlatformFile.FileExists(*StoreTestTracePath));
 					TestFalse("Renamed cache should be deleted", PlatformFile.FileExists(*StoreTestCachePath));
 				});
+
+			AfterEach([this]() {
+				TSharedPtr<FInsightsManager> InsightsManager = FInsightsManager::Get();
+				const FString StoreDir = InsightsManager->GetStoreDir();
+				const FString StoreTestTracePath = StoreDir / TEXT("Test.utrace");
+				const FString StoreTestCachePath = StoreDir / TEXT("Test.ucache");
+				IFileManager::Get().Delete(*StoreTestTracePath, false, true);
+				IFileManager::Get().Delete(*StoreTestCachePath, false, true);
+				});
 		});
 
 	// TestRail: C35547680
@@ -173,14 +198,24 @@ void FAutomationDriverUnrealInsightsSessionBrowserTest::Define()
 					FProcHandle EditorHandle = FPlatformProcess::CreateProc(*UEPath, *Parameters, bLaunchDetached, bLaunchHidden, bLaunchReallyHidden, &ProcessID, PriorityModifier, OptionalWorkingDirectory, PipeWriteChild, PipeReadChild);
 					if (!EditorHandle.IsValid())
 					{
-						AddError("Lyra should be started");
+						AddError("Editor should be started");
 						return;
 					}
 
 					// Verify that LIVE trace appeared
-					auto TraceWaiter = [Driver = Driver](void) -> bool {
-						return Driver->FindElements(By::Id("TraceStatusColumnList"))->GetElements()[0]->GetText().ToString() == TEXT("LIVE");
+					int Index = 0;
+					auto TraceWaiter = [Driver = Driver, &Index](void) -> bool
+					{
+						auto Elements = Driver->FindElements(By::Id("TraceStatusColumnList"))->GetElements();
+						for (int i = 0; i < Elements.Num(); ++i) {
+							if (Elements[i]->GetText().ToString() == TEXT("LIVE")) {
+								Index = i;
+								return true;
+							}
+						}
+						return false;
 					};
+
 					if (!Driver->Wait(Until::Condition(TraceWaiter, FWaitTimeout::InSeconds(10))))
 					{
 						AddError("Live trace should appear");
@@ -188,12 +223,13 @@ void FAutomationDriverUnrealInsightsSessionBrowserTest::Define()
 						return;
 					}
 
-					FDriverElementRef TraceElement = Driver->FindElements(By::Id("TraceList"))->GetElements()[0];
+					FDriverElementRef TraceElement = Driver->FindElements(By::Id("TraceList"))->GetElements()[Index];
 					const FString TraceName = TraceElement->GetText().ToString();
 
 					const FString StoreDir = InsightsManager->GetStoreDir();
 					const FString ProjectDir = FPaths::ProjectDir();
 					const FString StoreTracePath = StoreDir / FString::Printf(TEXT("%s.utrace"), *TraceName);
+					const FString StoreCachePath = StoreDir / FString::Printf(TEXT("%s.ucache"), *TraceName);
 					const FString LogDirPath = ProjectDir / TEXT("TestResults");
 					const FString LogPath = ProjectDir / TEXT("TestResults/Log.txt");
 					const FString SuccessTestResult = TEXT("Test Completed. Result={Success}");
@@ -213,6 +249,8 @@ void FAutomationDriverUnrealInsightsSessionBrowserTest::Define()
 					TestTrue("Test for stopped trace should pass", bLineFound);
 
 					check(IFileManager::Get().DeleteDirectory(*LogDirPath, false, true));
+					check(IFileManager::Get().Delete(*StoreTracePath, false, true));
+					check(IFileManager::Get().Delete(*StoreCachePath, false, true));
 				});
 		});
 	AfterEach([this]() {
