[Insights] Automated C27304222 "Unreal Insights - UnrealTraceServer - Starting".

diff --git a/Engine/Source/Developer/TraceInsights/Private/Insights/Tests/UITests.cpp b/Engine/Source/Developer/TraceInsights/Private/Insights/Tests/UITests.cpp
--- a/Engine/Source/Developer/TraceInsights/Private/Insights/Tests/UITests.cpp
+++ b/Engine/Source/Developer/TraceInsights/Private/Insights/Tests/UITests.cpp
@@ -293,6 +293,109 @@ bool FMemoryInsightsUploadLLMXMLReportsTraceTest::RunTest(const FString& Paramet
 	return true;
 }
 
+// TestRail: T223693543
+bool FUnrealTraceServerStartingTest::RunTest(const FString& Parameters)
+{
+	TSharedPtr<FInsightsManager> InsightsManager = FInsightsManager::Get();
+	const FString UnrealTraceServerName = TEXT("UnrealTraceServer");
+	const float Timeout = 10.0f;
+
+	auto WaitForUTSProcess = [](const FString& ProcessName, const float Timeout, const bool ExpectedStatus) -> bool
+	{
+		double StartTime = FPlatformTime::Seconds();
+		while (FPlatformTime::Seconds() - StartTime < Timeout)
+		{
+			if (FPlatformProcess::IsApplicationRunning(*ProcessName) == ExpectedStatus)
+			{
+				return ExpectedStatus;
+				break;
+			}
+			FPlatformProcess::Sleep(0.1f);
+		}
+		return !ExpectedStatus;
+	};
+
+	auto WaitForConnectionIconStatus = [](TSharedPtr<class FInsightsManager> InsightsManager, const float Timeout, const bool ExpectedStatus) -> bool
+	{
+		double StartTime = FPlatformTime::Seconds();
+		while (FPlatformTime::Seconds() - StartTime < Timeout)
+		{
+			TSharedPtr<STraceStoreWindow> Wnd = InsightsManager->GetTraceStoreWindow();
+			if (Wnd.IsValid())
+			{
+				if (Wnd->IsConnected() == ExpectedStatus)
+				{
+					return ExpectedStatus;
+					break;
+				}
+			}
+			FPlatformProcess::Sleep(0.1f);
+		}
+		return !ExpectedStatus;
+	};
+
+	bool bIsUnrealTraceServerRunning = WaitForUTSProcess(UnrealTraceServerName, Timeout, true);
+	TestTrue("UnrealTraceServer should be in processes", bIsUnrealTraceServerRunning);
+
+	bool bConnectionStatus = WaitForConnectionIconStatus(InsightsManager, Timeout, true);
+	TestTrue("Connection status should be true", bConnectionStatus);
+
+#if PLATFORM_WINDOWS
+	FString Command = FString::Printf(TEXT("/F /IM \"%s.exe\""), *UnrealTraceServerName);
+	FPlatformProcess::ExecProcess(TEXT("taskkill"), *Command, nullptr, nullptr, nullptr);
+#elif PLATFORM_MAC || PLATFORM_LINUX
+	FString Command = FString::Printf(TEXT("killall \"%s\""), *ProcessName);
+	FPlatformProcess::ExecProcess(TEXT("/bin/sh"), *Command, nullptr, nullptr, nullptr);
+#endif
+
+	bIsUnrealTraceServerRunning = WaitForUTSProcess(UnrealTraceServerName, Timeout, false);
+	TestFalse("UnrealTraceServer should not be in processes", bIsUnrealTraceServerRunning);
+
+	bConnectionStatus = WaitForConnectionIconStatus(InsightsManager, Timeout, false);
+	TestFalse("Connection status should be false", bConnectionStatus);
+
+	FString UEPath = FPlatformProcess::GenerateApplicationPath("UnrealEditor", EBuildConfiguration::Development);
+	FString UEParameters = TEXT("");
+	constexpr bool bLaunchDetached = true;
+	constexpr bool bLaunchHidden = false;
+	constexpr bool bLaunchReallyHidden = false;
+	uint32 ProcessID = 0;
+	const int32 PriorityModifier = 0;
+	const TCHAR* OptionalWorkingDirectory = nullptr;
+	void* PipeWriteChild = nullptr;
+	void* PipeReadChild = nullptr;
+	FProcHandle EditorHandle = FPlatformProcess::CreateProc(*UEPath, *UEParameters, bLaunchDetached, bLaunchHidden, bLaunchReallyHidden, &ProcessID, PriorityModifier, OptionalWorkingDirectory, PipeWriteChild, PipeReadChild);
+	if (!EditorHandle.IsValid())
+	{
+		AddError("Editor should be started");
+		return true;
+	}
+
+	bIsUnrealTraceServerRunning = WaitForUTSProcess(UnrealTraceServerName, Timeout, true);
+	TestTrue("UnrealTraceServer should be in processes", bIsUnrealTraceServerRunning);
+
+	bConnectionStatus = WaitForConnectionIconStatus(InsightsManager, Timeout, true);
+	TestTrue("Connection status should be true", bConnectionStatus);
+
+	FPlatformProcess::TerminateProc(EditorHandle);
+
+#if PLATFORM_WINDOWS
+	FPlatformProcess::ExecProcess(TEXT("taskkill"), *Command, nullptr, nullptr, nullptr);
+#elif PLATFORM_MAC || PLATFORM_LINUX
+	FPlatformProcess::ExecProcess(TEXT("/bin/sh"), *Command, nullptr, nullptr, nullptr);
+#endif
+
+	bIsUnrealTraceServerRunning = WaitForUTSProcess(UnrealTraceServerName, Timeout, false);
+	TestFalse("UnrealTraceServer should not be in processes", bIsUnrealTraceServerRunning);
+
+	FString Path = FPaths::RootDir() / TEXT("EngineTest/SourceAssets/Utrace/r423_win64_game_10478456.utrace");
+	InsightsManager->OpenUnrealInsights(TEXT("-InsightsTest -AutoQuit -NoUI -OpenTraceFile=%s"));
+	bIsUnrealTraceServerRunning = WaitForUTSProcess(UnrealTraceServerName, Timeout, true);
+	TestTrue("UnrealTraceServer should be in processes", bIsUnrealTraceServerRunning);
+
+	return true;
+}
+
 #endif // !WITH_EDITOR
 
 ////////////////////////////////////////////////////////////////////////////////////////////////////
diff --git a/Engine/Source/Developer/TraceInsights/Private/Insights/Tests/UITests.h b/Engine/Source/Developer/TraceInsights/Private/Insights/Tests/UITests.h
--- a/Engine/Source/Developer/TraceInsights/Private/Insights/Tests/UITests.h
+++ b/Engine/Source/Developer/TraceInsights/Private/Insights/Tests/UITests.h
@@ -18,6 +18,8 @@ FAutomationDriverPtr Driver;
 TSharedPtr<SWindow> AutomationWindow;
 END_DEFINE_SPEC(FAutomationDriverUnrealInsightsSessionBrowserTest)
 IMPLEMENT_SIMPLE_AUTOMATION_TEST(FMemoryInsightsUploadLLMXMLReportsTraceTest, "Insights.MemoryInsights.UploadMemoryInsightsLLMXMLReportsTrace", EAutomationTestFlags::ApplicationContextMask | EAutomationTestFlags::ProductFilter)
+IMPLEMENT_SIMPLE_AUTOMATION_TEST(FUnrealTraceServerStartingTest, "Insights.SessionBrowser.UnrealTraceServer.Startng", EAutomationTestFlags::ApplicationContextMask | EAutomationTestFlags::ProductFilter)
+
 #endif
 
 IMPLEMENT_SIMPLE_AUTOMATION_TEST(FMemoryFilterValueConverterTest, "Insights.FMemoryFilterValueConverterTest", EAutomationTestFlags::ApplicationContextMask | EAutomationTestFlags::ProductFilter)
diff --git a/Engine/Source/Developer/TraceInsights/Private/Insights/Widgets/SStartPageWindow.cpp b/Engine/Source/Developer/TraceInsights/Private/Insights/Widgets/SStartPageWindow.cpp
--- a/Engine/Source/Developer/TraceInsights/Private/Insights/Widgets/SStartPageWindow.cpp
+++ b/Engine/Source/Developer/TraceInsights/Private/Insights/Widgets/SStartPageWindow.cpp
@@ -2310,6 +2310,13 @@ void STraceStoreWindow::SetDeleteTraceConfirmationWindowVisibility(bool bIsVisib
 	bIsDeleteTraceConfirmWindowVisible = bIsVisibile;
 }
 
+bool STraceStoreWindow::IsConnected() const
+{
+	using Insights::FStoreBrowser;
+
+	return StoreBrowser->GetConnectionStatus() == FStoreBrowser::EConnectionStatus::Connected;
+}
+
 ////////////////////////////////////////////////////////////////////////////////////////////////////
 
 void STraceStoreWindow::UpdateTrace(FTraceViewModel& InOutTrace, const Insights::FStoreBrowserTraceInfo& InSourceTrace)
diff --git a/Engine/Source/Developer/TraceInsights/Private/Insights/Widgets/SStartPageWindow.h b/Engine/Source/Developer/TraceInsights/Private/Insights/Widgets/SStartPageWindow.h
--- a/Engine/Source/Developer/TraceInsights/Private/Insights/Widgets/SStartPageWindow.h
+++ b/Engine/Source/Developer/TraceInsights/Private/Insights/Widgets/SStartPageWindow.h
@@ -279,6 +279,7 @@ public:
 	void OnFilterChanged();
 	const TArray<TSharedPtr<FTraceViewModel>>& GetAllAvailableTraces() const;
 	void SetDeleteTraceConfirmationWindowVisibility(bool bIsVisibile);
+	bool IsConnected() const;
 
 private:
 	TSharedRef<SWidget> ConstructFiltersToolbar();
