[Insights] Added forgotten lines for code style

diff --git a/Engine/Source/Developer/TraceInsights/Private/Insights/Tests/UITests.cpp b/Engine/Source/Developer/TraceInsights/Private/Insights/Tests/UITests.cpp
--- a/Engine/Source/Developer/TraceInsights/Private/Insights/Tests/UITests.cpp
+++ b/Engine/Source/Developer/TraceInsights/Private/Insights/Tests/UITests.cpp
@@ -293,6 +293,8 @@ bool FMemoryInsightsUploadLLMXMLReportsTraceTest::RunTest(const FString& Paramet
 	return true;
 }
 
+////////////////////////////////////////////////////////////////////////////////////////////////////
+
 // TestRail: T223693543
 bool FUnrealTraceServerStartingTest::RunTest(const FString& Parameters)
 {
