[Insights] Fixed compile build for Linux in Insights

diff --git a/Engine/Source/Developer/TraceInsights/Private/Insights/Tests/UITests.cpp b/Engine/Source/Developer/TraceInsights/Private/Insights/Tests/UITests.cpp
--- a/Engine/Source/Developer/TraceInsights/Private/Insights/Tests/UITests.cpp
+++ b/Engine/Source/Developer/TraceInsights/Private/Insights/Tests/UITests.cpp
@@ -346,7 +346,7 @@ bool FUnrealTraceServerStartingTest::RunTest(const FString& Parameters)
 	FString Command = FString::Printf(TEXT("/F /IM \"%s.exe\""), *UnrealTraceServerName);
 	FPlatformProcess::ExecProcess(TEXT("taskkill"), *Command, nullptr, nullptr, nullptr);
 #elif PLATFORM_MAC || PLATFORM_LINUX
-	FString Command = FString::Printf(TEXT("killall \"%s\""), *ProcessName);
+	FString Command = FString::Printf(TEXT("killall \"%s\""), *UnrealTraceServerName);
 	FPlatformProcess::ExecProcess(TEXT("/bin/sh"), *Command, nullptr, nullptr, nullptr);
 #endif
 
