[Insights] Fix timing regions bugs

diff --git a/Engine/Source/Developer/TraceInsights/Private/Insights/ViewModels/RegionsTimingTrack.cpp b/Engine/Source/Developer/TraceInsights/Private/Insights/ViewModels/RegionsTimingTrack.cpp
--- a/Engine/Source/Developer/TraceInsights/Private/Insights/ViewModels/RegionsTimingTrack.cpp
+++ b/Engine/Source/Developer/TraceInsights/Private/Insights/ViewModels/RegionsTimingTrack.cpp
@@ -11,6 +11,9 @@
 #include "Insights/TimingProfilerCommon.h"
 #include "Insights/Common/InsightsMenuBuilder.h"
 #include "Insights/Common/TimeUtils.h"
+#include "Insights/ViewModels/FilterConfigurator.h"
+#include "Insights/ViewModels/Filters.h"
+#include "Insights/ViewModels/ThreadTimingTrack.h"
 #include "Insights/ViewModels/TimingTrackViewport.h"
 #include "Insights/Widgets/STimingView.h"
 #include "TraceServices/Model/Regions.h"
@@ -18,6 +21,9 @@
 
 #define LOCTEXT_NAMESPACE "RegionsTimingTrack"
 
+namespace Insights
+{
+
 ////////////////////////////////////////////////////////////////////////////////////////////////////
 // FFileActivityTimingViewCommands
 ////////////////////////////////////////////////////////////////////////////////////////////////////
@@ -53,19 +59,30 @@ UE_ENABLE_OPTIMIZATION_SHIP
 
 ////////////////////////////////////////////////////////////////////////////////////////////////////
 
-
-
-
 void FTimingRegionsSharedState::OnBeginSession(Insights::ITimingViewSession& InSession)
 {
+	if (&InSession != TimingView)
+	{
+		return;
+	}
+
 	TimingRegionsTrack.Reset();
 }
 
+////////////////////////////////////////////////////////////////////////////////////////////////////
+
 void FTimingRegionsSharedState::OnEndSession(Insights::ITimingViewSession& InSession)
 {
+	if (&InSession != TimingView)
+	{
+		return;
+	}
+
 	TimingRegionsTrack.Reset();
 }
 
+////////////////////////////////////////////////////////////////////////////////////////////////////
+
 void FTimingRegionsSharedState::Tick(Insights::ITimingViewSession& InSession,
 	const TraceServices::IAnalysisSession& InAnalysisSession)
 {
@@ -83,6 +100,8 @@ void FTimingRegionsSharedState::Tick(Insights::ITimingViewSession& InSession,
 	}
 }
 
+////////////////////////////////////////////////////////////////////////////////////////////////////
+
 void FTimingRegionsSharedState::ShowHideRegionsTrack()
 {
 	bShowHideRegionsTrack = !bShowHideRegionsTrack;
@@ -103,6 +122,8 @@ void FTimingRegionsSharedState::ShowHideRegionsTrack()
 	}
 }
 
+////////////////////////////////////////////////////////////////////////////////////////////////////
+
 void FTimingRegionsSharedState::ExtendOtherTracksFilterMenu(Insights::ITimingViewSession& InSession,
                                                             FMenuBuilder& InOutMenuBuilder)
 {
@@ -124,6 +145,8 @@ void FTimingRegionsSharedState::ExtendOtherTracksFilterMenu(Insights::ITimingVie
 	InOutMenuBuilder.EndSection();
 }
 
+////////////////////////////////////////////////////////////////////////////////////////////////////
+
 void FTimingRegionsSharedState::BindCommands()
 {
 	FTimingRegionsViewCommands::Register();
@@ -138,7 +161,6 @@ void FTimingRegionsSharedState::BindCommands()
 		FIsActionChecked::CreateSP(this, &FTimingRegionsSharedState::IsRegionsTrackVisible));
 }
 
-
 //////////////////////////////////////////////////////////////////////////
 
 INSIGHTS_IMPLEMENT_RTTI(FTimingRegionsTrack)
@@ -152,6 +174,7 @@ void FTimingRegionsTrack::BuildContextMenu(FMenuBuilder& MenuBuilder)
 	FTimingEventsTrack::BuildContextMenu(MenuBuilder);
 }
 
+//////////////////////////////////////////////////////////////////////////
 
 void FTimingRegionsTrack::InitTooltip(FTooltipDrawState& InOutTooltip, const ITimingEvent& InTooltipEvent) const
 {
@@ -210,6 +233,106 @@ void FTimingRegionsTrack::BuildDrawState(ITimingEventsTrackDrawStateBuilder& Bui
 	UE_CLOG(TotalTime > 1.0,TimingProfiler, Verbose, TEXT("[Regions] Updated draw state in %s."), *TimeUtils::FormatTimeAuto(TotalTime));
 }
 
+//////////////////////////////////////////////////////////////////////////
+
+void FTimingRegionsTrack::BuildFilteredDrawState(ITimingEventsTrackDrawStateBuilder& Builder, const ITimingTrackUpdateContext& Context)
+{
+	const TSharedPtr<ITimingEventFilter> EventFilterPtr = Context.GetEventFilter();
+	if (EventFilterPtr.IsValid() && EventFilterPtr->FilterTrack(*this))
+	{
+		bool bFilterOnlyByEventType = false; // this is the most often use case, so the below code tries to optimize it
+		TCHAR* FilterEventType = 0;
+		if (EventFilterPtr->Is<FTimingEventFilterByEventType>())
+		{
+			bFilterOnlyByEventType = true;
+			const FTimingEventFilterByEventType& EventFilter = EventFilterPtr->As<FTimingEventFilterByEventType>();
+			FilterEventType = reinterpret_cast<TCHAR*>(EventFilter.GetEventType());
+		}
+
+		TSharedPtr<const TraceServices::IAnalysisSession> Session = FInsightsManager::Get()->GetSession();
+		if (Session.IsValid())
+		{
+			const TraceServices::IRegionProvider& RegionProvider = TraceServices::ReadRegionProvider(*Session);
+			TraceServices::FProviderReadScopeLock RegionProviderScopedLock(RegionProvider);
+
+			const FTimingTrackViewport& Viewport = Context.GetViewport();
+
+			if (bFilterOnlyByEventType)
+			{
+				int32 CurDepth = 0;
+				RegionProvider.EnumerateLanes([this, Viewport, &CurDepth, &Builder, FilterEventType](const TraceServices::FRegionLane& Lane, const int32 Depth)
+					{
+						bool RegionHadEvents = false;
+						Lane.EnumerateRegions(Viewport.GetStartTime(), Viewport.GetEndTime(), [&Builder, &RegionHadEvents, &CurDepth, FilterEventType](const TraceServices::FTimeRegion& Region) -> bool
+							{
+								RegionHadEvents = true;
+								if (Region.Text == FilterEventType)
+								{
+									Builder.AddEvent(Region.BeginTime, Region.EndTime, CurDepth, Region.Text);
+								}
+								return true;
+							});
+
+						if (RegionHadEvents) CurDepth++;
+					});
+			}
+			else // generic filter
+			{
+				//TODO: if (EventFilterPtr->FilterEvent(TimingEvent))
+			}
+		}
+	}
+
+	if (HasCustomFilter())
+	{
+		if (!FilterConfigurator.IsValid())
+		{
+			return;
+		}
+
+		FFilterContext FilterContext;
+		FilterContext.SetReturnValueForUnsetFilters(false);
+
+		FilterContext.AddFilterData<double>(static_cast<int32>(EFilterField::StartTime), 0.0f);
+		FilterContext.AddFilterData<double>(static_cast<int32>(EFilterField::EndTime), 0.0f);
+		FilterContext.AddFilterData<double>(static_cast<int32>(EFilterField::Duration), 0.0f);
+		FilterContext.AddFilterData<FString>(static_cast<int32>(EFilterField::TrackName), this->GetName());
+
+		TSharedPtr<const TraceServices::IAnalysisSession> Session = FInsightsManager::Get()->GetSession();
+
+		if (Session.IsValid())
+		{
+			const TraceServices::IRegionProvider& RegionProvider = TraceServices::ReadRegionProvider(*Session);
+			TraceServices::FProviderReadScopeLock RegionProviderScopedLock(RegionProvider);
+			const FTimingTrackViewport& Viewport = Context.GetViewport();
+
+			int32 CurDepth = 0;
+			RegionProvider.EnumerateLanes([this, Viewport, &CurDepth, &Builder, &FilterContext](const TraceServices::FRegionLane& Lane, const int32 Depth)
+				{
+					bool RegionHadEvents = false;
+					Lane.EnumerateRegions(Viewport.GetStartTime(), Viewport.GetEndTime(), [&Builder, &RegionHadEvents, &CurDepth, &FilterContext, this](const TraceServices::FTimeRegion& Region) -> bool
+						{
+							FilterContext.SetFilterData<double>(static_cast<int32>(EFilterField::StartTime), Region.BeginTime);
+							FilterContext.SetFilterData<double>(static_cast<int32>(EFilterField::EndTime), Region.EndTime);
+							FilterContext.SetFilterData<double>(static_cast<int32>(EFilterField::Duration), Region.EndTime - Region.BeginTime);
+
+							RegionHadEvents = true;
+							if (FilterConfigurator->ApplyFilters(FilterContext))
+							{
+								Builder.AddEvent(Region.BeginTime, Region.EndTime, CurDepth, Region.Text);
+							}
+							
+							return true;
+						});
+
+					if (RegionHadEvents) CurDepth++;
+				});
+		}
+	}
+}
+
+//////////////////////////////////////////////////////////////////////////
+
 const TSharedPtr<const ITimingEvent> FTimingRegionsTrack::SearchEvent(
 	const FTimingEventSearchParameters& InSearchParameters) const
 {
@@ -217,15 +340,35 @@ const TSharedPtr<const ITimingEvent> FTimingRegionsTrack::SearchEvent(
 
 	FindRegionEvent(InSearchParameters, [this, &FoundEvent](double InFoundStartTime, double InFoundEndTime, uint32 InFoundDepth, const TraceServices::FTimeRegion& InEvent)
 	{
-		FoundEvent = MakeShared<FTimingEvent>(SharedThis(this), InFoundStartTime, InFoundEndTime, InFoundDepth);
+		FoundEvent = MakeShared<FTimingEvent>(SharedThis(this), InFoundStartTime, InFoundEndTime, InFoundDepth, reinterpret_cast<uint64>(InEvent.Text));
 	});
 
 	return FoundEvent;
 }
 
+//////////////////////////////////////////////////////////////////////////
+
 bool FTimingRegionsTrack::FindRegionEvent(const FTimingEventSearchParameters& InParameters,
 	TFunctionRef<void(double, double, uint32, const TraceServices::FTimeRegion&)> InFoundPredicate) const
 {
+	// If the query start time is larger than the end of the session return false.
+	{
+		TSharedPtr<const TraceServices::IAnalysisSession> Session = FInsightsManager::Get()->GetSession();
+		TraceServices::FAnalysisSessionReadScope SessionReadScope(*Session.Get());
+		if (Session.IsValid() && InParameters.StartTime > Session->GetDurationSeconds())
+		{
+			return false;
+		}
+	}
+
+	FFilterContext FilterConfiguratorContext;
+	FilterConfiguratorContext.SetReturnValueForUnsetFilters(false);
+	FilterConfiguratorContext.AddFilterData<double>(static_cast<int32>(EFilterField::StartTime), 0.0f);
+	FilterConfiguratorContext.AddFilterData<double>(static_cast<int32>(EFilterField::EndTime), 0.0f);
+	FilterConfiguratorContext.AddFilterData<double>(static_cast<int32>(EFilterField::Duration), 0.0f);
+	FilterConfiguratorContext.AddFilterData<FString>(static_cast<int32>(EFilterField::TrackName), this->GetName());
+	FilterConfiguratorContext.AddFilterData<int64>(static_cast<int32>(EFilterField::RegionName), 0);
+
 	return TTimingEventSearch<TraceServices::FTimeRegion>::Search(
 	InParameters,
 
@@ -234,7 +377,7 @@ bool FTimingRegionsTrack::FindRegionEvent(const FTimingEventSearchParameters& In
 	{
 		TSharedPtr<const TraceServices::IAnalysisSession> Session = FInsightsManager::Get()->GetSession();
 		const TraceServices::IRegionProvider& RegionProvider = TraceServices::ReadRegionProvider(*Session);
-			TraceServices::FProviderReadScopeLock RegionProviderScopedLock(RegionProvider);
+		TraceServices::FProviderReadScopeLock RegionProviderScopedLock(RegionProvider);
 
 		RegionProvider.EnumerateRegions(InContext.GetParameters().StartTime, InContext.GetParameters().EndTime, [&InContext](const TraceServices::FTimeRegion& Region)
 		{
@@ -248,8 +391,51 @@ bool FTimingRegionsTrack::FindRegionEvent(const FTimingEventSearchParameters& In
 			return true;
 		});
 	},
-	// Found!
-	InFoundPredicate);
+	[&FilterConfiguratorContext, &InParameters](double EventStartTime, double EventEndTime, uint32 EventDepth, const TraceServices::FTimeRegion& Region)
+	{
+		if (!InParameters.FilterExecutor.IsValid())
+		{
+			return true;
+		}
+
+		TSharedPtr<const TraceServices::IAnalysisSession> Session = FInsightsManager::Get()->GetSession();
+		if (Session.IsValid())
+		{
+			FilterConfiguratorContext.SetFilterData<double>(static_cast<int32>(EFilterField::StartTime), EventStartTime);
+			FilterConfiguratorContext.SetFilterData<double>(static_cast<int32>(EFilterField::EndTime), EventEndTime);
+			FilterConfiguratorContext.SetFilterData<double>(static_cast<int32>(EFilterField::Duration), EventEndTime - EventStartTime);
+			FilterConfiguratorContext.SetFilterData<int64>(static_cast<int32>(EFilterField::RegionName), reinterpret_cast<int64>(Region.Text));
+			return InParameters.FilterExecutor->ApplyFilters(FilterConfiguratorContext);
+		}
+
+		return false;
+	},
+	[&InFoundPredicate](double InFoundStartTime, double InFoundEndTime, uint32 InFoundDepth, const TraceServices::FTimeRegion& InEvent)
+	{
+		InFoundPredicate(InFoundStartTime, InFoundEndTime, InFoundDepth, InEvent);
+	},
+
+	TTimingEventSearch<TraceServices::FTimeRegion>::NoMatch);
+}
+
+//////////////////////////////////////////////////////////////////////////
+
+void FTimingRegionsTrack::SetFilterConfigurator(TSharedPtr<Insights::FFilterConfigurator> InFilterConfigurator)
+{
+	if (FilterConfigurator != InFilterConfigurator)
+	{
+		FilterConfigurator = InFilterConfigurator;
+		SetDirtyFlag();
+	}
 }
 
+//////////////////////////////////////////////////////////////////////////
+
+bool FTimingRegionsTrack::HasCustomFilter() const
+{
+	return FilterConfigurator.IsValid() && !FilterConfigurator->IsEmpty();
+}
+
+} // namespace Insights
+
 #undef LOCTEXT_NAMESPACE
\ No newline at end of file
diff --git a/Engine/Source/Developer/TraceInsights/Private/Insights/ViewModels/RegionsTimingTrack.h b/Engine/Source/Developer/TraceInsights/Private/Insights/ViewModels/RegionsTimingTrack.h
--- a/Engine/Source/Developer/TraceInsights/Private/Insights/ViewModels/RegionsTimingTrack.h
+++ b/Engine/Source/Developer/TraceInsights/Private/Insights/ViewModels/RegionsTimingTrack.h
@@ -4,10 +4,12 @@
 #include "Insights/ViewModels/TimingEventsTrack.h"
 #include "TraceServices/Model/Regions.h"
 
+////////////////////////////////////////////////////////////////////////////////////////////////////
 
-class FTimingRegionsTrack;
+namespace Insights
+{
 
-////////////////////////////////////////////////////////////////////////////////////////////////////
+class FTimingRegionsTrack;
 
 class FTimingRegionsViewCommands : public TCommands<FTimingRegionsViewCommands>
 {
@@ -42,7 +44,7 @@ public:
 private:
 	STimingView* TimingView;
 	TSharedPtr<FTimingRegionsTrack> TimingRegionsTrack;
-	
+
 	bool bShowHideRegionsTrack = true;
 };
 
@@ -58,9 +60,16 @@ public:
 	virtual void InitTooltip(FTooltipDrawState& InOutTooltip, const ITimingEvent& InTooltipEvent) const override;
 	virtual const TSharedPtr<const ITimingEvent> SearchEvent(const FTimingEventSearchParameters& InSearchParameters) const override;
 	virtual void BuildDrawState(ITimingEventsTrackDrawStateBuilder& Builder, const ITimingTrackUpdateContext& Context) override;
+	virtual void BuildFilteredDrawState(ITimingEventsTrackDrawStateBuilder& Builder, const ITimingTrackUpdateContext& Context) override;
+	virtual void SetFilterConfigurator(TSharedPtr<Insights::FFilterConfigurator> InFilterConfigurator) override;
+	virtual bool HasCustomFilter() const override;
+
 protected:
 	bool FindRegionEvent(const FTimingEventSearchParameters& InParameters, TFunctionRef<void(double, double, uint32, const TraceServices::FTimeRegion&)> InFoundPredicate) const;
+
 private:
-	
+	TSharedPtr<Insights::FFilterConfigurator> FilterConfigurator;
 	FTimingRegionsSharedState& SharedState;
 };
+
+} // namespace Insights
diff --git a/Engine/Source/Developer/TraceInsights/Private/Insights/ViewModels/ThreadTimingTrack.h b/Engine/Source/Developer/TraceInsights/Private/Insights/ViewModels/ThreadTimingTrack.h
--- a/Engine/Source/Developer/TraceInsights/Private/Insights/ViewModels/ThreadTimingTrack.h
+++ b/Engine/Source/Developer/TraceInsights/Private/Insights/ViewModels/ThreadTimingTrack.h
@@ -130,6 +130,7 @@ enum class EFilterField : int32
 	TimerId = 4,
 	TimerName = 5,
 	CoreEventName = 6,
+	RegionName = 7,
 };
 
 ////////////////////////////////////////////////////////////////////////////////////////////////////
diff --git a/Engine/Source/Developer/TraceInsights/Private/Insights/Widgets/STimingView.cpp b/Engine/Source/Developer/TraceInsights/Private/Insights/Widgets/STimingView.cpp
--- a/Engine/Source/Developer/TraceInsights/Private/Insights/Widgets/STimingView.cpp
+++ b/Engine/Source/Developer/TraceInsights/Private/Insights/Widgets/STimingView.cpp
@@ -101,7 +101,7 @@ STimingView::STimingView()
 	, ThreadTimingSharedState(MakeShared<FThreadTimingSharedState>(this))
 	, LoadingSharedState(MakeShared<FLoadingSharedState>(this))
 	, FileActivitySharedState(MakeShared<FFileActivitySharedState>(this))
-	, TimingRegionsSharedState(MakeShared<FTimingRegionsSharedState>(this))
+	, TimingRegionsSharedState(MakeShared<Insights::FTimingRegionsSharedState>(this))
 	, TimeRulerTrack(MakeShared<FTimeRulerTrack>())
 	, DefaultTimeMarker(MakeShared<Insights::FTimeMarker>())
 	, MarkersTrack(MakeShared<FMarkersTimingTrack>())
diff --git a/Engine/Source/Developer/TraceInsights/Private/Insights/Widgets/STimingView.h b/Engine/Source/Developer/TraceInsights/Private/Insights/Widgets/STimingView.h
--- a/Engine/Source/Developer/TraceInsights/Private/Insights/Widgets/STimingView.h
+++ b/Engine/Source/Developer/TraceInsights/Private/Insights/Widgets/STimingView.h
@@ -24,7 +24,6 @@
 #include "Insights/ViewModels/TimingViewDrawHelper.h"
 #include "Insights/ViewModels/TooltipDrawState.h"
 
-class FTimingRegionsSharedState;
 class FFileActivitySharedState;
 class FFrameSharedState;
 class FLoadingSharedState;
@@ -44,6 +43,7 @@ namespace Insights
 {
 	class ITimingViewExtender;
 	class FTimeMarker;
+	class FTimingRegionsSharedState;
 	class FQuickFind;
 	class SQuickFind;
 	enum class ETimingEventsColoringMode : uint32;
@@ -556,7 +556,7 @@ protected:
 	TSharedPtr<FFileActivitySharedState> FileActivitySharedState;
 
 	// Shared state for Regions tracks
-	TSharedPtr<FTimingRegionsSharedState> TimingRegionsSharedState;
+	TSharedPtr<Insights::FTimingRegionsSharedState> TimingRegionsSharedState;
 	
 	////////////////////////////////////////////////////////////
 
