[Insights] Fixed the Session main tab to show correctly the URI for traces located in "additional watch directories" (instead of the "trace store directory").

diff --git a/Engine/Source/Developer/TraceInsights/Private/Insights/InsightsManager.cpp b/Engine/Source/Developer/TraceInsights/Private/Insights/InsightsManager.cpp
--- a/Engine/Source/Developer/TraceInsights/Private/Insights/InsightsManager.cpp
+++ b/Engine/Source/Developer/TraceInsights/Private/Insights/InsightsManager.cpp
@@ -977,18 +977,24 @@ void FInsightsManager::LoadTrace(uint32 InTraceId, bool InAutoQuit)
 		return;
 	}
 
-	FString TraceName(StoreClient->GetStatus()->GetStoreDir());
+	FString TraceName;
 	const UE::Trace::FStoreClient::FTraceInfo* TraceInfo = StoreClient->GetTraceInfoById(InTraceId);
 	if (TraceInfo != nullptr)
 	{
 		const FUtf8StringView Utf8NameView = TraceInfo->GetName();
 		FString Name(Utf8NameView);
-		if (!Name.EndsWith(TEXT(".utrace")))
+		FUtf8StringView Uri = TraceInfo->GetUri();
+		if (Uri.Len() > 0)
 		{
-			Name += TEXT(".utrace");
+			TraceName = FString(Uri);
+		}
+		else
+		{
+			// Fallback for older versions of UTS which didn't write uri
+			FString StoreDirectory(StoreClient->GetStatus()->GetStoreDir());
+			TraceName = FPaths::SetExtension(FPaths::Combine(StoreDirectory, Name), TEXT(".utrace"));
+			FPaths::MakePlatformFilename(TraceName);
 		}
-		TraceName = FPaths::Combine(TraceName, Name);
-		FPaths::NormalizeFilename(TraceName);
 	}
 
 	StoreClientLock.Unlock();
