Fix Insights Enumerate.FastEnumerate unit test.
#rnx

diff --git a/Engine/Source/Developer/TraceInsights/Private/Insights/Tests/TimingProfilerTests.cpp b/Engine/Source/Developer/TraceInsights/Private/Insights/Tests/TimingProfilerTests.cpp
--- a/Engine/Source/Developer/TraceInsights/Private/Insights/Tests/TimingProfilerTests.cpp
+++ b/Engine/Source/Developer/TraceInsights/Private/Insights/Tests/TimingProfilerTests.cpp
@@ -288,6 +288,7 @@ void FTimingProfilerTests::RunEnumerateBenchmark(const FEnumerateTestParams& InP
 					Timeline.EnumerateEvents(Time, Time + InParams.Interval,
 						[&OutCheckValues](double EventStartTime, double EventEndTime, uint32 EventDepth, const TraceServices::FTimingProfilerEvent& Event)
 						{
+							EventEndTime = FMath::Min(EventEndTime, OutCheckValues.SessionDuration);
 							OutCheckValues.TotalEventDuration += EventEndTime - EventStartTime;
 							++OutCheckValues.EventCount;
 							OutCheckValues.SumDepth += EventDepth;
@@ -337,7 +338,7 @@ void FTimingProfilerTests::RunEnumerateAsyncBenchmark(const FEnumerateTestParams
 		TArray<FCheckValues> TaskCheckValues;
 
 		TimingProfilerProvider.ReadTimeline(TimelineIndex,
-			[&OutCheckValues, &InParams, TimeIncrement, &TaskCheckValues](const TraceServices::ITimingProfilerProvider::Timeline& Timeline)
+			[SessionTime, &InParams, TimeIncrement, &TaskCheckValues](const TraceServices::ITimingProfilerProvider::Timeline& Timeline)
 			{
 				double Time = 0.0;
 				for (int32 Index = 0; Index < InParams.NumEnumerations; ++Index)
@@ -351,8 +352,9 @@ void FTimingProfilerTests::RunEnumerateAsyncBenchmark(const FEnumerateTestParams
 					{
 						TaskCheckValues.AddDefaulted(NumTasks);
 					};
-					Params.Callback = [&TaskCheckValues](double EventStartTime, double EventEndTime, uint32 EventDepth, const TraceServices::FTimingProfilerEvent& Event, uint32 TaskIndex)
+					Params.Callback = [&TaskCheckValues, SessionTime](double EventStartTime, double EventEndTime, uint32 EventDepth, const TraceServices::FTimingProfilerEvent& Event, uint32 TaskIndex)
 					{
+						EventEndTime = FMath::Min(EventEndTime, SessionTime);
 						TaskCheckValues[TaskIndex].TotalEventDuration += EventEndTime - EventStartTime;
 						++TaskCheckValues[TaskIndex].EventCount;
 						TaskCheckValues[TaskIndex].SumDepth += EventDepth;
@@ -422,6 +424,7 @@ void FTimingProfilerTests::RunEnumerateAllTracksBenchmark(const FEnumerateTestPa
 							Timeline.EnumerateEvents(Time, Time + InParams.Interval,
 								[&OutCheckValues](double EventStartTime, double EventEndTime, uint32 EventDepth, const TraceServices::FTimingProfilerEvent& Event)
 								{
+									EventEndTime = FMath::Min(EventEndTime, OutCheckValues.SessionDuration);
 									OutCheckValues.TotalEventDuration += EventEndTime - EventStartTime;
 									++OutCheckValues.EventCount;
 									OutCheckValues.SumDepth += EventDepth;
@@ -492,6 +495,7 @@ void FTimingProfilerTests::RunEnumerateAsyncAllTracksBenchmark(const FEnumerateT
 							};
 							Params.Callback = [&OutCheckValues, &TaskCheckValues](double EventStartTime, double EventEndTime, uint32 EventDepth, const TraceServices::FTimingProfilerEvent& Event, uint32 TaskIndex)
 							{
+								EventEndTime = FMath::Min(EventEndTime, OutCheckValues.SessionDuration);
 								TaskCheckValues[TaskIndex].TotalEventDuration += EventEndTime - EventStartTime;
 								++TaskCheckValues[TaskIndex].EventCount;
 								TaskCheckValues[TaskIndex].SumDepth += EventDepth;
