[Insights] Fix context menu not opening for the time ruler track.

diff --git a/Engine/Source/Developer/TraceInsights/Private/Insights/Widgets/STimingView.cpp b/Engine/Source/Developer/TraceInsights/Private/Insights/Widgets/STimingView.cpp
--- a/Engine/Source/Developer/TraceInsights/Private/Insights/Widgets/STimingView.cpp
+++ b/Engine/Source/Developer/TraceInsights/Private/Insights/Widgets/STimingView.cpp
@@ -2207,7 +2207,6 @@ FReply STimingView::OnMouseButtonUp(const FGeometry& MyGeometry, const FPointerE
 			{
 				RaiseSelectionChanged();
 				bIsSelecting = false;
-				bIsMouseClick = false;
 			}
 			else if (TimeRulerTrack->IsScrubbing())
 			{
