[Insights] Fix timing region event copy to clipboard.
+ Fix Select Time Range of Event
+ Add icon

diff --git a/Engine/Source/Developer/TraceInsights/Private/Insights/InsightsStyle.cpp b/Engine/Source/Developer/TraceInsights/Private/Insights/InsightsStyle.cpp
--- a/Engine/Source/Developer/TraceInsights/Private/Insights/InsightsStyle.cpp
+++ b/Engine/Source/Developer/TraceInsights/Private/Insights/InsightsStyle.cpp
@@ -325,6 +325,7 @@ void FInsightsStyle::FStyle::Initialize()
 	Set("Icons.Delete", new CORE_IMAGE_BRUSH_SVG("Starship/Common/Delete", Icon16x16));
 
 	Set("Icons.ImportTable", new CORE_IMAGE_BRUSH_SVG("Starship/Common/Import", Icon16x16));
+	Set("Icons.SelectEventRange", new EDITOR_IMAGE_BRUSH_SVG("Starship/Common/SelectInViewport", Icon16x16));
 
 	//////////////////////////////////////////////////
 
diff --git a/Engine/Source/Developer/TraceInsights/Private/Insights/ViewModels/RegionsTimingTrack.cpp b/Engine/Source/Developer/TraceInsights/Private/Insights/ViewModels/RegionsTimingTrack.cpp
--- a/Engine/Source/Developer/TraceInsights/Private/Insights/ViewModels/RegionsTimingTrack.cpp
+++ b/Engine/Source/Developer/TraceInsights/Private/Insights/ViewModels/RegionsTimingTrack.cpp
@@ -2,6 +2,8 @@
 
 #include "RegionsTimingTrack.h"
 
+#include "HAL/PlatformApplicationMisc.h"
+
 #include "TimingViewDrawHelper.h"
 #include "Common/ProviderLock.h"
 #include "Insights/InsightsManager.h"
@@ -436,6 +438,24 @@ bool FTimingRegionsTrack::HasCustomFilter() const
 	return FilterConfigurator.IsValid() && !FilterConfigurator->IsEmpty();
 }
 
+//////////////////////////////////////////////////////////////////////////
+
+void FTimingRegionsTrack::OnClipboardCopyEvent(const ITimingEvent& InSelectedEvent) const
+{
+	if (InSelectedEvent.CheckTrack(this) && InSelectedEvent.Is<FTimingEvent>())
+	{
+		const FTimingEvent& TrackEvent = InSelectedEvent.As<FTimingEvent>();
+
+		// The pointer should be safe to access because it is stored in the Session string store.
+		FString EventName(reinterpret_cast<const TCHAR*>(TrackEvent.GetType()));
+		FTimingEventsTrackDrawStateBuilder::AppendDurationToEventName(EventName, TrackEvent.GetDuration());
+
+		FPlatformApplicationMisc::ClipboardCopy(*EventName);
+	}
+}
+
+//////////////////////////////////////////////////////////////////////////
+
 } // namespace Insights
 
 #undef LOCTEXT_NAMESPACE
\ No newline at end of file
diff --git a/Engine/Source/Developer/TraceInsights/Private/Insights/ViewModels/RegionsTimingTrack.h b/Engine/Source/Developer/TraceInsights/Private/Insights/ViewModels/RegionsTimingTrack.h
--- a/Engine/Source/Developer/TraceInsights/Private/Insights/ViewModels/RegionsTimingTrack.h
+++ b/Engine/Source/Developer/TraceInsights/Private/Insights/ViewModels/RegionsTimingTrack.h
@@ -63,6 +63,7 @@ public:
 	virtual void BuildFilteredDrawState(ITimingEventsTrackDrawStateBuilder& Builder, const ITimingTrackUpdateContext& Context) override;
 	virtual void SetFilterConfigurator(TSharedPtr<Insights::FFilterConfigurator> InFilterConfigurator) override;
 	virtual bool HasCustomFilter() const override;
+	virtual void OnClipboardCopyEvent(const ITimingEvent& InSelectedEvent) const override;
 
 protected:
 	bool FindRegionEvent(const FTimingEventSearchParameters& InParameters, TFunctionRef<void(double, double, uint32, const TraceServices::FTimeRegion&)> InFoundPredicate) const;
diff --git a/Engine/Source/Developer/TraceInsights/Private/Insights/Widgets/STimingView.cpp b/Engine/Source/Developer/TraceInsights/Private/Insights/Widgets/STimingView.cpp
--- a/Engine/Source/Developer/TraceInsights/Private/Insights/Widgets/STimingView.cpp
+++ b/Engine/Source/Developer/TraceInsights/Private/Insights/Widgets/STimingView.cpp
@@ -2914,12 +2914,12 @@ void STimingView::ShowContextMenu(const FPointerEvent& MouseEvent)
 		if (HoveredEvent)
 		{
 			double RangeStart = HoveredEvent->GetStartTime();
-			double RangeDuration = HoveredEvent->GetDuration();
+			double RangeDuration = Viewport.RestrictEndTime(HoveredEvent->GetEndTime()) - RangeStart;
 
 			MenuBuilder.AddMenuEntry(
 				LOCTEXT("ContextMenu_SelectEventTimeRange", "Select Time Range of Event"),
 				FText(),
-				FSlateIcon(),
+				FSlateIcon(FInsightsStyle::GetStyleSetName(), "Icons.SelectEventRange"),
 				FUIAction(FExecuteAction::CreateLambda([this, RangeStart, RangeDuration]()
 					{
 						SelectTimeInterval(RangeStart, RangeDuration);
