[Insights] Memory Insights: Fixed crash when grouping by Heap in the Allocs Table. Also added "Allocs" and "Heaps" sub-groups for the root heaps. Non root heaps only have "heap" items.

diff --git a/Engine/Source/Developer/TraceInsights/Private/Insights/MemoryProfiler/ViewModels/MemAllocGroupingByHeap.cpp b/Engine/Source/Developer/TraceInsights/Private/Insights/MemoryProfiler/ViewModels/MemAllocGroupingByHeap.cpp
--- a/Engine/Source/Developer/TraceInsights/Private/Insights/MemoryProfiler/ViewModels/MemAllocGroupingByHeap.cpp
+++ b/Engine/Source/Developer/TraceInsights/Private/Insights/MemoryProfiler/ViewModels/MemAllocGroupingByHeap.cpp
@@ -40,14 +40,39 @@ FMemAllocGroupingByHeap::~FMemAllocGroupingByHeap()
 
 FTableTreeNodePtr MakeGroupNodeHierarchy(const TraceServices::IAllocationsProvider::FHeapSpec& Spec, TWeakPtr<FTable>& InParentTable, TArray<FTableTreeNodePtr>& NodeTable)
 {
-	auto Node = MakeShared<FTableTreeNode>(FName(Spec.Name), InParentTable);
-	NodeTable[Spec.Id] = Node;
+	const FSlateBrush* IconBrush = FBaseTreeNode::GetDefaultIcon(true);
+	const FLinearColor Color(1.0f, 0.7f, 0.3f, 1.0f);
+
+	auto HeapGroup = MakeShared<FCustomTableTreeNode>(FName(Spec.Name), InParentTable, IconBrush, Color);
+	if (int32(Spec.Id) >= NodeTable.Num())
+	{
+		NodeTable.AddDefaulted(int32(Spec.Id) - NodeTable.Num() + 1);
+	}
+
+	FTableTreeNodePtr HeapsSubGroup;
+	if (uint32(Spec.Flags) & uint32(EMemoryTraceHeapFlags::Root))
+	{
+		HeapsSubGroup = MakeShared<FTableTreeNode>(FName(TEXT("Heaps")), InParentTable);
+		HeapGroup->AddChildAndSetParent(HeapsSubGroup);
+
+		auto AllocsSubGroup = MakeShared<FTableTreeNode>(FName(TEXT("Allocs")), InParentTable);
+		HeapGroup->AddChildAndSetParent(AllocsSubGroup);
+
+		NodeTable[int32(Spec.Id)] = AllocsSubGroup;
+	}
+	else
+	{
+		HeapsSubGroup = HeapGroup;
+		NodeTable[int32(Spec.Id)] = HeapGroup;
+	}
+
 	for (TraceServices::IAllocationsProvider::FHeapSpec* ChildSpec : Spec.Children)
 	{
 		auto ChildNode = MakeGroupNodeHierarchy(*ChildSpec, InParentTable, NodeTable);
-		Node->AddChildAndSetParent(ChildNode);
+		HeapsSubGroup->AddChildAndSetParent(ChildNode);
 	}
-	return Node;
+
+	return HeapGroup;
 }
 
 ////////////////////////////////////////////////////////////////////////////////////////////////////
@@ -59,13 +84,13 @@ void FMemAllocGroupingByHeap::GroupNodes(const TArray<FTableTreeNodePtr>& Nodes,
 
 	// Build heap hierarchy
 	TArray<FTableTreeNodePtr> HeapNodes;
-	HeapNodes.AddZeroed(256);
 
 	{
 		TraceServices::FProviderReadScopeLock _(AllocProvider);
 		AllocProvider.EnumerateRootHeaps([&](HeapId Id, const TraceServices::IAllocationsProvider::FHeapSpec& Spec)
 		{
-			ParentGroup.AddChildAndSetParent(MakeGroupNodeHierarchy(Spec, InParentTable, HeapNodes));
+			FTableTreeNodePtr RootHeapGroup = MakeGroupNodeHierarchy(Spec, InParentTable, HeapNodes);
+			ParentGroup.AddChildAndSetParent(RootHeapGroup);
 		});
 	}
 
@@ -91,9 +116,18 @@ void FMemAllocGroupingByHeap::GroupNodes(const TArray<FTableTreeNodePtr>& Nodes,
 			//TODO: Calculating the real HeapId when the allocation is first added to the provider was too expensive, so deferring that operation to here could make sense.
 			//const HeapId Heap = AllocProvider.GetParentBlock(Alloc->GetAddress());
 			const uint8 HeapId = static_cast<uint8>(Alloc->GetRootHeap());
-			if (FTableTreeNodePtr GroupPtr = HeapNodes[HeapId])
+			FTableTreeNodePtr GroupPtr = HeapNodes[HeapId];
+			if (ensure(GroupPtr.IsValid()))
 			{
-				GroupPtr->AddChildAndSetParent(NodePtr);
+				constexpr uint32 MaxRootHeaps = 16; // see TraceServices\Private\Model\AllocationsProvider.h
+				if (Alloc->IsHeap() && HeapId < MaxRootHeaps)
+				{
+					GroupPtr->GetParent()->GetChildren()[0]->AddChildAndSetParent(NodePtr);
+				}
+				else
+				{
+					GroupPtr->AddChildAndSetParent(NodePtr);
+				}
 			}
 		}
 	}
