[Insights] Changed naming of some automation tests to match code style.

diff --git a/Engine/Source/Developer/TraceInsights/Private/Insights/Tests/MiscTests.cpp b/Engine/Source/Developer/TraceInsights/Private/Insights/Tests/MiscTests.cpp
--- a/Engine/Source/Developer/TraceInsights/Private/Insights/Tests/MiscTests.cpp
+++ b/Engine/Source/Developer/TraceInsights/Private/Insights/Tests/MiscTests.cpp
@@ -9,7 +9,7 @@
 #include "Insights/InsightsCommands.h"
 #include "Insights/Tests/InsightsTestUtils.h"
 
-IMPLEMENT_SIMPLE_AUTOMATION_TEST(FPagedArrayFilteringTest, "Insights.FPagedArrayFilteringTest", EAutomationTestFlags::ApplicationContextMask | EAutomationTestFlags::SmokeFilter)
+IMPLEMENT_SIMPLE_AUTOMATION_TEST(FPagedArrayFilteringTest, "Insights.PagedArrayFiltering", EAutomationTestFlags::ApplicationContextMask | EAutomationTestFlags::SmokeFilter)
 
 bool operator!=(const FInt32Interval& lhs, const FInt32Interval& rhs)
 {
diff --git a/Engine/Source/Developer/TraceInsights/Private/Insights/Tests/UITests.h b/Engine/Source/Developer/TraceInsights/Private/Insights/Tests/UITests.h
--- a/Engine/Source/Developer/TraceInsights/Private/Insights/Tests/UITests.h
+++ b/Engine/Source/Developer/TraceInsights/Private/Insights/Tests/UITests.h
@@ -18,10 +18,10 @@ FAutomationDriverPtr Driver;
 TSharedPtr<SWindow> AutomationWindow;
 END_DEFINE_SPEC(FAutomationDriverUnrealInsightsSessionBrowserTest)
 IMPLEMENT_SIMPLE_AUTOMATION_TEST(FMemoryInsightsUploadLLMXMLReportsTraceTest, "Insights.MemoryInsights.UploadMemoryInsightsLLMXMLReportsTrace", EAutomationTestFlags::ApplicationContextMask | EAutomationTestFlags::ProductFilter)
-IMPLEMENT_SIMPLE_AUTOMATION_TEST(FUnrealTraceServerStartingTest, "Insights.SessionBrowser.UnrealTraceServer.Startng", EAutomationTestFlags::ApplicationContextMask | EAutomationTestFlags::ProductFilter)
+IMPLEMENT_SIMPLE_AUTOMATION_TEST(FUnrealTraceServerStartingTest, "Insights.SessionBrowser.UnrealTraceServer.Starting", EAutomationTestFlags::ApplicationContextMask | EAutomationTestFlags::ProductFilter)
 
 #endif
 
-IMPLEMENT_SIMPLE_AUTOMATION_TEST(FMemoryFilterValueConverterTest, "Insights.FMemoryFilterValueConverterTest", EAutomationTestFlags::ApplicationContextMask | EAutomationTestFlags::ProductFilter)
-IMPLEMENT_SIMPLE_AUTOMATION_TEST(FTimeFilterValueConverterTest, "Insights.FTimeFilterValueConverterTest", EAutomationTestFlags::ApplicationContextMask | EAutomationTestFlags::ProductFilter)
+IMPLEMENT_SIMPLE_AUTOMATION_TEST(FMemoryFilterValueConverterTest, "Insights.MemoryFilterValueConverter", EAutomationTestFlags::ApplicationContextMask | EAutomationTestFlags::ProductFilter)
+IMPLEMENT_SIMPLE_AUTOMATION_TEST(FTimeFilterValueConverterTest, "Insights.TimeFilterValueConverter", EAutomationTestFlags::ApplicationContextMask | EAutomationTestFlags::ProductFilter)
 
