[Insights] Automated T223693548, T223693546 and T223693547 from "Unreal Insights - Memory Insights - Allocations Query Table" suite

diff --git a/Engine/Source/Developer/TraceInsights/Private/Insights/Table/Widgets/STableTreeView.cpp b/Engine/Source/Developer/TraceInsights/Private/Insights/Table/Widgets/STableTreeView.cpp
--- a/Engine/Source/Developer/TraceInsights/Private/Insights/Table/Widgets/STableTreeView.cpp
+++ b/Engine/Source/Developer/TraceInsights/Private/Insights/Table/Widgets/STableTreeView.cpp
@@ -165,6 +165,15 @@ STableTreeView::~STableTreeView()
 
 ////////////////////////////////////////////////////////////////////////////////////////////////////
 
+void STableTreeView::SetCurrentGroupings(TArray<TSharedPtr<FTreeNodeGrouping>>& InCurrentGroupings)
+{
+	PreChangeGroupings();
+	CurrentGroupings = InCurrentGroupings;
+	PostChangeGroupings();
+}
+
+////////////////////////////////////////////////////////////////////////////////////////////////////
+
 void STableTreeView::InitCommandList()
 {
 	FTableTreeViewCommands::Register();
diff --git a/Engine/Source/Developer/TraceInsights/Private/Insights/Table/Widgets/STableTreeView.h b/Engine/Source/Developer/TraceInsights/Private/Insights/Table/Widgets/STableTreeView.h
--- a/Engine/Source/Developer/TraceInsights/Private/Insights/Table/Widgets/STableTreeView.h
+++ b/Engine/Source/Developer/TraceInsights/Private/Insights/Table/Widgets/STableTreeView.h
@@ -173,10 +173,19 @@ public:
 
 	////////////////////////////////////////////////////////////////////////////////////////////////////
 
-	/** Set a log listing name to be used for any errors or warnings. Must be preregistered by the caller with the MessageLog module. */
+	/** Sets a log listing name to be used for any errors or warnings. Must be preregistered by the caller with the MessageLog module. */
 	void SetLogListingName(const FName& InLogListingName) { LogListingName = InLogListingName; }
 	const FName& GetLogListingName() { return LogListingName; }
 
+	/** Gets the table row nodes. Each node corresponds to a table row. Index in this array corresponds to RowIndex in source table. */
+	const TArray<FTableTreeNodePtr>& GetTableRowNodes() const { return TableRowNodes; }
+	
+	/** Gets the avaiable grouping. */
+	const TArray<TSharedPtr<FTreeNodeGrouping>>& GetAvailableGroupings() const { return AvailableGroupings; }
+
+	/** Sets the current groupings. */
+	void SetCurrentGroupings(TArray<TSharedPtr<FTreeNodeGrouping>>& InCurrentGroupings);
+
 protected:
 	void InitCommandList();
 
diff --git a/Engine/Source/Developer/TraceInsights/Private/Insights/Tests/UITests.cpp b/Engine/Source/Developer/TraceInsights/Private/Insights/Tests/UITests.cpp
--- a/Engine/Source/Developer/TraceInsights/Private/Insights/Tests/UITests.cpp
+++ b/Engine/Source/Developer/TraceInsights/Private/Insights/Tests/UITests.cpp
@@ -19,6 +19,18 @@
 #include "Insights/Widgets/STimingView.h"
 #include "Insights/Tests/InsightsTestUtils.h"
 #include "Insights/TaskGraphProfiler/ViewModels/TaskTimingTrack.h"
+#include "Insights/ViewModels/TimeRulerTrack.h"
+#include "Insights/MemoryProfiler/Widgets/SMemInvestigationView.h"
+#include "Insights/MemoryProfiler/Widgets/SMemAllocTableTreeView.h"
+#include "Insights/MemoryProfiler/ViewModels/MemAllocGroupingByCallstack.h"
+#include "Insights/MemoryProfiler/ViewModels/MemAllocNode.h"
+#include "Insights/Table/ViewModels/TableTreeNode.h"
+#include "Insights/TaskGraphProfiler/ViewModels/TaskNode.h"
+#include "Insights/Table/ViewModels/TreeNodeGrouping.h"
+#include "Insights/MemoryProfiler/ViewModels/MemAllocGroupingByCallstack.h"
+#include "Insights/MemoryProfiler/ViewModels/MemAllocGroupingByHeap.h"
+#include "Insights/MemoryProfiler/ViewModels/MemAllocGroupingBySize.h"
+#include "Insights/MemoryProfiler/ViewModels/MemAllocGroupingByTag.h"
 
 ////////////////////////////////////////////////////////////////////////////////////////////////////
 
@@ -398,6 +410,200 @@ bool FUnrealTraceServerStartingTest::RunTest(const FString& Parameters)
 	return true;
 }
 
+////////////////////////////////////////////////////////////////////////////////////////////////////
+
+DEFINE_LATENT_AUTOMATION_COMMAND_THREE_PARAMETER(FWaitForRunningQuieryFinishedCommand, TSharedPtr<Insights::SMemAllocTableTreeView>, MemAllocTableTreeView, const float, Timeout, FAutomationTestBase*, Test);
+
+bool FWaitForRunningQuieryFinishedCommand::Update()
+{
+	if (!MemAllocTableTreeView->IsRunning())
+	{
+		return true;
+	}
+
+	if (FPlatformTime::Seconds() - StartTime >= Timeout)
+	{
+		Test->AddError(TEXT("FWaitForRunningQuieryFinishedCommand timed out"));
+		return true;
+	}
+
+	return false;
+}
+
+DEFINE_LATENT_AUTOMATION_COMMAND_TWO_PARAMETER(FChangeGroupingCommand, TSharedPtr<Insights::SMemAllocTableTreeView>, MemAllocTableTreeView, FAutomationTestBase*, Test);
+
+bool FChangeGroupingCommand::Update()
+{
+	TArray<TSharedPtr<Insights::FTreeNodeGrouping>> CurrentGroupings;
+	for (const auto& Grouping : MemAllocTableTreeView->GetAvailableGroupings())
+	{
+		if (Grouping->GetTitleName().ToString().Contains(TEXT("By Callstack")))
+		{
+			CurrentGroupings.Add(Grouping);
+		}
+	}
+
+	Test->TestTrue(TEXT("CurrentGroupings should not be empty"), !CurrentGroupings.IsEmpty());
+	MemAllocTableTreeView->SetCurrentGroupings(CurrentGroupings);
+
+	return true;
+}
+
+DEFINE_LATENT_AUTOMATION_COMMAND_FOUR_PARAMETER(FVerifyHierarchyCallStackCommand, TSharedPtr<Insights::SMemAllocTableTreeView>, MemAllocTableTreeView, FInsightsTestUtils, InsightsTestUtils, const float, Timeout, FAutomationTestBase*, Test);
+
+bool FVerifyHierarchyCallStackCommand::Update()
+{
+	if (!MemAllocTableTreeView->IsRunningAsyncUpdate())
+	{
+		for (const TSharedPtr<Insights::FTableTreeNode>& Node : MemAllocTableTreeView->GetTableRowNodes())
+		{
+			const Insights::FMemAllocNode& MemAllocNode = static_cast<const Insights::FMemAllocNode&>(*Node);
+			const Insights::FMemoryAlloc Alloc = MemAllocNode.GetMemAllocChecked();
+			if (!(!Alloc.GetCallstack() || Alloc.GetCallstack()->Num() == 0 || (Alloc.GetCallstack()->Num() != 0 && Alloc.GetCallstack()->Num() <= 256)))
+			{
+				Test->AddError(TEXT("Resolved callstack should be valid"));
+			}
+		}
+		return true;
+	}
+
+	if (FPlatformTime::Seconds() - StartTime >= Timeout)
+	{
+		Test->AddError(TEXT("FVerifyHierarchyCallStackCommand timed out"));
+		return true;
+	}
+
+	return false;
+}
+
+const TMap<TraceServices::IAllocationsProvider::EQueryRule, Insights::SMemAllocTableTreeView::FQueryParams> AllocsTimeMarkerStandaloneGameGetterMap
+{
+	{TraceServices::IAllocationsProvider::EQueryRule::aAf, {nullptr, {5.0, 0.0, 0.0, 0.0}}},
+	{TraceServices::IAllocationsProvider::EQueryRule::afA, {nullptr, {10.0, 0.0, 0.0, 0.0}}},
+	{TraceServices::IAllocationsProvider::EQueryRule::Aaf, {nullptr, {10.0, 0.0, 0.0, 0.0}}},
+	{TraceServices::IAllocationsProvider::EQueryRule::aAfB, {nullptr, {50.0, 51.0, 0.0, 0.0}}},
+	{TraceServices::IAllocationsProvider::EQueryRule::AaBf, {nullptr, {50.0, 51.0, 0.0, 0.0}}},
+	{TraceServices::IAllocationsProvider::EQueryRule::aAfaBf, {nullptr, {50.0, 51.0, 0.0, 0.0}}},
+	{TraceServices::IAllocationsProvider::EQueryRule::AfB, {nullptr, {50.0, 51.0, 0.0, 0.0}}},
+	{TraceServices::IAllocationsProvider::EQueryRule::AaBCf, {nullptr, {50.0, 51.0, 52.0, 0.0}}},
+	{TraceServices::IAllocationsProvider::EQueryRule::AaBfC, {nullptr, {50.0, 51.0, 52.0, 0.0}}},
+	{TraceServices::IAllocationsProvider::EQueryRule::aABfC, {nullptr, {50.0, 51.0, 52.0, 0.0}}},
+	{TraceServices::IAllocationsProvider::EQueryRule::AaBCfD, {nullptr, {50.0, 51.0, 52.0, 53.0}}},
+	{TraceServices::IAllocationsProvider::EQueryRule::aABf, {nullptr, {50.0, 51.0, 0.0, 0.0}}},
+	{TraceServices::IAllocationsProvider::EQueryRule::AafB, {nullptr, {50.0, 51.0, 0.0, 0.0}}},
+	{TraceServices::IAllocationsProvider::EQueryRule::AaB, {nullptr, {50.0, 51.0, 0.0, 0.0}}},
+};
+
+const TMap<TraceServices::IAllocationsProvider::EQueryRule, Insights::SMemAllocTableTreeView::FQueryParams> AllocsTimeMarkerEditorPackageGetterMap
+{
+	{TraceServices::IAllocationsProvider::EQueryRule::aAf, {nullptr, {5.0, 0.0, 0.0, 0.0}}},
+	{TraceServices::IAllocationsProvider::EQueryRule::afA, {nullptr, {10.0, 0.0, 0.0, 0.0}}},
+	{TraceServices::IAllocationsProvider::EQueryRule::Aaf, {nullptr, {10.0, 0.0, 0.0, 0.0}}},
+	{TraceServices::IAllocationsProvider::EQueryRule::aAfB, {nullptr, {2.0, 3.0, 0.0, 0.0}}},
+	{TraceServices::IAllocationsProvider::EQueryRule::AaBf, {nullptr, {2.0, 3.0, 0.0, 0.0}}},
+	{TraceServices::IAllocationsProvider::EQueryRule::aAfaBf, {nullptr, {2.0, 3.0, 0.0, 0.0}}},
+	{TraceServices::IAllocationsProvider::EQueryRule::AfB, {nullptr, {2.0, 3.0, 0.0, 0.0}}},
+	{TraceServices::IAllocationsProvider::EQueryRule::AaBCf, {nullptr, {1.0, 2.0, 3.0, 0.0}}},
+	{TraceServices::IAllocationsProvider::EQueryRule::AaBfC, {nullptr, {1.0, 2.0, 3.0, 0.0}}},
+	{TraceServices::IAllocationsProvider::EQueryRule::aABfC, {nullptr, {1.0, 2.0, 3.0, 0.0}}},
+	{TraceServices::IAllocationsProvider::EQueryRule::AaBCfD, {nullptr, {1.0, 2.0, 3.0, 4.0}}},
+	{TraceServices::IAllocationsProvider::EQueryRule::AafB, {nullptr, {2.0, 3.0, 0.0, 0.0}}},
+	{TraceServices::IAllocationsProvider::EQueryRule::AaB, {nullptr, {2.0, 3.0, 0.0, 0.0}}},
+	{TraceServices::IAllocationsProvider::EQueryRule::aABf, {nullptr, {2.0, 3.0, 0.0, 0.0}}},
+};
+
+bool MemoryInsightsAllocationsQueryTableTest(const FString& Parameters, const TMap <TraceServices::IAllocationsProvider::EQueryRule, Insights::SMemAllocTableTreeView::FQueryParams> AllocsTimeMarkerGetterMap, FAutomationTestBase* Test)
+{
+	const float Timeout = 30.0f;
+	FInsightsTestUtils InsightsTestUtils(Test);
+	TSharedPtr<SMemoryProfilerWindow> ProfilerWindow = FMemoryProfilerManager::Get()->GetProfilerWindow();
+	TSharedPtr<FInsightsManager> InsightsManager = FInsightsManager::Get();
+	FMemorySharedState& SharedState = ProfilerWindow->GetSharedState();
+
+	TSharedPtr<Insights::FMemoryRuleSpec> MemoryRule = *Algo::FindByPredicate(SharedState.GetMemoryRules(),
+		[&Parameters](const TSharedPtr<Insights::FMemoryRuleSpec>& Rule)
+		{
+			return Rule->GetShortName().ToString().Contains(Parameters);
+		});
+
+	if (!MemoryRule.IsValid())
+	{
+		Test->AddError(TEXT("MemoryRule should not be null"));
+		return false;
+	}
+
+	TSharedPtr<Insights::SMemAllocTableTreeView> MemAllocTableTreeView = ProfilerWindow->ShowMemAllocTableTreeViewTab();
+
+	Insights::SMemAllocTableTreeView::FQueryParams QueryParams = AllocsTimeMarkerGetterMap.FindChecked(MemoryRule->GetValue());
+	if (MemoryRule->GetValue() == TraceServices::IAllocationsProvider::EQueryRule::Aaf)
+	{
+		QueryParams.TimeMarkers[0] = InsightsManager->GetSessionDuration() - 10.0;
+	}
+	QueryParams.Rule = MemoryRule;
+	MemAllocTableTreeView->SetQueryParams(QueryParams);
+
+	ADD_LATENT_AUTOMATION_COMMAND(FWaitForRunningQuieryFinishedCommand(MemAllocTableTreeView, 120.0f, Test));
+	ADD_LATENT_AUTOMATION_COMMAND(FChangeGroupingCommand(MemAllocTableTreeView, Test));
+	ADD_LATENT_AUTOMATION_COMMAND(FWaitForRunningQuieryFinishedCommand(MemAllocTableTreeView, 120.0f, Test));
+	ADD_LATENT_AUTOMATION_COMMAND(FVerifyHierarchyCallStackCommand(MemAllocTableTreeView, InsightsTestUtils, 120.0f, Test));
+
+	return true;
+}
+
+// TestRail: T223693546, T223693548
+bool FMemoryInsightsAllocationsQueryTableEditorPackageTest::RunTest(const FString& Parameters)
+{
+	bool bSuccess = MemoryInsightsAllocationsQueryTableTest(Parameters, AllocsTimeMarkerEditorPackageGetterMap, this);
+	return bSuccess;
+}
+
+void FMemoryInsightsAllocationsQueryTableEditorPackageTest::GetTests(TArray<FString>& OutBeautifiedNames, TArray <FString>& OutTestCommands) const
+{
+	TSharedPtr<SMemoryProfilerWindow> ProfilerWindow = FMemoryProfilerManager::Get()->GetProfilerWindow();
+	if (!ProfilerWindow.IsValid())
+	{
+		return;
+	}
+
+	FMemorySharedState& SharedState = ProfilerWindow->GetSharedState();
+
+	for (const auto& MemoryRules : SharedState.GetMemoryRules())
+	{
+		const FString& MemoryRuleName = MemoryRules->GetShortName().ToString();
+
+		OutBeautifiedNames.Add(FString::Printf(TEXT("%s"), *MemoryRuleName));
+		OutTestCommands.Add(MemoryRuleName);
+	}
+}
+
+////////////////////////////////////////////////////////////////////////////////////////////////////
+
+// TestRail: T223693547
+bool FMemoryInsightsAllocationsQueryTableStandaloneTest::RunTest(const FString& Parameters)
+{
+	bool bSuccess = MemoryInsightsAllocationsQueryTableTest(Parameters, AllocsTimeMarkerStandaloneGameGetterMap, this);
+	return bSuccess;
+}
+
+void FMemoryInsightsAllocationsQueryTableStandaloneTest::GetTests(TArray<FString>& OutBeautifiedNames, TArray <FString>& OutTestCommands) const
+{
+	TSharedPtr<SMemoryProfilerWindow> ProfilerWindow = FMemoryProfilerManager::Get()->GetProfilerWindow();
+	if (!ProfilerWindow.IsValid())
+	{
+		return;
+	}
+
+	FMemorySharedState& SharedState = ProfilerWindow->GetSharedState();
+	TSharedPtr<Insights::FMemoryRuleSpec> MemoryRule = SharedState.GetMemoryRules()[0];
+
+	for (const auto& MemoryRules : SharedState.GetMemoryRules())
+	{
+		const FString& MemoryRuleName = MemoryRules->GetShortName().ToString();
+
+		OutBeautifiedNames.Add(FString::Printf(TEXT("%s"), *MemoryRuleName));
+		OutTestCommands.Add(MemoryRuleName);
+	}
+}
 #endif // !WITH_EDITOR
 
 ////////////////////////////////////////////////////////////////////////////////////////////////////
diff --git a/Engine/Source/Developer/TraceInsights/Private/Insights/Tests/UITests.h b/Engine/Source/Developer/TraceInsights/Private/Insights/Tests/UITests.h
--- a/Engine/Source/Developer/TraceInsights/Private/Insights/Tests/UITests.h
+++ b/Engine/Source/Developer/TraceInsights/Private/Insights/Tests/UITests.h
@@ -19,6 +19,8 @@ TSharedPtr<SWindow> AutomationWindow;
 END_DEFINE_SPEC(FAutomationDriverUnrealInsightsSessionBrowserTest)
 IMPLEMENT_SIMPLE_AUTOMATION_TEST(FMemoryInsightsUploadLLMXMLReportsTraceTest, "Insights.MemoryInsights.UploadMemoryInsightsLLMXMLReportsTrace", EAutomationTestFlags::ApplicationContextMask | EAutomationTestFlags::ProductFilter)
 IMPLEMENT_SIMPLE_AUTOMATION_TEST(FUnrealTraceServerStartingTest, "Insights.SessionBrowser.UnrealTraceServer.Starting", EAutomationTestFlags::ApplicationContextMask | EAutomationTestFlags::ProductFilter)
+IMPLEMENT_COMPLEX_AUTOMATION_TEST(FMemoryInsightsAllocationsQueryTableEditorPackageTest, "Insights.MemoryInsights.AllocationsQueryTable.Editor/Package", EAutomationTestFlags::ApplicationContextMask | EAutomationTestFlags::ProductFilter)
+IMPLEMENT_COMPLEX_AUTOMATION_TEST(FMemoryInsightsAllocationsQueryTableStandaloneTest, "Insights.MemoryInsights.AllocationsQueryTable.Standalone", EAutomationTestFlags::ApplicationContextMask | EAutomationTestFlags::ProductFilter)
 
 #endif
 
