[Insights] Fixed layout of major tabs to not create panels for unknown / non-persistent tabs, like the "Allocs Table" from Memory Insights.

diff --git a/Engine/Source/Developer/TraceInsights/Private/Insights/Widgets/SMajorTabWindow.cpp b/Engine/Source/Developer/TraceInsights/Private/Insights/Widgets/SMajorTabWindow.cpp
--- a/Engine/Source/Developer/TraceInsights/Private/Insights/Widgets/SMajorTabWindow.cpp
+++ b/Engine/Source/Developer/TraceInsights/Private/Insights/Widgets/SMajorTabWindow.cpp
@@ -278,7 +278,7 @@ void SMajorTabWindow::Construct(const FArguments& InArgs, const TSharedRef<SDock
 			+ SVerticalBox::Slot()
 			.FillHeight(1.0f)
 			[
-				TabManager->RestoreFrom(Layout, ConstructUnderWindow).ToSharedRef()
+				TabManager->RestoreFrom(Layout, ConstructUnderWindow, false, EOutputCanBeNullptr::IfNoTabValid).ToSharedRef()
 			]
 		]
 
