[Insights]
  - Memory Insights / Modules: Added status bar with info about total number of modules and total size of modules.
  - Memory Insights / Modules: Changed the numeric columns to be right aligned.

diff --git a/Engine/Source/Developer/TraceInsights/Private/Insights/Widgets/SModulesView.cpp b/Engine/Source/Developer/TraceInsights/Private/Insights/Widgets/SModulesView.cpp
--- a/Engine/Source/Developer/TraceInsights/Private/Insights/Widgets/SModulesView.cpp
+++ b/Engine/Source/Developer/TraceInsights/Private/Insights/Widgets/SModulesView.cpp
@@ -591,6 +591,7 @@ SModulesView::SModulesView()
 	: LastUpdateTime(0.0)
 	, SortColumn(ModulesViewColumns::DiscoveredColumnName)
 	, SortMode(EColumnSortMode::Descending)
+	, TotalSize(0)
 {
 }
 
@@ -608,6 +609,7 @@ void SModulesView::Reset()
 	//ListView
 	//ExternalScrollbar
 	Modules.Reset();
+	TotalSize = 0;
 
 	ListView->RebuildList();
 }
@@ -624,23 +626,6 @@ void SModulesView::Construct(const FArguments& InArgs)
 	[
 		SNew(SVerticalBox)
 
-		// Toolbar
-		//+ SVerticalBox::Slot()
-		//.AutoHeight()
-		//[
-		//	SNew(SHorizontalBox)
-		//
-		//	// Stats Text
-		//	+ SHorizontalBox::Slot()
-		//	.VAlign(VAlign_Center)
-		//	.Padding(2.0f, 0.0f)
-		//	.AutoWidth()
-		//	[
-		//		SNew(STextBlock)
-		//		.Text(this, &SModulesView::GetStatsText)
-		//	]
-		//]
-
 		+ SVerticalBox::Slot()
 		.FillHeight(1.0f)
 		[
@@ -682,7 +667,9 @@ void SModulesView::Construct(const FArguments& InArgs)
 							.OnSort(this, &SModulesView::OnSortModeChanged)
 							
 							+ SHeaderRow::Column(ModulesViewColumns::SizeColumnName)
-							.ManualWidth(76.0f)
+							.ManualWidth(80.0f)
+							.HAlignHeader(HAlign_Right)
+							.HAlignCell(HAlign_Right)
 							.DefaultLabel(LOCTEXT("SizeColumn", "Size"))
 							.InitialSortMode(EColumnSortMode::Descending)
 							.SortMode(this, &SModulesView::GetSortModeForColumn, ModulesViewColumns::SizeColumnName)
@@ -697,6 +684,8 @@ void SModulesView::Construct(const FArguments& InArgs)
 
 							+ SHeaderRow::Column(ModulesViewColumns::DiscoveredColumnName)
 							.ManualWidth(76.0f)
+							.HAlignHeader(HAlign_Right)
+							.HAlignCell(HAlign_Right)
 							.DefaultLabel(LOCTEXT("DiscoveredColumn", "Discovered"))
 							.DefaultTooltip(LOCTEXT("DiscoveredColumnTooltip", "The total number of symbols discovered in the trace.\nDiscovered = Cached + Resolved + Failed"))
 							.InitialSortMode(EColumnSortMode::Descending)
@@ -705,6 +694,8 @@ void SModulesView::Construct(const FArguments& InArgs)
 
 							+ SHeaderRow::Column(ModulesViewColumns::CachedColumnName)
 							.ManualWidth(76.0f)
+							.HAlignHeader(HAlign_Right)
+							.HAlignCell(HAlign_Right)
 							.DefaultLabel(LOCTEXT("CachedColumn", "Cached"))
 							.DefaultTooltip(LOCTEXT("CachedColumnTooltip", "The number of symbols loaded (pre-resolved) from the cache (*.ucache file)."))
 							.InitialSortMode(EColumnSortMode::Descending)
@@ -713,6 +704,8 @@ void SModulesView::Construct(const FArguments& InArgs)
 
 							+ SHeaderRow::Column(ModulesViewColumns::ResolvedColumnName)
 							.ManualWidth(76.0f)
+							.HAlignHeader(HAlign_Right)
+							.HAlignCell(HAlign_Right)
 							.DefaultLabel(LOCTEXT("ResolvedColumn", "Resolved"))
 							.DefaultTooltip(LOCTEXT("ResolvedColumnTooltip", "The number of symbols successfully resolved in the current analysis session."))
 							.InitialSortMode(EColumnSortMode::Descending)
@@ -721,6 +714,8 @@ void SModulesView::Construct(const FArguments& InArgs)
 
 							+ SHeaderRow::Column(ModulesViewColumns::FailedColumnName)
 							.ManualWidth(76.0f)
+							.HAlignHeader(HAlign_Right)
+							.HAlignCell(HAlign_Right)
 							.DefaultLabel(LOCTEXT("FailedColumn", "Failed"))
 							.DefaultTooltip(LOCTEXT("FailedColumnTooltip", "The number of symbols failed to resolve."))
 							.InitialSortMode(EColumnSortMode::Descending)
@@ -756,6 +751,22 @@ void SModulesView::Construct(const FArguments& InArgs)
 				]
 			]
 		]
+
+		// Status bar
+		+ SVerticalBox::Slot()
+		.AutoHeight()
+		[
+			SNew(SHorizontalBox)
+
+			+ SHorizontalBox::Slot()
+			.VAlign(VAlign_Center)
+			.Padding(2.0f, 0.0f)
+			.AutoWidth()
+			[
+				SNew(STextBlock)
+				.Text(this, &SModulesView::GetStatsText)
+			]
+		]
 	];
 }
 END_SLATE_FUNCTION_BUILD_OPTIMIZATION
@@ -787,6 +798,7 @@ void SModulesView::Tick(const FGeometry& AllottedGeometry, const double InCurren
 		ModuleProvider->EnumerateModules(Modules.Num(), [this](const TraceServices::FModule& InModule)
 		{
 			Modules.Emplace(MakeShared<FModule>(&InModule));
+			TotalSize += InModule.Size;
 		});
 		if (Modules.Num() != ModuleCount)
 		{
@@ -834,7 +846,7 @@ void SModulesView::OnSelectionChanged(TSharedPtr<FModule> Module, ESelectInfo::T
 
 FText SModulesView::GetStatsText() const
 {
-	return FText();
+	return FText::Format(LOCTEXT("StatusBarTextFmt", "{0} modules ({1} bytes)"), FText::AsNumber(Modules.Num()), FText::AsNumber(TotalSize));
 }
 
 ////////////////////////////////////////////////////////////////////////////////////////////////////
diff --git a/Engine/Source/Developer/TraceInsights/Private/Insights/Widgets/SModulesView.h b/Engine/Source/Developer/TraceInsights/Private/Insights/Widgets/SModulesView.h
--- a/Engine/Source/Developer/TraceInsights/Private/Insights/Widgets/SModulesView.h
+++ b/Engine/Source/Developer/TraceInsights/Private/Insights/Widgets/SModulesView.h
@@ -112,6 +112,9 @@ protected:
 
 	/** Sorting mode (ascending or descending) for current sorting column. */
 	EColumnSortMode::Type SortMode;
+
+	/** The total size of modules. */
+	uint64 TotalSize;
 };
 
 ////////////////////////////////////////////////////////////////////////////////////////////////////
