[Insights] Fixed crash when using TimingInsights.ExportTimerStatistics execute command with regions (ex: -ExecOnAnalysisCompleteCmd="TimingInsights.ExportTimerStatistics D:\out\*.csv -region=* -threads=GPU"). The "timing regions" provider lock was not set correctly.

diff --git a/Engine/Source/Developer/TraceInsights/Private/Insights/ViewModels/RegionsTimingTrack.cpp b/Engine/Source/Developer/TraceInsights/Private/Insights/ViewModels/RegionsTimingTrack.cpp
--- a/Engine/Source/Developer/TraceInsights/Private/Insights/ViewModels/RegionsTimingTrack.cpp
+++ b/Engine/Source/Developer/TraceInsights/Private/Insights/ViewModels/RegionsTimingTrack.cpp
@@ -2,24 +2,27 @@
 
 #include "RegionsTimingTrack.h"
 
+#include "Framework/MultiBox/MultiBoxBuilder.h"
 #include "HAL/PlatformApplicationMisc.h"
 
-#include "TimingViewDrawHelper.h"
+// TraceServices
 #include "Common/ProviderLock.h"
+#include "TraceServices/Model/Regions.h"
+
+// TraceInsights
+#include "Insights/Common/InsightsMenuBuilder.h"
+#include "Insights/Common/TimeUtils.h"
 #include "Insights/InsightsManager.h"
 #include "Insights/InsightsStyle.h"
 #include "Insights/ITimingViewSession.h"
 #include "Insights/Log.h"
 #include "Insights/TimingProfilerCommon.h"
-#include "Insights/Common/InsightsMenuBuilder.h"
-#include "Insights/Common/TimeUtils.h"
 #include "Insights/ViewModels/FilterConfigurator.h"
 #include "Insights/ViewModels/Filters.h"
 #include "Insights/ViewModels/ThreadTimingTrack.h"
 #include "Insights/ViewModels/TimingTrackViewport.h"
+#include "Insights/ViewModels/TimingViewDrawHelper.h"
 #include "Insights/Widgets/STimingView.h"
-#include "TraceServices/Model/Regions.h"
-#include "Framework/MultiBox/MultiBoxBuilder.h"
 
 #define LOCTEXT_NAMESPACE "RegionsTimingTrack"
 
@@ -214,7 +217,7 @@ void FTimingRegionsTrack::BuildDrawState(ITimingEventsTrackDrawStateBuilder& Bui
 
 	FStopwatch Stopwatch;
 	Stopwatch.Start();
-	
+
 	// whe're counting only non-empty lanes, so we can collapse empty ones in the visualization.
 	int32 CurDepth = 0;
 	RegionProvider.EnumerateLanes([this, Viewport, &CurDepth, &Builder](const TraceServices::FRegionLane& Lane, const int32 Depth)
@@ -323,7 +326,7 @@ void FTimingRegionsTrack::BuildFilteredDrawState(ITimingEventsTrackDrawStateBuil
 							{
 								Builder.AddEvent(Region.BeginTime, Region.EndTime, CurDepth, Region.Text);
 							}
-							
+
 							return true;
 						});
 
@@ -384,12 +387,12 @@ bool FTimingRegionsTrack::FindRegionEvent(const FTimingEventSearchParameters& In
 		RegionProvider.EnumerateRegions(InContext.GetParameters().StartTime, InContext.GetParameters().EndTime, [&InContext](const TraceServices::FTimeRegion& Region)
 		{
 			InContext.Check(Region.BeginTime, Region.EndTime, Region.Depth, Region);
-			
+
 			if (!InContext.ShouldContinueSearching())
 			{
 				return false;
 			}
-			
+
 			return true;
 		});
 	},
@@ -458,4 +461,4 @@ void FTimingRegionsTrack::OnClipboardCopyEvent(const ITimingEvent& InSelectedEve
 
 } // namespace Insights
 
-#undef LOCTEXT_NAMESPACE
\ No newline at end of file
+#undef LOCTEXT_NAMESPACE
diff --git a/Engine/Source/Developer/TraceInsights/Private/Insights/ViewModels/TimingExporter.cpp b/Engine/Source/Developer/TraceInsights/Private/Insights/ViewModels/TimingExporter.cpp
--- a/Engine/Source/Developer/TraceInsights/Private/Insights/ViewModels/TimingExporter.cpp
+++ b/Engine/Source/Developer/TraceInsights/Private/Insights/ViewModels/TimingExporter.cpp
@@ -1,24 +1,28 @@
 // Copyright Epic Games, Inc. All Rights Reserved.
 
 #include "TimingExporter.h"
+
 #include "HAL/PlatformFileManager.h"
 #include "Logging/MessageLog.h"
+
+// TraceServices
+#include "Common/ProviderLock.h"
 #include "TraceServices/Model/Bookmarks.h"
 #include "TraceServices/Model/Counters.h"
+#include "TraceServices/Model/Regions.h"
 #include "TraceServices/Model/Threads.h"
 
-// Insights
+// TraceInsights
 #include "Insights/Common/Stopwatch.h"
 #include "Insights/Log.h"
-#include "Insights/ViewModels/ThreadTimingTrack.h"
 #include "Insights/TimingProfilerManager.h"
-#include "TraceServices/Model/Regions.h"
+#include "Insights/ViewModels/ThreadTimingTrack.h"
+
+#define LOCTEXT_NAMESPACE "FTimingExporter"
 
 namespace Insights
 {
 
-#define LOCTEXT_NAMESPACE "FTimingExporter"
-
 ////////////////////////////////////////////////////////////////////////////////////////////////////
 // FTimingExporter
 ////////////////////////////////////////////////////////////////////////////////////////////////////
@@ -580,14 +584,13 @@ int32 FTimingExporter::ExportTimerStatisticsAsText(const FString& Filename, FExp
 			TArray<FString> NamePatterns;
 		};
 		FRegionNameSpec RegionNameSpec(Params.Region);
-		
+
 		TMap<FString, TraceServices::FTimeRegion> Regions;
 
 		// Detect regions
 		{
-			TraceServices::FAnalysisSessionReadScope SessionReadScope(Session);
-
 			const TraceServices::IRegionProvider& RegionProvider = TraceServices::ReadRegionProvider(Session);
+			TraceServices::FProviderReadScopeLock RegionProviderScopedLock(RegionProvider);
 
 			UE_LOG(TraceInsights, Log, TEXT("Looking for regions: '%s'"), *Params.Region);
 
@@ -598,7 +601,7 @@ int32 FTimingExporter::ExportTimerStatisticsAsText(const FString& Filename, FExp
 					{
 						return true;
 					}
-					
+
 					// Handle duplicate region names, individual regions may appear multiple times
 					// we append numbers to allow for unique export filenames.
 					FString RegionName = InRegion.Text;
@@ -1128,6 +1131,6 @@ int32 FTimingExporter::ExportCounterAsText(const FString& Filename, uint32 Count
 
 ////////////////////////////////////////////////////////////////////////////////////////////////////
 
-#undef LOCTEXT_NAMESPACE
-
 } // namespace Insights
+
+#undef LOCTEXT_NAMESPACE
