[Insights] Automated T223693536 "Unreal Insights - Trace.SnapshotFile - Active Editor Profiling Session".
Automated T223693540 "Unreal Insights - Trace.SnapshotFile - Late-Connect Editor Profiling Session".
Automated T223693523 "Unreal Insights - Trace Enable/Disable - Editor ".
Automated T223693549 "Unreal Insights - UnrealTraceServer - Stopping - Package".
Automated T223693525 "Unreal Insights - Trace Pause/Resume - Editor ".
Automated T223693535 "Verify that Trace.SnapshotFile creates a small utrace from an Editor session outside of an active profiling session"
Changed logic in FInsightsTestUtils::FileContainsString.

diff --git a/Engine/Source/Developer/TraceInsights/Private/Insights/InsightsManager.h b/Engine/Source/Developer/TraceInsights/Private/Insights/InsightsManager.h
--- a/Engine/Source/Developer/TraceInsights/Private/Insights/InsightsManager.h
+++ b/Engine/Source/Developer/TraceInsights/Private/Insights/InsightsManager.h
@@ -328,6 +328,9 @@ public:
 
 	void ScheduleCommand(const FString& InCmd);
 
+	/** Resets (closes) current session instance. */
+	void ResetSession(bool bNotify = true);
+
 	////////////////////////////////////////////////////////////////////////////////////////////////////
 	// SessionChangedEvent
 
@@ -371,6 +374,8 @@ private:
 	/** Called to spawn the Launcher major tab. */
 	TSharedRef<SDockTab> SpawnLauncherTab(const FSpawnTabArgs& Args);
 
+	TSharedRef<SDockTab> SpawnSessionBrowserAutomationWindowTab(const FSpawnTabArgs& Args);
+
 	/** Callback called when the Launcher major tab is closed. */
 	void OnLauncherTabClosed(TSharedRef<SDockTab> TabBeingClosed);
 
@@ -383,9 +388,6 @@ private:
 	/** Updates this manager, done through FCoreTicker. */
 	bool Tick(float DeltaTime);
 
-	/** Resets (closes) current session instance. */
-	void ResetSession(bool bNotify = true);
-
 	/** Extract messages from the session */
 	void PollAnalysisInfo();
 
diff --git a/Engine/Source/Developer/TraceInsights/Private/Insights/Tests/InsightsTestUtils.cpp b/Engine/Source/Developer/TraceInsights/Private/Insights/Tests/InsightsTestUtils.cpp
--- a/Engine/Source/Developer/TraceInsights/Private/Insights/Tests/InsightsTestUtils.cpp
+++ b/Engine/Source/Developer/TraceInsights/Private/Insights/Tests/InsightsTestUtils.cpp
@@ -2,17 +2,19 @@
 
 #include "Insights/Tests/InsightsTestUtils.h"
 
-#include "Misc/AutomationTest.h"
 #include "Modules/ModuleManager.h"
+#include "Misc/FileHelper.h"
+#include "HAL/FileManagerGeneric.h"
+
 #include "TraceServices/AnalysisService.h"
 #include "TraceServices/Model/AnalysisSession.h"
 #include "TraceServices/ITraceServicesModule.h"
 #include "TraceServices/ModuleService.h"
-
 #include "Insights/Common/Stopwatch.h"
 #include "Insights/IUnrealInsightsModule.h"
+#include "Insights/InsightsManager.h"
 
-#include "Misc/FileHelper.h"
+#include "Misc/AutomationTest.h"
 
 ////////////////////////////////////////////////////////////////////////////////////////////////////
 
@@ -75,14 +77,26 @@ bool FInsightsTestUtils::AnalyzeTrace(const TCHAR* Path) const
 
 ////////////////////////////////////////////////////////////////////////////////////////////////////
 
-bool FInsightsTestUtils::FileContainsString(const FString& PathToFile, const FString& ExpectedString, const float Timeout) const
+bool FInsightsTestUtils::FileContainsString(const FString& PathToFile, const FString& ExpectedString, const float& Timeout) const
 {
+	if (!FPaths::FileExists(PathToFile))
+	{
+		return false;
+	}
+
 	float StartTime = FPlatformTime::Seconds();
 	while ((FPlatformTime::Seconds() - StartTime) < Timeout)
 	{
 		FString LogFileContents;
-		if (FFileHelper::LoadFileToString(LogFileContents, *PathToFile))
+		IPlatformFile& PlatformFile = FPlatformFileManager::Get().GetPlatformFile();
+		TUniquePtr<IFileHandle> FileHandle(PlatformFile.OpenRead(*PathToFile, true)); // Open the file with shared read access
+		if (FileHandle)
 		{
+			TArray<uint8> FileData;
+			FileData.SetNumUninitialized(FileHandle->Size());
+			FileHandle->Read(FileData.GetData(), FileData.Num());
+			FFileHelper::BufferToString(LogFileContents, FileData.GetData(), FileData.Num());
+
 			if (LogFileContents.Contains(ExpectedString))
 			{
 				return true;
@@ -93,3 +107,109 @@ bool FInsightsTestUtils::FileContainsString(const FString& PathToFile, const FSt
 
 	return false;
 }
+
+////////////////////////////////////////////////////////////////////////////////////////////////////
+
+bool FInsightsTestUtils::StartTracing(FTraceAuxiliary::EConnectionType ConnectionType, const float& Timeout) const
+{
+	if (ConnectionType == FTraceAuxiliary::EConnectionType::Network)
+	{
+		return FTraceAuxiliary::Start(FTraceAuxiliary::EConnectionType::Network, TEXT("localhost"), nullptr);
+	}
+	else if (ConnectionType == FTraceAuxiliary::EConnectionType::File)
+	{
+		return FTraceAuxiliary::Start(FTraceAuxiliary::EConnectionType::File, nullptr, nullptr);
+	}
+
+	double TraceVerifyStartTime = FPlatformTime::Seconds();
+	while (FPlatformTime::Seconds() - TraceVerifyStartTime < Timeout)
+	{
+		if (FTraceAuxiliary::IsConnected())
+		{
+			return true;
+		}
+		FPlatformProcess::Sleep(0.1f);
+	}
+
+	return false;
+}
+
+////////////////////////////////////////////////////////////////////////////////////////////////////
+
+bool FInsightsTestUtils::SetupUTS(const float& Timeout) const
+{
+	const FString UnrealTraceServerName = TEXT("UnrealTraceServer");
+
+	FString UTSPath = FPlatformProcess::GenerateApplicationPath("UnrealTraceServer", EBuildConfiguration::Development);
+	FString UTSParameters = TEXT("daemon");
+	constexpr bool bLaunchDetached = true;
+	constexpr bool bLaunchHidden = false;
+	constexpr bool bLaunchReallyHidden = false;
+	uint32 ProcessID = 0;
+	const int32 PriorityModifier = 0;
+	const TCHAR* OptionalWorkingDirectory = nullptr;
+	void* PipeWriteChild = nullptr;
+	void* PipeReadChild = nullptr;
+	FProcHandle UTSHandle = FPlatformProcess::CreateProc(*UTSPath, *UTSParameters, bLaunchDetached, bLaunchHidden, bLaunchReallyHidden, &ProcessID, PriorityModifier, OptionalWorkingDirectory, PipeWriteChild, PipeReadChild);
+	if (!UTSHandle.IsValid())
+	{
+		return false;
+	}
+
+	double StartTime = FPlatformTime::Seconds();
+	while (FPlatformTime::Seconds() - StartTime < Timeout)
+	{
+		if (FPlatformProcess::IsApplicationRunning(*UnrealTraceServerName))
+		{
+			return true;
+		}
+		UTSHandle = FPlatformProcess::CreateProc(*UTSPath, *UTSParameters, bLaunchDetached, bLaunchHidden, bLaunchReallyHidden, &ProcessID, PriorityModifier, OptionalWorkingDirectory, PipeWriteChild, PipeReadChild);
+		FPlatformProcess::Sleep(0.1f);
+	}
+
+	return false;
+}
+
+////////////////////////////////////////////////////////////////////////////////////////////////////
+
+bool FInsightsTestUtils::KillUTS(const float& Timeout) const
+{
+	const FString UnrealTraceServerName = TEXT("UnrealTraceServer");
+
+	FString UTSPath = FPlatformProcess::GenerateApplicationPath("UnrealTraceServer", EBuildConfiguration::Development);
+	FString UTSParameters = TEXT("kill");
+	constexpr bool bLaunchDetached = true;
+	constexpr bool bLaunchHidden = false;
+	constexpr bool bLaunchReallyHidden = false;
+	uint32 ProcessID = 0;
+	const int32 PriorityModifier = 0;
+	const TCHAR* OptionalWorkingDirectory = nullptr;
+	void* PipeWriteChild = nullptr;
+	void* PipeReadChild = nullptr;
+	FProcHandle UTSHandle = FPlatformProcess::CreateProc(*UTSPath, *UTSParameters, bLaunchDetached, bLaunchHidden, bLaunchReallyHidden, &ProcessID, PriorityModifier, OptionalWorkingDirectory, PipeWriteChild, PipeReadChild);
+	if (!UTSHandle.IsValid())
+	{
+		return false;
+	}
+
+	double StartTime = FPlatformTime::Seconds();
+	while (FPlatformTime::Seconds() - StartTime < Timeout)
+	{
+		if (!FPlatformProcess::IsApplicationRunning(*UnrealTraceServerName))
+		{
+			return true;
+		}
+		UTSHandle = FPlatformProcess::CreateProc(*UTSPath, *UTSParameters, bLaunchDetached, bLaunchHidden, bLaunchReallyHidden, &ProcessID, PriorityModifier, OptionalWorkingDirectory, PipeWriteChild, PipeReadChild); 
+		FPlatformProcess::Sleep(0.1f);
+	}
+
+	return false;
+}
+
+////////////////////////////////////////////////////////////////////////////////////////////////////
+
+void FInsightsTestUtils::ResetSession() const
+{
+	TSharedPtr<FInsightsManager> InsightsManager = FInsightsManager::Get();
+	InsightsManager->ResetSession();
+}
\ No newline at end of file
diff --git a/Engine/Source/Developer/TraceInsights/Private/Insights/TraceInsightsModule.cpp b/Engine/Source/Developer/TraceInsights/Private/Insights/TraceInsightsModule.cpp
--- a/Engine/Source/Developer/TraceInsights/Private/Insights/TraceInsightsModule.cpp
+++ b/Engine/Source/Developer/TraceInsights/Private/Insights/TraceInsightsModule.cpp
@@ -198,6 +198,31 @@ void FTraceInsightsModule::CreateDefaultStore()
 
 ////////////////////////////////////////////////////////////////////////////////////////////////////
 
+FString FTraceInsightsModule::GetDefaultStoreDir()
+{
+	using UE::Trace::FStoreClient;
+
+	FStoreClient* StoreClientPtr = FStoreClient::Connect(TEXT("localhost"));
+	TUniquePtr<FStoreClient> StoreClient = TUniquePtr<FStoreClient>(StoreClientPtr);
+
+	if (!StoreClient)
+	{
+		UE_LOG(TraceInsights, Error, TEXT("Failed to connect to the store client."));
+		return FString("");
+	}
+
+	const FStoreClient::FStatus* Status = StoreClient->GetStatus();
+	if (!StoreClient)
+	{
+		UE_LOG(TraceInsights, Error, TEXT("Failed to get the status of the store client."));
+		return FString("");
+	}
+
+	return FString(Status->GetStoreDir());
+}
+
+////////////////////////////////////////////////////////////////////////////////////////////////////
+
 UE::Trace::FStoreClient* FTraceInsightsModule::GetStoreClient()
 {
 	return FInsightsManager::Get()->GetStoreClient();
diff --git a/Engine/Source/Developer/TraceInsights/Private/Insights/TraceInsightsModule.h b/Engine/Source/Developer/TraceInsights/Private/Insights/TraceInsightsModule.h
--- a/Engine/Source/Developer/TraceInsights/Private/Insights/TraceInsightsModule.h
+++ b/Engine/Source/Developer/TraceInsights/Private/Insights/TraceInsightsModule.h
@@ -43,6 +43,8 @@ public:
 
 	virtual void CreateDefaultStore() override;
 
+	FString GetDefaultStoreDir();
+
 	virtual UE::Trace::FStoreClient* GetStoreClient() override;
 	virtual bool ConnectToStore(const TCHAR* InStoreHost, uint32 InStorePort=0) override;
 
diff --git a/Engine/Source/Developer/TraceInsights/Public/Insights/IUnrealInsightsModule.h b/Engine/Source/Developer/TraceInsights/Public/Insights/IUnrealInsightsModule.h
--- a/Engine/Source/Developer/TraceInsights/Public/Insights/IUnrealInsightsModule.h
+++ b/Engine/Source/Developer/TraceInsights/Public/Insights/IUnrealInsightsModule.h
@@ -179,6 +179,11 @@ public:
 	 */
 	virtual void CreateDefaultStore() = 0;
 
+	/**
+	 * Gets the default trace store (for "Browser" mode).
+	*/
+	virtual FString GetDefaultStoreDir() = 0;
+
 	/**
 	 * Gets the store client.
 	 */
diff --git a/Engine/Source/Developer/TraceInsights/Public/Insights/Tests/InsightsTestUtils.h b/Engine/Source/Developer/TraceInsights/Public/Insights/Tests/InsightsTestUtils.h
--- a/Engine/Source/Developer/TraceInsights/Public/Insights/Tests/InsightsTestUtils.h
+++ b/Engine/Source/Developer/TraceInsights/Public/Insights/Tests/InsightsTestUtils.h
@@ -3,6 +3,7 @@
 #pragma once
 
 #include "CoreMinimal.h"
+#include "ProfilingDebugging/TraceAuxiliary.h"
 
 class FAutomationTestBase;
 
@@ -12,7 +13,11 @@ public:
 	FInsightsTestUtils(FAutomationTestBase* Test);
 
 	bool AnalyzeTrace(const TCHAR* Path) const;
-	bool FileContainsString(const FString& PathToFile, const FString& ExpectedString, const float Timeout) const;
+	bool FileContainsString(const FString& PathToFile, const FString& ExpectedString, const float& Timeout) const;
+	bool SetupUTS(const float& Timeout) const;
+	bool KillUTS(const float& Timeout) const;
+	bool StartTracing(FTraceAuxiliary::EConnectionType ConnectionType, const float& Timeout) const;
+	void ResetSession() const;
 
 private:
 	FAutomationTestBase* Test;
