Fix crash with FAutomationTestBase::GenerateTestNames()

#rnx

diff --git a/Engine/Source/Developer/TraceInsights/Private/Insights/Tests/UITests.cpp b/Engine/Source/Developer/TraceInsights/Private/Insights/Tests/UITests.cpp
--- a/Engine/Source/Developer/TraceInsights/Private/Insights/Tests/UITests.cpp
+++ b/Engine/Source/Developer/TraceInsights/Private/Insights/Tests/UITests.cpp
@@ -559,6 +559,10 @@ bool FMemoryInsightsAllocationsQueryTableEditorPackageTest::RunTest(const FStrin
 
 void FMemoryInsightsAllocationsQueryTableEditorPackageTest::GetTests(TArray<FString>& OutBeautifiedNames, TArray <FString>& OutTestCommands) const
 {
+	if (!FMemoryProfilerManager::Get().IsValid())
+	{
+		return;
+	}
 	TSharedPtr<SMemoryProfilerWindow> ProfilerWindow = FMemoryProfilerManager::Get()->GetProfilerWindow();
 	if (!ProfilerWindow.IsValid())
 	{
@@ -587,6 +591,10 @@ bool FMemoryInsightsAllocationsQueryTableStandaloneTest::RunTest(const FString&
 
 void FMemoryInsightsAllocationsQueryTableStandaloneTest::GetTests(TArray<FString>& OutBeautifiedNames, TArray <FString>& OutTestCommands) const
 {
+	if (!FMemoryProfilerManager::Get().IsValid())
+	{
+		return;
+	}
 	TSharedPtr<SMemoryProfilerWindow> ProfilerWindow = FMemoryProfilerManager::Get()->GetProfilerWindow();
 	if (!ProfilerWindow.IsValid())
 	{
