No PCH build fix

#rnx

diff --git a/Engine/Source/Developer/TraceInsights/Private/Insights/Tests/UITests.cpp b/Engine/Source/Developer/TraceInsights/Private/Insights/Tests/UITests.cpp
--- a/Engine/Source/Developer/TraceInsights/Private/Insights/Tests/UITests.cpp
+++ b/Engine/Source/Developer/TraceInsights/Private/Insights/Tests/UITests.cpp
@@ -2,6 +2,9 @@
 
 #include "UITests.h"
 
+#include "Algo/Find.h"
+
+#include "HAL/FileManager.h"
 #include "HAL/PlatformFileManager.h"
 
 #include "Insights/Common/Stopwatch.h"
