Deleted Test Rail references in Insights tests

diff --git a/Engine/Source/Developer/TraceInsights/Private/Insights/Tests/MiscTests.cpp b/Engine/Source/Developer/TraceInsights/Private/Insights/Tests/MiscTests.cpp
--- a/Engine/Source/Developer/TraceInsights/Private/Insights/Tests/MiscTests.cpp
+++ b/Engine/Source/Developer/TraceInsights/Private/Insights/Tests/MiscTests.cpp
@@ -90,7 +90,6 @@ bool FPagedArrayFilteringTest::RunTest(const FString& Parameters)
 
 IMPLEMENT_SIMPLE_AUTOMATION_TEST(FСommandsExportWindowsTest, "Insights.CommandsExport(Windows)", EAutomationTestFlags::ApplicationContextMask | EAutomationTestFlags::ProductFilter)
 
-// TestRail: C27304222
 bool FСommandsExportWindowsTest::RunTest(const FString& Parameters)
 {
 	TSharedPtr<FInsightsManager> InsightsManager = FInsightsManager::Get();
diff --git a/Engine/Source/Developer/TraceInsights/Private/Insights/Tests/UITests.cpp b/Engine/Source/Developer/TraceInsights/Private/Insights/Tests/UITests.cpp
--- a/Engine/Source/Developer/TraceInsights/Private/Insights/Tests/UITests.cpp
+++ b/Engine/Source/Developer/TraceInsights/Private/Insights/Tests/UITests.cpp
@@ -87,7 +87,6 @@ void FAutomationDriverUnrealInsightsSessionBrowserTest::Define()
 		Driver = IAutomationDriverModule::Get().CreateDriver();		
 		});
 
-	// TestRail: C28843555
 	Describe("CopyRenameDeleteTrace", [this]()
 		{
 			It("should verify that user copy, rename and delete traces", EAsyncExecution::ThreadPool, [this]()
@@ -185,7 +184,6 @@ void FAutomationDriverUnrealInsightsSessionBrowserTest::Define()
 				});
 		});
 
-	// TestRail: C35547680
 	Describe("MemoryInsights.XMLReportsUpload", [this]()
 		{
 			It("should verify that user can upload xml reports in Memory Insights tab", EAsyncExecution::ThreadPool, FTimespan::FromSeconds(120), [this]()
@@ -277,7 +275,6 @@ void FAutomationDriverUnrealInsightsSessionBrowserTest::Define()
 
 ////////////////////////////////////////////////////////////////////////////////////////////////////
 
-// TestRail: C35547680
 bool FMemoryInsightsUploadLLMXMLReportsTraceTest::RunTest(const FString& Parameters)
 {
 	const FString ReportGraphsXMLPath = FPaths::RootDir() / TEXT("EngineTest/SourceAssets/Utrace/ReportGraphs.xml");
@@ -310,7 +307,6 @@ bool FMemoryInsightsUploadLLMXMLReportsTraceTest::RunTest(const FString& Paramet
 
 ////////////////////////////////////////////////////////////////////////////////////////////////////
 
-// TestRail: T223693543
 bool FUnrealTraceServerStartingTest::RunTest(const FString& Parameters)
 {
 	TSharedPtr<FInsightsManager> InsightsManager = FInsightsManager::Get();
@@ -553,7 +549,6 @@ bool MemoryInsightsAllocationsQueryTableTest(const FString& Parameters, const TM
 	return true;
 }
 
-// TestRail: T223693546, T223693548
 bool FMemoryInsightsAllocationsQueryTableEditorPackageTest::RunTest(const FString& Parameters)
 {
 	bool bSuccess = MemoryInsightsAllocationsQueryTableTest(Parameters, AllocsTimeMarkerEditorPackageGetterMap, this);
@@ -585,7 +580,6 @@ void FMemoryInsightsAllocationsQueryTableEditorPackageTest::GetTests(TArray<FStr
 
 ////////////////////////////////////////////////////////////////////////////////////////////////////
 
-// TestRail: T223693547
 bool FMemoryInsightsAllocationsQueryTableStandaloneTest::RunTest(const FString& Parameters)
 {
 	bool bSuccess = MemoryInsightsAllocationsQueryTableTest(Parameters, AllocsTimeMarkerStandaloneGameGetterMap, this);
