[Insights] Unkown events in core tracks are now faded.
+ Added a new menu option to hide non target process events completely.

diff --git a/Engine/Source/Developer/TraceInsights/Private/Insights/ContextSwitches/ViewModels/ContextSwitchesSharedState.cpp b/Engine/Source/Developer/TraceInsights/Private/Insights/ContextSwitches/ViewModels/ContextSwitchesSharedState.cpp
--- a/Engine/Source/Developer/TraceInsights/Private/Insights/ContextSwitches/ViewModels/ContextSwitchesSharedState.cpp
+++ b/Engine/Source/Developer/TraceInsights/Private/Insights/ContextSwitches/ViewModels/ContextSwitchesSharedState.cpp
@@ -77,6 +77,12 @@ void FContextSwitchesStateCommands::RegisterCommands()
 		EUserInterfaceActionType::ToggleButton,
 		FInputChord(EModifierKey::Shift, EKeys::L));
 
+	UI_COMMAND(Command_ShowNonTargetProcessEvents,
+		"Non-Target Process Events",
+		"Shows/hides the cpu core events that do not belong to the target process.",
+		EUserInterfaceActionType::ToggleButton,
+		FInputChord());
+
 	UI_COMMAND(Command_NavigateToCpuThreadEvent,
 		"Go To CPU Thread Track",
 		"Selects the context switch event in the corresponding CPU Thread track.",
@@ -115,6 +121,7 @@ FContextSwitchesSharedState::FContextSwitchesSharedState(STimingView* InTimingVi
 	, bAreContextSwitchesVisible(true)
 	, bAreOverlaysVisible(true)
 	, bAreExtendedLinesVisible(true)
+	, bAreNonTargetProcessEventsVisible(true)
 	, bSyncWithProviders(true)
 {
 }
@@ -146,6 +153,7 @@ void FContextSwitchesSharedState::OnBeginSession(Insights::ITimingViewSession& I
 	bAreContextSwitchesVisible = true;
 	bAreOverlaysVisible = true;
 	bAreExtendedLinesVisible = true;
+	bAreNonTargetProcessEventsVisible = true;
 
 	bSyncWithProviders = true;
 
@@ -238,6 +246,7 @@ void FContextSwitchesSharedState::BuildSubMenu(FMenuBuilder& InMenuBuilder)
 	InMenuBuilder.BeginSection("ContextSwitches", LOCTEXT("ContextMenu_Section_ContextSwitches", "Context Switches"));
 	{
 		InMenuBuilder.AddMenuEntry(FContextSwitchesStateCommands::Get().Command_ShowCoreTracks);
+		InMenuBuilder.AddMenuEntry(FContextSwitchesStateCommands::Get().Command_ShowNonTargetProcessEvents);
 		InMenuBuilder.AddMenuEntry(FContextSwitchesStateCommands::Get().Command_ShowContextSwitches);
 		InMenuBuilder.AddMenuEntry(FContextSwitchesStateCommands::Get().Command_ShowOverlays);
 		InMenuBuilder.AddMenuEntry(FContextSwitchesStateCommands::Get().Command_ShowExtendedLines);
@@ -308,6 +317,12 @@ void FContextSwitchesSharedState::AddCommands()
 		FCanExecuteAction::CreateSP(this, &FContextSwitchesSharedState::Command_ShowCoreTracks_CanExecute),
 		FIsActionChecked::CreateSP(this, &FContextSwitchesSharedState::Command_ShowCoreTracks_IsChecked));
 
+	CommandList->MapAction(
+		FContextSwitchesStateCommands::Get().Command_ShowNonTargetProcessEvents,
+		FExecuteAction::CreateSP(this, &FContextSwitchesSharedState::Command_ShowNonTargetProcessEvents_Execute),
+		FCanExecuteAction::CreateSP(this, &FContextSwitchesSharedState::Command_ShowNonTargetProcessEvents_CanExecute),
+		FIsActionChecked::CreateSP(this, &FContextSwitchesSharedState::Command_ShowNonTargetProcessEvents_IsChecked));
+
 	CommandList->MapAction(
 		FContextSwitchesStateCommands::Get().Command_ShowContextSwitches,
 		FExecuteAction::CreateSP(this, &FContextSwitchesSharedState::Command_ShowContextSwitches_Execute),
@@ -564,6 +579,30 @@ void FContextSwitchesSharedState::SetExtendedLinesVisible(bool bOnOff)
 
 ////////////////////////////////////////////////////////////////////////////////////////////////////
 
+void FContextSwitchesSharedState::SetNonTargetProcessEventsVisible(bool bOnOff)
+{
+	if (bAreNonTargetProcessEventsVisible != bOnOff)
+	{
+		TSharedPtr<STimingView> TimingView = GetTimingView();
+		const TSharedPtr<ITimingEventFilter> EventFilter = TimingView->GetEventFilter();
+		for (auto Track : CpuCoreTimingTracks)
+		{
+			Track.Value->SetDirtyFlag();
+			if (!bOnOff && EventFilter)
+			{
+				if (EventFilter->FilterTrack(*Track.Value))
+				{
+					TimingView->ResetEventFilter();
+				}
+			}
+		}
+		
+		bAreNonTargetProcessEventsVisible = bOnOff;
+	}
+}
+
+////////////////////////////////////////////////////////////////////////////////////////////////////
+
 bool FContextSwitchesSharedState::IsValidCpuCoreEventSelected() const
 {
 	TSharedPtr<STimingView> TimingView = GetTimingView();
diff --git a/Engine/Source/Developer/TraceInsights/Private/Insights/ContextSwitches/ViewModels/ContextSwitchesSharedState.h b/Engine/Source/Developer/TraceInsights/Private/Insights/ContextSwitches/ViewModels/ContextSwitchesSharedState.h
--- a/Engine/Source/Developer/TraceInsights/Private/Insights/ContextSwitches/ViewModels/ContextSwitchesSharedState.h
+++ b/Engine/Source/Developer/TraceInsights/Private/Insights/ContextSwitches/ViewModels/ContextSwitchesSharedState.h
@@ -42,6 +42,7 @@ public:
 	TSharedPtr<FUICommandInfo> Command_ShowContextSwitches;
 	TSharedPtr<FUICommandInfo> Command_ShowOverlays;
 	TSharedPtr<FUICommandInfo> Command_ShowExtendedLines;
+	TSharedPtr<FUICommandInfo> Command_ShowNonTargetProcessEvents;
 
 	// Commands for a Cpu Core track (context menu).
 	TSharedPtr<FUICommandInfo> Command_NavigateToCpuThreadEvent;
@@ -99,6 +100,12 @@ public:
 	void ToggleExtendedLines() { SetExtendedLinesVisible(!bAreExtendedLinesVisible); }
 	void SetExtendedLinesVisible(bool bOnOff);
 
+	bool AreNonTargetProcessEventsVisible() const { return bAreNonTargetProcessEventsVisible; }
+	void ShowNonTargetProcessEvents() { SetNonTargetProcessEventsVisible(true); }
+	void HideNonTargetProcessEvents() { SetNonTargetProcessEventsVisible(false); }
+	void ToggleNonTargetProcessEvents() { SetNonTargetProcessEventsVisible(!bAreNonTargetProcessEventsVisible); }
+	void SetNonTargetProcessEventsVisible(bool bOnOff);
+
 	void SetTargetTimingEvent(const TSharedPtr<const ITimingEvent> InEvent)
 	{
 		TargetTimingEvent = InEvent;
@@ -135,6 +142,10 @@ private:
 	bool Command_ShowExtendedLines_CanExecute() const { return AreContextSwitchesAvailable() && AreContextSwitchesVisible(); }
 	bool Command_ShowExtendedLines_IsChecked() const { return AreExtendedLinesVisible(); }
 
+	void Command_ShowNonTargetProcessEvents_Execute() { ToggleNonTargetProcessEvents(); }
+	bool Command_ShowNonTargetProcessEvents_CanExecute() const { return AreContextSwitchesAvailable() && AreCoreTracksVisible(); }
+	bool Command_ShowNonTargetProcessEvents_IsChecked() const { return AreNonTargetProcessEventsVisible(); }
+
 	bool IsValidCpuCoreEventSelected() const;
 	bool IsValidContextSwitchEventSelected() const;
 
@@ -165,6 +176,7 @@ private:
 	bool bAreContextSwitchesVisible;
 	bool bAreOverlaysVisible;
 	bool bAreExtendedLinesVisible;
+	bool bAreNonTargetProcessEventsVisible;
 
 	bool bSyncWithProviders;
 
diff --git a/Engine/Source/Developer/TraceInsights/Private/Insights/ContextSwitches/ViewModels/CpuCoreTimingTrack.cpp b/Engine/Source/Developer/TraceInsights/Private/Insights/ContextSwitches/ViewModels/CpuCoreTimingTrack.cpp
--- a/Engine/Source/Developer/TraceInsights/Private/Insights/ContextSwitches/ViewModels/CpuCoreTimingTrack.cpp
+++ b/Engine/Source/Developer/TraceInsights/Private/Insights/ContextSwitches/ViewModels/CpuCoreTimingTrack.cpp
@@ -30,6 +30,26 @@ INSIGHTS_IMPLEMENT_RTTI(FCpuCoreTimingTrack)
 
 ////////////////////////////////////////////////////////////////////////////////////////////////////
 
+FCpuCoreTimingTrack::FCpuCoreTimingTrack(FContextSwitchesSharedState& InSharedState, const FString& InName, uint32 InCoreNumber)
+	: FTimingEventsTrack(InName)
+	, SharedState(InSharedState)
+	, CoreNumber(InCoreNumber)
+	, NonTargetProcessEventsDrawState(MakeShared<FTimingEventsTrackDrawState>())
+{
+}
+
+////////////////////////////////////////////////////////////////////////////////////////////////////
+
+void FCpuCoreTimingTrack::Reset()
+{
+	NonTargetProcessEventsMaxDepth = -1;
+	NonTargetProcessEventsDrawState->Reset();
+
+	FTimingEventsTrack::Reset();
+}
+
+////////////////////////////////////////////////////////////////////////////////////////////////////
+
 void FCpuCoreTimingTrack::BuildDrawState(ITimingEventsTrackDrawStateBuilder& Builder, const ITimingTrackUpdateContext& Context)
 {
 	TSharedPtr<const TraceServices::IAnalysisSession> Session = FInsightsManager::Get()->GetSession();
@@ -47,14 +67,25 @@ void FCpuCoreTimingTrack::BuildDrawState(ITimingEventsTrackDrawStateBuilder& Bui
 	}
 
 	const FTimingTrackViewport& Viewport = Context.GetViewport();
-
+	FTimingEventsTrackDrawStateBuilder NonTargetProcessEventsBuilder(*NonTargetProcessEventsDrawState, Viewport, Context.GetGeometry().Scale);
+	bool bShowNonTargetProcessEvents = SharedState.AreNonTargetProcessEventsVisible();
 	ContextSwitchesProvider->EnumerateCpuCoreEvents(CoreNumber, Viewport.GetStartTime(), Viewport.GetEndTime(), 
-		[&Builder, ContextSwitchesProvider, this](const TraceServices::FCpuCoreEvent& CpuCoreEvent)
+		[&Builder, &NonTargetProcessEventsBuilder, ContextSwitchesProvider, bShowNonTargetProcessEvents, this](const TraceServices::FCpuCoreEvent& CpuCoreEvent)
 		{
-			this->AddCoreTimingEvent(Builder, CpuCoreEvent);
+			uint32 ThreadId = 0;
+			if (ContextSwitchesProvider->GetThreadId(CpuCoreEvent.SystemThreadId, ThreadId))
+			{
+				this->AddCoreTimingEvent(Builder, CpuCoreEvent);
+			}
+			else if(bShowNonTargetProcessEvents)
+			{
+				this->AddCoreTimingEvent(NonTargetProcessEventsBuilder, CpuCoreEvent);
+			}
 
 			return TraceServices::EContextSwitchEnumerationResult::Continue;
 		});
+
+	NonTargetProcessEventsMaxDepth = NonTargetProcessEventsBuilder.GetMaxDepth();
 }
 
 ////////////////////////////////////////////////////////////////////////////////////////////////////
@@ -84,14 +115,21 @@ void FCpuCoreTimingTrack::BuildFilteredDrawState(ITimingEventsTrackDrawStateBuil
 
 			if (bFilterOnlyByEventType)
 			{
+				bool bShowNonTargetProcessEvents = SharedState.AreNonTargetProcessEventsVisible();
 				ContextSwitchesProvider->EnumerateCpuCoreEvents(CoreNumber, Viewport.GetStartTime(), Viewport.GetEndTime(), 
-					[&Builder, ContextSwitchesProvider, FilterEventType, this](const TraceServices::FCpuCoreEvent& CpuCoreEvent)
+					[&Builder, ContextSwitchesProvider, FilterEventType, this, bShowNonTargetProcessEvents](const TraceServices::FCpuCoreEvent& CpuCoreEvent)
 					{
 						if (CpuCoreEvent.SystemThreadId != FilterEventType)
 						{
 							return TraceServices::EContextSwitchEnumerationResult::Continue;
 						}
 
+						uint32 ThreadId;
+						if (!bShowNonTargetProcessEvents && !ContextSwitchesProvider->GetThreadId(CpuCoreEvent.SystemThreadId, ThreadId))
+						{
+							return TraceServices::EContextSwitchEnumerationResult::Continue;
+						}
+
 						this->AddCoreTimingEvent(Builder, CpuCoreEvent);
 
 						return TraceServices::EContextSwitchEnumerationResult::Continue;
@@ -132,8 +170,10 @@ void FCpuCoreTimingTrack::BuildFilteredDrawState(ITimingEventsTrackDrawStateBuil
 		const TraceServices::IContextSwitchesProvider* ContextSwitchesProvider = TraceServices::ReadContextSwitchesProvider(*Session.Get());
 		const FTimingTrackViewport& Viewport = Context.GetViewport();
 
+		bool bShowNonTargetProcessEvents = SharedState.AreNonTargetProcessEventsVisible();
+
 		ContextSwitchesProvider->EnumerateCpuCoreEvents(CoreNumber, Viewport.GetStartTime(), Viewport.GetEndTime(), 
-			[&Builder, ContextSwitchesProvider, &FilterContext, &FilterConfigurator, this](const TraceServices::FCpuCoreEvent& CpuCoreEvent)
+			[&Builder, ContextSwitchesProvider, &FilterContext, &FilterConfigurator, bShowNonTargetProcessEvents, this](const TraceServices::FCpuCoreEvent& CpuCoreEvent)
 			{
 				FilterContext.SetFilterData<double>(static_cast<int32>(EFilterField::StartTime), CpuCoreEvent.Start);
 				FilterContext.SetFilterData<double>(static_cast<int32>(EFilterField::EndTime), CpuCoreEvent.End);
@@ -145,6 +185,12 @@ void FCpuCoreTimingTrack::BuildFilteredDrawState(ITimingEventsTrackDrawStateBuil
 					return TraceServices::EContextSwitchEnumerationResult::Continue;
 				}
 
+				uint32 ThreadId;
+				if (!bShowNonTargetProcessEvents && !ContextSwitchesProvider->GetThreadId(CpuCoreEvent.SystemThreadId, ThreadId))
+				{
+					return TraceServices::EContextSwitchEnumerationResult::Continue;
+				}
+
 				this->AddCoreTimingEvent(Builder, CpuCoreEvent);
 
 				return TraceServices::EContextSwitchEnumerationResult::Continue;
@@ -179,6 +225,10 @@ void FCpuCoreTimingTrack::AddCoreTimingEvent(ITimingEventsTrackDrawStateBuilder&
 
 void FCpuCoreTimingTrack::Draw(const ITimingTrackDrawContext& Context) const
 {
+	const FTimingViewDrawHelper& Helper = *static_cast<const FTimingViewDrawHelper*>(&Context.GetHelper());
+
+	Helper.DrawFadedEvents(*NonTargetProcessEventsDrawState, *this, 1.0f, 0.1f);
+
 	FTimingEventsTrack::Draw(Context);
 }
 
@@ -227,11 +277,23 @@ const TSharedPtr<const ITimingEvent> FCpuCoreTimingTrack::GetEvent(double InTime
 
 	TraceServices::FCpuCoreEvent BestMatchEvent;
 	double Delta = 2 * SecondsPerPixel;
+	bool bShowNonTargetProcessEvents = SharedState.AreNonTargetProcessEventsVisible();
+
+	auto ShouldShowEvent = [bShowNonTargetProcessEvents, ContextSwitchesProvider](const TraceServices::FCpuCoreEvent& CpuCoreEvent)
+	{
+		uint32 ThreadId;
+		if (!bShowNonTargetProcessEvents)
+		{
+			return ContextSwitchesProvider->GetThreadId(CpuCoreEvent.SystemThreadId, ThreadId);
+		}
+
+		return true;
+	};
 
 	ContextSwitchesProvider->EnumerateCpuCoreEvents(CoreNumber, InTime - Delta, InTime + Delta,
-		[InTime, &BestMatchEvent, &Delta](const TraceServices::FCpuCoreEvent& CpuCoreEvent)
+		[InTime, &BestMatchEvent, &Delta, ShouldShowEvent](const TraceServices::FCpuCoreEvent& CpuCoreEvent)
 		{
-			if (CpuCoreEvent.Start <= InTime && CpuCoreEvent.End >= InTime)
+			if (CpuCoreEvent.Start <= InTime && CpuCoreEvent.End >= InTime && ShouldShowEvent(CpuCoreEvent))
 			{
 				BestMatchEvent = CpuCoreEvent;
 				Delta = 0.0f;
@@ -239,14 +301,14 @@ const TSharedPtr<const ITimingEvent> FCpuCoreTimingTrack::GetEvent(double InTime
 			}
 
 			double DeltaLeft = InTime - CpuCoreEvent.End;
-			if (DeltaLeft >= 0 && DeltaLeft < Delta)
+			if (DeltaLeft >= 0 && DeltaLeft < Delta && ShouldShowEvent(CpuCoreEvent))
 			{
 				Delta = DeltaLeft;
 				BestMatchEvent = CpuCoreEvent;
 			}
 
 			double DeltaRight = CpuCoreEvent.Start - InTime;
-			if (DeltaRight >= 0 && DeltaRight < Delta)
+			if (DeltaRight >= 0 && DeltaRight < Delta && ShouldShowEvent(CpuCoreEvent))
 			{
 				Delta = DeltaRight;
 				BestMatchEvent = CpuCoreEvent;
diff --git a/Engine/Source/Developer/TraceInsights/Private/Insights/ContextSwitches/ViewModels/CpuCoreTimingTrack.h b/Engine/Source/Developer/TraceInsights/Private/Insights/ContextSwitches/ViewModels/CpuCoreTimingTrack.h
--- a/Engine/Source/Developer/TraceInsights/Private/Insights/ContextSwitches/ViewModels/CpuCoreTimingTrack.h
+++ b/Engine/Source/Developer/TraceInsights/Private/Insights/ContextSwitches/ViewModels/CpuCoreTimingTrack.h
@@ -28,17 +28,14 @@ class FCpuCoreTimingTrack : public FTimingEventsTrack
 	INSIGHTS_DECLARE_RTTI(FCpuCoreTimingTrack, FTimingEventsTrack)
 
 public:
-	explicit FCpuCoreTimingTrack(FContextSwitchesSharedState& InSharedState, const FString& InName, uint32 InCoreNumber)
-		: FTimingEventsTrack(InName)
-		, SharedState(InSharedState)
-		, CoreNumber(InCoreNumber)
-	{
-	}
+	explicit FCpuCoreTimingTrack(FContextSwitchesSharedState& InSharedState, const FString& InName, uint32 InCoreNumber);
 
 	virtual ~FCpuCoreTimingTrack() {}
 
 	uint32 GetCoreNumber() const { return CoreNumber; }
 
+	virtual void Reset() override;
+
 	virtual void BuildDrawState(ITimingEventsTrackDrawStateBuilder& Builder, const ITimingTrackUpdateContext& Context) override;
 	virtual void BuildFilteredDrawState(ITimingEventsTrackDrawStateBuilder& Builder, const ITimingTrackUpdateContext& Context) override;
 
@@ -61,10 +58,16 @@ protected:
 
 	FString GetThreadName(uint32 InSystemThreadId) const;
 
+	virtual int32 GetMaxDepth() const override { return NonTargetProcessEventsMaxDepth; }
+
 private:
 	FContextSwitchesSharedState& SharedState;
 
 	uint32 CoreNumber;
+
+	TSharedRef<struct FTimingEventsTrackDrawState> NonTargetProcessEventsDrawState;
+
+	int32 NonTargetProcessEventsMaxDepth = -1;
 };
 
 ////////////////////////////////////////////////////////////////////////////////////////////////////
diff --git a/Engine/Source/Developer/TraceInsights/Private/Insights/ViewModels/TimingEventsTrack.cpp b/Engine/Source/Developer/TraceInsights/Private/Insights/ViewModels/TimingEventsTrack.cpp
--- a/Engine/Source/Developer/TraceInsights/Private/Insights/ViewModels/TimingEventsTrack.cpp
+++ b/Engine/Source/Developer/TraceInsights/Private/Insights/ViewModels/TimingEventsTrack.cpp
@@ -81,10 +81,8 @@ void FTimingEventsTrack::PreUpdate(const ITimingTrackUpdateContext& Context)
 
 			Builder.Flush();
 
-			if (Builder.GetMaxDepth() > MaxDepth)
-			{
-				MaxDepth = Builder.GetMaxDepth();
-			}
+			MaxDepth = FMath::Max(MaxDepth, Builder.GetMaxDepth());
+			MaxDepth = FMath::Max(MaxDepth, GetMaxDepth());
 		}
 
 		const TSharedPtr<ITimingEventFilter> EventFilter = Context.GetEventFilter();
diff --git a/Engine/Source/Developer/TraceInsights/Public/Insights/ViewModels/TimingEventsTrack.h b/Engine/Source/Developer/TraceInsights/Public/Insights/ViewModels/TimingEventsTrack.h
--- a/Engine/Source/Developer/TraceInsights/Public/Insights/ViewModels/TimingEventsTrack.h
+++ b/Engine/Source/Developer/TraceInsights/Public/Insights/ViewModels/TimingEventsTrack.h
@@ -93,6 +93,9 @@ protected:
 
 	virtual bool HasCustomFilter() const { return false; }
 
+	/* Can be overriden to force a max depth for the track. */
+	virtual int32 GetMaxDepth() const { return -1; }
+
 private:
 	int32 NumLanes; // number of lanes (sub-tracks)
 	TSharedRef<struct FTimingEventsTrackDrawState> DrawState;
