Changed and added some verifies to abort test case if it failed in Insights EngineTest tests.
Fixed TraceEnable/Disable.Editor case after fix accessing non-zero terminated string builder .
Divided Insights tests to Insights, InsightsEditor and InsightsCore.
Deleted Expected errors as workaround after verbosity of trace error change in TraceSnapshot.LateConnectEditorProfiling.
Changed FStopTracingCommand and FWaitForConnectionStatusCommand.
Changed FInsightsTestUtils::StartTracing function which never made it to the cycle.

diff --git a/Engine/Source/Developer/TraceInsights/Private/Insights/Tests/InsightsTestUtils.cpp b/Engine/Source/Developer/TraceInsights/Private/Insights/Tests/InsightsTestUtils.cpp
--- a/Engine/Source/Developer/TraceInsights/Private/Insights/Tests/InsightsTestUtils.cpp
+++ b/Engine/Source/Developer/TraceInsights/Private/Insights/Tests/InsightsTestUtils.cpp
@@ -112,23 +112,24 @@ bool FInsightsTestUtils::FileContainsString(const FString& PathToFile, const FSt
 
 bool FInsightsTestUtils::StartTracing(FTraceAuxiliary::EConnectionType ConnectionType, const float& Timeout) const
 {
+	bool bStarted = false;
 	if (ConnectionType == FTraceAuxiliary::EConnectionType::Network)
 	{
-		return FTraceAuxiliary::Start(FTraceAuxiliary::EConnectionType::Network, TEXT("localhost"), nullptr);
+		bStarted = FTraceAuxiliary::Start(FTraceAuxiliary::EConnectionType::Network, TEXT("localhost"), nullptr);
 	}
 	else if (ConnectionType == FTraceAuxiliary::EConnectionType::File)
 	{
-		return FTraceAuxiliary::Start(FTraceAuxiliary::EConnectionType::File, nullptr, nullptr);
+		bStarted = FTraceAuxiliary::Start(FTraceAuxiliary::EConnectionType::File, nullptr, nullptr);
 	}
 
 	double TraceVerifyStartTime = FPlatformTime::Seconds();
-	while (FPlatformTime::Seconds() - TraceVerifyStartTime < Timeout)
+	while (bStarted && (FPlatformTime::Seconds() - TraceVerifyStartTime < Timeout))
 	{
+		FPlatformProcess::Sleep(0.5f);
 		if (FTraceAuxiliary::IsConnected())
 		{
 			return true;
 		}
-		FPlatformProcess::Sleep(0.1f);
 	}
 
 	return false;
