[Insights] Add a new persistent setting "Auto Connect" to the Session Browser page that controls if the Editor and Game should auto connect and start tracing if Insights Session Browser is open.

diff --git a/Engine/Programs/UnrealInsights/Config/DefaultUnrealInsightsSessionBrowserSettings.ini b/Engine/Programs/UnrealInsights/Config/DefaultUnrealInsightsSessionBrowserSettings.ini
new file mode 100644
--- /dev/null
+++ b/Engine/Programs/UnrealInsights/Config/DefaultUnrealInsightsSessionBrowserSettings.ini
@@ -0,0 +1,3 @@
+[Insights.SessionBrowser]
+; Whether Insights should signal to an UE application to auto-connect and start tracing when Insights is running.
+AutoConnect=True
\ No newline at end of file
diff --git a/Engine/Source/Developer/TraceInsights/Private/Insights/InsightsManager.cpp b/Engine/Source/Developer/TraceInsights/Private/Insights/InsightsManager.cpp
--- a/Engine/Source/Developer/TraceInsights/Private/Insights/InsightsManager.cpp
+++ b/Engine/Source/Developer/TraceInsights/Private/Insights/InsightsManager.cpp
@@ -513,6 +513,12 @@ FInsightsActionManager& FInsightsManager::GetActionManager()
 
 ////////////////////////////////////////////////////////////////////////////////////////////////////
 
+FInsightsSessionBrowserSettings& FInsightsManager::GetSessionBrowserSettings()
+{
+	return FInsightsManager::Instance->SessionBrowserSettings;
+}
+////////////////////////////////////////////////////////////////////////////////////////////////////
+
 FInsightsSettings& FInsightsManager::GetSettings()
 {
 	return FInsightsManager::Instance->Settings;
diff --git a/Engine/Source/Developer/TraceInsights/Private/Insights/InsightsManager.h b/Engine/Source/Developer/TraceInsights/Private/Insights/InsightsManager.h
--- a/Engine/Source/Developer/TraceInsights/Private/Insights/InsightsManager.h
+++ b/Engine/Source/Developer/TraceInsights/Private/Insights/InsightsManager.h
@@ -11,6 +11,7 @@
 #include "Insights/Common/Stopwatch.h"
 #include "Insights/InsightsCommands.h"
 #include "Insights/InsightsSettings.h"
+#include "Insights/InsightsSessionBrowserSettings.h"
 #include "Insights/IUnrealInsightsModule.h"
 
 ////////////////////////////////////////////////////////////////////////////////////////////////////
@@ -152,6 +153,9 @@ public:
 	/** @return an instance of the main settings. */
 	static FInsightsSettings& GetSettings();
 
+	/** @return an instance of the session browser settings. */
+	static FInsightsSessionBrowserSettings& GetSessionBrowserSettings();
+
 	//////////////////////////////////////////////////
 	// Trace Store
 
@@ -448,6 +452,9 @@ private:
 	/** An instance of the main settings. */
 	FInsightsSettings Settings;
 
+	/** An instance of the session browser settings. */
+	FInsightsSessionBrowserSettings SessionBrowserSettings;
+
 	/** A weak pointer to the Trace Store window. */
 	TWeakPtr<class STraceStoreWindow> TraceStoreWindow;
 
diff --git a/Engine/Source/Developer/TraceInsights/Private/Insights/InsightsSessionBrowserSettings.cpp b/Engine/Source/Developer/TraceInsights/Private/Insights/InsightsSessionBrowserSettings.cpp
new file mode 100644
--- /dev/null
+++ b/Engine/Source/Developer/TraceInsights/Private/Insights/InsightsSessionBrowserSettings.cpp
@@ -0,0 +1,5 @@
+// Copyright Epic Games, Inc. All Rights Reserved.
+
+#include "InsightsSessionBrowserSettings.h"
+
+FInsightsSessionBrowserSettings FInsightsSessionBrowserSettings::Defaults(true);
diff --git a/Engine/Source/Developer/TraceInsights/Private/Insights/InsightsSessionBrowserSettings.h b/Engine/Source/Developer/TraceInsights/Private/Insights/InsightsSessionBrowserSettings.h
new file mode 100644
--- /dev/null
+++ b/Engine/Source/Developer/TraceInsights/Private/Insights/InsightsSessionBrowserSettings.h
@@ -0,0 +1,97 @@
+// Copyright Epic Games, Inc. All Rights Reserved.
+
+#pragma once
+
+#include "CoreMinimal.h"
+#include "Misc/ConfigCacheIni.h"
+#include "Misc/ConfigContext.h"
+#include "ProfilingDebugging/MiscTrace.h"
+
+/** Contains all settings for the Unreal Insights, accessible through the main manager. */
+class FInsightsSessionBrowserSettings
+{
+public:
+	FInsightsSessionBrowserSettings(bool bInIsDefault = false)
+		: bIsEditing(false)
+		, bIsDefault(bInIsDefault)
+		, bAutoConnect(true)
+	{
+		if (!bIsDefault)
+		{
+			LoadFromConfig();
+		}
+	}
+
+	~FInsightsSessionBrowserSettings()
+	{
+	}
+
+	void LoadFromConfig()
+	{
+		if (!FConfigContext::ReadIntoGConfig().Load(TEXT("UnrealInsightsSessionBrowserSettings"), SettingsIni))
+		{
+			return;
+		}
+
+		GConfig->GetBool(TEXT("Insights.SessionBrowser"), TEXT("AutoConnect"), bAutoConnect, SettingsIni);
+	}
+
+	void SaveToConfig()
+	{
+		GConfig->SetBool(TEXT("Insights.SessionBrowser"), TEXT("AutoConnect"), bAutoConnect, SettingsIni);
+
+		GConfig->Flush(false, SettingsIni);
+	}
+
+	void EnterEditMode()
+	{
+		bIsEditing = true;
+	}
+
+	void ExitEditMode()
+	{
+		bIsEditing = false;
+	}
+
+	const bool IsEditing() const
+	{
+		return bIsEditing;
+	}
+
+	const FInsightsSessionBrowserSettings& GetDefaults() const
+	{
+		return Defaults;
+	}
+
+	void ResetToDefaults()
+	{
+		bAutoConnect = Defaults.bAutoConnect;
+	}
+
+	#define SET_AND_SAVE(Option, Value) { if (Option != Value) { Option = Value; SaveToConfig(); } }
+
+	bool IsAutoConnectEnabled() const { return bAutoConnect; }
+	void SetAutoConnect(bool bOnOff) { bAutoConnect = bOnOff; }
+	void SetAndSaveAutoConnect(bool bOnOff) { SET_AND_SAVE(bAutoConnect, bOnOff); }
+
+	#undef SET_AND_SAVE
+
+private:
+	/** Contains default settings. */
+	static FInsightsSessionBrowserSettings Defaults;
+
+	/** Setting filename ini. */
+	FString SettingsIni;
+
+	/** Whether profiler settings is in edit mode. */
+	bool bIsEditing;
+
+	/** Whether this instance contains defaults. */
+	bool bIsDefault;
+
+	//////////////////////////////////////////////////
+	// Actual settings.
+
+	/** Whether Insights should signal to the Editor to auto connect and start tracing when Insights is running */
+	bool bAutoConnect;
+};
diff --git a/Engine/Source/Developer/TraceInsights/Private/Insights/Widgets/SStartPageWindow.cpp b/Engine/Source/Developer/TraceInsights/Private/Insights/Widgets/SStartPageWindow.cpp
--- a/Engine/Source/Developer/TraceInsights/Private/Insights/Widgets/SStartPageWindow.cpp
+++ b/Engine/Source/Developer/TraceInsights/Private/Insights/Widgets/SStartPageWindow.cpp
@@ -31,6 +31,7 @@
 #include "Widgets/Layout/SBox.h"
 #include "Widgets/Layout/SGridPanel.h"
 #include "Widgets/Layout/SScrollBox.h"
+#include "Widgets/Layout/SSeparator.h"
 #include "Widgets/Notifications/SNotificationList.h"
 #include "Widgets/Testing/SStarshipSuite.h"
 #include "Widgets/Text/STextBlock.h"
@@ -56,6 +57,13 @@
 // Driver
 #include "Framework/MetaData/DriverMetaData.h"
 
+#if PLATFORM_WINDOWS
+#include "Windows/AllowWindowsPlatformTypes.h"
+#include <handleapi.h> // for CreateEvent
+#include <synchapi.h> // for CloseHandle
+#include "Windows/HideWindowsPlatformTypes.h"
+#endif
+
 ////////////////////////////////////////////////////////////////////////////////////////////////////
 
 #define LOCTEXT_NAMESPACE "STraceStoreWindow"
@@ -1069,6 +1077,11 @@ void STraceStoreWindow::Construct(const FArguments& InArgs)
 
 	RefreshTraceList();
 
+	if (AutoConnect_IsChecked() == ECheckBoxState::Checked)
+	{
+		EnableAutoConnect();
+	}
+
 	bSetKeyboardFocusOnNextTick = true;
 }
 
@@ -1831,7 +1844,7 @@ TSharedRef<SWidget> STraceStoreWindow::ConstructTraceStoreDirectoryPanel()
 
 TSharedRef<SWidget> STraceStoreWindow::ConstructAutoStartPanel()
 {
-	TSharedRef<SWidget> Widget = SNew(SHorizontalBox)
+	TSharedRef<SHorizontalBox> Box = SNew(SHorizontalBox)
 
 	+ SHorizontalBox::Slot()
 	.AutoWidth()
@@ -1879,7 +1892,35 @@ TSharedRef<SWidget> STraceStoreWindow::ConstructAutoStartPanel()
 		]
 	];
 
-	return Widget;
+#if PLATFORM_WINDOWS
+	Box->AddSlot()
+		.AutoWidth()
+		.Padding(6.0f, 0.0f, 0.0f, 0.0f)
+		.HAlign(HAlign_Left)
+		.VAlign(VAlign_Center)
+		[
+			SNew(SSeparator)
+			.Orientation(Orient_Vertical)
+		];
+
+	Box->AddSlot()
+		.AutoWidth()
+		.Padding(6.0f, 0.0f, 0.0f, 0.0f)
+		.HAlign(HAlign_Left)
+		.VAlign(VAlign_Center)
+		[
+			SNew(SCheckBox)
+			.ToolTipText(LOCTEXT("AutoConnect_Tooltip", "Signal to an UE application to auto-connect and start tracing if Insights is running."))
+			.IsChecked(this, &STraceStoreWindow::AutoConnect_IsChecked)
+			.OnCheckStateChanged(this, &STraceStoreWindow::AutoConnect_OnCheckStateChanged)
+			[
+				SNew(STextBlock)
+				.Text(LOCTEXT("AutoConnect_Text", "Auto-connect"))
+			]
+		];
+#endif // PLATFORM_WINDOWS
+
+	return Box;
 }
 
 ////////////////////////////////////////////////////////////////////////////////////////////////////
@@ -2494,6 +2535,64 @@ void STraceStoreWindow::AutoStart_OnCheckStateChanged(ECheckBoxState NewState)
 
 ////////////////////////////////////////////////////////////////////////////////////////////////////
 
+ECheckBoxState STraceStoreWindow::AutoConnect_IsChecked() const
+{
+	return FInsightsManager::Get()->GetSessionBrowserSettings().IsAutoConnectEnabled() ? ECheckBoxState::Checked : ECheckBoxState::Unchecked;
+}
+
+////////////////////////////////////////////////////////////////////////////////////////////////////
+
+void STraceStoreWindow::AutoConnect_OnCheckStateChanged(ECheckBoxState NewState)
+{
+	if (AutoConnect_IsChecked() == NewState)
+	{
+		return;
+	}
+
+	FInsightsManager::Get()->GetSessionBrowserSettings().SetAndSaveAutoConnect(!FInsightsManager::Get()->GetSessionBrowserSettings().IsAutoConnectEnabled());
+
+	if (FInsightsManager::Get()->GetSessionBrowserSettings().IsAutoConnectEnabled())
+	{
+		EnableAutoConnect();
+	}
+	else
+	{
+		DisableAutoConnect();
+	}
+}
+
+////////////////////////////////////////////////////////////////////////////////////////////////////
+
+void STraceStoreWindow::EnableAutoConnect()
+{
+#if PLATFORM_WINDOWS
+	ensure(AutoConnectEvent == nullptr);
+	// The event is used by runtime to choose when to try to auto-connect.
+	// See FTraceAuxiliary::TryAutoConnect() in \Runtime\Core\Private\ProfilingDebugging\TraceAuxiliary.cpp
+	AutoConnectEvent = CreateEvent(NULL, true, false, TEXT("Local\\UnrealInsightsAutoConnect"));
+	if (AutoConnectEvent == nullptr || GetLastError() != ERROR_SUCCESS)
+	{
+		UE_LOG(TraceInsights, Warning, TEXT("[TraceStore] Failed to create AutoConnect event."));
+
+	}
+#endif
+}
+
+////////////////////////////////////////////////////////////////////////////////////////////////////
+
+void STraceStoreWindow::DisableAutoConnect()
+{
+#if PLATFORM_WINDOWS
+	if (AutoConnectEvent != nullptr)
+	{
+		CloseHandle(AutoConnectEvent);
+		AutoConnectEvent = nullptr;
+	}
+#endif
+}
+
+////////////////////////////////////////////////////////////////////////////////////////////////////
+
 void STraceStoreWindow::Tick(const FGeometry& AllottedGeometry, const double InCurrentTime, const float InDeltaTime)
 {
 	// We need to update the trace list, but not too often.
diff --git a/Engine/Source/Developer/TraceInsights/Private/Insights/Widgets/SStartPageWindow.h b/Engine/Source/Developer/TraceInsights/Private/Insights/Widgets/SStartPageWindow.h
--- a/Engine/Source/Developer/TraceInsights/Private/Insights/Widgets/SStartPageWindow.h
+++ b/Engine/Source/Developer/TraceInsights/Private/Insights/Widgets/SStartPageWindow.h
@@ -11,6 +11,12 @@
 #include "Widgets/SOverlay.h"
 #include "Widgets/Views/SListView.h"
 
+#if PLATFORM_WINDOWS
+#include "Windows/AllowWindowsPlatformTypes.h"
+#include <winnt.h> // for HANDLE
+#include "Windows/HideWindowsPlatformTypes.h"
+#endif
+
 ////////////////////////////////////////////////////////////////////////////////////////////////////
 
 class FActiveTimerHandle;
@@ -362,6 +368,12 @@ private:
 	ECheckBoxState AutoStart_IsChecked() const;
 	void AutoStart_OnCheckStateChanged(ECheckBoxState NewState);
 
+	//////////////////////////////////////////////////
+	// Auto Connect
+	
+	ECheckBoxState AutoConnect_IsChecked() const;
+	void AutoConnect_OnCheckStateChanged(ECheckBoxState NewState);
+
 	//////////////////////////////////////////////////
 	// Trace Store Directory
 
@@ -493,6 +505,9 @@ private:
 	void ShowSuccessMessage(FText& InMessage);
 	void ShowFailMessage(FText& InMessage);
 
+	void EnableAutoConnect();
+	void DisableAutoConnect();
+
 public:
 	/** Parameter that controls the visibility of the confirmation window in case the trace is deleted. */
 	bool bIsDeleteTraceConfirmWindowVisible = true;
@@ -576,6 +591,10 @@ private:
 	bool bStartProcessWithStompMalloc = false;
 
 	bool bSetKeyboardFocusOnNextTick = false;
+
+#if PLATFORM_WINDOWS
+	HANDLE AutoConnectEvent = nullptr;
+#endif
 };
 
 ////////////////////////////////////////////////////////////////////////////////////////////////////
diff --git a/Engine/Source/Programs/UnrealInsights/Private/UserInterfaceCommand.cpp b/Engine/Source/Programs/UnrealInsights/Private/UserInterfaceCommand.cpp
--- a/Engine/Source/Programs/UnrealInsights/Private/UserInterfaceCommand.cpp
+++ b/Engine/Source/Programs/UnrealInsights/Private/UserInterfaceCommand.cpp
@@ -79,8 +79,6 @@ bool CheckSessionBrowserSingleInstance()
 #if PLATFORM_WINDOWS
 	// Create a named event that other processes can detect.
 	// It allows only a single instance of Unreal Insights (Browser Mode).
-	// The event is also used by runtime to choose when to try to auto-connect.
-	// See FTraceAuxiliary::TryAutoConnect() in \Runtime\Core\Private\ProfilingDebugging\TraceAuxiliary.cpp
 	HANDLE SessionBrowserEvent = CreateEvent(NULL, true, false, TEXT("Local\\UnrealInsightsBrowser"));
 	if (SessionBrowserEvent == NULL || GetLastError() == ERROR_ALREADY_EXISTS)
 	{
diff --git a/Engine/Source/Runtime/Core/Private/ProfilingDebugging/TraceAuxiliary.cpp b/Engine/Source/Runtime/Core/Private/ProfilingDebugging/TraceAuxiliary.cpp
--- a/Engine/Source/Runtime/Core/Private/ProfilingDebugging/TraceAuxiliary.cpp
+++ b/Engine/Source/Runtime/Core/Private/ProfilingDebugging/TraceAuxiliary.cpp
@@ -1805,7 +1805,7 @@ void FTraceAuxiliary::TryAutoConnect()
 	{
 		// If we can detect a named event it means UnrealInsights (Browser Mode) is running.
 		// In this case, we try to auto-connect with the Trace Server.
-		HANDLE KnownEvent = ::OpenEvent(EVENT_ALL_ACCESS, false, TEXT("Local\\UnrealInsightsBrowser"));
+		HANDLE KnownEvent = ::OpenEvent(EVENT_ALL_ACCESS, false, TEXT("Local\\UnrealInsightsAutoConnect"));
 		if (KnownEvent != nullptr)
 		{
 			UE_LOG(LogTrace, Display, TEXT("Unreal Insights instance detected, auto-connecting to local trace server..."));
