[Insights] Fix unity build error

diff --git a/Engine/Source/Developer/TraceInsights/Private/Insights/ContextSwitches/ViewModels/CpuCoreTimingTrack.h b/Engine/Source/Developer/TraceInsights/Private/Insights/ContextSwitches/ViewModels/CpuCoreTimingTrack.h
--- a/Engine/Source/Developer/TraceInsights/Private/Insights/ContextSwitches/ViewModels/CpuCoreTimingTrack.h
+++ b/Engine/Source/Developer/TraceInsights/Private/Insights/ContextSwitches/ViewModels/CpuCoreTimingTrack.h
@@ -10,6 +10,7 @@
 ////////////////////////////////////////////////////////////////////////////////////////////////////
 
 class FThreadTrackEvent;
+struct FTimingEventsTrackDrawState;
 
 namespace TraceServices
 {
@@ -65,7 +66,7 @@ private:
 
 	uint32 CoreNumber;
 
-	TSharedRef<struct FTimingEventsTrackDrawState> NonTargetProcessEventsDrawState;
+	TSharedRef<FTimingEventsTrackDrawState> NonTargetProcessEventsDrawState;
 
 	int32 NonTargetProcessEventsMaxDepth = -1;
 };
