Reworked Unreal Insights - Commands - Export - Windows

diff --git a/Engine/Source/Developer/TraceInsights/Private/Insights/Tests/InsightsTestUtils.cpp b/Engine/Source/Developer/TraceInsights/Private/Insights/Tests/InsightsTestUtils.cpp
--- a/Engine/Source/Developer/TraceInsights/Private/Insights/Tests/InsightsTestUtils.cpp
+++ b/Engine/Source/Developer/TraceInsights/Private/Insights/Tests/InsightsTestUtils.cpp
@@ -79,30 +79,32 @@ bool FInsightsTestUtils::AnalyzeTrace(const TCHAR* Path) const
 
 bool FInsightsTestUtils::FileContainsString(const FString& PathToFile, const FString& ExpectedString, const float& Timeout) const
 {
-	if (!FPaths::FileExists(PathToFile))
-	{
-		return false;
-	}
-
 	float StartTime = FPlatformTime::Seconds();
 	while ((FPlatformTime::Seconds() - StartTime) < Timeout)
 	{
-		FString LogFileContents;
-		IPlatformFile& PlatformFile = FPlatformFileManager::Get().GetPlatformFile();
-		TUniquePtr<IFileHandle> FileHandle(PlatformFile.OpenRead(*PathToFile, true)); // Open the file with shared read access
-		if (FileHandle)
+		if (!FPaths::FileExists(PathToFile))
 		{
-			TArray<uint8> FileData;
-			FileData.SetNumUninitialized(FileHandle->Size());
-			FileHandle->Read(FileData.GetData(), FileData.Num());
-			FFileHelper::BufferToString(LogFileContents, FileData.GetData(), FileData.Num());
-
-			if (LogFileContents.Contains(ExpectedString))
+			FPlatformProcess::Sleep(0.1f);
+		}
+		else
+		{
+			FString LogFileContents;
+			IPlatformFile& PlatformFile = FPlatformFileManager::Get().GetPlatformFile();
+			TUniquePtr<IFileHandle> FileHandle(PlatformFile.OpenRead(*PathToFile, true)); // Open the file with shared read access
+			if (FileHandle)
 			{
-				return true;
+				TArray<uint8> FileData;
+				FileData.SetNumUninitialized(FileHandle->Size());
+				FileHandle->Read(FileData.GetData(), FileData.Num());
+				FFileHelper::BufferToString(LogFileContents, FileData.GetData(), FileData.Num());
+
+				if (LogFileContents.Contains(ExpectedString))
+				{
+					return true;
+				}
 			}
+			FPlatformProcess::Sleep(0.1f);
 		}
-		FPlatformProcess::Sleep(0.1f);
 	}
 
 	return false;
diff --git a/Engine/Source/Developer/TraceInsights/Private/Insights/Tests/MiscTests.cpp b/Engine/Source/Developer/TraceInsights/Private/Insights/Tests/MiscTests.cpp
--- a/Engine/Source/Developer/TraceInsights/Private/Insights/Tests/MiscTests.cpp
+++ b/Engine/Source/Developer/TraceInsights/Private/Insights/Tests/MiscTests.cpp
@@ -90,6 +90,39 @@ bool FPagedArrayFilteringTest::RunTest(const FString& Parameters)
 
 IMPLEMENT_SIMPLE_AUTOMATION_TEST(FСommandsExportWindowsTest, "Insights.CommandsExport(Windows)", EAutomationTestFlags::ApplicationContextMask | EAutomationTestFlags::ProductFilter)
 
+void VerifyExportedLines(const FString& ExportReportPath, const FString& CmdLogPath, const FString& Elements, FInsightsTestUtils Utils, FAutomationTestBase* Test, const float& Timeout)
+{
+	float StartTime = FPlatformTime::Seconds();
+	bool bLineFound = false;
+	while ((FPlatformTime::Seconds() - StartTime) < Timeout)
+	{
+		FString FileContent;
+
+		if (FFileHelper::LoadFileToString(FileContent, *ExportReportPath))
+		{
+			TArray<FString> Lines;
+			FileContent.ParseIntoArrayLines(Lines);
+			FString ExpectedResult = FString::Printf(TEXT("Exported %d %s to file"), Lines.Num() - 1, *Elements);
+			bLineFound = Utils.FileContainsString(CmdLogPath, ExpectedResult, 1.0f);
+			if (bLineFound)
+			{
+				Test->TestTrue(FString::Printf(TEXT("Line %s from %s should exists in file: %s"), *ExpectedResult, *ExportReportPath, *CmdLogPath), bLineFound);
+				return;
+			}
+			else
+			{
+				FPlatformProcess::Sleep(0.1f);
+			}
+		}
+		else
+		{
+			FPlatformProcess::Sleep(0.1f);
+		}
+	}
+
+	Test->AddError(FString::Printf(TEXT("VerifyExportedLines timed out while trying to find line from %s"), *ExportReportPath));
+}
+
 bool FСommandsExportWindowsTest::RunTest(const FString& Parameters)
 {
 	TSharedPtr<FInsightsManager> InsightsManager = FInsightsManager::Get();
@@ -97,124 +130,153 @@ bool FСommandsExportWindowsTest::RunTest(const FString& Parameters)
 	FInsightsTestUtils Utils(this);
 
 	const FString StoreDir = InsightsManager->GetStoreDir();
-	const FString SourceTracePath = FPaths::RootDir() / TEXT("EngineTest/SourceAssets/Utrace/Editor_2.utrace");
+	const FString SourceTracePath = FPaths::RootDir() / TEXT("EngineTest/SourceAssets/Utrace/CommandsExportTest_5.4.utrace");
 	const FString SourceExportPath = FPaths::RootDir() / TEXT("EngineTest/SourceAssets/Rsp/export.rsp");
-	const FString StoreTracePath = StoreDir / TEXT("Editor_2.utrace");
+	const FString StoreTracePath = StoreDir / TEXT("CommandsExportTest_5.4.utrace");
+	const FString TestResultsDirPath = TEXT("/TestResults");
+	const FString SingleCommandDirPath = TEXT("/TestResults/SingleCommand");
 	const FString LogResultExportPath = TEXT("/TestResults/export.rsp");
 	const FString CmdLogPath = TEXT("/TestResults/Logs/cmd.log");
+	const FString CmdThreadLogPath = TEXT("/TestResults/Logs/cmd_treads.log");
+	const FString CmdTimersLogPath = TEXT("/TestResults/Logs/cmd_timers.log");
+	const FString CmdTimingEventsLogPath = TEXT("/TestResults/Logs/cmd_timing_events.log");
+	const FString CmdTimingEventsNonDefaultLogPath = TEXT("/TestResults/Logs/cmd_timing_non_default.log");
+	const FString CmdRegionsLogPath = TEXT("/TestResults/Logs/cmd_regions.log");
+	const FString CmdExportLogPath = TEXT("/TestResults/Logs/cmd_export.log");
 	const FString ExportThreadsTask = TEXT("TimingInsights.ExportThreads /TestResults/SingleCommand/Threads.csv");
 	const FString ExportTimersTask = TEXT("TimingInsights.ExportTimers /TestResults/SingleCommand/Timers.csv");
-	const FString ExportTimingEventsTask = TEXT("TimingInsights.ExportTimingEvents /TestResults/SingleCommand/TimingEvents.TXT");
-	const FString ExportExportTimingEventsBorderedTask = TEXT("TimingInsights.ExportTimingEvents /TestResults/SingleCommand/TimingEventsNonDefault.CSV  -columns=ThreadId,ThreadName,TimerId,TimerName,StartTime,EndTime,Duration,Depth -threads=GameThread -timers=* -startTime=10 -endTime=20");
-	
-	TestTrue("Trace in project exists", PlatformFile.FileExists(*SourceTracePath));
-	TestTrue("Export in project exists", PlatformFile.FileExists(*SourceExportPath));
-	TestFalse("Trace in store should not exists before copy", PlatformFile.FileExists(*StoreTracePath));
+	const FString ExportTimingEventsTask = TEXT("TimingInsights.ExportTimingEvents /TestResults/SingleCommand/TimingEvents.csv");
+	const FString ExportTimingEventsBorderedTask = TEXT("TimingInsights.ExportTimingEvents /TestResults/SingleCommand/TimingEventsNonDefault.CSV  -columns=ThreadId,ThreadName,TimerId,TimerName,StartTime,EndTime,Duration,Depth -threads=GameThread -timers=* -startTime=10 -endTime=20");
+	const FString ExportTimerStatisticsTask = TEXT("TimingInsights.ExportTimerStatistics  /TestResults/SingleCommand/*.csv -region=* -threads=GPU");
+	const FString ExportThreadsReportPath = TEXT("/TestResults/SingleCommand/Threads.csv");
+	const FString ExportTimersReportPath = TEXT("/TestResults/SingleCommand/Timers.csv");
+	const FString ExportTimingEventsReportPath = TEXT("/TestResults/SingleCommand/TimingEvents.csv");
+	const FString ExportTimingEventsNonDefaultReportPath = TEXT("/TestResults/SingleCommand/TimingEventsNonDefault.csv");
+	const FString ExportTimerStatisticsReportPath = TEXT("/TestResults/SingleCommand/TimerStatistics.csv");
+	bool bLineFound = false;
+	const float Timeout = 50.0;
+
+	UTEST_TRUE("Trace in project exists", PlatformFile.FileExists(*SourceTracePath));
+	UTEST_TRUE("Export in project exists", PlatformFile.FileExists(*SourceExportPath));
+	UTEST_FALSE("Trace in store should not exists before copy", PlatformFile.FileExists(*StoreTracePath));
 	
 	PlatformFile.CopyFile(*StoreTracePath, *SourceTracePath);
-	if (!PlatformFile.FileExists(*StoreTracePath))
+	UTEST_TRUE(TEXT("Trace in store should exists after copy"), PlatformFile.FileExists(*StoreTracePath));
+
+	if (PlatformFile.DirectoryExists(*TestResultsDirPath))
 	{
-		AddError("Trace in store should exists after copy");
-		return false;
+		AddWarning(TEXT("The TestResults directory already exists. Deleting to avoid undefined behavior"));
+		IFileManager::Get().DeleteDirectory(*TestResultsDirPath, false, true);
+		float StartTime = FPlatformTime::Seconds();
+		while ((FPlatformTime::Seconds() - StartTime) < Timeout)
+		{
+			if (!PlatformFile.DirectoryExists(*TestResultsDirPath))
+			{ 
+				break;
+			}
+			FPlatformProcess::Sleep(0.1f);
+		}
+
+		UTEST_FALSE("Previously created TestResults directory is not deleted. Execution of the next steps will cause unstable behavior.", PlatformFile.DirectoryExists(*TestResultsDirPath));
 	}
 
     // ExportThreads
-	FString InsightsParameters = FString::Printf(TEXT("-OpenTraceFile=%s -ABSLOG=\"%s\" -AutoQuit -NoUI  -ExecOnAnalysisCompleteCmd=\"%s\" -log"), *StoreTracePath, *CmdLogPath, *ExportThreadsTask);
+	FString InsightsParameters = FString::Printf(TEXT("-OpenTraceFile=\"%s\" -ABSLOG=\"%s\" -AutoQuit -NoUI  -ExecOnAnalysisCompleteCmd=\"%s\" -log"), *StoreTracePath, *CmdThreadLogPath, *ExportThreadsTask);
 	InsightsManager->OpenUnrealInsights(*InsightsParameters);
 	
-	FString ExpectedThreadsElements[5] = { TEXT("Id,Name,Group"), TEXT("1196447025,GPU1,GPU"), TEXT("1196447026,GPU2,GPU"), TEXT("2,GameThread,"), TEXT("88,RTHeartBeat 0,Render") };
-	FString ExpectedResult = TEXT("Exported 107 threads to file");
-	bool bLineFound = Utils.FileContainsString(CmdLogPath, ExpectedResult, 15.0f);
-	TestTrue(FString::Printf(TEXT("Line %s should exists in file"), *ExpectedResult), bLineFound);
-
-	for (int i = 0; i < 5; i++)
-	{
-		bLineFound = Utils.FileContainsString(TEXT("/TestResults/SingleCommand/Threads.csv"), ExpectedThreadsElements[i], 15.0f);
-		TestTrue(FString::Printf(TEXT("Line %s should exists in file"), *ExpectedThreadsElements[i]), bLineFound);
-	}
+	VerifyExportedLines(ExportThreadsReportPath, CmdThreadLogPath, TEXT("threads"), Utils, this, Timeout);
 
 	// ExportTimers
-	InsightsParameters = FString::Printf(TEXT("-OpenTraceFile=%s -ABSLOG=\"%s\" -AutoQuit -NoUI  -ExecOnAnalysisCompleteCmd=\"%s\" -log"), *StoreTracePath, *CmdLogPath, *ExportTimersTask);
+	InsightsParameters = FString::Printf(TEXT("-OpenTraceFile=\"%s\" -ABSLOG=\"%s\" -AutoQuit -NoUI  -ExecOnAnalysisCompleteCmd=\"%s\" -log"), *StoreTracePath, *CmdTimersLogPath, *ExportTimersTask);
 	InsightsManager->OpenUnrealInsights(*InsightsParameters);
 
-	ExpectedResult = TEXT("Exported 1942 timers to file");
-	bLineFound = Utils.FileContainsString(CmdLogPath, ExpectedResult, 15.0f);
-	TestTrue(FString::Printf(TEXT("Line %s should exists in file"), *ExpectedResult), bLineFound);
-
-	bLineFound = Utils.FileContainsString(TEXT("/TestResults/SingleCommand/Timers.csv"), TEXT("Id,Type,Name,File,Line"), 15.0f);
-	TestTrue(FString::Printf(TEXT("Line %s should exists in file"), TEXT("Id,Type,Name,File,Line")), bLineFound);
-
-	bLineFound = Utils.FileContainsString(TEXT("/TestResults/SingleCommand/Timers.csv"), TEXT("15,CPU,LoadModule_StorageServerClient"), 15.0f);
-	TestTrue(FString::Printf(TEXT("Line %s should exists in file"), TEXT("15,CPU,LoadModule_StorageServerClient")), bLineFound);
+	VerifyExportedLines(ExportTimersReportPath, CmdTimersLogPath, TEXT("timers"), Utils, this, Timeout);
 
 	// UI verification cannot be executed in that case. 
 	
 	// ExportTimingEvents
-	InsightsParameters = FString::Printf(TEXT("-OpenTraceFile=%s -ABSLOG=\"%s\" -AutoQuit -NoUI  -ExecOnAnalysisCompleteCmd=\"%s\" -log"), *StoreTracePath, *CmdLogPath, *ExportTimingEventsTask);
+	InsightsParameters = FString::Printf(TEXT("-OpenTraceFile=\"%s\" -ABSLOG=\"%s\" -AutoQuit -NoUI  -ExecOnAnalysisCompleteCmd=\"%s\" -log"), *StoreTracePath, *CmdTimingEventsLogPath, *ExportTimingEventsTask);
 	InsightsManager->OpenUnrealInsights(*InsightsParameters);
 
-	FString ExpectedTimingEventsElements[3] = { TEXT("ThreadId	TimerId	StartTime	EndTime	Depth"), TEXT("1196447025	1916	46.510678	46.510955	0"), TEXT("1196447025	1916	46.51073	46.510955	1") };
-	ExpectedResult = TEXT("Exported 2381149 timing events to file");
-	bLineFound = Utils.FileContainsString(CmdLogPath, ExpectedResult, 15.0f);
-	TestTrue(FString::Printf(TEXT("Line %s should exists in file"), *ExpectedResult), bLineFound);
-
-	for (int i = 0; i < 3; i++)
-	{
-		bLineFound = Utils.FileContainsString(TEXT("/TestResults/SingleCommand/TimingEvents.TXT"), ExpectedTimingEventsElements[i], 15.0f);
-		TestTrue(FString::Printf(TEXT("Line %s should exists in file"), *ExpectedTimingEventsElements[i]), bLineFound);
-	}
+	VerifyExportedLines(ExportTimingEventsReportPath, CmdTimingEventsLogPath, TEXT("timing events"), Utils, this, Timeout);
 
 	// ExportTimingEventsNonDefault
-	InsightsParameters = FString::Printf(TEXT("-OpenTraceFile=%s -ABSLOG=\"%s\" -AutoQuit -NoUI  -ExecOnAnalysisCompleteCmd=\"%s\" -log"), *StoreTracePath, *CmdLogPath, *ExportExportTimingEventsBorderedTask);
+	InsightsParameters = FString::Printf(TEXT("-OpenTraceFile=\"%s\" -ABSLOG=\"%s\" -AutoQuit -NoUI  -ExecOnAnalysisCompleteCmd=\"%s\" -log"), *StoreTracePath, *CmdTimingEventsNonDefaultLogPath, *ExportTimingEventsBorderedTask);
 	InsightsManager->OpenUnrealInsights(*InsightsParameters);
 
-	FString ExpectedTimingEventsBorderedElements[3] = { TEXT("ThreadId,ThreadName,TimerId,TimerName,StartTime,EndTime,Duration,Depth"), 
-		TEXT("2,GameThread,0,FEngineLoop::PreInitPreStartupScreen,0.3376052,41.5370876,41.199482400,0"), 
-		TEXT("2,GameThread,182,InitDerivedData,1.7132565,40.2586897,38.545433200,1") };
-	ExpectedResult = TEXT("Exported 2 timing events to file");
-
-	bLineFound = Utils.FileContainsString(CmdLogPath, ExpectedResult, 15.0f);
-	TestTrue(FString::Printf(TEXT("Line %s should exists in file"), *ExpectedResult), bLineFound);
+	const TArray<FString> ExpectedTimingEventsBorderedElements = { TEXT("ThreadId"), TEXT("ThreadName"), TEXT("TimerId"), TEXT("TimerName"), TEXT("StartTime"), TEXT("EndTime"), TEXT("Duration"), TEXT("Depth"),
+		TEXT("2,GameThread,1,FEngineLoop::PreInitPreStartupScreen")
+	};
 
-	for (int i = 0; i < 3; i++)
+	for (int i = 0; i < ExpectedTimingEventsBorderedElements.Num(); i++)
 	{
-		bLineFound = Utils.FileContainsString(TEXT("/TestResults/SingleCommand/TimingEventsNonDefault.csv"), ExpectedTimingEventsBorderedElements[i], 15.0f);
-		TestTrue(FString::Printf(TEXT("Line %s should exists in file"), *ExpectedTimingEventsBorderedElements[i]), bLineFound);
+		bLineFound = Utils.FileContainsString(ExportTimingEventsNonDefaultReportPath, ExpectedTimingEventsBorderedElements[i], Timeout);
+		TestTrue(FString::Printf(TEXT("Line %s should exists in file: %s"), *ExpectedTimingEventsBorderedElements[i], *ExportTimingEventsNonDefaultReportPath), bLineFound);
 	}
 
+	VerifyExportedLines(ExportTimingEventsNonDefaultReportPath, CmdTimingEventsNonDefaultLogPath, TEXT("timing events"), Utils, this, Timeout);
+
+	// Export regions
+	InsightsParameters = FString::Printf(TEXT("-OpenTraceFile=\"%s\" -ABSLOG=\"%s\" -AutoQuit -NoUI  -ExecOnAnalysisCompleteCmd=\"%s\" -log"), *StoreTracePath, *CmdRegionsLogPath, *ExportTimerStatisticsTask);
+	InsightsManager->OpenUnrealInsights(*InsightsParameters);
+
+	bLineFound = Utils.FileContainsString(CmdRegionsLogPath, TEXT("Exported timing statistics for 10 regions"), Timeout);
+	TestTrue(FString::Printf(TEXT("Line %s should exists in file: %s"), TEXT("Exported timing statistics for 10 regions"), *CmdRegionsLogPath), bLineFound);
+
+	TArray<FString> RegionFiles;
+	FString FilePattern = TEXT("*.csv"); // Change the pattern to filter .csv files
+	IFileManager::Get().FindFiles(RegionFiles, *FPaths::Combine(*SingleCommandDirPath, *FilePattern), true, false);
+	TestEqual(FString::Printf(TEXT("Should be 14 csv files but found %d"), RegionFiles.Num()), RegionFiles.Num(), 14);
+
 	// Export.rsp
 	PlatformFile.CopyFile(*LogResultExportPath, *SourceExportPath);
 	TestTrue("Rsp in log directory should exists after copy", PlatformFile.FileExists(*LogResultExportPath));
  
-	InsightsParameters = FString::Printf(TEXT("-OpenTraceFile=%s -ABSLOG=\"%s\" -AutoQuit -NoUI  -ExecOnAnalysisCompleteCmd=\"@=/TestResults/export.rsp\" -log"), *StoreTracePath, *CmdLogPath);
+	InsightsParameters = FString::Printf(TEXT("-OpenTraceFile=\"%s\" -ABSLOG=\"%s\" -AutoQuit -NoUI  -ExecOnAnalysisCompleteCmd=\"@=/TestResults/export.rsp\" -log"), *StoreTracePath, *CmdExportLogPath);
 	InsightsManager->OpenUnrealInsights(*InsightsParameters);
 
-	FString ExpectedExportElements[14] = { TEXT("(\"/TestResults/RSPtest/CSV/Threads_rsp.csv\")"), 
-		TEXT("(\"/TestResults/RSPtest/TSV/Threads_rsp.tsv\")"),
-		TEXT("(\"/TestResults/RSPtest/TXT/Threads_rsp.txt\")"),
-		TEXT("(\"/TestResults/RSPtest/CSV/Timers_rsp.csv\")"),
-		TEXT("(\"/TestResults/RSPtest/TSV/Timers_rsp.tsv\")"),
-		TEXT("(\"/TestResults/RSPtest/TXT/Timers_rsp.txt\")"),
-		TEXT("(\"/TestResults/RSPtest/CSV/TimingEvents_rsp.csv\")"),
-		TEXT("(\"/TestResults/RSPtest/TSV/TimingEvents_rsp.tsv\")"),
-		TEXT("(\"/TestResults/RSPtest/TXT/TimingEvents_rsp.txt\")"), 
-		TEXT("Exported 107 threads to file"),
-		TEXT("Exported 1942 threads to file"),
-		TEXT("Exported 538 threads to file"),
-		TEXT("Exported 30 threads to file"),
-		TEXT("Exported 1535 threads to file"),
+	const TArray<FString> ExpectedThreadsElementsRsp = { TEXT("/TestResults/RSPtest/CSV/Threads_rsp.csv"),
+		TEXT("/TestResults/RSPtest/TSV/Threads_rsp.tsv"),
+		TEXT("/TestResults/RSPtest/TXT/Threads_rsp.txt"),
+	};
+
+	for (int i = 0; i < ExpectedThreadsElementsRsp.Num(); i++)
+	{
+		bLineFound = Utils.FileContainsString(CmdExportLogPath, ExpectedThreadsElementsRsp[i], Timeout);
+		TestTrue(FString::Printf(TEXT("Line %s should exists in file: %s"), *ExpectedThreadsElementsRsp[i], *CmdExportLogPath), bLineFound);
+
+		VerifyExportedLines(ExpectedThreadsElementsRsp[i], CmdExportLogPath, TEXT("threads"), Utils, this, Timeout);
+	}
+
+	const TArray<FString> ExpectedTimersElementsRsp = { TEXT("/TestResults/RSPtest/CSV/Timers_rsp.csv"),
+		TEXT("/TestResults/RSPtest/TSV/Timers_rsp.tsv"),
+		TEXT("/TestResults/RSPtest/TXT/Timers_rsp.txt"),
 	};
 
-	for (int i = 0; i < 9; i++)
+	for (int i = 0; i < ExpectedTimersElementsRsp.Num(); i++)
 	{
-		bLineFound = Utils.FileContainsString(CmdLogPath, ExpectedExportElements[i], 15.0f);
-		TestTrue(FString::Printf(TEXT("Line %s should exists in file"), *ExpectedExportElements[i]), bLineFound);
+		bLineFound = Utils.FileContainsString(CmdExportLogPath, ExpectedTimersElementsRsp[i], Timeout);
+		TestTrue(FString::Printf(TEXT("Line %s should exists in file: %s"), *ExpectedTimersElementsRsp[i], *CmdExportLogPath), bLineFound);
+
+		VerifyExportedLines(ExpectedTimersElementsRsp[i], CmdExportLogPath, TEXT("timers"), Utils, this, Timeout);
 	}
 
+	const TArray<FString> ExpectedTimingEventsRsp = { TEXT("/TestResults/RSPtest/CSV/TimingEvents_rsp.csv"),
+		TEXT("/TestResults/RSPtest/TSV/TimingEvents_rsp.tsv"),
+		TEXT("/TestResults/RSPtest/TXT/TimingEvents_rsp.txt"),
+	};
+
+	for (int i = 0; i < ExpectedTimingEventsRsp.Num(); i++)
+	{
+		bLineFound = Utils.FileContainsString(CmdExportLogPath, ExpectedTimingEventsRsp[i], Timeout);
+		TestTrue(FString::Printf(TEXT("Line %s should exists in file: %s"), *ExpectedTimingEventsRsp[i], *CmdExportLogPath), bLineFound);
+
+		VerifyExportedLines(ExpectedTimingEventsRsp[i], CmdExportLogPath, TEXT("timing events"), Utils, this, Timeout);
+	}
+	
 	IFileManager::Get().Delete(*StoreTracePath);
-	check(IFileManager::Get().DeleteDirectory(TEXT("/TestResults"), false, true));
 
 	return true;
 }
 
-#endif // !WITH_EDITOR
+#endif // !WITH_EDITOR
\ No newline at end of file
