[Insights] Fix unity compile error.

diff --git a/Engine/Source/Developer/TraceInsights/Private/Insights/Widgets/SStartPageWindow.h b/Engine/Source/Developer/TraceInsights/Private/Insights/Widgets/SStartPageWindow.h
--- a/Engine/Source/Developer/TraceInsights/Private/Insights/Widgets/SStartPageWindow.h
+++ b/Engine/Source/Developer/TraceInsights/Private/Insights/Widgets/SStartPageWindow.h
@@ -11,12 +11,6 @@
 #include "Widgets/SOverlay.h"
 #include "Widgets/Views/SListView.h"
 
-#if PLATFORM_WINDOWS
-#include "Windows/AllowWindowsPlatformTypes.h"
-#include <winnt.h> // for HANDLE
-#include "Windows/HideWindowsPlatformTypes.h"
-#endif
-
 ////////////////////////////////////////////////////////////////////////////////////////////////////
 
 class FActiveTimerHandle;
@@ -593,7 +587,7 @@ private:
 	bool bSetKeyboardFocusOnNextTick = false;
 
 #if PLATFORM_WINDOWS
-	HANDLE AutoConnectEvent = nullptr;
+	void* AutoConnectEvent = nullptr;
 #endif
 };
 
