[Insights] Fixed build issue of Unreal Insights  with -NoPCH -DisableUnity flags

diff --git a/Engine/Source/Developer/TraceInsights/Private/Insights/Tests/MiscTests.cpp b/Engine/Source/Developer/TraceInsights/Private/Insights/Tests/MiscTests.cpp
--- a/Engine/Source/Developer/TraceInsights/Private/Insights/Tests/MiscTests.cpp
+++ b/Engine/Source/Developer/TraceInsights/Private/Insights/Tests/MiscTests.cpp
@@ -3,6 +3,8 @@
 #include "Misc/AutomationTest.h"
 //awful hack since SlabAllocator is private
 #include "Developer/TraceServices/Private/Common/SlabAllocator.h"
+#include "HAL/FileManager.h"
+#include "Misc/FileHelper.h"
 #include "HAL/PlatformFileManager.h"
 #include "Common/PagedArray.h"
 #include "Insights/InsightsManager.h"
