[Insights] Minor fixes (incorrect include file; order of include files; "Analysis completed" log messages).

#rnx

diff --git a/Engine/Source/Developer/TraceInsights/Private/Insights/ViewModels/RegionsTimingTrack.h b/Engine/Source/Developer/TraceInsights/Private/Insights/ViewModels/RegionsTimingTrack.h
--- a/Engine/Source/Developer/TraceInsights/Private/Insights/ViewModels/RegionsTimingTrack.h
+++ b/Engine/Source/Developer/TraceInsights/Private/Insights/ViewModels/RegionsTimingTrack.h
@@ -1,16 +1,21 @@
 ﻿// Copyright Epic Games, Inc. All Rights Reserved.
+
 #pragma once
-#include "FileActivityTimingTrack.h"
+
+#include "Framework/Commands/Commands.h"
+#include "Insights/ITimingViewExtender.h"
 #include "Insights/ViewModels/TimingEventsTrack.h"
 #include "TraceServices/Model/Regions.h"
 
-////////////////////////////////////////////////////////////////////////////////////////////////////
+class STimingView;
 
 namespace Insights
 {
 
 class FTimingRegionsTrack;
 
+////////////////////////////////////////////////////////////////////////////////////////////////////
+
 class FTimingRegionsViewCommands : public TCommands<FTimingRegionsViewCommands>
 {
 public:
@@ -22,6 +27,8 @@ public:
 	TSharedPtr<FUICommandInfo> ShowHideRegionTrack;
 };
 
+////////////////////////////////////////////////////////////////////////////////////////////////////
+
 class FTimingRegionsSharedState : public Insights::ITimingViewExtender, public TSharedFromThis<FTimingRegionsSharedState>
 {
 	friend class FTimingRegionsTrack;
@@ -48,6 +55,8 @@ private:
 	bool bShowHideRegionsTrack = true;
 };
 
+////////////////////////////////////////////////////////////////////////////////////////////////////
+
 class FTimingRegionsTrack : public FTimingEventsTrack
 {
 	INSIGHTS_DECLARE_RTTI(FTimingRegionsTrack, FTimingEventsTrack)
@@ -73,4 +82,6 @@ private:
 	FTimingRegionsSharedState& SharedState;
 };
 
+////////////////////////////////////////////////////////////////////////////////////////////////////
+
 } // namespace Insights
diff --git a/Engine/Source/Developer/TraceInsights/Private/Insights/Widgets/SStartPageWindow.cpp b/Engine/Source/Developer/TraceInsights/Private/Insights/Widgets/SStartPageWindow.cpp
--- a/Engine/Source/Developer/TraceInsights/Private/Insights/Widgets/SStartPageWindow.cpp
+++ b/Engine/Source/Developer/TraceInsights/Private/Insights/Widgets/SStartPageWindow.cpp
@@ -3,23 +3,22 @@
 #include "SStartPageWindow.h"
 
 #include "DesktopPlatformModule.h"
-#include "IPAddress.h"
-#include "SlateOptMacros.h"
-#include "SocketSubsystem.h"
 #include "Framework/Application/SlateApplication.h"
+#include "Framework/MetaData/DriverMetaData.h"
 #include "Framework/MultiBox/MultiBoxBuilder.h"
 #include "HAL/FileManager.h"
 #include "HAL/PlatformProcess.h"
 #include "Internationalization/Text.h"
+#include "IPAddress.h"
 #include "Logging/MessageLog.h"
 #include "Misc/MessageDialog.h"
 #include "Misc/PathViews.h"
+#include "SlateOptMacros.h"
+#include "SocketSubsystem.h"
 #include "Styling/AppStyle.h"
 #include "Styling/StyleColors.h"
 #include "Trace/ControlClient.h"
 #include "Trace/StoreClient.h"
-#include "Widgets/SBoxPanel.h"
-#include "Widgets/SToolTip.h"
 #include "Widgets/Colors/SColorBlock.h"
 #include "Widgets/Images/SImage.h"
 #include "Widgets/Input/SButton.h"
@@ -33,30 +32,29 @@
 #include "Widgets/Layout/SScrollBox.h"
 #include "Widgets/Layout/SSeparator.h"
 #include "Widgets/Notifications/SNotificationList.h"
+#include "Widgets/SBoxPanel.h"
+#include "Widgets/SToolTip.h"
 #include "Widgets/Testing/SStarshipSuite.h"
 #include "Widgets/Text/STextBlock.h"
 #include "Widgets/Views/STileView.h"
 
 #if WITH_EDITOR
-	#include "EngineAnalytics.h"
 	#include "AnalyticsEventAttribute.h"
+	#include "EngineAnalytics.h"
 	#include "Interfaces/IAnalyticsProvider.h"
 #endif // WITH_EDITOR
 
 // Insights
+#include "Insights/Common/Stopwatch.h"
+#include "Insights/ImportTool/TableImportTool.h"
 #include "Insights/InsightsManager.h"
 #include "Insights/InsightsStyle.h"
 #include "Insights/Log.h"
-#include "Insights/Version.h"
-#include "Insights/Common/Stopwatch.h"
-#include "Insights/ImportTool/TableImportTool.h"
 #include "Insights/StoreService/StoreBrowser.h"
+#include "Insights/Version.h"
 #include "Insights/Widgets/SInsightsSettings.h"
 #include "Insights/Widgets/SLazyToolTip.h"
 
-// Driver
-#include "Framework/MetaData/DriverMetaData.h"
-
 #if PLATFORM_WINDOWS
 #include "Windows/AllowWindowsPlatformTypes.h"
 #include <handleapi.h> // for CreateEvent
@@ -2573,7 +2571,6 @@ void STraceStoreWindow::EnableAutoConnect()
 	if (AutoConnectEvent == nullptr || GetLastError() != ERROR_SUCCESS)
 	{
 		UE_LOG(TraceInsights, Warning, TEXT("[TraceStore] Failed to create AutoConnect event."));
-
 	}
 #endif
 }
@@ -4212,4 +4209,4 @@ END_SLATE_FUNCTION_BUILD_OPTIMIZATION
 
 ////////////////////////////////////////////////////////////////////////////////////////////////////
 
-#undef LOCTEXT_NAMESPACE
\ No newline at end of file
+#undef LOCTEXT_NAMESPACE
diff --git a/Engine/Source/Developer/TraceServices/Private/Model/AllocationsProvider.cpp b/Engine/Source/Developer/TraceServices/Private/Model/AllocationsProvider.cpp
--- a/Engine/Source/Developer/TraceServices/Private/Model/AllocationsProvider.cpp
+++ b/Engine/Source/Developer/TraceServices/Private/Model/AllocationsProvider.cpp
@@ -2328,7 +2328,8 @@ void FAllocationsProvider::EditOnAnalysisCompleted(double Time)
 			++NumHeapSpecs;
 		}
 	}
-	UE_LOG(LogTraceServices, Log, TEXT("[MemAlloc] Analysis Completed (%llu events, %llu allocs, %llu frees, %llu heaps, %d heap specs)"), TotalEventCount, AllocCount, FreeCount, HeapCount, NumHeapSpecs);
+	UE_LOG(LogTraceServices, Log, TEXT("[MemAlloc] Analysis completed (%llu events, %llu allocs, %llu frees, %llu heaps, %d heap specs)."),
+		TotalEventCount, AllocCount, FreeCount, HeapCount, NumHeapSpecs);
 }
 
 ////////////////////////////////////////////////////////////////////////////////////////////////////
diff --git a/Engine/Source/Developer/TraceServices/Private/Model/MetadataProvider.cpp b/Engine/Source/Developer/TraceServices/Private/Model/MetadataProvider.cpp
--- a/Engine/Source/Developer/TraceServices/Private/Model/MetadataProvider.cpp
+++ b/Engine/Source/Developer/TraceServices/Private/Model/MetadataProvider.cpp
@@ -558,7 +558,7 @@ void FMetadataProvider::OnAnalysisCompleted()
 		UE_LOG(LogTraceServices, Error, TEXT("[Meta] Restore scope errors: %u"), RestoreScopeErrors);
 	}
 
-	UE_LOG(LogTraceServices, Log, TEXT("[Meta] Analysis Completed (%llu events; %llu alloc events; %llu allocs; %llu bytes)"),
+	UE_LOG(LogTraceServices, Log, TEXT("[Meta] Analysis completed (%llu events, %llu alloc events, %llu allocs, %llu bytes)."),
 		EventCount, AllocationEventCount, AllocationCount, TotalAllocatedMemory);
 }
 
