Added some log for debug in TraceInsights MiscTests.cpp

diff --git a/Engine/Source/Developer/TraceInsights/Private/Insights/Tests/MiscTests.cpp b/Engine/Source/Developer/TraceInsights/Private/Insights/Tests/MiscTests.cpp
--- a/Engine/Source/Developer/TraceInsights/Private/Insights/Tests/MiscTests.cpp
+++ b/Engine/Source/Developer/TraceInsights/Private/Insights/Tests/MiscTests.cpp
@@ -96,6 +96,8 @@ void VerifyExportedLines(const FString& ExportReportPath, const FString& CmdLogP
 {
 	float StartTime = FPlatformTime::Seconds();
 	bool bLineFound = false;
+	FString ExpectedResult;
+
 	while ((FPlatformTime::Seconds() - StartTime) < Timeout)
 	{
 		FString FileContent;
@@ -104,11 +106,11 @@ void VerifyExportedLines(const FString& ExportReportPath, const FString& CmdLogP
 		{
 			TArray<FString> Lines;
 			FileContent.ParseIntoArrayLines(Lines);
-			FString ExpectedResult = FString::Printf(TEXT("Exported %d %s to file"), Lines.Num() - 1, *Elements);
+			ExpectedResult = FString::Printf(TEXT("Exported %d %s to file"), Lines.Num() - 1, *Elements);
 			bLineFound = Utils.FileContainsString(CmdLogPath, ExpectedResult, 1.0f);
 			if (bLineFound)
 			{
-				Test->TestTrue(FString::Printf(TEXT("Line %s from %s should exists in file: %s"), *ExpectedResult, *ExportReportPath, *CmdLogPath), bLineFound);
+				Test->TestTrue(FString::Printf(TEXT("Line '%s' from '%s' should exists in file: '%s'"), *ExpectedResult, *ExportReportPath, *CmdLogPath), bLineFound);
 				return;
 			}
 			else
@@ -122,7 +124,7 @@ void VerifyExportedLines(const FString& ExportReportPath, const FString& CmdLogP
 		}
 	}
 
-	Test->AddError(FString::Printf(TEXT("VerifyExportedLines timed out while trying to find line from %s"), *ExportReportPath));
+	Test->AddError(FString::Printf(TEXT("VerifyExportedLines timed out while trying to find line '%s' from '%s'"), *ExpectedResult, *ExportReportPath));
 }
 
 bool FСommandsExportWindowsTest::RunTest(const FString& Parameters)
