[Insights] Fixed crash when using "MENU > Open Trace File > Open in New/Same Instance..." option when analyzing a utrace file. The issue did not reproduce when analyzing a trace from a Trace Store.

diff --git a/Engine/Source/Developer/TraceInsights/Private/Insights/InsightsManager.cpp b/Engine/Source/Developer/TraceInsights/Private/Insights/InsightsManager.cpp
--- a/Engine/Source/Developer/TraceInsights/Private/Insights/InsightsManager.cpp
+++ b/Engine/Source/Developer/TraceInsights/Private/Insights/InsightsManager.cpp
@@ -446,6 +446,10 @@ void FInsightsManager::OnSessionInfoTabClosed(TSharedRef<SDockTab> TabBeingClose
 FString FInsightsManager::GetStoreDir()
 {
 	using namespace UE::Trace;
+	if (!StoreClient.IsValid())
+	{
+		return FString();
+	}
 	FScopeLock _(&StoreClientCriticalSection);
 	const FStoreClient::FStatus* Status = StoreClient->GetStatus();
 	return Status ? FString(Status->GetStoreDir()) : FString();
