[Insights] Implement navigation between CPU core events and context switches events.

diff --git a/Engine/Source/Developer/TraceInsights/Private/Insights/ContextSwitches/ViewModels/ContextSwitchesSharedState.cpp b/Engine/Source/Developer/TraceInsights/Private/Insights/ContextSwitches/ViewModels/ContextSwitchesSharedState.cpp
--- a/Engine/Source/Developer/TraceInsights/Private/Insights/ContextSwitches/ViewModels/ContextSwitchesSharedState.cpp
+++ b/Engine/Source/Developer/TraceInsights/Private/Insights/ContextSwitches/ViewModels/ContextSwitchesSharedState.cpp
@@ -659,11 +659,59 @@ void FContextSwitchesSharedState::Command_NavigateToCpuThreadEvent_Execute()
 				if (ThreadTimingTrack.IsValid() && ThreadTimingTrack->IsVisible())
 				{
 					TimingView->SelectTimingTrack(ThreadTimingTrack, true);
-					//TODO: ThreadTimingTrack->SearchEvent(..); + TimingView->SelectTimingEvent(..)
+					TSharedPtr<FBaseTimingTrack> ChildTrack = ThreadTimingTrack->GetChildTrack();
+					if (ChildTrack && ChildTrack->IsKindOf("FContextSwitchesTimingTrack"))
+					{
+						TSharedPtr<FContextSwitchesTimingTrack> ContextSwitchesTrack = StaticCastSharedPtr<FContextSwitchesTimingTrack>(ChildTrack);
+						if (ContextSwitchesTrack.IsValid())
+						{
+							TSharedPtr<const ITimingEvent> FoundEvent = ContextSwitchesTrack->SearchEvent(
+								FTimingEventSearchParameters(CpuCoreEvent.GetStartTime(), CpuCoreEvent.GetEndTime(), ETimingEventSearchFlags::StopAtFirstMatch));
+
+							if (FoundEvent.IsValid())
+							{
+								TimingView->SelectTimingEvent(FoundEvent, true);
+							}
+						}
+					}
+				}
+			}
+		}
+	}
+}
+
+////////////////////////////////////////////////////////////////////////////////////////////////////
+
+bool FContextSwitchesSharedState::Command_NavigateToCpuThreadEvent_CanExecute() const
+{ 
+	if (!AreContextSwitchesAvailable() || !AreContextSwitchesVisible() || !IsValidCpuCoreEventSelected())
+	{
+		return false;
+	}
+
+	TSharedPtr<STimingView> TimingView = GetTimingView();
+
+	if (TimingView && IsValidCpuCoreEventSelected())
+	{
+		if (TargetTimingEvent.IsValid() && TargetTimingEvent->Is<FCpuCoreTimingEvent>())
+		{
+			const FCpuCoreTimingEvent& CpuCoreEvent = TargetTimingEvent->As<FCpuCoreTimingEvent>();
+			const uint32 SystemThreadId = CpuCoreEvent.GetSystemThreadId();
+			uint32 ThreadId;
+			const TCHAR* ThreadName;
+			GetThreadInfo(SystemThreadId, ThreadId, ThreadName);
+			if (ThreadId != ~0)
+			{
+				TSharedPtr<FThreadTimingTrack> ThreadTimingTrack = GetThreadTimingTrack(ThreadId);
+				if (ThreadTimingTrack.IsValid() && ThreadTimingTrack->IsVisible())
+				{
+					return true;
 				}
 			}
 		}
 	}
+
+	return false;
 }
 
 ////////////////////////////////////////////////////////////////////////////////////////////////////
@@ -726,7 +774,13 @@ void FContextSwitchesSharedState::Command_NavigateToCpuCoreEvent_Execute()
 				if (CpuCoreTimingTrack.IsValid() && CpuCoreTimingTrack->IsVisible())
 				{
 					TimingView->SelectTimingTrack(CpuCoreTimingTrack, true);
-					//TODO: CpuCoreTimingTrack->SearchEvent(..); + TimingView->SelectTimingEvent(..)
+					TSharedPtr<const ITimingEvent> FoundEvent = CpuCoreTimingTrack->SearchEvent(
+						FTimingEventSearchParameters(ContextSwitchEvent.GetStartTime(), ContextSwitchEvent.GetEndTime(), ETimingEventSearchFlags::StopAtFirstMatch));
+
+					if (FoundEvent.IsValid())
+					{
+						TimingView->SelectTimingEvent(FoundEvent, true);
+					}
 				}
 			}
 		}
diff --git a/Engine/Source/Developer/TraceInsights/Private/Insights/ContextSwitches/ViewModels/ContextSwitchesSharedState.h b/Engine/Source/Developer/TraceInsights/Private/Insights/ContextSwitches/ViewModels/ContextSwitchesSharedState.h
--- a/Engine/Source/Developer/TraceInsights/Private/Insights/ContextSwitches/ViewModels/ContextSwitchesSharedState.h
+++ b/Engine/Source/Developer/TraceInsights/Private/Insights/ContextSwitches/ViewModels/ContextSwitchesSharedState.h
@@ -150,7 +150,7 @@ private:
 	bool IsValidContextSwitchEventSelected() const;
 
 	void Command_NavigateToCpuThreadEvent_Execute();
-	bool Command_NavigateToCpuThreadEvent_CanExecute() const { return AreContextSwitchesAvailable() && AreContextSwitchesVisible() && IsValidCpuCoreEventSelected(); }
+	bool Command_NavigateToCpuThreadEvent_CanExecute() const;
 
 	void Command_DockCpuThreadTrackToBottom_Execute();
 	bool Command_DockCpuThreadTrackToBottom_CanExecute() const { return AreContextSwitchesAvailable() && AreContextSwitchesVisible() && IsValidCpuCoreEventSelected(); }
diff --git a/Engine/Source/Developer/TraceInsights/Private/Insights/ContextSwitches/ViewModels/ContextSwitchesTimingTrack.cpp b/Engine/Source/Developer/TraceInsights/Private/Insights/ContextSwitches/ViewModels/ContextSwitchesTimingTrack.cpp
--- a/Engine/Source/Developer/TraceInsights/Private/Insights/ContextSwitches/ViewModels/ContextSwitchesTimingTrack.cpp
+++ b/Engine/Source/Developer/TraceInsights/Private/Insights/ContextSwitches/ViewModels/ContextSwitchesTimingTrack.cpp
@@ -276,6 +276,56 @@ void FContextSwitchesTimingTrack::BuildContextMenu(FMenuBuilder& InOutMenuBuilde
 
 ////////////////////////////////////////////////////////////////////////////////////////////////////
 
+const TSharedPtr<const ITimingEvent> FContextSwitchesTimingTrack::SearchEvent(const FTimingEventSearchParameters& InSearchParameters) const
+{
+	TSharedPtr<const ITimingEvent> FoundEvent;
+
+	TTimingEventSearch<TraceServices::FContextSwitch>::Search(
+		InSearchParameters,
+
+		[this, &InSearchParameters](TTimingEventSearch<TraceServices::FContextSwitch>::FContext& InContext)
+		{
+			TSharedPtr<const TraceServices::IAnalysisSession> Session = FInsightsManager::Get()->GetSession();
+			if (Session.IsValid())
+			{
+				TraceServices::FAnalysisSessionReadScope SessionReadScope(*Session.Get());
+				if (Session.IsValid() && TraceServices::ReadContextSwitchesProvider(*Session.Get()))
+				{
+					const TraceServices::IContextSwitchesProvider* ContextSwitchesProvider = TraceServices::ReadContextSwitchesProvider(*Session.Get());
+
+					auto Callback = [&InContext](double EventStartTime, double EventEndTime, uint32 EventDepth, const TraceServices::FContextSwitch& Event)
+					{
+						InContext.Check(EventStartTime, EventEndTime, EventDepth, Event);
+						return InContext.ShouldContinueSearching() ? TraceServices::EContextSwitchEnumerationResult::Continue : TraceServices::EContextSwitchEnumerationResult::Stop;
+					};
+
+					ContextSwitchesProvider->EnumerateContextSwitches(ThreadId, InSearchParameters.StartTime, InSearchParameters.EndTime,
+						[ContextSwitchesProvider, Callback](const TraceServices::FContextSwitch& ContextSwitchEvent)
+						{
+							return Callback(ContextSwitchEvent.Start, ContextSwitchEvent.End, 0, ContextSwitchEvent);
+						});
+				}
+			}
+		},
+
+		[](double EventStartTime, double EventEndTime, uint32 EventDepth, const TraceServices::FContextSwitch& Event)
+		{
+			return true;
+		},
+
+		[&FoundEvent, this](double InFoundStartTime, double InFoundEndTime, uint32 InFoundDepth, const TraceServices::FContextSwitch& InEvent)
+		{
+			FoundEvent = MakeShared<FTimingEvent>(SharedThis(this), InFoundStartTime, InFoundEndTime, InFoundDepth, InEvent.CoreNumber);
+		},
+
+		TTimingEventSearch<TraceServices::FContextSwitch>::NoMatch
+	);
+
+	return FoundEvent;
+}
+
+////////////////////////////////////////////////////////////////////////////////////////////////////
+
 } // namespace Insights
 
 #undef LOCTEXT_NAMESPACE
diff --git a/Engine/Source/Developer/TraceInsights/Private/Insights/ContextSwitches/ViewModels/ContextSwitchesTimingTrack.h b/Engine/Source/Developer/TraceInsights/Private/Insights/ContextSwitches/ViewModels/ContextSwitchesTimingTrack.h
--- a/Engine/Source/Developer/TraceInsights/Private/Insights/ContextSwitches/ViewModels/ContextSwitchesTimingTrack.h
+++ b/Engine/Source/Developer/TraceInsights/Private/Insights/ContextSwitches/ViewModels/ContextSwitchesTimingTrack.h
@@ -46,6 +46,8 @@ public:
 
 	virtual void BuildContextMenu(FMenuBuilder& MenuBuilder) override;
 
+	virtual const TSharedPtr<const ITimingEvent> SearchEvent(const FTimingEventSearchParameters& InSearchParameters) const override;
+
 protected:
 	virtual const TSharedPtr<const ITimingEvent> GetEvent(double InTime, double SecondsPerPixel, int32 Depth) const override;
 
