Made timeout more long in TraceInsights MiscTests.cpp
Added quotes for some log strings

diff --git a/Engine/Source/Developer/TraceInsights/Private/Insights/Tests/MiscTests.cpp b/Engine/Source/Developer/TraceInsights/Private/Insights/Tests/MiscTests.cpp
--- a/Engine/Source/Developer/TraceInsights/Private/Insights/Tests/MiscTests.cpp
+++ b/Engine/Source/Developer/TraceInsights/Private/Insights/Tests/MiscTests.cpp
@@ -158,7 +158,7 @@ bool FСommandsExportWindowsTest::RunTest(const FString& Parameters)
 	const FString ExportTimingEventsNonDefaultReportPath = TEXT("/TestResults/SingleCommand/TimingEventsNonDefault.csv");
 	const FString ExportTimerStatisticsReportPath = TEXT("/TestResults/SingleCommand/TimerStatistics.csv");
 	bool bLineFound = false;
-	const float Timeout = 50.0;
+	const float Timeout = 150.0;
 
 	UTEST_TRUE("Trace in project exists", PlatformFile.FileExists(*SourceTracePath));
 	UTEST_TRUE("Export in project exists", PlatformFile.FileExists(*SourceExportPath));
@@ -215,7 +215,7 @@ bool FСommandsExportWindowsTest::RunTest(const FString& Parameters)
 	for (int i = 0; i < ExpectedTimingEventsBorderedElements.Num(); i++)
 	{
 		bLineFound = Utils.FileContainsString(ExportTimingEventsNonDefaultReportPath, ExpectedTimingEventsBorderedElements[i], Timeout);
-		TestTrue(FString::Printf(TEXT("Line %s should exists in file: %s"), *ExpectedTimingEventsBorderedElements[i], *ExportTimingEventsNonDefaultReportPath), bLineFound);
+		TestTrue(FString::Printf(TEXT("Line '%s' should exists in file: '%s'"), *ExpectedTimingEventsBorderedElements[i], *ExportTimingEventsNonDefaultReportPath), bLineFound);
 	}
 
 	VerifyExportedLines(ExportTimingEventsNonDefaultReportPath, CmdTimingEventsNonDefaultLogPath, TEXT("timing events"), Utils, this, Timeout);
@@ -225,12 +225,12 @@ bool FСommandsExportWindowsTest::RunTest(const FString& Parameters)
 	InsightsManager->OpenUnrealInsights(*InsightsParameters);
 
 	bLineFound = Utils.FileContainsString(CmdRegionsLogPath, TEXT("Exported timing statistics for 10 regions"), Timeout);
-	TestTrue(FString::Printf(TEXT("Line %s should exists in file: %s"), TEXT("Exported timing statistics for 10 regions"), *CmdRegionsLogPath), bLineFound);
+	TestTrue(FString::Printf(TEXT("Line '%s' should exists in file: '%s'"), TEXT("Exported timing statistics for 10 regions"), *CmdRegionsLogPath), bLineFound);
 
 	TArray<FString> RegionFiles;
 	FString FilePattern = TEXT("*.csv"); // Change the pattern to filter .csv files
 	IFileManager::Get().FindFiles(RegionFiles, *FPaths::Combine(*SingleCommandDirPath, *FilePattern), true, false);
-	TestEqual(FString::Printf(TEXT("Should be 14 csv files but found %d"), RegionFiles.Num()), RegionFiles.Num(), 14);
+	TestEqual(FString::Printf(TEXT("Should be 14 csv files but found '%d'"), RegionFiles.Num()), RegionFiles.Num(), 14);
 
 	// Export.rsp
 	PlatformFile.CopyFile(*LogResultExportPath, *SourceExportPath);
@@ -247,7 +247,7 @@ bool FСommandsExportWindowsTest::RunTest(const FString& Parameters)
 	for (int i = 0; i < ExpectedThreadsElementsRsp.Num(); i++)
 	{
 		bLineFound = Utils.FileContainsString(CmdExportLogPath, ExpectedThreadsElementsRsp[i], Timeout);
-		TestTrue(FString::Printf(TEXT("Line %s should exists in file: %s"), *ExpectedThreadsElementsRsp[i], *CmdExportLogPath), bLineFound);
+		TestTrue(FString::Printf(TEXT("Line '%s' should exists in file: '%s'"), *ExpectedThreadsElementsRsp[i], *CmdExportLogPath), bLineFound);
 
 		VerifyExportedLines(ExpectedThreadsElementsRsp[i], CmdExportLogPath, TEXT("threads"), Utils, this, Timeout);
 	}
@@ -260,7 +260,7 @@ bool FСommandsExportWindowsTest::RunTest(const FString& Parameters)
 	for (int i = 0; i < ExpectedTimersElementsRsp.Num(); i++)
 	{
 		bLineFound = Utils.FileContainsString(CmdExportLogPath, ExpectedTimersElementsRsp[i], Timeout);
-		TestTrue(FString::Printf(TEXT("Line %s should exists in file: %s"), *ExpectedTimersElementsRsp[i], *CmdExportLogPath), bLineFound);
+		TestTrue(FString::Printf(TEXT("Line '%s' should exists in file: '%s'"), *ExpectedTimersElementsRsp[i], *CmdExportLogPath), bLineFound);
 
 		VerifyExportedLines(ExpectedTimersElementsRsp[i], CmdExportLogPath, TEXT("timers"), Utils, this, Timeout);
 	}
@@ -273,7 +273,7 @@ bool FСommandsExportWindowsTest::RunTest(const FString& Parameters)
 	for (int i = 0; i < ExpectedTimingEventsRsp.Num(); i++)
 	{
 		bLineFound = Utils.FileContainsString(CmdExportLogPath, ExpectedTimingEventsRsp[i], Timeout);
-		TestTrue(FString::Printf(TEXT("Line %s should exists in file: %s"), *ExpectedTimingEventsRsp[i], *CmdExportLogPath), bLineFound);
+		TestTrue(FString::Printf(TEXT("Line '%s' should exists in file: '%s'"), *ExpectedTimingEventsRsp[i], *CmdExportLogPath), bLineFound);
 
 		VerifyExportedLines(ExpectedTimingEventsRsp[i], CmdExportLogPath, TEXT("timing events"), Utils, this, Timeout);
 	}
