[Insights] Fix problem with selecting the file extension for export features.

diff --git a/Engine/Source/Developer/TraceInsights/Private/Insights/MemoryProfiler/Widgets/SMemAllocTableTreeView.cpp b/Engine/Source/Developer/TraceInsights/Private/Insights/MemoryProfiler/Widgets/SMemAllocTableTreeView.cpp
--- a/Engine/Source/Developer/TraceInsights/Private/Insights/MemoryProfiler/Widgets/SMemAllocTableTreeView.cpp
+++ b/Engine/Source/Developer/TraceInsights/Private/Insights/MemoryProfiler/Widgets/SMemAllocTableTreeView.cpp
@@ -1807,7 +1807,7 @@ void SMemAllocTableTreeView::OpenSourceFileInIDE(const TCHAR* InFile, uint32 Lin
 void SMemAllocTableTreeView::ExportMemorySnapshot() const
 {
 	// 1. Choose file
-	FString DefaultFile = TEXT("Table.tsv");
+	FString DefaultFile = TEXT("Table");
 	if (Table.IsValid() && !Table->GetDisplayName().IsEmpty())
 	{
 		DefaultFile = Table->GetDisplayName().ToString();
diff --git a/Engine/Source/Developer/TraceInsights/Private/Insights/Table/Widgets/STableTreeView.cpp b/Engine/Source/Developer/TraceInsights/Private/Insights/Table/Widgets/STableTreeView.cpp
--- a/Engine/Source/Developer/TraceInsights/Private/Insights/Table/Widgets/STableTreeView.cpp
+++ b/Engine/Source/Developer/TraceInsights/Private/Insights/Table/Widgets/STableTreeView.cpp
@@ -3671,7 +3671,7 @@ bool STableTreeView::ContextMenu_ExportToFile_CanExecute() const
 
 void STableTreeView::ContextMenu_ExportToFile_Execute(bool bInExportCollapsed, bool InExportLeafs)
 {
-	FString DefaultFile = TEXT("Table.tsv");
+	FString DefaultFile = TEXT("Table");
 	if (Table.IsValid() && !Table->GetDisplayName().IsEmpty())
 	{
 		DefaultFile = Table->GetDisplayName().ToString();
diff --git a/Engine/Source/Developer/TraceInsights/Private/Insights/Widgets/STimersView.cpp b/Engine/Source/Developer/TraceInsights/Private/Insights/Widgets/STimersView.cpp
--- a/Engine/Source/Developer/TraceInsights/Private/Insights/Widgets/STimersView.cpp
+++ b/Engine/Source/Developer/TraceInsights/Private/Insights/Widgets/STimersView.cpp
@@ -2743,7 +2743,7 @@ void STimersView::ContextMenu_Export_Execute()
 	}
 
 	const FString DialogTitle = LOCTEXT("Export_Title", "Export Aggregated Timer Stats").ToString();
-	const FString DefaultFile = TEXT("TimerStats.tsv");
+	const FString DefaultFile = TEXT("TimerStats");
 	FString Filename;
 	if (!OpenSaveTextFileDialog(DialogTitle, DefaultFile, Filename))
 	{
@@ -2892,7 +2892,7 @@ void STimersView::ContextMenu_ExportTimingEventsSelection_Execute() const
 	}
 
 	const FString DialogTitle = LOCTEXT("ExportTimingEventsSelection_Title", "Export Timing Events (Selection)").ToString();
-	const FString DefaultFile = TEXT("TimingEvents.tsv");
+	const FString DefaultFile = TEXT("TimingEvents");
 	FString Filename;
 	if (!OpenSaveTextFileDialog(DialogTitle, DefaultFile, Filename))
 	{
@@ -3014,7 +3014,7 @@ void STimersView::ContextMenu_ExportTimingEvents_Execute() const
 	}
 
 	const FString DialogTitle = LOCTEXT("ExportTimingEvents_Title", "Export Timing Events (All)").ToString();
-	const FString DefaultFile = TEXT("TimingEvents.tsv");
+	const FString DefaultFile = TEXT("TimingEvents");
 	FString Filename;
 	if (!OpenSaveTextFileDialog(DialogTitle, DefaultFile, Filename))
 	{
@@ -3043,7 +3043,7 @@ void STimersView::ContextMenu_ExportThreads_Execute() const
 	}
 
 	const FString DialogTitle = LOCTEXT("ExportThreads_Title", "Export Threads").ToString();
-	const FString DefaultFile = TEXT("Threads.tsv");
+	const FString DefaultFile = TEXT("Threads");
 	FString Filename;
 	if (!OpenSaveTextFileDialog(DialogTitle, DefaultFile, Filename))
 	{
@@ -3072,7 +3072,7 @@ void STimersView::ContextMenu_ExportTimers_Execute() const
 	}
 
 	const FString DialogTitle = LOCTEXT("ExportTimers_Title", "Export Timers").ToString();
-	const FString DefaultFile = TEXT("Timers.tsv");
+	const FString DefaultFile = TEXT("Timers");
 	FString Filename;
 	if (!OpenSaveTextFileDialog(DialogTitle, DefaultFile, Filename))
 	{
