[Insights] Reorganized the structure of functional tests for Unreal Insights.
Added a couple of .h files for compilation with the -NoPCH -DisableUnity flags
Wrapped OpenTraceFile in quotes in tests for custom paths with spaces
Changed type of Timeouts from float to double.

#prefligh 65096a64cf5b71e8360cec09

diff --git a/Engine/Source/Developer/TraceInsights/Private/Insights/Tests/FunctionalTests/ExportCommandsTests.cpp b/Engine/Source/Developer/TraceInsights/Private/Insights/Tests/FunctionalTests/ExportCommandsTests.cpp
new file mode 100644
--- /dev/null
+++ b/Engine/Source/Developer/TraceInsights/Private/Insights/Tests/FunctionalTests/ExportCommandsTests.cpp
@@ -0,0 +1,205 @@
+// Copyright Epic Games, Inc. All Rights Reserved.
+
+#include "HAL/FileManager.h"
+#include "HAL/PlatformFileManager.h"
+#include "Misc/FileHelper.h"
+#include "Insights/InsightsManager.h"
+#include "Insights/IUnrealInsightsModule.h"
+#include "Insights/Tests/InsightsTestUtils.h"
+#include "Misc/AutomationTest.h"
+
+void VerifyExportedLines(const FString& ExportReportPath, const FString& CmdLogPath, const FString& Elements, FInsightsTestUtils Utils, FAutomationTestBase* Test, double Timeout)
+{
+	double StartTime = FPlatformTime::Seconds();
+	bool bLineFound = false;
+	FString ExpectedResult;
+
+	while ((FPlatformTime::Seconds() - StartTime) < Timeout)
+	{
+		FString FileContent;
+
+		if (FFileHelper::LoadFileToString(FileContent, *ExportReportPath))
+		{
+			TArray<FString> Lines;
+			FileContent.ParseIntoArrayLines(Lines);
+			ExpectedResult = FString::Printf(TEXT("Exported %d %s to file"), Lines.Num() - 1, *Elements);
+			bLineFound = Utils.FileContainsString(CmdLogPath, ExpectedResult, 1.0f);
+			if (bLineFound)
+			{
+				Test->TestTrue(FString::Printf(TEXT("Line '%s' from '%s' should exists in file: '%s'"), *ExpectedResult, *ExportReportPath, *CmdLogPath), bLineFound);
+				return;
+			}
+			else
+			{
+				FPlatformProcess::Sleep(0.1f);
+			}
+		}
+		else
+		{
+			FPlatformProcess::Sleep(0.1f);
+		}
+	}
+
+	Test->AddError(FString::Printf(TEXT("VerifyExportedLines timed out while trying to find line '%s' from '%s'"), *ExpectedResult, *ExportReportPath));
+}
+
+#if !WITH_EDITOR
+
+IMPLEMENT_SIMPLE_AUTOMATION_TEST(FСommandsExportWindowsTest, "Insights.Analysis.ExecCmd.CommandsExport(Windows)", EAutomationTestFlags::ApplicationContextMask | EAutomationTestFlags::ProductFilter)
+bool FСommandsExportWindowsTest::RunTest(const FString& Parameters)
+{
+	TSharedPtr<FInsightsManager> InsightsManager = FInsightsManager::Get();
+	IPlatformFile& PlatformFile = FPlatformFileManager::Get().GetPlatformFile();
+	FInsightsTestUtils Utils(this);
+
+	const FString StoreDir = InsightsManager->GetStoreDir();
+	const FString SourceTracePath = FPaths::RootDir() / TEXT("EngineTest/SourceAssets/Utrace/CommandsExportTest_5.4.utrace");
+	const FString SourceExportPath = FPaths::RootDir() / TEXT("EngineTest/SourceAssets/Rsp/export.rsp");
+	const FString StoreTracePath = StoreDir / TEXT("CommandsExportTest_5.4.utrace");
+	const FString TestResultsDirPath = TEXT("/TestResults");
+	const FString SingleCommandDirPath = TEXT("/TestResults/SingleCommand");
+	const FString LogResultExportPath = TEXT("/TestResults/export.rsp");
+	const FString CmdLogPath = TEXT("/TestResults/Logs/cmd.log");
+	const FString CmdThreadLogPath = TEXT("/TestResults/Logs/cmd_treads.log");
+	const FString CmdTimersLogPath = TEXT("/TestResults/Logs/cmd_timers.log");
+	const FString CmdTimingEventsLogPath = TEXT("/TestResults/Logs/cmd_timing_events.log");
+	const FString CmdTimingEventsNonDefaultLogPath = TEXT("/TestResults/Logs/cmd_timing_non_default.log");
+	const FString CmdRegionsLogPath = TEXT("/TestResults/Logs/cmd_regions.log");
+	const FString CmdExportLogPath = TEXT("/TestResults/Logs/cmd_export.log");
+	const FString ExportThreadsTask = TEXT("TimingInsights.ExportThreads /TestResults/SingleCommand/Threads.csv");
+	const FString ExportTimersTask = TEXT("TimingInsights.ExportTimers /TestResults/SingleCommand/Timers.csv");
+	const FString ExportTimingEventsTask = TEXT("TimingInsights.ExportTimingEvents /TestResults/SingleCommand/TimingEvents.csv");
+	const FString ExportTimingEventsBorderedTask = TEXT("TimingInsights.ExportTimingEvents /TestResults/SingleCommand/TimingEventsNonDefault.CSV  -columns=ThreadId,ThreadName,TimerId,TimerName,StartTime,EndTime,Duration,Depth -threads=GameThread -timers=* -startTime=10 -endTime=20");
+	const FString ExportTimerStatisticsTask = TEXT("TimingInsights.ExportTimerStatistics  /TestResults/SingleCommand/*.csv -region=* -threads=GPU");
+	const FString ExportThreadsReportPath = TEXT("/TestResults/SingleCommand/Threads.csv");
+	const FString ExportTimersReportPath = TEXT("/TestResults/SingleCommand/Timers.csv");
+	const FString ExportTimingEventsReportPath = TEXT("/TestResults/SingleCommand/TimingEvents.csv");
+	const FString ExportTimingEventsNonDefaultReportPath = TEXT("/TestResults/SingleCommand/TimingEventsNonDefault.csv");
+	const FString ExportTimerStatisticsReportPath = TEXT("/TestResults/SingleCommand/TimerStatistics.csv");
+	bool bLineFound = false;
+	double Timeout = 150.0;
+
+	UTEST_TRUE("Trace in project exists", PlatformFile.FileExists(*SourceTracePath));
+	UTEST_TRUE("Export in project exists", PlatformFile.FileExists(*SourceExportPath));
+	UTEST_FALSE("Trace in store should not exists before copy", PlatformFile.FileExists(*StoreTracePath));
+
+	PlatformFile.CopyFile(*StoreTracePath, *SourceTracePath);
+	UTEST_TRUE(TEXT("Trace in store should exists after copy"), PlatformFile.FileExists(*StoreTracePath));
+
+	if (PlatformFile.DirectoryExists(*TestResultsDirPath))
+	{
+		AddWarning(TEXT("The TestResults directory already exists. Deleting to avoid undefined behavior"));
+		IFileManager::Get().DeleteDirectory(*TestResultsDirPath, false, true);
+		float StartTime = FPlatformTime::Seconds();
+		while ((FPlatformTime::Seconds() - StartTime) < Timeout)
+		{
+			if (!PlatformFile.DirectoryExists(*TestResultsDirPath))
+			{
+				break;
+			}
+			FPlatformProcess::Sleep(0.1f);
+		}
+
+		UTEST_FALSE("Previously created TestResults directory is not deleted. Execution of the next steps will cause unstable behavior.", PlatformFile.DirectoryExists(*TestResultsDirPath));
+	}
+
+	// ExportThreads
+	FString InsightsParameters = FString::Printf(TEXT("-OpenTraceFile=\"%s\" -ABSLOG=\"%s\" -AutoQuit -NoUI  -ExecOnAnalysisCompleteCmd=\"%s\" -log"), *StoreTracePath, *CmdThreadLogPath, *ExportThreadsTask);
+	InsightsManager->OpenUnrealInsights(*InsightsParameters);
+
+	VerifyExportedLines(ExportThreadsReportPath, CmdThreadLogPath, TEXT("threads"), Utils, this, Timeout);
+
+	// ExportTimers
+	InsightsParameters = FString::Printf(TEXT("-OpenTraceFile=\"%s\" -ABSLOG=\"%s\" -AutoQuit -NoUI  -ExecOnAnalysisCompleteCmd=\"%s\" -log"), *StoreTracePath, *CmdTimersLogPath, *ExportTimersTask);
+	InsightsManager->OpenUnrealInsights(*InsightsParameters);
+
+	VerifyExportedLines(ExportTimersReportPath, CmdTimersLogPath, TEXT("timers"), Utils, this, Timeout);
+
+	// UI verification cannot be executed in that case. 
+
+	// ExportTimingEvents
+	InsightsParameters = FString::Printf(TEXT("-OpenTraceFile=\"%s\" -ABSLOG=\"%s\" -AutoQuit -NoUI  -ExecOnAnalysisCompleteCmd=\"%s\" -log"), *StoreTracePath, *CmdTimingEventsLogPath, *ExportTimingEventsTask);
+	InsightsManager->OpenUnrealInsights(*InsightsParameters);
+
+	VerifyExportedLines(ExportTimingEventsReportPath, CmdTimingEventsLogPath, TEXT("timing events"), Utils, this, Timeout);
+
+	// ExportTimingEventsNonDefault
+	InsightsParameters = FString::Printf(TEXT("-OpenTraceFile=\"%s\" -ABSLOG=\"%s\" -AutoQuit -NoUI  -ExecOnAnalysisCompleteCmd=\"%s\" -log"), *StoreTracePath, *CmdTimingEventsNonDefaultLogPath, *ExportTimingEventsBorderedTask);
+	InsightsManager->OpenUnrealInsights(*InsightsParameters);
+
+	const TArray<FString> ExpectedTimingEventsBorderedElements = { TEXT("ThreadId"), TEXT("ThreadName"), TEXT("TimerId"), TEXT("TimerName"), TEXT("StartTime"), TEXT("EndTime"), TEXT("Duration"), TEXT("Depth"),
+		TEXT("2,GameThread,1,FEngineLoop::PreInitPreStartupScreen")
+	};
+
+	for (int i = 0; i < ExpectedTimingEventsBorderedElements.Num(); i++)
+	{
+		bLineFound = Utils.FileContainsString(ExportTimingEventsNonDefaultReportPath, ExpectedTimingEventsBorderedElements[i], Timeout);
+		TestTrue(FString::Printf(TEXT("Line '%s' should exists in file: '%s'"), *ExpectedTimingEventsBorderedElements[i], *ExportTimingEventsNonDefaultReportPath), bLineFound);
+	}
+
+	VerifyExportedLines(ExportTimingEventsNonDefaultReportPath, CmdTimingEventsNonDefaultLogPath, TEXT("timing events"), Utils, this, Timeout);
+
+	// Export regions
+	InsightsParameters = FString::Printf(TEXT("-OpenTraceFile=\"%s\" -ABSLOG=\"%s\" -AutoQuit -NoUI  -ExecOnAnalysisCompleteCmd=\"%s\" -log"), *StoreTracePath, *CmdRegionsLogPath, *ExportTimerStatisticsTask);
+	InsightsManager->OpenUnrealInsights(*InsightsParameters);
+
+	bLineFound = Utils.FileContainsString(CmdRegionsLogPath, TEXT("Exported timing statistics for 10 regions"), Timeout);
+	TestTrue(FString::Printf(TEXT("Line '%s' should exists in file: '%s'"), TEXT("Exported timing statistics for 10 regions"), *CmdRegionsLogPath), bLineFound);
+
+	TArray<FString> RegionFiles;
+	FString FilePattern = TEXT("*.csv"); // Change the pattern to filter .csv files
+	IFileManager::Get().FindFiles(RegionFiles, *FPaths::Combine(*SingleCommandDirPath, *FilePattern), true, false);
+	TestEqual(FString::Printf(TEXT("Should be 14 csv files but found '%d'"), RegionFiles.Num()), RegionFiles.Num(), 14);
+
+	// Export.rsp
+	PlatformFile.CopyFile(*LogResultExportPath, *SourceExportPath);
+	TestTrue("Rsp in log directory should exists after copy", PlatformFile.FileExists(*LogResultExportPath));
+
+	InsightsParameters = FString::Printf(TEXT("-OpenTraceFile=\"%s\" -ABSLOG=\"%s\" -AutoQuit -NoUI  -ExecOnAnalysisCompleteCmd=\"@=/TestResults/export.rsp\" -log"), *StoreTracePath, *CmdExportLogPath);
+	InsightsManager->OpenUnrealInsights(*InsightsParameters);
+
+	const TArray<FString> ExpectedThreadsElementsRsp = { TEXT("/TestResults/RSPtest/CSV/Threads_rsp.csv"),
+		TEXT("/TestResults/RSPtest/TSV/Threads_rsp.tsv"),
+		TEXT("/TestResults/RSPtest/TXT/Threads_rsp.txt"),
+	};
+
+	for (int i = 0; i < ExpectedThreadsElementsRsp.Num(); i++)
+	{
+		bLineFound = Utils.FileContainsString(CmdExportLogPath, ExpectedThreadsElementsRsp[i], Timeout);
+		TestTrue(FString::Printf(TEXT("Line '%s' should exists in file: '%s'"), *ExpectedThreadsElementsRsp[i], *CmdExportLogPath), bLineFound);
+
+		VerifyExportedLines(ExpectedThreadsElementsRsp[i], CmdExportLogPath, TEXT("threads"), Utils, this, Timeout);
+	}
+
+	const TArray<FString> ExpectedTimersElementsRsp = { TEXT("/TestResults/RSPtest/CSV/Timers_rsp.csv"),
+		TEXT("/TestResults/RSPtest/TSV/Timers_rsp.tsv"),
+		TEXT("/TestResults/RSPtest/TXT/Timers_rsp.txt"),
+	};
+
+	for (int i = 0; i < ExpectedTimersElementsRsp.Num(); i++)
+	{
+		bLineFound = Utils.FileContainsString(CmdExportLogPath, ExpectedTimersElementsRsp[i], Timeout);
+		TestTrue(FString::Printf(TEXT("Line '%s' should exists in file: '%s'"), *ExpectedTimersElementsRsp[i], *CmdExportLogPath), bLineFound);
+
+		VerifyExportedLines(ExpectedTimersElementsRsp[i], CmdExportLogPath, TEXT("timers"), Utils, this, Timeout);
+	}
+
+	const TArray<FString> ExpectedTimingEventsRsp = { TEXT("/TestResults/RSPtest/CSV/TimingEvents_rsp.csv"),
+		TEXT("/TestResults/RSPtest/TSV/TimingEvents_rsp.tsv"),
+		TEXT("/TestResults/RSPtest/TXT/TimingEvents_rsp.txt"),
+	};
+
+	for (int i = 0; i < ExpectedTimingEventsRsp.Num(); i++)
+	{
+		bLineFound = Utils.FileContainsString(CmdExportLogPath, ExpectedTimingEventsRsp[i], Timeout);
+		TestTrue(FString::Printf(TEXT("Line '%s' should exists in file: '%s'"), *ExpectedTimingEventsRsp[i], *CmdExportLogPath), bLineFound);
+
+		VerifyExportedLines(ExpectedTimingEventsRsp[i], CmdExportLogPath, TEXT("timing events"), Utils, this, Timeout);
+	}
+
+	IFileManager::Get().Delete(*StoreTracePath);
+
+	return true;
+}
+
+#endif //!WITH_EDITOR
\ No newline at end of file
diff --git a/Engine/Source/Developer/TraceInsights/Private/Insights/Tests/FunctionalTests/MemoryInsightsTests.cpp b/Engine/Source/Developer/TraceInsights/Private/Insights/Tests/FunctionalTests/MemoryInsightsTests.cpp
new file mode 100644
--- /dev/null
+++ b/Engine/Source/Developer/TraceInsights/Private/Insights/Tests/FunctionalTests/MemoryInsightsTests.cpp
@@ -0,0 +1,370 @@
+// Copyright Epic Games, Inc. All Rights Reserved.
+
+#include "AutomationDriverCommon.h"
+#include "Algo/Find.h"
+#include "HAL/FileManager.h"
+#include "HAL/PlatformFileManager.h"
+#include "Misc/FileHelper.h"
+#include "Insights/InsightsManager.h"
+#include "Insights/IUnrealInsightsModule.h"
+#include "Insights/MemoryProfiler/MemoryProfilerManager.h"
+#include "Insights/MemoryProfiler/ViewModels/MemAllocGroupingByCallstack.h"
+#include "Insights/MemoryProfiler/ViewModels/MemAllocNode.h"
+#include "Insights/MemoryProfiler/ViewModels/MemoryAlloc.h"
+#include "Insights/MemoryProfiler/ViewModels/MemorySharedState.h"
+#include "Insights/MemoryProfiler/Widgets/SMemAllocTableTreeView.h"
+#include "Insights/MemoryProfiler/Widgets/SMemoryProfilerWindow.h"
+#include "Insights/Table/ViewModels/TableTreeNode.h"
+#include "Insights/Table/ViewModels/TreeNodeGrouping.h"
+#include "Insights/Table/Widgets/STableTreeView.h"
+#include "Insights/Tests/InsightsTestUtils.h"
+#include "Insights/Widgets/SStartPageWindow.h"
+#include "Insights/Widgets/STimingView.h"
+#include "Misc/AutomationTest.h"
+#include "TraceServices/Model/AllocationsProvider.h"
+
+DECLARE_LOG_CATEGORY_EXTERN(MemoryInsightsTests, Log, All);
+
+#if !WITH_EDITOR
+
+BEGIN_DEFINE_SPEC(FAutomationDriverUnrealInsightsHubMemoryInsightsTest, "Insights.Hub.MemoryInsights", EAutomationTestFlags::ApplicationContextMask | EAutomationTestFlags::ProductFilter)
+FAutomationDriverPtr Driver;
+TSharedPtr<SWindow> AutomationWindow;
+END_DEFINE_SPEC(FAutomationDriverUnrealInsightsHubMemoryInsightsTest)
+void FAutomationDriverUnrealInsightsHubMemoryInsightsTest::Define()
+{
+	BeforeEach([this]() {
+		AutomationWindow = FSlateApplication::Get().GetActiveTopLevelWindow();
+		const FString AutomationWindowName = TEXT("Automation");
+		if (AutomationWindow->GetTitle().ToString().Contains(AutomationWindowName))
+		{
+			AutomationWindow->Minimize();
+		}
+
+		IUnrealInsightsModule& UnrealInsightsModule = FModuleManager::LoadModuleChecked<IUnrealInsightsModule>("TraceInsights");
+		if (IAutomationDriverModule::Get().IsEnabled())
+		{
+			IAutomationDriverModule::Get().Disable();
+		}
+		IAutomationDriverModule::Get().Enable();
+
+		Driver = IAutomationDriverModule::Get().CreateDriver();
+		});
+
+	Describe("XMLReportsUpload", [this]()
+		{
+			It("should verify that user can upload xml reports in Memory Insights tab", EAsyncExecution::ThreadPool, FTimespan::FromSeconds(120), [this]()
+				{
+					FInsightsTestUtils Utils(this);
+					IPlatformFile& PlatformFile = FPlatformFileManager::Get().GetPlatformFile();
+					TSharedPtr<FInsightsManager> InsightsManager = FInsightsManager::Get();
+					if (!InsightsManager.IsValid())
+					{
+						AddError("Insigts manager should not be null");
+						return;
+					}
+
+					// Start tracing editor instance, not Lyra. There is no difference between them in this test.
+					FString UEPath = FPlatformProcess::GenerateApplicationPath("UnrealEditor", EBuildConfiguration::Development);
+					FString Parameters = TEXT("-trace=Bookmark,Memory -tracehost=127.0.0.1");
+					constexpr bool bLaunchDetached = true;
+					constexpr bool bLaunchHidden = false;
+					constexpr bool bLaunchReallyHidden = false;
+					uint32 ProcessID = 0;
+					const int32 PriorityModifier = 0;
+					const TCHAR* OptionalWorkingDirectory = nullptr;
+					void* PipeWriteChild = nullptr;
+					void* PipeReadChild = nullptr;
+					FProcHandle EditorHandle = FPlatformProcess::CreateProc(*UEPath, *Parameters, bLaunchDetached, bLaunchHidden, bLaunchReallyHidden, &ProcessID, PriorityModifier, OptionalWorkingDirectory, PipeWriteChild, PipeReadChild);
+					if (!EditorHandle.IsValid())
+					{
+						AddError("Editor should be started");
+						return;
+					}
+
+					// Verify that LIVE trace appeared
+					int Index = 0;
+					auto TraceWaiter = [Driver = Driver, &Index](void) -> bool
+					{
+						auto Elements = Driver->FindElements(By::Id("TraceStatusColumnList"))->GetElements();
+						for (int i = 0; i < Elements.Num(); ++i) {
+							if (Elements[i]->GetText().ToString() == TEXT("LIVE")) {
+								Index = i;
+								return true;
+							}
+						}
+						return false;
+					};
+
+					if (!Driver->Wait(Until::Condition(TraceWaiter, FWaitTimeout::InSeconds(10))))
+					{
+						AddError("Live trace should appear");
+						FPlatformProcess::TerminateProc(EditorHandle);
+						return;
+					}
+
+					FDriverElementRef TraceElement = Driver->FindElements(By::Id("TraceList"))->GetElements()[Index];
+					const FString TraceName = TraceElement->GetText().ToString();
+
+					const FString StoreDir = InsightsManager->GetStoreDir();
+					const FString ProjectDir = FPaths::ProjectDir();
+					const FString StoreTracePath = StoreDir / FString::Printf(TEXT("%s.utrace"), *TraceName);
+					const FString StoreCachePath = StoreDir / FString::Printf(TEXT("%s.ucache"), *TraceName);
+					const FString LogDirPath = ProjectDir / TEXT("TestResults");
+					const FString LogPath = ProjectDir / TEXT("TestResults/Log.txt");
+					const FString SuccessTestResult = TEXT("Test Completed. Result={Success}");
+
+					// Test live trace
+					FString TraceParameters = FString::Printf(TEXT("-InsightsTest -ABSLOG=\"%s\" -AutoQuit -ExecOnAnalysisCompleteCmd=\"Automation RunTests Insights.UploadMemoryInsightsLLMXMLReportsTrace\" -OpenTraceFile=\"%s\""), *LogPath, *StoreTracePath);
+					InsightsManager->OpenUnrealInsights(*TraceParameters);
+					bool bLineFound = Utils.FileContainsString(LogPath, SuccessTestResult, 60.0f);
+					TestTrue("Test for live trace should pass", bLineFound);
+
+					IFileManager::Get().DeleteDirectory(*LogDirPath, false, true);
+					FPlatformProcess::TerminateProc(EditorHandle);
+
+					// Test stopped trace 
+					InsightsManager->OpenUnrealInsights(*TraceParameters);
+					bLineFound = Utils.FileContainsString(LogPath, SuccessTestResult, 60.0f);
+					TestTrue("Test for stopped trace should pass", bLineFound);
+
+					IFileManager::Get().DeleteDirectory(*LogDirPath, false, true);
+					IFileManager::Get().Delete(*StoreTracePath);
+					IFileManager::Get().Delete(*StoreCachePath);
+				});
+		});
+	AfterEach([this]() {
+		Driver.Reset();
+		IAutomationDriverModule::Get().Disable();
+		AutomationWindow->Restore();
+		});
+}
+
+IMPLEMENT_SIMPLE_AUTOMATION_TEST(FMemoryInsightsUploadLLMXMLReportsTraceTest, "Insights.Analysis.MemoryInsights.UploadMemoryInsightsLLMXMLReportsTrace", EAutomationTestFlags::ApplicationContextMask | EAutomationTestFlags::ProductFilter)
+bool FMemoryInsightsUploadLLMXMLReportsTraceTest::RunTest(const FString& Parameters)
+{
+	const FString ReportGraphsXMLPath = FPaths::RootDir() / TEXT("EngineTest/SourceAssets/Utrace/ReportGraphs.xml");
+	const FString LLMReportTypesXMLPath = FPaths::RootDir() / TEXT("EngineTest/SourceAssets/Utrace/LLMReportTypes.xml");
+
+	FMemorySharedState* SharedState = FMemoryProfilerManager::Get()->GetSharedState();
+	if (!SharedState)
+	{
+		AddError("ProfilerWindow should be valid. Please, run this test throught Insights Session automation tab");
+		return false;
+	}
+
+	SharedState->RemoveAllMemTagGraphTracks();
+	const int DefaultTracksAmount = SharedState->GetTimingView()->GetAllTracks().Num();
+	SharedState->RemoveAllMemTagGraphTracks();
+	AddExpectedError("Failed to load Report");
+	SharedState->CreateTracksFromReport(ReportGraphsXMLPath);
+	SharedState = FMemoryProfilerManager::Get()->GetSharedState();
+	const int AfterReportGraphsUploadTrackAmount = SharedState->GetTimingView()->GetAllTracks().Num();
+	TestTrue("Tracks amount should be default ", DefaultTracksAmount == AfterReportGraphsUploadTrackAmount);
+
+	SharedState->RemoveAllMemTagGraphTracks();
+	SharedState->CreateTracksFromReport(LLMReportTypesXMLPath);
+	SharedState = FMemoryProfilerManager::Get()->GetSharedState();
+	const int AfterLLMReportTypesUploadTrackAmount = SharedState->GetTimingView()->GetAllTracks().Num();
+	TestTrue("Tracks should not be default", DefaultTracksAmount != AfterLLMReportTypesUploadTrackAmount);
+
+	return true;
+}
+
+DEFINE_LATENT_AUTOMATION_COMMAND_THREE_PARAMETER(FWaitForRunningQuieryFinishedCommand, TSharedPtr<Insights::SMemAllocTableTreeView>, MemAllocTableTreeView, double, Timeout, FAutomationTestBase*, Test);
+bool FWaitForRunningQuieryFinishedCommand::Update()
+{
+	if (!MemAllocTableTreeView->IsRunning())
+	{
+		return true;
+	}
+
+	if (FPlatformTime::Seconds() - StartTime >= Timeout)
+	{
+		Test->AddError(TEXT("FWaitForRunningQuieryFinishedCommand timed out"));
+		return true;
+	}
+
+	return false;
+}
+
+DEFINE_LATENT_AUTOMATION_COMMAND_TWO_PARAMETER(FChangeGroupingCommand, TSharedPtr<Insights::SMemAllocTableTreeView>, MemAllocTableTreeView, FAutomationTestBase*, Test);
+bool FChangeGroupingCommand::Update()
+{
+	TArray<TSharedPtr<Insights::FTreeNodeGrouping>> CurrentGroupings;
+	for (const auto& Grouping : MemAllocTableTreeView->GetAvailableGroupings())
+	{
+		if (Grouping->GetTitleName().ToString().Contains(TEXT("By Callstack")))
+		{
+			CurrentGroupings.Add(Grouping);
+		}
+	}
+
+	Test->TestTrue(TEXT("CurrentGroupings should not be empty"), !CurrentGroupings.IsEmpty());
+	MemAllocTableTreeView->SetCurrentGroupings(CurrentGroupings);
+
+	return true;
+}
+
+DEFINE_LATENT_AUTOMATION_COMMAND_FOUR_PARAMETER(FVerifyHierarchyCallStackCommand, TSharedPtr<Insights::SMemAllocTableTreeView>, MemAllocTableTreeView, FInsightsTestUtils, InsightsTestUtils, double, Timeout, FAutomationTestBase*, Test);
+bool FVerifyHierarchyCallStackCommand::Update()
+{
+	if (!MemAllocTableTreeView->IsRunningAsyncUpdate())
+	{
+		for (const TSharedPtr<Insights::FTableTreeNode>& Node : MemAllocTableTreeView->GetTableRowNodes())
+		{
+			const Insights::FMemAllocNode& MemAllocNode = static_cast<const Insights::FMemAllocNode&>(*Node);
+			const Insights::FMemoryAlloc Alloc = MemAllocNode.GetMemAllocChecked();
+			if (!(!Alloc.GetCallstack() || Alloc.GetCallstack()->Num() == 0 || (Alloc.GetCallstack()->Num() != 0 && Alloc.GetCallstack()->Num() <= 256)))
+			{
+				Test->AddError(TEXT("Resolved callstack should be valid"));
+			}
+		}
+		return true;
+	}
+
+	if (FPlatformTime::Seconds() - StartTime >= Timeout)
+	{
+		Test->AddError(TEXT("FVerifyHierarchyCallStackCommand timed out"));
+		return true;
+	}
+
+	return false;
+}
+
+const TMap<TraceServices::IAllocationsProvider::EQueryRule, Insights::SMemAllocTableTreeView::FQueryParams> AllocsTimeMarkerStandaloneGameGetterMap
+{
+	{TraceServices::IAllocationsProvider::EQueryRule::aAf, {nullptr, {5.0, 0.0, 0.0, 0.0}}},
+	{TraceServices::IAllocationsProvider::EQueryRule::afA, {nullptr, {10.0, 0.0, 0.0, 0.0}}},
+	{TraceServices::IAllocationsProvider::EQueryRule::Aaf, {nullptr, {10.0, 0.0, 0.0, 0.0}}},
+	{TraceServices::IAllocationsProvider::EQueryRule::aAfB, {nullptr, {50.0, 51.0, 0.0, 0.0}}},
+	{TraceServices::IAllocationsProvider::EQueryRule::AaBf, {nullptr, {50.0, 51.0, 0.0, 0.0}}},
+	{TraceServices::IAllocationsProvider::EQueryRule::aAfaBf, {nullptr, {50.0, 51.0, 0.0, 0.0}}},
+	{TraceServices::IAllocationsProvider::EQueryRule::AfB, {nullptr, {50.0, 51.0, 0.0, 0.0}}},
+	{TraceServices::IAllocationsProvider::EQueryRule::AaBCf, {nullptr, {50.0, 51.0, 52.0, 0.0}}},
+	{TraceServices::IAllocationsProvider::EQueryRule::AaBfC, {nullptr, {50.0, 51.0, 52.0, 0.0}}},
+	{TraceServices::IAllocationsProvider::EQueryRule::aABfC, {nullptr, {50.0, 51.0, 52.0, 0.0}}},
+	{TraceServices::IAllocationsProvider::EQueryRule::AaBCfD, {nullptr, {50.0, 51.0, 52.0, 53.0}}},
+	{TraceServices::IAllocationsProvider::EQueryRule::aABf, {nullptr, {50.0, 51.0, 0.0, 0.0}}},
+	{TraceServices::IAllocationsProvider::EQueryRule::AafB, {nullptr, {50.0, 51.0, 0.0, 0.0}}},
+	{TraceServices::IAllocationsProvider::EQueryRule::AaB, {nullptr, {50.0, 51.0, 0.0, 0.0}}},
+};
+
+const TMap<TraceServices::IAllocationsProvider::EQueryRule, Insights::SMemAllocTableTreeView::FQueryParams> AllocsTimeMarkerEditorPackageGetterMap
+{
+	{TraceServices::IAllocationsProvider::EQueryRule::aAf, {nullptr, {5.0, 0.0, 0.0, 0.0}}},
+	{TraceServices::IAllocationsProvider::EQueryRule::afA, {nullptr, {10.0, 0.0, 0.0, 0.0}}},
+	{TraceServices::IAllocationsProvider::EQueryRule::Aaf, {nullptr, {10.0, 0.0, 0.0, 0.0}}},
+	{TraceServices::IAllocationsProvider::EQueryRule::aAfB, {nullptr, {2.0, 3.0, 0.0, 0.0}}},
+	{TraceServices::IAllocationsProvider::EQueryRule::AaBf, {nullptr, {2.0, 3.0, 0.0, 0.0}}},
+	{TraceServices::IAllocationsProvider::EQueryRule::aAfaBf, {nullptr, {2.0, 3.0, 0.0, 0.0}}},
+	{TraceServices::IAllocationsProvider::EQueryRule::AfB, {nullptr, {2.0, 3.0, 0.0, 0.0}}},
+	{TraceServices::IAllocationsProvider::EQueryRule::AaBCf, {nullptr, {1.0, 2.0, 3.0, 0.0}}},
+	{TraceServices::IAllocationsProvider::EQueryRule::AaBfC, {nullptr, {1.0, 2.0, 3.0, 0.0}}},
+	{TraceServices::IAllocationsProvider::EQueryRule::aABfC, {nullptr, {1.0, 2.0, 3.0, 0.0}}},
+	{TraceServices::IAllocationsProvider::EQueryRule::AaBCfD, {nullptr, {1.0, 2.0, 3.0, 4.0}}},
+	{TraceServices::IAllocationsProvider::EQueryRule::AafB, {nullptr, {2.0, 3.0, 0.0, 0.0}}},
+	{TraceServices::IAllocationsProvider::EQueryRule::AaB, {nullptr, {2.0, 3.0, 0.0, 0.0}}},
+	{TraceServices::IAllocationsProvider::EQueryRule::aABf, {nullptr, {2.0, 3.0, 0.0, 0.0}}},
+};
+
+bool MemoryInsightsAllocationsQueryTableTest(const FString& Parameters, const TMap <TraceServices::IAllocationsProvider::EQueryRule, Insights::SMemAllocTableTreeView::FQueryParams> AllocsTimeMarkerGetterMap, FAutomationTestBase* Test)
+{
+	double Timeout = 30.0;
+	FInsightsTestUtils InsightsTestUtils(Test);
+	TSharedPtr<SMemoryProfilerWindow> ProfilerWindow = FMemoryProfilerManager::Get()->GetProfilerWindow();
+	TSharedPtr<FInsightsManager> InsightsManager = FInsightsManager::Get();
+	FMemorySharedState& SharedState = ProfilerWindow->GetSharedState();
+
+	TSharedPtr<Insights::FMemoryRuleSpec> MemoryRule = *Algo::FindByPredicate(SharedState.GetMemoryRules(),
+		[&Parameters](const TSharedPtr<Insights::FMemoryRuleSpec>& Rule)
+		{
+			return Rule->GetShortName().ToString().Contains(Parameters);
+		});
+
+	if (!MemoryRule.IsValid())
+	{
+		Test->AddError(TEXT("MemoryRule should not be null"));
+		return false;
+	}
+
+	TSharedPtr<Insights::SMemAllocTableTreeView> MemAllocTableTreeView = ProfilerWindow->ShowMemAllocTableTreeViewTab();
+
+	Insights::SMemAllocTableTreeView::FQueryParams QueryParams = AllocsTimeMarkerGetterMap.FindChecked(MemoryRule->GetValue());
+	if (MemoryRule->GetValue() == TraceServices::IAllocationsProvider::EQueryRule::Aaf)
+	{
+		QueryParams.TimeMarkers[0] = InsightsManager->GetSessionDuration() - 10.0;
+	}
+	QueryParams.Rule = MemoryRule;
+	MemAllocTableTreeView->SetQueryParams(QueryParams);
+
+	ADD_LATENT_AUTOMATION_COMMAND(FWaitForRunningQuieryFinishedCommand(MemAllocTableTreeView, 120.0f, Test));
+	ADD_LATENT_AUTOMATION_COMMAND(FChangeGroupingCommand(MemAllocTableTreeView, Test));
+	ADD_LATENT_AUTOMATION_COMMAND(FWaitForRunningQuieryFinishedCommand(MemAllocTableTreeView, 120.0f, Test));
+	ADD_LATENT_AUTOMATION_COMMAND(FVerifyHierarchyCallStackCommand(MemAllocTableTreeView, InsightsTestUtils, 120.0f, Test));
+
+	return true;
+}
+
+IMPLEMENT_COMPLEX_AUTOMATION_TEST(FMemoryInsightsAllocationsQueryTableEditorPackageTest, "Insights.Analysis.MemoryInsights.AllocationsQueryTable.Editor/Package", EAutomationTestFlags::ApplicationContextMask | EAutomationTestFlags::ProductFilter)
+bool FMemoryInsightsAllocationsQueryTableEditorPackageTest::RunTest(const FString& Parameters)
+{
+	bool bSuccess = MemoryInsightsAllocationsQueryTableTest(Parameters, AllocsTimeMarkerEditorPackageGetterMap, this);
+	return bSuccess;
+}
+
+void FMemoryInsightsAllocationsQueryTableEditorPackageTest::GetTests(TArray<FString>& OutBeautifiedNames, TArray <FString>& OutTestCommands) const
+{
+	if (!FMemoryProfilerManager::Get().IsValid())
+	{
+		return;
+	}
+	TSharedPtr<SMemoryProfilerWindow> ProfilerWindow = FMemoryProfilerManager::Get()->GetProfilerWindow();
+	if (!ProfilerWindow.IsValid())
+	{
+		return;
+	}
+
+	FMemorySharedState& SharedState = ProfilerWindow->GetSharedState();
+
+	for (const auto& MemoryRules : SharedState.GetMemoryRules())
+	{
+		const FString& MemoryRuleName = MemoryRules->GetShortName().ToString();
+
+		OutBeautifiedNames.Add(FString::Printf(TEXT("%s"), *MemoryRuleName));
+		OutTestCommands.Add(MemoryRuleName);
+	}
+}
+
+IMPLEMENT_COMPLEX_AUTOMATION_TEST(FMemoryInsightsAllocationsQueryTableStandaloneTest, "Insights.Analysis.MemoryInsights.AllocationsQueryTable.Standalone", EAutomationTestFlags::ApplicationContextMask | EAutomationTestFlags::ProductFilter)
+bool FMemoryInsightsAllocationsQueryTableStandaloneTest::RunTest(const FString& Parameters)
+{
+	bool bSuccess = MemoryInsightsAllocationsQueryTableTest(Parameters, AllocsTimeMarkerStandaloneGameGetterMap, this);
+	return bSuccess;
+}
+
+void FMemoryInsightsAllocationsQueryTableStandaloneTest::GetTests(TArray<FString>& OutBeautifiedNames, TArray <FString>& OutTestCommands) const
+{
+	if (!FMemoryProfilerManager::Get().IsValid())
+	{
+		return;
+	}
+	TSharedPtr<SMemoryProfilerWindow> ProfilerWindow = FMemoryProfilerManager::Get()->GetProfilerWindow();
+	if (!ProfilerWindow.IsValid())
+	{
+		return;
+	}
+
+	FMemorySharedState& SharedState = ProfilerWindow->GetSharedState();
+	TSharedPtr<Insights::FMemoryRuleSpec> MemoryRule = SharedState.GetMemoryRules()[0];
+
+	for (const auto& MemoryRules : SharedState.GetMemoryRules())
+	{
+		const FString& MemoryRuleName = MemoryRules->GetShortName().ToString();
+
+		OutBeautifiedNames.Add(FString::Printf(TEXT("%s"), *MemoryRuleName));
+		OutTestCommands.Add(MemoryRuleName);
+	}
+}
+
+#endif //!WITH_EDITOR
\ No newline at end of file
diff --git a/Engine/Source/Developer/TraceInsights/Private/Insights/Tests/FunctionalTests/SessionBrowserTests.cpp b/Engine/Source/Developer/TraceInsights/Private/Insights/Tests/FunctionalTests/SessionBrowserTests.cpp
new file mode 100644
--- /dev/null
+++ b/Engine/Source/Developer/TraceInsights/Private/Insights/Tests/FunctionalTests/SessionBrowserTests.cpp
@@ -0,0 +1,148 @@
+// Copyright Epic Games, Inc. All Rights Reserved.
+
+#include "AutomationDriverCommon.h"
+#include "Framework/Application/SlateApplication.h"
+#include "HAL/FileManager.h"
+#include "HAL/PlatformFileManager.h"
+#include "Misc/AutomationTest.h"
+#include "Misc/FileHelper.h"
+#include "Insights/InsightsManager.h"
+#include "Insights/IUnrealInsightsModule.h"
+#include "Insights/Tests/InsightsTestUtils.h"
+#include "Insights/Widgets/SStartPageWindow.h"
+#include "Logging/LogMacros.h"
+#include "TraceServices/Model/AllocationsProvider.h"
+
+DECLARE_LOG_CATEGORY_EXTERN(UnrealInsightsHubTests, Log, All);
+
+#if !WITH_EDITOR
+
+BEGIN_DEFINE_SPEC(FAutomationDriverUnrealInsightsSessionBrowserTest, "Insights.Hub.SessionBrowser", EAutomationTestFlags::ApplicationContextMask | EAutomationTestFlags::ProductFilter)
+FAutomationDriverPtr Driver;
+TSharedPtr<SWindow> AutomationWindow;
+END_DEFINE_SPEC(FAutomationDriverUnrealInsightsSessionBrowserTest)
+void FAutomationDriverUnrealInsightsSessionBrowserTest::Define()
+{
+	BeforeEach([this]() {
+		AutomationWindow = FSlateApplication::Get().GetActiveTopLevelWindow();
+		const FString AutomationWindowName = TEXT("Automation");
+		if (AutomationWindow->GetTitle().ToString().Contains(AutomationWindowName))
+		{
+			AutomationWindow->Minimize();
+		}
+
+		IUnrealInsightsModule& UnrealInsightsModule = FModuleManager::LoadModuleChecked<IUnrealInsightsModule>("TraceInsights");
+		if (IAutomationDriverModule::Get().IsEnabled())
+		{
+			IAutomationDriverModule::Get().Disable();
+		}
+		IAutomationDriverModule::Get().Enable();
+
+		Driver = IAutomationDriverModule::Get().CreateDriver();
+		});
+
+	Describe("CopyRenameDeleteTrace", [this]()
+		{
+			It("should verify that user copy, rename and delete traces", EAsyncExecution::ThreadPool, [this]()
+				{
+					IPlatformFile& PlatformFile = FPlatformFileManager::Get().GetPlatformFile();
+					TSharedPtr<FInsightsManager> InsightsManager = FInsightsManager::Get();
+					TestTrue("Insigts manager should not be null", InsightsManager.IsValid());
+					InsightsManager->GetTraceStoreWindow()->SetDeleteTraceConfirmationWindowVisibility(false);
+
+					const FString StoreDir = InsightsManager->GetStoreDir();
+					const FString ProjectDir = FPaths::ProjectDir();
+
+					const FString SourceTestTracePath = FPaths::RootDir() / TEXT("EngineTest/SourceAssets/Utrace/Test.utrace");
+					const FString SourceTestCachePath = FPaths::RootDir() / TEXT("EngineTest/SourceAssets/Utrace/Test.ucache");
+
+					FString StoreTestTracePath = StoreDir / TEXT("Test.utrace");
+					FString StoreTestCachePath = StoreDir / TEXT("Test.ucache");
+
+					TestTrue("Trace in project exists", PlatformFile.FileExists(*SourceTestTracePath));
+					TestTrue("Cache in project exists", PlatformFile.FileExists(*SourceTestCachePath));
+
+					TestFalse("Trace in store should not exist before copy", PlatformFile.FileExists(*StoreTestTracePath));
+					TestFalse("Cache in store should not exist before copy", PlatformFile.FileExists(*StoreTestCachePath));
+
+					// Copy trace
+					// Here we just check that button can be clicked. Unable to copy and paste via Automation Driver 
+					FDriverElementRef ExploreTraceStoreDirButton = Driver->FindElement(By::Id("ExploreTraceStoreDirButton"));
+					TestTrue("Explore Trace Store Dir Button clicked", ExploreTraceStoreDirButton->IsInteractable());
+
+					PlatformFile.CopyFile(*StoreTestTracePath, *SourceTestTracePath);
+					PlatformFile.CopyFile(*StoreTestCachePath, *SourceTestCachePath);
+
+					TestTrue("Trace copied", PlatformFile.FileExists(*StoreTestTracePath));
+					TestTrue("Cache copied", PlatformFile.FileExists(*StoreTestCachePath));
+
+					StoreTestTracePath = StoreDir / TEXT("TestUcacheRenaming.utrace");
+					StoreTestCachePath = StoreDir / TEXT("TestUcacheRenaming.ucache");
+
+					TestFalse("Renamed trace should not exist before renaming", PlatformFile.FileExists(*StoreTestTracePath));
+					TestFalse("Renamed cache should not exist before renaming", PlatformFile.FileExists(*StoreTestCachePath));
+
+					// Rename 
+					int Index = 0;
+					auto TraceWaiter = [Driver = Driver, &Index](void) -> bool
+					{
+						auto Elements = Driver->FindElements(By::Id("TraceList"))->GetElements();
+						for (int i = 0; i < Elements.Num(); ++i) {
+							if (Elements[i]->GetText().ToString() == TEXT("Test")) {
+								Index = i;
+								return true;
+							}
+						}
+						return false;
+					};
+
+					bool bTestTraceExists = Driver->Wait(Until::Condition(TraceWaiter, FWaitTimeout::InSeconds(10)));
+					if (!bTestTraceExists)
+					{
+						AddError("Trace should exists in Session Browser");
+						return;
+					}
+
+					FDriverElementRef TraceElement = Driver->FindElements(By::Id("TraceList"))->GetElements()[Index];
+
+					FDriverSequenceRef Sequence = Driver->CreateSequence();
+					Sequence->Actions()
+						.Click(TraceElement)
+						.Type(EKeys::F2)
+						.Type(TEXT("UcacheRenaming"))
+						.Type(EKeys::Enter);
+
+					TestTrue("Trace renamed", Sequence->Perform());
+
+					TestTrue("Renamed trace should exists", PlatformFile.FileExists(*StoreTestTracePath));
+					TestTrue("Renamed cache should exists", PlatformFile.FileExists(*StoreTestCachePath));
+
+					// Delete
+					FDriverElementRef OpenTraceButton = Driver->FindElement(By::Id("OpenTraceButton"));
+					Driver->Wait(Until::ElementIsInteractable(OpenTraceButton, FWaitTimeout::InSeconds(10)));
+
+					TraceElement = Driver->FindElements(By::Id("TraceList"))->GetElements()[0];
+					TraceElement->Type(EKeys::Delete);
+
+					TestFalse("Renamed trace should be deleted", PlatformFile.FileExists(*StoreTestTracePath));
+					TestFalse("Renamed cache should be deleted", PlatformFile.FileExists(*StoreTestCachePath));
+				});
+
+			AfterEach([this]() {
+				TSharedPtr<FInsightsManager> InsightsManager = FInsightsManager::Get();
+				const FString StoreDir = InsightsManager->GetStoreDir();
+				const FString StoreTestTracePath = StoreDir / TEXT("Test.utrace");
+				const FString StoreTestCachePath = StoreDir / TEXT("Test.ucache");
+				IFileManager::Get().Delete(*StoreTestTracePath, false, true);
+				IFileManager::Get().Delete(*StoreTestCachePath, false, true);
+				});
+		});
+
+	AfterEach([this]() {
+		Driver.Reset();
+		IAutomationDriverModule::Get().Disable();
+		AutomationWindow->Restore();
+		});
+}
+
+#endif //!WITH_EDITOR
\ No newline at end of file
diff --git a/Engine/Source/Developer/TraceInsights/Private/Insights/Tests/FunctionalTests/TimingProfilerTests.cpp b/Engine/Source/Developer/TraceInsights/Private/Insights/Tests/FunctionalTests/TimingProfilerTests.cpp
new file mode 100644
--- /dev/null
+++ b/Engine/Source/Developer/TraceInsights/Private/Insights/Tests/FunctionalTests/TimingProfilerTests.cpp
@@ -0,0 +1,654 @@
+// Copyright Epic Games, Inc. All Rights Reserved.
+
+#include "Insights/Tests/TimingProfilerTests.h"
+
+#include "TraceServices/Model/Threads.h"
+#include "Misc/AutomationTest.h"
+#include "Misc/FileHelper.h"
+
+#include "Insights/Common/Stopwatch.h"
+#include "Insights/InsightsManager.h"
+#include "Insights/Tests/InsightsTestUtils.h"
+
+////////////////////////////////////////////////////////////////////////////////////////////////////
+
+DEFINE_LOG_CATEGORY(TimingProfilerTests);
+
+////////////////////////////////////////////////////////////////////////////////////////////////////
+
+#if !WITH_EDITOR && WITH_DEV_AUTOMATION_TESTS
+
+IMPLEMENT_SIMPLE_AUTOMATION_TEST(FEnumerateTest, "Insights.Analysis.TimingInsights.Enumerate", EAutomationTestFlags::ApplicationContextMask | EAutomationTestFlags::ProductFilter)
+bool FEnumerateTest::RunTest(const FString& Parameters)
+{
+	FInsightsTestUtils Utils(this);
+
+	FString Path = FPaths::RootDir() + TEXT("EngineTest/SourceAssets/Utrace/r424_win64_game_11590231.utrace");
+	bool ret = Utils.AnalyzeTrace(*Path);
+	if (!ret)
+	{
+		return !HasAnyErrors();
+	}
+
+	FTimingProfilerTests::FEnumerateTestParams Params;
+	Params.Interval = 0.01;
+	Params.NumEnumerations = 10000;
+
+	FTimingProfilerTests::FCheckValues CheckValues;
+	FTimingProfilerTests::RunEnumerateBenchmark(Params, CheckValues);
+
+	TestEqual(TEXT("SessionDuration"), CheckValues.SessionDuration, 307.0172116, 1.e-6);
+	TestEqual(TEXT("TotalEventDuration"), CheckValues.TotalEventDuration, 680.943945, 1.e-6);
+	TestEqual(TEXT("EventCount"), CheckValues.EventCount, 10836057ull);
+	TestEqual(TEXT("SumDepth"), (uint64)CheckValues.SumDepth, 80030008ull);
+	TestEqual(TEXT("SumTimerIndex"), (uint64)CheckValues.SumTimerIndex, 4126772211ull);
+
+	AddInfo(FString::Printf(TEXT("Enumeration Duration: %f seconds."), CheckValues.EnumerationDuration));
+
+	return !HasAnyErrors();
+}
+
+IMPLEMENT_SIMPLE_AUTOMATION_TEST(FEnumeratePerformanceTest, "Insights.Analysis.TimingInsights.EnumeratePerformance", EAutomationTestFlags::ApplicationContextMask | EAutomationTestFlags::PerfFilter)
+bool FEnumeratePerformanceTest::RunTest(const FString& Parameters)
+{
+	FInsightsTestUtils Utils(this);
+
+	FString Path = FPaths::RootDir() + TEXT("EngineTest/SourceAssets/Utrace/r425_win64_game_13649855.utrace");
+	bool ret = Utils.AnalyzeTrace(*Path);
+	if (!ret)
+	{
+		return !HasAnyErrors();
+	}
+
+	FTimingProfilerTests::FEnumerateTestParams Params;
+	Params.Interval = 0.01;
+	Params.NumEnumerations = 100000;
+
+	FTimingProfilerTests::FCheckValues CheckValues;
+	FTimingProfilerTests::RunEnumerateBenchmark(Params, CheckValues);
+
+	TestEqual(TEXT("SessionDuration"), CheckValues.SessionDuration, 341.073285, 1.e-6);
+	TestEqual(TEXT("TotalEventDuration"), CheckValues.TotalEventDuration, 10912.775537, 1.e-6);
+	TestEqual(TEXT("EventCount"), CheckValues.EventCount, 137000700ull);
+	TestEqual(TEXT("SumDepth"), (uint64)CheckValues.SumDepth, 1134384338ull);
+	TestEqual(TEXT("SumTimerIndex"), (uint64)CheckValues.SumTimerIndex, 3499618755ull);
+
+	const double BenchmarkBaseline = 16.0;
+	AddInfo(FString::Printf(TEXT("Enumeration Duration: %f seconds."), CheckValues.EnumerationDuration));
+
+	if (CheckValues.EnumerationDuration > 1.5 * BenchmarkBaseline)
+	{
+		AddWarning(FString::Printf(TEXT("Enumeration duration (%f seconds) exceeded baseline by %.2f%%."), CheckValues.EnumerationDuration, CheckValues.EnumerationDuration / BenchmarkBaseline * 100.0));
+	}
+	else if (CheckValues.EnumerationDuration > 1.25 * BenchmarkBaseline)
+	{
+		AddInfo(FString::Printf(TEXT("Enumeration duration (%f seconds) exceeded baseline by %.2f%%."), CheckValues.EnumerationDuration, CheckValues.EnumerationDuration / BenchmarkBaseline * 100.0));
+	}
+
+	return !HasAnyErrors();
+}
+
+IMPLEMENT_SIMPLE_AUTOMATION_TEST(FEnumerateFastTest, "Insights.Analysis.TimingInsights.EnumerateFast", EAutomationTestFlags::ApplicationContextMask | EAutomationTestFlags::ProductFilter)
+bool FEnumerateFastTest::RunTest(const FString& Parameters)
+{
+	FInsightsTestUtils Utils(this);
+
+	FString Path = FPaths::RootDir() + TEXT("EngineTest/SourceAssets/Utrace/r423_win64_game_10478456.utrace");
+	bool ret = Utils.AnalyzeTrace(*Path);
+	if (!ret)
+	{
+		return !HasAnyErrors();
+	}
+
+	FTimingProfilerTests::FEnumerateTestParams Params;
+	Params.Interval = 0.01;
+	Params.NumEnumerations = 10000;
+
+	FTimingProfilerTests::FCheckValues CheckValues;
+	FTimingProfilerTests::RunEnumerateBenchmark(Params, CheckValues);
+
+	TestEqual(TEXT("SessionDuration"), CheckValues.SessionDuration, 305.232584, 1.e-6);
+	TestEqual(TEXT("TotalEventDuration"), CheckValues.TotalEventDuration, 1647.693886, 1.e-6);
+	TestEqual(TEXT("EventCount"), CheckValues.EventCount, 1759740ull);
+	TestEqual(TEXT("SumDepth"), (uint64)CheckValues.SumDepth, 15189227ull);
+	TestEqual(TEXT("SumTimerIndex"), (uint64)CheckValues.SumTimerIndex, 1239801518ull);
+
+	AddInfo(FString::Printf(TEXT("Enumeration Duration: %f seconds."), CheckValues.EnumerationDuration));
+
+	return !HasAnyErrors();
+}
+
+IMPLEMENT_SIMPLE_AUTOMATION_TEST(EnumerateEventsToFile, "Insights.Analysis.TimingInsights.EnumerateEventsToFile", EAutomationTestFlags::ApplicationContextMask | EAutomationTestFlags::ProductFilter)
+bool EnumerateEventsToFile::RunTest(const FString& Parameters)
+{
+	double SessionTime = 0.0;
+	uint32 TimelineIndex = (uint32)-1;
+
+	TSharedPtr<const TraceServices::IAnalysisSession> Session = FInsightsManager::Get()->GetSession();
+	if (Session.IsValid() && TraceServices::ReadTimingProfilerProvider(*Session.Get()))
+	{
+		TraceServices::FAnalysisSessionReadScope SessionReadScope(*Session.Get());
+
+		const TraceServices::ITimingProfilerProvider& TimingProfilerProvider = *TraceServices::ReadTimingProfilerProvider(*Session.Get());
+
+		TimelineIndex = FTimingProfilerTests::GetTimelineIndex(TEXT("GameThread"));
+
+		AddErrorIfFalse(TimelineIndex != (uint32)-1, TEXT("Failed to get track named GameThread"));
+
+		const TraceServices::ITimingProfilerTimerReader* TimerReader;
+		TimingProfilerProvider.ReadTimers([&TimerReader](const TraceServices::ITimingProfilerTimerReader& Out) { TimerReader = &Out; });
+
+		FString EnumeratedEventsStr;
+		TimingProfilerProvider.ReadTimeline(TimelineIndex,
+			[Session, &EnumeratedEventsStr, TimerReader](const TraceServices::ITimingProfilerProvider::Timeline& Timeline)
+			{
+				Timeline.EnumerateEvents(-1.0, Session->GetDurationSeconds() + 1.0,
+					[&EnumeratedEventsStr, TimerReader](double EventStartTime, double EventEndTime, uint32 EventDepth, const TraceServices::FTimingProfilerEvent& Event)
+					{
+						const TraceServices::FTimingProfilerTimer* Timer = TimerReader->GetTimer(Event.TimerIndex);
+
+						EnumeratedEventsStr.Appendf(TEXT("%s %f %f %d\n"), Timer->Name, EventStartTime, EventEndTime, EventDepth);
+						return TraceServices::EEventEnumerate::Continue;
+					});
+			});
+
+		FString FilePath = FPaths::ProjectSavedDir() + TEXT("/EnumerateEventsToFile.txt");
+		FFileHelper::SaveStringToFile(EnumeratedEventsStr, *FilePath);
+	}
+
+	return !HasAnyErrors();
+}
+
+IMPLEMENT_SIMPLE_AUTOMATION_TEST(EnumerateScopesToFile, "Insights.Analysis.TimingInsights.EnumerateScopesToFile", EAutomationTestFlags::ApplicationContextMask | EAutomationTestFlags::ProductFilter)
+bool EnumerateScopesToFile::RunTest(const FString& Parameters)
+{
+	double SessionTime = 0.0;
+	uint32 TimelineIndex = (uint32)-1;
+
+	TSharedPtr<const TraceServices::IAnalysisSession> Session = FInsightsManager::Get()->GetSession();
+	if (Session.IsValid() && TraceServices::ReadTimingProfilerProvider(*Session.Get()))
+	{
+		TraceServices::FAnalysisSessionReadScope SessionReadScope(*Session.Get());
+
+		const TraceServices::ITimingProfilerProvider& TimingProfilerProvider = *TraceServices::ReadTimingProfilerProvider(*Session.Get());
+
+		TimelineIndex = FTimingProfilerTests::GetTimelineIndex(TEXT("GameThread"));
+
+		AddErrorIfFalse(TimelineIndex != (uint32)-1, TEXT("Failed to get track named GameThread"));
+
+		const TraceServices::ITimingProfilerTimerReader* TimerReader;
+		TimingProfilerProvider.ReadTimers([&TimerReader](const TraceServices::ITimingProfilerTimerReader& Out) { TimerReader = &Out; });
+
+		FString EnumeratedEventsStr;
+		TimingProfilerProvider.ReadTimeline(TimelineIndex,
+			[Session, &EnumeratedEventsStr, TimerReader](const TraceServices::ITimingProfilerProvider::Timeline& Timeline)
+			{
+				Timeline.EnumerateEvents(-1.0, Session->GetDurationSeconds() + 1.0,
+					[&EnumeratedEventsStr, TimerReader](bool bStart, double Time, const TraceServices::FTimingProfilerEvent& Event)
+					{
+						const TraceServices::FTimingProfilerTimer* Timer = TimerReader->GetTimer(Event.TimerIndex);
+
+						EnumeratedEventsStr.Appendf(TEXT("%s %d %f\n"), Timer->Name, (int32)bStart, Time);
+						return TraceServices::EEventEnumerate::Continue;
+					});
+			});
+
+		FString FilePath = FPaths::ProjectSavedDir() + TEXT("/EnumerateScopesToFile.txt");
+		FFileHelper::SaveStringToFile(EnumeratedEventsStr, *FilePath);
+	}
+
+	return !HasAnyErrors();
+}
+
+IMPLEMENT_SIMPLE_AUTOMATION_TEST(InsightsEnumerate10K, "Insights.Analysis.TimingInsights.Enumerate10K", EAutomationTestFlags::ApplicationContextMask | EAutomationTestFlags::ProductFilter)
+bool InsightsEnumerate10K::RunTest(const FString& Parameters)
+{
+	FTimingProfilerTests::FEnumerateTestParams Params;
+	Params.Interval = 0.01;
+	Params.NumEnumerations = 10000;
+
+	FTimingProfilerTests::FCheckValues CheckValues;
+	FTimingProfilerTests::RunEnumerateBenchmark(Params, CheckValues);
+
+	return !HasAnyErrors();
+}
+
+IMPLEMENT_SIMPLE_AUTOMATION_TEST(InsightsEnumerate100K, "Insights.Analysis.TimingInsights.Enumerate100K", EAutomationTestFlags::ApplicationContextMask | EAutomationTestFlags::ProductFilter)
+bool InsightsEnumerate100K::RunTest(const FString& Parameters)
+{
+	FTimingProfilerTests::FEnumerateTestParams Params;
+	Params.Interval = 0.01;
+	Params.NumEnumerations = 100000;
+
+	FTimingProfilerTests::FCheckValues CheckValues;
+	FTimingProfilerTests::RunEnumerateBenchmark(Params, CheckValues);
+
+	return !HasAnyErrors();
+}
+
+IMPLEMENT_SIMPLE_AUTOMATION_TEST(EnumerateByEndTimeAsyncAllTracks, "Insights.Analysis.TimingInsights.EnumerateByEndTimeAsyncAllTracks", EAutomationTestFlags::ApplicationContextMask | EAutomationTestFlags::ProductFilter)
+bool EnumerateByEndTimeAsyncAllTracks::RunTest(const FString& Parameters)
+{
+	FTimingProfilerTests::FEnumerateTestParams Params;
+	Params.Interval = 24.0 * 3600; //A day - Should be big enought to contain any valid session in [0, 0 + interval]
+	Params.NumEnumerations = 1;
+	Params.SortOrder = TraceServices::EEventSortOrder::ByEndTime;
+
+	FTimingProfilerTests::RunEnumerateSyncAsyncComparisonTest(*this, Params, false);
+
+	return !HasAnyErrors();
+}
+
+IMPLEMENT_SIMPLE_AUTOMATION_TEST(EnumerateByEndTimeAsyncGameThreadTrack, "Insights.Analysis.TimingInsights.EnumerateByEndTimeAsyncGameThreadTrack", EAutomationTestFlags::ApplicationContextMask | EAutomationTestFlags::ProductFilter)
+bool EnumerateByEndTimeAsyncGameThreadTrack::RunTest(const FString& Parameters)
+{
+	FTimingProfilerTests::FEnumerateTestParams Params;
+	Params.Interval = 24.0 * 3600; //A day - Should be big enought to contain any valid session in [0, 0 + interval]
+	Params.NumEnumerations = 1;
+	Params.SortOrder = TraceServices::EEventSortOrder::ByEndTime;
+
+	FTimingProfilerTests::RunEnumerateSyncAsyncComparisonTest(*this, Params, true);
+
+	return !HasAnyErrors();
+}
+
+IMPLEMENT_SIMPLE_AUTOMATION_TEST(EnumerateByEndTimeAllTracks10sIntervals, "Insights.Analysis.TimingInsights.EnumerateByEndTimeAllTracks10sIntervals", EAutomationTestFlags::ApplicationContextMask | EAutomationTestFlags::ProductFilter)
+bool EnumerateByEndTimeAllTracks10sIntervals::RunTest(const FString& Parameters)
+{
+	FTimingProfilerTests::FEnumerateTestParams Params;
+	Params.Interval = 10.0;
+	Params.NumEnumerations = 100;
+	Params.SortOrder = TraceServices::EEventSortOrder::ByEndTime;
+
+	FTimingProfilerTests::RunEnumerateSyncAsyncComparisonTest(*this, Params, false);
+
+	return !HasAnyErrors();
+}
+
+IMPLEMENT_SIMPLE_AUTOMATION_TEST(EnumerateByEndTimeAllTracks5sIntervals, "Insights.Analysis.TimingInsights.EnumerateByEndTimeAllTracks5sIntervals", EAutomationTestFlags::ApplicationContextMask | EAutomationTestFlags::ProductFilter)
+bool EnumerateByEndTimeAllTracks5sIntervals::RunTest(const FString& Parameters)
+{
+	FTimingProfilerTests::FEnumerateTestParams Params;
+	Params.Interval = 5.0;
+	Params.NumEnumerations = 200;
+	Params.SortOrder = TraceServices::EEventSortOrder::ByEndTime;
+
+	FTimingProfilerTests::RunEnumerateSyncAsyncComparisonTest(*this, Params, false);
+
+	return !HasAnyErrors();
+}
+
+////////////////////////////////////////////////////////////////////////////////////////////////////
+////////////////////////////Enumerate Async Ordered by Start Time Tests/////////////////////////////
+////////////////////////////////////////////////////////////////////////////////////////////////////
+
+IMPLEMENT_SIMPLE_AUTOMATION_TEST(EnumerateByStartTimeAsyncAllTracks, "Insights.Analysis.TimingInsights.EnumerateByStartTimeAsyncAllTracks", EAutomationTestFlags::ApplicationContextMask | EAutomationTestFlags::ProductFilter)
+bool EnumerateByStartTimeAsyncAllTracks::RunTest(const FString& Parameters)
+{
+	FTimingProfilerTests::FEnumerateTestParams Params;
+	Params.Interval = 24.0 * 3600; //A day - Should be big enought to contain any valid session in [0, 0 + interval]
+	Params.NumEnumerations = 1;
+	Params.SortOrder = TraceServices::EEventSortOrder::ByStartTime;
+
+	FTimingProfilerTests::RunEnumerateSyncAsyncComparisonTest(*this, Params, false);
+
+	return !HasAnyErrors();
+}
+
+IMPLEMENT_SIMPLE_AUTOMATION_TEST(EnumerateByStartTimeAsyncGameThreadTrack, "Insights.Analysis.TimingInsights.EnumerateByStartTimeAsyncGameThreadTrack", EAutomationTestFlags::ApplicationContextMask | EAutomationTestFlags::ProductFilter)
+bool EnumerateByStartTimeAsyncGameThreadTrack::RunTest(const FString& Parameters)
+{
+	FTimingProfilerTests::FEnumerateTestParams Params;
+	Params.Interval = 24.0 * 3600; //A day - Should be big enought to contain any valid session in [0, 0 + interval]
+	Params.NumEnumerations = 1;
+	Params.SortOrder = TraceServices::EEventSortOrder::ByStartTime;
+
+	FTimingProfilerTests::RunEnumerateSyncAsyncComparisonTest(*this, Params, true);
+
+	return !HasAnyErrors();
+}
+
+IMPLEMENT_SIMPLE_AUTOMATION_TEST(EnumerateByStartTimeAllTracks10sIntervals, "Insights.Analysis.TimingInsights.EnumerateByStartTimeAllTracks10sIntervals", EAutomationTestFlags::ApplicationContextMask | EAutomationTestFlags::ProductFilter)
+bool EnumerateByStartTimeAllTracks10sIntervals::RunTest(const FString& Parameters)
+{
+	FTimingProfilerTests::FEnumerateTestParams Params;
+	Params.Interval = 10.0;
+	Params.NumEnumerations = 100;
+	Params.SortOrder = TraceServices::EEventSortOrder::ByStartTime;
+
+	FTimingProfilerTests::RunEnumerateSyncAsyncComparisonTest(*this, Params, false);
+
+	return !HasAnyErrors();
+}
+
+IMPLEMENT_SIMPLE_AUTOMATION_TEST(EnumerateByStartTimeAllTracks5sIntervals, "Insights.Analysis.TimingInsights.EnumerateByStartTimeAllTracks5sIntervals", EAutomationTestFlags::ApplicationContextMask | EAutomationTestFlags::ProductFilter)
+bool EnumerateByStartTimeAllTracks5sIntervals::RunTest(const FString& Parameters)
+{
+	FTimingProfilerTests::FEnumerateTestParams Params;
+	Params.Interval = 5.0;
+	Params.NumEnumerations = 200;
+	Params.SortOrder = TraceServices::EEventSortOrder::ByStartTime;
+
+	FTimingProfilerTests::RunEnumerateSyncAsyncComparisonTest(*this, Params, false);
+
+	return !HasAnyErrors();
+}
+#endif //!WITH_EDITOR && WITH_DEV_AUTOMATION_TESTS
+
+void FTimingProfilerTests::RunEnumerateBenchmark(const FEnumerateTestParams& InParams, FCheckValues& OutCheckValues)
+{
+	UE_LOG(TimingProfilerTests, Log, TEXT("RUNNING BENCHMARK..."));
+
+	FStopwatch Stopwatch;
+	Stopwatch.Start();
+
+	double SessionTime = 0.0;
+	uint32 TimelineIndex = (uint32)-1;
+
+	TSharedPtr<const TraceServices::IAnalysisSession> Session = FInsightsManager::Get()->GetSession();
+	if (Session.IsValid() && TraceServices::ReadTimingProfilerProvider(*Session.Get()))
+	{
+		TraceServices::FAnalysisSessionReadScope SessionReadScope(*Session.Get());
+
+		const TraceServices::ITimingProfilerProvider& TimingProfilerProvider = *TraceServices::ReadTimingProfilerProvider(*Session.Get());
+
+		SessionTime = Session->GetDurationSeconds();
+		OutCheckValues.SessionDuration = SessionTime;
+
+		const double TimeIncrement = SessionTime / static_cast<double>(InParams.NumEnumerations);
+
+		TimelineIndex = GetTimelineIndex(TEXT("GameThread"));
+
+		TimingProfilerProvider.ReadTimeline(TimelineIndex,
+			[&OutCheckValues, &InParams, TimeIncrement](const TraceServices::ITimingProfilerProvider::Timeline& Timeline)
+			{
+				double Time = 0.0;
+				for (int32 Index = 0; Index < InParams.NumEnumerations; ++Index)
+				{
+					Timeline.EnumerateEvents(Time, Time + InParams.Interval,
+						[&OutCheckValues](double EventStartTime, double EventEndTime, uint32 EventDepth, const TraceServices::FTimingProfilerEvent& Event)
+						{
+							EventEndTime = FMath::Min(EventEndTime, OutCheckValues.SessionDuration);
+							OutCheckValues.TotalEventDuration += EventEndTime - EventStartTime;
+							++OutCheckValues.EventCount;
+							OutCheckValues.SumDepth += EventDepth;
+							OutCheckValues.SumTimerIndex += Event.TimerIndex;
+							return TraceServices::EEventEnumerate::Continue;
+						});
+
+					Time += TimeIncrement;
+				}
+			});
+	}
+
+	Stopwatch.Stop();
+	OutCheckValues.EnumerationDuration = Stopwatch.GetAccumulatedTime();
+	UE_LOG(TimingProfilerTests, Log, TEXT("BENCHMARK RESULT: %f seconds"), OutCheckValues.EnumerationDuration);
+	UE_LOG(TimingProfilerTests, Log, TEXT("SessionTime: %f seconds"), SessionTime);
+	UE_LOG(TimingProfilerTests, Log, TEXT("TimelineIndex: %u"), TimelineIndex);
+	UE_LOG(TimingProfilerTests, Log, TEXT("Check Values: %f %llu %u %u"), OutCheckValues.TotalEventDuration, OutCheckValues.EventCount, OutCheckValues.SumDepth, OutCheckValues.SumTimerIndex);
+}
+
+void FTimingProfilerTests::RunEnumerateAsyncBenchmark(const FEnumerateTestParams& InParams, FCheckValues& OutCheckValues)
+{
+	UE_LOG(TimingProfilerTests, Log, TEXT("RUNNING ASYNC ENUMERATE BENCHMARK..."));
+
+	FStopwatch Stopwatch;
+	Stopwatch.Start();
+
+	double SessionTime = 0.0;
+	uint32 TimelineIndex = (uint32)-1;
+
+	TSharedPtr<const TraceServices::IAnalysisSession> Session = FInsightsManager::Get()->GetSession();
+	if (Session.IsValid() && TraceServices::ReadTimingProfilerProvider(*Session.Get()))
+	{
+		TraceServices::FAnalysisSessionReadScope SessionReadScope(*Session.Get());
+
+		const TraceServices::ITimingProfilerProvider& TimingProfilerProvider = *TraceServices::ReadTimingProfilerProvider(*Session.Get());
+
+		SessionTime = Session->GetDurationSeconds();
+		OutCheckValues.SessionDuration = SessionTime;
+
+		const double TimeIncrement = SessionTime / static_cast<double>(InParams.NumEnumerations);
+
+		TimelineIndex = GetTimelineIndex(TEXT("GameThread"));
+
+		TArray<FCheckValues> TaskCheckValues;
+
+		TimingProfilerProvider.ReadTimeline(TimelineIndex,
+			[SessionTime, &InParams, TimeIncrement, &TaskCheckValues](const TraceServices::ITimingProfilerProvider::Timeline& Timeline)
+			{
+				double Time = 0.0;
+				for (int32 Index = 0; Index < InParams.NumEnumerations; ++Index)
+				{
+					TraceServices::ITimeline<TraceServices::FTimingProfilerEvent>::EnumerateAsyncParams Params;
+					Params.IntervalStart = Time;
+					Params.IntervalEnd = Time + InParams.Interval;
+					Params.Resolution = 0.0;
+					Params.SortOrder = InParams.SortOrder;
+					Params.SetupCallback = [&TaskCheckValues](uint32 NumTasks)
+					{
+						TaskCheckValues.AddDefaulted(NumTasks);
+					};
+					Params.Callback = [&TaskCheckValues, SessionTime](double EventStartTime, double EventEndTime, uint32 EventDepth, const TraceServices::FTimingProfilerEvent& Event, uint32 TaskIndex)
+					{
+						EventEndTime = FMath::Min(EventEndTime, SessionTime);
+						TaskCheckValues[TaskIndex].TotalEventDuration += EventEndTime - EventStartTime;
+						++TaskCheckValues[TaskIndex].EventCount;
+						TaskCheckValues[TaskIndex].SumDepth += EventDepth;
+						TaskCheckValues[TaskIndex].SumTimerIndex += Event.TimerIndex;
+						return TraceServices::EEventEnumerate::Continue;
+					};
+
+					Timeline.EnumerateEventsDownSampledAsync(Params);
+
+					Time += TimeIncrement;
+				}
+			});
+
+		for (auto& CheckValues : TaskCheckValues)
+		{
+			OutCheckValues.TotalEventDuration += CheckValues.TotalEventDuration;
+			OutCheckValues.EventCount += CheckValues.EventCount;
+			OutCheckValues.SumDepth += CheckValues.SumDepth;
+			OutCheckValues.SumTimerIndex += CheckValues.SumTimerIndex;
+		}
+	}
+
+	Stopwatch.Stop();
+	OutCheckValues.EnumerationDuration = Stopwatch.GetAccumulatedTime();
+	UE_LOG(TimingProfilerTests, Log, TEXT("ASYNC ENUMERATE BENCHMARK RESULT: %f seconds"), OutCheckValues.EnumerationDuration);
+	UE_LOG(TimingProfilerTests, Log, TEXT("SessionTime: %f seconds"), SessionTime);
+	UE_LOG(TimingProfilerTests, Log, TEXT("TimelineIndex: %u"), TimelineIndex);
+	UE_LOG(TimingProfilerTests, Log, TEXT("Check Values: %f %llu %u %u"), OutCheckValues.TotalEventDuration, OutCheckValues.EventCount, OutCheckValues.SumDepth, OutCheckValues.SumTimerIndex);
+}
+
+void FTimingProfilerTests::RunEnumerateAllTracksBenchmark(const FEnumerateTestParams& InParams, FCheckValues& OutCheckValues)
+{
+	UE_LOG(TimingProfilerTests, Log, TEXT("RUNNING ENUMERATE ALL TRACKS BENCHMARK..."));
+
+	FStopwatch Stopwatch;
+	Stopwatch.Start();
+
+	double SessionTime = 0.0;
+	uint32 TimelineIndex = (uint32)-1;
+
+	TSharedPtr<const TraceServices::IAnalysisSession> Session = FInsightsManager::Get()->GetSession();
+	if (Session.IsValid() && TraceServices::ReadTimingProfilerProvider(*Session.Get()))
+	{
+		TraceServices::FAnalysisSessionReadScope SessionReadScope(*Session.Get());
+
+		const TraceServices::ITimingProfilerProvider& TimingProfilerProvider = *TraceServices::ReadTimingProfilerProvider(*Session.Get());
+
+		SessionTime = Session->GetDurationSeconds();
+		OutCheckValues.SessionDuration = SessionTime;
+
+		const double TimeIncrement = SessionTime / static_cast<double>(InParams.NumEnumerations);
+
+		const TraceServices::IThreadProvider& ThreadProvider = TraceServices::ReadThreadProvider(*Session.Get());
+		ThreadProvider.EnumerateThreads(
+			[&TimelineIndex, &TimingProfilerProvider, &OutCheckValues, &InParams, TimeIncrement](const TraceServices::FThreadInfo& ThreadInfo)
+			{
+				TimingProfilerProvider.GetCpuThreadTimelineIndex(ThreadInfo.Id, TimelineIndex);
+
+				TimingProfilerProvider.ReadTimeline(TimelineIndex,
+					[&OutCheckValues, &InParams, TimeIncrement](const TraceServices::ITimingProfilerProvider::Timeline& Timeline)
+					{
+						double Time = 0.0;
+						for (int32 Index = 0; Index < InParams.NumEnumerations; ++Index)
+						{
+							Timeline.EnumerateEvents(Time, Time + InParams.Interval,
+								[&OutCheckValues](double EventStartTime, double EventEndTime, uint32 EventDepth, const TraceServices::FTimingProfilerEvent& Event)
+								{
+									EventEndTime = FMath::Min(EventEndTime, OutCheckValues.SessionDuration);
+									OutCheckValues.TotalEventDuration += EventEndTime - EventStartTime;
+									++OutCheckValues.EventCount;
+									OutCheckValues.SumDepth += EventDepth;
+									OutCheckValues.SumTimerIndex += Event.TimerIndex;
+									return TraceServices::EEventEnumerate::Continue;
+								});
+
+							Time += TimeIncrement;
+						}
+					});
+			});
+	}
+
+	Stopwatch.Stop();
+	OutCheckValues.EnumerationDuration = Stopwatch.GetAccumulatedTime();
+	UE_LOG(TimingProfilerTests, Log, TEXT("ENUMERATE ALL TRACKS BENCHMARK RESULT: %f seconds"), OutCheckValues.EnumerationDuration);
+	UE_LOG(TimingProfilerTests, Log, TEXT("SessionTime: %f seconds"), SessionTime);
+	UE_LOG(TimingProfilerTests, Log, TEXT("TimelineIndex: %u"), TimelineIndex);
+	UE_LOG(TimingProfilerTests, Log, TEXT("Check Values: %f %llu %u %u"), OutCheckValues.TotalEventDuration, OutCheckValues.EventCount, OutCheckValues.SumDepth, OutCheckValues.SumTimerIndex);
+}
+
+void FTimingProfilerTests::RunEnumerateAsyncAllTracksBenchmark(const FEnumerateTestParams& InParams, FCheckValues& OutCheckValues)
+{
+	UE_LOG(TimingProfilerTests, Log, TEXT("RUNNING ASYNC ENUMERATE ALL TRACKS BENCHMARK..."));
+
+	FStopwatch Stopwatch;
+	Stopwatch.Start();
+
+	double SessionTime = 0.0;
+	uint32 TimelineIndex = (uint32)-1;
+
+	TSharedPtr<const TraceServices::IAnalysisSession> Session = FInsightsManager::Get()->GetSession();
+	if (Session.IsValid() && TraceServices::ReadTimingProfilerProvider(*Session.Get()))
+	{
+		TraceServices::FAnalysisSessionReadScope SessionReadScope(*Session.Get());
+
+		const TraceServices::ITimingProfilerProvider& TimingProfilerProvider = *TraceServices::ReadTimingProfilerProvider(*Session.Get());
+
+		SessionTime = Session->GetDurationSeconds();
+		OutCheckValues.SessionDuration = SessionTime;
+
+		const double TimeIncrement = SessionTime / static_cast<double>(InParams.NumEnumerations);
+
+		TArray<FCheckValues> TaskCheckValues;
+
+		const TraceServices::IThreadProvider& ThreadProvider = TraceServices::ReadThreadProvider(*Session.Get());
+		ThreadProvider.EnumerateThreads(
+			[&TimelineIndex, &TimingProfilerProvider, &OutCheckValues, &InParams, TimeIncrement, &TaskCheckValues](const TraceServices::FThreadInfo& ThreadInfo)
+			{
+				TimingProfilerProvider.GetCpuThreadTimelineIndex(ThreadInfo.Id, TimelineIndex);
+
+				TimingProfilerProvider.ReadTimeline(TimelineIndex,
+					[&OutCheckValues, &InParams, TimeIncrement, &TaskCheckValues](const TraceServices::ITimingProfilerProvider::Timeline& Timeline)
+					{
+						double Time = 0.0;
+						for (int32 Index = 0; Index < InParams.NumEnumerations; ++Index)
+						{
+							TraceServices::ITimeline<TraceServices::FTimingProfilerEvent>::EnumerateAsyncParams Params;
+							Params.IntervalStart = Time;
+							Params.IntervalEnd = Time + InParams.Interval;
+							Params.Resolution = 0.0;
+							Params.SortOrder = InParams.SortOrder;
+							Params.SetupCallback = [&TaskCheckValues](uint32 NumTasks)
+							{
+								TaskCheckValues.AddDefaulted(NumTasks);
+							};
+							Params.Callback = [&OutCheckValues, &TaskCheckValues](double EventStartTime, double EventEndTime, uint32 EventDepth, const TraceServices::FTimingProfilerEvent& Event, uint32 TaskIndex)
+							{
+								EventEndTime = FMath::Min(EventEndTime, OutCheckValues.SessionDuration);
+								TaskCheckValues[TaskIndex].TotalEventDuration += EventEndTime - EventStartTime;
+								++TaskCheckValues[TaskIndex].EventCount;
+								TaskCheckValues[TaskIndex].SumDepth += EventDepth;
+								TaskCheckValues[TaskIndex].SumTimerIndex += Event.TimerIndex;
+								return TraceServices::EEventEnumerate::Continue;
+							};
+
+							Timeline.EnumerateEventsDownSampledAsync(Params);
+
+							for (auto& CheckValues : TaskCheckValues)
+							{
+								OutCheckValues.TotalEventDuration += CheckValues.TotalEventDuration;
+								OutCheckValues.EventCount += CheckValues.EventCount;
+								OutCheckValues.SumDepth += CheckValues.SumDepth;
+								OutCheckValues.SumTimerIndex += CheckValues.SumTimerIndex;
+							}
+
+							TaskCheckValues.Empty();
+							Time += TimeIncrement;
+						}
+					});
+			});
+	}
+
+	Stopwatch.Stop();
+	OutCheckValues.EnumerationDuration = Stopwatch.GetAccumulatedTime();
+	UE_LOG(TimingProfilerTests, Log, TEXT("ASYNC ENUMERATE ALL TRACKS BENCHMARK RESULT: %f seconds"), OutCheckValues.EnumerationDuration);
+	UE_LOG(TimingProfilerTests, Log, TEXT("SessionTime: %f seconds"), SessionTime);
+	UE_LOG(TimingProfilerTests, Log, TEXT("TimelineIndex: %u"), TimelineIndex);
+	UE_LOG(TimingProfilerTests, Log, TEXT("Check Values: %f %llu %u %u"), OutCheckValues.TotalEventDuration, OutCheckValues.EventCount, OutCheckValues.SumDepth, OutCheckValues.SumTimerIndex);
+}
+
+uint32 FTimingProfilerTests::GetTimelineIndex(const TCHAR* InName)
+{
+	TSharedPtr<const TraceServices::IAnalysisSession> Session = FInsightsManager::Get()->GetSession();
+	const TraceServices::ITimingProfilerProvider& TimingProfilerProvider = *TraceServices::ReadTimingProfilerProvider(*Session.Get());
+	const TraceServices::IThreadProvider& ThreadProvider = TraceServices::ReadThreadProvider(*Session.Get());
+	uint32 TimelineIndex = (uint32)-1;
+	ThreadProvider.EnumerateThreads(
+		[&TimelineIndex, &TimingProfilerProvider, InName](const TraceServices::FThreadInfo& ThreadInfo)
+		{
+			if (!FCString::Strcmp(ThreadInfo.Name, InName))
+			{
+				TimingProfilerProvider.GetCpuThreadTimelineIndex(ThreadInfo.Id, TimelineIndex);
+			}
+		});
+
+	return TimelineIndex;
+}
+
+bool FTimingProfilerTests::RunEnumerateSyncAsyncComparisonTest(FAutomationTestBase& Test, const FEnumerateTestParams& InParams, bool bGameThreadOnly)
+{
+	FTimingProfilerTests::FCheckValues CheckValues;
+	FTimingProfilerTests::FCheckValues CheckValuesAsync;
+
+	if (bGameThreadOnly)
+	{
+		FTimingProfilerTests::RunEnumerateBenchmark(InParams, CheckValues);
+		FTimingProfilerTests::RunEnumerateAsyncBenchmark(InParams, CheckValuesAsync);
+	}
+	else
+	{
+		FTimingProfilerTests::RunEnumerateAllTracksBenchmark(InParams, CheckValues);
+		FTimingProfilerTests::RunEnumerateAsyncAllTracksBenchmark(InParams, CheckValuesAsync);
+	}
+
+	FTimingProfilerTests::VerifyCheckValues(Test, CheckValues, CheckValuesAsync);
+
+	Test.AddInfo(FString::Printf(TEXT("Enumeration Duration: %f seconds."), CheckValues.EnumerationDuration));
+	Test.AddInfo(FString::Printf(TEXT("Async Enumeration Duration: %f seconds."), CheckValuesAsync.EnumerationDuration));
+
+	return !Test.HasAnyErrors();
+}
+
+void FTimingProfilerTests::VerifyCheckValues(FAutomationTestBase& Test, FCheckValues First, FCheckValues Second)
+{
+	Test.TestEqual(TEXT("SessionDuration"), First.SessionDuration, Second.SessionDuration, 1.e-6);
+	Test.TestEqual(TEXT("TotalEventDuration"), First.TotalEventDuration, Second.TotalEventDuration, 1.e-3);
+	Test.TestEqual(TEXT("EventCount"), First.EventCount, Second.EventCount);
+	Test.TestEqual(TEXT("SumDepth"), (uint64)First.SumDepth, (uint64)Second.SumDepth);
+	Test.TestEqual(TEXT("SumTimerIndex"), (uint64)First.SumTimerIndex, (uint64)Second.SumTimerIndex);
+}
\ No newline at end of file
diff --git a/Engine/Source/Developer/TraceInsights/Private/Insights/Tests/FunctionalTests/UITests.cpp b/Engine/Source/Developer/TraceInsights/Private/Insights/Tests/FunctionalTests/UITests.cpp
new file mode 100644
--- /dev/null
+++ b/Engine/Source/Developer/TraceInsights/Private/Insights/Tests/FunctionalTests/UITests.cpp
@@ -0,0 +1,121 @@
+// Copyright Epic Games, Inc. All Rights Reserved.
+
+#include "CoreMinimal.h"
+#include "Insights/Common/Stopwatch.h"
+#include "Insights/MemoryProfiler/ViewModels/MemAllocFilterValueConverter.h"
+#include "Insights/TaskGraphProfiler/ViewModels/TaskTimingTrack.h"
+#include "Insights/TimingProfilerManager.h"
+#include "Insights/ViewModels/TimeFilterValueConverter.h"
+#include "Insights/ViewModels/TimeRulerTrack.h"
+#include "Insights/Widgets/STimingView.h"
+#include "Logging/LogMacros.h"
+#include "Misc/AutomationTest.h"
+#include "TraceServices/Model/TimingProfiler.h"
+
+DECLARE_LOG_CATEGORY_EXTERN(UITests, Log, All);
+
+IMPLEMENT_SIMPLE_AUTOMATION_TEST(FHideAndShowAllTimingViewTabs, "Insights.Analysis.TimingInsights.HideAndShowAllTimingViewTabs", EAutomationTestFlags::ApplicationContextMask | EAutomationTestFlags::ProductFilter)
+bool FHideAndShowAllTimingViewTabs::RunTest(const FString& Parameters)
+{
+	TSharedPtr<FTimingProfilerManager> TimingProfilerManager = FTimingProfilerManager::Get();
+
+	TimingProfilerManager->ShowHideTimingView(false);
+	TimingProfilerManager->ShowHideCalleesTreeView(false);
+	TimingProfilerManager->ShowHideCallersTreeView(false);
+	TimingProfilerManager->ShowHideFramesTrack(false);
+	TimingProfilerManager->ShowHideLogView(false);
+	TimingProfilerManager->ShowHideTimersView(false);
+
+	TimingProfilerManager->ShowHideTimingView(true);
+	TimingProfilerManager->ShowHideCalleesTreeView(true);
+	TimingProfilerManager->ShowHideCallersTreeView(true);
+	TimingProfilerManager->ShowHideFramesTrack(true);
+	TimingProfilerManager->ShowHideLogView(true);
+	TimingProfilerManager->ShowHideTimersView(true);
+
+	return true;
+}
+
+IMPLEMENT_SIMPLE_AUTOMATION_TEST(FMemoryFilterValueConverterTest, "Insights.Analysis.MemoryFilterValueConverter", EAutomationTestFlags::ApplicationContextMask | EAutomationTestFlags::ProductFilter)
+bool FMemoryFilterValueConverterTest::RunTest(const FString& Parameters)
+{
+	Insights::FMemoryFilterValueConverter Converter;
+
+	FText Error;
+	int64 Value;
+
+	Converter.Convert(TEXT("152485"), Value, Error);
+	TestEqual(TEXT("BasicValue"), 152485LL, Value);
+
+	Converter.Convert(TEXT("125.56"), Value, Error);
+	TestEqual(TEXT("DoubleValue"), 125LL, Value);
+
+	Converter.Convert(TEXT("3 KiB"), Value, Error);
+	TestEqual(TEXT("Kib"), 3072LL, Value);
+
+	Converter.Convert(TEXT("7.14 KiB"), Value, Error);
+	TestEqual(TEXT("KibDouble"), 7311LL, Value);
+
+	Converter.Convert(TEXT("5 MiB"), Value, Error);
+	TestEqual(TEXT("Mib"), 5242880LL, Value);
+
+	Converter.Convert(TEXT("1 EiB"), Value, Error);
+	TestEqual(TEXT("Eib"), 1152921504606846976LL, Value);
+
+	Converter.Convert(TEXT("2 kib"), Value, Error);
+	TestEqual(TEXT("CaseInsesitive"), 2048LL, Value);
+
+	TestFalse(TEXT("Fail1"), Converter.Convert(TEXT("23test"), Value, Error));
+	TestFalse(TEXT("Fail2"), Converter.Convert(TEXT("43 kOb"), Value, Error));
+	TestFalse(TEXT("FailInvalidChar"), Converter.Convert(TEXT("45,"), Value, Error));
+
+	return !HasAnyErrors();
+}
+
+IMPLEMENT_SIMPLE_AUTOMATION_TEST(FTimeFilterValueConverterTest, "Insights.Analysis.TimeFilterValueConverter", EAutomationTestFlags::ApplicationContextMask | EAutomationTestFlags::ProductFilter)
+bool FTimeFilterValueConverterTest::RunTest(const FString& Parameters)
+{
+	Insights::FTimeFilterValueConverter Converter;
+
+	FText Error;
+	double Value;
+
+	Converter.Convert(TEXT("15.3"), Value, Error);
+	TestEqual(TEXT("ValueInSeconds"), 15.3, Value);
+
+	Converter.Convert(TEXT("0.2"), Value, Error);
+	TestEqual(TEXT("ValueInSeconds2"), 0.2, Value);
+
+	Converter.Convert(TEXT("0.4s"), Value, Error);
+	TestEqual(TEXT("ValueInSeconds3"), 0.4, Value);
+
+	Converter.Convert(TEXT("125.56ms"), Value, Error);
+	TestEqual(TEXT("ValueInMiliseconds"), 0.12556, Value);
+
+	Converter.Convert(TEXT("14.2 µs"), Value, Error);
+	TestEqual(TEXT("ValueInMicroseconds"), 14.2 * 1.e-6, Value);
+
+	Converter.Convert(TEXT("0.72us"), Value, Error);
+	TestEqual(TEXT("ValueInMicroseconds2"), 0.72 * 1.e-6, Value); 
+
+	Converter.Convert(TEXT("3ns"), Value, Error);
+	TestEqual(TEXT("ValueInNanoseconds"), 3 * 1.e-9, Value);
+
+	Converter.Convert(TEXT("17ns"), Value, Error);
+	TestEqual(TEXT("ValueInPicoseconds"), 17 * 1.e-12, Value);
+
+	Converter.Convert(TEXT("0.5m"), Value, Error);
+	TestEqual(TEXT("ValueInMinutes"), 30.0, Value);
+
+	Converter.Convert(TEXT("1.1h"), Value, Error);
+	TestEqual(TEXT("ValueInHours"), 3960.0, Value);
+
+	Converter.Convert(TEXT("2d"), Value, Error);
+	TestEqual(TEXT("ValueInHours"), 2 * 60 * 60 * 24.0, Value);
+
+	TestFalse(TEXT("Fail1"), Converter.Convert(TEXT("1.23ss"), Value, Error));
+	TestFalse(TEXT("Fail2"), Converter.Convert(TEXT("abc"), Value, Error));
+	TestFalse(TEXT("FailInvalidChar"), Converter.Convert(TEXT("0,2"), Value, Error));
+
+	return !HasAnyErrors();
+}
\ No newline at end of file
diff --git a/Engine/Source/Developer/TraceInsights/Private/Insights/Tests/FunctionalTests/UnrealTraceServer.cpp b/Engine/Source/Developer/TraceInsights/Private/Insights/Tests/FunctionalTests/UnrealTraceServer.cpp
new file mode 100644
--- /dev/null
+++ b/Engine/Source/Developer/TraceInsights/Private/Insights/Tests/FunctionalTests/UnrealTraceServer.cpp
@@ -0,0 +1,116 @@
+// Copyright Epic Games, Inc. All Rights Reserved.
+
+#include "HAL/PlatformFileManager.h"
+#include "Insights/InsightsManager.h"
+#include "Insights/IUnrealInsightsModule.h"
+#include "Insights/Tests/InsightsTestUtils.h"
+#include "Insights/Widgets/SStartPageWindow.h"
+#include "Logging/LogMacros.h"
+#include "Misc/AutomationTest.h"
+
+#if !WITH_EDITOR
+
+IMPLEMENT_SIMPLE_AUTOMATION_TEST(FUnrealTraceServerStartingTest, "Insights.Hub.UnrealTraceServer.Starting", EAutomationTestFlags::ApplicationContextMask | EAutomationTestFlags::ProductFilter)
+bool FUnrealTraceServerStartingTest::RunTest(const FString& Parameters)
+{
+	TSharedPtr<FInsightsManager> InsightsManager = FInsightsManager::Get();
+	const FString UnrealTraceServerName = TEXT("UnrealTraceServer");
+	double Timeout = 10.0;
+
+	auto WaitForUTSProcess = [](const FString& ProcessName, double Timeout, const bool ExpectedStatus) -> bool
+	{
+		double StartTime = FPlatformTime::Seconds();
+		while (FPlatformTime::Seconds() - StartTime < Timeout)
+		{
+			if (FPlatformProcess::IsApplicationRunning(*ProcessName) == ExpectedStatus)
+			{
+				return ExpectedStatus;
+				break;
+			}
+			FPlatformProcess::Sleep(0.1f);
+		}
+		return !ExpectedStatus;
+	};
+
+	auto WaitForConnectionIconStatus = [](TSharedPtr<class FInsightsManager> InsightsManager, double Timeout, const bool ExpectedStatus) -> bool
+	{
+		double StartTime = FPlatformTime::Seconds();
+		while (FPlatformTime::Seconds() - StartTime < Timeout)
+		{
+			TSharedPtr<STraceStoreWindow> Wnd = InsightsManager->GetTraceStoreWindow();
+			if (Wnd.IsValid())
+			{
+				if (Wnd->IsConnected() == ExpectedStatus)
+				{
+					return ExpectedStatus;
+					break;
+				}
+			}
+			FPlatformProcess::Sleep(0.1f);
+		}
+		return !ExpectedStatus;
+	};
+
+	bool bIsUnrealTraceServerRunning = WaitForUTSProcess(UnrealTraceServerName, Timeout, true);
+	TestTrue("UnrealTraceServer should be in processes", bIsUnrealTraceServerRunning);
+
+	bool bConnectionStatus = WaitForConnectionIconStatus(InsightsManager, Timeout, true);
+	TestTrue("Connection status should be true", bConnectionStatus);
+
+#if PLATFORM_WINDOWS
+	FString Command = FString::Printf(TEXT("/F /IM \"%s.exe\""), *UnrealTraceServerName);
+	FPlatformProcess::ExecProcess(TEXT("taskkill"), *Command, nullptr, nullptr, nullptr);
+#elif PLATFORM_MAC || PLATFORM_LINUX
+	FString Command = FString::Printf(TEXT("killall \"%s\""), *UnrealTraceServerName);
+	FPlatformProcess::ExecProcess(TEXT("/bin/sh"), *Command, nullptr, nullptr, nullptr);
+#endif
+
+	bIsUnrealTraceServerRunning = WaitForUTSProcess(UnrealTraceServerName, Timeout, false);
+	TestFalse("UnrealTraceServer should not be in processes", bIsUnrealTraceServerRunning);
+
+	bConnectionStatus = WaitForConnectionIconStatus(InsightsManager, Timeout, false);
+	TestFalse("Connection status should be false", bConnectionStatus);
+
+	FString UEPath = FPlatformProcess::GenerateApplicationPath("UnrealEditor", EBuildConfiguration::Development);
+	FString UEParameters = TEXT("");
+	constexpr bool bLaunchDetached = true;
+	constexpr bool bLaunchHidden = false;
+	constexpr bool bLaunchReallyHidden = false;
+	uint32 ProcessID = 0;
+	const int32 PriorityModifier = 0;
+	const TCHAR* OptionalWorkingDirectory = nullptr;
+	void* PipeWriteChild = nullptr;
+	void* PipeReadChild = nullptr;
+	FProcHandle EditorHandle = FPlatformProcess::CreateProc(*UEPath, *UEParameters, bLaunchDetached, bLaunchHidden, bLaunchReallyHidden, &ProcessID, PriorityModifier, OptionalWorkingDirectory, PipeWriteChild, PipeReadChild);
+	if (!EditorHandle.IsValid())
+	{
+		AddError("Editor should be started");
+		return true;
+	}
+
+	bIsUnrealTraceServerRunning = WaitForUTSProcess(UnrealTraceServerName, Timeout, true);
+	TestTrue("UnrealTraceServer should be in processes", bIsUnrealTraceServerRunning);
+
+	bConnectionStatus = WaitForConnectionIconStatus(InsightsManager, Timeout, true);
+	TestTrue("Connection status should be true", bConnectionStatus);
+
+	FPlatformProcess::TerminateProc(EditorHandle);
+
+#if PLATFORM_WINDOWS
+	FPlatformProcess::ExecProcess(TEXT("taskkill"), *Command, nullptr, nullptr, nullptr);
+#elif PLATFORM_MAC || PLATFORM_LINUX
+	FPlatformProcess::ExecProcess(TEXT("/bin/sh"), *Command, nullptr, nullptr, nullptr);
+#endif
+
+	bIsUnrealTraceServerRunning = WaitForUTSProcess(UnrealTraceServerName, Timeout, false);
+	TestFalse("UnrealTraceServer should not be in processes", bIsUnrealTraceServerRunning);
+
+	FString Path = FPaths::RootDir() / TEXT("EngineTest/SourceAssets/Utrace/r423_win64_game_10478456.utrace");
+	InsightsManager->OpenUnrealInsights(*(FString::Printf(TEXT("-InsightsTest -AutoQuit -NoUI -OpenTraceFile=\"%s\""), *Path)));
+	bIsUnrealTraceServerRunning = WaitForUTSProcess(UnrealTraceServerName, Timeout, true);
+	TestTrue("UnrealTraceServer should be in processes", bIsUnrealTraceServerRunning);
+
+	return true;
+}
+
+#endif //!WITH_EDITOR
\ No newline at end of file
diff --git a/Engine/Source/Developer/TraceInsights/Private/Insights/Tests/InsightsTestRunner.cpp b/Engine/Source/Developer/TraceInsights/Private/Insights/Tests/InsightsTestRunner.cpp
--- a/Engine/Source/Developer/TraceInsights/Private/Insights/Tests/InsightsTestRunner.cpp
+++ b/Engine/Source/Developer/TraceInsights/Private/Insights/Tests/InsightsTestRunner.cpp
@@ -17,7 +17,6 @@
 // Insights
 #include "Insights/Common/InsightsMenuBuilder.h"
 #include "Insights/InsightsManager.h"
-#include "Insights/Tests/UITests.h"
 #include "Insights/InsightsStyle.h"
 
 ////////////////////////////////////////////////////////////////////////////////////////////////////
diff --git a/Engine/Source/Developer/TraceInsights/Private/Insights/Tests/MiscTests.cpp b/Engine/Source/Developer/TraceInsights/Private/Insights/Tests/MiscTests.cpp
deleted file mode 100644
diff --git a/Engine/Source/Developer/TraceInsights/Private/Insights/Tests/TimingProfilerTests.cpp b/Engine/Source/Developer/TraceInsights/Private/Insights/Tests/TimingProfilerTests.cpp
deleted file mode 100644
diff --git a/Engine/Source/Developer/TraceInsights/Private/Insights/Tests/UITests.cpp b/Engine/Source/Developer/TraceInsights/Private/Insights/Tests/UITests.cpp
deleted file mode 100644
diff --git a/Engine/Source/Developer/TraceInsights/Private/Insights/Tests/UITests.h b/Engine/Source/Developer/TraceInsights/Private/Insights/Tests/UITests.h
deleted file mode 100644
diff --git a/Engine/Source/Developer/TraceInsights/Private/Insights/Tests/UnitTests.cpp b/Engine/Source/Developer/TraceInsights/Private/Insights/Tests/UnitTests.cpp
new file mode 100644
--- /dev/null
+++ b/Engine/Source/Developer/TraceInsights/Private/Insights/Tests/UnitTests.cpp
@@ -0,0 +1,10 @@
+// Copyright Epic Games, Inc. All Rights Reserved.
+#include "CoreMinimal.h"
+#include "Misc/AutomationTest.h"
+
+IMPLEMENT_SIMPLE_AUTOMATION_TEST(FTraceInsightsUnitTest, "Insights.TraceInsightsUnitTest", EAutomationTestFlags::ApplicationContextMask | EAutomationTestFlags::SmokeFilter)
+
+bool FTraceInsightsUnitTest::RunTest(const FString& Parameters)
+{
+	return true;
+}
\ No newline at end of file
diff --git a/Engine/Source/Developer/TraceServices/Private/Tests/UnitTests.cpp b/Engine/Source/Developer/TraceServices/Private/Tests/UnitTests.cpp
new file mode 100644
--- /dev/null
+++ b/Engine/Source/Developer/TraceServices/Private/Tests/UnitTests.cpp
@@ -0,0 +1,82 @@
+// Copyright Epic Games, Inc. All Rights Reserved.
+#include "CoreMinimal.h"
+#include "Common/SlabAllocator.h"
+#include "Common/PagedArray.h"
+#include "Misc/AutomationTest.h"
+
+IMPLEMENT_SIMPLE_AUTOMATION_TEST(FPagedArrayFilteringTest, "Insights.Analysis.PagedArrayFiltering", EAutomationTestFlags::ApplicationContextMask | EAutomationTestFlags::SmokeFilter)
+
+bool operator!=(const FInt32Interval& lhs, const FInt32Interval& rhs)
+{
+	return lhs.Min != rhs.Min || lhs.Max != rhs.Max;
+}
+
+bool FPagedArrayFilteringTest::RunTest(const FString& Parameters)
+{
+	using namespace TraceServices;
+
+	TArray<int32> Integers{ 0,1,2,3,4,5,6,7,8,9,10 };
+	TArray<int32> Integer{ 1 };
+
+	//first element >= Value
+	TestEqual(TEXT("First index of element found in range"), Algo::LowerBound(Integers, 4), 4);
+	TestEqual(TEXT("First index of element not in range"), Algo::LowerBound(Integers, 100), 11);
+	//first element > Value
+	TestEqual(TEXT("Upper bound value found in range"), Algo::UpperBound(Integers, 4), 5);
+	TestEqual(TEXT("Upper bound value not found in range"), Algo::UpperBound(Integers, 100), 11);
+
+	TestEqual(TEXT("Upper bound value found in single-item range"), Algo::UpperBound(Integer, 10), 1);
+
+	FSlabAllocator alloc(32 << 20);
+	int32 PageSize = 4;
+
+	struct FTimeRegion
+	{
+		double BeginTime;
+		double EndTime;
+	};
+
+	TPagedArray<FTimeRegion> EmptyLane(alloc, PageSize);
+
+	TPagedArray<FTimeRegion> OneItemLane(alloc, PageSize);
+	OneItemLane.EmplaceBack(FTimeRegion{ 1.0, 2.0 });
+
+	TPagedArray<FTimeRegion> TestLane(alloc, PageSize);
+	TestLane.EmplaceBack(FTimeRegion{ 1.0, 2.0 });
+	TestLane.EmplaceBack(FTimeRegion{ 3.0, 4.0 });
+	TestLane.EmplaceBack(FTimeRegion{ 5.0, 6.0 });
+	TestLane.EmplaceBack(FTimeRegion{ 7.0, 8.0 });
+	TestLane.EmplaceBack(FTimeRegion{ 9.0, 10.0 });
+	TestLane.EmplaceBack(FTimeRegion{ 11.0, 12.0 });
+	TestLane.EmplaceBack(FTimeRegion{ 13.0, 14.0 });
+
+	FInt32Interval Result;
+	auto BeginProj = [](const FTimeRegion& r) { return r.BeginTime; };
+	auto EndProj = [](const FTimeRegion& r) { return r.EndTime; };
+
+	Result = GetElementRangeOverlappingGivenRange<FTimeRegion>(EmptyLane, -10.0, 0.0, BeginProj, EndProj);
+	TestEqual(TEXT("No overlaps for empty range"), Result, FInt32Interval{ -1,-1 });
+
+	Result = GetElementRangeOverlappingGivenRange<FTimeRegion>(OneItemLane, -10.0, 10.0, BeginProj, EndProj);
+	TestEqual(TEXT("Overlap is interval is larger than element-range"), Result, { 0,0 });
+
+	Result = GetElementRangeOverlappingGivenRange<FTimeRegion>(TestLane, -10.0, 0.0, BeginProj, EndProj);
+	TestEqual(TEXT("No overlaps for interval before element-range"), Result, { -1,-1 });
+
+	Result = GetElementRangeOverlappingGivenRange<FTimeRegion>(TestLane, 100.0, 200.0, BeginProj, EndProj);
+	TestEqual(TEXT("No overlaps for interval after element-range"), Result, { -1,-1 });
+
+	//partially overlapping begin
+	Result = GetElementRangeOverlappingGivenRange<FTimeRegion>(TestLane, 0, 5, BeginProj, EndProj);
+	TestEqual(TEXT("Overlaps for interval earlier to inside element-range"), Result, { 0,2 });
+
+	// full inside
+	Result = GetElementRangeOverlappingGivenRange<FTimeRegion>(TestLane, 4, 8, BeginProj, EndProj);
+	TestEqual(TEXT("Overlaps for interval inside element-range"), Result, { 2,3 });
+
+	//partially overlapping end
+	Result = GetElementRangeOverlappingGivenRange<FTimeRegion>(TestLane, 11.3, 100, BeginProj, EndProj);
+	TestEqual(TEXT("Overlaps for interval inside to after element-range"), Result, { 5,6 });
+
+	return !HasAnyErrors();
+}
\ No newline at end of file
