Added forgotten Analysis. prefix to Insights unit test

diff --git a/Engine/Source/Developer/TraceInsights/Private/Insights/Tests/UnitTests.cpp b/Engine/Source/Developer/TraceInsights/Private/Insights/Tests/UnitTests.cpp
--- a/Engine/Source/Developer/TraceInsights/Private/Insights/Tests/UnitTests.cpp
+++ b/Engine/Source/Developer/TraceInsights/Private/Insights/Tests/UnitTests.cpp
@@ -2,7 +2,7 @@
 #include "CoreMinimal.h"
 #include "Misc/AutomationTest.h"
 
-IMPLEMENT_SIMPLE_AUTOMATION_TEST(FTraceInsightsUnitTest, "Insights.TraceInsightsUnitTest", EAutomationTestFlags::ApplicationContextMask | EAutomationTestFlags::SmokeFilter)
+IMPLEMENT_SIMPLE_AUTOMATION_TEST(FTraceInsightsUnitTest, "Insights.Analysis.TraceInsightsUnitTest", EAutomationTestFlags::ApplicationContextMask | EAutomationTestFlags::SmokeFilter)
 
 bool FTraceInsightsUnitTest::RunTest(const FString& Parameters)
 {
