[Insights] Fixed key shortcut for visibility of regions track to Ctrl+R (instread of just R key which was already used by Frame tracks).

diff --git a/Engine/Source/Developer/TraceInsights/Private/Insights/ViewModels/RegionsTimingTrack.cpp b/Engine/Source/Developer/TraceInsights/Private/Insights/ViewModels/RegionsTimingTrack.cpp
--- a/Engine/Source/Developer/TraceInsights/Private/Insights/ViewModels/RegionsTimingTrack.cpp
+++ b/Engine/Source/Developer/TraceInsights/Private/Insights/ViewModels/RegionsTimingTrack.cpp
@@ -58,7 +58,7 @@ void FTimingRegionsViewCommands::RegisterCommands()
 		"Timing Regions Track",
 		"Shows/hides the Timing Regions track.",
 		EUserInterfaceActionType::ToggleButton,
-		FInputChord(EKeys::R));
+		FInputChord(EModifierKey::Control, EKeys::R));
 }
 UE_ENABLE_OPTIMIZATION_SHIP
 
