[Insights] Open file now remembers the last used location until the application is closed.

diff --git a/Engine/Source/Developer/TraceInsights/Private/Insights/InsightsManager.cpp b/Engine/Source/Developer/TraceInsights/Private/Insights/InsightsManager.cpp
--- a/Engine/Source/Developer/TraceInsights/Private/Insights/InsightsManager.cpp
+++ b/Engine/Source/Developer/TraceInsights/Private/Insights/InsightsManager.cpp
@@ -805,7 +805,7 @@ void FInsightsManager::ActivateTimingInsightsTab()
 
 bool FInsightsManager::ShowOpenTraceFileDialog(FString& OutTraceFile) const
 {
-	const FString ProfilingDirectory(FPaths::ConvertRelativePathToFull(FInsightsManager::Get()->GetStoreDir()));
+	static FString DefaultDirectory(FPaths::ConvertRelativePathToFull(FInsightsManager::Get()->GetStoreDir()));
 
 	TArray<FString> OutFiles;
 	bool bOpened = false;
@@ -819,7 +819,7 @@ bool FInsightsManager::ShowOpenTraceFileDialog(FString& OutTraceFile) const
 		(
 			FSlateApplication::Get().FindBestParentWindowHandleForDialogs(nullptr),
 			LOCTEXT("LoadTrace_FileDesc", "Open trace file...").ToString(),
-			ProfilingDirectory,
+			DefaultDirectory,
 			TEXT(""),
 			LOCTEXT("LoadTrace_FileFilter", "Trace files (*.utrace)|*.utrace|All files (*.*)|*.*").ToString(),
 			EFileDialogFlags::None,
@@ -830,6 +830,7 @@ bool FInsightsManager::ShowOpenTraceFileDialog(FString& OutTraceFile) const
 	if (bOpened == true && OutFiles.Num() == 1)
 	{
 		OutTraceFile = OutFiles[0];
+		DefaultDirectory = FPaths::GetPath(OutTraceFile);
 		return true;
 	}
 
