[Insights] Fixed trimming of 64bit numbers to 32bit when displayed as hex (i.e. %x was used instead of %llx for uint64 numbers). This includes truncation of address displayed in Memory Insights when callstack resolving is not completed.

diff --git a/Engine/Source/Developer/TraceInsights/Private/Insights/MemoryProfiler/ViewModels/CallstackFormatting.h b/Engine/Source/Developer/TraceInsights/Private/Insights/MemoryProfiler/ViewModels/CallstackFormatting.h
--- a/Engine/Source/Developer/TraceInsights/Private/Insights/MemoryProfiler/ViewModels/CallstackFormatting.h
+++ b/Engine/Source/Developer/TraceInsights/Private/Insights/MemoryProfiler/ViewModels/CallstackFormatting.h
@@ -87,7 +87,14 @@ inline void FormatStackFrame(const TraceServices::FStackFrame& Frame, FStringBui
 			}
 			if (EnumHasAnyFlags(FormatFlags, EStackFrameFormatFlags::Symbol))
 			{
-				OutString.Appendf(TEXT("0x%08x"), Frame.Addr);
+				if (Frame.Addr == 0)
+				{
+					OutString.Append(TEXT("0x00000000"));
+				}
+				else
+				{
+					OutString.Appendf(TEXT("0x%llX"), Frame.Addr);
+				}
 			}
 			if (EnumHasAnyFlags(FormatFlags, EStackFrameFormatFlags::ModuleAndSymbol))
 			{
diff --git a/Engine/Source/Developer/TraceInsights/Private/Insights/MemoryProfiler/ViewModels/MemAllocGroupingByCallstack.cpp b/Engine/Source/Developer/TraceInsights/Private/Insights/MemoryProfiler/ViewModels/MemAllocGroupingByCallstack.cpp
--- a/Engine/Source/Developer/TraceInsights/Private/Insights/MemoryProfiler/ViewModels/MemAllocGroupingByCallstack.cpp
+++ b/Engine/Source/Developer/TraceInsights/Private/Insights/MemoryProfiler/ViewModels/MemAllocGroupingByCallstack.cpp
@@ -226,11 +226,11 @@ FName FMemAllocGroupingByCallstack::GetGroupName(const TraceServices::FStackFram
 	}
 	else if (Result == TraceServices::ESymbolQueryResult::Pending)
 	{
-		return FName(FString::Printf(TEXT("%s!0x%X [...]"), Frame->Symbol->Module, Frame->Addr), 0);
+		return FName(FString::Printf(TEXT("%s!0x%llX [...]"), Frame->Symbol->Module, Frame->Addr), 0);
 	}
 	else
 	{
-		return FName(FString::Printf(TEXT("%s!0x%X"), Frame->Symbol->Module, Frame->Addr), 0);
+		return FName(FString::Printf(TEXT("%s!0x%llX"), Frame->Symbol->Module, Frame->Addr), 0);
 	}
 }
 
diff --git a/Engine/Source/Developer/TraceInsights/Private/Insights/MemoryProfiler/ViewModels/MemoryGraphTrack.cpp b/Engine/Source/Developer/TraceInsights/Private/Insights/MemoryProfiler/ViewModels/MemoryGraphTrack.cpp
--- a/Engine/Source/Developer/TraceInsights/Private/Insights/MemoryProfiler/ViewModels/MemoryGraphTrack.cpp
+++ b/Engine/Source/Developer/TraceInsights/Private/Insights/MemoryProfiler/ViewModels/MemoryGraphTrack.cpp
@@ -99,7 +99,7 @@ FMemoryGraphTrack::FMemoryGraphTrack(FMemorySharedState& InSharedState)
 
 	for (uint32 Mode = 0; Mode < static_cast<uint32>(EMemoryTrackHeightMode::Count); ++Mode)
 	{
-		SetAvailableTrackHeight(static_cast<EMemoryTrackHeightMode>(Mode), 100.0f * (Mode + 1));
+		SetAvailableTrackHeight(static_cast<EMemoryTrackHeightMode>(Mode), 100.0f * static_cast<float>(Mode + 1));
 	}
 	SetHeight(200.0f);
 }
@@ -1126,7 +1126,7 @@ void FMemoryGraphTrack::InitTooltip(FTooltipDrawState& InOutTooltip, const ITimi
 
 		if (Series->GetTimelineType() == FMemoryGraphSeries::ETimelineType::MemTag)
 		{
-			FString SubTitle = FString::Printf(TEXT("(tag id 0x%X, tracker id %i)"), (uint64)Series->GetTagId(), (int32)Series->GetTrackerId());
+			FString SubTitle = FString::Printf(TEXT("(tag id 0x%llX, tracker id %i)"), (uint64)Series->GetTagId(), (int32)Series->GetTrackerId());
 			InOutTooltip.AddTitle(SubTitle, Series->GetColor());
 		}
 
diff --git a/Engine/Source/Developer/TraceInsights/Private/Insights/MemoryProfiler/ViewModels/MemorySharedState.cpp b/Engine/Source/Developer/TraceInsights/Private/Insights/MemoryProfiler/ViewModels/MemorySharedState.cpp
--- a/Engine/Source/Developer/TraceInsights/Private/Insights/MemoryProfiler/ViewModels/MemorySharedState.cpp
+++ b/Engine/Source/Developer/TraceInsights/Private/Insights/MemoryProfiler/ViewModels/MemorySharedState.cpp
@@ -564,7 +564,7 @@ TSharedPtr<FMemoryGraphTrack> FMemorySharedState::CreateMemTagGraphTrack(Insight
 	}
 	else
 	{
-		SeriesName = FString::Printf(TEXT("Unknown LLM Tag (tag id: 0x%X, tracker id: %i)"), InMemTagId, int32(InMemTrackerId));
+		SeriesName = FString::Printf(TEXT("Unknown LLM Tag (tag id: 0x%llX, tracker id: %i)"), uint64(InMemTagId), int32(InMemTrackerId));
 	}
 
 	const FLinearColor Color = TagPtr ? TagPtr->GetColor() : FLinearColor(0.5f, 0.5f, 0.5f, 1.0f);
diff --git a/Engine/Source/Developer/TraceInsights/Private/Insights/MemoryProfiler/Widgets/SMemAllocTableTreeView.cpp b/Engine/Source/Developer/TraceInsights/Private/Insights/MemoryProfiler/Widgets/SMemAllocTableTreeView.cpp
--- a/Engine/Source/Developer/TraceInsights/Private/Insights/MemoryProfiler/Widgets/SMemAllocTableTreeView.cpp
+++ b/Engine/Source/Developer/TraceInsights/Private/Insights/MemoryProfiler/Widgets/SMemAllocTableTreeView.cpp
@@ -1726,7 +1726,7 @@ void SMemAllocTableTreeView::BuildOpenSourceSubMenu(FMenuBuilder& MenuBuilder)
 						{
 							ItemLabel = FText::Format(LOCTEXT("ContextMenu_OpenSource_Fmt2", "{0} ({1}) \u2192 {2}"),
 								FText::FromString(Frame->Symbol->Module),
-								FText::FromString(FString::Printf(TEXT("0x%X"), Frame->Addr)),
+								FText::FromString(FString::Printf(TEXT("0x%llX"), Frame->Addr)),
 								FText::FromString(TraceServices::QueryResultToString(Frame->Symbol->GetResult())));
 
 							ItemToolTip = FText::Format(LOCTEXT("ContextMenu_OpenSource_Desc_Fmt2", "Open source file of selected callstack frame in {0}."),
diff --git a/Engine/Source/Developer/TraceInsights/Private/Insights/MemoryProfiler/Widgets/SMemTagTreeViewTooltip.cpp b/Engine/Source/Developer/TraceInsights/Private/Insights/MemoryProfiler/Widgets/SMemTagTreeViewTooltip.cpp
--- a/Engine/Source/Developer/TraceInsights/Private/Insights/MemoryProfiler/Widgets/SMemTagTreeViewTooltip.cpp
+++ b/Engine/Source/Developer/TraceInsights/Private/Insights/MemoryProfiler/Widgets/SMemTagTreeViewTooltip.cpp
@@ -100,7 +100,7 @@ TSharedPtr<SToolTip> SMemTagTreeViewTooltip::GetRowTooltip(const TSharedPtr<FMem
 	if (MemTag)
 	{
 		TagText = FText::FromString(MemTag->GetStatName());
-		TagTextEx = FText::FromString(FString::Printf(TEXT(" (0x%X)"), MemTag->GetId()));
+		TagTextEx = FText::FromString(FString::Printf(TEXT(" (0x%llX)"), (uint64)MemTag->GetId()));
 	}
 
 	FText ParentTagText = FText::GetEmpty();
@@ -109,7 +109,7 @@ TSharedPtr<SToolTip> SMemTagTreeViewTooltip::GetRowTooltip(const TSharedPtr<FMem
 	if (ParentMemTag)
 	{
 		ParentTagText = FText::FromString(ParentMemTag->GetStatName());
-		ParentTagTextEx = FText::FromString(FString::Printf(TEXT(" (0x%X)"), ParentMemTag->GetId()));
+		ParentTagTextEx = FText::FromString(FString::Printf(TEXT(" (0x%llX)"), (uint64)ParentMemTag->GetId()));
 	}
 
 #if TOOLTIP_SHOW_AGGREGATED_STATS
diff --git a/Engine/Source/Developer/TraceInsights/Private/Insights/ViewModels/FileActivityTimingTrack.cpp b/Engine/Source/Developer/TraceInsights/Private/Insights/ViewModels/FileActivityTimingTrack.cpp
--- a/Engine/Source/Developer/TraceInsights/Private/Insights/ViewModels/FileActivityTimingTrack.cpp
+++ b/Engine/Source/Developer/TraceInsights/Private/Insights/ViewModels/FileActivityTimingTrack.cpp
@@ -695,13 +695,13 @@ void FFileActivityTimingTrack::InitTooltip(FTooltipDrawState& InOutTooltip, cons
 
 			if (InEvent.FileHandle != uint64(-1))
 			{
-				const FString Value = FString::Printf(TEXT("0x%X"), InEvent.FileHandle);
+				const FString Value = FString::Printf(TEXT("0x%llX"), InEvent.FileHandle);
 				InOutTooltip.AddNameValueTextLine(TEXT("File Handle:"), Value);
 			}
 
 			if (InEvent.ReadWriteHandle != uint64(-1))
 			{
-				const FString Value = FString::Printf(TEXT("0x%X"), InEvent.ReadWriteHandle);
+				const FString Value = FString::Printf(TEXT("0x%llX"), InEvent.ReadWriteHandle);
 				InOutTooltip.AddNameValueTextLine(TEXT("Read/Write Handle:"), Value);
 			}
 
@@ -876,7 +876,8 @@ void FOverviewFileActivityTimingTrack::BuildDrawState(ITimingEventsTrackDrawStat
 					EventName += TEXT(" [!]");
 				}
 
-				if (Width > EventName.Len() * 4.0f + 32.0f)
+				const float MinWidth = static_cast<float>(EventName.Len()) * 4.0f + 32.0f;
+				if (Width > MinWidth)
 				{
 					const double Duration = Event.EndTime - Event.StartTime; // actual event duration
 					FTimingEventsTrackDrawStateBuilder::AppendDurationToEventName(EventName, Duration);
@@ -950,7 +951,8 @@ void FDetailedFileActivityTimingTrack::BuildDrawState(ITimingEventsTrackDrawStat
 				{
 					FString EventName = Activity->Path;
 
-					if (Width > EventName.Len() * 4.0f + 32.0f)
+					const float MinWidth = static_cast<float>(EventName.Len()) * 4.0f + 32.0f;
+					if (Width > MinWidth)
 					{
 						const double Duration = Activity->EndTime - Activity->StartTime;
 						FTimingEventsTrackDrawStateBuilder::AppendDurationToEventName(EventName, Duration);
@@ -1015,7 +1017,8 @@ void FDetailedFileActivityTimingTrack::BuildDrawState(ITimingEventsTrackDrawStat
 					EventName += "]";
 				}
 
-				if (Width > EventName.Len() * 4.0f + 32.0f)
+				const float MinWidth = static_cast<float>(EventName.Len()) * 4.0f + 32.0f;
+				if (Width > MinWidth)
 				{
 					const double Duration = Event.EndTime - Event.StartTime;
 					FTimingEventsTrackDrawStateBuilder::AppendDurationToEventName(EventName, Duration);
