Enabled UnsafeTypeCast warnings as errors in the TraceInsights module.
Fixed conversion warnings.

diff --git a/Engine/Source/Developer/TraceInsights/Private/Insights/ContextSwitches/ViewModels/CpuCoreTimingTrack.cpp b/Engine/Source/Developer/TraceInsights/Private/Insights/ContextSwitches/ViewModels/CpuCoreTimingTrack.cpp
--- a/Engine/Source/Developer/TraceInsights/Private/Insights/ContextSwitches/ViewModels/CpuCoreTimingTrack.cpp
+++ b/Engine/Source/Developer/TraceInsights/Private/Insights/ContextSwitches/ViewModels/CpuCoreTimingTrack.cpp
@@ -211,7 +211,8 @@ void FCpuCoreTimingTrack::AddCoreTimingEvent(ITimingEventsTrackDrawStateBuilder&
 		{
 			FString EventName = GetThreadName(SystemThreadId);
 
-			if (Width > EventName.Len() * 4.0f + 32.0f)
+			const float MinWidth = static_cast<float>(EventName.Len()) * 4.0f + 32.0f;
+			if (Width > MinWidth)
 			{
 				const double Duration = CpuCoreEvent.End - CpuCoreEvent.Start;
 				FTimingEventsTrackDrawStateBuilder::AppendDurationToEventName(EventName, Duration);
diff --git a/Engine/Source/Developer/TraceInsights/Private/Insights/InsightsManager.cpp b/Engine/Source/Developer/TraceInsights/Private/Insights/InsightsManager.cpp
--- a/Engine/Source/Developer/TraceInsights/Private/Insights/InsightsManager.cpp
+++ b/Engine/Source/Developer/TraceInsights/Private/Insights/InsightsManager.cpp
@@ -632,7 +632,7 @@ void FInsightsManager::CheckMemoryUsage()
 		constexpr double MemUsageLimitHysteresisPercent = 50.0;
 
 		const uint64 Time = FPlatformTime::Cycles64();
-		const double DurationSeconds = (Time - MemUsageLimitLastTimestamp) * FPlatformTime::GetSecondsPerCycle64();
+		const double DurationSeconds = static_cast<double>(Time - MemUsageLimitLastTimestamp) * FPlatformTime::GetSecondsPerCycle64();
 		if (DurationSeconds > 1.0) // only check once per second
 		{
 			MemUsageLimitLastTimestamp = Time;
diff --git a/Engine/Source/Developer/TraceInsights/Private/Insights/MemoryProfiler/ViewModels/MemAllocFilterValueConverter.cpp b/Engine/Source/Developer/TraceInsights/Private/Insights/MemoryProfiler/ViewModels/MemAllocFilterValueConverter.cpp
--- a/Engine/Source/Developer/TraceInsights/Private/Insights/MemoryProfiler/ViewModels/MemAllocFilterValueConverter.cpp
+++ b/Engine/Source/Developer/TraceInsights/Private/Insights/MemoryProfiler/ViewModels/MemAllocFilterValueConverter.cpp
@@ -55,7 +55,7 @@ bool FMemoryFilterValueConverter::Convert(const FString& Input, int64& Output, F
 
 			const int64 Unit = 1LL << ((PrefixPos - Prefixes + 1) * 10);
 			double Value = FCString::Atod(*TrimmedInput.Left(Index));
-			Output = static_cast<int64>(Value * Unit);
+			Output = static_cast<int64>(Value * static_cast<double>(Unit));
 			break;
 		}
 		else
diff --git a/Engine/Source/Developer/TraceInsights/Private/Insights/NetworkingProfiler/ViewModels/PacketContentViewDrawHelper.cpp b/Engine/Source/Developer/TraceInsights/Private/Insights/NetworkingProfiler/ViewModels/PacketContentViewDrawHelper.cpp
--- a/Engine/Source/Developer/TraceInsights/Private/Insights/NetworkingProfiler/ViewModels/PacketContentViewDrawHelper.cpp
+++ b/Engine/Source/Developer/TraceInsights/Private/Insights/NetworkingProfiler/ViewModels/PacketContentViewDrawHelper.cpp
@@ -186,7 +186,8 @@ void FPacketContentViewDrawStateBuilder::AddEvent(const TraceServices::FNetProfi
 		TStringBuilder<512> Builder;
 
 		Builder.Appendf(TEXT("%s"), EventName ? EventName : TEXT("?"));
-		if (EventW > Builder.Len() * 2.0f + 48.0f)
+		const float MinWidth = static_cast<float>(Builder.Len()) * 2.0f + 48.0f;
+		if (EventW > MinWidth)
 		{
 			const TraceServices::FNetProfilerBunchInfo& Info = Event.BunchInfo;
 			if (Info.bIsValid)
@@ -380,7 +381,7 @@ void FPacketContentViewDrawHelper::Draw(const FPacketContentViewDrawState& DrawS
 		const float EventFillH = Viewport.GetEventHeight();
 		for (const FPacketContentViewDrawState::FBoxPrimitive& Box : DrawState.Boxes)
 		{
-			DrawContext.DrawBox(Box.X, BoxY0 + BoxDY * Box.Depth, Box.W, EventFillH, WhiteBrush, Box.Color.CopyWithNewOpacity(Opacity));
+			DrawContext.DrawBox(Box.X, BoxY0 + BoxDY * (float)Box.Depth, Box.W, EventFillH, WhiteBrush, Box.Color.CopyWithNewOpacity(Opacity));
 		}
 		DrawContext.LayerId++;
 	}
@@ -393,7 +394,7 @@ void FPacketContentViewDrawHelper::Draw(const FPacketContentViewDrawState& DrawS
 		const float EventFillH = Viewport.GetEventHeight() - 2.0f;
 		for (const FPacketContentViewDrawState::FBoxPrimitive& Box : DrawState.InsideBoxes)
 		{
-			DrawContext.DrawBox(Box.X, BoxY0 + BoxDY * Box.Depth, Box.W, EventFillH, WhiteBrush, Box.Color.CopyWithNewOpacity(Opacity));
+			DrawContext.DrawBox(Box.X, BoxY0 + BoxDY * (float)Box.Depth, Box.W, EventFillH, WhiteBrush, Box.Color.CopyWithNewOpacity(Opacity));
 		}
 		DrawContext.LayerId++;
 	}
@@ -407,7 +408,7 @@ void FPacketContentViewDrawHelper::Draw(const FPacketContentViewDrawState& DrawS
 		const float EventBorderH = Viewport.GetEventHeight();
 		for (const FPacketContentViewDrawState::FBoxPrimitive& Box : DrawState.Borders)
 		{
-			DrawContext.DrawBox(Box.X, BoxY0 + BoxDY * Box.Depth, Box.W, EventBorderH, EventBorderBrush, Box.Color.CopyWithNewOpacity(Opacity));
+			DrawContext.DrawBox(Box.X, BoxY0 + BoxDY * (float)Box.Depth, Box.W, EventBorderH, EventBorderBrush, Box.Color.CopyWithNewOpacity(Opacity));
 		}
 		DrawContext.LayerId++;
 #else
@@ -419,7 +420,7 @@ void FPacketContentViewDrawHelper::Draw(const FPacketContentViewDrawState& DrawS
 			const float EventBorderH = Viewport.GetEventHeight();
 			for (const FPacketContentViewDrawState::FBoxPrimitive& Box : DrawState.Borders)
 			{
-				DrawContext.DrawBox(EventBorderLayerId, Box.X, BoxY0 + BoxDY * Box.Depth, Box.W, EventBorderH, WhiteBrush, Box.Color);
+				DrawContext.DrawBox(EventBorderLayerId, Box.X, BoxY0 + BoxDY * (float)Box.Depth, Box.W, EventBorderH, WhiteBrush, Box.Color);
 			}
 		}
 		else
@@ -430,7 +431,7 @@ void FPacketContentViewDrawHelper::Draw(const FPacketContentViewDrawState& DrawS
 			for (const FPacketContentViewDrawState::FBoxPrimitive& Box : DrawState.Borders)
 			{
 				FSlateRoundedBoxBrush Brush(FLinearColor(0.0f, 0.0f, 0.0f, 0.0f), 0.0f, Box.Color.CopyWithNewOpacity(Opacity), 1.0f);
-				DrawContext.DrawBox(Box.X, BoxY0 + BoxDY * Box.Depth, Box.W, EventBorderH, &Brush, FLinearColor(0.0f, 0.0f, 0.0f, 0.0f));
+				DrawContext.DrawBox(Box.X, BoxY0 + BoxDY * (float)Box.Depth, Box.W, EventBorderH, &Brush, FLinearColor(0.0f, 0.0f, 0.0f, 0.0f));
 			}
 			DrawContext.LayerId++;
 		}
@@ -446,7 +447,7 @@ void FPacketContentViewDrawHelper::Draw(const FPacketContentViewDrawState& DrawS
 		const FLinearColor BlackColor(0.0f, 0.0f, 0.0f, Opacity);
 		for (const FPacketContentViewDrawState::FTextPrimitive& Text : DrawState.Texts)
 		{
-			DrawContext.DrawText(Text.X, TextY0 + TextDY * Text.Depth, Text.Text, EventFont, Text.bWhite ? WhiteColor : BlackColor);
+			DrawContext.DrawText(Text.X, TextY0 + TextDY * (float)Text.Depth, Text.Text, EventFont, Text.bWhite ? WhiteColor : BlackColor);
 		}
 		DrawContext.LayerId++;
 	}
@@ -491,7 +492,7 @@ void FPacketContentViewDrawHelper::DrawEventHighlight(const FNetworkPacketEvent&
 	}
 
 	const float EventW = EventX2 - EventX1;
-	const float EventY = Viewport.GetTopEventPosY() + (Viewport.GetEventHeight() + Viewport.GetEventDY()) * Event.Level;
+	const float EventY = Viewport.GetTopEventPosY() + (Viewport.GetEventHeight() + Viewport.GetEventDY()) * static_cast<float>(Event.Level);
 
 #if INSIGHTS_USE_LEGACY_BORDER
 	constexpr float BorderOffset = 1.0f;
diff --git a/Engine/Source/Developer/TraceInsights/Private/Insights/NetworkingProfiler/ViewModels/PacketViewDrawHelper.cpp b/Engine/Source/Developer/TraceInsights/Private/Insights/NetworkingProfiler/ViewModels/PacketViewDrawHelper.cpp
--- a/Engine/Source/Developer/TraceInsights/Private/Insights/NetworkingProfiler/ViewModels/PacketViewDrawHelper.cpp
+++ b/Engine/Source/Developer/TraceInsights/Private/Insights/NetworkingProfiler/ViewModels/PacketViewDrawHelper.cpp
@@ -191,7 +191,7 @@ void FPacketViewDrawHelper::DrawCached(const FNetworkPacketSeries& Series) const
 
 		NumDrawSamples++;
 
-		const float X = SampleIndex * SampleW;
+		const float X = static_cast<float>(SampleIndex) * SampleW;
 
 		//const float ValueY = FMath::RoundToFloat(ViewportY.GetOffsetForValue(static_cast<double>(Sample.LargestPacket.ContentSizeInBits)));
 		const float ValueY = FMath::RoundToFloat(ViewportY.GetOffsetForValue(static_cast<double>(Sample.LargestPacket.TotalSizeInBytes * 8)));
@@ -248,7 +248,7 @@ void FPacketViewDrawHelper::DrawSampleHighlight(const FNetworkPacketAggregatedSa
 	const int32 PacketsPerSample = Viewport.GetNumPacketsPerSample();
 	const int32 FirstPacketIndex = Viewport.GetFirstPacketIndex();
 	const int32 SampleIndex = (Sample.LargestPacket.Index - FirstPacketIndex) / PacketsPerSample;
-	const float X = SampleIndex * SampleW;
+	const float X = static_cast<float>(SampleIndex) * SampleW;
 
 	const FAxisViewportDouble& ViewportY = Viewport.GetVerticalAxisViewport();
 
@@ -309,10 +309,10 @@ void FPacketViewDrawHelper::DrawSelection(int32 StartPacketIndex, int32 EndPacke
 	const int32 FirstPacketIndex = Viewport.GetFirstPacketIndex();
 
 	const int32 StartSampleIndex = (StartPacketIndex - FirstPacketIndex) / PacketsPerSample;
-	const float X1 = FMath::Max(-2.0f, StartSampleIndex * SampleW);
+	const float X1 = FMath::Max(-2.0f, static_cast<float>(StartSampleIndex) * SampleW);
 
 	const int32 EndSampleIndex = (EndPacketIndex - FirstPacketIndex) / PacketsPerSample;
-	const float X2 = FMath::Min(Viewport.GetWidth() + 2.0f, EndSampleIndex * SampleW);
+	const float X2 = FMath::Min(Viewport.GetWidth() + 2.0f, static_cast<float>(EndSampleIndex) * SampleW);
 
 	const FAxisViewportDouble& ViewportY = Viewport.GetVerticalAxisViewport();
 
diff --git a/Engine/Source/Developer/TraceInsights/Private/Insights/NetworkingProfiler/Widgets/SPacketContentView.cpp b/Engine/Source/Developer/TraceInsights/Private/Insights/NetworkingProfiler/Widgets/SPacketContentView.cpp
--- a/Engine/Source/Developer/TraceInsights/Private/Insights/NetworkingProfiler/Widgets/SPacketContentView.cpp
+++ b/Engine/Source/Developer/TraceInsights/Private/Insights/NetworkingProfiler/Widgets/SPacketContentView.cpp
@@ -1234,7 +1234,7 @@ FNetworkPacketEventRef SPacketContentView::GetEventAtMousePosition(float X, floa
 			const float EventX1 = ViewportX.GetRoundedOffsetForValue(static_cast<double>(Event.BitOffset));
 			const float EventX2 = ViewportX.GetRoundedOffsetForValue(static_cast<double>(Event.BitOffset + Event.BitSize));
 
-			const float EventY = Viewport.GetTopEventPosY() + (Viewport.GetEventHeight() + Viewport.GetEventDY()) * Event.Level;
+			const float EventY = Viewport.GetTopEventPosY() + (Viewport.GetEventHeight() + Viewport.GetEventDY()) * static_cast<float>(Event.Level);
 
 			constexpr float ToleranceX = 1.0f;
 
@@ -1328,10 +1328,10 @@ int32 SPacketContentView::OnPaint(const FPaintArgs& Args, const FGeometry& Allot
 
 		DrawContext.LayerId++;
 
-		DrawContext.DrawBox(DbgX - 2.0f, DbgY - 2.0f, DbgW, DbgH, WhiteBrush, FLinearColor(1.0, 1.0, 1.0, 0.9));
+		DrawContext.DrawBox(DbgX - 2.0f, DbgY - 2.0f, DbgW, DbgH, WhiteBrush, FLinearColor(1.0f, 1.0f, 1.0f, 0.9f));
 		DrawContext.LayerId++;
 
-		FLinearColor DbgTextColor(0.0, 0.0, 0.0, 0.9);
+		FLinearColor DbgTextColor(0.0f, 0.0f, 0.0f, 0.9f);
 
 		// Time interval since last OnPaint call.
 		const uint64 CurrentTime = FPlatformTime::Cycles64();
diff --git a/Engine/Source/Developer/TraceInsights/Private/Insights/NetworkingProfiler/Widgets/SPacketView.cpp b/Engine/Source/Developer/TraceInsights/Private/Insights/NetworkingProfiler/Widgets/SPacketView.cpp
--- a/Engine/Source/Developer/TraceInsights/Private/Insights/NetworkingProfiler/Widgets/SPacketView.cpp
+++ b/Engine/Source/Developer/TraceInsights/Private/Insights/NetworkingProfiler/Widgets/SPacketView.cpp
@@ -891,10 +891,10 @@ int32 SPacketView::OnPaint(const FPaintArgs& Args, const FGeometry& AllottedGeom
 			}
 
 			const float Y = 10.0f;
-			const float H = 2.0f + 13.0f * NumLines;
-			DrawContext.DrawBox(CX - W2, Y, 2 * W2, H, WhiteBrush, FLinearColor(0.7, 0.7, 0.7, TooltipOpacity));
+			const float H = 2.0f + 13.0f * static_cast<float>(NumLines);
+			DrawContext.DrawBox(CX - W2, Y, 2 * W2, H, WhiteBrush, FLinearColor(0.7f, 0.7f, 0.7f, TooltipOpacity));
 			DrawContext.LayerId++;
-			DrawContext.DrawText(CX - W2 + DX, Y + 1.0f, Text, SummaryFont, FLinearColor(0.0, 0.0, 0.0, TooltipOpacity));
+			DrawContext.DrawText(CX - W2 + DX, Y + 1.0f, Text, SummaryFont, FLinearColor(0.0f, 0.0f, 0.0f, TooltipOpacity));
 			DrawContext.LayerId++;
 		}
 
@@ -916,10 +916,10 @@ int32 SPacketView::OnPaint(const FPaintArgs& Args, const FGeometry& AllottedGeom
 
 		DrawContext.LayerId++;
 
-		DrawContext.DrawBox(DbgX - 2.0f, DbgY - 2.0f, DbgW, DbgH, WhiteBrush, FLinearColor(1.0, 1.0, 1.0, 0.9));
+		DrawContext.DrawBox(DbgX - 2.0f, DbgY - 2.0f, DbgW, DbgH, WhiteBrush, FLinearColor(1.0f, 1.0f, 1.0f, 0.9f));
 		DrawContext.LayerId++;
 
-		FLinearColor DbgTextColor(0.0, 0.0, 0.0, 0.9);
+		FLinearColor DbgTextColor(0.0f, 0.0f, 0.0f, 0.9f);
 
 		// Time interval since last OnPaint call.
 		const uint64 CurrentTime = FPlatformTime::Cycles64();
diff --git a/Engine/Source/Developer/TraceInsights/Private/Insights/Table/ViewModels/UntypedTable.cpp b/Engine/Source/Developer/TraceInsights/Private/Insights/Table/ViewModels/UntypedTable.cpp
--- a/Engine/Source/Developer/TraceInsights/Private/Insights/Table/ViewModels/UntypedTable.cpp
+++ b/Engine/Source/Developer/TraceInsights/Private/Insights/Table/ViewModels/UntypedTable.cpp
@@ -288,7 +288,7 @@ void FUntypedTable::CreateColumns(const TraceServices::ITableLayout& TableLayout
 		case TraceServices::TableColumnType_CString:
 			Column.SetDataType(ETableCellDataType::CString);
 			HorizontalAlignment = HAlign_Left;
-			InitialColumnWidth = FMath::Max(120.0f, 6.0f * ColumnNameStr.Len());
+			InitialColumnWidth = FMath::Max(120.0f, 6.0f * static_cast<float>(ColumnNameStr.Len()));
 			FormatterPtr = MakeShared<FCStringValueFormatterAsText>();
 			SorterPtr = MakeShared<FSorterByCStringValue>(ColumnRef);
 			InitialSortMode = EColumnSortMode::Ascending;
diff --git a/Engine/Source/Developer/TraceInsights/Private/Insights/Table/Widgets/STableTreeView.cpp b/Engine/Source/Developer/TraceInsights/Private/Insights/Table/Widgets/STableTreeView.cpp
--- a/Engine/Source/Developer/TraceInsights/Private/Insights/Table/Widgets/STableTreeView.cpp
+++ b/Engine/Source/Developer/TraceInsights/Private/Insights/Table/Widgets/STableTreeView.cpp
@@ -1982,7 +1982,7 @@ void STableTreeView::PostChangeGroupings()
 	float HierarchyColumnWidth = DefaultHierarchyColumnWidth;
 	FString GroupingStr;
 
-	int32 GroupingDepth = 0;
+	float GroupingDepth = 0;
 	for (TSharedPtr<FTreeNodeGrouping>& GroupingPtr : CurrentGroupings)
 	{
 		const FName& ColumnId = GroupingPtr->GetColumnId();
@@ -1998,7 +1998,7 @@ void STableTreeView::PostChangeGroupings()
 					const SHeaderRow::FColumn& CurrentColumn = TreeViewHeaderRow->GetColumns()[ColumnIndex];
 					if (CurrentColumn.ColumnId == ColumnId)
 					{
-						const float Width = HierarchyMinWidth + GroupingDepth * HierarchyIndentation + CurrentColumn.GetWidth();
+						const float Width = HierarchyMinWidth + GroupingDepth + CurrentColumn.GetWidth();
 						if (Width > HierarchyColumnWidth)
 						{
 							HierarchyColumnWidth = Width;
@@ -2022,7 +2022,7 @@ void STableTreeView::PostChangeGroupings()
 			GroupingStr.Append(GroupingPtr->GetShortName().ToString());
 		}
 
-		++GroupingDepth;
+		GroupingDepth += HierarchyIndentation;
 	}
 
 	//////////////////////////////////////////////////
diff --git a/Engine/Source/Developer/TraceInsights/Private/Insights/Tests/FunctionalTests/ExportCommandsTests.cpp b/Engine/Source/Developer/TraceInsights/Private/Insights/Tests/FunctionalTests/ExportCommandsTests.cpp
--- a/Engine/Source/Developer/TraceInsights/Private/Insights/Tests/FunctionalTests/ExportCommandsTests.cpp
+++ b/Engine/Source/Developer/TraceInsights/Private/Insights/Tests/FunctionalTests/ExportCommandsTests.cpp
@@ -2,19 +2,19 @@
 
 #include "HAL/FileManager.h"
 #include "HAL/PlatformFileManager.h"
-#include "Misc/FileHelper.h"
 #include "Insights/InsightsManager.h"
 #include "Insights/IUnrealInsightsModule.h"
 #include "Insights/Tests/InsightsTestUtils.h"
 #include "Misc/AutomationTest.h"
+#include "Misc/FileHelper.h"
 
 void VerifyExportedLines(const FString& ExportReportPath, const FString& CmdLogPath, const FString& Elements, FInsightsTestUtils Utils, FAutomationTestBase* Test, double Timeout)
 {
-	double StartTime = FPlatformTime::Seconds();
 	bool bLineFound = false;
 	FString ExpectedResult;
 
-	while ((FPlatformTime::Seconds() - StartTime) < Timeout)
+	double StartTime = FPlatformTime::Seconds();
+	while (FPlatformTime::Seconds() - StartTime < Timeout)
 	{
 		FString FileContent;
 
@@ -90,8 +90,8 @@ bool FСommandsExportWindowsTest::RunTest(const FString& Parameters)
 	{
 		AddWarning(TEXT("The TestResults directory already exists. Deleting to avoid undefined behavior"));
 		IFileManager::Get().DeleteDirectory(*TestResultsDirPath, false, true);
-		float StartTime = FPlatformTime::Seconds();
-		while ((FPlatformTime::Seconds() - StartTime) < Timeout)
+		double StartTime = FPlatformTime::Seconds();
+		while (FPlatformTime::Seconds() - StartTime < Timeout)
 		{
 			if (!PlatformFile.DirectoryExists(*TestResultsDirPath))
 			{
@@ -115,7 +115,7 @@ bool FСommandsExportWindowsTest::RunTest(const FString& Parameters)
 
 	VerifyExportedLines(ExportTimersReportPath, CmdTimersLogPath, TEXT("timers"), Utils, this, Timeout);
 
-	// UI verification cannot be executed in that case. 
+	// UI verification cannot be executed in that case.
 
 	// ExportTimingEvents
 	InsightsParameters = FString::Printf(TEXT("-OpenTraceFile=\"%s\" -ABSLOG=\"%s\" -AutoQuit -NoUI  -ExecOnAnalysisCompleteCmd=\"%s\" -log"), *StoreTracePath, *CmdTimingEventsLogPath, *ExportTimingEventsTask);
@@ -127,7 +127,16 @@ bool FСommandsExportWindowsTest::RunTest(const FString& Parameters)
 	InsightsParameters = FString::Printf(TEXT("-OpenTraceFile=\"%s\" -ABSLOG=\"%s\" -AutoQuit -NoUI  -ExecOnAnalysisCompleteCmd=\"%s\" -log"), *StoreTracePath, *CmdTimingEventsNonDefaultLogPath, *ExportTimingEventsBorderedTask);
 	InsightsManager->OpenUnrealInsights(*InsightsParameters);
 
-	const TArray<FString> ExpectedTimingEventsBorderedElements = { TEXT("ThreadId"), TEXT("ThreadName"), TEXT("TimerId"), TEXT("TimerName"), TEXT("StartTime"), TEXT("EndTime"), TEXT("Duration"), TEXT("Depth"),
+	const TArray<FString> ExpectedTimingEventsBorderedElements =
+	{
+		TEXT("ThreadId"),
+		TEXT("ThreadName"),
+		TEXT("TimerId"),
+		TEXT("TimerName"),
+		TEXT("StartTime"),
+		TEXT("EndTime"),
+		TEXT("Duration"),
+		TEXT("Depth"),
 		TEXT("2,GameThread,1,FEngineLoop::PreInitPreStartupScreen")
 	};
 
@@ -158,7 +167,9 @@ bool FСommandsExportWindowsTest::RunTest(const FString& Parameters)
 	InsightsParameters = FString::Printf(TEXT("-OpenTraceFile=\"%s\" -ABSLOG=\"%s\" -AutoQuit -NoUI  -ExecOnAnalysisCompleteCmd=\"@=/TestResults/export.rsp\" -log"), *StoreTracePath, *CmdExportLogPath);
 	InsightsManager->OpenUnrealInsights(*InsightsParameters);
 
-	const TArray<FString> ExpectedThreadsElementsRsp = { TEXT("/TestResults/RSPtest/CSV/Threads_rsp.csv"),
+	const TArray<FString> ExpectedThreadsElementsRsp =
+	{
+		TEXT("/TestResults/RSPtest/CSV/Threads_rsp.csv"),
 		TEXT("/TestResults/RSPtest/TSV/Threads_rsp.tsv"),
 		TEXT("/TestResults/RSPtest/TXT/Threads_rsp.txt"),
 	};
@@ -171,7 +182,9 @@ bool FСommandsExportWindowsTest::RunTest(const FString& Parameters)
 		VerifyExportedLines(ExpectedThreadsElementsRsp[i], CmdExportLogPath, TEXT("threads"), Utils, this, Timeout);
 	}
 
-	const TArray<FString> ExpectedTimersElementsRsp = { TEXT("/TestResults/RSPtest/CSV/Timers_rsp.csv"),
+	const TArray<FString> ExpectedTimersElementsRsp =
+	{
+		TEXT("/TestResults/RSPtest/CSV/Timers_rsp.csv"),
 		TEXT("/TestResults/RSPtest/TSV/Timers_rsp.tsv"),
 		TEXT("/TestResults/RSPtest/TXT/Timers_rsp.txt"),
 	};
@@ -184,7 +197,9 @@ bool FСommandsExportWindowsTest::RunTest(const FString& Parameters)
 		VerifyExportedLines(ExpectedTimersElementsRsp[i], CmdExportLogPath, TEXT("timers"), Utils, this, Timeout);
 	}
 
-	const TArray<FString> ExpectedTimingEventsRsp = { TEXT("/TestResults/RSPtest/CSV/TimingEvents_rsp.csv"),
+	const TArray<FString> ExpectedTimingEventsRsp =
+	{
+		TEXT("/TestResults/RSPtest/CSV/TimingEvents_rsp.csv"),
 		TEXT("/TestResults/RSPtest/TSV/TimingEvents_rsp.tsv"),
 		TEXT("/TestResults/RSPtest/TXT/TimingEvents_rsp.txt"),
 	};
@@ -202,4 +217,4 @@ bool FСommandsExportWindowsTest::RunTest(const FString& Parameters)
 	return true;
 }
 
-#endif //!WITH_EDITOR
\ No newline at end of file
+#endif //!WITH_EDITOR
diff --git a/Engine/Source/Developer/TraceInsights/Private/Insights/Tests/InsightsTestUtils.cpp b/Engine/Source/Developer/TraceInsights/Private/Insights/Tests/InsightsTestUtils.cpp
--- a/Engine/Source/Developer/TraceInsights/Private/Insights/Tests/InsightsTestUtils.cpp
+++ b/Engine/Source/Developer/TraceInsights/Private/Insights/Tests/InsightsTestUtils.cpp
@@ -77,10 +77,10 @@ bool FInsightsTestUtils::AnalyzeTrace(const TCHAR* Path) const
 
 ////////////////////////////////////////////////////////////////////////////////////////////////////
 
-bool FInsightsTestUtils::FileContainsString(const FString& PathToFile, const FString& ExpectedString, const float& Timeout) const
+bool FInsightsTestUtils::FileContainsString(const FString& PathToFile, const FString& ExpectedString, double Timeout) const
 {
-	float StartTime = FPlatformTime::Seconds();
-	while ((FPlatformTime::Seconds() - StartTime) < Timeout)
+	double StartTime = FPlatformTime::Seconds();
+	while (FPlatformTime::Seconds() - StartTime < Timeout)
 	{
 		if (!FPaths::FileExists(PathToFile))
 		{
@@ -94,7 +94,7 @@ bool FInsightsTestUtils::FileContainsString(const FString& PathToFile, const FSt
 			if (FileHandle)
 			{
 				TArray<uint8> FileData;
-				FileData.SetNumUninitialized(FileHandle->Size());
+				FileData.SetNumUninitialized(static_cast<int32>(FileHandle->Size()));
 				FileHandle->Read(FileData.GetData(), FileData.Num());
 				FFileHelper::BufferToString(LogFileContents, FileData.GetData(), FileData.Num());
 
@@ -112,7 +112,7 @@ bool FInsightsTestUtils::FileContainsString(const FString& PathToFile, const FSt
 
 ////////////////////////////////////////////////////////////////////////////////////////////////////
 
-bool FInsightsTestUtils::StartTracing(FTraceAuxiliary::EConnectionType ConnectionType, const float& Timeout) const
+bool FInsightsTestUtils::StartTracing(FTraceAuxiliary::EConnectionType ConnectionType, double Timeout) const
 {
 	bool bStarted = false;
 	if (ConnectionType == FTraceAuxiliary::EConnectionType::Network)
@@ -139,7 +139,7 @@ bool FInsightsTestUtils::StartTracing(FTraceAuxiliary::EConnectionType Connectio
 
 ////////////////////////////////////////////////////////////////////////////////////////////////////
 
-bool FInsightsTestUtils::SetupUTS(const float& Timeout) const
+bool FInsightsTestUtils::SetupUTS(double Timeout) const
 {
 	const FString UnrealTraceServerName = TEXT("UnrealTraceServer");
 
@@ -175,7 +175,7 @@ bool FInsightsTestUtils::SetupUTS(const float& Timeout) const
 
 ////////////////////////////////////////////////////////////////////////////////////////////////////
 
-bool FInsightsTestUtils::KillUTS(const float& Timeout) const
+bool FInsightsTestUtils::KillUTS(double Timeout) const
 {
 	const FString UnrealTraceServerName = TEXT("UnrealTraceServer");
 
diff --git a/Engine/Source/Developer/TraceInsights/Private/Insights/ViewModels/FrameTrackHelper.cpp b/Engine/Source/Developer/TraceInsights/Private/Insights/ViewModels/FrameTrackHelper.cpp
--- a/Engine/Source/Developer/TraceInsights/Private/Insights/ViewModels/FrameTrackHelper.cpp
+++ b/Engine/Source/Developer/TraceInsights/Private/Insights/ViewModels/FrameTrackHelper.cpp
@@ -169,7 +169,7 @@ void FFrameTrackDrawHelper::DrawCached(const FFrameTrackSeries& Series) const
 
 		NumDrawSamples++;
 
-		const float X = SampleIndex * SampleW;
+		const float X = static_cast<float>(SampleIndex) * SampleW;
 		float ValueY;
 
 		FLinearColor ColorFill = SeriesColor;
@@ -227,7 +227,7 @@ void FFrameTrackDrawHelper::DrawHoveredSample(const FFrameTrackSample& Sample) c
 	const int32 FramesPerSample = Viewport.GetNumFramesPerSample();
 	const int32 FirstFrameIndex = Viewport.GetFirstFrameIndex();
 	const int32 SampleIndex = (Sample.LargestFrameIndex - FirstFrameIndex) / FramesPerSample;
-	const float X = SampleIndex * SampleW;
+	const float X = static_cast<float>(SampleIndex) * SampleW;
 
 	const FAxisViewportDouble& ViewportY = Viewport.GetVerticalAxisViewport();
 
@@ -272,8 +272,8 @@ void FFrameTrackDrawHelper::DrawHighlightedInterval(const FFrameTrackSeries& Ser
 	if (Index1 <= Index2)
 	{
 		const float SampleW = Viewport.GetSampleWidth();
-		float X1 = Index1 * SampleW;
-		float X2 = (Index2 + 1) * SampleW;
+		float X1 = static_cast<float>(Index1) * SampleW;
+		float X2 = static_cast<float>(Index2 + 1) * SampleW;
 
 		constexpr float Y1 = 0.0f; // allows 12px for the horizontal scrollbar (one displayed on top of the track)
 		const float Y2 = Viewport.GetHeight();
diff --git a/Engine/Source/Developer/TraceInsights/Private/Insights/ViewModels/ThreadTimingTrack.cpp b/Engine/Source/Developer/TraceInsights/Private/Insights/ViewModels/ThreadTimingTrack.cpp
--- a/Engine/Source/Developer/TraceInsights/Private/Insights/ViewModels/ThreadTimingTrack.cpp
+++ b/Engine/Source/Developer/TraceInsights/Private/Insights/ViewModels/ThreadTimingTrack.cpp
@@ -322,7 +322,8 @@ static void AddTimingEventToBuilder(ITimingEventsTrackDrawStateBuilder& Builder,
 		{
 			FString EventName = Timer->Name;
 
-			if (Width > EventName.Len() * 4.0f + 32.0f)
+			const float MinWidth = static_cast<float>(EventName.Len()) * 4.0f + 32.0f;
+			if (Width > MinWidth)
 			{
 				//EventName = TEXT("*") + EventName; // for debugging
 
diff --git a/Engine/Source/Developer/TraceInsights/Private/Insights/ViewModels/TimingGraphTrack.cpp b/Engine/Source/Developer/TraceInsights/Private/Insights/ViewModels/TimingGraphTrack.cpp
--- a/Engine/Source/Developer/TraceInsights/Private/Insights/ViewModels/TimingGraphTrack.cpp
+++ b/Engine/Source/Developer/TraceInsights/Private/Insights/ViewModels/TimingGraphTrack.cpp
@@ -815,7 +815,7 @@ void FTimingGraphTrack::DrawVerticalAxisGrid(const ITimingTrackDrawContext& Cont
 		if (Index < NumThresholds - 1)
 		{
 			const double NextTickUnit = Thresholds[Index + 1];
-			if (NextTickUnit <= (DeltaTicks + 1) * TickUnit)
+			if (NextTickUnit <= static_cast<double>(DeltaTicks + 1) * TickUnit)
 			{
 				TickUnit = NextTickUnit;
 				DeltaTicks = 1;
@@ -825,7 +825,7 @@ void FTimingGraphTrack::DrawVerticalAxisGrid(const ITimingTrackDrawContext& Cont
 				DeltaTicks++;
 			}
 		}
-		const double Grid = DeltaTicks * TickUnit;
+		const double Grid = static_cast<double>(DeltaTicks) * TickUnit;
 		const double StartValue = FMath::GridSnap(BottomValue, Grid);
 
 		const FLinearColor GridColor(0.0f, 0.0f, 0.0f, 0.1f);
diff --git a/Engine/Source/Developer/TraceInsights/Private/Insights/ViewModels/TimingViewDrawHelper.cpp b/Engine/Source/Developer/TraceInsights/Private/Insights/ViewModels/TimingViewDrawHelper.cpp
--- a/Engine/Source/Developer/TraceInsights/Private/Insights/ViewModels/TimingViewDrawHelper.cpp
+++ b/Engine/Source/Developer/TraceInsights/Private/Insights/ViewModels/TimingViewDrawHelper.cpp
@@ -60,7 +60,8 @@ void FTimingEventsTrackDrawStateBuilder::AddEvent(double EventStartTime, double
 			//FString Name = FString::Printf(TEXT("%s [%llu]"), EventName, EventType);
 			FString Name = EventName;
 
-			if (Width > Name.Len() * 4.0f + 32.0f)
+			const float MinWidth = static_cast<float>(Name.Len()) * 4.0f + 32.0f;
+			if (Width > MinWidth)
 			{
 				AppendDurationToEventName(Name, EventEndTime - EventStartTime);
 			}
@@ -331,7 +332,7 @@ void FTimingViewDrawHelper::DrawEvents(const FTimingEventsTrackDrawState& DrawSt
 			const float EventFillH = Layout.EventH;
 			for (const FTimingEventsTrackDrawState::FBoxPrimitive& Box : DrawState.Boxes)
 			{
-				const float Y = TopLaneY + (Layout.EventH + Layout.EventDY) * Box.Depth;
+				const float Y = TopLaneY + (Layout.EventH + Layout.EventDY) * static_cast<float>(Box.Depth);
 				DrawContext.DrawBox(EventFillLayerId, Box.X, Y, Box.W, EventFillH, WhiteBrush, Box.Color);
 			}
 			NumDrawBoxes += DrawState.Boxes.Num();
@@ -344,7 +345,7 @@ void FTimingViewDrawHelper::DrawEvents(const FTimingEventsTrackDrawState& DrawSt
 			const float EventFillH = Layout.EventH - 2.0f;
 			for (const FTimingEventsTrackDrawState::FBoxPrimitive& Box : DrawState.InsideBoxes)
 			{
-				const float Y = TopLaneY + (Layout.EventH + Layout.EventDY) * Box.Depth + 1.0f;
+				const float Y = TopLaneY + (Layout.EventH + Layout.EventDY) * static_cast<float>(Box.Depth) + 1.0f;
 				DrawContext.DrawBox(EventFillLayerId, Box.X, Y, Box.W, EventFillH, WhiteBrush, Box.Color);
 			}
 			NumDrawBoxes += DrawState.InsideBoxes.Num();
@@ -357,7 +358,7 @@ void FTimingViewDrawHelper::DrawEvents(const FTimingEventsTrackDrawState& DrawSt
 			const float EventBorderH = Layout.EventH;
 			for (const FTimingEventsTrackDrawState::FBoxPrimitive& Box : DrawState.Borders)
 			{
-				const float Y = TopLaneY + (Layout.EventH + Layout.EventDY) * Box.Depth;
+				const float Y = TopLaneY + (Layout.EventH + Layout.EventDY) * static_cast<float>(Box.Depth);
 #if INSIGHTS_USE_LEGACY_BORDER
 				DrawContext.DrawBox(EventBorderLayerId, Box.X, Y, Box.W, EventBorderH, EventBorderBrush, Box.Color);
 #else
@@ -385,7 +386,7 @@ void FTimingViewDrawHelper::DrawEvents(const FTimingEventsTrackDrawState& DrawSt
 			const int32 EventTextLayerId = ReservedLayerId + ToInt32(EDrawLayer::EventText);
 			for (const FTimingEventsTrackDrawState::FTextPrimitive& Text : DrawState.Texts)
 			{
-				const float Y = TextY0 + TextDY * Text.Depth;
+				const float Y = TextY0 + TextDY * static_cast<float>(Text.Depth);
 				DrawContext.DrawText(EventTextLayerId, Text.X, Y, Text.Text, EventFont, Text.bWhite ? WhiteColor : BlackColor);
 			}
 			NumDrawTexts += DrawState.Texts.Num();
@@ -422,7 +423,7 @@ void FTimingViewDrawHelper::DrawFadedEvents(const FTimingEventsTrackDrawState& D
 			const float EventFillH = Layout.EventH;
 			for (const FTimingEventsTrackDrawState::FBoxPrimitive& Box : DrawState.Boxes)
 			{
-				const float Y = TopLaneY + (Layout.EventH + Layout.EventDY) * Box.Depth;
+				const float Y = TopLaneY + (Layout.EventH + Layout.EventDY) * static_cast<float>(Box.Depth);
 				DrawContext.DrawBox(EventFillLayerId, Box.X, Y, Box.W, EventFillH, WhiteBrush, Box.Color.CopyWithNewOpacity(Opacity));
 			}
 			NumDrawBoxes += DrawState.Boxes.Num();
@@ -435,7 +436,7 @@ void FTimingViewDrawHelper::DrawFadedEvents(const FTimingEventsTrackDrawState& D
 			const float EventFillH = Layout.EventH - 2.0f;
 			for (const FTimingEventsTrackDrawState::FBoxPrimitive& Box : DrawState.InsideBoxes)
 			{
-				const float Y = TopLaneY + (Layout.EventH + Layout.EventDY) * Box.Depth + 1.0f;
+				const float Y = TopLaneY + (Layout.EventH + Layout.EventDY) * static_cast<float>(Box.Depth) + 1.0f;
 				DrawContext.DrawBox(EventFillLayerId, Box.X, Y, Box.W, EventFillH, WhiteBrush, Box.Color.CopyWithNewOpacity(Opacity));
 			}
 			NumDrawBoxes += DrawState.InsideBoxes.Num();
@@ -448,7 +449,7 @@ void FTimingViewDrawHelper::DrawFadedEvents(const FTimingEventsTrackDrawState& D
 			const float EventBorderH = Layout.EventH;
 			for (const FTimingEventsTrackDrawState::FBoxPrimitive& Box : DrawState.Borders)
 			{
-				const float Y = TopLaneY + (Layout.EventH + Layout.EventDY) * Box.Depth;
+				const float Y = TopLaneY + (Layout.EventH + Layout.EventDY) * static_cast<float>(Box.Depth);
 #if INSIGHTS_USE_LEGACY_BORDER
 				DrawContext.DrawBox(EventBorderLayerId, Box.X, Y, Box.W, EventBorderH, EventBorderBrush, Box.Color.CopyWithNewOpacity(Opacity));
 #else
@@ -477,7 +478,7 @@ void FTimingViewDrawHelper::DrawFadedEvents(const FTimingEventsTrackDrawState& D
 			const int32 EventTextLayerId = ReservedLayerId + ToInt32(EDrawLayer::EventText);
 			for (const FTimingEventsTrackDrawState::FTextPrimitive& Text : DrawState.Texts)
 			{
-				const float Y = TextY0 + (TextDY) * Text.Depth;
+				const float Y = TextY0 + (TextDY) * static_cast<float>(Text.Depth);
 				DrawContext.DrawText(EventTextLayerId, Text.X, Y, Text.Text, EventFont, Text.bWhite ? WhiteColor : BlackColor);
 			}
 			NumDrawTexts += DrawState.Texts.Num();
diff --git a/Engine/Source/Developer/TraceInsights/Private/Insights/Widgets/STimingView.cpp b/Engine/Source/Developer/TraceInsights/Private/Insights/Widgets/STimingView.cpp
--- a/Engine/Source/Developer/TraceInsights/Private/Insights/Widgets/STimingView.cpp
+++ b/Engine/Source/Developer/TraceInsights/Private/Insights/Widgets/STimingView.cpp
@@ -1498,7 +1498,7 @@ int32 STimingView::OnPaint(const FPaintArgs& Args, const FGeometry& AllottedGeom
 		for (int32 LineIndex = 0; LineIndex < OverscrollLineCount; ++LineIndex)
 		{
 			const float Opacity = OverscrollLeft * static_cast<float>(OverscrollLineCount - LineIndex) / static_cast<float>(OverscrollLineCount);
-			DrawContext.DrawBox(LineIndex * OverscrollLineSize, OverscrollLineY, OverscrollLineSize, OverscrollLineH, WhiteBrush, FLinearColor(1.0f, 0.1f, 0.1f, Opacity));
+			DrawContext.DrawBox(static_cast<float>(LineIndex) * OverscrollLineSize, OverscrollLineY, OverscrollLineSize, OverscrollLineH, WhiteBrush, FLinearColor(1.0f, 0.1f, 0.1f, Opacity));
 		}
 	}
 	if (OverscrollRight > 0.0f)
@@ -1508,7 +1508,7 @@ int32 STimingView::OnPaint(const FPaintArgs& Args, const FGeometry& AllottedGeom
 		for (int32 LineIndex = 0; LineIndex < OverscrollLineCount; ++LineIndex)
 		{
 			const float Opacity = OverscrollRight * static_cast<float>(OverscrollLineCount - LineIndex) / static_cast<float>(OverscrollLineCount);
-			DrawContext.DrawBox(ViewWidth - (1 + LineIndex) * OverscrollLineSize, OverscrollLineY, OverscrollLineSize, OverscrollLineH, WhiteBrush, FLinearColor(1.0f, 0.1f, 0.1f, Opacity));
+			DrawContext.DrawBox(ViewWidth - static_cast<float>(1 + LineIndex) * OverscrollLineSize, OverscrollLineY, OverscrollLineSize, OverscrollLineH, WhiteBrush, FLinearColor(1.0f, 0.1f, 0.1f, Opacity));
 		}
 	}
 	if (OverscrollTop > 0.0f)
@@ -1517,7 +1517,7 @@ int32 STimingView::OnPaint(const FPaintArgs& Args, const FGeometry& AllottedGeom
 		for (int32 LineIndex = 0; LineIndex < OverscrollLineCount; ++LineIndex)
 		{
 			const float Opacity = OverscrollTop * static_cast<float>(OverscrollLineCount - LineIndex) / static_cast<float>(OverscrollLineCount);
-			DrawContext.DrawBox(0.0f, OverscrollLineY + LineIndex * OverscrollLineSize, ViewWidth, OverscrollLineSize, WhiteBrush, FLinearColor(1.0f, 0.1f, 0.1f, Opacity));
+			DrawContext.DrawBox(0.0f, OverscrollLineY + static_cast<float>(LineIndex) * OverscrollLineSize, ViewWidth, OverscrollLineSize, WhiteBrush, FLinearColor(1.0f, 0.1f, 0.1f, Opacity));
 		}
 	}
 	if (OverscrollBottom > 0.0f)
@@ -1526,7 +1526,7 @@ int32 STimingView::OnPaint(const FPaintArgs& Args, const FGeometry& AllottedGeom
 		for (int32 LineIndex = 0; LineIndex < OverscrollLineCount; ++LineIndex)
 		{
 			const float Opacity = OverscrollBottom * static_cast<float>(OverscrollLineCount - LineIndex) / static_cast<float>(OverscrollLineCount);
-			DrawContext.DrawBox(0.0f, OverscrollLineY - (1 + LineIndex) * OverscrollLineSize, ViewWidth, OverscrollLineSize, WhiteBrush, FLinearColor(1.0f, 0.1f, 0.1f, Opacity));
+			DrawContext.DrawBox(0.0f, OverscrollLineY - static_cast<float>(1 + LineIndex) * OverscrollLineSize, ViewWidth, OverscrollLineSize, WhiteBrush, FLinearColor(1.0f, 0.1f, 0.1f, Opacity));
 		}
 	}
 
@@ -1541,7 +1541,7 @@ int32 STimingView::OnPaint(const FPaintArgs& Args, const FGeometry& AllottedGeom
 		const float DbgDY = MaxFontCharHeight;
 
 		const float DbgW = 320.0f;
-		const float DbgH = DbgDY * 9 + 3.0f;
+		const float DbgH = DbgDY * 9.0f + 3.0f;
 		const float DbgX = ViewWidth - DbgW - 20.0f;
 		float DbgY = Viewport.GetPosY() + Viewport.GetTopOffset() + 10.0f;
 
diff --git a/Engine/Source/Developer/TraceInsights/Public/Insights/Tests/InsightsTestUtils.h b/Engine/Source/Developer/TraceInsights/Public/Insights/Tests/InsightsTestUtils.h
--- a/Engine/Source/Developer/TraceInsights/Public/Insights/Tests/InsightsTestUtils.h
+++ b/Engine/Source/Developer/TraceInsights/Public/Insights/Tests/InsightsTestUtils.h
@@ -13,10 +13,10 @@ public:
 	FInsightsTestUtils(FAutomationTestBase* Test);
 
 	bool AnalyzeTrace(const TCHAR* Path) const;
-	bool FileContainsString(const FString& PathToFile, const FString& ExpectedString, const float& Timeout) const;
-	bool SetupUTS(const float& Timeout) const;
-	bool KillUTS(const float& Timeout) const;
-	bool StartTracing(FTraceAuxiliary::EConnectionType ConnectionType, const float& Timeout) const;
+	bool FileContainsString(const FString& PathToFile, const FString& ExpectedString, double Timeout) const;
+	bool SetupUTS(double Timeout) const;
+	bool KillUTS(double Timeout) const;
+	bool StartTracing(FTraceAuxiliary::EConnectionType ConnectionType, double Timeout) const;
 	void ResetSession() const;
 
 private:
diff --git a/Engine/Source/Developer/TraceInsights/Public/Insights/ViewModels/TimingViewLayout.h b/Engine/Source/Developer/TraceInsights/Public/Insights/ViewModels/TimingViewLayout.h
--- a/Engine/Source/Developer/TraceInsights/Public/Insights/ViewModels/TimingViewLayout.h
+++ b/Engine/Source/Developer/TraceInsights/Public/Insights/ViewModels/TimingViewLayout.h
@@ -39,14 +39,14 @@ struct TRACEINSIGHTS_API FTimingViewLayout
 	float GetLaneY(int32 Depth) const
 	{
 		// 1.0f is for the top horizontal line of each track
-		return 1.0f + TimelineDY + (EventDY + EventH) * Depth;
+		return 1.0f + TimelineDY + (EventDY + EventH) * (float)Depth;
 	}
 
 	//////////////////////////////////////////////////
 
 	float GetChildLaneY(int32 Depth) const
 	{
-		return (EventDY + EventH) * Depth;
+		return (EventDY + EventH) * (float)Depth;
 	}
 
 	/* The layout of a track:
@@ -67,7 +67,7 @@ struct TRACEINSIGHTS_API FTimingViewLayout
 		else //if (NumLanes > 0)
 		{
 			// 1.0f is for horizontal line between timelines
-			float TrackHeight = 1.0f + TimelineDY + (EventH + EventDY) * NumLanes - EventDY + TimelineDY;
+			float TrackHeight = 1.0f + TimelineDY + (EventH + EventDY) * (float)NumLanes - EventDY + TimelineDY;
 
 			if (TrackHeight < FTimingViewLayout::RealMinTimelineH)
 			{
@@ -86,7 +86,7 @@ struct TRACEINSIGHTS_API FTimingViewLayout
 		}
 		else //if (NumLanes > 0)
 		{
-			const float TrackHeight = (EventH + EventDY) * NumLanes - EventDY;
+			const float TrackHeight = (EventH + EventDY) * (float)NumLanes - EventDY;
 
 			if (TrackHeight < FTimingViewLayout::RealMinTimelineH)
 			{
diff --git a/Engine/Source/Developer/TraceInsights/TraceInsights.Build.cs b/Engine/Source/Developer/TraceInsights/TraceInsights.Build.cs
--- a/Engine/Source/Developer/TraceInsights/TraceInsights.Build.cs
+++ b/Engine/Source/Developer/TraceInsights/TraceInsights.Build.cs
@@ -6,8 +6,9 @@ public class TraceInsights : ModuleRules
 {
 	public TraceInsights(ReadOnlyTargetRules Target) : base(Target)
 	{
-		PublicDependencyModuleNames.AddRange
-		(
+		UnsafeTypeCastWarningLevel = WarningLevel.Error;
+
+		PublicDependencyModuleNames.AddRange(
 			new string[] {
 				"AppFramework", // for SColorPicker
 				"ApplicationCore",
@@ -27,7 +28,7 @@ public class TraceInsights : ModuleRules
 				"TraceServices",
 				"WorkspaceMenuStructure",
 				"XmlParser",
-		}
+			}
 		);
 
 		//Modules required for running automation in stand alone Insights
