Network insights
- Fixed aggregate display in PacketView

diff --git a/Engine/Source/Developer/TraceInsights/Private/Insights/NetworkingProfiler/Widgets/SPacketView.cpp b/Engine/Source/Developer/TraceInsights/Private/Insights/NetworkingProfiler/Widgets/SPacketView.cpp
--- a/Engine/Source/Developer/TraceInsights/Private/Insights/NetworkingProfiler/Widgets/SPacketView.cpp
+++ b/Engine/Source/Developer/TraceInsights/Private/Insights/NetworkingProfiler/Widgets/SPacketView.cpp
@@ -362,8 +362,9 @@ void SPacketView::UpdateState()
 								const uint32 EndPos = ~0U;
 								uint32 EndNetIdMatchPos = ~0U;
 								uint32 EndEventTypeMatchPos = ~0U;
+								uint32 LastMatchingLevel = ~0U;
 
-								NetProfilerProvider->EnumeratePacketContentEventsByPosition(ConnectionIndex, ConnectionMode, PacketIndex - 1, StartPos, EndPos, [this, &bFilterMatch, &bOldEventMatchesFilter, &Filter, NetProfilerProvider, &FilterMatchAggregatedEventSizeInBits,&FilterMatchMaxEventSizeBits,  &FilterMatchEventTypeIndex, &EndNetIdMatchPos, &EndEventTypeMatchPos](const TraceServices::FNetProfilerContentEvent& Event)
+								NetProfilerProvider->EnumeratePacketContentEventsByPosition(ConnectionIndex, ConnectionMode, PacketIndex - 1, StartPos, EndPos, [this, &LastMatchingLevel, &bFilterMatch, &bOldEventMatchesFilter, &Filter, NetProfilerProvider, &FilterMatchAggregatedEventSizeInBits,&FilterMatchMaxEventSizeBits,  &FilterMatchEventTypeIndex, &EndNetIdMatchPos, &EndEventTypeMatchPos](const TraceServices::FNetProfilerContentEvent& Event)
 								{
 									if (!bFilterMatch || (Filter.AggregationMode != TraceServices::ENetProfilerAggregationMode::None))
 									{
@@ -403,12 +404,13 @@ void SPacketView::UpdateState()
 										}
 
 										// Check if all conditions are fulfilled but only aggregate stats for top-level event.
-										const bool bEventMatchesFilter = (!Filter.bByNetId || EndNetIdMatchPos != ~0U) && (!Filter.bByEventType || EndEventTypeMatchPos != ~0U);
-										if (bEventMatchesFilter && !bOldEventMatchesFilter)
+										const bool bEventMatchesFilter = (!Filter.bByNetId || EndNetIdMatchPos != ~0U) && (!Filter.bByEventType || EndEventTypeMatchPos != ~0U);										
+										if (bEventMatchesFilter && (!bOldEventMatchesFilter || LastMatchingLevel == Event.Level))
 										{
 											const uint32 EventSize = static_cast<uint32>(Event.EndPos - Event.StartPos);
 											FilterMatchAggregatedEventSizeInBits += EventSize;
 											FilterMatchMaxEventSizeBits = FMath::Max(FilterMatchMaxEventSizeBits, EventSize);
+											LastMatchingLevel = Event.Level;
 
 											if (!bFilterMatch)
 											{
