Added kill UTS logic if tracing isn't started in Engine Test Insights test. Changed start UTS logic, changed kill UTS logic. Added default timeout for trace, UTS start and for UTS kill.
Added check posibillity to connect to UTS function.
Added check of live trace

#tests win64, mac

diff --git a/Engine/Source/Developer/TraceInsights/Private/Insights/Tests/InsightsTestUtils.cpp b/Engine/Source/Developer/TraceInsights/Private/Insights/Tests/InsightsTestUtils.cpp
--- a/Engine/Source/Developer/TraceInsights/Private/Insights/Tests/InsightsTestUtils.cpp
+++ b/Engine/Source/Developer/TraceInsights/Private/Insights/Tests/InsightsTestUtils.cpp
@@ -10,6 +10,7 @@
 #include "TraceServices/Model/AnalysisSession.h"
 #include "TraceServices/ITraceServicesModule.h"
 #include "TraceServices/ModuleService.h"
+#include "Trace/StoreClient.h"
 #include "Insights/Common/Stopwatch.h"
 #include "Insights/IUnrealInsightsModule.h"
 #include "Insights/InsightsManager.h"
@@ -112,39 +113,72 @@ bool FInsightsTestUtils::FileContainsString(const FString& PathToFile, const FSt
 
 ////////////////////////////////////////////////////////////////////////////////////////////////////
 
-bool FInsightsTestUtils::StartTracing(FTraceAuxiliary::EConnectionType ConnectionType, double Timeout) const
+bool FInsightsTestUtils::IsUnrealTraceServerReady(const TCHAR* Host, int32 Port) const
 {
-	bool bStarted = false;
-	if (ConnectionType == FTraceAuxiliary::EConnectionType::Network)
+	UE::Trace::FStoreClient* StoreClient = UE::Trace::FStoreClient::Connect(Host, Port);
+	if (!StoreClient)
 	{
-		bStarted = FTraceAuxiliary::Start(FTraceAuxiliary::EConnectionType::Network, TEXT("localhost"), nullptr);
+		Test->AddInfo(TEXT("Cannot connect to UTS. Trying again"));
+		return false;
 	}
-	else if (ConnectionType == FTraceAuxiliary::EConnectionType::File)
+	const UE::Trace::FStoreClient::FVersion* Version = StoreClient->GetVersion();
+	if (Version == nullptr)
 	{
-		bStarted = FTraceAuxiliary::Start(FTraceAuxiliary::EConnectionType::File, nullptr, nullptr);
+		Test->AddError(TEXT("Cannot get version of UTS"));
+		delete StoreClient;
+		return false;
 	}
 
+	delete StoreClient;
+	Test->AddInfo(FString::Printf(TEXT("Connected to UTS version %u.%u"), Version->GetMajorVersion(), Version->GetMinorVersion()));
+	return true;
+}
+
+////////////////////////////////////////////////////////////////////////////////////////////////////
+
+bool FInsightsTestUtils::StartTracing(FTraceAuxiliary::EConnectionType ConnectionType, double Timeout) const
+{
+	bool bStarted = false;
+
 	double TraceVerifyStartTime = FPlatformTime::Seconds();
-	while (bStarted && (FPlatformTime::Seconds() - TraceVerifyStartTime < Timeout))
+	while(FPlatformTime::Seconds() - TraceVerifyStartTime < Timeout)
 	{
-		FPlatformProcess::Sleep(0.5f);
+		if (ConnectionType == FTraceAuxiliary::EConnectionType::Network)
+		{
+			bStarted = FTraceAuxiliary::Start(FTraceAuxiliary::EConnectionType::Network, TEXT("localhost"), nullptr);
+		}
+		else if (ConnectionType == FTraceAuxiliary::EConnectionType::File)
+		{
+			bStarted = FTraceAuxiliary::Start(FTraceAuxiliary::EConnectionType::File, nullptr, nullptr);
+		}
+
 		if (FTraceAuxiliary::IsConnected())
 		{
-			return true;
+			FPlatformProcess::Sleep(0.5f);
+			return bStarted;
 		}
+		FPlatformProcess::Sleep(0.1f);
 	}
 
-	return false;
+	return bStarted;
 }
 
 ////////////////////////////////////////////////////////////////////////////////////////////////////
 
-bool FInsightsTestUtils::SetupUTS(double Timeout) const
+bool FInsightsTestUtils::SetupUTS(double Timeout, bool bUseFork) const
 {
 	const FString UnrealTraceServerName = TEXT("UnrealTraceServer");
 
 	FString UTSPath = FPlatformProcess::GenerateApplicationPath("UnrealTraceServer", EBuildConfiguration::Development);
-	FString UTSParameters = TEXT("daemon");
+	FString UTSParameters;
+	if (bUseFork)
+	{
+		UTSParameters = TEXT("fork");
+	}
+	else
+	{
+		UTSParameters = TEXT("daemon");
+	}
 	constexpr bool bLaunchDetached = true;
 	constexpr bool bLaunchHidden = false;
 	constexpr bool bLaunchReallyHidden = false;
@@ -156,20 +190,29 @@ bool FInsightsTestUtils::SetupUTS(double Timeout) const
 	FProcHandle UTSHandle = FPlatformProcess::CreateProc(*UTSPath, *UTSParameters, bLaunchDetached, bLaunchHidden, bLaunchReallyHidden, &ProcessID, PriorityModifier, OptionalWorkingDirectory, PipeWriteChild, PipeReadChild);
 	if (!UTSHandle.IsValid())
 	{
+		Test->AddError(TEXT("The UTSHandle should be valid"));
 		return false;
 	}
 
 	double StartTime = FPlatformTime::Seconds();
 	while (FPlatformTime::Seconds() - StartTime < Timeout)
 	{
-		if (FPlatformProcess::IsApplicationRunning(*UnrealTraceServerName))
+		FPlatformProcess::Sleep(0.1f);
+		if (!FPlatformProcess::IsApplicationRunning(*UnrealTraceServerName))
 		{
-			return true;
+			Test->AddInfo(TEXT("UTS not started yet"));
+			continue;
 		}
-		UTSHandle = FPlatformProcess::CreateProc(*UTSPath, *UTSParameters, bLaunchDetached, bLaunchHidden, bLaunchReallyHidden, &ProcessID, PriorityModifier, OptionalWorkingDirectory, PipeWriteChild, PipeReadChild);
-		FPlatformProcess::Sleep(0.1f);
+		if (!IsUnrealTraceServerReady())
+		{
+			Test->AddInfo(TEXT("UTS not ready yet"));
+			continue;
+		}
+		Test->AddInfo(TEXT("UTS is ready"));
+		return true;
 	}
 
+	Test->AddError(TEXT("UTS failed to start"));
 	return false;
 }
 
@@ -192,20 +235,22 @@ bool FInsightsTestUtils::KillUTS(double Timeout) const
 	FProcHandle UTSHandle = FPlatformProcess::CreateProc(*UTSPath, *UTSParameters, bLaunchDetached, bLaunchHidden, bLaunchReallyHidden, &ProcessID, PriorityModifier, OptionalWorkingDirectory, PipeWriteChild, PipeReadChild);
 	if (!UTSHandle.IsValid())
 	{
+		Test->AddError(TEXT("The UTSHandle should be valid"));
 		return false;
 	}
 
 	double StartTime = FPlatformTime::Seconds();
 	while (FPlatformTime::Seconds() - StartTime < Timeout)
 	{
+		FPlatformProcess::Sleep(0.1f);
 		if (!FPlatformProcess::IsApplicationRunning(*UnrealTraceServerName))
 		{
+			Test->AddInfo(TEXT("The UTS successfully killed"));
 			return true;
 		}
-		UTSHandle = FPlatformProcess::CreateProc(*UTSPath, *UTSParameters, bLaunchDetached, bLaunchHidden, bLaunchReallyHidden, &ProcessID, PriorityModifier, OptionalWorkingDirectory, PipeWriteChild, PipeReadChild); 
-		FPlatformProcess::Sleep(0.1f);
 	}
 
+	Test->AddError(TEXT("UTS failed to kill"));
 	return false;
 }
 
@@ -215,4 +260,44 @@ void FInsightsTestUtils::ResetSession() const
 {
 	TSharedPtr<FInsightsManager> InsightsManager = FInsightsManager::Get();
 	InsightsManager->ResetSession();
+}
+
+////////////////////////////////////////////////////////////////////////////////////////////////////
+
+bool FInsightsTestUtils::IsTraceHasLiveStatus(const FString& TraceName, const TCHAR* Host, int32 Port) const
+{
+	UE::Trace::FStoreClient* StoreClient = UE::Trace::FStoreClient::Connect(Host, Port);
+	if (!StoreClient)
+	{
+		Test->AddInfo(TEXT("The StoreClient shouldn't be null"));
+		return false;
+	}
+	uint32 SessionCount = StoreClient->GetSessionCount();
+	if (!SessionCount)
+	{
+		Test->AddInfo(TEXT("The SessionCount shouldn't be 0"));
+		delete StoreClient;
+		return false;
+	}
+
+	for (uint32 Index = 0; Index < SessionCount; ++Index)
+	{
+		const UE::Trace::FStoreClient::FSessionInfo* SessionInfo = StoreClient->GetSessionInfo(Index);
+		if (!SessionInfo)
+		{
+			continue;
+		}
+		uint32 TraceId = SessionInfo->GetTraceId();
+		const UE::Trace::FStoreClient::FTraceInfo* Info = StoreClient->GetTraceInfoById(TraceId);
+		if (TraceName.Contains(static_cast<FString>(Info->GetName())))
+		{
+			Test->AddInfo(TEXT("Trace is live"));
+			delete StoreClient;
+			return true;
+		}
+	}
+
+	Test->AddInfo(FString::Printf(TEXT("The trace with name %s does not have live status. Trying to find live trace"), *TraceName));
+	delete StoreClient;
+	return false;
 }
\ No newline at end of file
diff --git a/Engine/Source/Developer/TraceInsights/Public/Insights/Tests/InsightsTestUtils.h b/Engine/Source/Developer/TraceInsights/Public/Insights/Tests/InsightsTestUtils.h
--- a/Engine/Source/Developer/TraceInsights/Public/Insights/Tests/InsightsTestUtils.h
+++ b/Engine/Source/Developer/TraceInsights/Public/Insights/Tests/InsightsTestUtils.h
@@ -14,9 +14,11 @@ public:
 
 	bool AnalyzeTrace(const TCHAR* Path) const;
 	bool FileContainsString(const FString& PathToFile, const FString& ExpectedString, double Timeout) const;
-	bool SetupUTS(double Timeout) const;
-	bool KillUTS(double Timeout) const;
-	bool StartTracing(FTraceAuxiliary::EConnectionType ConnectionType, double Timeout) const;
+	bool SetupUTS(double Timeout = 30.0, bool bUseFork = false) const;
+	bool KillUTS(double Timeout = 30.0) const;
+	bool StartTracing(FTraceAuxiliary::EConnectionType ConnectionType, double Timeout = 10.0) const;
+	bool IsUnrealTraceServerReady(const TCHAR* Host = TEXT("localhost"), int32 Port = 0U) const;
+	bool IsTraceHasLiveStatus(const FString& TraceName, const TCHAR* Host = TEXT("localhost"), int32 Port = 0U) const;
 	void ResetSession() const;
 
 private:
