[Insights] Add persistence for Main Graph Track setting in Timing Profiler.

diff --git a/Engine/Programs/UnrealInsights/Config/DefaultUnrealInsightsSettings.ini b/Engine/Programs/UnrealInsights/Config/DefaultUnrealInsightsSettings.ini
--- a/Engine/Programs/UnrealInsights/Config/DefaultUnrealInsightsSettings.ini
+++ b/Engine/Programs/UnrealInsights/Config/DefaultUnrealInsightsSettings.ini
@@ -22,6 +22,24 @@ AutoScrollMinDelay=0.3
 ; If enabled, the panning is allowed to continue when the mouse cursor reaches the edges of the screen.
 bAllowPanningOnScreenEdges=False
 
+[Insights.TimingProfiler.MainGraph]
+; If enabled, values will be displayed as points in the Main Graph Track in Timing Insights.
+ShowPoints=False
+; If enabled, values will be displayed as points with border in the Main Graph Track in Timing Insights.
+ShowPointsWithBorder=True
+; If enabled, values will be displayed as connected lines in the Main Graph Track in Timing Insights.
+ShowConnectedLines=True
+; If enabled, values will be displayed as polygons in the Main Graph Track in Timing Insights.
+ShowPolygons=True
+; If enabled, uses duration of timing events for connected lines and polygons in the Main Graph Track in Timing Insights.
+ShowEventDuration=True
+; If enabled, shows bars corresponding to the duration of the timing events in the Main Graph Track in Timing Insights.
+ShowBars=False
+; If enabled, shows game frames in the Main Graph Track in Timing Insights.
+ShowGameFrames=True
+; If enabled, shows rendering frames in the Main Graph Track in Timing Insights.
+ShowRenderingFrame=True
+
 [Insights.TimingProfiler.TimersView]
 ; Names of visible columns in the Timers Panel in the Game mode.
 +GameFrameColumns=MaxInclTime
diff --git a/Engine/Source/Developer/TraceInsights/Private/Insights/InsightsSettings.h b/Engine/Source/Developer/TraceInsights/Private/Insights/InsightsSettings.h
--- a/Engine/Source/Developer/TraceInsights/Private/Insights/InsightsSettings.h
+++ b/Engine/Source/Developer/TraceInsights/Private/Insights/InsightsSettings.h
@@ -28,6 +28,14 @@ public:
 		, bTimersViewShowCpuTimers(true)
 		, bTimersViewShowGpuTimers(true)
 		, bTimersViewShowZeroCountTimers(true)
+		, bTimingViewMainGraphShowPoints(false)
+		, bTimingViewMainGraphShowPointsWithBorder(true)
+		, bTimingViewMainGraphShowConnectedLines(true)
+		, bTimingViewMainGraphShowPolygons(true)
+		, bTimingViewMainGraphShowEventDuration(true)
+		, bTimingViewMainGraphShowBars(false)
+		, bTimingViewMainGraphShowGameFrames(true)
+		, bTimingViewMainGraphShowRenderingFrames(true)
 	{
 		if (!bIsDefault)
 		{
@@ -103,6 +111,16 @@ public:
 		GConfig->GetBool(TEXT("Insights.TimingProfiler.TimersView"), TEXT("ShowCpuTimers"), bTimersViewShowCpuTimers, SettingsIni);
 		GConfig->GetBool(TEXT("Insights.TimingProfiler.TimersView"), TEXT("ShowGpuTimers"), bTimersViewShowGpuTimers, SettingsIni);
 		GConfig->GetBool(TEXT("Insights.TimingProfiler.TimersView"), TEXT("ShowZeroCountTimers"), bTimersViewShowZeroCountTimers, SettingsIni);
+
+		GConfig->GetBool(TEXT("Insights.TimingProfiler.MainGraph"), TEXT("ShowPoints"), bTimingViewMainGraphShowPoints, SettingsIni);
+		GConfig->GetBool(TEXT("Insights.TimingProfiler.MainGraph"), TEXT("ShowPointsWithBorder"), bTimingViewMainGraphShowPointsWithBorder, SettingsIni);
+		GConfig->GetBool(TEXT("Insights.TimingProfiler.MainGraph"), TEXT("ShowConnectedLines"), bTimingViewMainGraphShowConnectedLines, SettingsIni);
+		GConfig->GetBool(TEXT("Insights.TimingProfiler.MainGraph"), TEXT("ShowPolygons"), bTimingViewMainGraphShowPolygons, SettingsIni);
+		GConfig->GetBool(TEXT("Insights.TimingProfiler.MainGraph"), TEXT("ShowEventDuration"), bTimingViewMainGraphShowEventDuration, SettingsIni);
+		GConfig->GetBool(TEXT("Insights.TimingProfiler.MainGraph"), TEXT("ShowBars"), bTimingViewMainGraphShowBars, SettingsIni);
+		GConfig->GetBool(TEXT("Insights.TimingProfiler.MainGraph"), TEXT("ShowGameFrames"), bTimingViewMainGraphShowGameFrames, SettingsIni);
+		GConfig->GetBool(TEXT("Insights.TimingProfiler.MainGraph"), TEXT("ShowRenderingFrame"), bTimingViewMainGraphShowRenderingFrames, SettingsIni);
+
 	}
 
 	void SaveToConfig()
@@ -132,6 +150,15 @@ public:
 		GConfig->SetInt(TEXT("Insights.TimingProfiler.TimersView"), TEXT("ShowGpuTimers"), bTimersViewShowGpuTimers, SettingsIni);
 		GConfig->SetInt(TEXT("Insights.TimingProfiler.TimersView"), TEXT("ShowZeroCountTimers"), bTimersViewShowZeroCountTimers, SettingsIni);
 
+		GConfig->SetBool(TEXT("Insights.TimingProfiler.MainGraph"), TEXT("ShowPoints"), bTimingViewMainGraphShowPoints, SettingsIni);
+		GConfig->SetBool(TEXT("Insights.TimingProfiler.MainGraph"), TEXT("ShowPointsWithBorder"), bTimingViewMainGraphShowPointsWithBorder, SettingsIni);
+		GConfig->SetBool(TEXT("Insights.TimingProfiler.MainGraph"), TEXT("ShowConnectedLines"), bTimingViewMainGraphShowConnectedLines, SettingsIni);
+		GConfig->SetBool(TEXT("Insights.TimingProfiler.MainGraph"), TEXT("ShowPolygons"), bTimingViewMainGraphShowPolygons, SettingsIni);
+		GConfig->SetBool(TEXT("Insights.TimingProfiler.MainGraph"), TEXT("ShowEventDuration"), bTimingViewMainGraphShowEventDuration, SettingsIni);
+		GConfig->SetBool(TEXT("Insights.TimingProfiler.MainGraph"), TEXT("ShowBars"), bTimingViewMainGraphShowBars, SettingsIni);
+		GConfig->SetBool(TEXT("Insights.TimingProfiler.MainGraph"), TEXT("ShowGameFrames"), bTimingViewMainGraphShowGameFrames, SettingsIni);
+		GConfig->SetBool(TEXT("Insights.TimingProfiler.MainGraph"), TEXT("ShowRenderingFrame"), bTimingViewMainGraphShowRenderingFrames, SettingsIni);
+
 		GConfig->Flush(false, SettingsIni);
 	}
 
@@ -236,6 +263,38 @@ public:
 	void SetTimersViewShowZeroCountTimers(bool InValue) { bTimersViewShowZeroCountTimers = InValue; }
 	void SetAndSaveTimersViewShowZeroCountTimers(bool InValue) { SET_AND_SAVE(bTimersViewShowZeroCountTimers, InValue); }
 
+	bool GetTimingViewMainGraphShowPoints() const { return bTimingViewMainGraphShowPoints; }
+	void SetTimingViewMainGraphShowPoints(bool InValue) { bTimingViewMainGraphShowPoints = InValue; }
+	void SetAndSaveTimingViewMainGraphShowPoints(bool InValue) { SET_AND_SAVE(bTimingViewMainGraphShowPoints, InValue); }
+
+	bool GetTimingViewMainGraphShowPointsWithBorder() const { return bTimingViewMainGraphShowPointsWithBorder; }
+	void SetTimingViewMainGraphShowPointsWithBorder(bool InValue) { bTimingViewMainGraphShowPointsWithBorder = InValue; }
+	void SetAndSaveTimingViewMainGraphShowPointsWithBorder(bool InValue) { SET_AND_SAVE(bTimingViewMainGraphShowPointsWithBorder, InValue); }
+
+	bool GetTimingViewMainGraphShowConnectedLines() const { return bTimingViewMainGraphShowConnectedLines; }
+	void SetTimingViewMainGraphShowConnectedLines(bool InValue) { bTimingViewMainGraphShowConnectedLines = InValue; }
+	void SetAndSaveTimingViewMainGraphShowConnectedLines(bool InValue) { SET_AND_SAVE(bTimingViewMainGraphShowConnectedLines, InValue); }
+
+	bool GetTimingViewMainGraphShowPolygons() const { return bTimingViewMainGraphShowPolygons; }
+	void SetTimingViewMainGraphShowPolygons(bool InValue) { bTimingViewMainGraphShowPolygons = InValue; }
+	void SetAndTimingViewMainGraphShowPolygons(bool InValue) { SET_AND_SAVE(bTimingViewMainGraphShowPolygons, InValue); }
+
+	bool GetTimingViewMainGraphShowEventDuration() const { return bTimingViewMainGraphShowEventDuration; }
+	void SetTimingViewMainGraphShowEventDuration(bool InValue) { bTimingViewMainGraphShowEventDuration = InValue; }
+	void SetAndSaveTimingViewMainGraphShowEventDuration(bool InValue) { SET_AND_SAVE(bTimingViewMainGraphShowEventDuration, InValue); }
+
+	bool GetTimingViewMainGraphShowBars() const { return bTimingViewMainGraphShowBars; }
+	void SetTimingViewMainGraphShowBars(bool InValue) { bTimingViewMainGraphShowBars = InValue; }
+	void SetAndSaveTimingViewMainGraphShowBars(bool InValue) { SET_AND_SAVE(bTimingViewMainGraphShowBars, InValue); }
+
+	bool GetTimingViewMainGraphShowGameFrames() const { return bTimingViewMainGraphShowGameFrames; }
+	void SetTimingViewMainGraphShowGameFrames(bool InValue) { bTimingViewMainGraphShowGameFrames = InValue; }
+	void SetAndSaveTimingViewMainGraphShowGameFrames(bool InValue) { SET_AND_SAVE(bTimingViewMainGraphShowGameFrames, InValue); }
+
+	bool GetTimingViewMainGraphShowRenderingFrames() const { return bTimingViewMainGraphShowRenderingFrames; }
+	void SetTimingViewMainGraphShowRenderingFrames(bool InValue) { bTimingViewMainGraphShowRenderingFrames = InValue; }
+	void SetAndSaveTimingViewMainGraphShowRenderingFrames(bool InValue) { SET_AND_SAVE(bTimingViewMainGraphShowRenderingFrames, InValue); }
+
 	#undef SET_AND_SAVE
 
 private:
@@ -308,4 +367,28 @@ private:
 
 	//** If enabled, timers with no instances in the selected interval will still be displayed in the Timers View. */
 	bool bTimersViewShowZeroCountTimers;
+
+	//** If enabled, values will be displayed as points in the Main Graph Track in Timing Insights. */
+	bool bTimingViewMainGraphShowPoints;
+
+	//** If enabled, values will be displayed as points with border in the Main Graph Track in Timing Insights. */
+	bool bTimingViewMainGraphShowPointsWithBorder;
+
+	//** If enabled, values will be displayed as connected lines in the Main Graph Track in Timing Insights. */
+	bool bTimingViewMainGraphShowConnectedLines;
+
+	//** If enabled, values will be displayed as polygons in the Main Graph Track in Timing Insights. */
+	bool bTimingViewMainGraphShowPolygons;
+
+	//** If enabled, uses duration of timing events for connected lines and polygons in the Main Graph Track in Timing Insights. */
+	bool bTimingViewMainGraphShowEventDuration;
+
+	//** If enabled, shows bars corresponding to the duration of the timing events in the Main Graph Track in Timing Insights. */
+	bool bTimingViewMainGraphShowBars;
+
+	//** If enabled, shows game frames in the Main Graph Track in Timing Insights. */
+	bool bTimingViewMainGraphShowGameFrames;
+
+	//** If enabled, shows rendering frames in the Main Graph Track in Timing Insights. */
+	bool bTimingViewMainGraphShowRenderingFrames;
 };
diff --git a/Engine/Source/Developer/TraceInsights/Private/Insights/ViewModels/TimingGraphTrack.cpp b/Engine/Source/Developer/TraceInsights/Private/Insights/ViewModels/TimingGraphTrack.cpp
--- a/Engine/Source/Developer/TraceInsights/Private/Insights/ViewModels/TimingGraphTrack.cpp
+++ b/Engine/Source/Developer/TraceInsights/Private/Insights/ViewModels/TimingGraphTrack.cpp
@@ -112,6 +112,15 @@ FString FTimingGraphSeries::FormatValue(double Value) const
 	return FGraphSeries::FormatValue(Value);
 }
 
+////////////////////////////////////////////////////////////////////////////////////////////////////
+
+void FTimingGraphSeries::SetVisibility(bool bOnOff)
+{
+	FGraphSeries::SetVisibility(bOnOff);
+
+	VisibilityChangedDelegate.Broadcast(bOnOff);
+}
+
 ////////////////////////////////////////////////////////////////////////////////////////////////////
 // FTimingGraphTrack
 ////////////////////////////////////////////////////////////////////////////////////////////////////
@@ -120,8 +129,9 @@ INSIGHTS_IMPLEMENT_RTTI(FTimingGraphTrack)
 
 ////////////////////////////////////////////////////////////////////////////////////////////////////
 
-FTimingGraphTrack::FTimingGraphTrack()
+FTimingGraphTrack::FTimingGraphTrack(TSharedPtr<STimingView> InTimingView)
 	: FGraphTrack()
+	, TimingView(InTimingView)
 	//, SharedValueViewport()
 {
 	EnabledOptions = //EGraphOptions::ShowDebugInfo |
@@ -135,6 +145,8 @@ FTimingGraphTrack::FTimingGraphTrack()
 					 EGraphOptions::ShowVerticalAxisGrid |
 					 EGraphOptions::ShowHeader |
 					 EGraphOptions::None;
+
+	LoadDefaultSettings();
 }
 
 ////////////////////////////////////////////////////////////////////////////////////////////////////
@@ -143,16 +155,30 @@ FTimingGraphTrack::~FTimingGraphTrack()
 {
 	if (OnTrackVisibilityChangedHandle.IsValid())
 	{
-		TSharedPtr<class STimingProfilerWindow> TimingWindow = FTimingProfilerManager::Get()->GetProfilerWindow();
-		if (TimingWindow.IsValid())
+		TSharedPtr<STimingView> TimingViewPtr = TimingView.Pin();
+		if (TimingView.IsValid())
 		{
-			TSharedPtr<STimingView> TimingView = TimingWindow->GetTimingView();
-			if (TimingView.IsValid())
-			{
-				TimingView->OnTrackVisibilityChanged().Remove(OnTrackVisibilityChangedHandle);
-				TimingView->OnTrackAdded().Remove(OnTrackAddedHandle);
-				TimingView->OnTrackRemoved().Remove(OnTrackRemovedHandle);
-			}
+			TimingViewPtr->OnTrackVisibilityChanged().Remove(OnTrackVisibilityChangedHandle);
+			TimingViewPtr->OnTrackAdded().Remove(OnTrackAddedHandle);
+			TimingViewPtr->OnTrackRemoved().Remove(OnTrackRemovedHandle);
+		}
+	}
+
+	if (GameFrameSeriesVisibilityHandle.IsValid())
+	{
+		TSharedPtr<FTimingGraphSeries> GameFramesSeries = GetFrameSeries(ETraceFrameType::TraceFrameType_Game);
+		if (GameFramesSeries.IsValid())
+		{
+			GameFramesSeries->VisibilityChangedDelegate.Remove(GameFrameSeriesVisibilityHandle);
+		}
+	}
+
+	if (RenderingFrameSeriesVisibilityHandle.IsValid())
+	{
+		TSharedPtr<FTimingGraphSeries> RenderingFramesSeries = GetFrameSeries(ETraceFrameType::TraceFrameType_Game);
+		if (RenderingFramesSeries.IsValid())
+		{
+			RenderingFramesSeries->VisibilityChangedDelegate.Remove(RenderingFrameSeriesVisibilityHandle);
 		}
 	}
 }
@@ -162,34 +188,31 @@ FTimingGraphTrack::~FTimingGraphTrack()
 void FTimingGraphTrack::Update(const ITimingTrackUpdateContext& Context)
 {
 	FGraphTrack::Update(Context);
+
 	if (!OnTrackVisibilityChangedHandle.IsValid())
 	{
-		TSharedPtr<class STimingProfilerWindow> TimingWindow = FTimingProfilerManager::Get()->GetProfilerWindow();
-		if (TimingWindow.IsValid())
+		TSharedPtr<STimingView> TimingViewPtr = TimingView.Pin();
+		if (TimingViewPtr.IsValid())
 		{
-			TSharedPtr<STimingView> TimingView = TimingWindow->GetTimingView();
-			if (TimingView.IsValid())
+			auto OnTrackAddedRemovedLamda = [this](const TSharedPtr<const FBaseTimingTrack> Track)
 			{
-				auto OnTrackAddedRemovedLamda = [this](const TSharedPtr<const FBaseTimingTrack> Track)
+				if (Track->Is<FThreadTimingTrack>())
 				{
-					if (Track->Is<FThreadTimingTrack>())
+					// If there are more series than the default frame series.
+					if (this->AllSeries.Num() > ETraceFrameType::TraceFrameType_Count)
 					{
-						// If there are more series than the default frame series.
-						if (this->AllSeries.Num() > ETraceFrameType::TraceFrameType_Count)
-						{
-							this->SetDirtyFlag();
-						}
+						this->SetDirtyFlag();
 					}
-				};
+				}
+			};
 
-				OnTrackAddedHandle = TimingView->OnTrackAdded().AddLambda(OnTrackAddedRemovedLamda);
-				OnTrackRemovedHandle = TimingView->OnTrackRemoved().AddLambda(OnTrackAddedRemovedLamda);
+			OnTrackAddedHandle = TimingViewPtr->OnTrackAdded().AddLambda(OnTrackAddedRemovedLamda);
+			OnTrackRemovedHandle = TimingViewPtr->OnTrackRemoved().AddLambda(OnTrackAddedRemovedLamda);
 
-				OnTrackVisibilityChangedHandle = TimingView->OnTrackVisibilityChanged().AddLambda([this]()
-					{
-						this->SetDirtyFlag();
-					});
-			}
+			OnTrackVisibilityChangedHandle = TimingViewPtr->OnTrackVisibilityChanged().AddLambda([this]()
+				{
+					this->SetDirtyFlag();
+				});
 		}
 	}
 
@@ -256,6 +279,9 @@ void FTimingGraphTrack::Update(const ITimingTrackUpdateContext& Context)
 
 void FTimingGraphTrack::AddDefaultFrameSeries()
 {
+	const FInsightsSettings& Settings = FInsightsManager::Get()->GetSettings();
+	TSharedPtr<STimingView> TimingViewPtr = TimingView.Pin();
+
 	TSharedRef<FTimingGraphSeries> GameFramesSeries = MakeShared<FTimingGraphSeries>(FTimingGraphSeries::ESeriesType::Frame);
 	GameFramesSeries->SetName(TEXT("Game Frames"));
 	GameFramesSeries->SetDescription(TEXT("Duration of Game frames"));
@@ -264,6 +290,15 @@ void FTimingGraphTrack::AddDefaultFrameSeries()
 	GameFramesSeries->SetBaselineY(SharedValueViewport.GetBaselineY());
 	GameFramesSeries->SetScaleY(SharedValueViewport.GetScaleY());
 	GameFramesSeries->EnableSharedViewport();
+	if (TimingViewPtr.IsValid() && TimingViewPtr->GetName() == FInsightsManagerTabs::TimingProfilerTabId)
+	{
+		GameFramesSeries->SetVisibility(Settings.GetTimingViewMainGraphShowGameFrames());
+		GameFrameSeriesVisibilityHandle = GameFramesSeries->VisibilityChangedDelegate.AddLambda([](bool bOnOff)
+			{
+				FInsightsSettings& Settings = FInsightsManager::Get()->GetSettings();
+				Settings.SetAndSaveTimingViewMainGraphShowGameFrames(bOnOff);
+			});
+	}
 	AllSeries.Add(GameFramesSeries);
 
 	TSharedRef<FTimingGraphSeries> RenderingFramesSeries = MakeShared<FTimingGraphSeries>(FTimingGraphSeries::ESeriesType::Frame);
@@ -274,11 +309,32 @@ void FTimingGraphTrack::AddDefaultFrameSeries()
 	RenderingFramesSeries->SetBaselineY(SharedValueViewport.GetBaselineY());
 	RenderingFramesSeries->SetScaleY(SharedValueViewport.GetScaleY());
 	RenderingFramesSeries->EnableSharedViewport();
+	if (TimingViewPtr.IsValid() && TimingViewPtr->GetName() == FInsightsManagerTabs::TimingProfilerTabId)
+	{
+		RenderingFramesSeries->SetVisibility(Settings.GetTimingViewMainGraphShowRenderingFrames());
+		RenderingFrameSeriesVisibilityHandle = RenderingFramesSeries->VisibilityChangedDelegate.AddLambda([](bool bOnOff)
+			{
+				FInsightsSettings& Settings = FInsightsManager::Get()->GetSettings();
+				Settings.SetAndSaveTimingViewMainGraphShowRenderingFrames(bOnOff);
+			});
+	}
 	AllSeries.Add(RenderingFramesSeries);
 }
 
 ////////////////////////////////////////////////////////////////////////////////////////////////////
 
+TSharedPtr<FTimingGraphSeries> FTimingGraphTrack::GetFrameSeries(ETraceFrameType FrameType)
+{
+	TSharedPtr<FGraphSeries>* Ptr = AllSeries.FindByPredicate([FrameType](const TSharedPtr<FGraphSeries>& Series)
+		{
+			const TSharedPtr<FTimingGraphSeries> TimingSeries = StaticCastSharedPtr<FTimingGraphSeries>(Series);
+			return TimingSeries->Type == FTimingGraphSeries::ESeriesType::Frame && TimingSeries->FrameType == FrameType;
+		});
+	return (Ptr != nullptr) ? StaticCastSharedPtr<FTimingGraphSeries>(*Ptr) : nullptr;
+}
+
+////////////////////////////////////////////////////////////////////////////////////////////////////
+
 void FTimingGraphTrack::UpdateFrameSeries(FTimingGraphSeries& Series, const FTimingTrackViewport& Viewport)
 {
 	FGraphTrackBuilder Builder(*this, Series, Viewport);
@@ -718,23 +774,95 @@ void FTimingGraphTrack::UpdateStatsCounterSeries(FTimingGraphSeries& Series, con
 
 void FTimingGraphTrack::GetVisibleTimelineIndexes(TSet<uint32>& TimelineIndexes)
 {
-	TSharedPtr<class STimingProfilerWindow> TimingWindow = FTimingProfilerManager::Get()->GetProfilerWindow();
-	if (!TimingWindow.IsValid())
+	TSharedPtr<STimingView> TimingViewPtr = TimingView.Pin();
+	if (!TimingViewPtr.IsValid())
 	{
 		return;
 	}
 
-	TSharedPtr<STimingView> TimingView = TimingWindow->GetTimingView();
-	if (!TimingView.IsValid())
+	TSharedPtr<FThreadTimingSharedState> ThreadSharedState = TimingViewPtr->GetThreadTimingSharedState();
+	ThreadSharedState->GetVisibleTimelineIndexes(TimelineIndexes);
+}
+
+////////////////////////////////////////////////////////////////////////////////////////////////////
+
+void FTimingGraphTrack::ContextMenu_ToggleOption_Execute(EGraphOptions Option)
+{
+	FGraphTrack::ContextMenu_ToggleOption_Execute(Option);
+
+	TSharedPtr<STimingView> TimingViewPtr = TimingView.Pin();
+	if (!TimingViewPtr.IsValid())
+	{
+		return;
+	}
+	if (TimingViewPtr->GetName() != FInsightsManagerTabs::TimingProfilerTabId)
 	{
 		return;
 	}
 
-	TSharedPtr<FThreadTimingSharedState> ThreadSharedState = TimingView->GetThreadTimingSharedState();
-	ThreadSharedState->GetVisibleTimelineIndexes(TimelineIndexes);
+	FInsightsSettings& Settings = FInsightsManager::Get()->GetSettings();
+	if (EnumHasAnyFlags(Option, EGraphOptions::ShowPoints))
+	{
+		Settings.SetAndSaveTimingViewMainGraphShowPoints(EnumHasAnyFlags(EnabledOptions, EGraphOptions::ShowPoints));
+	}
+	if (EnumHasAnyFlags(Option, EGraphOptions::ShowPointsWithBorder))
+	{
+		Settings.SetAndSaveTimingViewMainGraphShowPointsWithBorder(EnumHasAnyFlags(EnabledOptions, EGraphOptions::ShowPointsWithBorder));
+	}
+	if (EnumHasAnyFlags(Option, EGraphOptions::ShowLines))
+	{
+		Settings.SetAndSaveTimingViewMainGraphShowConnectedLines(EnumHasAnyFlags(EnabledOptions, EGraphOptions::ShowLines));
+	}
+	if (EnumHasAnyFlags(Option, EGraphOptions::ShowPolygon))
+	{
+		Settings.SetAndTimingViewMainGraphShowPolygons(EnumHasAnyFlags(EnabledOptions, EGraphOptions::ShowPolygon));
+	}
+	if (EnumHasAnyFlags(Option, EGraphOptions::UseEventDuration))
+	{
+		Settings.SetAndSaveTimingViewMainGraphShowEventDuration(EnumHasAnyFlags(EnabledOptions, EGraphOptions::UseEventDuration));
+	}
+	if (EnumHasAnyFlags(Option, EGraphOptions::ShowBars))
+	{
+		Settings.SetAndSaveTimingViewMainGraphShowBars(EnumHasAnyFlags(EnabledOptions, EGraphOptions::ShowBars));
+	}
 }
 
 ////////////////////////////////////////////////////////////////////////////////////////////////////
+
+void FTimingGraphTrack::LoadDefaultSettings()
+{
+	TSharedPtr<STimingView> TimingViewPtr = TimingView.Pin();
+	if (TimingViewPtr.IsValid() && TimingViewPtr->GetName() == FInsightsManagerTabs::TimingProfilerTabId)
+	{
+		EnabledOptions = EGraphOptions::None;
+		const FInsightsSettings& Settings = FInsightsManager::Get()->GetSettings();
+		if (Settings.GetTimingViewMainGraphShowPoints())
+		{
+			EnabledOptions |= EGraphOptions::ShowPoints;
+		}
+		if (Settings.GetTimingViewMainGraphShowPointsWithBorder())
+		{
+			EnabledOptions |= EGraphOptions::ShowPointsWithBorder;
+		}
+		if (Settings.GetTimingViewMainGraphShowConnectedLines())
+		{
+			EnabledOptions |= EGraphOptions::ShowLines;
+		}
+		if (Settings.GetTimingViewMainGraphShowPolygons())
+		{
+			EnabledOptions |= EGraphOptions::ShowPolygon;
+		}
+		if (Settings.GetTimingViewMainGraphShowEventDuration())
+		{
+			EnabledOptions |= EGraphOptions::UseEventDuration;
+		}
+		if (Settings.GetTimingViewMainGraphShowBars())
+		{
+			EnabledOptions |= EGraphOptions::ShowBars;
+		}
+	}
+}
+
 ////////////////////////////////////////////////////////////////////////////////////////////////////
 
 void FTimingGraphTrack::DrawVerticalAxisGrid(const ITimingTrackDrawContext& Context) const
diff --git a/Engine/Source/Developer/TraceInsights/Private/Insights/ViewModels/TimingGraphTrack.h b/Engine/Source/Developer/TraceInsights/Private/Insights/ViewModels/TimingGraphTrack.h
--- a/Engine/Source/Developer/TraceInsights/Private/Insights/ViewModels/TimingGraphTrack.h
+++ b/Engine/Source/Developer/TraceInsights/Private/Insights/ViewModels/TimingGraphTrack.h
@@ -13,8 +13,13 @@ namespace Insights
 	struct FFrameStatsCachedEvent;
 }
 
+class STimingView;
+
 ////////////////////////////////////////////////////////////////////////////////////////////////////
 
+/** The delegate to be invoked when a series visibility is changed. */
+DECLARE_MULTICAST_DELEGATE_OneParam(FSeriesVisibilityChangedDelegate, bool bOnOff);
+
 class FTimingGraphSeries : public FGraphSeries
 {
 public:
@@ -43,6 +48,8 @@ public:
 		return EventA.StartTime < EventB.StartTime;
 	}
 
+	virtual void SetVisibility(bool bOnOff) override;
+
 public:
 	ESeriesType Type;
 	union
@@ -60,6 +67,8 @@ public:
 	bool bIsTime; // the unit for values is [second]
 	bool bIsMemory; // the unit for value is [byte]
 	bool bIsFloatingPoint; // for stats counters
+
+	FSeriesVisibilityChangedDelegate VisibilityChangedDelegate;
 };
 
 ////////////////////////////////////////////////////////////////////////////////////////////////////
@@ -69,13 +78,15 @@ class FTimingGraphTrack : public FGraphTrack
 	INSIGHTS_DECLARE_RTTI(FTimingGraphTrack, FGraphTrack)
 
 public:
-	FTimingGraphTrack();
+	FTimingGraphTrack(TSharedPtr<STimingView> InTimingView);
 	virtual ~FTimingGraphTrack();
 
 	virtual void Update(const ITimingTrackUpdateContext& Context) override;
 
 	void AddDefaultFrameSeries();
 
+	TSharedPtr<FTimingGraphSeries> GetFrameSeries(ETraceFrameType FrameType);
+
 	TSharedPtr<FTimingGraphSeries> GetTimerSeries(uint32 TimerId);
 	TSharedPtr<FTimingGraphSeries> AddTimerSeries(uint32 TimerId, FLinearColor Color);
 	void RemoveTimerSeries(uint32 TimerId);
@@ -98,10 +109,20 @@ protected:
 
 	virtual void DrawVerticalAxisGrid(const ITimingTrackDrawContext& Context) const override;
 
+	void LoadDefaultSettings();
+
+private:
+	virtual void ContextMenu_ToggleOption_Execute(EGraphOptions Option);
+
 private:
 	FDelegateHandle OnTrackVisibilityChangedHandle;
 	FDelegateHandle OnTrackAddedHandle;
 	FDelegateHandle OnTrackRemovedHandle;
+
+	FDelegateHandle GameFrameSeriesVisibilityHandle;
+	FDelegateHandle RenderingFrameSeriesVisibilityHandle;
+
+	TWeakPtr<STimingView> TimingView;
 };
 
 ////////////////////////////////////////////////////////////////////////////////////////////////////
diff --git a/Engine/Source/Developer/TraceInsights/Private/Insights/Widgets/STimingView.cpp b/Engine/Source/Developer/TraceInsights/Private/Insights/Widgets/STimingView.cpp
--- a/Engine/Source/Developer/TraceInsights/Private/Insights/Widgets/STimingView.cpp
+++ b/Engine/Source/Developer/TraceInsights/Private/Insights/Widgets/STimingView.cpp
@@ -105,7 +105,6 @@ STimingView::STimingView()
 	, TimeRulerTrack(MakeShared<FTimeRulerTrack>())
 	, DefaultTimeMarker(MakeShared<Insights::FTimeMarker>())
 	, MarkersTrack(MakeShared<FMarkersTimingTrack>())
-	, GraphTrack(MakeShared<FTimingGraphTrack>())
 	, bAllowPanningOnScreenEdges(false)
 	, DPIScaleFactor(1.0f)
 	, EdgeFrameCountX(0)
@@ -117,7 +116,6 @@ STimingView::STimingView()
 	DefaultTimeMarker->SetName(TEXT(""));
 	DefaultTimeMarker->SetColor(FLinearColor(0.85f, 0.5f, 0.03f, 0.5f));
 
-	GraphTrack->SetName(TEXT("Main Graph"));
 
 	IModularFeatures::Get().RegisterModularFeature(Insights::TimingViewExtenderFeatureName, FrameSharedState.Get());
 	IModularFeatures::Get().RegisterModularFeature(Insights::TimingViewExtenderFeatureName, ThreadTimingSharedState.Get());
@@ -159,6 +157,9 @@ void STimingView::Construct(const FArguments& InArgs, FName InViewName)
 {
 	ViewName = InViewName;
 
+	GraphTrack = MakeShared<FTimingGraphTrack>(SharedThis(this));
+	GraphTrack->SetName(TEXT("Main Graph"));
+
 	FSlimHorizontalToolBarBuilder LeftToolbar(nullptr, FMultiBoxCustomization::None);
 	LeftToolbar.SetStyle(&FInsightsStyle::Get(), "SecondaryToolbar");
 
diff --git a/Engine/Source/Developer/TraceInsights/Public/Insights/ViewModels/GraphSeries.h b/Engine/Source/Developer/TraceInsights/Public/Insights/ViewModels/GraphSeries.h
--- a/Engine/Source/Developer/TraceInsights/Public/Insights/ViewModels/GraphSeries.h
+++ b/Engine/Source/Developer/TraceInsights/Public/Insights/ViewModels/GraphSeries.h
@@ -92,7 +92,7 @@ public:
 	void SetDescription(const FText& InDescription) { Description = InDescription; }
 
 	bool IsVisible() const { return bIsVisible; }
-	void SetVisibility(bool bOnOff) { bIsVisible = bOnOff; }
+	virtual void SetVisibility(bool bOnOff) { bIsVisible = bOnOff; }
 
 	bool IsDirty() const { return bIsDirty; }
 	void SetDirtyFlag() { bIsDirty = true; }
diff --git a/Engine/Source/Developer/TraceInsights/Public/Insights/ViewModels/GraphTrack.h b/Engine/Source/Developer/TraceInsights/Public/Insights/ViewModels/GraphTrack.h
--- a/Engine/Source/Developer/TraceInsights/Public/Insights/ViewModels/GraphTrack.h
+++ b/Engine/Source/Developer/TraceInsights/Public/Insights/ViewModels/GraphTrack.h
@@ -123,10 +123,11 @@ protected:
 	// Get the Y value that is used to provide a clipping border between adjacent graph tracks.
 	virtual float GetBorderY() const { return 0.0f; }
 
-private:
 	bool ContextMenu_ToggleOption_CanExecute(EGraphOptions Option);
-	void ContextMenu_ToggleOption_Execute(EGraphOptions Option);
+	virtual void ContextMenu_ToggleOption_Execute(EGraphOptions Option);
 	bool ContextMenu_ToggleOption_IsChecked(EGraphOptions Option);
+
+private:
 	bool ContextMenu_ShowPointsWithBorder_CanExecute();
 	bool ContextMenu_UseEventDuration_CanExecute();
 
