[Insights] Add persistence to the Auto-Start Analysis settings in Session Browser.

diff --git a/Engine/Programs/UnrealInsights/Config/DefaultUnrealInsightsSessionBrowserSettings.ini b/Engine/Programs/UnrealInsights/Config/DefaultUnrealInsightsSessionBrowserSettings.ini
--- a/Engine/Programs/UnrealInsights/Config/DefaultUnrealInsightsSessionBrowserSettings.ini
+++ b/Engine/Programs/UnrealInsights/Config/DefaultUnrealInsightsSessionBrowserSettings.ini
@@ -1,3 +1,9 @@
 [Insights.SessionBrowser]
 ; Whether Insights should signal to an UE application to auto-connect and start tracing when Insights is running.
-AutoConnect=True
\ No newline at end of file
+AutoConnect=True
+; Whether Insights should auto start analysis when a new live trace in received.
+AutoStartAnalysis=False
+; If specified, auto-start analysis will be enabled only for live trace sessions with this specified Platform.
+AutoStartAnalysisPlatform=
+; If specified, auto-start analysis will be enabled only for live trace sessions with this specified App Name.
+AutoStartAnalysisAppName=
\ No newline at end of file
diff --git a/Engine/Source/Developer/TraceInsights/Private/Insights/InsightsSessionBrowserSettings.h b/Engine/Source/Developer/TraceInsights/Private/Insights/InsightsSessionBrowserSettings.h
--- a/Engine/Source/Developer/TraceInsights/Private/Insights/InsightsSessionBrowserSettings.h
+++ b/Engine/Source/Developer/TraceInsights/Private/Insights/InsightsSessionBrowserSettings.h
@@ -15,6 +15,7 @@ public:
 		: bIsEditing(false)
 		, bIsDefault(bInIsDefault)
 		, bAutoConnect(true)
+		, bAutoStartAnalysis(false)
 	{
 		if (!bIsDefault)
 		{
@@ -34,11 +35,17 @@ public:
 		}
 
 		GConfig->GetBool(TEXT("Insights.SessionBrowser"), TEXT("AutoConnect"), bAutoConnect, SettingsIni);
+		GConfig->GetBool(TEXT("Insights.SessionBrowser"), TEXT("AutoStartAnalysis"), bAutoStartAnalysis, SettingsIni);
+		GConfig->GetString(TEXT("Insights.SessionBrowser"), TEXT("AutoStartAnalysisPlatform"), AutoStartAnalysisPlatform, SettingsIni);
+		GConfig->GetString(TEXT("Insights.SessionBrowser"), TEXT("AutoStartAnalysisAppName"), AutoStartAnalysisAppName, SettingsIni);
 	}
 
 	void SaveToConfig()
 	{
 		GConfig->SetBool(TEXT("Insights.SessionBrowser"), TEXT("AutoConnect"), bAutoConnect, SettingsIni);
+		GConfig->SetBool(TEXT("Insights.SessionBrowser"), TEXT("AutoStartAnalysis"), bAutoStartAnalysis, SettingsIni);
+		GConfig->SetString(TEXT("Insights.SessionBrowser"), TEXT("AutoStartAnalysisPlatform"), *AutoStartAnalysisPlatform, SettingsIni);
+		GConfig->SetString(TEXT("Insights.SessionBrowser"), TEXT("AutoStartAnalysisAppName"), *AutoStartAnalysisAppName, SettingsIni);
 
 		GConfig->Flush(false, SettingsIni);
 	}
@@ -74,6 +81,18 @@ public:
 	void SetAutoConnect(bool bOnOff) { bAutoConnect = bOnOff; }
 	void SetAndSaveAutoConnect(bool bOnOff) { SET_AND_SAVE(bAutoConnect, bOnOff); }
 
+	bool IsAutoStartAnalysisEnabled() const { return bAutoStartAnalysis; }
+	void SetAutoStartAnalysis(bool bOnOff) { bAutoStartAnalysis = bOnOff; }
+	void SetAndSaveAutoStartAnalysis(bool bOnOff) { SET_AND_SAVE(bAutoStartAnalysis, bOnOff); }
+
+	FString GetAutoStartAnalysisPlatform() const { return AutoStartAnalysisPlatform; }
+	void SetAutoStartAnalysisPlatform(const FString& Value) { AutoStartAnalysisPlatform = Value; }
+	void SetAndSaveAutoStartAnalysisPlatform(const FString& Value) { SET_AND_SAVE(AutoStartAnalysisPlatform, Value); }
+
+	FString GetAutoStartAnalysisAppName() const { return AutoStartAnalysisAppName; }
+	void SetAutoStartAnalysisAppName(const FString& Value) { AutoStartAnalysisAppName = Value; }
+	void SetAndSaveAutoStartAnalysisAppName(const FString& Value) { SET_AND_SAVE(AutoStartAnalysisAppName, Value); }
+
 	#undef SET_AND_SAVE
 
 private:
@@ -94,4 +113,13 @@ private:
 
 	/** Whether Insights should signal to the Editor to auto connect and start tracing when Insights is running */
 	bool bAutoConnect;
+
+	/** Whether Insights should auto start analysis when a new live trace in received.*/
+	bool bAutoStartAnalysis;
+
+	/** If specified, auto-start analysis will be enabled only for live trace sessions with this specified Platform.*/
+	FString AutoStartAnalysisPlatform;
+
+	/** If specified, auto-start analysis will be enabled only for live trace sessions with this specified App Name.*/
+	FString AutoStartAnalysisAppName;
 };
diff --git a/Engine/Source/Developer/TraceInsights/Private/Insights/Widgets/SStartPageWindow.cpp b/Engine/Source/Developer/TraceInsights/Private/Insights/Widgets/SStartPageWindow.cpp
--- a/Engine/Source/Developer/TraceInsights/Private/Insights/Widgets/SStartPageWindow.cpp
+++ b/Engine/Source/Developer/TraceInsights/Private/Insights/Widgets/SStartPageWindow.cpp
@@ -923,7 +923,6 @@ STraceStoreWindow::STraceStoreWindow()
 	, FilterStatsText()
 	, SortColumn(TraceStoreColumns::Date)
 	, SortMode(EColumnSortMode::Ascending)
-	, bAutoStartAnalysisForLiveSessions(false)
 	, AutoStartedSessions()
 	, AutoStartPlatformFilter()
 	, AutoStartAppNameFilter()
@@ -1870,6 +1869,8 @@ TSharedRef<SWidget> STraceStoreWindow::ConstructAutoStartPanel()
 		.MaxDesiredWidth(200.0f)
 		[
 			SAssignNew(AutoStartPlatformFilter, SSearchBox)
+			.InitialText(FText::FromString(FInsightsManager::Get()->GetSessionBrowserSettings().GetAutoStartAnalysisPlatform()))
+			.OnTextCommitted(this, &STraceStoreWindow::AutoStartPlatformFilterBox_OnValueCommitted)
 			.HintText(LOCTEXT("AutoStartPlatformFilter_Hint", "Platform"))
 			.ToolTipText(LOCTEXT("AutoStartPlatformFilter_Tooltip", "Type here to specify the Platform filter.\nAuto-start analysis will be enabled only for live trace sessions with this specified Platform."))
 		]
@@ -1885,6 +1886,8 @@ TSharedRef<SWidget> STraceStoreWindow::ConstructAutoStartPanel()
 		.MaxDesiredWidth(200.0f)
 		[
 			SAssignNew(AutoStartAppNameFilter, SSearchBox)
+			.InitialText(FText::FromString(FInsightsManager::Get()->GetSessionBrowserSettings().GetAutoStartAnalysisAppName()))
+			.OnTextCommitted(this, &STraceStoreWindow::AutoStartAppNameFilterBox_OnValueCommitted)
 			.HintText(LOCTEXT("AutoStartAppNameFilter_Hint", "AppName"))
 			.ToolTipText(LOCTEXT("AutoStartAppNameFilter_Tooltip", "Type here to specify the AppName filter.\nAuto-start analysis will be enabled only for live trace sessions with this specified AppName."))
 		]
@@ -2398,14 +2401,16 @@ void STraceStoreWindow::UpdateTrace(FTraceViewModel& InOutTrace, const Insights:
 		InOutTrace.TargetType = InSourceTrace.TargetType;
 	}
 
+	const FInsightsSessionBrowserSettings& Settings = FInsightsManager::Get()->GetSessionBrowserSettings();
+
 	// Auto start analysis for a live trace session.
 	if (InOutTrace.bIsLive &&
 		InOutTrace.bIsMetadataUpdated &&
-		bAutoStartAnalysisForLiveSessions && // is auto start enabled?
+		Settings.IsAutoStartAnalysisEnabled() && // is auto start enabled?
 		!AutoStartedSessions.Contains(InOutTrace.TraceId)) // is not already auto-started?
 	{
-		const FString AutoStartPlatformFilterStr = AutoStartPlatformFilter->GetText().ToString();
-		const FString AutoStartAppNameFilterStr = AutoStartAppNameFilter->GetText().ToString();
+		const FString& AutoStartPlatformFilterStr = Settings.GetAutoStartAnalysisPlatform();
+		const FString& AutoStartAppNameFilterStr = Settings.GetAutoStartAnalysisAppName();
 
 		// matches filter?
 		if ((AutoStartPlatformFilterStr.IsEmpty() || FCString::Strcmp(*AutoStartPlatformFilterStr, *InOutTrace.Platform.ToString()) == 0) &&
@@ -2521,14 +2526,19 @@ void STraceStoreWindow::TraceList_OnMouseButtonDoubleClick(TSharedPtr<FTraceView
 
 ECheckBoxState STraceStoreWindow::AutoStart_IsChecked() const
 {
-	return bAutoStartAnalysisForLiveSessions ? ECheckBoxState::Checked : ECheckBoxState::Unchecked;
+	return FInsightsManager::Get()->GetSessionBrowserSettings().IsAutoStartAnalysisEnabled() ? ECheckBoxState::Checked : ECheckBoxState::Unchecked;
 }
 
 ////////////////////////////////////////////////////////////////////////////////////////////////////
 
 void STraceStoreWindow::AutoStart_OnCheckStateChanged(ECheckBoxState NewState)
 {
-	bAutoStartAnalysisForLiveSessions = (NewState == ECheckBoxState::Checked);
+	if (AutoStart_IsChecked() == NewState)
+	{
+		return;
+	}
+
+	FInsightsManager::Get()->GetSessionBrowserSettings().SetAndSaveAutoStartAnalysis(!FInsightsManager::Get()->GetSessionBrowserSettings().IsAutoStartAnalysisEnabled());
 }
 
 ////////////////////////////////////////////////////////////////////////////////////////////////////
@@ -3784,6 +3794,20 @@ void STraceStoreWindow::ShowFailMessage(FText& InMessage)
 
 ////////////////////////////////////////////////////////////////////////////////////////////////////
 
+void STraceStoreWindow::AutoStartPlatformFilterBox_OnValueCommitted(const FText& InText, ETextCommit::Type InCommitType)
+{
+	FInsightsManager::Get()->GetSessionBrowserSettings().SetAndSaveAutoStartAnalysisPlatform(InText.ToString());
+}
+
+////////////////////////////////////////////////////////////////////////////////////////////////////
+
+void STraceStoreWindow::AutoStartAppNameFilterBox_OnValueCommitted(const FText& InText, ETextCommit::Type InCommitType)
+{
+	FInsightsManager::Get()->GetSessionBrowserSettings().SetAndSaveAutoStartAnalysisAppName(InText.ToString());
+}
+
+////////////////////////////////////////////////////////////////////////////////////////////////////
+
 #undef LOCTEXT_NAMESPACE
 #define LOCTEXT_NAMESPACE "SConnectionWindow"
 
diff --git a/Engine/Source/Developer/TraceInsights/Private/Insights/Widgets/SStartPageWindow.h b/Engine/Source/Developer/TraceInsights/Private/Insights/Widgets/SStartPageWindow.h
--- a/Engine/Source/Developer/TraceInsights/Private/Insights/Widgets/SStartPageWindow.h
+++ b/Engine/Source/Developer/TraceInsights/Private/Insights/Widgets/SStartPageWindow.h
@@ -502,6 +502,9 @@ private:
 	void EnableAutoConnect();
 	void DisableAutoConnect();
 
+	void AutoStartPlatformFilterBox_OnValueCommitted(const FText& InText, ETextCommit::Type InCommitType);
+	void AutoStartAppNameFilterBox_OnValueCommitted(const FText& InText, ETextCommit::Type InCommitType);
+
 public:
 	/** Parameter that controls the visibility of the confirmation window in case the trace is deleted. */
 	bool bIsDeleteTraceConfirmWindowVisible = true;
@@ -567,7 +570,6 @@ private:
 	//////////////////////////////////////////////////
 	// Auto-start functionality
 
-	bool bAutoStartAnalysisForLiveSessions;
 	TArray<uint32> AutoStartedSessions; // tracks sessions that were auto started (in order to not start them again)
 
 	TSharedPtr<SSearchBox> AutoStartPlatformFilter;
