[Insights] Auto Start analysis now also works when the Trace Store is minimized.

diff --git a/Engine/Source/Developer/TraceInsights/Private/Insights/Widgets/SStartPageWindow.cpp b/Engine/Source/Developer/TraceInsights/Private/Insights/Widgets/SStartPageWindow.cpp
--- a/Engine/Source/Developer/TraceInsights/Private/Insights/Widgets/SStartPageWindow.cpp
+++ b/Engine/Source/Developer/TraceInsights/Private/Insights/Widgets/SStartPageWindow.cpp
@@ -936,6 +936,10 @@ STraceStoreWindow::STraceStoreWindow()
 
 STraceStoreWindow::~STraceStoreWindow()
 {
+	if (OnTickHandle.IsValid())
+	{
+		FTSTicker::GetCoreTicker().RemoveTicker(OnTickHandle);
+	}
 #if WITH_EDITOR
 	if (DurationActive > 0.0f && FEngineAnalytics::IsAvailable())
 	{
@@ -1070,6 +1074,10 @@ void STraceStoreWindow::Construct(const FArguments& InArgs)
 		.Expose(OverlaySettingsSlot)
 	];
 
+	// Register tick functions.
+	OnTick = FTickerDelegate::CreateSP(this, &STraceStoreWindow::CoreTick);
+	OnTickHandle = FTSTicker::GetCoreTicker().AddTicker(OnTick, 0.0f);
+
 	CreateFilters();
 
 	RefreshTraceList();
@@ -2601,6 +2609,12 @@ void STraceStoreWindow::DisableAutoConnect()
 ////////////////////////////////////////////////////////////////////////////////////////////////////
 
 void STraceStoreWindow::Tick(const FGeometry& AllottedGeometry, const double InCurrentTime, const float InDeltaTime)
+{
+}
+
+////////////////////////////////////////////////////////////////////////////////////////////////////
+
+bool STraceStoreWindow::CoreTick(float DeltaTime)
 {
 	// We need to update the trace list, but not too often.
 	static uint64 NextTimestamp = 0;
@@ -2624,7 +2638,8 @@ void STraceStoreWindow::Tick(const FGeometry& AllottedGeometry, const double InC
 		FSlateApplication::Get().SetKeyboardFocus(TraceListView);
 	}
 
-	TickSplashScreenOverlay(InDeltaTime);
+	TickSplashScreenOverlay(DeltaTime);
+	return true;
 }
 
 ////////////////////////////////////////////////////////////////////////////////////////////////////
diff --git a/Engine/Source/Developer/TraceInsights/Private/Insights/Widgets/SStartPageWindow.h b/Engine/Source/Developer/TraceInsights/Private/Insights/Widgets/SStartPageWindow.h
--- a/Engine/Source/Developer/TraceInsights/Private/Insights/Widgets/SStartPageWindow.h
+++ b/Engine/Source/Developer/TraceInsights/Private/Insights/Widgets/SStartPageWindow.h
@@ -3,6 +3,7 @@
 #pragma once
 
 #include "Async/TaskGraphInterfaces.h"
+#include "Containers/Ticker.h"
 #include "Misc/FilterCollection.h"
 #include "Misc/TextFilter.h"
 #include "Styling/SlateTypes.h"
@@ -462,6 +463,9 @@ private:
 	 */
 	virtual FReply OnDragOver(const FGeometry& MyGeometry, const FDragDropEvent& DragDropEvent)  override;
 
+	/** Updates this class, done through FCoreTicker. Updates also when the page is not visible, unlike the Tick() function */
+	bool CoreTick(float DeltaTime);
+
 	//////////////////////////////////////////////////
 	// Filtering
 
@@ -588,6 +592,9 @@ private:
 
 	bool bSetKeyboardFocusOnNextTick = false;
 
+	FTickerDelegate OnTick;
+	FTSTicker::FDelegateHandle OnTickHandle;
+
 #if PLATFORM_WINDOWS
 	void* AutoConnectEvent = nullptr;
 #endif
