[MemoryInsights] Add free memory event callstack groupings/columns.

#rnx
#tests win64

diff --git a/Engine/Source/Developer/TraceInsights/Private/Insights/MemoryProfiler/ViewModels/MemAllocGroupingByCallstack.cpp b/Engine/Source/Developer/TraceInsights/Private/Insights/MemoryProfiler/ViewModels/MemAllocGroupingByCallstack.cpp
--- a/Engine/Source/Developer/TraceInsights/Private/Insights/MemoryProfiler/ViewModels/MemAllocGroupingByCallstack.cpp
+++ b/Engine/Source/Developer/TraceInsights/Private/Insights/MemoryProfiler/ViewModels/MemAllocGroupingByCallstack.cpp
@@ -55,15 +55,22 @@ INSIGHTS_IMPLEMENT_RTTI(FMemAllocGroupingByCallstack)
 
 ////////////////////////////////////////////////////////////////////////////////////////////////////
 
-FMemAllocGroupingByCallstack::FMemAllocGroupingByCallstack(bool bInIsInverted, bool bInIsGroupingByFunction)
+FMemAllocGroupingByCallstack::FMemAllocGroupingByCallstack(bool bInIsAllocCallstack, bool bInIsInverted, bool bInIsGroupingByFunction)
 	: FTreeNodeGrouping(
-		bInIsInverted ? LOCTEXT("Grouping_ByCallstack2_ShortName", "Inverted Callstack")
-					  : LOCTEXT("Grouping_ByCallstack1_ShortName", "Callstack"),
-		bInIsInverted ? LOCTEXT("Grouping_ByCallstack2_TitleName", "By Inverted Callstack")
-					  : LOCTEXT("Grouping_ByCallstack1_TitleName", "By Callstack"),
+		bInIsAllocCallstack
+		? (bInIsInverted ? LOCTEXT("Grouping_ByCallstack2_ShortName", "Inverted Alloc Callstack")
+						: LOCTEXT("Grouping_ByCallstack1_ShortName", "Alloc Callstack"))
+		: (bInIsInverted ? LOCTEXT("Grouping_ByCallstack4_ShortName", "Inverted Free Callstack")
+						: LOCTEXT("Grouping_ByCallstack3_ShortName", "Free Callstack")),
+		bInIsAllocCallstack
+		? (bInIsInverted ? LOCTEXT("Grouping_ByCallstack2_TitleName", "By Inverted Alloc Callstack")
+					  : LOCTEXT("Grouping_ByCallstack1_TitleName", "By Alloc Callstack"))
+		: (bInIsInverted ? LOCTEXT("Grouping_ByCallstack4_TitleName", "By Inverted Free Callstack")
+					  : LOCTEXT("Grouping_ByCallstack3_TitleName", "By Free Callstack")),
 		LOCTEXT("Grouping_Callstack_Desc", "Creates a tree based on callstack of each allocation."),
 		TEXT("Icons.Group.TreeItem"),
 		nullptr)
+	, bIsAllocCallstack(bInIsAllocCallstack)
 	, bIsInverted(bInIsInverted)
 	, bIsGroupingByFunction(bInIsGroupingByFunction)
 {
@@ -111,7 +118,7 @@ void FMemAllocGroupingByCallstack::GroupNodes(const TArray<FTableTreeNodePtr>& N
 		const FMemoryAlloc* Alloc = MemAllocNode.GetMemAlloc();
 		if (Alloc)
 		{
-			const TraceServices::FCallstack* Callstack = Alloc->GetCallstack();
+			const TraceServices::FCallstack* Callstack = bIsAllocCallstack ? Alloc->GetCallstack() : Alloc->GetFreeCallstack();
 
 			FCallstackGroup** FoundGroupPtrPtr = GroupMapByCallstack.Find(Callstack);
 			if (FoundGroupPtrPtr)
diff --git a/Engine/Source/Developer/TraceInsights/Private/Insights/MemoryProfiler/ViewModels/MemAllocGroupingByCallstack.h b/Engine/Source/Developer/TraceInsights/Private/Insights/MemoryProfiler/ViewModels/MemAllocGroupingByCallstack.h
--- a/Engine/Source/Developer/TraceInsights/Private/Insights/MemoryProfiler/ViewModels/MemAllocGroupingByCallstack.h
+++ b/Engine/Source/Developer/TraceInsights/Private/Insights/MemoryProfiler/ViewModels/MemAllocGroupingByCallstack.h
@@ -70,12 +70,13 @@ private:
 	};
 
 public:
-	FMemAllocGroupingByCallstack(bool bInIsInverted, bool bInIsGroupingByFunction);
+	FMemAllocGroupingByCallstack(bool bInIsAllocCallstack, bool bInIsInverted, bool bInIsGroupingByFunction);
 	virtual ~FMemAllocGroupingByCallstack();
 
 	virtual void GroupNodes(const TArray<FTableTreeNodePtr>& Nodes, FTableTreeNode& ParentGroup, TWeakPtr<FTable> InParentTable, IAsyncOperationProgress& InAsyncOperationProgress) const override;
 
 	bool IsInverted() const { return bIsInverted; }
+	bool IsAllocCallstack() const { return bIsAllocCallstack; }
 
 	bool IsGroupingByFunction() const { return bIsGroupingByFunction; }
 	void SetGroupingByFunction(bool bOnOff) { bIsGroupingByFunction = bOnOff; }
@@ -89,6 +90,7 @@ private:
 	FTableTreeNode* CreateEmptyCallstackGroup(TWeakPtr<FTable> ParentTable, FTableTreeNode& Parent) const;
 
 private:
+	bool bIsAllocCallstack;
 	bool bIsInverted;
 	bool bIsGroupingByFunction;
 };
diff --git a/Engine/Source/Developer/TraceInsights/Private/Insights/MemoryProfiler/ViewModels/MemAllocNode.cpp b/Engine/Source/Developer/TraceInsights/Private/Insights/MemoryProfiler/ViewModels/MemAllocNode.cpp
--- a/Engine/Source/Developer/TraceInsights/Private/Insights/MemoryProfiler/ViewModels/MemAllocNode.cpp
+++ b/Engine/Source/Developer/TraceInsights/Private/Insights/MemoryProfiler/ViewModels/MemAllocNode.cpp
@@ -28,49 +28,49 @@ uint32 FMemAllocNode::GetFreeCallstackId() const
 
 ////////////////////////////////////////////////////////////////////////////////////////////////////
 
-FText FMemAllocNode::GetFullCallstack() const
+FText FMemAllocNode::GetFullCallstack(ECallstackType InCallstackType) const
 {
-	return IsValidMemAlloc() ? GetMemAllocChecked().GetFullCallstack() : FText();
+	return GetFullCallstackOrSourceFiles(InCallstackType, (uint8)EStackFrameFormatFlags::ModuleSymbolFileAndLine);
 }
 
 ////////////////////////////////////////////////////////////////////////////////////////////////////
 
-FText FMemAllocNode::GetFullCallstackSourceFiles() const
+FText FMemAllocNode::GetFullCallstackSourceFiles(ECallstackType InCallstackType) const
 {
-	return IsValidMemAlloc() ? GetMemAllocChecked().GetFullCallstackSourceFiles() : FText();
+	return GetFullCallstackOrSourceFiles(InCallstackType, (uint8)EStackFrameFormatFlags::File);
 }
 
 ////////////////////////////////////////////////////////////////////////////////////////////////////
 
-FText FMemAllocNode::GetTopFunction() const
+FText FMemAllocNode::GetTopFunction(ECallstackType InCallstackType) const
 {
-	return GetTopFunctionOrSourceFile((uint8)EStackFrameFormatFlags::ModuleAndSymbol);
+	return GetTopFunctionOrSourceFile(InCallstackType, (uint8)EStackFrameFormatFlags::ModuleAndSymbol);
 }
 
 ////////////////////////////////////////////////////////////////////////////////////////////////////
 
-FText FMemAllocNode::GetTopFunctionEx() const
+FText FMemAllocNode::GetTopFunctionEx(ECallstackType InCallstackType) const
 {
-	return GetTopFunctionOrSourceFile((uint8)EStackFrameFormatFlags::ModuleSymbolFileAndLine);
+	return GetTopFunctionOrSourceFile(InCallstackType, (uint8)EStackFrameFormatFlags::ModuleSymbolFileAndLine);
 }
 
 ////////////////////////////////////////////////////////////////////////////////////////////////////
 
-FText FMemAllocNode::GetTopSourceFile() const
+FText FMemAllocNode::GetTopSourceFile(ECallstackType InCallstackType) const
 {
-	return GetTopFunctionOrSourceFile((uint8)EStackFrameFormatFlags::File);
+	return GetTopFunctionOrSourceFile(InCallstackType, (uint8)EStackFrameFormatFlags::File);
 }
 
 ////////////////////////////////////////////////////////////////////////////////////////////////////
 
-FText FMemAllocNode::GetTopSourceFileEx() const
+FText FMemAllocNode::GetTopSourceFileEx(ECallstackType InCallstackType) const
 {
-	return GetTopFunctionOrSourceFile((uint8)EStackFrameFormatFlags::FileAndLine);
+	return GetTopFunctionOrSourceFile(InCallstackType, (uint8)EStackFrameFormatFlags::FileAndLine);
 }
 
 ////////////////////////////////////////////////////////////////////////////////////////////////////
 
-FText FMemAllocNode::GetTopFunctionOrSourceFile(uint8 Flags) const
+FText FMemAllocNode::GetFullCallstackOrSourceFiles(ECallstackType InCallstackType, uint8 Flags) const
 {
 	if (!IsValidMemAlloc())
 	{
@@ -78,7 +78,51 @@ FText FMemAllocNode::GetTopFunctionOrSourceFile(uint8 Flags) const
 	}
 
 	const Insights::FMemoryAlloc& Alloc = GetMemAllocChecked();
-	const TraceServices::FCallstack* Callstack = Alloc.GetCallstack();
+	const TraceServices::FCallstack* Callstack = 
+		InCallstackType == ECallstackType::AllocCallstack
+			? Alloc.GetCallstack() 
+			: Alloc.GetFreeCallstack();
+
+	if (!Callstack)
+	{
+		return FText::FromString(GetCallstackNotAvailableString());
+	}
+
+	if (Callstack->Num() == 0)
+	{
+		return FText::FromString(GetEmptyCallstackString());
+	}
+
+	TStringBuilder<1024> Tooltip;
+	const uint32 NumCallstackFrames = Callstack->Num();
+	check(NumCallstackFrames <= 256); // see Callstack->Frame(uint8)
+	for (uint32 FrameIndex = 0; FrameIndex < NumCallstackFrames; ++FrameIndex)
+	{
+		if (FrameIndex != 0)
+		{
+			Tooltip << TEXT("\n");
+		}
+		const TraceServices::FStackFrame* Frame = Callstack->Frame(static_cast<uint8>(FrameIndex));
+		check(Frame != nullptr);
+		FormatStackFrame(*Frame, Tooltip, (EStackFrameFormatFlags)Flags);
+	}
+	return FText::FromString(FString(Tooltip));
+}
+
+////////////////////////////////////////////////////////////////////////////////////////////////////
+
+FText FMemAllocNode::GetTopFunctionOrSourceFile(ECallstackType InCallstackType, uint8 Flags) const
+{
+	if (!IsValidMemAlloc())
+	{
+		return FText();
+	}
+
+	const Insights::FMemoryAlloc& Alloc = GetMemAllocChecked();
+	const TraceServices::FCallstack* Callstack = 
+		InCallstackType == ECallstackType::AllocCallstack
+			? Alloc.GetCallstack() 
+			: Alloc.GetFreeCallstack();
 
 	if (!Callstack)
 	{
diff --git a/Engine/Source/Developer/TraceInsights/Private/Insights/MemoryProfiler/ViewModels/MemAllocNode.h b/Engine/Source/Developer/TraceInsights/Private/Insights/MemoryProfiler/ViewModels/MemAllocNode.h
--- a/Engine/Source/Developer/TraceInsights/Private/Insights/MemoryProfiler/ViewModels/MemAllocNode.h
+++ b/Engine/Source/Developer/TraceInsights/Private/Insights/MemoryProfiler/ViewModels/MemAllocNode.h
@@ -41,6 +41,12 @@ class FMemAllocNode : public FTableTreeNode
 	INSIGHTS_DECLARE_RTTI(FMemAllocNode, FTableTreeNode)
 
 public:
+	enum ECallstackType : uint8
+	{
+		AllocCallstack,
+		FreeCallstack
+	};
+
 	/** Initialization constructor for the MemAlloc node. */
 	explicit FMemAllocNode(const FName InName, TWeakPtr<FMemAllocTable> InParentTable, int32 InRowIndex)
 		: FTableTreeNode(InName, InParentTable, InRowIndex)
@@ -64,15 +70,16 @@ public:
 	uint32 GetCallstackId() const;
 	uint32 GetFreeCallstackId() const;
 
-	FText GetFullCallstack() const;
-	FText GetFullCallstackSourceFiles() const;
-	FText GetTopFunction() const;
-	FText GetTopFunctionEx() const;
-	FText GetTopSourceFile() const;
-	FText GetTopSourceFileEx() const;
+	FText GetFullCallstack(ECallstackType InCallstackType) const;
+	FText GetFullCallstackSourceFiles(ECallstackType InCallstackType) const;
+	FText GetTopFunction(ECallstackType InCallstackType) const;
+	FText GetTopFunctionEx(ECallstackType InCallstackType) const;
+	FText GetTopSourceFile(ECallstackType InCallstackType) const;
+	FText GetTopSourceFileEx(ECallstackType InCallstackType) const;
 
 private:
-	FText GetTopFunctionOrSourceFile(uint8 Flags) const;
+	FText GetFullCallstackOrSourceFiles(ECallstackType InCallstackType, uint8 Flags) const;
+	FText GetTopFunctionOrSourceFile(ECallstackType InCallstackType, uint8 Flags) const;
 
 private:
 	FMemAllocTable* MemAllocTable;
diff --git a/Engine/Source/Developer/TraceInsights/Private/Insights/MemoryProfiler/ViewModels/MemAllocTable.cpp b/Engine/Source/Developer/TraceInsights/Private/Insights/MemoryProfiler/ViewModels/MemAllocTable.cpp
--- a/Engine/Source/Developer/TraceInsights/Private/Insights/MemoryProfiler/ViewModels/MemAllocTable.cpp
+++ b/Engine/Source/Developer/TraceInsights/Private/Insights/MemoryProfiler/ViewModels/MemAllocTable.cpp
@@ -40,9 +40,12 @@ const FName FMemAllocTableColumns::TagColumnId(TEXT("Tag"));
 const FName FMemAllocTableColumns::AssetColumnId(TEXT("Asset"));
 const FName FMemAllocTableColumns::PackageColumnId(TEXT("Package"));
 const FName FMemAllocTableColumns::ClassNameColumnId(TEXT("ClassName"));
-const FName FMemAllocTableColumns::FunctionColumnId(TEXT("Function"));
-const FName FMemAllocTableColumns::SourceFileColumnId(TEXT("SourceFile"));
-const FName FMemAllocTableColumns::CallstackSizeColumnId(TEXT("CallstackSize"));
+const FName FMemAllocTableColumns::AllocFunctionColumnId(TEXT("AllocFunction"));
+const FName FMemAllocTableColumns::AllocSourceFileColumnId(TEXT("AllocSourceFile"));
+const FName FMemAllocTableColumns::AllocCallstackSizeColumnId(TEXT("AllocCallstackSize"));
+const FName FMemAllocTableColumns::FreeFunctionColumnId(TEXT("FreeFunction"));
+const FName FMemAllocTableColumns::FreeSourceFileColumnId(TEXT("FreeSourceFile"));
+const FName FMemAllocTableColumns::FreeCallstackSizeColumnId(TEXT("FreeCallstackSize"));
 
 ////////////////////////////////////////////////////////////////////////////////////////////////////
 // FMemAllocTable
@@ -914,14 +917,14 @@ void FMemAllocTable::AddDefaultColumns()
 	//////////////////////////////////////////////////
 	// Top Function Column
 	{
-		TSharedRef<FTableColumn> ColumnRef = MakeShared<FTableColumn>(FMemAllocTableColumns::FunctionColumnId);
+		TSharedRef<FTableColumn> ColumnRef = MakeShared<FTableColumn>(FMemAllocTableColumns::AllocFunctionColumnId);
 		FTableColumn& Column = *ColumnRef;
 
 		Column.SetIndex(ColumnIndex++);
 
-		Column.SetShortName(LOCTEXT("FunctionColumnName", "Function"));
-		Column.SetTitleName(LOCTEXT("FunctionColumnTitle", "Top Function"));
-		Column.SetDescription(LOCTEXT("FunctionColumnDesc", "Resolved top function from the callstack of allocation"));
+		Column.SetShortName(LOCTEXT("FunctionColumnName_Alloc", "Function (Alloc)"));
+		Column.SetTitleName(LOCTEXT("FunctionColumnTitle_Alloc", "Top Function (Alloc Callstack)"));
+		Column.SetDescription(LOCTEXT("FunctionColumnDesc_Alloc", "Resolved top function from the callstack of allocation"));
 
 		Column.SetFlags(ETableColumnFlags::ShouldBeVisible | ETableColumnFlags::CanBeHidden | ETableColumnFlags::CanBeFiltered | ETableColumnFlags::IsDynamic);
 
@@ -947,7 +950,7 @@ void FMemAllocTable::AddDefaultColumns()
 				{
 					const FMemAllocNode& MemAllocNode = static_cast<const FMemAllocNode&>(Node);
 					const uint64 CallstackId = MemAllocNode.GetCallstackId();
-					const FText CallstackText = MemAllocNode.GetTopFunction();
+					const FText CallstackText = MemAllocNode.GetTopFunction(FMemAllocNode::ECallstackType::AllocCallstack);
 					return FTableCellValue(CallstackText, CallstackId);
 				}
 
@@ -984,7 +987,7 @@ void FMemAllocTable::AddDefaultColumns()
 						.Padding(2.0f)
 						[
 							SNew(STextBlock)
-							.Text(&MemAllocNode, &FMemAllocNode::GetTopFunction)
+							.Text(&MemAllocNode, &FMemAllocNode::GetTopFunction, FMemAllocNode::ECallstackType::AllocCallstack)
 							.ColorAndOpacity(FSlateColor(EStyleColor::AccentBlue))
 						]
 						+ SVerticalBox::Slot()
@@ -992,7 +995,7 @@ void FMemAllocTable::AddDefaultColumns()
 						.Padding(2.0f, 6.0f, 2.0f, 2.0f)
 						[
 							SNew(STextBlock)
-							.Text(&MemAllocNode, &FMemAllocNode::GetFullCallstack)
+							.Text(&MemAllocNode, &FMemAllocNode::GetFullCallstack, FMemAllocNode::ECallstackType::AllocCallstack)
 						]
 					];
 			}
@@ -1000,7 +1003,7 @@ void FMemAllocTable::AddDefaultColumns()
 			virtual FText FormatValueForTooltip(const FTableColumn& Column, const FBaseTreeNode& Node) const override
 			{
 				const FMemAllocNode& MemAllocNode = static_cast<const FMemAllocNode&>(Node);
-				return MemAllocNode.GetTopFunction();
+				return MemAllocNode.GetTopFunction(FMemAllocNode::ECallstackType::AllocCallstack);
 			}
 		};
 
@@ -1015,14 +1018,14 @@ void FMemAllocTable::AddDefaultColumns()
 	//////////////////////////////////////////////////
 	// Source File Column
 	{
-		TSharedRef<FTableColumn> ColumnRef = MakeShared<FTableColumn>(FMemAllocTableColumns::SourceFileColumnId);
+		TSharedRef<FTableColumn> ColumnRef = MakeShared<FTableColumn>(FMemAllocTableColumns::AllocSourceFileColumnId);
 		FTableColumn& Column = *ColumnRef;
 
 		Column.SetIndex(ColumnIndex++);
 
-		Column.SetShortName(LOCTEXT("SourceFileColumnName", "Source"));
-		Column.SetTitleName(LOCTEXT("SourceFileColumnTitle", "Top Source File"));
-		Column.SetDescription(LOCTEXT("SourceFileColumnDesc", "Source file of the top function from the callstack of allocation"));
+		Column.SetShortName(LOCTEXT("SourceFileColumnName_Alloc", "Source (Alloc)"));
+		Column.SetTitleName(LOCTEXT("SourceFileColumnTitle_Alloc", "Top Source File (Alloc Callstack)"));
+		Column.SetDescription(LOCTEXT("SourceFileColumnDesc_Alloc", "Source file of the top function from the callstack of allocation"));
 
 		Column.SetFlags(ETableColumnFlags::CanBeHidden | ETableColumnFlags::CanBeFiltered | ETableColumnFlags::IsDynamic);
 
@@ -1048,7 +1051,7 @@ void FMemAllocTable::AddDefaultColumns()
 				{
 					const FMemAllocNode& MemAllocNode = static_cast<const FMemAllocNode&>(Node);
 					const uint64 CallstackId = MemAllocNode.GetCallstackId();
-					const FText CallstackText = MemAllocNode.GetTopSourceFile();
+					const FText CallstackText = MemAllocNode.GetTopSourceFile(FMemAllocNode::ECallstackType::AllocCallstack);
 					return FTableCellValue(CallstackText, CallstackId);
 				}
 
@@ -1084,7 +1087,7 @@ void FMemAllocTable::AddDefaultColumns()
 						.Padding(2.0f)
 						[
 							SNew(STextBlock)
-							.Text(&MemAllocNode, &FMemAllocNode::GetTopSourceFileEx)
+							.Text(&MemAllocNode, &FMemAllocNode::GetTopSourceFileEx, FMemAllocNode::ECallstackType::AllocCallstack)
 							.ColorAndOpacity(FSlateColor(EStyleColor::AccentBlue))
 						]
 						+ SVerticalBox::Slot()
@@ -1092,7 +1095,7 @@ void FMemAllocTable::AddDefaultColumns()
 						.Padding(2.0f, 6.0f, 2.0f, 2.0f)
 						[
 							SNew(STextBlock)
-							.Text(&MemAllocNode, &FMemAllocNode::GetFullCallstackSourceFiles)
+							.Text(&MemAllocNode, &FMemAllocNode::GetFullCallstackSourceFiles, FMemAllocNode::ECallstackType::AllocCallstack)
 						]
 					];
 			}
@@ -1100,7 +1103,7 @@ void FMemAllocTable::AddDefaultColumns()
 			virtual FText FormatValueForTooltip(const FTableColumn& Column, const FBaseTreeNode& Node) const override
 			{
 				const FMemAllocNode& MemAllocNode = static_cast<const FMemAllocNode&>(Node);
-				return MemAllocNode.GetTopSourceFileEx();
+				return MemAllocNode.GetTopSourceFileEx(FMemAllocNode::ECallstackType::AllocCallstack);
 			}
 		};
 
@@ -1115,14 +1118,14 @@ void FMemAllocTable::AddDefaultColumns()
 	//////////////////////////////////////////////////
 	// Callstack (Size) Column
 	{
-		TSharedRef<FTableColumn> ColumnRef = MakeShared<FTableColumn>(FMemAllocTableColumns::CallstackSizeColumnId);
+		TSharedRef<FTableColumn> ColumnRef = MakeShared<FTableColumn>(FMemAllocTableColumns::AllocCallstackSizeColumnId);
 		FTableColumn& Column = *ColumnRef;
 
 		Column.SetIndex(ColumnIndex++);
 
-		Column.SetShortName(LOCTEXT("CallstackSizeColumnName", "Callstack"));
-		Column.SetTitleName(LOCTEXT("CallstackSizeColumnTitle", "Callstack"));
-		Column.SetDescription(LOCTEXT("CallstackSizeColumnDesc", "Number of callstack frames.\nTooltip shows the entire callstack."));
+		Column.SetShortName(LOCTEXT("CallstackSizeColumnName_Alloc", "Alloc Callstack"));
+		Column.SetTitleName(LOCTEXT("CallstackSizeColumnTitle_Alloc", "Alloc Callstack"));
+		Column.SetDescription(LOCTEXT("CallstackSizeColumnDesc_Alloc", "Number of alloc callstack frames.\nTooltip shows the entire alloc callstack."));
 
 		Column.SetFlags(ETableColumnFlags::ShouldBeVisible | ETableColumnFlags::CanBeHidden | ETableColumnFlags::CanBeFiltered);
 
@@ -1205,7 +1208,7 @@ void FMemAllocTable::AddDefaultColumns()
 						.Padding(2.0f)
 						[
 							SNew(STextBlock)
-							.Text(&MemAllocNode, &FMemAllocNode::GetFullCallstack)
+							.Text(&MemAllocNode, &FMemAllocNode::GetFullCallstack, FMemAllocNode::ECallstackType::AllocCallstack)
 						]
 					];
 			}
@@ -1213,7 +1216,326 @@ void FMemAllocTable::AddDefaultColumns()
 			virtual FText FormatValueForTooltip(const FTableColumn& Column, const FBaseTreeNode& Node) const override
 			{
 				const FMemAllocNode& MemAllocNode = static_cast<const FMemAllocNode&>(Node);
-				return MemAllocNode.GetFullCallstack();
+				return MemAllocNode.GetFullCallstack(FMemAllocNode::ECallstackType::AllocCallstack);
+			}
+
+			virtual FText FormatValueForGrouping(const FTableColumn& Column, const FBaseTreeNode& Node) const override
+			{
+				return FormatValue(Column.GetValue(Node));
+			}
+		};
+
+		TSharedRef<ITableCellValueFormatter> Formatter = MakeShared<FCallstackValueFormatter>();
+		Column.SetValueFormatter(Formatter);
+
+		TSharedRef<ITableCellValueSorter> Sorter = MakeShared<FSorterByInt64Value>(ColumnRef);
+		Column.SetValueSorter(Sorter);
+
+		AddColumn(ColumnRef);
+	}
+	//////////////////////////////////////////////////
+	// Free Top Function Column
+	{
+		TSharedRef<FTableColumn> ColumnRef = MakeShared<FTableColumn>(FMemAllocTableColumns::FreeFunctionColumnId);
+		FTableColumn& Column = *ColumnRef;
+
+		Column.SetIndex(ColumnIndex++);
+
+		Column.SetShortName(LOCTEXT("FunctionColumnName_Free", "Function (Free)"));
+		Column.SetTitleName(LOCTEXT("FunctionColumnTitle_Free", "Top Function (Free Callstack)"));
+		Column.SetDescription(LOCTEXT("FunctionColumnDesc_Free", "Resolved top function from the callstack of freeing an allocation"));
+
+		Column.SetFlags(ETableColumnFlags::ShouldBeVisible | ETableColumnFlags::CanBeHidden | ETableColumnFlags::CanBeFiltered | ETableColumnFlags::IsDynamic);
+
+		Column.SetHorizontalAlignment(HAlign_Left);
+		Column.SetInitialWidth(550.0f);
+
+		Column.SetDataType(ETableCellDataType::Text);
+
+		class FFunctionValueGetter : public FTableCellValueGetter
+		{
+		public:
+			virtual const TOptional<FTableCellValue> GetValue(const FTableColumn& Column, const FBaseTreeNode& Node) const override
+			{
+				if (Node.IsGroup())
+				{
+					const FTableTreeNode& NodePtr = static_cast<const FTableTreeNode&>(Node);
+					if (NodePtr.HasAggregatedValue(Column.GetId()))
+					{
+						return NodePtr.GetAggregatedValue(Column.GetId());
+					}
+				}
+				else //if (Node->Is<FMemAllocNode>())
+				{
+					const FMemAllocNode& MemAllocNode = static_cast<const FMemAllocNode&>(Node);
+					const uint64 CallstackId = MemAllocNode.GetCallstackId();
+					const FText CallstackText = MemAllocNode.GetTopFunction(FMemAllocNode::ECallstackType::FreeCallstack);
+					return FTableCellValue(CallstackText, CallstackId);
+				}
+
+				return TOptional<FTableCellValue>();
+			}
+
+			virtual uint64 GetValueId(const FTableColumn& Column, const FBaseTreeNode& Node) const override
+			{
+				if (!Node.IsGroup()) //if (Node->Is<FMemAllocNode>())
+				{
+					const FMemAllocNode& MemAllocNode = static_cast<const FMemAllocNode&>(Node);
+					return MemAllocNode.GetCallstackId();
+				}
+
+				return 0;
+			}
+		};
+		TSharedRef<ITableCellValueGetter> Getter = MakeShared<FFunctionValueGetter>();
+		Column.SetValueGetter(Getter);
+
+		class FFunctionValueFormatter : public FTextValueFormatter
+		{
+		public:
+			virtual TSharedPtr<IToolTip> GetCustomTooltip(const FTableColumn& Column, const FBaseTreeNode& Node) const override
+			{
+				const FMemAllocNode& MemAllocNode = static_cast<const FMemAllocNode&>(Node);
+
+				return SNew(SToolTip)
+					.Visibility(TAttribute<EVisibility>::Create(TAttribute<EVisibility>::FGetter::CreateStatic(&FTableCellValueFormatter::GetTooltipVisibility)))
+					[
+						SNew(SVerticalBox)
+						+ SVerticalBox::Slot()
+						.AutoHeight()
+						.Padding(2.0f)
+						[
+							SNew(STextBlock)
+							.Text(&MemAllocNode, &FMemAllocNode::GetTopFunction, FMemAllocNode::ECallstackType::FreeCallstack)
+							.ColorAndOpacity(FSlateColor(EStyleColor::AccentBlue))
+						]
+						+ SVerticalBox::Slot()
+						.AutoHeight()
+						.Padding(2.0f, 6.0f, 2.0f, 2.0f)
+						[
+							SNew(STextBlock)
+							.Text(&MemAllocNode, &FMemAllocNode::GetFullCallstack, FMemAllocNode::ECallstackType::FreeCallstack)
+						]
+					];
+			}
+
+			virtual FText FormatValueForTooltip(const FTableColumn& Column, const FBaseTreeNode& Node) const override
+			{
+				const FMemAllocNode& MemAllocNode = static_cast<const FMemAllocNode&>(Node);
+				return MemAllocNode.GetTopFunction(FMemAllocNode::ECallstackType::FreeCallstack);
+			}
+		};
+
+		TSharedRef<ITableCellValueFormatter> Formatter = MakeShared<FFunctionValueFormatter>();
+		Column.SetValueFormatter(Formatter);
+
+		TSharedRef<ITableCellValueSorter> Sorter = MakeShared<FSorterByTextValueWithId>(ColumnRef);
+		Column.SetValueSorter(Sorter);
+
+		AddColumn(ColumnRef);
+	}
+	//////////////////////////////////////////////////
+	// Free Source File Column
+	{
+		TSharedRef<FTableColumn> ColumnRef = MakeShared<FTableColumn>(FMemAllocTableColumns::FreeSourceFileColumnId);
+		FTableColumn& Column = *ColumnRef;
+
+		Column.SetIndex(ColumnIndex++);
+
+		Column.SetShortName(LOCTEXT("SourceFileColumnName_Free", "Source (Free)"));
+		Column.SetTitleName(LOCTEXT("SourceFileColumnTitle_Free", "Top Source File (Free Callstack)"));
+		Column.SetDescription(LOCTEXT("SourceFileColumnDesc_Free", "Source file of the top function from the callstack of freeing an allocation"));
+
+		Column.SetFlags(ETableColumnFlags::CanBeHidden | ETableColumnFlags::CanBeFiltered | ETableColumnFlags::IsDynamic);
+
+		Column.SetHorizontalAlignment(HAlign_Left);
+		Column.SetInitialWidth(550.0f);
+
+		Column.SetDataType(ETableCellDataType::Text);
+
+		class FSourceValueGetter : public FTableCellValueGetter
+		{
+		public:
+			virtual const TOptional<FTableCellValue> GetValue(const FTableColumn& Column, const FBaseTreeNode& Node) const override
+			{
+				if (Node.IsGroup())
+				{
+					const FTableTreeNode& NodePtr = static_cast<const FTableTreeNode&>(Node);
+					if (NodePtr.HasAggregatedValue(Column.GetId()))
+					{
+						return NodePtr.GetAggregatedValue(Column.GetId());
+					}
+				}
+				else //if (Node->Is<FMemAllocNode>())
+				{
+					const FMemAllocNode& MemAllocNode = static_cast<const FMemAllocNode&>(Node);
+					const uint64 CallstackId = MemAllocNode.GetCallstackId();
+					const FText CallstackText = MemAllocNode.GetTopSourceFile(FMemAllocNode::ECallstackType::FreeCallstack);
+					return FTableCellValue(CallstackText, CallstackId);
+				}
+
+				return TOptional<FTableCellValue>();
+			}
+
+			virtual uint64 GetValueId(const FTableColumn& Column, const FBaseTreeNode& Node) const override
+			{
+				if (!Node.IsGroup()) //if (Node->Is<FMemAllocNode>())
+				{
+					const FMemAllocNode& MemAllocNode = static_cast<const FMemAllocNode&>(Node);
+					return MemAllocNode.GetCallstackId();
+				}
+				return 0;
+			}
+		};
+		TSharedRef<ITableCellValueGetter> Getter = MakeShared<FSourceValueGetter>();
+		Column.SetValueGetter(Getter);
+
+		class FSourceValueFormatter : public FTextValueFormatter
+		{
+		public:
+			virtual TSharedPtr<IToolTip> GetCustomTooltip(const FTableColumn& Column, const FBaseTreeNode& Node) const override
+			{
+				const FMemAllocNode& MemAllocNode = static_cast<const FMemAllocNode&>(Node);
+
+				return SNew(SToolTip)
+					.Visibility(TAttribute<EVisibility>::Create(TAttribute<EVisibility>::FGetter::CreateStatic(&FTableCellValueFormatter::GetTooltipVisibility)))
+					[
+						SNew(SVerticalBox)
+						+ SVerticalBox::Slot()
+						.AutoHeight()
+						.Padding(2.0f)
+						[
+							SNew(STextBlock)
+							.Text(&MemAllocNode, &FMemAllocNode::GetTopSourceFileEx, FMemAllocNode::ECallstackType::FreeCallstack)
+							.ColorAndOpacity(FSlateColor(EStyleColor::AccentBlue))
+						]
+						+ SVerticalBox::Slot()
+						.AutoHeight()
+						.Padding(2.0f, 6.0f, 2.0f, 2.0f)
+						[
+							SNew(STextBlock)
+							.Text(&MemAllocNode, &FMemAllocNode::GetFullCallstackSourceFiles, FMemAllocNode::ECallstackType::FreeCallstack)
+						]
+					];
+			}
+
+			virtual FText FormatValueForTooltip(const FTableColumn& Column, const FBaseTreeNode& Node) const override
+			{
+				const FMemAllocNode& MemAllocNode = static_cast<const FMemAllocNode&>(Node);
+				return MemAllocNode.GetTopSourceFileEx(FMemAllocNode::ECallstackType::FreeCallstack);
+			}
+		};
+
+		TSharedRef<ITableCellValueFormatter> Formatter = MakeShared<FSourceValueFormatter>();
+		Column.SetValueFormatter(Formatter);
+
+		TSharedRef<ITableCellValueSorter> Sorter = MakeShared<FSorterByTextValueWithId>(ColumnRef);
+		Column.SetValueSorter(Sorter);
+
+		AddColumn(ColumnRef);
+	}
+	//////////////////////////////////////////////////
+	// Free Callstack (Size) Column
+	{
+		TSharedRef<FTableColumn> ColumnRef = MakeShared<FTableColumn>(FMemAllocTableColumns::FreeCallstackSizeColumnId);
+		FTableColumn& Column = *ColumnRef;
+
+		Column.SetIndex(ColumnIndex++);
+
+		Column.SetShortName(LOCTEXT("CallstackSizeColumnName_Free", "Free Callstack"));
+		Column.SetTitleName(LOCTEXT("CallstackSizeColumnTitle_Free", "Free Callstack"));
+		Column.SetDescription(LOCTEXT("CallstackSizeColumnDesc_Free", "Number of free callstack frames.\nTooltip shows the entire free callstack."));
+
+		Column.SetFlags(ETableColumnFlags::ShouldBeVisible | ETableColumnFlags::CanBeHidden | ETableColumnFlags::CanBeFiltered);
+
+		Column.SetHorizontalAlignment(HAlign_Left);
+		Column.SetInitialWidth(100.0f);
+
+		Column.SetDataType(ETableCellDataType::Int64);
+
+		class FCallstackSizeValueGetter : public FTableCellValueGetter
+		{
+		public:
+			virtual const TOptional<FTableCellValue> GetValue(const FTableColumn& Column, const FBaseTreeNode& Node) const override
+			{
+				if (Node.IsGroup())
+				{
+					const FTableTreeNode& NodePtr = static_cast<const FTableTreeNode&>(Node);
+					if (NodePtr.HasAggregatedValue(Column.GetId()))
+					{
+						return NodePtr.GetAggregatedValue(Column.GetId());
+					}
+				}
+				else //if (Node->Is<FMemAllocNode>())
+				{
+					const FMemAllocNode& MemAllocNode = static_cast<const FMemAllocNode&>(Node);
+					const FMemoryAlloc* Alloc = MemAllocNode.GetMemAlloc();
+					if (Alloc)
+					{
+						const TraceServices::FCallstack* Callstack = Alloc->GetCallstack();
+						const int64 CallstackSize = (Callstack && (Callstack->Num() != 1 || Callstack->Addr(0) != 0)) ? Callstack->Num() : 0;
+						return FTableCellValue(CallstackSize);
+					}
+				}
+
+				return TOptional<FTableCellValue>();
+			}
+
+			virtual uint64 GetValueId(const FTableColumn& Column, const FBaseTreeNode& Node) const override
+			{
+				if (!Node.IsGroup()) //if (Node->Is<FMemAllocNode>())
+				{
+					const FMemAllocNode& MemAllocNode = static_cast<const FMemAllocNode&>(Node);
+					return MemAllocNode.GetCallstackId();
+				}
+				return 0;
+			}
+		};
+		TSharedRef<ITableCellValueGetter> Getter = MakeShared<FCallstackSizeValueGetter>();
+		Column.SetValueGetter(Getter);
+
+		class FCallstackValueFormatter : public FTableCellValueFormatter
+		{
+		public:
+			virtual FText FormatValue(const TOptional<FTableCellValue>& InValue) const override
+			{
+				return InValue.IsSet() ? FText::AsNumber(InValue.GetValue().Int64) : FText::GetEmpty();
+			}
+
+			virtual TSharedPtr<IToolTip> GetCustomTooltip(const FTableColumn& Column, const FBaseTreeNode& Node) const override
+			{
+				const FMemAllocNode& MemAllocNode = static_cast<const FMemAllocNode&>(Node);
+
+				FText HeaderText = FText::Format(LOCTEXT("CallstackHeaderFmt", "CallstackId={0} FreeCallstackId={1}"),
+					FText::AsNumber(MemAllocNode.GetCallstackId(), &FNumberFormattingOptions::DefaultNoGrouping()),
+					FText::AsNumber(MemAllocNode.GetFreeCallstackId(), &FNumberFormattingOptions::DefaultNoGrouping()));
+
+				return SNew(SToolTip)
+					.Visibility(TAttribute<EVisibility>::Create(TAttribute<EVisibility>::FGetter::CreateStatic(&FTableCellValueFormatter::GetTooltipVisibility)))
+					[
+						SNew(SVerticalBox)
+						+ SVerticalBox::Slot()
+						.AutoHeight()
+						.Padding(2.0f)
+						[
+							SNew(STextBlock)
+							.Text(HeaderText)
+							.ColorAndOpacity(FSlateColor(EStyleColor::White25))
+						]
+						+ SVerticalBox::Slot()
+						.AutoHeight()
+						.Padding(2.0f)
+						[
+							SNew(STextBlock)
+							.Text(&MemAllocNode, &FMemAllocNode::GetFullCallstack, FMemAllocNode::ECallstackType::FreeCallstack)
+						]
+					];
+			}
+
+			virtual FText FormatValueForTooltip(const FTableColumn& Column, const FBaseTreeNode& Node) const override
+			{
+				const FMemAllocNode& MemAllocNode = static_cast<const FMemAllocNode&>(Node);
+				return MemAllocNode.GetFullCallstack(FMemAllocNode::ECallstackType::FreeCallstack);
 			}
 
 			virtual FText FormatValueForGrouping(const FTableColumn& Column, const FBaseTreeNode& Node) const override
diff --git a/Engine/Source/Developer/TraceInsights/Private/Insights/MemoryProfiler/ViewModels/MemAllocTable.h b/Engine/Source/Developer/TraceInsights/Private/Insights/MemoryProfiler/ViewModels/MemAllocTable.h
--- a/Engine/Source/Developer/TraceInsights/Private/Insights/MemoryProfiler/ViewModels/MemAllocTable.h
+++ b/Engine/Source/Developer/TraceInsights/Private/Insights/MemoryProfiler/ViewModels/MemAllocTable.h
@@ -32,9 +32,12 @@ struct FMemAllocTableColumns
 	static const FName AssetColumnId;
 	static const FName PackageColumnId;
 	static const FName ClassNameColumnId;
-	static const FName FunctionColumnId;
-	static const FName SourceFileColumnId;
-	static const FName CallstackSizeColumnId;
+	static const FName AllocFunctionColumnId;
+	static const FName AllocSourceFileColumnId;
+	static const FName AllocCallstackSizeColumnId;
+	static const FName FreeFunctionColumnId;
+	static const FName FreeSourceFileColumnId;
+	static const FName FreeCallstackSizeColumnId;
 };
 
 ////////////////////////////////////////////////////////////////////////////////////////////////////
diff --git a/Engine/Source/Developer/TraceInsights/Private/Insights/MemoryProfiler/ViewModels/MemoryAlloc.cpp b/Engine/Source/Developer/TraceInsights/Private/Insights/MemoryProfiler/ViewModels/MemoryAlloc.cpp
--- a/Engine/Source/Developer/TraceInsights/Private/Insights/MemoryProfiler/ViewModels/MemoryAlloc.cpp
+++ b/Engine/Source/Developer/TraceInsights/Private/Insights/MemoryProfiler/ViewModels/MemoryAlloc.cpp
@@ -36,64 +36,4 @@ FMemoryAlloc::~FMemoryAlloc()
 
 ////////////////////////////////////////////////////////////////////////////////////////////////////
 
-FText FMemoryAlloc::GetFullCallstack() const
-{
-	if (!Callstack)
-	{
-		return FText::FromString(GetCallstackNotAvailableString());
-	}
-
-	if (Callstack->Num() == 0)
-	{
-		return FText::FromString(GetEmptyCallstackString());
-	}
-
-	TStringBuilder<1024> Tooltip;
-	const uint32 NumCallstackFrames = Callstack->Num();
-	check(NumCallstackFrames <= 256); // see Callstack->Frame(uint8)
-	for (uint32 Index = 0; Index < NumCallstackFrames; ++Index)
-	{
-		if (Index != 0)
-		{
-			Tooltip << TEXT("\n");
-		}
-		const TraceServices::FStackFrame* Frame = Callstack->Frame(static_cast<uint8>(Index));
-		check(Frame != nullptr);
-		FormatStackFrame(*Frame, Tooltip, EStackFrameFormatFlags::ModuleSymbolFileAndLine);
-	}
-	return FText::FromString(FString(Tooltip));
-}
-
-////////////////////////////////////////////////////////////////////////////////////////////////////
-
-FText FMemoryAlloc::GetFullCallstackSourceFiles() const
-{
-	if (!Callstack)
-	{
-		return FText::FromString(GetCallstackNotAvailableString());
-	}
-
-	if (Callstack->Num() == 0)
-	{
-		return FText::FromString(GetEmptyCallstackString());
-	}
-
-	TStringBuilder<1024> Tooltip;
-	const uint32 NumCallstackFrames = Callstack->Num();
-	check(NumCallstackFrames <= 256); // see Callstack->Frame(uint8)
-	for (uint32 FrameIndex = 0; FrameIndex < NumCallstackFrames; ++FrameIndex)
-	{
-		if (FrameIndex != 0)
-		{
-			Tooltip << TEXT("\n");
-		}
-		const TraceServices::FStackFrame* Frame = Callstack->Frame(static_cast<uint8>(FrameIndex));
-		check(Frame != nullptr);
-		FormatStackFrame(*Frame, Tooltip, EStackFrameFormatFlags::FileAndLine);
-	}
-	return FText::FromString(FString(Tooltip));
-}
-
-////////////////////////////////////////////////////////////////////////////////////////////////////
-
 } // namespace Insights
diff --git a/Engine/Source/Developer/TraceInsights/Private/Insights/MemoryProfiler/ViewModels/MemoryAlloc.h b/Engine/Source/Developer/TraceInsights/Private/Insights/MemoryProfiler/ViewModels/MemoryAlloc.h
--- a/Engine/Source/Developer/TraceInsights/Private/Insights/MemoryProfiler/ViewModels/MemoryAlloc.h
+++ b/Engine/Source/Developer/TraceInsights/Private/Insights/MemoryProfiler/ViewModels/MemoryAlloc.h
@@ -39,8 +39,6 @@ public:
 	uint32 GetFreeCallstackId() const { return FreeCallstackId; }
 	const TraceServices::FCallstack* GetCallstack() const { return Callstack; }
 	const TraceServices::FCallstack* GetFreeCallstack() const { return FreeCallstack; }
-	FText GetFullCallstack() const;
-	FText GetFullCallstackSourceFiles() const;
 	HeapId GetRootHeap() const { return RootHeap; }
 
 private:
diff --git a/Engine/Source/Developer/TraceInsights/Private/Insights/MemoryProfiler/Widgets/SMemAllocTableTreeView.cpp b/Engine/Source/Developer/TraceInsights/Private/Insights/MemoryProfiler/Widgets/SMemAllocTableTreeView.cpp
--- a/Engine/Source/Developer/TraceInsights/Private/Insights/MemoryProfiler/Widgets/SMemAllocTableTreeView.cpp
+++ b/Engine/Source/Developer/TraceInsights/Private/Insights/MemoryProfiler/Widgets/SMemAllocTableTreeView.cpp
@@ -653,8 +653,8 @@ void SMemAllocTableTreeView::InitAvailableViewPresets()
 			InOutConfigSet.Add({ FMemAllocTableColumns::CountColumnId,         true, 100.0f });
 			InOutConfigSet.Add({ FMemAllocTableColumns::SizeColumnId,          true, 100.0f });
 			InOutConfigSet.Add({ FMemAllocTableColumns::TagColumnId,           true, 120.0f });
-			InOutConfigSet.Add({ FMemAllocTableColumns::FunctionColumnId,      true, 550.0f });
-			InOutConfigSet.Add({ FMemAllocTableColumns::CallstackSizeColumnId, true, 100.0f });
+			InOutConfigSet.Add({ FMemAllocTableColumns::AllocFunctionColumnId,      true, 550.0f });
+			InOutConfigSet.Add({ FMemAllocTableColumns::AllocCallstackSizeColumnId, true, 100.0f });
 		}
 	};
 	AvailableViewPresets.Add(MakeShared<FDefaultViewPreset>());
@@ -702,9 +702,9 @@ void SMemAllocTableTreeView::InitAvailableViewPresets()
 			InOutConfigSet.Add({ FMemAllocTableColumns::CountColumnId,           true, 100.0f });
 			InOutConfigSet.Add({ FMemAllocTableColumns::SizeColumnId,            true, 100.0f });
 			InOutConfigSet.Add({ FMemAllocTableColumns::TagColumnId,             true, 120.0f });
-			InOutConfigSet.Add({ FMemAllocTableColumns::FunctionColumnId,        true, 550.0f });
-			InOutConfigSet.Add({ FMemAllocTableColumns::SourceFileColumnId,      true, 550.0f });
-			InOutConfigSet.Add({ FMemAllocTableColumns::CallstackSizeColumnId,   true, 100.0f });
+			InOutConfigSet.Add({ FMemAllocTableColumns::AllocFunctionColumnId,        true, 550.0f });
+			InOutConfigSet.Add({ FMemAllocTableColumns::AllocSourceFileColumnId,      true, 550.0f });
+			InOutConfigSet.Add({ FMemAllocTableColumns::AllocCallstackSizeColumnId,   true, 100.0f });
 		}
 	};
 	AvailableViewPresets.Add(MakeShared<FDetailedViewPreset>());
@@ -754,7 +754,7 @@ void SMemAllocTableTreeView::InitAvailableViewPresets()
 			InOutConfigSet.Add({ FMemAllocTableColumns::CountColumnId,    true, 100.0f });
 			InOutConfigSet.Add({ FMemAllocTableColumns::SizeColumnId,     true, 100.0f });
 			InOutConfigSet.Add({ FMemAllocTableColumns::TagColumnId,      true, 200.0f });
-			InOutConfigSet.Add({ FMemAllocTableColumns::FunctionColumnId, true, 200.0f });
+			InOutConfigSet.Add({ FMemAllocTableColumns::AllocFunctionColumnId, true, 200.0f });
 		}
 	};
 	AvailableViewPresets.Add(MakeShared<FHeapViewPreset>());
@@ -805,7 +805,7 @@ void SMemAllocTableTreeView::InitAvailableViewPresets()
 			InOutConfigSet.Add({ FMemAllocTableColumns::CountColumnId,    true, 100.0f });
 			InOutConfigSet.Add({ FMemAllocTableColumns::SizeColumnId,     true, 100.0f });
 			InOutConfigSet.Add({ FMemAllocTableColumns::TagColumnId,      true, 120.0f });
-			InOutConfigSet.Add({ FMemAllocTableColumns::FunctionColumnId, true, 400.0f });
+			InOutConfigSet.Add({ FMemAllocTableColumns::AllocFunctionColumnId, true, 400.0f });
 		}
 	};
 	AvailableViewPresets.Add(MakeShared<FSizeViewPreset>());
@@ -854,7 +854,7 @@ void SMemAllocTableTreeView::InitAvailableViewPresets()
 			InOutConfigSet.Add({ FTable::GetHierarchyColumnId(),          true, 200.0f });
 			InOutConfigSet.Add({ FMemAllocTableColumns::CountColumnId,    true, 100.0f });
 			InOutConfigSet.Add({ FMemAllocTableColumns::SizeColumnId,     true, 100.0f });
-			InOutConfigSet.Add({ FMemAllocTableColumns::FunctionColumnId, true, 400.0f });
+			InOutConfigSet.Add({ FMemAllocTableColumns::AllocFunctionColumnId, true, 400.0f });
 		}
 	};
 	AvailableViewPresets.Add(MakeShared<FTagViewPreset>());
@@ -905,7 +905,7 @@ void SMemAllocTableTreeView::InitAvailableViewPresets()
 			InOutConfigSet.Add({ FMemAllocTableColumns::CountColumnId,    true, 100.0f });
 			InOutConfigSet.Add({ FMemAllocTableColumns::SizeColumnId,     true, 100.0f });
 			InOutConfigSet.Add({ FMemAllocTableColumns::TagColumnId,      true, 120.0f });
-			InOutConfigSet.Add({ FMemAllocTableColumns::FunctionColumnId, true, 400.0f });
+			InOutConfigSet.Add({ FMemAllocTableColumns::AllocFunctionColumnId, true, 400.0f });
 		}
 	};
 	AvailableViewPresets.Add(MakeShared<FAssetViewPreset>());
@@ -967,7 +967,7 @@ void SMemAllocTableTreeView::InitAvailableViewPresets()
 			InOutConfigSet.Add({ FMemAllocTableColumns::CountColumnId,    true, 100.0f });
 			InOutConfigSet.Add({ FMemAllocTableColumns::SizeColumnId,     true, 100.0f });
 			InOutConfigSet.Add({ FMemAllocTableColumns::TagColumnId,      true, 120.0f });
-			InOutConfigSet.Add({ FMemAllocTableColumns::FunctionColumnId, true, 400.0f });
+			InOutConfigSet.Add({ FMemAllocTableColumns::AllocFunctionColumnId, true, 400.0f });
 		}
 	};
 	AvailableViewPresets.Add(MakeShared<FClassNameViewPreset>());
@@ -978,22 +978,33 @@ void SMemAllocTableTreeView::InitAvailableViewPresets()
 	class FCallstackViewPreset : public ITableTreeViewPreset
 	{
 	public:
-		FCallstackViewPreset(bool bIsInverted)
+		FCallstackViewPreset(bool bIsInverted, bool bIsAlloc)
 			: bIsInvertedCallstack(bIsInverted)
+			, bIsAllocCallstack(bIsAlloc)
 		{
 		}
 
 		virtual FText GetName() const override
 		{
-			return bIsInvertedCallstack ?
-				LOCTEXT("InvertedCallstack_PresetName", "Inverted Callstack") :
-				LOCTEXT("Callstack_PresetName", "Callstack");
+			return
+				bIsAllocCallstack
+				? (bIsInvertedCallstack ?
+				LOCTEXT("InvertedCallstack_Alloc_PresetName", "Inverted Alloc Callstack") :
+				LOCTEXT("Callstack_Alloc_PresetName", "Alloc Callstack"))
+				: (bIsInvertedCallstack ?
+				LOCTEXT("InvertedCallstack_Free_PresetName", "Inverted Free Callstack") :
+				LOCTEXT("Callstack_Free_PresetName", "Free Callstack"));
 		}
 		virtual FText GetToolTip() const override
 		{
-			return bIsInvertedCallstack ?
-				LOCTEXT("InvertedCallstack_PresetToolTip", "Inverted Callstack Breakdown View\nConfigure the tree view to show a breakdown of allocations by inverted callstack.") :
-				LOCTEXT("Callstack_PresetToolTip", "Callstack Breakdown View\nConfigure the tree view to show a breakdown of allocations by callstack.");
+			return 
+				bIsAllocCallstack
+				? (bIsInvertedCallstack ?
+				LOCTEXT("InvertedCallstack_Alloc_PresetToolTip", "Inverted Alloc Callstack Breakdown View\nConfigure the tree view to show a breakdown of allocations by inverted callstack.") :
+				LOCTEXT("Callstack_Alloc_PresetToolTip", "Alloc Callstack Breakdown View\nConfigure the tree view to show a breakdown of allocations by callstack."))
+				: (bIsInvertedCallstack ?
+				LOCTEXT("InvertedCallstack_Free_PresetToolTip", "Inverted Free Callstack Breakdown View\nConfigure the tree view to show a breakdown of allocations by inverted callstack.") :
+				LOCTEXT("Callstack_Free_PresetToolTip", "Free Callstack Breakdown View\nConfigure the tree view to show a breakdown of allocations by callstack."));
 		}
 		virtual FName GetSortColumn() const override
 		{
@@ -1011,11 +1022,13 @@ void SMemAllocTableTreeView::InitAvailableViewPresets()
 			InOutCurrentGroupings.Add(InAvailableGroupings[0]);
 
 			const bool bIsInverted = bIsInvertedCallstack;
+			const bool bIsAlloc = bIsAllocCallstack;
 			const TSharedPtr<FTreeNodeGrouping>* CallstackGrouping = InAvailableGroupings.FindByPredicate(
-				[bIsInverted](TSharedPtr<FTreeNodeGrouping>& Grouping)
+				[bIsInverted, bIsAlloc](TSharedPtr<FTreeNodeGrouping>& Grouping)
 				{
 					return Grouping->Is<FMemAllocGroupingByCallstack>() &&
-						   Grouping->As<FMemAllocGroupingByCallstack>().IsInverted() == bIsInverted;
+						   Grouping->As<FMemAllocGroupingByCallstack>().IsInverted() == bIsInverted &&
+						   Grouping->As<FMemAllocGroupingByCallstack>().IsAllocCallstack() == bIsAlloc;
 				});
 			if (CallstackGrouping)
 			{
@@ -1028,14 +1041,24 @@ void SMemAllocTableTreeView::InitAvailableViewPresets()
 			InOutConfigSet.Add({ FMemAllocTableColumns::CountColumnId,    true, 100.0f });
 			InOutConfigSet.Add({ FMemAllocTableColumns::SizeColumnId,     true, 100.0f });
 			InOutConfigSet.Add({ FMemAllocTableColumns::TagColumnId,      true, 200.0f });
-			InOutConfigSet.Add({ FMemAllocTableColumns::FunctionColumnId, true, 200.0f });
+			if (bIsAllocCallstack)
+			{
+				InOutConfigSet.Add({ FMemAllocTableColumns::AllocFunctionColumnId, true, 200.0f });
+			}
+			else
+			{
+				InOutConfigSet.Add({ FMemAllocTableColumns::FreeFunctionColumnId, true, 200.0f });
+			}
 		}
 
 	private:
 		bool bIsInvertedCallstack;
+		bool bIsAllocCallstack;
 	};
-	AvailableViewPresets.Add(MakeShared<FCallstackViewPreset>(false));
-	AvailableViewPresets.Add(MakeShared<FCallstackViewPreset>(true));
+	AvailableViewPresets.Add(MakeShared<FCallstackViewPreset>(false, true));
+	AvailableViewPresets.Add(MakeShared<FCallstackViewPreset>(true, true));
+	AvailableViewPresets.Add(MakeShared<FCallstackViewPreset>(false, false));
+	AvailableViewPresets.Add(MakeShared<FCallstackViewPreset>(true, false));
 
 	//////////////////////////////////////////////////
 	// Address (4K Page) Breakdown View
@@ -1084,7 +1107,7 @@ void SMemAllocTableTreeView::InitAvailableViewPresets()
 			InOutConfigSet.Add({ FMemAllocTableColumns::CountColumnId,    true, 100.0f });
 			InOutConfigSet.Add({ FMemAllocTableColumns::SizeColumnId,     true, 100.0f });
 			InOutConfigSet.Add({ FMemAllocTableColumns::TagColumnId,      true, 120.0f });
-			InOutConfigSet.Add({ FMemAllocTableColumns::FunctionColumnId, true, 400.0f });
+			InOutConfigSet.Add({ FMemAllocTableColumns::AllocFunctionColumnId, true, 400.0f });
 		}
 	};
 	AvailableViewPresets.Add(MakeShared<FPageViewPreset>());
@@ -1303,7 +1326,8 @@ void SMemAllocTableTreeView::InternalCreateGroupings()
 			else if (Grouping->Is<FTreeNodeGroupingByPathBreakdown>())
 			{
 				const FName ColumnId = Grouping->As<FTreeNodeGroupingByPathBreakdown>().GetColumnId();
-				if (ColumnId == FMemAllocTableColumns::FunctionColumnId ||
+				if (ColumnId == FMemAllocTableColumns::AllocFunctionColumnId ||
+					ColumnId == FMemAllocTableColumns::FreeFunctionColumnId ||
 					ColumnId == FMemAllocTableColumns::ClassNameColumnId)
 				{
 					return true;
@@ -1323,8 +1347,10 @@ void SMemAllocTableTreeView::InternalCreateGroupings()
 		AvailableGroupings.Insert(MakeShared<FMemAllocGroupingByTag>(*AllocationsProvider), Index++);
 	}
 
-	AvailableGroupings.Insert(MakeShared<FMemAllocGroupingByCallstack>(false, bIsCallstackGroupingByFunction), Index++);
-	AvailableGroupings.Insert(MakeShared<FMemAllocGroupingByCallstack>(true, bIsCallstackGroupingByFunction), Index++);
+	AvailableGroupings.Insert(MakeShared<FMemAllocGroupingByCallstack>(true, false, bIsCallstackGroupingByFunction), Index++);
+	AvailableGroupings.Insert(MakeShared<FMemAllocGroupingByCallstack>(true, true, bIsCallstackGroupingByFunction), Index++);
+	AvailableGroupings.Insert(MakeShared<FMemAllocGroupingByCallstack>(false, false, bIsCallstackGroupingByFunction), Index++);
+	AvailableGroupings.Insert(MakeShared<FMemAllocGroupingByCallstack>(false, true, bIsCallstackGroupingByFunction), Index++);
 
 	if (AllocationsProvider)
 	{
@@ -1377,7 +1403,7 @@ void SMemAllocTableTreeView::UpdateFilterContext(const FFilterConfigurator& InFi
 		// GetFullCallstack is super heavy to compute. Validate that the filter has a use for this key before computing it.
 		if (InFilterConfigurator.IsKeyUsed(FullCallStackIndex))
 		{
-			FilterContext.SetFilterData<FString>(FullCallStackIndex, MemNode.GetFullCallstack().ToString());
+			FilterContext.SetFilterData<FString>(FullCallStackIndex, MemNode.GetFullCallstack(FMemAllocNode::ECallstackType::AllocCallstack).ToString());
 		}
 
 		if (InFilterConfigurator.IsKeyUsed(LLMFilterIndex))
@@ -1395,8 +1421,8 @@ void SMemAllocTableTreeView::InitFilterConfigurator(FFilterConfigurator& InOutFi
 
 	TSharedRef<FFilter> FullCallStackFilter = MakeShared<FFilter>(
 		FullCallStackIndex,
-		LOCTEXT("FullCallstack", "Full Callstack"),
-		LOCTEXT("SearchFullCallstack", "Search in all the callstack frames"),
+		LOCTEXT("FullCallstack", "Full Alloc Callstack"),
+		LOCTEXT("SearchFullCallstack", "Search in all the alloc callstack frames"),
 		EFilterDataType::String,
 		nullptr,
 		FFilterService::Get()->GetStringOperators());
@@ -1568,14 +1594,28 @@ void SMemAllocTableTreeView::ExtendMenu(FMenuBuilder& MenuBuilder)
 	{
 		MenuBuilder.BeginSection("Allocation", LOCTEXT("ContextMenu_Section_OpenSource", "Allocation"));
 		{
-			FText ItemLabel = FText::Format(LOCTEXT("ContextMenu_Open_SubMenu", "Open in {0}"), SourceCodeAccessor.GetNameText());
-			FText ItemToolTip = FText::Format(LOCTEXT("ContextMenu_Open_Desc_SubMenu", "Open source file of selected callstack frame in {0}."), SourceCodeAccessor.GetNameText());
+			FText ItemLabel = FText::Format(LOCTEXT("ContextMenu_Open_SubMenu_Alloc", "Open in {0} | Alloc Callstack"), SourceCodeAccessor.GetNameText());
+			FText ItemToolTip = FText::Format(LOCTEXT("ContextMenu_Open_Desc_SubMenu_Alloc", "Open source file of selected allocation callstack frame in {0}."), SourceCodeAccessor.GetNameText());
+
+			// Alloc callstack
+			MenuBuilder.AddSubMenu
+			(
+				ItemLabel,
+				ItemToolTip,
+				FNewMenuDelegate::CreateSP(this, &SMemAllocTableTreeView::BuildOpenSourceSubMenu, true),
+				false,
+				FSlateIcon(FAppStyle::GetAppStyleSetName(), SourceCodeAccessor.GetOpenIconName())
+			);
 
+			ItemLabel = FText::Format(LOCTEXT("ContextMenu_Open_SubMenu_Free", "Open in {0} | Free Callstack"), SourceCodeAccessor.GetNameText());
+			ItemToolTip = FText::Format(LOCTEXT("ContextMenu_Open_Desc_SubMenu_Free", "Open source file of selected free callstack frame in {0}."), SourceCodeAccessor.GetNameText());
+
+			// Free Callstack
 			MenuBuilder.AddSubMenu
 			(
 				ItemLabel,
 				ItemToolTip,
-				FNewMenuDelegate::CreateSP(this, &SMemAllocTableTreeView::BuildOpenSourceSubMenu),
+				FNewMenuDelegate::CreateSP(this, &SMemAllocTableTreeView::BuildOpenSourceSubMenu, false),
 				false,
 				FSlateIcon(FAppStyle::GetAppStyleSetName(), SourceCodeAccessor.GetOpenIconName())
 			);
@@ -1659,104 +1699,117 @@ uint32 SMemAllocTableTreeView::CountSourceFiles(FMemAllocNode& MemAllocNode)
 
 ////////////////////////////////////////////////////////////////////////////////////////////////////
 
-void SMemAllocTableTreeView::BuildOpenSourceSubMenu(FMenuBuilder& MenuBuilder)
+bool SMemAllocTableTreeView::BuildOpenSourceSubMenuItems(FMenuBuilder& MenuBuilder, const TraceServices::FCallstack& Callstack)
 {
-	MenuBuilder.BeginSection("OpenSource");
-	{
-		uint32 NumSourceFiles = 0;
+	ISourceCodeAccessModule& SourceCodeAccessModule = FModuleManager::LoadModuleChecked<ISourceCodeAccessModule>("SourceCodeAccess");
+	ISourceCodeAccessor& SourceCodeAccessor = SourceCodeAccessModule.GetAccessor();
 
-		FMemAllocNodePtr MemAllocNode = GetSingleSelectedMemAllocNode();
-		if (MemAllocNode.IsValid())
+	uint32 NumSourceFiles = 0;
+	const uint32 NumCallstackFrames = Callstack.Num();
+	check(NumCallstackFrames <= 256); // see Callstack.Frame(uint8)
+	for (uint32 FrameIndex = 0; FrameIndex < NumCallstackFrames; ++FrameIndex)
+	{
+		const TraceServices::FStackFrame* Frame = Callstack.Frame(static_cast<uint8>(FrameIndex));
+		if (Frame && Frame->Symbol && Frame->Symbol->File)
 		{
-			const FMemoryAlloc* Alloc = MemAllocNode->GetMemAlloc();
-			if (Alloc && Alloc->Callstack)
+			FText ItemLabel;
+			FText ItemToolTip;
+			if (Frame->Symbol->GetResult() == TraceServices::ESymbolQueryResult::OK)
 			{
-				ISourceCodeAccessModule& SourceCodeAccessModule = FModuleManager::LoadModuleChecked<ISourceCodeAccessModule>("SourceCodeAccess");
-				ISourceCodeAccessor& SourceCodeAccessor = SourceCodeAccessModule.GetAccessor();
+				FText FileName;
+				constexpr int32 MaxFileNameLen = 120;
+				const int32 FileNameLen = FCString::Strlen(Frame->Symbol->File);
+				if (FileNameLen > MaxFileNameLen)
+				{
+					FString FileNameStr = TEXT("...") + FString(MaxFileNameLen, Frame->Symbol->File + (FileNameLen - MaxFileNameLen));
+					FileName = FText::FromString(FileNameStr);
+				}
+				else
+				{
+					FileName = FText::FromString(Frame->Symbol->File);
+				}
 
-				const uint32 NumCallstackFrames = Alloc->Callstack->Num();
-				check(NumCallstackFrames <= 256); // see Callstack->Frame(uint8)
-				for (uint32 FrameIndex = 0; FrameIndex < NumCallstackFrames; ++FrameIndex)
+				FText SymbolName;
+				constexpr int32 MaxSymbolNameLen = 100;
+				const int32 SymbolNameLen = FCString::Strlen(Frame->Symbol->Name);
+				if (SymbolNameLen > MaxSymbolNameLen)
 				{
-					const TraceServices::FStackFrame* Frame = Alloc->Callstack->Frame(static_cast<uint8>(FrameIndex));
-					if (Frame && Frame->Symbol && Frame->Symbol->File)
-					{
-						FText ItemLabel;
-						FText ItemToolTip;
-						if (Frame->Symbol->GetResult() == TraceServices::ESymbolQueryResult::OK)
-						{
-							FText FileName;
-							constexpr int32 MaxFileNameLen = 120;
-							const int32 FileNameLen = FCString::Strlen(Frame->Symbol->File);
-							if (FileNameLen > MaxFileNameLen)
-							{
-								FString FileNameStr = TEXT("...") + FString(MaxFileNameLen, Frame->Symbol->File + (FileNameLen - MaxFileNameLen));
-								FileName = FText::FromString(FileNameStr);
-							}
-							else
-							{
-								FileName = FText::FromString(Frame->Symbol->File);
-							}
+					FString SymbolNameStr = TEXT("...") + FString(MaxSymbolNameLen, Frame->Symbol->Name + (SymbolNameLen - MaxSymbolNameLen));
+					SymbolName = FText::FromString(SymbolNameStr);
+				}
+				else
+				{
+					SymbolName = FText::FromString(Frame->Symbol->Name);
+				}
 
-							FText SymbolName;
-							constexpr int32 MaxSymbolNameLen = 100;
-							const int32 SymbolNameLen = FCString::Strlen(Frame->Symbol->Name);
-							if (SymbolNameLen > MaxSymbolNameLen)
-							{
-								FString SymbolNameStr = TEXT("...") + FString(MaxSymbolNameLen, Frame->Symbol->Name + (SymbolNameLen - MaxSymbolNameLen));
-								SymbolName = FText::FromString(SymbolNameStr);
-							}
-							else
-							{
-								SymbolName = FText::FromString(Frame->Symbol->Name);
-							}
-
-							ItemLabel = FText::Format(LOCTEXT("ContextMenu_OpenSource_Fmt1", "{0} ({1}) \u2192 {2}"),
-								FileName,
-								FText::AsNumber(Frame->Symbol->Line),
-								SymbolName);
-
-							ItemToolTip = FText::Format(LOCTEXT("ContextMenu_OpenSource_Desc_Fmt1", "Open source file of selected callstack frame in {0}.\n{1} (line {2})\n\u2192 {3}"),
-								SourceCodeAccessor.GetNameText(),
-								FText::FromString(Frame->Symbol->File),
-								FText::AsNumber(Frame->Symbol->Line),
-								FText::FromString(Frame->Symbol->Name));
-						}
-						else
-						{
-							ItemLabel = FText::Format(LOCTEXT("ContextMenu_OpenSource_Fmt2", "{0} ({1}) \u2192 {2}"),
-								FText::FromString(Frame->Symbol->Module),
-								FText::FromString(FString::Printf(TEXT("0x%llX"), Frame->Addr)),
-								FText::FromString(TraceServices::QueryResultToString(Frame->Symbol->GetResult())));
+				ItemLabel = FText::Format(LOCTEXT("ContextMenu_OpenSource_Fmt1", "{0} ({1}) \u2192 {2}"),
+					FileName,
+					FText::AsNumber(Frame->Symbol->Line),
+					SymbolName);
 
-							ItemToolTip = FText::Format(LOCTEXT("ContextMenu_OpenSource_Desc_Fmt2", "Open source file of selected callstack frame in {0}."),
-								SourceCodeAccessor.GetNameText());
-						}
+				ItemToolTip = FText::Format(LOCTEXT("ContextMenu_OpenSource_Desc_Fmt1", "Open source file of selected callstack frame in {0}.\n{1} (line {2})\n\u2192 {3}"),
+					SourceCodeAccessor.GetNameText(),
+					FText::FromString(Frame->Symbol->File),
+					FText::AsNumber(Frame->Symbol->Line),
+					FText::FromString(Frame->Symbol->Name));
+			}
+			else
+			{
+				ItemLabel = FText::Format(LOCTEXT("ContextMenu_OpenSource_Fmt2", "{0} ({1}) \u2192 {2}"),
+					FText::FromString(Frame->Symbol->Module),
+					FText::FromString(FString::Printf(TEXT("0x%llX"), Frame->Addr)),
+					FText::FromString(TraceServices::QueryResultToString(Frame->Symbol->GetResult())));
 
-						const bool bFileExists = FPaths::FileExists(Frame->Symbol->File);
-
-						FUIAction Action_OpenIDE
-						(
-							FExecuteAction::CreateSP(this, &SMemAllocTableTreeView::OpenSourceFileInIDE, Frame->Symbol->File, uint32(Frame->Symbol->Line)),
-							FCanExecuteAction::CreateLambda([bFileExists]() { return bFileExists; })
-						);
-						MenuBuilder.AddMenuEntry
-						(
-							ItemLabel,
-							ItemToolTip,
-							FSlateIcon(),
-							Action_OpenIDE,
-							NAME_None,
-							EUserInterfaceActionType::Button
-						);
-
-						++NumSourceFiles;
-					}
+				ItemToolTip = FText::Format(LOCTEXT("ContextMenu_OpenSource_Desc_Fmt2", "Open source file of selected callstack frame in {0}."),
+					SourceCodeAccessor.GetNameText());
+			}
+
+			const bool bFileExists = FPaths::FileExists(Frame->Symbol->File);
+
+			FUIAction Action_OpenIDE
+			(
+				FExecuteAction::CreateSP(this, &SMemAllocTableTreeView::OpenSourceFileInIDE, Frame->Symbol->File, uint32(Frame->Symbol->Line)),
+				FCanExecuteAction::CreateLambda([bFileExists]() { return bFileExists; })
+			);
+			MenuBuilder.AddMenuEntry
+			(
+				ItemLabel,
+				ItemToolTip,
+				FSlateIcon(),
+				Action_OpenIDE,
+				NAME_None,
+				EUserInterfaceActionType::Button
+			);
+
+			++NumSourceFiles;
+		}
+	}
+
+	return NumSourceFiles > 0;
+}
+
+////////////////////////////////////////////////////////////////////////////////////////////////////
+
+void SMemAllocTableTreeView::BuildOpenSourceSubMenu(FMenuBuilder& MenuBuilder, bool bIsAllocCallstack)
+{
+	MenuBuilder.BeginSection("OpenSource");
+	{
+		bool bHasAnySourceFilesToOpen = false;
+		FMemAllocNodePtr MemAllocNode = GetSingleSelectedMemAllocNode();
+		if (MemAllocNode.IsValid())
+		{
+			const FMemoryAlloc* Alloc = MemAllocNode->GetMemAlloc();
+			if (Alloc)
+			{
+				const TraceServices::FCallstack* Callstack = bIsAllocCallstack ? Alloc->Callstack : Alloc->FreeCallstack;
+				if (Callstack)
+				{
+					bHasAnySourceFilesToOpen = BuildOpenSourceSubMenuItems(MenuBuilder, *Callstack);
 				}
 			}
 		}
 
-		if (NumSourceFiles == 0)
+		if (!bHasAnySourceFilesToOpen)
 		{
 			struct FLocal
 			{
@@ -1887,13 +1940,14 @@ void SMemAllocTableTreeView::ExportMemorySnapshot() const
 				OutData << Value;
 			}
 		};
-	auto WriteCallstackColumn = [Separator](FStringBuilderBase& OutData, const FTableTreeNode& Node)
+	auto WriteCallstackColumn = [Separator](FStringBuilderBase& OutData, const FTableTreeNode& Node, bool bIsAllocCallstack)
 		{
 			const FMemAllocNode& MemAllocNode = static_cast<const FMemAllocNode&>(Node);
 			const FMemoryAlloc* Alloc = MemAllocNode.GetMemAlloc();
 			if (Alloc)
 			{
-				if (!Alloc->Callstack)
+				const TraceServices::FCallstack* Callstack = bIsAllocCallstack ? Alloc->Callstack : Alloc->FreeCallstack;
+				if (!Callstack)
 				{
 					OutData << QuotationMarkBegin;
 					OutData << GetCallstackNotAvailableString();
@@ -1901,7 +1955,7 @@ void SMemAllocTableTreeView::ExportMemorySnapshot() const
 					return;
 				}
 
-				if (Alloc->Callstack->Num() == 0)
+				if (Callstack->Num() == 0)
 				{
 					OutData << QuotationMarkBegin;
 					OutData << GetEmptyCallstackString();
@@ -1909,7 +1963,7 @@ void SMemAllocTableTreeView::ExportMemorySnapshot() const
 					return;
 				}
 
-				const uint32 NumCallstackFrames = Alloc->Callstack->Num();
+				const uint32 NumCallstackFrames = Callstack->Num();
 				check(NumCallstackFrames <= 256);
 				OutData << QuotationMarkBegin;
 				for (uint32 Index = 0; Index < NumCallstackFrames; ++Index)
@@ -1918,7 +1972,7 @@ void SMemAllocTableTreeView::ExportMemorySnapshot() const
 					{
 						OutData << TEXT("/");
 					}
-					const TraceServices::FStackFrame* Frame = Alloc->Callstack->Frame(static_cast<uint8>(Index));
+					const TraceServices::FStackFrame* Frame = Callstack->Frame(static_cast<uint8>(Index));
 					check(Frame != nullptr);
 					FormatStackFrame(*Frame, OutData, EStackFrameFormatFlags::ModuleAndSymbol);
 				}
@@ -1937,8 +1991,10 @@ void SMemAllocTableTreeView::ExportMemorySnapshot() const
 	const FTableColumn& TagColumn = *GetTable()->FindColumn(FMemAllocTableColumns::TagColumnId);
 	const FTableColumn& AssetColumn = *GetTable()->FindColumn(FMemAllocTableColumns::AssetColumnId);
 	const FTableColumn& ClassNameColumn = *GetTable()->FindColumn(FMemAllocTableColumns::ClassNameColumnId);
-	const FTableColumn& FunctionColumn = *GetTable()->FindColumn(FMemAllocTableColumns::FunctionColumnId);
-	const FTableColumn& SourceFileColumn = *GetTable()->FindColumn(FMemAllocTableColumns::SourceFileColumnId);
+	const FTableColumn& AllocFunctionColumn = *GetTable()->FindColumn(FMemAllocTableColumns::AllocFunctionColumnId);
+	const FTableColumn& FreeFunctionColumn = *GetTable()->FindColumn(FMemAllocTableColumns::FreeFunctionColumnId);
+	const FTableColumn& AllocSourceFileColumn = *GetTable()->FindColumn(FMemAllocTableColumns::AllocSourceFileColumnId);
+	const FTableColumn& FreeSourceFileColumn = *GetTable()->FindColumn(FMemAllocTableColumns::FreeSourceFileColumnId);
 
 	// 2. Iterate over TreeNodes
 	TStringBuilder<2048> Buffer;
@@ -1963,9 +2019,12 @@ void SMemAllocTableTreeView::ExportMemorySnapshot() const
 			WriteColumnHeader(Buffer, TagColumn); Buffer += Separator;
 			WriteColumnHeader(Buffer, AssetColumn); Buffer += Separator;
 			WriteColumnHeader(Buffer, ClassNameColumn); Buffer += Separator;
-			WriteColumnHeader(Buffer, FunctionColumn); Buffer += Separator;
-			WriteColumnHeader(Buffer, SourceFileColumn); Buffer += Separator;
-			Buffer += TEXT("Callstack"); Buffer += LineEnd;
+			WriteColumnHeader(Buffer, AllocFunctionColumn); Buffer += Separator;
+			WriteColumnHeader(Buffer, FreeFunctionColumn); Buffer += Separator;
+			WriteColumnHeader(Buffer, AllocSourceFileColumn); Buffer += Separator;
+			WriteColumnHeader(Buffer, FreeSourceFileColumn); Buffer += Separator;
+			Buffer += TEXT("Alloc Callstack"); Buffer += Separator;
+			Buffer += TEXT("Free Callstack"); Buffer += LineEnd;
 
 			IsFirstRow = false;
 		}
@@ -1983,9 +2042,12 @@ void SMemAllocTableTreeView::ExportMemorySnapshot() const
 		WriteColumn(Buffer, TagColumn, *Node); Buffer += Separator;
 		WriteColumn(Buffer, AssetColumn, *Node); Buffer += Separator;
 		WriteColumn(Buffer, ClassNameColumn, *Node); Buffer += Separator;
-		WriteColumn(Buffer, FunctionColumn, *Node); Buffer += Separator;
-		WriteColumn(Buffer, SourceFileColumn, *Node); Buffer += Separator;
-		WriteCallstackColumn(Buffer, *Node); Buffer += LineEnd;
+		WriteColumn(Buffer, AllocFunctionColumn, *Node); Buffer += Separator;
+		WriteColumn(Buffer, FreeFunctionColumn, *Node); Buffer += Separator;
+		WriteColumn(Buffer, AllocSourceFileColumn, *Node); Buffer += Separator;
+		WriteColumn(Buffer, FreeSourceFileColumn, *Node); Buffer += Separator;
+		WriteCallstackColumn(Buffer, *Node, true); Buffer += Separator;
+		WriteCallstackColumn(Buffer, *Node, false); Buffer += LineEnd;
 
 		// 5. Write rows to file
 		FTCHARToUTF16 UTF16String(*Buffer, Buffer.Len());
diff --git a/Engine/Source/Developer/TraceInsights/Private/Insights/MemoryProfiler/Widgets/SMemAllocTableTreeView.h b/Engine/Source/Developer/TraceInsights/Private/Insights/MemoryProfiler/Widgets/SMemAllocTableTreeView.h
--- a/Engine/Source/Developer/TraceInsights/Private/Insights/MemoryProfiler/Widgets/SMemAllocTableTreeView.h
+++ b/Engine/Source/Developer/TraceInsights/Private/Insights/MemoryProfiler/Widgets/SMemAllocTableTreeView.h
@@ -98,7 +98,8 @@ protected:
 	uint32 CountSourceFiles(FMemAllocNode& MemAllocNode);
 
 	virtual void ExtendMenu(FMenuBuilder& MenuBuilder) override;
-	void BuildOpenSourceSubMenu(FMenuBuilder& MenuBuilder);
+	bool BuildOpenSourceSubMenuItems(FMenuBuilder& MenuBuilder, const TraceServices::FCallstack& Callstack);
+	void BuildOpenSourceSubMenu(FMenuBuilder& MenuBuilder, bool bIsAllocCallstack);
 	bool CanOpenCallstackFrameSourceFileInIDE() const;
 	void OpenCallstackFrameSourceFileInIDE();
 	FText GetSelectedCallstackFrameFileName() const;
diff --git a/Engine/Source/Developer/TraceServices/Private/Model/ModuleProvider.cpp b/Engine/Source/Developer/TraceServices/Private/Model/ModuleProvider.cpp
--- a/Engine/Source/Developer/TraceServices/Private/Model/ModuleProvider.cpp
+++ b/Engine/Source/Developer/TraceServices/Private/Model/ModuleProvider.cpp
@@ -451,8 +451,10 @@ FResolvedSymbolFilter::FResolvedSymbolFilter()
 	IgnoreSymbolsByFunctionName.Add(TEXT("FVirtualWinApiHooks::"));
 	IgnoreSymbolsByFunctionName.Add(TEXT("Malloc"));
 	IgnoreSymbolsByFunctionName.Add(TEXT("Realloc"));
+	IgnoreSymbolsByFunctionName.Add(TEXT("Free"));
 	IgnoreSymbolsByFunctionName.Add(TEXT("MemoryTrace_"));
 	IgnoreSymbolsByFunctionName.Add(TEXT("operator new"));
+	IgnoreSymbolsByFunctionName.Add(TEXT("operator delete"));
 	IgnoreSymbolsByFunctionName.Add(TEXT("std::"));
 	IgnoreSymbolsByFunctionName.Add(TEXT("FWindowsPlatformMemory::"));
 	IgnoreSymbolsByFunctionName.Add(TEXT("FCachedOSPageAllocator::"));
