Increased timeout for some tests to avoid timeout fail on Mac arm64 in Insights tests EngineTest
Changed UTS version visualisation in log to avoid "heap-use-after-free" error in IsUnrealTraceServerReady()

#tests mac, linux, win64

diff --git a/Engine/Source/Developer/TraceInsights/Private/Insights/Tests/InsightsTestUtils.cpp b/Engine/Source/Developer/TraceInsights/Private/Insights/Tests/InsightsTestUtils.cpp
--- a/Engine/Source/Developer/TraceInsights/Private/Insights/Tests/InsightsTestUtils.cpp
+++ b/Engine/Source/Developer/TraceInsights/Private/Insights/Tests/InsightsTestUtils.cpp
@@ -128,9 +128,10 @@ bool FInsightsTestUtils::IsUnrealTraceServerReady(const TCHAR* Host, int32 Port)
 		delete StoreClient;
 		return false;
 	}
-
+	uint32 MajorVersion = Version->GetMajorVersion();
+	uint32 MinorVersion = Version->GetMinorVersion();
 	delete StoreClient;
-	Test->AddInfo(FString::Printf(TEXT("Connected to UTS version %u.%u"), Version->GetMajorVersion(), Version->GetMinorVersion()));
+	Test->AddInfo(FString::Printf(TEXT("Connected to UTS version %u.%u"), MajorVersion, MinorVersion));
 	return true;
 }
 
@@ -289,15 +290,57 @@ bool FInsightsTestUtils::IsTraceHasLiveStatus(const FString& TraceName, const TC
 		}
 		uint32 TraceId = SessionInfo->GetTraceId();
 		const UE::Trace::FStoreClient::FTraceInfo* Info = StoreClient->GetTraceInfoById(TraceId);
-		if (TraceName.Contains(static_cast<FString>(Info->GetName())))
+		const FUtf8StringView Utf8TraceNameView = Info->GetName();
+		FString ActualTraceName(Utf8TraceNameView);
+		if (TraceName.Contains(ActualTraceName))
 		{
 			Test->AddInfo(TEXT("Trace is live"));
 			delete StoreClient;
 			return true;
 		}
+		Test->AddInfo(FString::Printf(TEXT("The live trace is %s"), *ActualTraceName));
 	}
 
 	Test->AddInfo(FString::Printf(TEXT("The trace with name %s does not have live status. Trying to find live trace"), *TraceName));
 	delete StoreClient;
 	return false;
+}
+
+////////////////////////////////////////////////////////////////////////////////////////////////////
+
+FString FInsightsTestUtils::GetLiveTrace(const TCHAR* Host, int32 Port) const
+{
+	UE::Trace::FStoreClient* StoreClient = UE::Trace::FStoreClient::Connect(Host, Port);
+	if (!StoreClient)
+	{
+		Test->AddInfo(TEXT("The StoreClient shouldn't be null"));
+		return TEXT("");
+	}
+	uint32 SessionCount = StoreClient->GetSessionCount();
+	if (!SessionCount)
+	{
+		Test->AddInfo(TEXT("The SessionCount shouldn't be 0"));
+		delete StoreClient;
+		return TEXT("");
+	}
+
+	for (uint32 Index = 0; Index < SessionCount; ++Index)
+	{
+		const UE::Trace::FStoreClient::FSessionInfo* SessionInfo = StoreClient->GetSessionInfo(Index);
+		if (!SessionInfo)
+		{
+			continue;
+		}
+
+		uint32 TraceId = SessionInfo->GetTraceId();
+		const UE::Trace::FStoreClient::FTraceInfo* Info = StoreClient->GetTraceInfoById(TraceId);
+		FString LiveTraceName = static_cast<FString>(Info->GetName());
+		LiveTraceName = LiveTraceName + TEXT(".utrace");
+		delete StoreClient;
+		return LiveTraceName;
+	}
+
+	Test->AddInfo(TEXT("There isn't any live trace"));
+	delete StoreClient;
+	return TEXT("");
 }
\ No newline at end of file
diff --git a/Engine/Source/Developer/TraceInsights/Public/Insights/Tests/InsightsTestUtils.h b/Engine/Source/Developer/TraceInsights/Public/Insights/Tests/InsightsTestUtils.h
--- a/Engine/Source/Developer/TraceInsights/Public/Insights/Tests/InsightsTestUtils.h
+++ b/Engine/Source/Developer/TraceInsights/Public/Insights/Tests/InsightsTestUtils.h
@@ -20,6 +20,7 @@ public:
 	bool IsUnrealTraceServerReady(const TCHAR* Host = TEXT("localhost"), int32 Port = 0U) const;
 	bool IsTraceHasLiveStatus(const FString& TraceName, const TCHAR* Host = TEXT("localhost"), int32 Port = 0U) const;
 	void ResetSession() const;
+	FString GetLiveTrace(const TCHAR* Host = TEXT("localhost"), int32 Port = 0U) const;
 
 private:
 	FAutomationTestBase* Test;
