Fixed UE-197542

diff --git a/Engine/Source/Developer/TraceInsights/Private/Insights/Tests/InsightsTestUtils.cpp b/Engine/Source/Developer/TraceInsights/Private/Insights/Tests/InsightsTestUtils.cpp
--- a/Engine/Source/Developer/TraceInsights/Private/Insights/Tests/InsightsTestUtils.cpp
+++ b/Engine/Source/Developer/TraceInsights/Private/Insights/Tests/InsightsTestUtils.cpp
@@ -290,6 +290,10 @@ bool FInsightsTestUtils::IsTraceHasLiveStatus(const FString& TraceName, const TC
 		}
 		uint32 TraceId = SessionInfo->GetTraceId();
 		const UE::Trace::FStoreClient::FTraceInfo* Info = StoreClient->GetTraceInfoById(TraceId);
+		if (!Info)
+		{
+			continue;
+		}
 		const FUtf8StringView Utf8TraceNameView = Info->GetName();
 		FString ActualTraceName(Utf8TraceNameView);
 		if (TraceName.Contains(ActualTraceName))
