[Insights] MemoryInsights/Mem Allocs table:
  - Added "Alloc/Free Callstack Id" columns, separated by "Alloc/Free Callstack Size" columns. Both shows entire callstack in tooltip (and can be copied to clipboard). Also removed the full callstacks from the main tooltip of an allocation (where values of all columns are displayed).
  - Fixed "Free Callstack" columns (some showing the wrong "alloc callstack" and other still using "alloc callstack" for sorting and grouping).
  - Added "Alloc/Free Thread Id" columns (shows the trace assigned thread id for "alloc" and "free" events).

diff --git a/Engine/Source/Developer/TraceInsights/Private/Insights/MemoryProfiler/ViewModels/MemAllocGroupingByCallstack.cpp b/Engine/Source/Developer/TraceInsights/Private/Insights/MemoryProfiler/ViewModels/MemAllocGroupingByCallstack.cpp
--- a/Engine/Source/Developer/TraceInsights/Private/Insights/MemoryProfiler/ViewModels/MemAllocGroupingByCallstack.cpp
+++ b/Engine/Source/Developer/TraceInsights/Private/Insights/MemoryProfiler/ViewModels/MemAllocGroupingByCallstack.cpp
@@ -118,7 +118,7 @@ void FMemAllocGroupingByCallstack::GroupNodes(const TArray<FTableTreeNodePtr>& N
 		const FMemoryAlloc* Alloc = MemAllocNode.GetMemAlloc();
 		if (Alloc)
 		{
-			const TraceServices::FCallstack* Callstack = bIsAllocCallstack ? Alloc->GetCallstack() : Alloc->GetFreeCallstack();
+			const TraceServices::FCallstack* Callstack = bIsAllocCallstack ? Alloc->GetAllocCallstack() : Alloc->GetFreeCallstack();
 
 			FCallstackGroup** FoundGroupPtrPtr = GroupMapByCallstack.Find(Callstack);
 			if (FoundGroupPtrPtr)
diff --git a/Engine/Source/Developer/TraceInsights/Private/Insights/MemoryProfiler/ViewModels/MemAllocNode.cpp b/Engine/Source/Developer/TraceInsights/Private/Insights/MemoryProfiler/ViewModels/MemAllocNode.cpp
--- a/Engine/Source/Developer/TraceInsights/Private/Insights/MemoryProfiler/ViewModels/MemAllocNode.cpp
+++ b/Engine/Source/Developer/TraceInsights/Private/Insights/MemoryProfiler/ViewModels/MemAllocNode.cpp
@@ -14,9 +14,9 @@ INSIGHTS_IMPLEMENT_RTTI(FMemAllocNode)
 
 ////////////////////////////////////////////////////////////////////////////////////////////////////
 
-uint32 FMemAllocNode::GetCallstackId() const
+uint32 FMemAllocNode::GetAllocCallstackId() const
 {
-	return IsValidMemAlloc() ? GetMemAllocChecked().GetCallstackId() : 0u;
+	return IsValidMemAlloc() ? GetMemAllocChecked().GetAllocCallstackId() : 0u;
 }
 
 ////////////////////////////////////////////////////////////////////////////////////////////////////
@@ -80,7 +80,7 @@ FText FMemAllocNode::GetFullCallstackOrSourceFiles(ECallstackType InCallstackTyp
 	const Insights::FMemoryAlloc& Alloc = GetMemAllocChecked();
 	const TraceServices::FCallstack* Callstack = 
 		InCallstackType == ECallstackType::AllocCallstack
-			? Alloc.GetCallstack() 
+			? Alloc.GetAllocCallstack() 
 			: Alloc.GetFreeCallstack();
 
 	if (!Callstack)
@@ -121,7 +121,7 @@ FText FMemAllocNode::GetTopFunctionOrSourceFile(ECallstackType InCallstackType,
 	const Insights::FMemoryAlloc& Alloc = GetMemAllocChecked();
 	const TraceServices::FCallstack* Callstack = 
 		InCallstackType == ECallstackType::AllocCallstack
-			? Alloc.GetCallstack() 
+			? Alloc.GetAllocCallstack() 
 			: Alloc.GetFreeCallstack();
 
 	if (!Callstack)
diff --git a/Engine/Source/Developer/TraceInsights/Private/Insights/MemoryProfiler/ViewModels/MemAllocNode.h b/Engine/Source/Developer/TraceInsights/Private/Insights/MemoryProfiler/ViewModels/MemAllocNode.h
--- a/Engine/Source/Developer/TraceInsights/Private/Insights/MemoryProfiler/ViewModels/MemAllocNode.h
+++ b/Engine/Source/Developer/TraceInsights/Private/Insights/MemoryProfiler/ViewModels/MemAllocNode.h
@@ -67,7 +67,7 @@ public:
 	const FMemoryAlloc* GetMemAlloc() const { return GetMemTableChecked().GetMemAlloc(GetRowIndex()); }
 	const FMemoryAlloc& GetMemAllocChecked() const { return GetMemTableChecked().GetMemAllocChecked(GetRowIndex()); }
 
-	uint32 GetCallstackId() const;
+	uint32 GetAllocCallstackId() const;
 	uint32 GetFreeCallstackId() const;
 
 	FText GetFullCallstack(ECallstackType InCallstackType) const;
diff --git a/Engine/Source/Developer/TraceInsights/Private/Insights/MemoryProfiler/ViewModels/MemAllocTable.cpp b/Engine/Source/Developer/TraceInsights/Private/Insights/MemoryProfiler/ViewModels/MemAllocTable.cpp
--- a/Engine/Source/Developer/TraceInsights/Private/Insights/MemoryProfiler/ViewModels/MemAllocTable.cpp
+++ b/Engine/Source/Developer/TraceInsights/Private/Insights/MemoryProfiler/ViewModels/MemAllocTable.cpp
@@ -32,6 +32,8 @@ const FName FMemAllocTableColumns::EventDistanceColumnId(TEXT("EventDistance"));
 const FName FMemAllocTableColumns::StartTimeColumnId(TEXT("StartTime"));
 const FName FMemAllocTableColumns::EndTimeColumnId(TEXT("EndTime"));
 const FName FMemAllocTableColumns::DurationColumnId(TEXT("Duration"));
+const FName FMemAllocTableColumns::AllocThreadColumnId(TEXT("AllocThread"));
+const FName FMemAllocTableColumns::FreeThreadColumnId(TEXT("FreeThread"));
 const FName FMemAllocTableColumns::AddressColumnId(TEXT("Address"));
 const FName FMemAllocTableColumns::MemoryPageColumnId(TEXT("MemoryPage"));
 const FName FMemAllocTableColumns::CountColumnId(TEXT("Count"));
@@ -42,9 +44,11 @@ const FName FMemAllocTableColumns::PackageColumnId(TEXT("Package"));
 const FName FMemAllocTableColumns::ClassNameColumnId(TEXT("ClassName"));
 const FName FMemAllocTableColumns::AllocFunctionColumnId(TEXT("AllocFunction"));
 const FName FMemAllocTableColumns::AllocSourceFileColumnId(TEXT("AllocSourceFile"));
+const FName FMemAllocTableColumns::AllocCallstackIdColumnId(TEXT("AllocCallstackId"));
 const FName FMemAllocTableColumns::AllocCallstackSizeColumnId(TEXT("AllocCallstackSize"));
 const FName FMemAllocTableColumns::FreeFunctionColumnId(TEXT("FreeFunction"));
 const FName FMemAllocTableColumns::FreeSourceFileColumnId(TEXT("FreeSourceFile"));
+const FName FMemAllocTableColumns::FreeCallstackIdColumnId(TEXT("FreeCallstackId"));
 const FName FMemAllocTableColumns::FreeCallstackSizeColumnId(TEXT("FreeCallstackSize"));
 
 ////////////////////////////////////////////////////////////////////////////////////////////////////
@@ -448,6 +452,118 @@ void FMemAllocTable::AddDefaultColumns()
 		AddColumn(ColumnRef);
 	}
 	//////////////////////////////////////////////////
+	// Alloc Thread Column
+	{
+		TSharedRef<FTableColumn> ColumnRef = MakeShared<FTableColumn>(FMemAllocTableColumns::AllocThreadColumnId);
+		FTableColumn& Column = *ColumnRef;
+
+		Column.SetIndex(ColumnIndex++);
+
+		Column.SetShortName(LOCTEXT("AllocThreadColumnName", "Alloc Thread"));
+		Column.SetTitleName(LOCTEXT("AllocThreadColumnTitle", "Alloc Thread"));
+		Column.SetDescription(LOCTEXT("AllocThreadColumnDesc", "The thread the allocation was made on"));
+
+		Column.SetFlags(ETableColumnFlags::CanBeHidden | ETableColumnFlags::CanBeFiltered);
+
+		Column.SetHorizontalAlignment(HAlign_Left);
+		Column.SetInitialWidth(80.0f);
+
+		Column.SetDataType(ETableCellDataType::Int64);
+
+		class FAllocThreadValueGetter : public FTableCellValueGetter
+		{
+		public:
+			virtual const TOptional<FTableCellValue> GetValue(const FTableColumn& Column, const FBaseTreeNode& Node) const override
+			{
+				if (Node.IsGroup())
+				{
+					const FTableTreeNode& NodePtr = static_cast<const FTableTreeNode&>(Node);
+					if (NodePtr.HasAggregatedValue(Column.GetId()))
+					{
+						return NodePtr.GetAggregatedValue(Column.GetId());
+					}
+				}
+				else //if (Node->Is<FMemAllocNode>())
+				{
+					const FMemAllocNode& MemAllocNode = static_cast<const FMemAllocNode&>(Node);
+					const FMemoryAlloc* Alloc = MemAllocNode.GetMemAlloc();
+					if (Alloc)
+					{
+						return FTableCellValue(static_cast<int64>(Alloc->GetAllocThreadId()));
+					}
+				}
+
+				return TOptional<FTableCellValue>();
+			}
+		};
+		TSharedRef<ITableCellValueGetter> Getter = MakeShared<FAllocThreadValueGetter>();
+		Column.SetValueGetter(Getter);
+
+		TSharedRef<ITableCellValueFormatter> Formatter = MakeShared<FInt64ValueFormatterAsNumber>();
+		Column.SetValueFormatter(Formatter);
+
+		TSharedRef<ITableCellValueSorter> Sorter = MakeShared<FSorterByInt64Value>(ColumnRef);
+		Column.SetValueSorter(Sorter);
+
+		AddColumn(ColumnRef);
+	}
+	//////////////////////////////////////////////////
+	// Free Thread Column
+	{
+		TSharedRef<FTableColumn> ColumnRef = MakeShared<FTableColumn>(FMemAllocTableColumns::FreeThreadColumnId);
+		FTableColumn& Column = *ColumnRef;
+
+		Column.SetIndex(ColumnIndex++);
+
+		Column.SetShortName(LOCTEXT("FreeThreadColumnName", "Free Thread"));
+		Column.SetTitleName(LOCTEXT("FreeThreadColumnTitle", "Free Thread"));
+		Column.SetDescription(LOCTEXT("FreeThreadColumnDesc", "The thread the allocation was freed on"));
+
+		Column.SetFlags(ETableColumnFlags::CanBeHidden | ETableColumnFlags::CanBeFiltered);
+
+		Column.SetHorizontalAlignment(HAlign_Left);
+		Column.SetInitialWidth(80.0f);
+
+		Column.SetDataType(ETableCellDataType::Int64);
+
+		class FAllocThreadValueGetter : public FTableCellValueGetter
+		{
+		public:
+			virtual const TOptional<FTableCellValue> GetValue(const FTableColumn& Column, const FBaseTreeNode& Node) const override
+			{
+				if (Node.IsGroup())
+				{
+					const FTableTreeNode& NodePtr = static_cast<const FTableTreeNode&>(Node);
+					if (NodePtr.HasAggregatedValue(Column.GetId()))
+					{
+						return NodePtr.GetAggregatedValue(Column.GetId());
+					}
+				}
+				else //if (Node->Is<FMemAllocNode>())
+				{
+					const FMemAllocNode& MemAllocNode = static_cast<const FMemAllocNode&>(Node);
+					const FMemoryAlloc* Alloc = MemAllocNode.GetMemAlloc();
+					if (Alloc)
+					{
+						return FTableCellValue(static_cast<int64>(Alloc->GetFreeThreadId()));
+					}
+				}
+
+				return TOptional<FTableCellValue>();
+			}
+		};
+		TSharedRef<ITableCellValueGetter> Getter = MakeShared<FAllocThreadValueGetter>();
+		Column.SetValueGetter(Getter);
+
+		TSharedRef<ITableCellValueFormatter> Formatter = MakeShared<FInt64ValueFormatterAsNumber>();
+		Column.SetValueFormatter(Formatter);
+
+		TSharedRef<ITableCellValueSorter> Sorter = MakeShared<FSorterByInt64Value>(ColumnRef);
+		Column.SetValueSorter(Sorter);
+
+		AddColumn(ColumnRef);
+	}
+	//////////////////////////////////////////////////
 	// Address Column
 	{
 		TSharedRef<FTableColumn> ColumnRef = MakeShared<FTableColumn>(FMemAllocTableColumns::AddressColumnId);
@@ -615,7 +731,7 @@ void FMemAllocTable::AddDefaultColumns()
 		Column.SetInitialSortMode(EColumnSortMode::Descending);
 
 		Column.SetAggregation(ETableColumnAggregation::Sum);
-
+		
 		AddColumn(ColumnRef);
 	}
 	//////////////////////////////////////////////////
@@ -915,7 +1031,194 @@ void FMemAllocTable::AddDefaultColumns()
 		AddColumn(ColumnRef);
 	}
 	//////////////////////////////////////////////////
-	// Top Function Column
+
+	class FCallstackValueFormatter : public FTableCellValueFormatter
+	{
+	public:
+		FCallstackValueFormatter(FMemAllocNode::ECallstackType InCallstackType)
+			: CallstackType(InCallstackType)
+		{
+		}
+
+		virtual FText FormatValue(const TOptional<FTableCellValue>& InValue) const override
+		{
+			return InValue.IsSet() ? FText::AsNumber(InValue.GetValue().Int64) : FText::GetEmpty();
+		}
+
+		FText GetHeaderText(const FMemAllocNode& MemAllocNode) const
+		{
+			FText HeaderText;
+			if (CallstackType == FMemAllocNode::ECallstackType::AllocCallstack)
+			{
+				int64 CallstackSize = 0;
+				const FMemoryAlloc* Alloc = MemAllocNode.GetMemAlloc();
+				if (Alloc)
+				{
+					const TraceServices::FCallstack* Callstack = Alloc->GetAllocCallstack();
+					if (Callstack && (Callstack->Num() != 1 || Callstack->Addr(0) != 0))
+					{
+						CallstackSize = Callstack->Num();
+					}
+				}
+				HeaderText = FText::Format(LOCTEXT("AllocCallstackHeaderFmt", "AllocCallstackId={0} ({1} frames)"),
+					FText::AsNumber(MemAllocNode.GetAllocCallstackId(), &FNumberFormattingOptions::DefaultNoGrouping()),
+					FText::AsNumber(CallstackSize));
+			}
+			else
+			{
+				int64 CallstackSize = 0;
+				const FMemoryAlloc* Alloc = MemAllocNode.GetMemAlloc();
+				if (Alloc)
+				{
+					const TraceServices::FCallstack* Callstack = Alloc->GetFreeCallstack();
+					if (Callstack && (Callstack->Num() != 1 || Callstack->Addr(0) != 0))
+					{
+						CallstackSize = Callstack->Num();
+					}
+				}
+				HeaderText = FText::Format(LOCTEXT("FreeCallstackHeaderFmt", "FreeCallstackId={0} ({1} frames)"),
+					FText::AsNumber(MemAllocNode.GetFreeCallstackId(), &FNumberFormattingOptions::DefaultNoGrouping()),
+					FText::AsNumber(CallstackSize));
+			}
+			return HeaderText;
+		}
+
+		virtual TSharedPtr<IToolTip> GetCustomTooltip(const FTableColumn& Column, const FBaseTreeNode& Node) const override
+		{
+			const FMemAllocNode& MemAllocNode = static_cast<const FMemAllocNode&>(Node);
+
+			FText HeaderText = GetHeaderText(MemAllocNode);
+			return SNew(SToolTip)
+				.Visibility(TAttribute<EVisibility>::Create(TAttribute<EVisibility>::FGetter::CreateStatic(&FTableCellValueFormatter::GetTooltipVisibility)))
+				[
+					SNew(SVerticalBox)
+					+ SVerticalBox::Slot()
+					.AutoHeight()
+					.Padding(2.0f)
+					[
+						SNew(STextBlock)
+						.Text(HeaderText)
+						.ColorAndOpacity(FSlateColor(EStyleColor::White25))
+					]
+					+ SVerticalBox::Slot()
+					.AutoHeight()
+					.Padding(2.0f)
+					[
+						SNew(STextBlock)
+						.Text(&MemAllocNode, &FMemAllocNode::GetFullCallstack, CallstackType)
+					]
+				];
+		}
+
+		virtual FText FormatValueForGrouping(const FTableColumn& Column, const FBaseTreeNode& Node) const override
+		{
+			return FormatValue(Column.GetValue(Node));
+		}
+
+		virtual FText CopyTooltip(const FTableColumn& Column, const FBaseTreeNode& Node) const override
+		{
+			const FMemAllocNode& MemAllocNode = static_cast<const FMemAllocNode&>(Node);
+			return MemAllocNode.GetFullCallstack(CallstackType);
+		}
+
+	protected:
+		FMemAllocNode::ECallstackType CallstackType;
+	};
+
+	//////////////////////////////////////////////////
+
+	class FFunctionValueFormatter : public FTextValueFormatter
+	{
+	public:
+		FFunctionValueFormatter(FMemAllocNode::ECallstackType InCallstackType)
+			: CallstackType(InCallstackType)
+		{
+		}
+
+		virtual TSharedPtr<IToolTip> GetCustomTooltip(const FTableColumn& Column, const FBaseTreeNode& Node) const override
+		{
+			const FMemAllocNode& MemAllocNode = static_cast<const FMemAllocNode&>(Node);
+
+			return SNew(SToolTip)
+				.Visibility(TAttribute<EVisibility>::Create(TAttribute<EVisibility>::FGetter::CreateStatic(&FTableCellValueFormatter::GetTooltipVisibility)))
+				[
+					SNew(SVerticalBox)
+					+ SVerticalBox::Slot()
+					.AutoHeight()
+					.Padding(2.0f)
+					[
+						SNew(STextBlock)
+						.Text(&MemAllocNode, &FMemAllocNode::GetTopFunction, CallstackType)
+						.ColorAndOpacity(FSlateColor(EStyleColor::AccentBlue))
+					]
+					+ SVerticalBox::Slot()
+					.AutoHeight()
+					.Padding(2.0f, 6.0f, 2.0f, 2.0f)
+					[
+						SNew(STextBlock)
+						.Text(&MemAllocNode, &FMemAllocNode::GetFullCallstack, CallstackType)
+					]
+				];
+		}
+
+		virtual FText FormatValueForTooltip(const FTableColumn& Column, const FBaseTreeNode& Node) const override
+		{
+			const FMemAllocNode& MemAllocNode = static_cast<const FMemAllocNode&>(Node);
+			return MemAllocNode.GetTopFunction(CallstackType);
+		}
+
+	protected:
+		FMemAllocNode::ECallstackType CallstackType;
+	};
+
+	//////////////////////////////////////////////////
+
+	class FSourceValueFormatter : public FTextValueFormatter
+	{
+	public:
+		FSourceValueFormatter(FMemAllocNode::ECallstackType InCallstackType)
+			: CallstackType(InCallstackType)
+		{
+		}
+
+		virtual TSharedPtr<IToolTip> GetCustomTooltip(const FTableColumn& Column, const FBaseTreeNode& Node) const override
+		{
+			const FMemAllocNode& MemAllocNode = static_cast<const FMemAllocNode&>(Node);
+
+			return SNew(SToolTip)
+				.Visibility(TAttribute<EVisibility>::Create(TAttribute<EVisibility>::FGetter::CreateStatic(&FTableCellValueFormatter::GetTooltipVisibility)))
+				[
+					SNew(SVerticalBox)
+					+ SVerticalBox::Slot()
+					.AutoHeight()
+					.Padding(2.0f)
+					[
+						SNew(STextBlock)
+						.Text(&MemAllocNode, &FMemAllocNode::GetTopSourceFileEx, CallstackType)
+						.ColorAndOpacity(FSlateColor(EStyleColor::AccentBlue))
+					]
+					+ SVerticalBox::Slot()
+					.AutoHeight()
+					.Padding(2.0f, 6.0f, 2.0f, 2.0f)
+					[
+						SNew(STextBlock)
+						.Text(&MemAllocNode, &FMemAllocNode::GetFullCallstackSourceFiles, CallstackType)
+					]
+				];
+		}
+
+		virtual FText FormatValueForTooltip(const FTableColumn& Column, const FBaseTreeNode& Node) const override
+		{
+			const FMemAllocNode& MemAllocNode = static_cast<const FMemAllocNode&>(Node);
+			return MemAllocNode.GetTopSourceFileEx(CallstackType);
+		}
+
+	protected:
+		FMemAllocNode::ECallstackType CallstackType;
+	};
+
+	//////////////////////////////////////////////////
+	// Top Function (Alloc Callstack) Column
 	{
 		TSharedRef<FTableColumn> ColumnRef = MakeShared<FTableColumn>(FMemAllocTableColumns::AllocFunctionColumnId);
 		FTableColumn& Column = *ColumnRef;
@@ -949,7 +1252,7 @@ void FMemAllocTable::AddDefaultColumns()
 				else //if (Node->Is<FMemAllocNode>())
 				{
 					const FMemAllocNode& MemAllocNode = static_cast<const FMemAllocNode&>(Node);
-					const uint64 CallstackId = MemAllocNode.GetCallstackId();
+					const uint64 CallstackId = MemAllocNode.GetAllocCallstackId();
 					const FText CallstackText = MemAllocNode.GetTopFunction(FMemAllocNode::ECallstackType::AllocCallstack);
 					return FTableCellValue(CallstackText, CallstackId);
 				}
@@ -962,7 +1265,7 @@ void FMemAllocTable::AddDefaultColumns()
 				if (!Node.IsGroup()) //if (Node->Is<FMemAllocNode>())
 				{
 					const FMemAllocNode& MemAllocNode = static_cast<const FMemAllocNode&>(Node);
-					return MemAllocNode.GetCallstackId();
+					return MemAllocNode.GetAllocCallstackId();
 				}
 
 				return 0;
@@ -971,43 +1274,7 @@ void FMemAllocTable::AddDefaultColumns()
 		TSharedRef<ITableCellValueGetter> Getter = MakeShared<FFunctionValueGetter>();
 		Column.SetValueGetter(Getter);
 
-		class FFunctionValueFormatter : public FTextValueFormatter
-		{
-		public:
-			virtual TSharedPtr<IToolTip> GetCustomTooltip(const FTableColumn& Column, const FBaseTreeNode& Node) const override
-			{
-				const FMemAllocNode& MemAllocNode = static_cast<const FMemAllocNode&>(Node);
-
-				return SNew(SToolTip)
-					.Visibility(TAttribute<EVisibility>::Create(TAttribute<EVisibility>::FGetter::CreateStatic(&FTableCellValueFormatter::GetTooltipVisibility)))
-					[
-						SNew(SVerticalBox)
-						+ SVerticalBox::Slot()
-						.AutoHeight()
-						.Padding(2.0f)
-						[
-							SNew(STextBlock)
-							.Text(&MemAllocNode, &FMemAllocNode::GetTopFunction, FMemAllocNode::ECallstackType::AllocCallstack)
-							.ColorAndOpacity(FSlateColor(EStyleColor::AccentBlue))
-						]
-						+ SVerticalBox::Slot()
-						.AutoHeight()
-						.Padding(2.0f, 6.0f, 2.0f, 2.0f)
-						[
-							SNew(STextBlock)
-							.Text(&MemAllocNode, &FMemAllocNode::GetFullCallstack, FMemAllocNode::ECallstackType::AllocCallstack)
-						]
-					];
-			}
-
-			virtual FText FormatValueForTooltip(const FTableColumn& Column, const FBaseTreeNode& Node) const override
-			{
-				const FMemAllocNode& MemAllocNode = static_cast<const FMemAllocNode&>(Node);
-				return MemAllocNode.GetTopFunction(FMemAllocNode::ECallstackType::AllocCallstack);
-			}
-		};
-
-		TSharedRef<ITableCellValueFormatter> Formatter = MakeShared<FFunctionValueFormatter>();
+		TSharedRef<ITableCellValueFormatter> Formatter = MakeShared<FFunctionValueFormatter>(FMemAllocNode::ECallstackType::AllocCallstack);
 		Column.SetValueFormatter(Formatter);
 
 		TSharedRef<ITableCellValueSorter> Sorter = MakeShared<FSorterByTextValueWithId>(ColumnRef);
@@ -1016,7 +1283,7 @@ void FMemAllocTable::AddDefaultColumns()
 		AddColumn(ColumnRef);
 	}
 	//////////////////////////////////////////////////
-	// Source File Column
+	// Top Source File (Alloc Callstack) Column
 	{
 		TSharedRef<FTableColumn> ColumnRef = MakeShared<FTableColumn>(FMemAllocTableColumns::AllocSourceFileColumnId);
 		FTableColumn& Column = *ColumnRef;
@@ -1050,7 +1317,7 @@ void FMemAllocTable::AddDefaultColumns()
 				else //if (Node->Is<FMemAllocNode>())
 				{
 					const FMemAllocNode& MemAllocNode = static_cast<const FMemAllocNode&>(Node);
-					const uint64 CallstackId = MemAllocNode.GetCallstackId();
+					const uint64 CallstackId = MemAllocNode.GetAllocCallstackId();
 					const FText CallstackText = MemAllocNode.GetTopSourceFile(FMemAllocNode::ECallstackType::AllocCallstack);
 					return FTableCellValue(CallstackText, CallstackId);
 				}
@@ -1063,7 +1330,7 @@ void FMemAllocTable::AddDefaultColumns()
 				if (!Node.IsGroup()) //if (Node->Is<FMemAllocNode>())
 				{
 					const FMemAllocNode& MemAllocNode = static_cast<const FMemAllocNode&>(Node);
-					return MemAllocNode.GetCallstackId();
+					return MemAllocNode.GetAllocCallstackId();
 				}
 				return 0;
 			}
@@ -1071,63 +1338,93 @@ void FMemAllocTable::AddDefaultColumns()
 		TSharedRef<ITableCellValueGetter> Getter = MakeShared<FSourceValueGetter>();
 		Column.SetValueGetter(Getter);
 
-		class FSourceValueFormatter : public FTextValueFormatter
+		TSharedRef<ITableCellValueFormatter> Formatter = MakeShared<FSourceValueFormatter>(FMemAllocNode::ECallstackType::AllocCallstack);
+		Column.SetValueFormatter(Formatter);
+
+		TSharedRef<ITableCellValueSorter> Sorter = MakeShared<FSorterByTextValueWithId>(ColumnRef);
+		Column.SetValueSorter(Sorter);
+
+		AddColumn(ColumnRef);
+	}
+	//////////////////////////////////////////////////
+	// Alloc Callstack Id Column
+	{
+		TSharedRef<FTableColumn> ColumnRef = MakeShared<FTableColumn>(FMemAllocTableColumns::AllocCallstackIdColumnId);
+		FTableColumn& Column = *ColumnRef;
+
+		Column.SetIndex(ColumnIndex++);
+
+		Column.SetShortName(LOCTEXT("CallstackIdColumnName_Alloc", "Alloc Callstack"));
+		Column.SetTitleName(LOCTEXT("CallstackIdColumnTitle_Alloc", "Alloc Callstack Id"));
+		Column.SetDescription(LOCTEXT("CallstackIdColumnDesc_Alloc", "Id of alloc callstack.\nTooltip shows the entire alloc callstack."));
+
+		Column.SetFlags(ETableColumnFlags::ShouldBeVisible | ETableColumnFlags::CanBeHidden | ETableColumnFlags::CanBeFiltered);
+
+		Column.SetHorizontalAlignment(HAlign_Left);
+		Column.SetInitialWidth(100.0f);
+
+		Column.SetDataType(ETableCellDataType::Int64);
+
+		class FCallstackIdValueGetter : public FTableCellValueGetter
 		{
 		public:
-			virtual TSharedPtr<IToolTip> GetCustomTooltip(const FTableColumn& Column, const FBaseTreeNode& Node) const override
+			virtual const TOptional<FTableCellValue> GetValue(const FTableColumn& Column, const FBaseTreeNode& Node) const override
 			{
-				const FMemAllocNode& MemAllocNode = static_cast<const FMemAllocNode&>(Node);
+				if (Node.IsGroup())
+				{
+					const FTableTreeNode& NodePtr = static_cast<const FTableTreeNode&>(Node);
+					if (NodePtr.HasAggregatedValue(Column.GetId()))
+					{
+						return NodePtr.GetAggregatedValue(Column.GetId());
+					}
+				}
+				else //if (Node->Is<FMemAllocNode>())
+				{
+					const FMemAllocNode& MemAllocNode = static_cast<const FMemAllocNode&>(Node);
+					const FMemoryAlloc* Alloc = MemAllocNode.GetMemAlloc();
+					if (Alloc)
+					{
+						return FTableCellValue(int64(Alloc->GetAllocCallstackId()));
+					}
+				}
 
-				return SNew(SToolTip)
-					.Visibility(TAttribute<EVisibility>::Create(TAttribute<EVisibility>::FGetter::CreateStatic(&FTableCellValueFormatter::GetTooltipVisibility)))
-					[
-						SNew(SVerticalBox)
-						+ SVerticalBox::Slot()
-						.AutoHeight()
-						.Padding(2.0f)
-						[
-							SNew(STextBlock)
-							.Text(&MemAllocNode, &FMemAllocNode::GetTopSourceFileEx, FMemAllocNode::ECallstackType::AllocCallstack)
-							.ColorAndOpacity(FSlateColor(EStyleColor::AccentBlue))
-						]
-						+ SVerticalBox::Slot()
-						.AutoHeight()
-						.Padding(2.0f, 6.0f, 2.0f, 2.0f)
-						[
-							SNew(STextBlock)
-							.Text(&MemAllocNode, &FMemAllocNode::GetFullCallstackSourceFiles, FMemAllocNode::ECallstackType::AllocCallstack)
-						]
-					];
+				return TOptional<FTableCellValue>();
 			}
 
-			virtual FText FormatValueForTooltip(const FTableColumn& Column, const FBaseTreeNode& Node) const override
+			virtual uint64 GetValueId(const FTableColumn& Column, const FBaseTreeNode& Node) const override
 			{
-				const FMemAllocNode& MemAllocNode = static_cast<const FMemAllocNode&>(Node);
-				return MemAllocNode.GetTopSourceFileEx(FMemAllocNode::ECallstackType::AllocCallstack);
+				if (!Node.IsGroup()) //if (Node->Is<FMemAllocNode>())
+				{
+					const FMemAllocNode& MemAllocNode = static_cast<const FMemAllocNode&>(Node);
+					return MemAllocNode.GetAllocCallstackId();
+				}
+				return 0;
 			}
 		};
+		TSharedRef<ITableCellValueGetter> Getter = MakeShared<FCallstackIdValueGetter>();
+		Column.SetValueGetter(Getter);
 
-		TSharedRef<ITableCellValueFormatter> Formatter = MakeShared<FSourceValueFormatter>();
+		TSharedRef<ITableCellValueFormatter> Formatter = MakeShared<FCallstackValueFormatter>(FMemAllocNode::ECallstackType::AllocCallstack);
 		Column.SetValueFormatter(Formatter);
 
-		TSharedRef<ITableCellValueSorter> Sorter = MakeShared<FSorterByTextValueWithId>(ColumnRef);
+		TSharedRef<ITableCellValueSorter> Sorter = MakeShared<FSorterByInt64Value>(ColumnRef);
 		Column.SetValueSorter(Sorter);
 
 		AddColumn(ColumnRef);
 	}
 	//////////////////////////////////////////////////
-	// Callstack (Size) Column
+	// Alloc Callstack Size Column
 	{
 		TSharedRef<FTableColumn> ColumnRef = MakeShared<FTableColumn>(FMemAllocTableColumns::AllocCallstackSizeColumnId);
 		FTableColumn& Column = *ColumnRef;
 
 		Column.SetIndex(ColumnIndex++);
 
-		Column.SetShortName(LOCTEXT("CallstackSizeColumnName_Alloc", "Alloc Callstack"));
-		Column.SetTitleName(LOCTEXT("CallstackSizeColumnTitle_Alloc", "Alloc Callstack"));
+		Column.SetShortName(LOCTEXT("CallstackSizeColumnName_Alloc", "Alloc Callstack Size"));
+		Column.SetTitleName(LOCTEXT("CallstackSizeColumnTitle_Alloc", "Alloc Callstack Size"));
 		Column.SetDescription(LOCTEXT("CallstackSizeColumnDesc_Alloc", "Number of alloc callstack frames.\nTooltip shows the entire alloc callstack."));
 
-		Column.SetFlags(ETableColumnFlags::ShouldBeVisible | ETableColumnFlags::CanBeHidden | ETableColumnFlags::CanBeFiltered);
+		Column.SetFlags(ETableColumnFlags::CanBeHidden | ETableColumnFlags::CanBeFiltered);
 
 		Column.SetHorizontalAlignment(HAlign_Left);
 		Column.SetInitialWidth(100.0f);
@@ -1153,7 +1450,7 @@ void FMemAllocTable::AddDefaultColumns()
 					const FMemoryAlloc* Alloc = MemAllocNode.GetMemAlloc();
 					if (Alloc)
 					{
-						const TraceServices::FCallstack* Callstack = Alloc->GetCallstack();
+						const TraceServices::FCallstack* Callstack = Alloc->GetAllocCallstack();
 						const int64 CallstackSize = (Callstack && (Callstack->Num() != 1 || Callstack->Addr(0) != 0)) ? Callstack->Num() : 0;
 						return FTableCellValue(CallstackSize);
 					}
@@ -1167,7 +1464,7 @@ void FMemAllocTable::AddDefaultColumns()
 				if (!Node.IsGroup()) //if (Node->Is<FMemAllocNode>())
 				{
 					const FMemAllocNode& MemAllocNode = static_cast<const FMemAllocNode&>(Node);
-					return MemAllocNode.GetCallstackId();
+					return MemAllocNode.GetAllocCallstackId();
 				}
 				return 0;
 			}
@@ -1175,57 +1472,7 @@ void FMemAllocTable::AddDefaultColumns()
 		TSharedRef<ITableCellValueGetter> Getter = MakeShared<FCallstackSizeValueGetter>();
 		Column.SetValueGetter(Getter);
 
-		class FCallstackValueFormatter : public FTableCellValueFormatter
-		{
-		public:
-			virtual FText FormatValue(const TOptional<FTableCellValue>& InValue) const override
-			{
-				return InValue.IsSet() ? FText::AsNumber(InValue.GetValue().Int64) : FText::GetEmpty();
-			}
-
-			virtual TSharedPtr<IToolTip> GetCustomTooltip(const FTableColumn& Column, const FBaseTreeNode& Node) const override
-			{
-				const FMemAllocNode& MemAllocNode = static_cast<const FMemAllocNode&>(Node);
-
-				FText HeaderText = FText::Format(LOCTEXT("CallstackHeaderFmt", "CallstackId={0} FreeCallstackId={1}"),
-					FText::AsNumber(MemAllocNode.GetCallstackId(), &FNumberFormattingOptions::DefaultNoGrouping()),
-					FText::AsNumber(MemAllocNode.GetFreeCallstackId(), &FNumberFormattingOptions::DefaultNoGrouping()));
-
-				return SNew(SToolTip)
-					.Visibility(TAttribute<EVisibility>::Create(TAttribute<EVisibility>::FGetter::CreateStatic(&FTableCellValueFormatter::GetTooltipVisibility)))
-					[
-						SNew(SVerticalBox)
-						+ SVerticalBox::Slot()
-						.AutoHeight()
-						.Padding(2.0f)
-						[
-							SNew(STextBlock)
-							.Text(HeaderText)
-							.ColorAndOpacity(FSlateColor(EStyleColor::White25))
-						]
-						+ SVerticalBox::Slot()
-						.AutoHeight()
-						.Padding(2.0f)
-						[
-							SNew(STextBlock)
-							.Text(&MemAllocNode, &FMemAllocNode::GetFullCallstack, FMemAllocNode::ECallstackType::AllocCallstack)
-						]
-					];
-			}
-
-			virtual FText FormatValueForTooltip(const FTableColumn& Column, const FBaseTreeNode& Node) const override
-			{
-				const FMemAllocNode& MemAllocNode = static_cast<const FMemAllocNode&>(Node);
-				return MemAllocNode.GetFullCallstack(FMemAllocNode::ECallstackType::AllocCallstack);
-			}
-
-			virtual FText FormatValueForGrouping(const FTableColumn& Column, const FBaseTreeNode& Node) const override
-			{
-				return FormatValue(Column.GetValue(Node));
-			}
-		};
-
-		TSharedRef<ITableCellValueFormatter> Formatter = MakeShared<FCallstackValueFormatter>();
+		TSharedRef<ITableCellValueFormatter> Formatter = MakeShared<FCallstackValueFormatter>(FMemAllocNode::ECallstackType::AllocCallstack);
 		Column.SetValueFormatter(Formatter);
 
 		TSharedRef<ITableCellValueSorter> Sorter = MakeShared<FSorterByInt64Value>(ColumnRef);
@@ -1234,7 +1481,7 @@ void FMemAllocTable::AddDefaultColumns()
 		AddColumn(ColumnRef);
 	}
 	//////////////////////////////////////////////////
-	// Free Top Function Column
+	// Top Function (Free Callstack) Column
 	{
 		TSharedRef<FTableColumn> ColumnRef = MakeShared<FTableColumn>(FMemAllocTableColumns::FreeFunctionColumnId);
 		FTableColumn& Column = *ColumnRef;
@@ -1268,7 +1515,7 @@ void FMemAllocTable::AddDefaultColumns()
 				else //if (Node->Is<FMemAllocNode>())
 				{
 					const FMemAllocNode& MemAllocNode = static_cast<const FMemAllocNode&>(Node);
-					const uint64 CallstackId = MemAllocNode.GetCallstackId();
+					const uint64 CallstackId = MemAllocNode.GetFreeCallstackId();
 					const FText CallstackText = MemAllocNode.GetTopFunction(FMemAllocNode::ECallstackType::FreeCallstack);
 					return FTableCellValue(CallstackText, CallstackId);
 				}
@@ -1281,7 +1528,7 @@ void FMemAllocTable::AddDefaultColumns()
 				if (!Node.IsGroup()) //if (Node->Is<FMemAllocNode>())
 				{
 					const FMemAllocNode& MemAllocNode = static_cast<const FMemAllocNode&>(Node);
-					return MemAllocNode.GetCallstackId();
+					return MemAllocNode.GetFreeCallstackId();
 				}
 
 				return 0;
@@ -1290,43 +1537,7 @@ void FMemAllocTable::AddDefaultColumns()
 		TSharedRef<ITableCellValueGetter> Getter = MakeShared<FFunctionValueGetter>();
 		Column.SetValueGetter(Getter);
 
-		class FFunctionValueFormatter : public FTextValueFormatter
-		{
-		public:
-			virtual TSharedPtr<IToolTip> GetCustomTooltip(const FTableColumn& Column, const FBaseTreeNode& Node) const override
-			{
-				const FMemAllocNode& MemAllocNode = static_cast<const FMemAllocNode&>(Node);
-
-				return SNew(SToolTip)
-					.Visibility(TAttribute<EVisibility>::Create(TAttribute<EVisibility>::FGetter::CreateStatic(&FTableCellValueFormatter::GetTooltipVisibility)))
-					[
-						SNew(SVerticalBox)
-						+ SVerticalBox::Slot()
-						.AutoHeight()
-						.Padding(2.0f)
-						[
-							SNew(STextBlock)
-							.Text(&MemAllocNode, &FMemAllocNode::GetTopFunction, FMemAllocNode::ECallstackType::FreeCallstack)
-							.ColorAndOpacity(FSlateColor(EStyleColor::AccentBlue))
-						]
-						+ SVerticalBox::Slot()
-						.AutoHeight()
-						.Padding(2.0f, 6.0f, 2.0f, 2.0f)
-						[
-							SNew(STextBlock)
-							.Text(&MemAllocNode, &FMemAllocNode::GetFullCallstack, FMemAllocNode::ECallstackType::FreeCallstack)
-						]
-					];
-			}
-
-			virtual FText FormatValueForTooltip(const FTableColumn& Column, const FBaseTreeNode& Node) const override
-			{
-				const FMemAllocNode& MemAllocNode = static_cast<const FMemAllocNode&>(Node);
-				return MemAllocNode.GetTopFunction(FMemAllocNode::ECallstackType::FreeCallstack);
-			}
-		};
-
-		TSharedRef<ITableCellValueFormatter> Formatter = MakeShared<FFunctionValueFormatter>();
+		TSharedRef<ITableCellValueFormatter> Formatter = MakeShared<FFunctionValueFormatter>(FMemAllocNode::ECallstackType::FreeCallstack);
 		Column.SetValueFormatter(Formatter);
 
 		TSharedRef<ITableCellValueSorter> Sorter = MakeShared<FSorterByTextValueWithId>(ColumnRef);
@@ -1335,7 +1546,7 @@ void FMemAllocTable::AddDefaultColumns()
 		AddColumn(ColumnRef);
 	}
 	//////////////////////////////////////////////////
-	// Free Source File Column
+	// Top Source File (Free Callstack) Column
 	{
 		TSharedRef<FTableColumn> ColumnRef = MakeShared<FTableColumn>(FMemAllocTableColumns::FreeSourceFileColumnId);
 		FTableColumn& Column = *ColumnRef;
@@ -1369,7 +1580,7 @@ void FMemAllocTable::AddDefaultColumns()
 				else //if (Node->Is<FMemAllocNode>())
 				{
 					const FMemAllocNode& MemAllocNode = static_cast<const FMemAllocNode&>(Node);
-					const uint64 CallstackId = MemAllocNode.GetCallstackId();
+					const uint64 CallstackId = MemAllocNode.GetFreeCallstackId();
 					const FText CallstackText = MemAllocNode.GetTopSourceFile(FMemAllocNode::ECallstackType::FreeCallstack);
 					return FTableCellValue(CallstackText, CallstackId);
 				}
@@ -1382,7 +1593,7 @@ void FMemAllocTable::AddDefaultColumns()
 				if (!Node.IsGroup()) //if (Node->Is<FMemAllocNode>())
 				{
 					const FMemAllocNode& MemAllocNode = static_cast<const FMemAllocNode&>(Node);
-					return MemAllocNode.GetCallstackId();
+					return MemAllocNode.GetFreeCallstackId();
 				}
 				return 0;
 			}
@@ -1390,63 +1601,93 @@ void FMemAllocTable::AddDefaultColumns()
 		TSharedRef<ITableCellValueGetter> Getter = MakeShared<FSourceValueGetter>();
 		Column.SetValueGetter(Getter);
 
-		class FSourceValueFormatter : public FTextValueFormatter
+		TSharedRef<ITableCellValueFormatter> Formatter = MakeShared<FSourceValueFormatter>(FMemAllocNode::ECallstackType::FreeCallstack);
+		Column.SetValueFormatter(Formatter);
+
+		TSharedRef<ITableCellValueSorter> Sorter = MakeShared<FSorterByTextValueWithId>(ColumnRef);
+		Column.SetValueSorter(Sorter);
+
+		AddColumn(ColumnRef);
+	}
+	//////////////////////////////////////////////////
+	// Free Callstack Id Column
+	{
+		TSharedRef<FTableColumn> ColumnRef = MakeShared<FTableColumn>(FMemAllocTableColumns::FreeCallstackIdColumnId);
+		FTableColumn& Column = *ColumnRef;
+
+		Column.SetIndex(ColumnIndex++);
+
+		Column.SetShortName(LOCTEXT("CallstackIdColumnName_Free", "Free Callstack"));
+		Column.SetTitleName(LOCTEXT("CallstackIdColumnTitle_Free", "Free Callstack Id"));
+		Column.SetDescription(LOCTEXT("CallstackIdColumnDesc_Free", "Id of free callstack.\nTooltip shows the entire alloc callstack."));
+
+		Column.SetFlags(ETableColumnFlags::ShouldBeVisible | ETableColumnFlags::CanBeHidden | ETableColumnFlags::CanBeFiltered);
+
+		Column.SetHorizontalAlignment(HAlign_Left);
+		Column.SetInitialWidth(100.0f);
+
+		Column.SetDataType(ETableCellDataType::Int64);
+
+		class FCallstackIdValueGetter : public FTableCellValueGetter
 		{
 		public:
-			virtual TSharedPtr<IToolTip> GetCustomTooltip(const FTableColumn& Column, const FBaseTreeNode& Node) const override
+			virtual const TOptional<FTableCellValue> GetValue(const FTableColumn& Column, const FBaseTreeNode& Node) const override
 			{
-				const FMemAllocNode& MemAllocNode = static_cast<const FMemAllocNode&>(Node);
+				if (Node.IsGroup())
+				{
+					const FTableTreeNode& NodePtr = static_cast<const FTableTreeNode&>(Node);
+					if (NodePtr.HasAggregatedValue(Column.GetId()))
+					{
+						return NodePtr.GetAggregatedValue(Column.GetId());
+					}
+				}
+				else //if (Node->Is<FMemAllocNode>())
+				{
+					const FMemAllocNode& MemAllocNode = static_cast<const FMemAllocNode&>(Node);
+					const FMemoryAlloc* Alloc = MemAllocNode.GetMemAlloc();
+					if (Alloc)
+					{
+						return FTableCellValue(int64(Alloc->GetFreeCallstackId()));
+					}
+				}
 
-				return SNew(SToolTip)
-					.Visibility(TAttribute<EVisibility>::Create(TAttribute<EVisibility>::FGetter::CreateStatic(&FTableCellValueFormatter::GetTooltipVisibility)))
-					[
-						SNew(SVerticalBox)
-						+ SVerticalBox::Slot()
-						.AutoHeight()
-						.Padding(2.0f)
-						[
-							SNew(STextBlock)
-							.Text(&MemAllocNode, &FMemAllocNode::GetTopSourceFileEx, FMemAllocNode::ECallstackType::FreeCallstack)
-							.ColorAndOpacity(FSlateColor(EStyleColor::AccentBlue))
-						]
-						+ SVerticalBox::Slot()
-						.AutoHeight()
-						.Padding(2.0f, 6.0f, 2.0f, 2.0f)
-						[
-							SNew(STextBlock)
-							.Text(&MemAllocNode, &FMemAllocNode::GetFullCallstackSourceFiles, FMemAllocNode::ECallstackType::FreeCallstack)
-						]
-					];
+				return TOptional<FTableCellValue>();
 			}
 
-			virtual FText FormatValueForTooltip(const FTableColumn& Column, const FBaseTreeNode& Node) const override
+			virtual uint64 GetValueId(const FTableColumn& Column, const FBaseTreeNode& Node) const override
 			{
-				const FMemAllocNode& MemAllocNode = static_cast<const FMemAllocNode&>(Node);
-				return MemAllocNode.GetTopSourceFileEx(FMemAllocNode::ECallstackType::FreeCallstack);
+				if (!Node.IsGroup()) //if (Node->Is<FMemAllocNode>())
+				{
+					const FMemAllocNode& MemAllocNode = static_cast<const FMemAllocNode&>(Node);
+					return MemAllocNode.GetFreeCallstackId();
+				}
+				return 0;
 			}
 		};
+		TSharedRef<ITableCellValueGetter> Getter = MakeShared<FCallstackIdValueGetter>();
+		Column.SetValueGetter(Getter);
 
-		TSharedRef<ITableCellValueFormatter> Formatter = MakeShared<FSourceValueFormatter>();
+		TSharedRef<ITableCellValueFormatter> Formatter = MakeShared<FCallstackValueFormatter>(FMemAllocNode::ECallstackType::FreeCallstack);
 		Column.SetValueFormatter(Formatter);
 
-		TSharedRef<ITableCellValueSorter> Sorter = MakeShared<FSorterByTextValueWithId>(ColumnRef);
+		TSharedRef<ITableCellValueSorter> Sorter = MakeShared<FSorterByInt64Value>(ColumnRef);
 		Column.SetValueSorter(Sorter);
 
 		AddColumn(ColumnRef);
 	}
 	//////////////////////////////////////////////////
-	// Free Callstack (Size) Column
+	// Free Callstack Size Column
 	{
 		TSharedRef<FTableColumn> ColumnRef = MakeShared<FTableColumn>(FMemAllocTableColumns::FreeCallstackSizeColumnId);
 		FTableColumn& Column = *ColumnRef;
 
 		Column.SetIndex(ColumnIndex++);
 
-		Column.SetShortName(LOCTEXT("CallstackSizeColumnName_Free", "Free Callstack"));
-		Column.SetTitleName(LOCTEXT("CallstackSizeColumnTitle_Free", "Free Callstack"));
+		Column.SetShortName(LOCTEXT("CallstackSizeColumnName_Free", "Free Callstack Size"));
+		Column.SetTitleName(LOCTEXT("CallstackSizeColumnTitle_Free", "Free Callstack Size"));
 		Column.SetDescription(LOCTEXT("CallstackSizeColumnDesc_Free", "Number of free callstack frames.\nTooltip shows the entire free callstack."));
 
-		Column.SetFlags(ETableColumnFlags::ShouldBeVisible | ETableColumnFlags::CanBeHidden | ETableColumnFlags::CanBeFiltered);
+		Column.SetFlags(ETableColumnFlags::CanBeHidden | ETableColumnFlags::CanBeFiltered);
 
 		Column.SetHorizontalAlignment(HAlign_Left);
 		Column.SetInitialWidth(100.0f);
@@ -1472,7 +1713,7 @@ void FMemAllocTable::AddDefaultColumns()
 					const FMemoryAlloc* Alloc = MemAllocNode.GetMemAlloc();
 					if (Alloc)
 					{
-						const TraceServices::FCallstack* Callstack = Alloc->GetCallstack();
+						const TraceServices::FCallstack* Callstack = Alloc->GetFreeCallstack();
 						const int64 CallstackSize = (Callstack && (Callstack->Num() != 1 || Callstack->Addr(0) != 0)) ? Callstack->Num() : 0;
 						return FTableCellValue(CallstackSize);
 					}
@@ -1486,7 +1727,7 @@ void FMemAllocTable::AddDefaultColumns()
 				if (!Node.IsGroup()) //if (Node->Is<FMemAllocNode>())
 				{
 					const FMemAllocNode& MemAllocNode = static_cast<const FMemAllocNode&>(Node);
-					return MemAllocNode.GetCallstackId();
+					return MemAllocNode.GetFreeCallstackId();
 				}
 				return 0;
 			}
@@ -1494,57 +1735,7 @@ void FMemAllocTable::AddDefaultColumns()
 		TSharedRef<ITableCellValueGetter> Getter = MakeShared<FCallstackSizeValueGetter>();
 		Column.SetValueGetter(Getter);
 
-		class FCallstackValueFormatter : public FTableCellValueFormatter
-		{
-		public:
-			virtual FText FormatValue(const TOptional<FTableCellValue>& InValue) const override
-			{
-				return InValue.IsSet() ? FText::AsNumber(InValue.GetValue().Int64) : FText::GetEmpty();
-			}
-
-			virtual TSharedPtr<IToolTip> GetCustomTooltip(const FTableColumn& Column, const FBaseTreeNode& Node) const override
-			{
-				const FMemAllocNode& MemAllocNode = static_cast<const FMemAllocNode&>(Node);
-
-				FText HeaderText = FText::Format(LOCTEXT("CallstackHeaderFmt", "CallstackId={0} FreeCallstackId={1}"),
-					FText::AsNumber(MemAllocNode.GetCallstackId(), &FNumberFormattingOptions::DefaultNoGrouping()),
-					FText::AsNumber(MemAllocNode.GetFreeCallstackId(), &FNumberFormattingOptions::DefaultNoGrouping()));
-
-				return SNew(SToolTip)
-					.Visibility(TAttribute<EVisibility>::Create(TAttribute<EVisibility>::FGetter::CreateStatic(&FTableCellValueFormatter::GetTooltipVisibility)))
-					[
-						SNew(SVerticalBox)
-						+ SVerticalBox::Slot()
-						.AutoHeight()
-						.Padding(2.0f)
-						[
-							SNew(STextBlock)
-							.Text(HeaderText)
-							.ColorAndOpacity(FSlateColor(EStyleColor::White25))
-						]
-						+ SVerticalBox::Slot()
-						.AutoHeight()
-						.Padding(2.0f)
-						[
-							SNew(STextBlock)
-							.Text(&MemAllocNode, &FMemAllocNode::GetFullCallstack, FMemAllocNode::ECallstackType::FreeCallstack)
-						]
-					];
-			}
-
-			virtual FText FormatValueForTooltip(const FTableColumn& Column, const FBaseTreeNode& Node) const override
-			{
-				const FMemAllocNode& MemAllocNode = static_cast<const FMemAllocNode&>(Node);
-				return MemAllocNode.GetFullCallstack(FMemAllocNode::ECallstackType::FreeCallstack);
-			}
-
-			virtual FText FormatValueForGrouping(const FTableColumn& Column, const FBaseTreeNode& Node) const override
-			{
-				return FormatValue(Column.GetValue(Node));
-			}
-		};
-
-		TSharedRef<ITableCellValueFormatter> Formatter = MakeShared<FCallstackValueFormatter>();
+		TSharedRef<ITableCellValueFormatter> Formatter = MakeShared<FCallstackValueFormatter>(FMemAllocNode::ECallstackType::FreeCallstack);
 		Column.SetValueFormatter(Formatter);
 
 		TSharedRef<ITableCellValueSorter> Sorter = MakeShared<FSorterByInt64Value>(ColumnRef);
diff --git a/Engine/Source/Developer/TraceInsights/Private/Insights/MemoryProfiler/ViewModels/MemAllocTable.h b/Engine/Source/Developer/TraceInsights/Private/Insights/MemoryProfiler/ViewModels/MemAllocTable.h
--- a/Engine/Source/Developer/TraceInsights/Private/Insights/MemoryProfiler/ViewModels/MemAllocTable.h
+++ b/Engine/Source/Developer/TraceInsights/Private/Insights/MemoryProfiler/ViewModels/MemAllocTable.h
@@ -24,6 +24,8 @@ struct FMemAllocTableColumns
 	static const FName StartTimeColumnId;
 	static const FName EndTimeColumnId;
 	static const FName DurationColumnId;
+	static const FName AllocThreadColumnId;
+	static const FName FreeThreadColumnId;
 	static const FName AddressColumnId;
 	static const FName MemoryPageColumnId;
 	static const FName CountColumnId;
@@ -34,9 +36,11 @@ struct FMemAllocTableColumns
 	static const FName ClassNameColumnId;
 	static const FName AllocFunctionColumnId;
 	static const FName AllocSourceFileColumnId;
+	static const FName AllocCallstackIdColumnId;
 	static const FName AllocCallstackSizeColumnId;
 	static const FName FreeFunctionColumnId;
 	static const FName FreeSourceFileColumnId;
+	static const FName FreeCallstackIdColumnId;
 	static const FName FreeCallstackSizeColumnId;
 };
 
diff --git a/Engine/Source/Developer/TraceInsights/Private/Insights/MemoryProfiler/ViewModels/MemoryAlloc.cpp b/Engine/Source/Developer/TraceInsights/Private/Insights/MemoryProfiler/ViewModels/MemoryAlloc.cpp
--- a/Engine/Source/Developer/TraceInsights/Private/Insights/MemoryProfiler/ViewModels/MemoryAlloc.cpp
+++ b/Engine/Source/Developer/TraceInsights/Private/Insights/MemoryProfiler/ViewModels/MemoryAlloc.cpp
@@ -10,21 +10,6 @@ namespace Insights
 ////////////////////////////////////////////////////////////////////////////////////////////////////
 
 FMemoryAlloc::FMemoryAlloc()
-	: StartEventIndex(0)
-	, EndEventIndex(0)
-	, StartTime(0.0)
-	, EndTime(0.0)
-	, Address(0)
-	, Size(0)
-	, TagId(0)
-	, Tag(nullptr)
-	, Asset(nullptr)
-	, Package(nullptr)
-	, Callstack(nullptr)
-	, FreeCallstack(nullptr)
-	, RootHeap(0)
-	, bIsHeap(false)
-	, bIsDecline(false)
 {
 }
 
diff --git a/Engine/Source/Developer/TraceInsights/Private/Insights/MemoryProfiler/ViewModels/MemoryAlloc.h b/Engine/Source/Developer/TraceInsights/Private/Insights/MemoryProfiler/ViewModels/MemoryAlloc.h
--- a/Engine/Source/Developer/TraceInsights/Private/Insights/MemoryProfiler/ViewModels/MemoryAlloc.h
+++ b/Engine/Source/Developer/TraceInsights/Private/Insights/MemoryProfiler/ViewModels/MemoryAlloc.h
@@ -23,43 +23,57 @@ public:
 	int64 GetStartEventIndex() const { return int64(StartEventIndex); }
 	int64 GetEndEventIndex() const { return int64(EndEventIndex); }
 	int64 GetEventDistance() const { return EndEventIndex == ~0 ? int64(EndEventIndex) : int64(EndEventIndex) - int64(StartEventIndex); }
+
 	double GetStartTime() const { return StartTime; }
 	double GetEndTime() const { return EndTime; }
 	double GetDuration() const { return EndTime - StartTime; }
+
 	uint64 GetAddress() const { return Address; }
 	uint64 GetPage() const { return Address & ~(4llu*1024-1); }
+
 	int64 GetSize() const { return Size; }
-	bool IsHeap() const { return bIsHeap; }
+
 	TraceServices::TagIdType GetTagId() const { return TagId; }
 	const TCHAR* GetTag() const { return Tag; }
+
 	const TCHAR* GetAsset() const { return Asset; }
 	const TCHAR* GetPackage() const { return Package; }
 	const TCHAR* GetClassName() const { return ClassName; }
-	uint32 GetCallstackId() const { return CallstackId; }
+
+	uint16 GetAllocThreadId() const { return AllocThreadId; }
+	uint16 GetFreeThreadId() const { return FreeThreadId; }
+
+	uint32 GetAllocCallstackId() const { return AllocCallstackId; }
 	uint32 GetFreeCallstackId() const { return FreeCallstackId; }
-	const TraceServices::FCallstack* GetCallstack() const { return Callstack; }
+	const TraceServices::FCallstack* GetAllocCallstack() const { return AllocCallstack; }
 	const TraceServices::FCallstack* GetFreeCallstack() const { return FreeCallstack; }
+
 	HeapId GetRootHeap() const { return RootHeap; }
+	bool IsHeap() const { return bIsHeap; }
+
+	bool IsDecline() const { return bIsDecline; }
 
 private:
-	uint32 StartEventIndex;
-	uint32 EndEventIndex;
-	double StartTime;
-	double EndTime;
-	uint64 Address;
-	int64 Size;
-	TraceServices::TagIdType TagId;
-	const TCHAR* Tag;
-	const TCHAR* Asset;
-	const TCHAR* Package;
-	const TCHAR* ClassName;
-	uint32 CallstackId;
-	uint32 FreeCallstackId;
-	const TraceServices::FCallstack* Callstack;
-	const TraceServices::FCallstack* FreeCallstack;
-	HeapId RootHeap;
-	bool bIsHeap;
-	bool bIsDecline;
+	uint32 StartEventIndex = 0;
+	uint32 EndEventIndex = 0;
+	double StartTime = 0.0;
+	double EndTime = 0.0;
+	uint64 Address = 0;
+	int64 Size = 0;
+	uint16 AllocThreadId = 0;
+	uint16 FreeThreadId = 0;
+	TraceServices::TagIdType TagId = 0;
+	const TCHAR* Tag = nullptr;
+	const TCHAR* Asset = nullptr;
+	const TCHAR* Package = nullptr;
+	const TCHAR* ClassName = nullptr;
+	uint32 AllocCallstackId = 0;
+	uint32 FreeCallstackId = 0;
+	const TraceServices::FCallstack* AllocCallstack = nullptr;
+	const TraceServices::FCallstack* FreeCallstack = nullptr;
+	HeapId RootHeap = 0;
+	bool bIsHeap = false;
+	bool bIsDecline = false;
 };
 
 ////////////////////////////////////////////////////////////////////////////////////////////////////
diff --git a/Engine/Source/Developer/TraceInsights/Private/Insights/MemoryProfiler/Widgets/SMemAllocTableTreeView.cpp b/Engine/Source/Developer/TraceInsights/Private/Insights/MemoryProfiler/Widgets/SMemAllocTableTreeView.cpp
--- a/Engine/Source/Developer/TraceInsights/Private/Insights/MemoryProfiler/Widgets/SMemAllocTableTreeView.cpp
+++ b/Engine/Source/Developer/TraceInsights/Private/Insights/MemoryProfiler/Widgets/SMemAllocTableTreeView.cpp
@@ -413,7 +413,7 @@ void SMemAllocTableTreeView::UpdateQuery(TraceServices::IAllocationsProvider::EQ
 					if (MetadataId != TraceServices::IMetadataProvider::InvalidMetadataId && MetadataProvider && DefinitionProvider)
 					{
 						TraceServices::FProviderReadScopeLock MetadataProviderReadLock(*MetadataProvider);
-						MetadataProvider->EnumerateMetadata(Allocation->GetThreadId(), MetadataId,
+						MetadataProvider->EnumerateMetadata(Allocation->GetAllocThreadId(), MetadataId,
 							[AssetMetadataType, Schema, &Alloc, DefinitionProvider](uint32 StackDepth, uint16 Type, const void* Data, uint32 Size) -> bool
 							{
 								if (Type == AssetMetadataType)
@@ -444,20 +444,23 @@ void SMemAllocTableTreeView::UpdateQuery(TraceServices::IAllocationsProvider::EQ
 							});
 					}
 
-					Alloc.CallstackId = Allocation->GetCallstackId();
+					Alloc.AllocThreadId = uint16(Allocation->GetAllocThreadId());
+					Alloc.FreeThreadId = uint16(Allocation->GetFreeThreadId());
+
+					Alloc.AllocCallstackId = Allocation->GetAllocCallstackId();
 					Alloc.FreeCallstackId = Allocation->GetFreeCallstackId();
 
 					if (CallstacksProvider)
 					{
-						Alloc.Callstack = CallstacksProvider->GetCallstack(Allocation->GetCallstackId());
-						check(Alloc.Callstack != nullptr);
+						Alloc.AllocCallstack = CallstacksProvider->GetCallstack(Allocation->GetAllocCallstackId());
+						check(Alloc.AllocCallstack != nullptr);
 
 						Alloc.FreeCallstack = CallstacksProvider->GetCallstack(Allocation->GetFreeCallstackId());
 						check(Alloc.FreeCallstack != nullptr);
 					}
 					else
 					{
-						Alloc.Callstack = nullptr;
+						Alloc.AllocCallstack = nullptr;
 						Alloc.FreeCallstack = nullptr;
 					}
 
@@ -1678,17 +1681,17 @@ uint32 SMemAllocTableTreeView::CountSourceFiles(FMemAllocNode& MemAllocNode)
 	}
 
 	const FMemoryAlloc* Alloc = MemAllocNode.GetMemAlloc();
-	if (!Alloc || !Alloc->Callstack)
+	if (!Alloc || !Alloc->AllocCallstack)
 	{
 		return 0;
 	}
 
 	uint32 NumSourceFiles = 0;
-	const uint32 NumCallstackFrames = Alloc->Callstack->Num();
+	const uint32 NumCallstackFrames = Alloc->AllocCallstack->Num();
 	check(NumCallstackFrames <= 256); // see Callstack->Frame(uint8)
 	for (uint32 FrameIndex = 0; FrameIndex < NumCallstackFrames; ++FrameIndex)
 	{
-		const TraceServices::FStackFrame* Frame = Alloc->Callstack->Frame(static_cast<uint8>(FrameIndex));
+		const TraceServices::FStackFrame* Frame = Alloc->AllocCallstack->Frame(static_cast<uint8>(FrameIndex));
 		if (Frame && Frame->Symbol && Frame->Symbol->File)
 		{
 			++NumSourceFiles;
@@ -1801,7 +1804,7 @@ void SMemAllocTableTreeView::BuildOpenSourceSubMenu(FMenuBuilder& MenuBuilder, b
 			const FMemoryAlloc* Alloc = MemAllocNode->GetMemAlloc();
 			if (Alloc)
 			{
-				const TraceServices::FCallstack* Callstack = bIsAllocCallstack ? Alloc->Callstack : Alloc->FreeCallstack;
+				const TraceServices::FCallstack* Callstack = bIsAllocCallstack ? Alloc->AllocCallstack : Alloc->FreeCallstack;
 				if (Callstack)
 				{
 					bHasAnySourceFilesToOpen = BuildOpenSourceSubMenuItems(MenuBuilder, *Callstack);
@@ -1946,7 +1949,7 @@ void SMemAllocTableTreeView::ExportMemorySnapshot() const
 			const FMemoryAlloc* Alloc = MemAllocNode.GetMemAlloc();
 			if (Alloc)
 			{
-				const TraceServices::FCallstack* Callstack = bIsAllocCallstack ? Alloc->Callstack : Alloc->FreeCallstack;
+				const TraceServices::FCallstack* Callstack = bIsAllocCallstack ? Alloc->AllocCallstack : Alloc->FreeCallstack;
 				if (!Callstack)
 				{
 					OutData << QuotationMarkBegin;
diff --git a/Engine/Source/Developer/TraceInsights/Private/Insights/Table/ViewModels/TableCellValueFormatter.cpp b/Engine/Source/Developer/TraceInsights/Private/Insights/Table/ViewModels/TableCellValueFormatter.cpp
--- a/Engine/Source/Developer/TraceInsights/Private/Insights/Table/ViewModels/TableCellValueFormatter.cpp
+++ b/Engine/Source/Developer/TraceInsights/Private/Insights/Table/ViewModels/TableCellValueFormatter.cpp
@@ -41,6 +41,20 @@ FText FTableCellValueFormatter::FormatValueForGrouping(const FTableColumn& Colum
 
 ////////////////////////////////////////////////////////////////////////////////////////////////////
 
+FText FTableCellValueFormatter::CopyValue(const FTableColumn& Column, const FBaseTreeNode& Node) const
+{
+	return FormatValue(Column, Node);
+}
+
+////////////////////////////////////////////////////////////////////////////////////////////////////
+
+FText FTableCellValueFormatter::CopyTooltip(const FTableColumn& Column, const FBaseTreeNode& Node) const
+{
+	return FormatValueForTooltip(Column, Node);
+}
+
+////////////////////////////////////////////////////////////////////////////////////////////////////
+
 TSharedPtr<IToolTip> FTableCellValueFormatter::GetCustomTooltip(const FTableColumn& Column, const FBaseTreeNode& Node) const
 {
 	return SNew(SToolTip)
diff --git a/Engine/Source/Developer/TraceInsights/Private/Insights/Table/ViewModels/TableCellValueFormatter.h b/Engine/Source/Developer/TraceInsights/Private/Insights/Table/ViewModels/TableCellValueFormatter.h
--- a/Engine/Source/Developer/TraceInsights/Private/Insights/Table/ViewModels/TableCellValueFormatter.h
+++ b/Engine/Source/Developer/TraceInsights/Private/Insights/Table/ViewModels/TableCellValueFormatter.h
@@ -28,6 +28,9 @@ public:
 	virtual FText FormatValueForTooltip(const FTableColumn& Column, const FBaseTreeNode& Node) const = 0;
 	virtual FText FormatValueForGrouping(const FTableColumn& Column, const FBaseTreeNode& Node) const = 0;
 
+	virtual FText CopyValue(const FTableColumn& Column, const FBaseTreeNode& Node) const = 0;
+	virtual FText CopyTooltip(const FTableColumn& Column, const FBaseTreeNode& Node) const = 0;
+
 	virtual TSharedPtr<IToolTip> GetCustomTooltip(const FTableColumn& Column, const FBaseTreeNode& Node) const = 0;
 };
 
@@ -46,6 +49,9 @@ public:
 	virtual FText FormatValueForTooltip(const FTableColumn& Column, const FBaseTreeNode& Node) const override; // { return FormatValueForTooltip(Column.GetValue(Node)); }
 	virtual FText FormatValueForGrouping(const FTableColumn& Column, const FBaseTreeNode& Node) const override; // { return FormatValueForTooltip(Column.GetValue(Node)); }
 
+	virtual FText CopyValue(const FTableColumn& Column, const FBaseTreeNode& Node) const override; // { return FormatValue(Column, Node); }
+	virtual FText CopyTooltip(const FTableColumn& Column, const FBaseTreeNode& Node) const override; // { return FormatValueForTooltip(Column, Node); }
+
 	virtual TSharedPtr<IToolTip> GetCustomTooltip(const FTableColumn& Column, const FBaseTreeNode& Node) const override;
 
 	static EVisibility GetTooltipVisibility();
diff --git a/Engine/Source/Developer/TraceInsights/Private/Insights/Table/ViewModels/TableColumn.cpp b/Engine/Source/Developer/TraceInsights/Private/Insights/Table/ViewModels/TableColumn.cpp
--- a/Engine/Source/Developer/TraceInsights/Private/Insights/Table/ViewModels/TableColumn.cpp
+++ b/Engine/Source/Developer/TraceInsights/Private/Insights/Table/ViewModels/TableColumn.cpp
@@ -62,6 +62,20 @@ FText FTableColumn::GetValueAsGroupingText(const FBaseTreeNode& InNode) const
 
 ////////////////////////////////////////////////////////////////////////////////////////////////////
 
+FText FTableColumn::CopyValue(const FBaseTreeNode& InNode) const
+{
+	return ValueFormatter->CopyValue(*this, InNode);
+}
+
+////////////////////////////////////////////////////////////////////////////////////////////////////
+
+FText FTableColumn::CopyTooltip(const FBaseTreeNode& InNode) const
+{
+	return ValueFormatter->CopyTooltip(*this, InNode);
+}
+
+////////////////////////////////////////////////////////////////////////////////////////////////////
+
 inline void TrimEndingZero(FString& Str)
 {
 	const TCHAR* Start = *Str;
diff --git a/Engine/Source/Developer/TraceInsights/Private/Insights/Table/ViewModels/TableColumn.h b/Engine/Source/Developer/TraceInsights/Private/Insights/Table/ViewModels/TableColumn.h
--- a/Engine/Source/Developer/TraceInsights/Private/Insights/Table/ViewModels/TableColumn.h
+++ b/Engine/Source/Developer/TraceInsights/Private/Insights/Table/ViewModels/TableColumn.h
@@ -179,6 +179,9 @@ public:
 	FText GetValueAsTooltipText(const FBaseTreeNode& InNode) const;
 	FText GetValueAsGroupingText(const FBaseTreeNode& InNode) const;
 
+	FText CopyValue(const FBaseTreeNode& InNode) const;
+	FText CopyTooltip(const FBaseTreeNode& InNode) const;
+
 	FString GetValueAsSerializableString(const FBaseTreeNode& InNode) const;
 
 	//////////////////////////////////////////////////
diff --git a/Engine/Source/Developer/TraceInsights/Private/Insights/Table/Widgets/STableTreeView.cpp b/Engine/Source/Developer/TraceInsights/Private/Insights/Table/Widgets/STableTreeView.cpp
--- a/Engine/Source/Developer/TraceInsights/Private/Insights/Table/Widgets/STableTreeView.cpp
+++ b/Engine/Source/Developer/TraceInsights/Private/Insights/Table/Widgets/STableTreeView.cpp
@@ -3507,7 +3507,7 @@ void STableTreeView::ContextMenu_CopyColumnToClipboard_Execute()
 		const TSharedPtr<FTableColumn> HoveredColumnPtr = Table->FindColumn(HoveredColumnId);
 		if (HoveredColumnPtr.IsValid())
 		{
-			FString Text = HoveredColumnPtr->GetValueAsText(*SelectedNode).ToString();
+			FString Text = HoveredColumnPtr->CopyValue(*SelectedNode).ToString();
 			FPlatformApplicationMisc::ClipboardCopy(*Text);
 		}
 	}
@@ -3537,7 +3537,7 @@ void STableTreeView::ContextMenu_CopyColumnTooltipToClipboard_Execute()
 		const TSharedPtr<FTableColumn> HoveredColumnPtr = Table->FindColumn(HoveredColumnId);
 		if (HoveredColumnPtr.IsValid())
 		{
-			FString Text = HoveredColumnPtr->GetValueAsTooltipText(*SelectedNode).ToString();
+			FString Text = HoveredColumnPtr->CopyTooltip(*SelectedNode).ToString();
 			FPlatformApplicationMisc::ClipboardCopy(*Text);
 		}
 	}
diff --git a/Engine/Source/Developer/TraceInsights/Private/Insights/Tests/FunctionalTests/MemoryInsightsTests.cpp b/Engine/Source/Developer/TraceInsights/Private/Insights/Tests/FunctionalTests/MemoryInsightsTests.cpp
--- a/Engine/Source/Developer/TraceInsights/Private/Insights/Tests/FunctionalTests/MemoryInsightsTests.cpp
+++ b/Engine/Source/Developer/TraceInsights/Private/Insights/Tests/FunctionalTests/MemoryInsightsTests.cpp
@@ -215,9 +215,13 @@ bool FVerifyHierarchyCallStackCommand::Update()
 		{
 			const Insights::FMemAllocNode& MemAllocNode = static_cast<const Insights::FMemAllocNode&>(*Node);
 			const Insights::FMemoryAlloc Alloc = MemAllocNode.GetMemAllocChecked();
-			if (!(!Alloc.GetCallstack() || Alloc.GetCallstack()->Num() == 0 || (Alloc.GetCallstack()->Num() != 0 && Alloc.GetCallstack()->Num() <= 256)))
+			if (!(!Alloc.GetAllocCallstack() || Alloc.GetAllocCallstack()->Num() == 0 || (Alloc.GetAllocCallstack()->Num() != 0 && Alloc.GetAllocCallstack()->Num() < 256)))
 			{
-				Test->AddError(TEXT("Resolved callstack should be valid"));
+				Test->AddError(TEXT("Resolved alloc callstack should be valid"));
+			}
+			if (!(!Alloc.GetFreeCallstack() || Alloc.GetFreeCallstack()->Num() == 0 || (Alloc.GetFreeCallstack()->Num() != 0 && Alloc.GetFreeCallstack()->Num() < 256)))
+			{
+				Test->AddError(TEXT("Resolved free callstack should be valid"));
 			}
 		}
 		return true;
diff --git a/Engine/Source/Developer/TraceServices/Private/Model/AllocationItem.h b/Engine/Source/Developer/TraceServices/Private/Model/AllocationItem.h
--- a/Engine/Source/Developer/TraceServices/Private/Model/AllocationItem.h
+++ b/Engine/Source/Developer/TraceServices/Private/Model/AllocationItem.h
@@ -32,9 +32,9 @@ struct FAllocationItem
 	uint32 EndEventIndex;
 	double StartTime;
 	double EndTime;
-	uint16 ThreadId;
+	uint16 AllocThreadId;
 	uint16 FreeThreadId;
-	uint32 CallstackId;
+	uint32 AllocCallstackId;
 	uint32 FreeCallstackId;
 	uint32 MetadataId;
 	TagIdType Tag; // uint32
diff --git a/Engine/Source/Developer/TraceServices/Private/Model/AllocationsProvider.cpp b/Engine/Source/Developer/TraceServices/Private/Model/AllocationsProvider.cpp
--- a/Engine/Source/Developer/TraceServices/Private/Model/AllocationsProvider.cpp
+++ b/Engine/Source/Developer/TraceServices/Private/Model/AllocationsProvider.cpp
@@ -466,18 +466,26 @@ uint32 IAllocationsProvider::FAllocation::GetAlignment() const
 
 ////////////////////////////////////////////////////////////////////////////////////////////////////
 
-uint32 IAllocationsProvider::FAllocation::GetThreadId() const
+uint32 IAllocationsProvider::FAllocation::GetAllocThreadId() const
 {
 	const auto* Inner = (const FAllocationItem*)this;
-	return Inner->ThreadId;
+	return Inner->AllocThreadId;
 }
 
 ////////////////////////////////////////////////////////////////////////////////////////////////////
 
-uint32 IAllocationsProvider::FAllocation::GetCallstackId() const
+uint32 IAllocationsProvider::FAllocation::GetFreeThreadId() const
 {
 	const auto* Inner = (const FAllocationItem*)this;
-	return Inner->CallstackId;
+	return Inner->FreeThreadId;
+}
+
+////////////////////////////////////////////////////////////////////////////////////////////////////
+
+uint32 IAllocationsProvider::FAllocation::GetAllocCallstackId() const
+{
+	const auto* Inner = (const FAllocationItem*)this;
+	return Inner->AllocCallstackId;
 }
 
 ////////////////////////////////////////////////////////////////////////////////////////////////////
@@ -1450,10 +1458,10 @@ void FAllocationsProvider::EditAlloc(uint32 ThreadId, double Time, uint32 Callst
 		Allocation.EndEventIndex = (uint32)-1;
 		Allocation.StartTime = Time;
 		Allocation.EndTime = std::numeric_limits<double>::infinity();
-		Allocation.ThreadId = static_cast<uint16>(ThreadId);
-		check(uint32(Allocation.ThreadId) == ThreadId);
+		Allocation.AllocThreadId = static_cast<uint16>(ThreadId);
+		check(uint32(Allocation.AllocThreadId) == ThreadId);
 		Allocation.FreeThreadId = 0;
-		Allocation.CallstackId = CallstackId;
+		Allocation.AllocCallstackId = CallstackId;
 		Allocation.FreeCallstackId = 0; // no callstack yet
 		Allocation.MetadataId = MetadataId;
 		Allocation.Tag = Tag;
@@ -1851,7 +1859,7 @@ void FAllocationsProvider::EditMarkAllocationAsHeap(uint32 ThreadId, double Time
 		Alloc->RootHeap = static_cast<uint8>(Heap);
 		if (CallstackId != 0)
 		{
-			Alloc->CallstackId = CallstackId;
+			Alloc->AllocCallstackId = CallstackId;
 		}
 
 		// Re-add it to the Live allocs as a heap allocation.
@@ -1954,7 +1962,7 @@ void FAllocationsProvider::EditUnmarkAllocationAsHeap(uint32 ThreadId, double Ti
 		// Mark allocation as a "heap" allocation.
 		Alloc->Flags = Alloc->Flags | EMemoryTraceHeapAllocationFlags::Heap;
 		Alloc->RootHeap = static_cast<uint8>(Heap);
-		Alloc->CallstackId = CallstackId;
+		Alloc->AllocCallstackId = CallstackId;
 	}
 #endif // INSIGHTS_VALIDATE_HEAP_UNMARK_ALLOC_EVENTS
 
@@ -2036,7 +2044,7 @@ void FAllocationsProvider::EditUnmarkAllocationAsHeap(uint32 ThreadId, double Ti
 
 		const uint64 Size = Alloc->GetSize();
 		const uint32 Alignment = Alloc->GetAlignment();
-		const uint32 AllocCallstackId = Alloc->CallstackId;
+		const uint32 AllocCallstackId = Alloc->AllocCallstackId;
 
 		// Re-add this allocation to the Live allocs.
 		RootHeap.LiveAllocs->Add(Alloc); // the Live allocs takes ownership of Alloc
diff --git a/Engine/Source/Developer/TraceServices/Public/TraceServices/Model/AllocationsProvider.h b/Engine/Source/Developer/TraceServices/Public/TraceServices/Model/AllocationsProvider.h
--- a/Engine/Source/Developer/TraceServices/Public/TraceServices/Model/AllocationsProvider.h
+++ b/Engine/Source/Developer/TraceServices/Public/TraceServices/Model/AllocationsProvider.h
@@ -65,8 +65,9 @@ public:
 		uint64 GetAddress() const;
 		uint64 GetSize() const;
 		uint32 GetAlignment() const;
-		uint32 GetThreadId() const;
-		uint32 GetCallstackId() const;
+		uint32 GetAllocThreadId() const;
+		uint32 GetFreeThreadId() const;
+		uint32 GetAllocCallstackId() const;
 		uint32 GetFreeCallstackId() const;
 		uint32 GetMetadataId() const;
 		TagIdType GetTag() const;
