Add feature to find the minimum or maximum instance for a timer event.

diff --git a/Engine/Source/Developer/TraceInsights/Private/Insights/InsightsStyle.cpp b/Engine/Source/Developer/TraceInsights/Private/Insights/InsightsStyle.cpp
--- a/Engine/Source/Developer/TraceInsights/Private/Insights/InsightsStyle.cpp
+++ b/Engine/Source/Developer/TraceInsights/Private/Insights/InsightsStyle.cpp
@@ -327,6 +327,10 @@ void FInsightsStyle::FStyle::Initialize()
 	Set("Icons.ImportTable", new CORE_IMAGE_BRUSH_SVG("Starship/Common/Import", Icon16x16));
 	Set("Icons.SelectEventRange", new EDITOR_IMAGE_BRUSH_SVG("Starship/Common/SelectInViewport", Icon16x16));
 
+	Set("Icons.FindInstance", new CORE_IMAGE_BRUSH_SVG("Starship/Common/search", Icon16x16));
+	Set("Icons.FindMinInstance", new EDITOR_IMAGE_BRUSH_SVG("Starship/Common/Profile", Icon16x16));
+	Set("Icons.FindMaxInstance", new EDITOR_IMAGE_BRUSH_SVG("Starship/Common/Profile", Icon16x16));
+
 	//////////////////////////////////////////////////
 
 	Set("TreeTable.RowBackground", new EDITOR_IMAGE_BRUSH("Old/White", Icon16x16, FLinearColor(1.0f, 1.0f, 1.0f, 0.25f)));
diff --git a/Engine/Source/Developer/TraceInsights/Private/Insights/ViewModels/ThreadTimingTrack.cpp b/Engine/Source/Developer/TraceInsights/Private/Insights/ViewModels/ThreadTimingTrack.cpp
--- a/Engine/Source/Developer/TraceInsights/Private/Insights/ViewModels/ThreadTimingTrack.cpp
+++ b/Engine/Source/Developer/TraceInsights/Private/Insights/ViewModels/ThreadTimingTrack.cpp
@@ -785,6 +785,90 @@ void FThreadTimingSharedState::ToggleTrackVisibilityByGroup_Execute(const TCHAR*
 	}
 }
 
+////////////////////////////////////////////////////////////////////////////////////////////////////
+
+TSharedPtr<const ITimingEvent> FThreadTimingSharedState::FindMaxEventInstance(uint32 TimerId, double StartTime, double EndTime)
+{
+	auto CompareAndAssignEvent = [](TSharedPtr<const ITimingEvent>& TimingEvent, TSharedPtr<const ITimingEvent>& TrackEvent)
+	{
+		if (!TrackEvent.IsValid())
+		{
+			return;
+		}
+
+		if (!TimingEvent.IsValid() || TrackEvent->GetDuration() > TimingEvent->GetDuration())
+		{
+			TimingEvent = TrackEvent;
+		}
+	};
+
+	TSharedPtr<const ITimingEvent> TimingEvent;
+	TSharedPtr<const ITimingEvent> TrackEvent;
+
+	for (const auto& KV : CpuTracks)
+	{
+		const FCpuTimingTrack& Track = *KV.Value;
+		if (Track.IsVisible())
+		{
+			TrackEvent = Track.FindMaxEventInstance(TimerId, StartTime, EndTime);
+			CompareAndAssignEvent(TimingEvent, TrackEvent);
+		}
+	}
+
+	GpuTrack->FindMaxEventInstance(TimerId, StartTime, EndTime);
+	CompareAndAssignEvent(TimingEvent, TrackEvent);
+
+	Gpu2Track->FindMaxEventInstance(TimerId, StartTime, EndTime);
+	CompareAndAssignEvent(TimingEvent, TrackEvent);
+	
+	return TimingEvent;
+}
+
+////////////////////////////////////////////////////////////////////////////////////////////////////
+
+TSharedPtr<const ITimingEvent> FThreadTimingSharedState::FindMinEventInstance(uint32 TimerId, double StartTime, double EndTime)
+{
+	auto CompareAndAssignEvent = [](TSharedPtr<const ITimingEvent>& TimingEvent, TSharedPtr<const ITimingEvent>& TrackEvent)
+	{
+		if (!TrackEvent.IsValid())
+		{
+			return;
+		}
+
+		if (!TimingEvent.IsValid() || TrackEvent->GetDuration() < TimingEvent->GetDuration())
+		{
+			TimingEvent = TrackEvent;
+		}
+	};
+
+	TSharedPtr<const ITimingEvent> TimingEvent;
+	TSharedPtr<const ITimingEvent> TrackEvent;
+
+	for (const auto& KV : CpuTracks)
+	{
+		const FCpuTimingTrack& Track = *KV.Value;
+		if (Track.IsVisible())
+		{
+			TrackEvent = Track.FindMinEventInstance(TimerId, StartTime, EndTime);
+			CompareAndAssignEvent(TimingEvent, TrackEvent);
+		}
+	}
+
+	if (GpuTrack.IsValid() && GpuTrack->IsVisible())
+	{
+		TrackEvent = GpuTrack->FindMinEventInstance(TimerId, StartTime, EndTime);
+		CompareAndAssignEvent(TimingEvent, TrackEvent);
+	}
+
+	if (Gpu2Track.IsValid() && Gpu2Track->IsVisible())
+	{
+		TrackEvent = Gpu2Track->FindMinEventInstance(TimerId, StartTime, EndTime);
+		CompareAndAssignEvent(TimingEvent, TrackEvent);
+	}
+
+	return TimingEvent;
+}
+
 ////////////////////////////////////////////////////////////////////////////////////////////////////
 // FThreadTimingTrack
 ////////////////////////////////////////////////////////////////////////////////////////////////////
@@ -1855,4 +1939,169 @@ void FThreadTimingTrack::SetFilterConfigurator(TSharedPtr<Insights::FFilterConfi
 
 ////////////////////////////////////////////////////////////////////////////////////////////////////
 
+TSharedPtr<const ITimingEvent> FThreadTimingTrack::FindMaxEventInstance(uint32 InTimerId, double InStartTime, double InEndTime) const
+{
+	TSharedPtr<const TraceServices::IAnalysisSession> Session = FInsightsManager::Get()->GetSession();
+	check(Session.IsValid());
+
+	TraceServices::FAnalysisSessionReadScope SessionReadScope(*Session.Get());
+
+	const TraceServices::ITimingProfilerProvider& TimingProfilerProvider = *TraceServices::ReadTimingProfilerProvider(*Session.Get());
+
+	const TraceServices::ITimingProfilerTimerReader* TimerReader;
+	TimingProfilerProvider.ReadTimers([&TimerReader](const TraceServices::ITimingProfilerTimerReader& Out) { TimerReader = &Out; });
+
+	struct CandidateEvent
+	{
+		double StartTime = 0.0f;
+		double EndTime = -1.0f;
+		uint32 Depth = 0;
+		uint32 TimerIndex = 0;
+	};
+
+	TSharedPtr<FThreadTrackEvent> TimingEvent;
+
+	TimingProfilerProvider.ReadTimeline(TimelineIndex,
+		[TimerReader, InStartTime, InEndTime, InTimerId, &TimingEvent, this](const TraceServices::ITimingProfilerProvider::Timeline& Timeline)
+		{
+			TArray<CandidateEvent> CandidateEvents;
+
+			TraceServices::ITimeline<TraceServices::FTimingProfilerEvent>::EnumerateAsyncParams Params;
+			Params.IntervalStart = InStartTime;
+			Params.IntervalEnd = InEndTime;
+			Params.Resolution = 0.0;
+			Params.SetupCallback = [&CandidateEvents](uint32 NumTasks)
+			{
+				CandidateEvents.AddDefaulted(NumTasks);
+			};
+			Params.Callback = [TimerReader, &CandidateEvents, InTimerId](double StartTime, double EndTime, uint32 Depth, const TraceServices::FTimingProfilerEvent& Event, uint32 TaskIndex)
+			{
+				const TraceServices::FTimingProfilerTimer* Timer = TimerReader->GetTimer(Event.TimerIndex);
+				if (ensure(Timer != nullptr))
+				{
+					if (Timer->Id == InTimerId)
+					{
+						double CandidateDuration = CandidateEvents[TaskIndex].EndTime - CandidateEvents[TaskIndex].StartTime;
+						double EventDuration = EndTime - StartTime;
+
+						if (EventDuration > CandidateDuration)
+						{
+							CandidateEvents[TaskIndex].StartTime = StartTime;
+							CandidateEvents[TaskIndex].EndTime = EndTime;
+							CandidateEvents[TaskIndex].Depth = Depth;
+							CandidateEvents[TaskIndex].TimerIndex = Event.TimerIndex;
+						}
+
+					}
+				}
+				return TraceServices::EEventEnumerate::Continue;
+			};
+
+			// Note: Enumerating events for filtering should not use downsampling.
+			Timeline.EnumerateEventsDownSampledAsync(Params);
+
+			CandidateEvent BestMatch;
+			for (const CandidateEvent& Event : CandidateEvents)
+			{
+				if ((Event.EndTime - Event.StartTime) > BestMatch.EndTime - BestMatch.StartTime)
+				{
+					BestMatch = Event;
+				}
+			}
+
+			if (BestMatch.EndTime > BestMatch.StartTime)
+			{
+				TimingEvent = MakeShared<FThreadTrackEvent>(SharedThis(this), BestMatch.StartTime, BestMatch.EndTime, BestMatch.Depth);
+				TimingEvent->SetTimerId(InTimerId);
+				TimingEvent->SetTimerIndex(BestMatch.TimerIndex);
+			}
+		});
+
+	return TimingEvent;
+}
+
+////////////////////////////////////////////////////////////////////////////////////////////////////
+
+TSharedPtr<const ITimingEvent> FThreadTimingTrack::FindMinEventInstance(uint32 InTimerId, double InStartTime, double InEndTime) const
+{
+	TSharedPtr<const TraceServices::IAnalysisSession> Session = FInsightsManager::Get()->GetSession();
+	check(Session.IsValid());
+
+	TraceServices::FAnalysisSessionReadScope SessionReadScope(*Session.Get());
+
+	const TraceServices::ITimingProfilerProvider& TimingProfilerProvider = *TraceServices::ReadTimingProfilerProvider(*Session.Get());
+
+	const TraceServices::ITimingProfilerTimerReader* TimerReader;
+	TimingProfilerProvider.ReadTimers([&TimerReader](const TraceServices::ITimingProfilerTimerReader& Out) { TimerReader = &Out; });
+
+	struct CandidateEvent
+	{
+		double StartTime = -std::numeric_limits<double>::infinity();
+		double EndTime = std::numeric_limits<double>::infinity();
+		uint32 Depth = 0;
+		uint32 TimerIndex = 0;
+	};
+
+	TSharedPtr<FThreadTrackEvent> TimingEvent;
+
+	TimingProfilerProvider.ReadTimeline(TimelineIndex,
+		[TimerReader, InStartTime, InEndTime, InTimerId, &TimingEvent, this](const TraceServices::ITimingProfilerProvider::Timeline& Timeline)
+		{
+			TArray<CandidateEvent> CandidateEvents;
+
+			TraceServices::ITimeline<TraceServices::FTimingProfilerEvent>::EnumerateAsyncParams Params;
+			Params.IntervalStart = InStartTime;
+			Params.IntervalEnd = InEndTime;
+			Params.Resolution = 0.0;
+			Params.SetupCallback = [&CandidateEvents](uint32 NumTasks)
+			{
+				CandidateEvents.AddDefaulted(NumTasks);
+			};
+			Params.Callback = [TimerReader, &CandidateEvents, InTimerId](double StartTime, double EndTime, uint32 Depth, const TraceServices::FTimingProfilerEvent& Event, uint32 TaskIndex)
+			{
+				const TraceServices::FTimingProfilerTimer* Timer = TimerReader->GetTimer(Event.TimerIndex);
+				if (ensure(Timer != nullptr))
+				{
+					if (Timer->Id == InTimerId)
+					{
+						double CandidateDuration = CandidateEvents[TaskIndex].EndTime - CandidateEvents[TaskIndex].StartTime;
+						double EventDuration = EndTime - StartTime;
+
+						if (EventDuration < CandidateDuration)
+						{
+							CandidateEvents[TaskIndex].StartTime = StartTime;
+							CandidateEvents[TaskIndex].EndTime = EndTime;
+							CandidateEvents[TaskIndex].Depth = Depth;
+							CandidateEvents[TaskIndex].TimerIndex = Event.TimerIndex;
+						}
+					}
+				}
+				return TraceServices::EEventEnumerate::Continue;
+			};
+
+			// Note: Enumerating events for filtering should not use downsampling.
+			Timeline.EnumerateEventsDownSampledAsync(Params);
+
+			CandidateEvent BestMatch;
+			for (const CandidateEvent& Event : CandidateEvents)
+			{
+				if ((Event.EndTime - Event.StartTime) < BestMatch.EndTime - BestMatch.StartTime)
+				{
+					BestMatch = Event;
+				}
+			}
+
+			if (BestMatch.StartTime != -std::numeric_limits<double>::infinity())
+			{
+				TimingEvent = MakeShared<FThreadTrackEvent>(SharedThis(this), BestMatch.StartTime, BestMatch.EndTime, BestMatch.Depth);
+				TimingEvent->SetTimerId(InTimerId);
+				TimingEvent->SetTimerIndex(BestMatch.TimerIndex);
+			}
+		});
+
+	return TimingEvent;
+}
+
+////////////////////////////////////////////////////////////////////////////////////////////////////
+
 #undef LOCTEXT_NAMESPACE
diff --git a/Engine/Source/Developer/TraceInsights/Private/Insights/ViewModels/ThreadTimingTrack.h b/Engine/Source/Developer/TraceInsights/Private/Insights/ViewModels/ThreadTimingTrack.h
--- a/Engine/Source/Developer/TraceInsights/Private/Insights/ViewModels/ThreadTimingTrack.h
+++ b/Engine/Source/Developer/TraceInsights/Private/Insights/ViewModels/ThreadTimingTrack.h
@@ -94,6 +94,9 @@ public:
 	void HideAllCpuTracks() { SetAllCpuTracksToggle(false); }
 	void ShowHideAllCpuTracks() { SetAllCpuTracksToggle(!IsAllCpuTracksToggleOn()); }
 
+	TSharedPtr<const ITimingEvent> FindMaxEventInstance(uint32 TimerId, double StartTime, double EndTime);
+	TSharedPtr<const ITimingEvent> FindMinEventInstance(uint32 TimerId, double StartTime, double EndTime);
+
 private:
 	void CreateThreadGroupsMenu(FMenuBuilder& MenuBuilder);
 
@@ -188,6 +191,9 @@ public:
 
 	virtual void SetFilterConfigurator(TSharedPtr<Insights::FFilterConfigurator> InFilterConfigurator);
 
+	TSharedPtr<const ITimingEvent> FindMaxEventInstance(uint32 TimerId, double StartTime, double EndTime) const;
+	TSharedPtr<const ITimingEvent> FindMinEventInstance(uint32 TimerId, double StartTime, double EndTime) const;
+
 protected:
 	virtual bool HasCustomFilter() const override;
 
diff --git a/Engine/Source/Developer/TraceInsights/Private/Insights/Widgets/STimersView.cpp b/Engine/Source/Developer/TraceInsights/Private/Insights/Widgets/STimersView.cpp
--- a/Engine/Source/Developer/TraceInsights/Private/Insights/Widgets/STimersView.cpp
+++ b/Engine/Source/Developer/TraceInsights/Private/Insights/Widgets/STimersView.cpp
@@ -116,6 +116,30 @@ public:
 			"Opens the source file of the selected timer in the registered IDE.",
 			EUserInterfaceActionType::Button,
 			FInputChord());
+
+		UI_COMMAND(Command_FindMaxInstance,
+			"Maximum Duration Instance",
+			"Navigates to and selects the timing event instance with the maximum duration, for the selected timer.",
+			EUserInterfaceActionType::Button,
+			FInputChord());
+
+		UI_COMMAND(Command_FindMinInstance,
+			"Minimum Duration Instance",
+			"Navigates to and selects the timing event instance with the minimum duration, for the selected timer.",
+			EUserInterfaceActionType::Button,
+			FInputChord());
+
+		UI_COMMAND(Command_FindMaxInstanceInSelection,
+			"Maximum Duration Instance in Selection",
+			"Navigates to and selects the timing event instance with the maximum duration, for the selected timer, in the selected time range.",
+			EUserInterfaceActionType::Button,
+			FInputChord());
+
+		UI_COMMAND(Command_FindMinInstanceInSelection,
+			"Minimum Duration Instance in Selection",
+			"Navigates to and selects the timing event instance with the minimum duration, for the selected timer, in the selected time range.",
+			EUserInterfaceActionType::Button,
+			FInputChord());
 	}
 	UE_ENABLE_OPTIMIZATION_SHIP
 
@@ -126,6 +150,11 @@ public:
 	TSharedPtr<FUICommandInfo> Command_ExportThreads;
 	TSharedPtr<FUICommandInfo> Command_ExportTimers;
 	TSharedPtr<FUICommandInfo> Command_OpenSource;
+
+	TSharedPtr<FUICommandInfo> Command_FindMaxInstance;
+	TSharedPtr<FUICommandInfo> Command_FindMinInstance;
+	TSharedPtr<FUICommandInfo> Command_FindMaxInstanceInSelection;
+	TSharedPtr<FUICommandInfo> Command_FindMinInstanceInSelection;
 };
 
 ////////////////////////////////////////////////////////////////////////////////////////////////////
@@ -178,6 +207,10 @@ void STimersView::InitCommandList()
 	CommandList->MapAction(FTimersViewCommands::Get().Command_ExportThreads, FExecuteAction::CreateSP(this, &STimersView::ContextMenu_ExportThreads_Execute), FCanExecuteAction::CreateSP(this, &STimersView::ContextMenu_ExportThreads_CanExecute));
 	CommandList->MapAction(FTimersViewCommands::Get().Command_ExportTimers, FExecuteAction::CreateSP(this, &STimersView::ContextMenu_ExportTimers_Execute), FCanExecuteAction::CreateSP(this, &STimersView::ContextMenu_ExportTimers_CanExecute));
 	CommandList->MapAction(FTimersViewCommands::Get().Command_OpenSource, FExecuteAction::CreateSP(this, &STimersView::ContextMenu_OpenSource_Execute), FCanExecuteAction::CreateSP(this, &STimersView::ContextMenu_OpenSource_CanExecute));
+	CommandList->MapAction(FTimersViewCommands::Get().Command_FindMaxInstance, FExecuteAction::CreateSP(this, &STimersView::ContextMenu_FindInstance_Execute, true), FCanExecuteAction::CreateSP(this, &STimersView::ContextMenu_FindInstance_CanExecute));
+	CommandList->MapAction(FTimersViewCommands::Get().Command_FindMinInstance, FExecuteAction::CreateSP(this, &STimersView::ContextMenu_FindInstance_Execute, false), FCanExecuteAction::CreateSP(this, &STimersView::ContextMenu_FindInstance_CanExecute));
+	CommandList->MapAction(FTimersViewCommands::Get().Command_FindMaxInstanceInSelection, FExecuteAction::CreateSP(this, &STimersView::ContextMenu_FindInstanceInSelection_Execute, true), FCanExecuteAction::CreateSP(this, &STimersView::ContextMenu_FindInstanceInSelection_CanExecute));
+	CommandList->MapAction(FTimersViewCommands::Get().Command_FindMinInstanceInSelection, FExecuteAction::CreateSP(this, &STimersView::ContextMenu_FindInstanceInSelection_Execute, false), FCanExecuteAction::CreateSP(this, &STimersView::ContextMenu_FindInstanceInSelection_CanExecute));
 }
 
 ////////////////////////////////////////////////////////////////////////////////////////////////////
@@ -555,6 +588,15 @@ TSharedPtr<SWidget> STimersView::TreeView_GetMenuContent()
 			FSlateIcon(FAppStyle::Get().GetStyleSetName(), "Icons.AddGraphSeries")
 		);
 
+		MenuBuilder.AddSubMenu
+		(
+			LOCTEXT("ContextMenu_FindInstance_SubMenu", "Find Instance"),
+			LOCTEXT("ContextMenu_PlotInstance_SubMenu_Desc", "Find the instance of this timer with the minimum or maximum duration."),
+			FNewMenuDelegate::CreateSP(this, &STimersView::TreeView_FindMenu),
+			false,
+			FSlateIcon(FAppStyle::Get().GetStyleSetName(), "Icons.FindInstance")
+		);
+
 		// Open Source in IDE
 		{
 			ISourceCodeAccessModule& SourceCodeAccessModule = FModuleManager::LoadModuleChecked<ISourceCodeAccessModule>("SourceCodeAccess");
@@ -971,6 +1013,49 @@ void STimersView::TreeView_BuildPlotTimerMenu(FMenuBuilder& MenuBuilder)
 
 ////////////////////////////////////////////////////////////////////////////////////////////////////
 
+void STimersView::TreeView_FindMenu(FMenuBuilder& MenuBuilder)
+{
+	MenuBuilder.AddMenuEntry
+	(
+		FTimersViewCommands::Get().Command_FindMaxInstance,
+		NAME_None,
+		TAttribute<FText>(),
+		TAttribute<FText>(),
+		FSlateIcon(FAppStyle::Get().GetStyleSetName(), "Icons.FindMaxInstance")
+	);
+
+	MenuBuilder.AddMenuEntry
+	(
+		FTimersViewCommands::Get().Command_FindMinInstance,
+		NAME_None,
+		TAttribute<FText>(),
+		TAttribute<FText>(),
+		FSlateIcon(FAppStyle::Get().GetStyleSetName(), "Icons.FindMinInstance")
+	);
+
+	MenuBuilder.AddMenuSeparator();
+
+	MenuBuilder.AddMenuEntry
+	(
+		FTimersViewCommands::Get().Command_FindMaxInstanceInSelection,
+		NAME_None,
+		TAttribute<FText>(),
+		TAttribute<FText>(),
+		FSlateIcon(FAppStyle::Get().GetStyleSetName(), "Icons.FindMaxInstance")
+	);
+
+	MenuBuilder.AddMenuEntry
+	(
+		FTimersViewCommands::Get().Command_FindMinInstanceInSelection,
+		NAME_None,
+		TAttribute<FText>(),
+		TAttribute<FText>(),
+		FSlateIcon(FAppStyle::Get().GetStyleSetName(), "Icons.FindMinInstance")
+	);
+}
+
+////////////////////////////////////////////////////////////////////////////////////////////////////
+
 void STimersView::InitializeAndShowHeaderColumns()
 {
 	// Create columns.
@@ -3166,6 +3251,84 @@ void STimersView::ContextMenu_OpenSource_Execute() const
 
 ////////////////////////////////////////////////////////////////////////////////////////////////////
 
+bool STimersView::ContextMenu_FindInstance_CanExecute() const
+{
+	const TArray<FTimerNodePtr> SelectedNodes = TreeView->GetSelectedItems();
+	return SelectedNodes.Num() == 1 && !(SelectedNodes[0]->GetType() == ETimerNodeType::Group);
+}
+
+////////////////////////////////////////////////////////////////////////////////////////////////////
+
+void STimersView::ContextMenu_FindInstance_Execute(bool bFindMax) const
+{
+	const TArray<FTimerNodePtr> SelectedNodes = TreeView->GetSelectedItems();
+	if (SelectedNodes.Num() != 1)
+	{
+		return;
+	}
+
+	FTimerNodePtr SelectedNode = SelectedNodes[0];
+
+	TSharedPtr<STimingProfilerWindow> Wnd = FTimingProfilerManager::Get()->GetProfilerWindow();
+	TSharedPtr<STimingView> TimingView = Wnd.IsValid() ? Wnd->GetTimingView() : nullptr;
+
+	if (!TimingView.IsValid())
+	{
+		return;
+	}
+
+	ESelectEventType Type = bFindMax ? ESelectEventType::Max : ESelectEventType::Min;
+	TimingView->SelectEventInstance(SelectedNode->GetTimerId(), Type, false);
+}
+
+////////////////////////////////////////////////////////////////////////////////////////////////////
+
+bool STimersView::ContextMenu_FindInstanceInSelection_CanExecute() const
+{
+	const TArray<FTimerNodePtr> SelectedNodes = TreeView->GetSelectedItems();
+
+	if (SelectedNodes.Num() != 1 || SelectedNodes[0]->GetType() == ETimerNodeType::Group)
+	{
+		return false;
+	}
+
+	TSharedPtr<STimingProfilerWindow> Wnd = FTimingProfilerManager::Get()->GetProfilerWindow();
+	TSharedPtr<STimingView> TimingView = Wnd.IsValid() ? Wnd->GetTimingView() : nullptr;
+
+	if (TimingView.IsValid())
+	{
+		return TimingView->GetSelectionEndTime() > TimingView->GetSelectionStartTime();
+	}
+
+	return false;
+}
+
+////////////////////////////////////////////////////////////////////////////////////////////////////
+
+void STimersView::ContextMenu_FindInstanceInSelection_Execute(bool bFindMax) const
+{
+	const TArray<FTimerNodePtr> SelectedNodes = TreeView->GetSelectedItems();
+	if (SelectedNodes.Num() != 1)
+	{
+		return;
+	}
+
+	FTimerNodePtr SelectedNode = SelectedNodes[0];
+
+	TSharedPtr<STimingProfilerWindow> Wnd = FTimingProfilerManager::Get()->GetProfilerWindow();
+	TSharedPtr<STimingView> TimingView = Wnd.IsValid() ? Wnd->GetTimingView() : nullptr;
+
+	if (!TimingView.IsValid())
+	{
+		return;
+	}
+
+	ESelectEventType Type = bFindMax ? ESelectEventType::Max : ESelectEventType::Min;
+	TimingView->SelectEventInstance(SelectedNode->GetTimerId(), Type, true);
+}
+
+////////////////////////////////////////////////////////////////////////////////////////////////////
+
 void STimersView::OpenSourceFileInIDE(FTimerNodePtr InNode) const
 {
 	if (InNode.IsValid() && InNode->GetType() != ETimerNodeType::Group)
diff --git a/Engine/Source/Developer/TraceInsights/Private/Insights/Widgets/STimersView.h b/Engine/Source/Developer/TraceInsights/Private/Insights/Widgets/STimersView.h
--- a/Engine/Source/Developer/TraceInsights/Private/Insights/Widgets/STimersView.h
+++ b/Engine/Source/Developer/TraceInsights/Private/Insights/Widgets/STimersView.h
@@ -148,6 +148,12 @@ private:
 	bool ContextMenu_OpenSource_CanExecute() const;
 	void ContextMenu_OpenSource_Execute() const;
 
+	bool ContextMenu_FindInstance_CanExecute() const;
+	void ContextMenu_FindInstance_Execute(bool bFindMax) const;
+
+	bool ContextMenu_FindInstanceInSelection_CanExecute() const;
+	void ContextMenu_FindInstanceInSelection_Execute(bool bFindMax) const;
+
 	////////////////////////////////////////////////////////////////////////////////////////////////////
 	// Tree View - Columns' Header
 
@@ -304,6 +310,8 @@ private:
 
 	void TreeView_BuildPlotTimerMenu(FMenuBuilder& MenuBuilder);
 
+	void TreeView_FindMenu(FMenuBuilder& MenuBuilder);
+
 	TSharedPtr<FTimingGraphTrack> GetTimingViewMainGraphTrack() const;
 	TSharedPtr<SFrameTrack> GetFrameTrack() const;
 
diff --git a/Engine/Source/Developer/TraceInsights/Private/Insights/Widgets/STimingView.cpp b/Engine/Source/Developer/TraceInsights/Private/Insights/Widgets/STimingView.cpp
--- a/Engine/Source/Developer/TraceInsights/Private/Insights/Widgets/STimingView.cpp
+++ b/Engine/Source/Developer/TraceInsights/Private/Insights/Widgets/STimingView.cpp
@@ -1084,6 +1084,16 @@ void STimingView::Tick(const FGeometry& AllottedGeometry, const double InCurrent
 
 	Viewport.ResetDirtyFlags();
 
+	if (bBringSelectedEventIntoViewVerticallyOnNextTick && SelectedEvent.IsValid())
+	{
+		if (SelectedEvent->GetTrack()->GetLocation() == ETimingTrackLocation::Scrollable)
+		{
+			BringScrollableTrackIntoView(*SelectedEvent->GetTrack());
+		}
+	}
+
+	bBringSelectedEventIntoViewVerticallyOnNextTick = false;
+
 	TickStopwatch.Stop();
 	TickDurationHistory.AddValue(TickStopwatch.AccumulatedTime);
 }
@@ -3821,7 +3831,7 @@ void STimingView::SelectHoveredTimingTrack()
 
 void STimingView::SelectHoveredTimingEvent()
 {
-	SelectTimingEvent(HoveredEvent, true);
+	SelectTimingEvent(HoveredEvent, true, false);
 }
 
 ////////////////////////////////////////////////////////////////////////////////////////////////////
@@ -3855,7 +3865,7 @@ void STimingView::SelectTimingTrack(const TSharedPtr<FBaseTimingTrack> InTrack,
 
 ////////////////////////////////////////////////////////////////////////////////////////////////////
 
-void STimingView::SelectTimingEvent(const TSharedPtr<const ITimingEvent> InEvent, bool bBringEventIntoView)
+void STimingView::SelectTimingEvent(const TSharedPtr<const ITimingEvent> InEvent, bool bBringEventIntoViewHorizontally, bool bBringEventIntoViewVertically)
 {
 	if (SelectedEvent != InEvent)
 	{
@@ -3864,10 +3874,16 @@ void STimingView::SelectTimingEvent(const TSharedPtr<const ITimingEvent> InEvent
 		if (SelectedEvent.IsValid())
 		{
 			LastSelectionType = ESelectionType::TimingEvent;
-			if (bBringEventIntoView)
+			if (bBringEventIntoViewHorizontally)
 			{
 				BringIntoView(SelectedEvent->GetStartTime(), SelectedEvent->GetEndTime());
 			}
+			if (bBringEventIntoViewVertically)
+			{
+				bBringSelectedEventIntoViewVerticallyOnNextTick = true;
+				// We need the layout to be callculated in one frame, no animations, otherwise the event might not be in view at the end of the animation.
+				Viewport.AddDirtyFlags(ETimingTrackViewportDirtyFlags::VLayoutChanged);
+			}
 		}
 
 		OnSelectedTimingEventChanged();
@@ -5321,5 +5337,43 @@ ETraceFrameType STimingView::GetFrameTypeToSnapTo()
 
 ////////////////////////////////////////////////////////////////////////////////////////////////////
 
+void STimingView::SelectEventInstance(uint32 TimerId, ESelectEventType Type, bool bUseSelection)
+{
+	SelectTimingEvent(nullptr, false, false);
+
+	double IntervalStart = 0.0f;
+	double IntervalEnd = std::numeric_limits<double>::infinity();
+
+	if (bUseSelection && SelectionEndTime > SelectionStartTime)
+	{
+		IntervalStart = SelectionStartTime;
+		IntervalEnd = SelectionEndTime;
+	}
+
+	TSharedPtr<const ITimingEvent> TimingEvent;
+
+	if (Type == ESelectEventType::Min)
+	{
+		TimingEvent = ThreadTimingSharedState->FindMinEventInstance(TimerId, IntervalStart, IntervalEnd);
+	}
+	else if (Type == ESelectEventType::Max)
+	{
+		TimingEvent = ThreadTimingSharedState->FindMaxEventInstance(TimerId, IntervalStart, IntervalEnd);
+	}
+
+	if (TimingEvent.IsValid())
+	{
+		SelectTimingEvent(TimingEvent, true, true);
+	}
+	else
+	{
+		FMessageLog ReportMessageLog(FTimingProfilerManager::Get()->GetLogListingName());
+		ReportMessageLog.Error(LOCTEXT("NoEventFound", "No event instance found!"));
+		ReportMessageLog.Notify();
+	}
+}
+
+////////////////////////////////////////////////////////////////////////////////////////////////////
+
 #undef INSIGHTS_ACTIVATE_BENCHMARK
 #undef LOCTEXT_NAMESPACE
diff --git a/Engine/Source/Developer/TraceInsights/Private/Insights/Widgets/STimingView.h b/Engine/Source/Developer/TraceInsights/Private/Insights/Widgets/STimingView.h
--- a/Engine/Source/Developer/TraceInsights/Private/Insights/Widgets/STimingView.h
+++ b/Engine/Source/Developer/TraceInsights/Private/Insights/Widgets/STimingView.h
@@ -49,6 +49,12 @@ namespace Insights
 	enum class ETimingEventsColoringMode : uint32;
 }
 
+enum class ESelectEventType : uint32
+{
+	Min = 0,
+	Max = 1
+};
+
 /** A custom widget used to display timing events. */
 class STimingView : public SCompoundWidget, public Insights::ITimingViewSession
 {
@@ -365,7 +371,7 @@ public:
 	const TSharedPtr<const ITimingEvent> GetSelectedEvent() const { return SelectedEvent; }
 
 	void SelectTimingTrack(const TSharedPtr<FBaseTimingTrack> InTrack, bool bBringTrackIntoView);
-	void SelectTimingEvent(const TSharedPtr<const ITimingEvent> InEvent, bool bBringEventIntoView);
+	void SelectTimingEvent(const TSharedPtr<const ITimingEvent> InEvent, bool bBringEventIntoViewHorizontally, bool bBringEventIntoViewVertically = false);
 
 	const TSharedPtr<ITimingEventFilter> GetEventFilter() const { return TimingEventFilter; }
 	void SetEventFilter(const TSharedPtr<ITimingEventFilter> InEventFilter);
@@ -389,6 +395,8 @@ public:
 
 	TMap<uint64, TSharedPtr<FBaseTimingTrack>>& GetAllTracks() { return AllTracks; }
 
+	void SelectEventInstance(uint32 TimerId, ESelectEventType Type, bool bUseSelection);
+
 protected:
 	virtual FVector2D ComputeDesiredSize(float) const override
 	{
@@ -699,6 +707,8 @@ protected:
 	bool bDrawTopSeparatorLine;
 	bool bDrawBottomSeparatorLine;
 
+	bool bBringSelectedEventIntoViewVerticallyOnNextTick = false;
+
 	// Debug stats
 	int32 NumUpdatedEvents;
 	TFixedCircularBuffer<uint64, 32> PreUpdateTracksDurationHistory;
