[Insights] Fix localization key error.

diff --git a/Engine/Source/Developer/TraceInsights/Private/Insights/Widgets/STimingView.cpp b/Engine/Source/Developer/TraceInsights/Private/Insights/Widgets/STimingView.cpp
--- a/Engine/Source/Developer/TraceInsights/Private/Insights/Widgets/STimingView.cpp
+++ b/Engine/Source/Developer/TraceInsights/Private/Insights/Widgets/STimingView.cpp
@@ -5368,7 +5368,7 @@ void STimingView::SelectEventInstance(uint32 TimerId, ESelectEventType Type, boo
 	else
 	{
 		FMessageLog ReportMessageLog(FTimingProfilerManager::Get()->GetLogListingName());
-		ReportMessageLog.Error(LOCTEXT("NoEventFound", "No event instance found!"));
+		ReportMessageLog.Error(LOCTEXT("NoEventInstanceFound", "No event instance found!"));
 		ReportMessageLog.Notify();
 	}
 }
