[Insights] The package tables now shows a descriptive message until analysis is complete

diff --git a/Engine/Source/Developer/TraceInsights/Private/Insights/CookProfiler/Widgets/SPackageTableTreeView.cpp b/Engine/Source/Developer/TraceInsights/Private/Insights/CookProfiler/Widgets/SPackageTableTreeView.cpp
--- a/Engine/Source/Developer/TraceInsights/Private/Insights/CookProfiler/Widgets/SPackageTableTreeView.cpp
+++ b/Engine/Source/Developer/TraceInsights/Private/Insights/CookProfiler/Widgets/SPackageTableTreeView.cpp
@@ -428,6 +428,20 @@ void SPackageTableTreeView::InitAvailableViewPresets()
 
 ////////////////////////////////////////////////////////////////////////////////////////////////////
 
+void SPackageTableTreeView::UpdateBannerText()
+{
+	if (!bDataLoaded)
+	{
+		TreeViewBannerText = LOCTEXT("DataWillLoad", "Package data will load when session analysis is complete.");
+	}
+	else
+	{
+		STableTreeView::UpdateBannerText();
+	}
+}
+
+////////////////////////////////////////////////////////////////////////////////////////////////////
+
 } // namespace Insights
 
 #undef LOCTEXT_NAMESPACE
diff --git a/Engine/Source/Developer/TraceInsights/Private/Insights/CookProfiler/Widgets/SPackageTableTreeView.h b/Engine/Source/Developer/TraceInsights/Private/Insights/CookProfiler/Widgets/SPackageTableTreeView.h
--- a/Engine/Source/Developer/TraceInsights/Private/Insights/CookProfiler/Widgets/SPackageTableTreeView.h
+++ b/Engine/Source/Developer/TraceInsights/Private/Insights/CookProfiler/Widgets/SPackageTableTreeView.h
@@ -73,6 +73,8 @@ protected:
 
 	virtual void ExtendMenu(FMenuBuilder& MenuBuilder) override;
 
+	virtual void UpdateBannerText() override;
+
 private:
 	void AddCommmands();
 
