[Insights] Fix missing labels in Timing Graph Track.

diff --git a/Engine/Source/Developer/TraceInsights/Private/Insights/ViewModels/TimingGraphTrack.cpp b/Engine/Source/Developer/TraceInsights/Private/Insights/ViewModels/TimingGraphTrack.cpp
--- a/Engine/Source/Developer/TraceInsights/Private/Insights/ViewModels/TimingGraphTrack.cpp
+++ b/Engine/Source/Developer/TraceInsights/Private/Insights/ViewModels/TimingGraphTrack.cpp
@@ -134,19 +134,10 @@ FTimingGraphTrack::FTimingGraphTrack(TSharedPtr<STimingView> InTimingView)
 	, TimingView(InTimingView)
 	//, SharedValueViewport()
 {
-	EnabledOptions = //EGraphOptions::ShowDebugInfo |
-					 //EGraphOptions::ShowPoints |
-					 EGraphOptions::ShowPointsWithBorder |
-					 EGraphOptions::ShowLines |
-					 EGraphOptions::ShowPolygon |
-					 EGraphOptions::UseEventDuration |
-					 //EGraphOptions::ShowBars |
-					 EGraphOptions::ShowBaseline |
-					 EGraphOptions::ShowVerticalAxisGrid |
-					 EGraphOptions::ShowHeader |
-					 EGraphOptions::None;
-
 	LoadDefaultSettings();
+	
+	// Add non editable options.
+	EnabledOptions = EnabledOptions | EGraphOptions::ShowBaseline | EGraphOptions::ShowVerticalAxisGrid | EGraphOptions::ShowHeader;
 }
 
 ////////////////////////////////////////////////////////////////////////////////////////////////////
@@ -834,7 +825,6 @@ void FTimingGraphTrack::LoadDefaultSettings()
 	TSharedPtr<STimingView> TimingViewPtr = TimingView.Pin();
 	if (TimingViewPtr.IsValid() && TimingViewPtr->GetName() == FInsightsManagerTabs::TimingProfilerTabId)
 	{
-		EnabledOptions = EGraphOptions::None;
 		const FInsightsSettings& Settings = FInsightsManager::Get()->GetSettings();
 		if (Settings.GetTimingViewMainGraphShowPoints())
 		{
