[Insights] Fix issues with plotting timers when Timing View and the Frame Track are recreated.

diff --git a/Engine/Source/Developer/TraceInsights/Private/Insights/ViewModels/TimingGraphTrack.cpp b/Engine/Source/Developer/TraceInsights/Private/Insights/ViewModels/TimingGraphTrack.cpp
--- a/Engine/Source/Developer/TraceInsights/Private/Insights/ViewModels/TimingGraphTrack.cpp
+++ b/Engine/Source/Developer/TraceInsights/Private/Insights/ViewModels/TimingGraphTrack.cpp
@@ -20,6 +20,7 @@
 #include "Insights/ViewModels/TimingTrackViewport.h"
 #include "Insights/ViewModels/ThreadTimingTrack.h"
 #include "Insights/Widgets/STimingProfilerWindow.h"
+#include "Insights/Widgets/STimersView.h"
 #include "Insights/Widgets/STimingView.h"
 
 #define LOCTEXT_NAMESPACE "GraphTrack"
@@ -138,6 +139,8 @@ FTimingGraphTrack::FTimingGraphTrack(TSharedPtr<STimingView> InTimingView)
 	
 	// Add non editable options.
 	EnabledOptions = EnabledOptions | EGraphOptions::ShowBaseline | EGraphOptions::ShowVerticalAxisGrid | EGraphOptions::ShowHeader;
+
+	bNotifyTimersOnDestruction = InTimingView->GetName() == FInsightsManagerTabs::TimingProfilerTabId;
 }
 
 ////////////////////////////////////////////////////////////////////////////////////////////////////
@@ -172,6 +175,33 @@ FTimingGraphTrack::~FTimingGraphTrack()
 			RenderingFramesSeries->VisibilityChangedDelegate.Remove(RenderingFrameSeriesVisibilityHandle);
 		}
 	}
+
+	TSharedPtr<STimersView> TimersView;
+	if (bNotifyTimersOnDestruction)
+	{
+		TSharedPtr<STimingProfilerWindow> ProfilerWindow = FTimingProfilerManager::Get()->GetProfilerWindow();
+		if (ProfilerWindow.IsValid())
+		{
+			TimersView = ProfilerWindow->GetTimersView();
+		}
+	}
+
+	if(TimersView)
+	{
+		for (const TSharedPtr<FGraphSeries>& Series : AllSeries)
+		{
+			const TSharedPtr<FTimingGraphSeries> TimingSeries = StaticCastSharedPtr<FTimingGraphSeries>(Series);
+			if (TimingSeries.IsValid() &&
+				(TimingSeries->Type == FTimingGraphSeries::ESeriesType::Timer || TimingSeries->Type == FTimingGraphSeries::ESeriesType::FrameStatsTimer))
+			{
+				FTimerNodePtr TimerNode = TimersView->GetTimerNode(TimingSeries->TimerId);
+				if (TimerNode)
+				{
+					TimerNode->OnRemovedFromGraph();
+				}
+			}
+		};
+	}
 }
 
 ////////////////////////////////////////////////////////////////////////////////////////////////////
@@ -976,4 +1006,24 @@ void FTimingGraphTrack::DrawVerticalAxisGrid(const ITimingTrackDrawContext& Cont
 
 ////////////////////////////////////////////////////////////////////////////////////////////////////
 
+uint32 FTimingGraphTrack::GetNumSeriesForTimer(uint32 TimerId)
+{
+	uint32 NumSeries = 0;
+
+	for(const TSharedPtr<FGraphSeries>& Series : AllSeries)
+	{
+		const TSharedPtr<FTimingGraphSeries> TimingSeries = StaticCastSharedPtr<FTimingGraphSeries>(Series);
+		if (TimingSeries.IsValid() && 
+			(TimingSeries->Type == FTimingGraphSeries::ESeriesType::Timer || TimingSeries->Type == FTimingGraphSeries::ESeriesType::FrameStatsTimer) && 
+			TimingSeries->TimerId == TimerId)
+		{
+			++NumSeries;
+		}
+	};
+
+	return NumSeries;
+}
+
+////////////////////////////////////////////////////////////////////////////////////////////////////
+
 #undef LOCTEXT_NAMESPACE
diff --git a/Engine/Source/Developer/TraceInsights/Private/Insights/ViewModels/TimingGraphTrack.h b/Engine/Source/Developer/TraceInsights/Private/Insights/ViewModels/TimingGraphTrack.h
--- a/Engine/Source/Developer/TraceInsights/Private/Insights/ViewModels/TimingGraphTrack.h
+++ b/Engine/Source/Developer/TraceInsights/Private/Insights/ViewModels/TimingGraphTrack.h
@@ -99,6 +99,8 @@ public:
 	TSharedPtr<FTimingGraphSeries> AddStatsCounterSeries(uint32 CounterId, FLinearColor Color);
 	void RemoveStatsCounterSeries(uint32 CounterId);
 
+	uint32 GetNumSeriesForTimer(uint32 TimerId);
+
 protected:
 	void UpdateFrameSeries(FTimingGraphSeries& Series, const FTimingTrackViewport& Viewport);
 	void UpdateTimerSeries(FTimingGraphSeries& Series, const FTimingTrackViewport& Viewport);
@@ -123,6 +125,7 @@ private:
 	FDelegateHandle RenderingFrameSeriesVisibilityHandle;
 
 	TWeakPtr<STimingView> TimingView;
+	bool bNotifyTimersOnDestruction;
 };
 
 ////////////////////////////////////////////////////////////////////////////////////////////////////
diff --git a/Engine/Source/Developer/TraceInsights/Private/Insights/Widgets/SFrameTrack.cpp b/Engine/Source/Developer/TraceInsights/Private/Insights/Widgets/SFrameTrack.cpp
--- a/Engine/Source/Developer/TraceInsights/Private/Insights/Widgets/SFrameTrack.cpp
+++ b/Engine/Source/Developer/TraceInsights/Private/Insights/Widgets/SFrameTrack.cpp
@@ -27,6 +27,7 @@
 #include "Insights/ViewModels/FrameTrackHelper.h"
 #include "Insights/ViewModels/ThreadTimingTrack.h"
 #include "Insights/Widgets/STimingProfilerWindow.h"
+#include "Insights/Widgets/STimersView.h"
 #include "Insights/Widgets/STimingView.h"
 
 #include <limits>
@@ -63,6 +64,29 @@ SFrameTrack::~SFrameTrack()
 			}
 		}
 	}
+
+	TSharedPtr<STimersView> TimersView;
+	TSharedPtr<STimingProfilerWindow> ProfilerWindow = FTimingProfilerManager::Get()->GetProfilerWindow();
+	if (ProfilerWindow.IsValid())
+	{
+		TimersView = ProfilerWindow->GetTimersView();
+	}
+
+	if (TimersView)
+	{
+		for (const TSharedPtr<FFrameTrackSeries>& Series : AllSeries)
+		{
+			if (Series.IsValid() && Series->Type == EFrameTrackSeriesType::TimerFrameStats)
+			{
+				const TSharedPtr<FTimerFrameStatsTrackSeries> TimerSeries = StaticCastSharedPtr<FTimerFrameStatsTrackSeries>(Series);
+				FTimerNodePtr TimerNode = TimersView->GetTimerNode(TimerSeries->TimerId);
+				if (TimerNode)
+				{
+					TimerNode->OnRemovedFromGraph();
+				}
+			}
+		};
+	}
 }
 
 ////////////////////////////////////////////////////////////////////////////////////////////////////
@@ -1650,4 +1674,25 @@ bool SFrameTrack::RemoveTimerFrameStatSeries(ETraceFrameType FrameType, uint32 T
 
 ////////////////////////////////////////////////////////////////////////////////////////////////////
 
+uint32 SFrameTrack::GetNumSeriesForTimer(uint32 TimerId)
+{
+	uint32 NumSeries = 0;
+
+	for (const TSharedPtr<FFrameTrackSeries>& Series : AllSeries)
+	{
+		if (Series->Type == EFrameTrackSeriesType::TimerFrameStats)
+		{
+			const TSharedPtr<FTimerFrameStatsTrackSeries> TimerSeries = StaticCastSharedPtr<FTimerFrameStatsTrackSeries>(Series);
+			if (TimerSeries->TimerId == TimerId)
+			{
+				++NumSeries;
+			}
+		}
+	};
+
+	return NumSeries;
+}
+
+////////////////////////////////////////////////////////////////////////////////////////////////////
+
 #undef LOCTEXT_NAMESPACE
diff --git a/Engine/Source/Developer/TraceInsights/Private/Insights/Widgets/SFrameTrack.h b/Engine/Source/Developer/TraceInsights/Private/Insights/Widgets/SFrameTrack.h
--- a/Engine/Source/Developer/TraceInsights/Private/Insights/Widgets/SFrameTrack.h
+++ b/Engine/Source/Developer/TraceInsights/Private/Insights/Widgets/SFrameTrack.h
@@ -133,6 +133,8 @@ public:
 	TSharedPtr<FTimerFrameStatsTrackSeries> AddTimerFrameStatSeries(ETraceFrameType FrameType, uint32 TimerId, FLinearColor Color, FText Name);
 	bool RemoveTimerFrameStatSeries(ETraceFrameType FrameType, uint32 TimerId);
 
+	uint32 GetNumSeriesForTimer(uint32 TimerId);
+
 protected:
 	TSharedRef<FFrameTrackSeries> FindOrAddSeries(ETraceFrameType FrameType);
 	TSharedPtr<FFrameTrackSeries> FindSeries(ETraceFrameType FrameType) const;
diff --git a/Engine/Source/Developer/TraceInsights/Private/Insights/Widgets/STimersView.cpp b/Engine/Source/Developer/TraceInsights/Private/Insights/Widgets/STimersView.cpp
--- a/Engine/Source/Developer/TraceInsights/Private/Insights/Widgets/STimersView.cpp
+++ b/Engine/Source/Developer/TraceInsights/Private/Insights/Widgets/STimersView.cpp
@@ -821,19 +821,26 @@ void STimersView::TreeView_BuildPlotTimerMenu(FMenuBuilder& MenuBuilder)
 	const int32 NumSelectedNodes = SelectedNodes.Num();
 	FTimerNodePtr SelectedNode = NumSelectedNodes ? SelectedNodes[0] : nullptr;
 
-	auto CanExecute = [NumSelectedNodes, SelectedNode]()
+	auto CanExecuteAddToGraphTrack = [NumSelectedNodes, SelectedNode]()
 	{
 		TSharedPtr<STimingProfilerWindow> Wnd = FTimingProfilerManager::Get()->GetProfilerWindow();
 		TSharedPtr<STimingView> TimingView = Wnd.IsValid() ? Wnd->GetTimingView() : nullptr;
 		return TimingView.IsValid() && NumSelectedNodes == 1 && SelectedNode.IsValid() && SelectedNode->GetType() != ETimerNodeType::Group;
 	};
 
+	auto CanExecuteAddToFramesTrack = [NumSelectedNodes, SelectedNode]()
+	{
+		TSharedPtr<STimingProfilerWindow> Wnd = FTimingProfilerManager::Get()->GetProfilerWindow();
+		TSharedPtr<SFrameTrack> FrameTrack = Wnd.IsValid() ? Wnd->GetFrameView() : nullptr;
+		return FrameTrack.IsValid() && NumSelectedNodes == 1 && SelectedNode.IsValid() && SelectedNode->GetType() != ETimerNodeType::Group;
+	};
+
 	MenuBuilder.BeginSection("Instance", LOCTEXT("Plot_Series_Instance_Section", "Instance"));
 
 	// Add/remove series to/from graph track
 	{
 		FUIAction Action_ToggleTimerInGraphTrack;
-		Action_ToggleTimerInGraphTrack.CanExecuteAction = FCanExecuteAction::CreateLambda(CanExecute);
+		Action_ToggleTimerInGraphTrack.CanExecuteAction = FCanExecuteAction::CreateLambda(CanExecuteAddToGraphTrack);
 		Action_ToggleTimerInGraphTrack.ExecuteAction = FExecuteAction::CreateSP(this, &STimersView::ToggleTimingViewMainGraphEventSeries, SelectedNode);
 
 		if (SelectedNode.IsValid() &&
@@ -871,7 +878,7 @@ void STimersView::TreeView_BuildPlotTimerMenu(FMenuBuilder& MenuBuilder)
 	// Add/remove game frame stats series to/from graph track
 	{
 		FUIAction Action_ToggleFrameStatsTimerInGraphTrack;
-		Action_ToggleFrameStatsTimerInGraphTrack.CanExecuteAction = FCanExecuteAction::CreateLambda(CanExecute);
+		Action_ToggleFrameStatsTimerInGraphTrack.CanExecuteAction = FCanExecuteAction::CreateLambda(CanExecuteAddToGraphTrack);
 		Action_ToggleFrameStatsTimerInGraphTrack.ExecuteAction = FExecuteAction::CreateSP(this, &STimersView::ToggleTimingViewMainGraphEventFrameStatsSeries, SelectedNode, ETraceFrameType::TraceFrameType_Game);
 
 		if (SelectedNode.IsValid() &&
@@ -905,7 +912,7 @@ void STimersView::TreeView_BuildPlotTimerMenu(FMenuBuilder& MenuBuilder)
 	// Add/remove game frame stats series to/from frame track
 	{
 		FUIAction Action_ToggleFrameStatsTimerInFrameTrack;
-		Action_ToggleFrameStatsTimerInFrameTrack.CanExecuteAction = FCanExecuteAction::CreateLambda(CanExecute);
+		Action_ToggleFrameStatsTimerInFrameTrack.CanExecuteAction = FCanExecuteAction::CreateLambda(CanExecuteAddToFramesTrack);
 		Action_ToggleFrameStatsTimerInFrameTrack.ExecuteAction = FExecuteAction::CreateSP(this, &STimersView::ToggleFrameTrackSeries, SelectedNode, ETraceFrameType::TraceFrameType_Game);
 
 		if (SelectedNode.IsValid() &&
@@ -943,7 +950,7 @@ void STimersView::TreeView_BuildPlotTimerMenu(FMenuBuilder& MenuBuilder)
 	// Add/remove rendering frame stats series to/from graph track
 	{
 		FUIAction Action_ToggleFrameStatsTimerInGraphTrack;
-		Action_ToggleFrameStatsTimerInGraphTrack.CanExecuteAction = FCanExecuteAction::CreateLambda(CanExecute);
+		Action_ToggleFrameStatsTimerInGraphTrack.CanExecuteAction = FCanExecuteAction::CreateLambda(CanExecuteAddToGraphTrack);
 		Action_ToggleFrameStatsTimerInGraphTrack.ExecuteAction = FExecuteAction::CreateSP(this, &STimersView::ToggleTimingViewMainGraphEventFrameStatsSeries, SelectedNode, ETraceFrameType::TraceFrameType_Rendering);
 
 		if (SelectedNode.IsValid() &&
@@ -977,7 +984,7 @@ void STimersView::TreeView_BuildPlotTimerMenu(FMenuBuilder& MenuBuilder)
 	// Add/remove rendering frame stats series to/from frame track
 	{
 		FUIAction Action_ToggleFrameStatsTimerInFrameTrack;
-		Action_ToggleFrameStatsTimerInFrameTrack.CanExecuteAction = FCanExecuteAction::CreateLambda(CanExecute);
+		Action_ToggleFrameStatsTimerInFrameTrack.CanExecuteAction = FCanExecuteAction::CreateLambda(CanExecuteAddToFramesTrack);
 		Action_ToggleFrameStatsTimerInFrameTrack.ExecuteAction = FExecuteAction::CreateSP(this, &STimersView::ToggleFrameTrackSeries, SelectedNode, ETraceFrameType::TraceFrameType_Rendering);
 
 		if (SelectedNode.IsValid() &&
@@ -2419,6 +2426,9 @@ void STimersView::RebuildTree(bool bResync)
 			check(TimerCount > PreviousNodeCount);
 			TimerNodes.Reserve(TimerCount);
 
+			TSharedPtr<FTimingGraphTrack> GraphTrack = GetTimingViewMainGraphTrack();
+			TSharedPtr<SFrameTrack> FrameTrack = GetFrameTrack();
+
 			// Add nodes only for new timers.
 			for (uint32 TimerIndex = PreviousNodeCount; TimerIndex < TimerCount; ++TimerIndex)
 			{
@@ -2427,6 +2437,27 @@ void STimersView::RebuildTree(bool bResync)
 				const ETimerNodeType Type = Timer.IsGpuTimer ? ETimerNodeType::GpuScope : ETimerNodeType::CpuScope;
 				FTimerNodePtr TimerNodePtr = MakeShared<FTimerNode>(Timer.Id, Timer.Name, Type);
 				TimerNodePtr->SetDefaultSortOrder(TimerIndex + 1);
+
+				if (GraphTrack.IsValid())
+				{
+					uint32 NumSeries = GraphTrack->GetNumSeriesForTimer(Timer.Id);
+					while (NumSeries > 0)
+					{
+						TimerNodePtr->OnAddedToGraph();
+						--NumSeries;
+					}
+				}
+
+				if (FrameTrack.IsValid())
+				{
+					uint32 NumSeries = FrameTrack->GetNumSeriesForTimer(Timer.Id);
+					while (NumSeries > 0)
+					{
+						TimerNodePtr->OnAddedToGraph();
+						--NumSeries;
+					}
+				}
+
 				TimerNodes.Add(TimerNodePtr);
 			}
 			ensure(TimerNodes.Num() == TimerCount);
