[Insights] Fix issue where GPU timers are not ploted in the graph track.

diff --git a/Engine/Source/Developer/TraceInsights/Private/Insights/ViewModels/ThreadTimingTrack.cpp b/Engine/Source/Developer/TraceInsights/Private/Insights/ViewModels/ThreadTimingTrack.cpp
--- a/Engine/Source/Developer/TraceInsights/Private/Insights/ViewModels/ThreadTimingTrack.cpp
+++ b/Engine/Source/Developer/TraceInsights/Private/Insights/ViewModels/ThreadTimingTrack.cpp
@@ -425,6 +425,16 @@ void FThreadTimingSharedState::GetVisibleTimelineIndexes(TSet<uint32>& OutSet) c
 			OutSet.Add(Track.GetTimelineIndex());
 		}
 	}
+
+	if (GpuTrack.IsValid() && GpuTrack->IsVisible())
+	{
+		OutSet.Add(GpuTrack->GetTimelineIndex());
+	}
+
+	if (Gpu2Track.IsValid() && Gpu2Track->IsVisible())
+	{
+		OutSet.Add(Gpu2Track->GetTimelineIndex());
+	}
 }
 
 ////////////////////////////////////////////////////////////////////////////////////////////////////
