[Insights] Session Browser: Fixed locking for settings and traces synced from the store browser's thread.

diff --git a/Engine/Source/Developer/TraceInsights/Private/Insights/StoreService/StoreBrowser.cpp b/Engine/Source/Developer/TraceInsights/Private/Insights/StoreService/StoreBrowser.cpp
--- a/Engine/Source/Developer/TraceInsights/Private/Insights/StoreService/StoreBrowser.cpp
+++ b/Engine/Source/Developer/TraceInsights/Private/Insights/StoreService/StoreBrowser.cpp
@@ -19,16 +19,8 @@ namespace Insights
 ////////////////////////////////////////////////////////////////////////////////////////////////////
 
 FStoreBrowser::FStoreBrowser()
-	: bRunning(true)
-	, Thread(nullptr)
-	, TracesCriticalSection()
-	, bTracesLocked(false)
-	, StoreChangeSerial(0)
-	, TracesChangeSerial(0)
-	, Traces()
-	, TraceMap()
-	, LiveTraceMap()
 {
+	bRunning = true;
 	Thread = FRunnableThread::Create(this, TEXT("StoreBrowser"));
 }
 
@@ -82,12 +74,30 @@ void FStoreBrowser::Exit()
 
 void FStoreBrowser::Refresh()
 {
-	if (!IsLocked())
+	if (AreSettingsLocked() || AreTracesLocked())
 	{
-		FScopeLock Lock(&TracesCriticalSection);
+		return;
+	}
 
-		StoreChangeSerial = 0;
+	StoreChangeSerial = 0;
+	StoreSettingsChangeSerial = 0;
+
+	{
+		FScopeLock Lock(&SettingsCriticalSection);
+		SettingsChangeSerial = 0;
+		Host.Reset();
+		Version.Reset();
+		StoreDirectory.Reset();
+		WatchDirectories.Reset();
+	}
+
+	{
+		FScopeLock Lock(&TracesCriticalSection);
 		TracesChangeSerial = 0;
+		for (const TSharedPtr<FStoreBrowserTraceInfo>& Trace : Traces)
+		{
+			Trace->MetadataUpdateCount = 0;
+		}
 		Traces.Reset();
 		TraceMap.Reset();
 		LiveTraceMap.Reset();
@@ -122,9 +132,19 @@ void FStoreBrowser::UpdateTraces()
 	FStopwatch StopwatchTotal;
 	StopwatchTotal.Start();
 
+	// Update store host.
+	if (Host.IsEmpty())
+	{
+		FScopeLock Lock(&SettingsCriticalSection);
+		SettingsChangeSerial++;
+		Host = FInsightsManager::Get()->GetLastStoreHost();
+	}
+
 	// Check if connection to store is still active We want to do this without locking the critical
 	// section, in case the UI needs to read values. The output is an atomic anyway.
 	{
+		EConnectionStatus PreviousConnectionStatus = ConnectionStatus;
+
 		bool bIsValid = StoreClient->IsValid();
 		while (!bIsValid && bRunning) // TODO: Perhaps a max reconnection attempts is needed?
 		{
@@ -132,7 +152,7 @@ void FStoreBrowser::UpdateTraces()
 			uint32 SecondsToSleep = static_cast<uint32>(ReconnectionFrequency);
 			do
 			{
-				FPlatformProcess::Sleep(1.0);
+				FPlatformProcess::Sleep(1.0f);
 				ConnectionStatus.store(static_cast<EConnectionStatus>(--SecondsToSleep));
 			} while (SecondsToSleep);
 
@@ -152,12 +172,12 @@ void FStoreBrowser::UpdateTraces()
 			// If there is no connection, no point in checking for sessions
 			return;
 		}
-		else
-		{
-			FScopeLock Lock(&TracesCriticalSection);
-			Host = FInsightsManager::Get()->GetLastStoreHost();
 
-			// Get the server version
+		if (ConnectionStatus == EConnectionStatus::Connected &&
+			(PreviousConnectionStatus != EConnectionStatus::Connected || StoreChangeSerial == 0))
+		{
+			// Get the server version.
+			FString StoreVersion;
 			{
 				FScopeLock StoreClientLock(&GetStoreClientCriticalSection());
 				const auto VersionResult = StoreClient->GetVersion();
@@ -170,23 +190,30 @@ void FStoreBrowser::UpdateTraces()
 					{
 						VersionString << TEXT(" (") << VersionResult->GetConfiguration() << TEXT(")");
 					}
-					Version = VersionString.ToString();
+					StoreVersion = VersionString.ToString();
 				}
 				else
 				{
-					Version = TEXT("Unknown");
+					StoreVersion = TEXT("Unknown");
 				}
 			}
+
+			// Update settings.
+			{
+				FScopeLock Lock(&SettingsCriticalSection);
+				SettingsChangeSerial++;
+				Host = FInsightsManager::Get()->GetLastStoreHost();
+				Version = StoreVersion;
+			}
 		}
 	}
 
-
 	// Check if the list of trace sessions or store settings have changed.
 	{
-		TArray<FString> NewWatchDirectories;
-		FString NewStoreDirectory;
-		uint32 NewSettingsChangeSerial = 0;
 		uint32 NewStoreChangeSerial = 0;
+		uint32 NewStoreSettingsChangeSerial = 0;
+		FString NewStoreDirectory;
+		TArray<FString> NewWatchDirectories;
 
 		{
 			// Get Trace Store status.
@@ -194,25 +221,26 @@ void FStoreBrowser::UpdateTraces()
 			const UE::Trace::FStoreClient::FStatus* Status = StoreClient->GetStatus();
 			if (Status)
 			{
-				NewSettingsChangeSerial = Status->GetSettingsSerial();
 				NewStoreChangeSerial = Status->GetChangeSerial();
-				if (SettingsChangeSerial != NewSettingsChangeSerial)
+				NewStoreSettingsChangeSerial = Status->GetSettingsSerial();
+				if (StoreSettingsChangeSerial != NewStoreSettingsChangeSerial)
 				{
-					Status->GetWatchDirectories(NewWatchDirectories);
 					NewStoreDirectory = FString(Status->GetStoreDir());
+					Status->GetWatchDirectories(NewWatchDirectories);
 				}
 			}
 		}
 
-		if (SettingsChangeSerial != NewSettingsChangeSerial)
+		if (StoreSettingsChangeSerial != NewStoreSettingsChangeSerial)
 		{
-			SettingsChangeSerial = NewSettingsChangeSerial;
+			StoreSettingsChangeSerial = NewStoreSettingsChangeSerial;
 
-			// Update watch dirs
-			FScopeLock Lock(&TracesCriticalSection);
+			// Update settings.
+			FScopeLock Lock(&SettingsCriticalSection);
+			SettingsChangeSerial++;
+			StoreDirectory = NewStoreDirectory;
 			WatchDirectories.Empty();
 			WatchDirectories = NewWatchDirectories;
-			StoreDirectory = NewStoreDirectory;
 		}
 
 		if (StoreChangeSerial != NewStoreChangeSerial)
@@ -440,37 +468,34 @@ void FStoreBrowser::UpdateTraces()
 		FStopwatch Stopwatch;
 		Stopwatch.Start();
 
-		struct FStoreBrowserTraceInfoBox
-		{
-			FStoreBrowserTraceInfo* Ptr;
-		};
-		TArray<FStoreBrowserTraceInfoBox> TracesToUpdate;
+		TArray<TSharedPtr<FStoreBrowserTraceInfo>> TracesToUpdate;
 
-		for (TSharedPtr<FStoreBrowserTraceInfo>& TracePtr : Traces)
+		for (const TSharedPtr<FStoreBrowserTraceInfo>& TraceInfo : Traces)
 		{
-			if (TracePtr->MetadataUpdateCount > 0)
+			if (TraceInfo->MetadataUpdateCount > 0)
 			{
-				TracesToUpdate.Add({ TracePtr.Get() });
+				TracesToUpdate.Add(TraceInfo);
 			}
 		}
 
 		// Sort descending by timestamp (i.e. to update the newer traces first).
-		TracesToUpdate.Sort([](const FStoreBrowserTraceInfoBox& A, const FStoreBrowserTraceInfoBox& B)
+		TracesToUpdate.Sort([](const TSharedPtr<FStoreBrowserTraceInfo>& A, const TSharedPtr<FStoreBrowserTraceInfo>& B)
 			{
-				return A.Ptr->Timestamp.GetTicks() > B.Ptr->Timestamp.GetTicks();
+				return A->Timestamp.GetTicks() > B->Timestamp.GetTicks();
 			});
 
 #if 0
-		for (const FStoreBrowserTraceInfoBox& Trace : TracesToUpdate)
+		for (TSharedPtr<FStoreBrowserTraceInfo>& Trace : TracesToUpdate)
 		{
-			UpdateMetadata(*Trace.Ptr);
+			UpdateMetadata(Trace);
 		}
 #else
 		if (TracesToUpdate.Num() > 0)
 		{
+			UE_LOG(TraceInsights, Log, TEXT("[StoreBrowser] Starting metadata update for %d trace(s)."), TracesToUpdate.Num());
 			ParallelFor(TracesToUpdate.Num(), [this, &TracesToUpdate](uint32 Index)
 				{
-					UpdateMetadata(*TracesToUpdate[Index].Ptr);
+					UpdateMetadata(TracesToUpdate[Index]);
 				},
 				EParallelForFlags::BackgroundPriority);
 		}
@@ -503,6 +528,10 @@ void FStoreBrowser::ResetTraces()
 	{
 		FScopeLock Lock(&TracesCriticalSection);
 		TracesChangeSerial++;
+		for (const TSharedPtr<FStoreBrowserTraceInfo>& Trace : Traces)
+		{
+			Trace->MetadataUpdateCount = 0;
+		}
 		Traces.Reset();
 		TraceMap.Reset();
 	}
@@ -510,8 +539,14 @@ void FStoreBrowser::ResetTraces()
 
 ////////////////////////////////////////////////////////////////////////////////////////////////////
 
-void FStoreBrowser::UpdateMetadata(FStoreBrowserTraceInfo& Trace)
+void FStoreBrowser::UpdateMetadata(TSharedPtr<FStoreBrowserTraceInfo> TraceInfoPtr)
 {
+	FStoreBrowserTraceInfo& Trace = *TraceInfoPtr;
+	if (Trace.MetadataUpdateCount == 0)
+	{
+		return;
+	}
+
 	UE::Trace::FStoreClient* StoreClient = GetStoreClient();
 	if (StoreClient == nullptr)
 	{
@@ -535,6 +570,7 @@ void FStoreBrowser::UpdateMetadata(FStoreBrowserTraceInfo& Trace)
 			Ready = 0,
 			StoppedByReadSizeLimit,
 			StoppedByTimeLimit,
+			Cancelled
 		};
 
 		virtual int32 Read(void* Data, uint32 Size) override
@@ -545,6 +581,12 @@ void FStoreBrowser::UpdateMetadata(FStoreBrowserTraceInfo& Trace)
 				return 0;
 			}
 
+			if (Trace->MetadataUpdateCount == 0)
+			{
+				Status = EReadStatus::Cancelled;
+				return 0;
+			}
+
 			Stopwatch.Update();
 			if (Stopwatch.GetAccumulatedTime() > TimeLimit)
 			{
@@ -564,6 +606,7 @@ void FStoreBrowser::UpdateMetadata(FStoreBrowserTraceInfo& Trace)
 		}
 
 		UE::Trace::IInDataStream* Inner;
+		FStoreBrowserTraceInfo* Trace;
 		double TimeLimit = 1.0;
 		FStopwatch Stopwatch;
 		int32 BytesRead = 0;
@@ -572,6 +615,7 @@ void FStoreBrowser::UpdateMetadata(FStoreBrowserTraceInfo& Trace)
 
 	FDataStream DataStream;
 	DataStream.Inner = TraceData.Get();
+	DataStream.Trace = &Trace;
 	DataStream.TimeLimit = 1.0 + (double)(Trace.MetadataUpdateCount - 1) * 2.0; // 1s, 3s, 5s, ...
 	DataStream.Stopwatch.Start();
 
@@ -581,6 +625,7 @@ void FStoreBrowser::UpdateMetadata(FStoreBrowserTraceInfo& Trace)
 	Context.Process(DataStream).Wait();
 
 	// Update the FStoreBrowserTraceInfo object.
+	if (Trace.MetadataUpdateCount != 0)
 	{
 		FScopeLock Lock(&TracesCriticalSection);
 
diff --git a/Engine/Source/Developer/TraceInsights/Private/Insights/StoreService/StoreBrowser.h b/Engine/Source/Developer/TraceInsights/Private/Insights/StoreService/StoreBrowser.h
--- a/Engine/Source/Developer/TraceInsights/Private/Insights/StoreService/StoreBrowser.h
+++ b/Engine/Source/Developer/TraceInsights/Private/Insights/StoreService/StoreBrowser.h
@@ -4,11 +4,10 @@
 
 #include "Async/AsyncWork.h"
 #include "CoreMinimal.h"
+#include "HAL/CriticalSection.h"
+#include "HAL/PlatformAtomics.h"
 #include "HAL/Runnable.h"
 #include "HAL/RunnableThread.h"
-#include "HAL/ThreadSafeBool.h"
-#include "HAL/CriticalSection.h"
-#include "Templates/Atomic.h"
 
 namespace UE
 {
@@ -47,7 +46,7 @@ struct FStoreBrowserTraceInfo
 	EBuildTargetType TargetType = EBuildTargetType::Unknown;
 
 	bool bIsLive = false;
-	uint8 MetadataUpdateCount = 1;
+	std::atomic<uint8> MetadataUpdateCount = 1;
 
 	uint32 IpAddress = 0;
 
@@ -84,7 +83,7 @@ public:
 		// Attempting connection
 		Connecting = 0,
 		// Values between (start,end) is interpreted as
-		// number of seconds until next reconnection attemp
+		// number of seconds until next reconnection attempt
 		SecondsToReconnectStart,
 		SecondsToReconnectEnd = 0xfd,
 		// No connection could be made, no more reconnection
@@ -94,20 +93,25 @@ public:
 		Connected = 0xff
 	};
 
-	EConnectionStatus GetConnectionStatus() const { return ConnectionStatus; }
 	bool IsRunning() const { return bRunning; }
+	EConnectionStatus GetConnectionStatus() const { return ConnectionStatus; }
+
+	bool AreSettingsLocked() const { return bSettingsLocked; }
+	void LockSettings() { check(!bSettingsLocked); bSettingsLocked = true; SettingsCriticalSection.Lock(); }
+	void UnlockSettings() { check(bSettingsLocked); bSettingsLocked = false; SettingsCriticalSection.Unlock(); }
+	uint32 GetSettingsChangeSerial() const { check(bSettingsLocked); return SettingsChangeSerial; };
+	const FString& GetHost() const { check(bSettingsLocked); return Host; }
+	const FString& GetVersion() const { check(bSettingsLocked); return Version; }
+	const FString& GetStoreDirectory() const { check(bSettingsLocked); return StoreDirectory; }
+	const TArray<FString>& GetWatchDirectories() const { check(bSettingsLocked); return WatchDirectories; }
+
+	bool AreTracesLocked() const { return bTracesLocked; }
+	void LockTraces() { check(!bTracesLocked); bTracesLocked = true; TracesCriticalSection.Lock(); }
+	void UnlockTraces() { check(bTracesLocked); bTracesLocked = false; TracesCriticalSection.Unlock(); }
+	uint32 GetTracesChangeSerial() const { check(bTracesLocked); return TracesChangeSerial; }
+	const TArray<TSharedPtr<FStoreBrowserTraceInfo>>& GetTraces() const { check(bTracesLocked); return Traces; }
+	const TMap<uint32, TSharedPtr<FStoreBrowserTraceInfo>>& GetTraceMap() const { check(bTracesLocked); return TraceMap; }
 
-	bool IsLocked() const { return bTracesLocked; }
-	void Lock() { check(!bTracesLocked); bTracesLocked = true; TracesCriticalSection.Lock(); }
-	uint64 GetLockedTracesChangeSerial() const { check(bTracesLocked); return TracesChangeSerial; }
-	uint32 GetLockedSettingsSerial() const { check(bTracesLocked); return SettingsChangeSerial; };
-	const TArray<TSharedPtr<FStoreBrowserTraceInfo>>& GetLockedTraces() const { check(bTracesLocked); return Traces; }
-	const TMap<uint32, TSharedPtr<FStoreBrowserTraceInfo>>& GetLockedTraceMap() const { check(bTracesLocked); return TraceMap; }
-	const TArray<FString>& GetLockedWatchDirectories() const { check(bTracesLocked); return WatchDirectories; }
-	const FString& GetLockedStoreDirectory() const { check(bTracesLocked); return StoreDirectory; }
-	const FString& GetLockedHost() const { check(bTracesLocked); return Host; }
-	const FString& GetVersion() const { return Version; }
-	void Unlock() { check(bTracesLocked); bTracesLocked = false; TracesCriticalSection.Unlock(); }
 	void Refresh();
 
 private:
@@ -117,28 +121,34 @@ private:
 	void UpdateTraces();
 	void ResetTraces();
 
-	void UpdateMetadata(FStoreBrowserTraceInfo& TraceSession);
+	void UpdateMetadata(TSharedPtr<FStoreBrowserTraceInfo> TraceInfoPtr);
 
 private:
 	// Thread safe bool for stopping the thread
-	FThreadSafeBool bRunning;
+	std::atomic<bool> bRunning = false;
+
+	// Thread for continuously updating and caching info about trace store.
+	FRunnableThread* Thread = nullptr;
+
+	std::atomic<EConnectionStatus> ConnectionStatus = EConnectionStatus::Connecting;
+
+	uint32 StoreChangeSerial = 0;
+	uint32 StoreSettingsChangeSerial = 0;
 
-	// Thread for continously updating and caching info about trace store.
-	FRunnableThread* Thread;
+	mutable FCriticalSection SettingsCriticalSection;
+	std::atomic<bool> bSettingsLocked = false; // for debugging, to ensure Get*() methods for settings are only called between Lock() - Unlock() calls.
+	uint32 SettingsChangeSerial = 0;
+	FString Host;
+	FString Version;
+	FString StoreDirectory;
+	TArray<FString> WatchDirectories;
 
 	mutable FCriticalSection TracesCriticalSection;
-	FThreadSafeBool bTracesLocked; // for debugging, to ensure GetLocked*() methods are only called between Lock() - Unlock() calls.
-	std::atomic<EConnectionStatus> ConnectionStatus;
-	uint32 StoreChangeSerial;
-	uint64 TracesChangeSerial;
-	uint32 SettingsChangeSerial;
+	std::atomic<bool> bTracesLocked = false; // for debugging, to ensure Get*() methods for traces are only called between Lock() - Unlock() calls.
+	uint32 TracesChangeSerial = 0;
 	TArray<TSharedPtr<FStoreBrowserTraceInfo>> Traces;
 	TMap<uint32, TSharedPtr<FStoreBrowserTraceInfo>> TraceMap;
 	TMap<uint32, TSharedPtr<FStoreBrowserTraceInfo>> LiveTraceMap;
-	TArray<FString> WatchDirectories;
-	FString StoreDirectory;
-	FString Host;
-	FString Version;
 };
 
 ////////////////////////////////////////////////////////////////////////////////////////////////////
diff --git a/Engine/Source/Developer/TraceInsights/Private/Insights/Widgets/SStartPageWindow.cpp b/Engine/Source/Developer/TraceInsights/Private/Insights/Widgets/SStartPageWindow.cpp
--- a/Engine/Source/Developer/TraceInsights/Private/Insights/Widgets/SStartPageWindow.cpp
+++ b/Engine/Source/Developer/TraceInsights/Private/Insights/Widgets/SStartPageWindow.cpp
@@ -132,7 +132,7 @@ public:
 					.HighlightColor(FLinearColor(0.75f, 0.75f, 0.75f, 1.0f))
 					.HighlightShape(FInsightsStyle::Get().GetBrush("DarkGreenBrush"))
 					.ToolTip(STraceListRow::GetTraceTooltip())
-			     	.AddMetaData(FDriverMetaData::Id("TraceList"))
+					.AddMetaData(FDriverMetaData::Id("TraceList"))
 				]
 
 				+ SOverlay::Slot()
@@ -897,39 +897,11 @@ private:
 ////////////////////////////////////////////////////////////////////////////////////////////////////
 
 STraceStoreWindow::STraceStoreWindow()
-	: NotificationList()
-	, OverlaySettingsSlot(nullptr)
-	, DurationActive(0.0f)
-	, ActiveTimerHandle()
-	, MainContentPanel()
-	, StoreBrowser(new Insights::FStoreBrowser())
-	, TracesChangeSerial(0)
-	, TraceViewModels()
-	, FilteredTraceViewModels()
-	, TraceViewModelMap()
-	, TraceListView()
-	, SelectedTrace()
-	, bIsUserSelectedTrace(false)
-	, Filters()
-	, bSearchByCommandLine(false)
-	, FilterByNameSearchBox()
-	, FilterByName()
-	, FilterByPlatform()
-	, FilterByAppName()
-	, FilterByBuildConfig()
-	, FilterByBuildTarget()
-	, FilterByBranch()
-	, bFilterStatsTextIsDirty(true)
-	, FilterStatsText()
-	, SortColumn(TraceStoreColumns::Date)
-	, SortMode(EColumnSortMode::Ascending)
-	, AutoStartedSessions()
-	, AutoStartPlatformFilter()
-	, AutoStartAppNameFilter()
-	, AutoStartConfigurationTypeFilter(EBuildConfiguration::Unknown)
-	, AutoStartTargetTypeFilter(EBuildTargetType::Unknown)
-	, SplashScreenOverlayFadeTime(0.0f)
 {
+	StoreBrowser.Reset(new Insights::FStoreBrowser());
+
+	SortColumn = TraceStoreColumns::Date;
+	SortMode = EColumnSortMode::Ascending;
 }
 
 ////////////////////////////////////////////////////////////////////////////////////////////////////
@@ -1080,6 +1052,11 @@ void STraceStoreWindow::Construct(const FArguments& InArgs)
 
 	CreateFilters();
 
+	if (StoreHostTextBox)
+	{
+		StoreHostTextBox->SetText(FText::FromString(FInsightsManager::Get()->GetLastStoreHost()));
+	}
+
 	RefreshTraceList();
 
 	if (AutoConnect_IsChecked() == ECheckBoxState::Checked)
@@ -1863,7 +1840,7 @@ TSharedRef<SWidget> STraceStoreWindow::ConstructAutoStartPanel()
 		.OnCheckStateChanged(this, &STraceStoreWindow::AutoStart_OnCheckStateChanged)
 		[
 			SNew(STextBlock)
-			.Text(LOCTEXT("AutoStart_Text", "Auto-start analysis for LIVE trace sessions"))
+			.Text(LOCTEXT("AutoStart_Text", "Auto-start (LIVE)"))
 		]
 	]
 
@@ -1919,7 +1896,7 @@ TSharedRef<SWidget> STraceStoreWindow::ConstructAutoStartPanel()
 		.VAlign(VAlign_Center)
 		[
 			SNew(SCheckBox)
-			.ToolTipText(LOCTEXT("AutoConnect_Tooltip", "Signal to an UE application to auto-connect and start tracing if Insights is running."))
+			.ToolTipText(LOCTEXT("AutoConnect_Tooltip", "Signal to UE applications to auto-connect with local trace server and start tracing if Insights is running."))
 			.IsChecked(this, &STraceStoreWindow::AutoConnect_IsChecked)
 			.OnCheckStateChanged(this, &STraceStoreWindow::AutoConnect_OnCheckStateChanged)
 			[
@@ -1953,11 +1930,16 @@ FText STraceStoreWindow::GetConnectionStatusTooltip() const
 	static FText Disconnected = LOCTEXT("Disconnected", "Connection to trace server has been lost. Attempting to reconnect in {0} seconds.");
 
 	const FStoreBrowser::EConnectionStatus Status = StoreBrowser->GetConnectionStatus();
-	const FString& Version = StoreBrowser->GetVersion();
+
 	switch (Status)
 	{
 		case FStoreBrowser::EConnectionStatus::Connected:
-			return FText::Format(Connected, FText::FromString(Version));
+		{
+			StoreBrowser->LockSettings();
+			FText Version = FText::FromString(StoreBrowser->GetVersion());
+			StoreBrowser->UnlockSettings();
+			return FText::Format(Connected, Version);
+		}
 
 		case FStoreBrowser::EConnectionStatus::NoConnection:
 			return NotConnected;
@@ -2222,29 +2204,29 @@ void STraceStoreWindow::RefreshTraceList()
 	bool bSettingsChanged = false;
 
 	{
-		StoreBrowser->Lock();
+		StoreBrowser->LockSettings();
 
-		const uint32 NewSettingsSerial = StoreBrowser->GetLockedSettingsSerial();
-		if (NewSettingsSerial != SettingsChangeSerial)
+		const uint32 NewSettingsChangeSerial = StoreBrowser->GetSettingsChangeSerial();
+		if (NewSettingsChangeSerial != SettingsChangeSerial)
 		{
-			SettingsChangeSerial = NewSettingsSerial;
+			SettingsChangeSerial = NewSettingsChangeSerial;
 
 			// Update the host text
 			if (StoreHostTextBox)
 			{
-				StoreHostTextBox->SetText(FText::FromString(StoreBrowser->GetLockedHost()));
+				StoreHostTextBox->SetText(FText::FromString(StoreBrowser->GetHost()));
 			}
 
 			// Update the store text box
 			if (StoreDirTextBox)
 			{
-				StoreDirTextBox->SetText(FText::FromString(StoreBrowser->GetLockedStoreDirectory()));
+				StoreDirTextBox->SetText(FText::FromString(StoreBrowser->GetStoreDirectory()));
 			}
 
 			// Update store directory model
 			StoreDirectoryModel.Empty(1);
 			StoreDirectoryModel.Push(MakeShared<FTraceDirectoryModel>(
-				FString(StoreBrowser->GetLockedStoreDirectory()),
+				FString(StoreBrowser->GetStoreDirectory()),
 				NAME_None,
 				ETraceDirOperations::ModifyStore|ETraceDirOperations::Explore
 			));
@@ -2254,7 +2236,8 @@ void STraceStoreWindow::RefreshTraceList()
 			}
 
 			// Update additional watch directory models
-			static const FName DirColor[] = {
+			static const FName DirColor[] =
+			{
 				FName("Colors.AccentBlue"),
 				FName("Colors.AccentGreen"),
 				FName("Colors.AccentYellow"),
@@ -2264,7 +2247,7 @@ void STraceStoreWindow::RefreshTraceList()
 			};
 			int32 ColorIdx = 0;
 			WatchDirectoriesModel.Empty();
-			for (const auto& Dir : StoreBrowser->GetLockedWatchDirectories())
+			for (const auto& Dir : StoreBrowser->GetWatchDirectories())
 			{
 				WatchDirectoriesModel.Emplace(MakeShared<FTraceDirectoryModel>(
 					FString(Dir),
@@ -2281,14 +2264,17 @@ void STraceStoreWindow::RefreshTraceList()
 			bSettingsChanged = true;
 		}
 
-		const uint64 NewChangeSerial = StoreBrowser->GetLockedTracesChangeSerial();
-		if (NewChangeSerial != TracesChangeSerial || bSettingsChanged)
+		StoreBrowser->UnlockSettings();
+		StoreBrowser->LockTraces();
+
+		const uint32 NewTracesChangeSerial = StoreBrowser->GetTracesChangeSerial();
+		if (NewTracesChangeSerial != TracesChangeSerial || bSettingsChanged)
 		{
-			TracesChangeSerial = NewChangeSerial;
+			TracesChangeSerial = NewTracesChangeSerial;
 			//UE_LOG(TraceInsights, Log, TEXT("[TraceStore] Synching the trace list with StoreBrowser..."));
 
-			const TArray<TSharedPtr<Insights::FStoreBrowserTraceInfo>>& InTraces = StoreBrowser->GetLockedTraces();
-			const TMap<uint32, TSharedPtr<Insights::FStoreBrowserTraceInfo>>& InTraceMap = StoreBrowser->GetLockedTraceMap();
+			const TArray<TSharedPtr<Insights::FStoreBrowserTraceInfo>>& InTraces = StoreBrowser->GetTraces();
+			const TMap<uint32, TSharedPtr<Insights::FStoreBrowserTraceInfo>>& InTraceMap = StoreBrowser->GetTraceMap();
 
 			// Check for removed traces.
 			{
@@ -2336,7 +2322,7 @@ void STraceStoreWindow::RefreshTraceList()
 			}
 		}
 
-		StoreBrowser->Unlock();
+		StoreBrowser->UnlockTraces();
 	}
 
 	if (AddedTraces > 0 || RemovedTraces > 0)
@@ -2687,7 +2673,7 @@ FReply STraceStoreWindow::OnKeyDown(const FGeometry& MyGeometry, const FKeyEvent
 	if (InKeyEvent.GetKey() == EKeys::F5) // refresh metadata for all trace sessions
 	{
 		StoreBrowser->Refresh();
-
+		SettingsChangeSerial = 0;
 		TracesChangeSerial = 0;
 		TraceViewModels.Reset();
 		TraceViewModelMap.Reset();
diff --git a/Engine/Source/Developer/TraceInsights/Private/Insights/Widgets/SStartPageWindow.h b/Engine/Source/Developer/TraceInsights/Private/Insights/Widgets/SStartPageWindow.h
--- a/Engine/Source/Developer/TraceInsights/Private/Insights/Widgets/SStartPageWindow.h
+++ b/Engine/Source/Developer/TraceInsights/Private/Insights/Widgets/SStartPageWindow.h
@@ -365,7 +365,7 @@ private:
 
 	//////////////////////////////////////////////////
 	// Auto Connect
-	
+
 	ECheckBoxState AutoConnect_IsChecked() const;
 	void AutoConnect_OnCheckStateChanged(ECheckBoxState NewState);
 
@@ -518,10 +518,10 @@ private:
 	TSharedPtr<SNotificationList> NotificationList;
 
 	/** Overlay slot which contains the profiler settings widget. */
-	SOverlay::FOverlaySlot* OverlaySettingsSlot;
+	SOverlay::FOverlaySlot* OverlaySettingsSlot = nullptr;
 
 	/** The number of seconds the profiler has been active */
-	float DurationActive;
+	float DurationActive = 0.0f;
 
 	/** The handle to the active update duration tick */
 	TWeakPtr<FActiveTimerHandle> ActiveTimerHandle;
@@ -531,8 +531,8 @@ private:
 	//////////////////////////////////////////////////
 
 	TUniquePtr<Insights::FStoreBrowser> StoreBrowser;
-	uint64 TracesChangeSerial;
-	uint32 SettingsChangeSerial;
+	uint32 SettingsChangeSerial = 0;
+	uint32 TracesChangeSerial = 0;
 
 	TArray<TSharedPtr<FTraceViewModel>> TraceViewModels; // all available trace view models
 	TArray<TSharedPtr<FTraceViewModel>> FilteredTraceViewModels; // the filtered list of trace view models
@@ -544,15 +544,16 @@ private:
 	TSharedPtr<SScrollBox> StoreSettingsArea;
 	TSharedPtr<STableViewBase> WatchDirsListView;
 	TSharedPtr<SListView<TSharedPtr<FTraceViewModel>>> TraceListView;
+
 	TSharedPtr<FTraceViewModel> SelectedTrace;
-	bool bIsUserSelectedTrace;
+	bool bIsUserSelectedTrace = false;
 
 	//////////////////////////////////////////////////
 	// Filtering
 
 	TSharedPtr<FTraceViewModelFilterCollection> Filters;
 
-	bool bSearchByCommandLine;
+	bool bSearchByCommandLine = false;
 	TSharedPtr<SSearchBox> FilterByNameSearchBox;
 	TSharedPtr<FTraceTextFilter> FilterByName;
 
@@ -562,14 +563,14 @@ private:
 	TSharedPtr<FTraceFilterByBuildTarget> FilterByBuildTarget;
 	TSharedPtr<FTraceFilterByBranch> FilterByBranch;
 
-	bool bFilterStatsTextIsDirty;
+	bool bFilterStatsTextIsDirty = true;
 	FText FilterStatsText;
 
 	//////////////////////////////////////////////////
 	// Sorting
 
 	FName SortColumn;
-	EColumnSortMode::Type SortMode;
+	EColumnSortMode::Type SortMode = EColumnSortMode::None;
 
 	//////////////////////////////////////////////////
 	// Auto-start functionality
@@ -578,13 +579,13 @@ private:
 
 	TSharedPtr<SSearchBox> AutoStartPlatformFilter;
 	TSharedPtr<SSearchBox> AutoStartAppNameFilter;
-	EBuildConfiguration AutoStartConfigurationTypeFilter;
-	EBuildTargetType AutoStartTargetTypeFilter;
+	EBuildConfiguration AutoStartConfigurationTypeFilter = EBuildConfiguration::Unknown;
+	EBuildTargetType AutoStartTargetTypeFilter = EBuildTargetType::Unknown;
 
 	//////////////////////////////////////////////////
 
 	FString SplashScreenOverlayTraceFile;
-	float SplashScreenOverlayFadeTime;
+	float SplashScreenOverlayFadeTime = 0.0f;
 
 	bool bEnableAutomaticTesting = false;
 	bool bEnableDebugTools = false;
