[Insights] Fix Find Max Instance for GPU Timers.

diff --git a/Engine/Source/Developer/TraceInsights/Private/Insights/ViewModels/ThreadTimingTrack.cpp b/Engine/Source/Developer/TraceInsights/Private/Insights/ViewModels/ThreadTimingTrack.cpp
--- a/Engine/Source/Developer/TraceInsights/Private/Insights/ViewModels/ThreadTimingTrack.cpp
+++ b/Engine/Source/Developer/TraceInsights/Private/Insights/ViewModels/ThreadTimingTrack.cpp
@@ -825,12 +825,18 @@ TSharedPtr<const ITimingEvent> FThreadTimingSharedState::FindMaxEventInstance(ui
 		}
 	}
 
-	GpuTrack->FindMaxEventInstance(TimerId, StartTime, EndTime);
-	CompareAndAssignEvent(TimingEvent, TrackEvent);
+	if (GpuTrack.IsValid() && GpuTrack->IsVisible())
+	{
+		TrackEvent = GpuTrack->FindMaxEventInstance(TimerId, StartTime, EndTime);
+		CompareAndAssignEvent(TimingEvent, TrackEvent);
+	}
+
+	if (Gpu2Track.IsValid() && Gpu2Track->IsVisible())
+	{
+		TrackEvent = Gpu2Track->FindMaxEventInstance(TimerId, StartTime, EndTime);
+		CompareAndAssignEvent(TimingEvent, TrackEvent);
+	}
 
-	Gpu2Track->FindMaxEventInstance(TimerId, StartTime, EndTime);
-	CompareAndAssignEvent(TimingEvent, TrackEvent);
-	
 	return TimingEvent;
 }
 
