[Insights] Fix changing the visibility of a track doesn't update the timer panel aggregated stats in some cases.

diff --git a/Engine/Source/Developer/TraceInsights/Private/Insights/TimingProfilerManager.cpp b/Engine/Source/Developer/TraceInsights/Private/Insights/TimingProfilerManager.cpp
--- a/Engine/Source/Developer/TraceInsights/Private/Insights/TimingProfilerManager.cpp
+++ b/Engine/Source/Developer/TraceInsights/Private/Insights/TimingProfilerManager.cpp
@@ -519,8 +519,17 @@ void FTimingProfilerManager::SetSelectedTimer(uint32 InTimerId)
 void FTimingProfilerManager::OnThreadFilterChanged()
 {
 	UpdateCallersAndCallees();
-	UpdateAggregatedTimerStats();
 	UpdateAggregatedCounterStats();
+
+	TSharedPtr<STimingProfilerWindow> Wnd = GetProfilerWindow();
+	if (Wnd)
+	{
+		TSharedPtr<STimersView> TimersView = Wnd->GetTimersView();
+		if (TimersView)
+		{
+			TimersView->OnTimingViewTrackListChanged();
+		}
+	}
 }
 
 ////////////////////////////////////////////////////////////////////////////////////////////////////
diff --git a/Engine/Source/Developer/TraceInsights/Private/Insights/Widgets/STimersView.cpp b/Engine/Source/Developer/TraceInsights/Private/Insights/Widgets/STimersView.cpp
--- a/Engine/Source/Developer/TraceInsights/Private/Insights/Widgets/STimersView.cpp
+++ b/Engine/Source/Developer/TraceInsights/Private/Insights/Widgets/STimersView.cpp
@@ -3514,4 +3514,16 @@ void STimersView::LoadSettings()
 
 ////////////////////////////////////////////////////////////////////////////////////////////////////
 
+void STimersView::OnTimingViewTrackListChanged()
+{
+	if (!Aggregator->IsEmptyTimeInterval())
+	{
+		Aggregator->Cancel();
+		Aggregator->SetFrameType(ModeFrameType);
+		Aggregator->Start();
+	}
+}
+
+////////////////////////////////////////////////////////////////////////////////////////////////////
+
 #undef LOCTEXT_NAMESPACE
diff --git a/Engine/Source/Developer/TraceInsights/Private/Insights/Widgets/STimersView.h b/Engine/Source/Developer/TraceInsights/Private/Insights/Widgets/STimersView.h
--- a/Engine/Source/Developer/TraceInsights/Private/Insights/Widgets/STimersView.h
+++ b/Engine/Source/Developer/TraceInsights/Private/Insights/Widgets/STimersView.h
@@ -91,6 +91,8 @@ public:
 
 	ETraceFrameType GetFrameTypeMode() { return ModeFrameType; }
 
+	void OnTimingViewTrackListChanged();
+
 private:
 	void InitCommandList();
 
