[Insights] MemoryInsights: Added "LLM Size" column to Allocs Table. It shows the size of LLM reported stats for "tag" groups (when allocs are grouped by tags). Also added "LLM Delta Size" column, to show the difference between computed tag size based on traced allocs and LLM reported size. This helps to discover mismatches between LLM tracking and Insights' memory tracing.

diff --git a/Engine/Source/Developer/TraceInsights/Private/Insights/MemoryProfiler/ViewModels/MemAllocGroupingByTag.cpp b/Engine/Source/Developer/TraceInsights/Private/Insights/MemoryProfiler/ViewModels/MemAllocGroupingByTag.cpp
--- a/Engine/Source/Developer/TraceInsights/Private/Insights/MemoryProfiler/ViewModels/MemAllocGroupingByTag.cpp
+++ b/Engine/Source/Developer/TraceInsights/Private/Insights/MemoryProfiler/ViewModels/MemAllocGroupingByTag.cpp
@@ -6,7 +6,9 @@
 
 // Insights
 #include "Insights/Common/AsyncOperationProgress.h"
+#include "Insights/InsightsManager.h"
 #include "Insights/MemoryProfiler/ViewModels/MemAllocNode.h"
+#include "Insights/MemoryProfiler/MemoryProfilerManager.h"
 
 #define LOCTEXT_NAMESPACE "Insights::FMemAllocGroupingByTag"
 
@@ -17,8 +19,11 @@ namespace Insights
 
 ////////////////////////////////////////////////////////////////////////////////////////////////////
 
-INSIGHTS_IMPLEMENT_RTTI(FMemAllocGroupingByTag);
+INSIGHTS_IMPLEMENT_RTTI(FMemAllocGroupingByTag)
+INSIGHTS_IMPLEMENT_RTTI(FMemTagTableTreeNode)
 
+////////////////////////////////////////////////////////////////////////////////////////////////////
+// FMemAllocGroupingByTag
 ////////////////////////////////////////////////////////////////////////////////////////////////////
 
 FMemAllocGroupingByTag::FMemAllocGroupingByTag(const TraceServices::IAllocationsProvider& InTagProvider)
@@ -114,9 +119,6 @@ void FMemAllocGroupingByTag::GroupNodes(const TArray<FTableTreeNodePtr>& Nodes,
 	};
 	FScopedMemory _(TagNodes);
 
-	const FSlateBrush* IconBrush = FBaseTreeNode::GetDefaultIcon(true);
-	FLinearColor GroupNodeColor(0.75f, 0.5f, 1.0f, 1.0f);
-
 	// Create tag nodes.
 	{
 		TraceServices::FProviderReadScopeLock TagProviderReadScopeLock(TagProvider);
@@ -124,10 +126,10 @@ void FMemAllocGroupingByTag::GroupNodes(const TArray<FTableTreeNodePtr>& Nodes,
 			{
 				const FName NodeName(Name);
 #if INSIGHTS_MERGE_MEM_TAGS_BY_NAME
-				FTableTreeNodePtr TreeNode = MakeShared<FCustomTableTreeNode>(NodeName, InParentTable, IconBrush, GroupNodeColor);
+				FTableTreeNodePtr TreeNode = MakeShared<FMemTagTableTreeNode>(NodeName, InParentTable, FullName);
 #else
 				const FName NodeNameEx(Name, int32(Id));
-				FTableTreeNodePtr TreeNode = MakeShared<FCustomTableTreeNode>(NodeNameEx, InParentTable, IconBrush, GroupNodeColor);
+				FTableTreeNodePtr TreeNode = MakeShared<FMemTagTableTreeNode>(NodeNameEx, InParentTable, FullName);
 #endif
 				FTagNode* TagNode = new FTagNode(NodeName, Id, ParentId, TreeNode);
 				TagNodes.Add(TagNode);
@@ -141,7 +143,7 @@ void FMemAllocGroupingByTag::GroupNodes(const TArray<FTableTreeNodePtr>& Nodes,
 	if (!UntaggedNode)
 	{
 		const FName NodeName(TEXT("Untagged"));
-		FTableTreeNodePtr UntaggedTreeNode = MakeShared<FCustomTableTreeNode>(NodeName, InParentTable, IconBrush, GroupNodeColor);
+		FTableTreeNodePtr UntaggedTreeNode = MakeShared<FMemTagTableTreeNode>(NodeName, InParentTable, TEXT("Untagged"));
 		UntaggedNode = new FTagNode(NodeName, UntaggedNodeId, InvalidTagId, UntaggedTreeNode);
 		TagNodes.Add(UntaggedNode);
 		IdToNodeMap.Add(UntaggedNodeId, UntaggedNode);
@@ -153,7 +155,7 @@ void FMemAllocGroupingByTag::GroupNodes(const TArray<FTableTreeNodePtr>& Nodes,
 	if (!UnknownTagNode)
 	{
 		const FName NodeName(TEXT("Unknown"));
-		FTableTreeNodePtr UnknownTagTreeNode = MakeShared<FCustomTableTreeNode>(NodeName, InParentTable, IconBrush, GroupNodeColor);
+		FTableTreeNodePtr UnknownTagTreeNode = MakeShared<FMemTagTableTreeNode>(NodeName, InParentTable, TEXT("Unknown"));
 		UnknownTagNode = new FTagNode(NodeName, UnknownTagNodeId, InvalidTagId, UnknownTagTreeNode);
 		TagNodes.Add(UnknownTagNode);
 		IdToNodeMap.Add(UnknownTagNodeId, UnknownTagNode);
@@ -261,9 +263,94 @@ void FMemAllocGroupingByTag::GroupNodes(const TArray<FTableTreeNodePtr>& Nodes,
 	}
 }
 
+////////////////////////////////////////////////////////////////////////////////////////////////////
+// FMemTagTableTreeNode
 ////////////////////////////////////////////////////////////////////////////////////////////////////
 
-#undef INSIGHTS_MERGE_MEM_TAGS_BY_NAME
-#undef LOCTEXT_NAMESPACE
+void FMemTagTableTreeNode::UpdateLLMSize() const
+{
+	LLMSize = 0;
+
+	TSharedPtr<FTable> Table = ParentTable.Pin();
+	if (!Table.IsValid())
+	{
+		return;
+	}
+
+	TSharedPtr<FMemAllocTable> MemAllocTable = StaticCastSharedPtr<FMemAllocTable>(Table);
+	if (!MemAllocTable.IsValid())
+	{
+		return;
+	}
+
+	double Time = MemAllocTable->GetTimeMarkerA();
+
+	TSharedPtr<const TraceServices::IAnalysisSession> Session = FInsightsManager::Get()->GetSession();
+	if (!Session.IsValid())
+	{
+		return;
+	}
+
+	::FMemorySharedState* SharedState = FMemoryProfilerManager::Get()->GetSharedState();
+	if (!SharedState)
+	{
+		return;
+	}
+
+	TraceServices::FMemoryTrackerId TrackerId = FMemoryTracker::InvalidTrackerId;
+	TraceServices::FMemoryTagId TagId = FMemoryTag::InvalidTagId;
+
+	const Insights::FMemoryTagList& TagList = SharedState->GetTagList();
+	for (const Insights::FMemoryTag* MemTag : TagList.GetTags())
+	{
+		if (MemTag)
+		{
+			const FString& FullName = MemTag->GetStatFullName();
+			if (FCString::Stricmp(*FullName, TagFullName) == 0)
+			{
+				TrackerId = MemTag->GetTrackerId();
+				TagId = MemTag->GetId();
+				break;
+			}
+		}
+	}
+	if (TrackerId == FMemoryTracker::InvalidTrackerId || TagId == FMemoryTag::InvalidTagId)
+	{
+		return;
+	}
+
+	const TraceServices::IMemoryProvider* MemoryProvider = TraceServices::ReadMemoryProvider(*Session.Get());
+	if (!MemoryProvider)
+	{
+		return;
+	}
+
+	TraceServices::FAnalysisSessionReadScope SessionReadScope(*Session.Get());
+
+	const uint64 TotalSampleCount = MemoryProvider->GetTagSampleCount(TrackerId, TagId);
+	if (TotalSampleCount == 0)
+	{
+		return;
+	}
+
+	int64 LocalLLMSize = INT64_MAX;
+	MemoryProvider->EnumerateTagSamples(TrackerId, TagId, Time, Time, true,
+		[&LocalLLMSize](double Time, double Duration, const TraceServices::FMemoryTagSample& Sample)
+		{
+			if (LocalLLMSize == INT64_MAX)
+			{
+				LocalLLMSize = Sample.Value;
+			}
+		});
+	if (LocalLLMSize != INT64_MAX)
+	{
+		LLMSize = LocalLLMSize;
+	}
+}
+
+////////////////////////////////////////////////////////////////////////////////////////////////////
 
 } // namespace Insights
+
+#undef INSIGHTS_MERGE_MEM_TAGS_BY_NAME
+#undef LOCTEXT_NAMESPACE
diff --git a/Engine/Source/Developer/TraceInsights/Private/Insights/MemoryProfiler/ViewModels/MemAllocGroupingByTag.h b/Engine/Source/Developer/TraceInsights/Private/Insights/MemoryProfiler/ViewModels/MemAllocGroupingByTag.h
--- a/Engine/Source/Developer/TraceInsights/Private/Insights/MemoryProfiler/ViewModels/MemAllocGroupingByTag.h
+++ b/Engine/Source/Developer/TraceInsights/Private/Insights/MemoryProfiler/ViewModels/MemAllocGroupingByTag.h
@@ -6,6 +6,7 @@
 #include "Insights/Table/ViewModels/TableTreeNode.h"
 #include "Insights/Table/ViewModels/TreeNodeGrouping.h"
 #include "TraceServices/Model/AllocationsProvider.h"
+#include "TraceServices/Model/Memory.h"
 
 namespace Insights
 {
@@ -26,6 +27,56 @@ public:
 private:
 	const TraceServices::IAllocationsProvider& TagProvider;
 };
-	
+
+////////////////////////////////////////////////////////////////////////////////////////////////////
+
+class FMemTagTableTreeNode : public FTableTreeNode
+{
+	INSIGHTS_DECLARE_RTTI(FMemTagTableTreeNode, FTableTreeNode)
+
+public:
+	/** Initialization constructor for a table record node. */
+	explicit FMemTagTableTreeNode(const FName InName, TWeakPtr<FTable> InParentTable, int32 InRowIndex, const TCHAR* InTagFullName, bool IsGroup)
+		: FTableTreeNode(InName, InParentTable, InRowIndex, IsGroup)
+		, TagFullName(InTagFullName)
+	{
+	}
+
+	/** Initialization constructor for the group node. */
+	explicit FMemTagTableTreeNode(const FName InName, TWeakPtr<FTable> InParentTable, const TCHAR* InTagFullName)
+		: FTableTreeNode(InName, InParentTable)
+		, TagFullName(InTagFullName)
+	{
+	}
+
+	virtual ~FMemTagTableTreeNode()
+	{
+	}
+
+	/**
+	 * @return the color tint for icon and name text.
+	 */
+	virtual FLinearColor GetColor() const override
+	{
+		return FLinearColor(0.75f, 0.5f, 1.0f, 1.0f);
+	}
+
+	int64 GetLLMSize() const
+	{
+		if (LLMSize == INT64_MAX)
+		{
+			UpdateLLMSize();
+		}
+		return LLMSize;
+	}
+
+private:
+	void UpdateLLMSize() const;
+
+private:
+	const TCHAR* TagFullName = nullptr;
+	mutable int64 LLMSize = INT64_MAX;
+};
+
 ////////////////////////////////////////////////////////////////////////////////////////////////////
 }
diff --git a/Engine/Source/Developer/TraceInsights/Private/Insights/MemoryProfiler/ViewModels/MemAllocTable.cpp b/Engine/Source/Developer/TraceInsights/Private/Insights/MemoryProfiler/ViewModels/MemAllocTable.cpp
--- a/Engine/Source/Developer/TraceInsights/Private/Insights/MemoryProfiler/ViewModels/MemAllocTable.cpp
+++ b/Engine/Source/Developer/TraceInsights/Private/Insights/MemoryProfiler/ViewModels/MemAllocTable.cpp
@@ -10,6 +10,7 @@
 // Insights
 #include "Insights/MemoryProfiler/ViewModels/CallstackFormatting.h"
 #include "Insights/MemoryProfiler/ViewModels/MemAllocFilterValueConverter.h"
+#include "Insights/MemoryProfiler/ViewModels/MemAllocGroupingByTag.h"
 #include "Insights/MemoryProfiler/ViewModels/MemAllocNode.h"
 #include "Insights/MemoryProfiler/ViewModels/MemAllocTable.h"
 #include "Insights/Table/ViewModels/TableCellValueFormatter.h"
@@ -38,6 +39,8 @@ const FName FMemAllocTableColumns::AddressColumnId(TEXT("Address"));
 const FName FMemAllocTableColumns::MemoryPageColumnId(TEXT("MemoryPage"));
 const FName FMemAllocTableColumns::CountColumnId(TEXT("Count"));
 const FName FMemAllocTableColumns::SizeColumnId(TEXT("Size"));
+const FName FMemAllocTableColumns::LLMSizeColumnId(TEXT("LLMSize"));
+const FName FMemAllocTableColumns::LLMDeltaSizeColumnId(TEXT("LLMDeltaSize"));
 const FName FMemAllocTableColumns::TagColumnId(TEXT("Tag"));
 const FName FMemAllocTableColumns::AssetColumnId(TEXT("Asset"));
 const FName FMemAllocTableColumns::PackageColumnId(TEXT("Package"));
@@ -731,7 +734,7 @@ void FMemAllocTable::AddDefaultColumns()
 		Column.SetInitialSortMode(EColumnSortMode::Descending);
 
 		Column.SetAggregation(ETableColumnAggregation::Sum);
-		
+
 		AddColumn(ColumnRef);
 	}
 	//////////////////////////////////////////////////
@@ -797,6 +800,136 @@ void FMemAllocTable::AddDefaultColumns()
 		AddColumn(ColumnRef);
 	}
 	//////////////////////////////////////////////////
+	// LLM Size Column
+	{
+		TSharedRef<FTableColumn> ColumnRef = MakeShared<FTableColumn>(FMemAllocTableColumns::LLMSizeColumnId);
+		FTableColumn& Column = *ColumnRef;
+
+		Column.SetIndex(ColumnIndex++);
+
+		Column.SetShortName(LOCTEXT("LLMSizeColumnName", "LLM Size"));
+		Column.SetTitleName(LOCTEXT("LLMSizeColumnTitle", "LLM Size"));
+		Column.SetDescription(LOCTEXT("LLMSizeColumnDesc", "Size reported by LLM (only for Mem Tag group nodes)"));
+
+		Column.SetFlags(ETableColumnFlags::CanBeHidden | ETableColumnFlags::CanBeFiltered);
+
+		Column.SetHorizontalAlignment(HAlign_Right);
+		Column.SetInitialWidth(100.0f);
+
+		Column.SetDataType(ETableCellDataType::Int64);
+
+		class FLLMSizeValueGetter : public FTableCellValueGetter
+		{
+		public:
+			virtual const TOptional<FTableCellValue> GetValue(const FTableColumn& Column, const FBaseTreeNode& Node) const override
+			{
+				if (Node.IsGroup())
+				{
+					if (Node.Is<FMemTagTableTreeNode>())
+					{
+						int64 LLMSize = Node.As<FMemTagTableTreeNode>().GetLLMSize();
+						return FTableCellValue(LLMSize);
+					}
+					else
+					{
+						const FTableTreeNode& NodePtr = static_cast<const FTableTreeNode&>(Node);
+						if (NodePtr.HasAggregatedValue(Column.GetId()))
+						{
+							return NodePtr.GetAggregatedValue(Column.GetId());
+						}
+					}
+				}
+				return TOptional<FTableCellValue>();
+			}
+		};
+		TSharedRef<ITableCellValueGetter> Getter = MakeShared<FLLMSizeValueGetter>();
+		Column.SetValueGetter(Getter);
+
+		TSharedRef<ITableCellValueFormatter> Formatter = MakeShared<FInt64ValueFormatterAsMemory>();
+		Column.SetValueFormatter(Formatter);
+
+		TSharedRef<ITableCellValueSorter> Sorter = MakeShared<FSorterByInt64Value>(ColumnRef);
+		Column.SetValueSorter(Sorter);
+		Column.SetInitialSortMode(EColumnSortMode::Descending);
+
+		TSharedRef<IFilterValueConverter> Converter = MakeShared<FMemoryFilterValueConverter>();
+		Column.SetValueConverter(Converter);
+
+		Column.SetAggregation(ETableColumnAggregation::Sum);
+
+		AddColumn(ColumnRef);
+	}
+	//////////////////////////////////////////////////
+	// LLM Delta Size Column
+	{
+		TSharedRef<FTableColumn> ColumnRef = MakeShared<FTableColumn>(FMemAllocTableColumns::LLMDeltaSizeColumnId);
+		FTableColumn& Column = *ColumnRef;
+
+		Column.SetIndex(ColumnIndex++);
+
+		Column.SetShortName(LOCTEXT("LLMDeltaSizeColumnName", "LLM Delta"));
+		Column.SetTitleName(LOCTEXT("LLMDeltaSizeColumnTitle", "LLM Delta Size"));
+		Column.SetDescription(LOCTEXT("LLMDeltaSizeColumnDesc", "Difference from size reported by LLM (only for Mem Tag group nodes)"));
+
+		Column.SetFlags(ETableColumnFlags::CanBeHidden | ETableColumnFlags::CanBeFiltered);
+
+		Column.SetHorizontalAlignment(HAlign_Right);
+		Column.SetInitialWidth(100.0f);
+
+		Column.SetDataType(ETableCellDataType::Int64);
+
+		class FLLMDeltaSizeValueGetter : public FTableCellValueGetter
+		{
+		public:
+			virtual const TOptional<FTableCellValue> GetValue(const FTableColumn& Column, const FBaseTreeNode& Node) const override
+			{
+				if (Node.IsGroup())
+				{
+					if (Node.Is<FMemTagTableTreeNode>())
+					{
+						int64 LLMSize = Node.As<FMemTagTableTreeNode>().GetLLMSize();
+						int64 Size = 0;
+						const FTableTreeNode& NodePtr = static_cast<const FTableTreeNode&>(Node);
+						if (NodePtr.HasAggregatedValue(FMemAllocTableColumns::SizeColumnId))
+						{
+							const FTableCellValue& Value = NodePtr.GetAggregatedValue(FMemAllocTableColumns::SizeColumnId);
+							if (Value.DataType == ETableCellDataType::Int64)
+							{
+								Size = Value.AsInt64();
+							}
+						}
+						return FTableCellValue(Size - LLMSize);
+					}
+					else
+					{
+						const FTableTreeNode& NodePtr = static_cast<const FTableTreeNode&>(Node);
+						if (NodePtr.HasAggregatedValue(Column.GetId()))
+						{
+							return NodePtr.GetAggregatedValue(Column.GetId());
+						}
+					}
+				}
+				return TOptional<FTableCellValue>();
+			}
+		};
+		TSharedRef<ITableCellValueGetter> Getter = MakeShared<FLLMDeltaSizeValueGetter>();
+		Column.SetValueGetter(Getter);
+
+		TSharedRef<ITableCellValueFormatter> Formatter = MakeShared<FInt64ValueFormatterAsMemory>();
+		Column.SetValueFormatter(Formatter);
+
+		TSharedRef<ITableCellValueSorter> Sorter = MakeShared<FSorterByInt64Value>(ColumnRef);
+		Column.SetValueSorter(Sorter);
+		Column.SetInitialSortMode(EColumnSortMode::Ascending);
+
+		TSharedRef<IFilterValueConverter> Converter = MakeShared<FMemoryFilterValueConverter>();
+		Column.SetValueConverter(Converter);
+
+		Column.SetAggregation(ETableColumnAggregation::Sum);
+
+		AddColumn(ColumnRef);
+	}
+	//////////////////////////////////////////////////
 	// LLM Tag Column
 	{
 		TSharedRef<FTableColumn> ColumnRef = MakeShared<FTableColumn>(FMemAllocTableColumns::TagColumnId);
diff --git a/Engine/Source/Developer/TraceInsights/Private/Insights/MemoryProfiler/ViewModels/MemAllocTable.h b/Engine/Source/Developer/TraceInsights/Private/Insights/MemoryProfiler/ViewModels/MemAllocTable.h
--- a/Engine/Source/Developer/TraceInsights/Private/Insights/MemoryProfiler/ViewModels/MemAllocTable.h
+++ b/Engine/Source/Developer/TraceInsights/Private/Insights/MemoryProfiler/ViewModels/MemAllocTable.h
@@ -30,6 +30,8 @@ struct FMemAllocTableColumns
 	static const FName MemoryPageColumnId;
 	static const FName CountColumnId;
 	static const FName SizeColumnId;
+	static const FName LLMSizeColumnId;
+	static const FName LLMDeltaSizeColumnId;
 	static const FName TagColumnId;
 	static const FName AssetColumnId;
 	static const FName PackageColumnId;
@@ -61,11 +63,15 @@ public:
 	const FMemoryAlloc* GetMemAlloc(int32 InIndex) const { return IsValidRowIndex(InIndex) ? &Allocs[InIndex] : nullptr; }
 	const FMemoryAlloc& GetMemAllocChecked(int32 InIndex) const { check(IsValidRowIndex(InIndex)); return Allocs[InIndex]; }
 
+	double GetTimeMarkerA() const { return TimeA; }
+	void SetTimeMarkerA(double InTime) { TimeA = InTime; }
+
 private:
 	void AddDefaultColumns();
 
 private:
 	TArray<FMemoryAlloc> Allocs;
+	double TimeA = 0.0;
 };
 
 ////////////////////////////////////////////////////////////////////////////////////////////////////
diff --git a/Engine/Source/Developer/TraceInsights/Private/Insights/MemoryProfiler/ViewModels/MemorySharedState.cpp b/Engine/Source/Developer/TraceInsights/Private/Insights/MemoryProfiler/ViewModels/MemorySharedState.cpp
--- a/Engine/Source/Developer/TraceInsights/Private/Insights/MemoryProfiler/ViewModels/MemorySharedState.cpp
+++ b/Engine/Source/Developer/TraceInsights/Private/Insights/MemoryProfiler/ViewModels/MemorySharedState.cpp
@@ -27,6 +27,9 @@
 
 #include <limits>
 
+static_assert(Insights::FMemoryTracker::InvalidTrackerId == TraceServices::FMemoryTrackerInfo::InvalidTrackerId, "InvalidTrackerId");
+static_assert(Insights::FMemoryTag::InvalidTagId == TraceServices::FMemoryTagInfo::InvalidTagId, "InvalidTagId");
+
 #define LOCTEXT_NAMESPACE "MemorySharedState"
 
 const FName Insights::FQueryTargetWindowSpec::NewWindow = TEXT("New Window");
diff --git a/Engine/Source/Developer/TraceInsights/Private/Insights/MemoryProfiler/Widgets/SMemAllocTableTreeView.cpp b/Engine/Source/Developer/TraceInsights/Private/Insights/MemoryProfiler/Widgets/SMemAllocTableTreeView.cpp
--- a/Engine/Source/Developer/TraceInsights/Private/Insights/MemoryProfiler/Widgets/SMemAllocTableTreeView.cpp
+++ b/Engine/Source/Developer/TraceInsights/Private/Insights/MemoryProfiler/Widgets/SMemAllocTableTreeView.cpp
@@ -234,6 +234,8 @@ void SMemAllocTableTreeView::ResetAndStartQuery()
 	TSharedPtr<Insights::FMemAllocTable> MemAllocTable = GetMemAllocTable();
 	if (MemAllocTable)
 	{
+		MemAllocTable->SetTimeMarkerA(TimeMarkers[0]); // to be used by LLM Size and LLM Delta Size columns
+
 		TArray<FMemoryAlloc>& Allocs = MemAllocTable->GetAllocs();
 		Allocs.Reset(10 * 1024 * 1024);
 	}
