[Insights]
  - Tree Views: Added message for cases where an async op (sorting, filtering, grouping) is canceled and tree is updating (meanwhile showing an empty table).
  - MemoryInsights / Allocs Table: Fixed "Asset (Package)" preset to use "path" breakdown by package (not by asset name).
  - MemoryInsights / Allocs Table: Added grouping  by Package (unique values) to "Asset Class" preset.
  - MemoryInsights / Allocs Table: Fixed "inverted callstack" grouping to skip filtered frames no matter if the "grouping by function" option is on or off.

diff --git a/Engine/Source/Developer/TraceInsights/Private/Insights/MemoryProfiler/ViewModels/MemAllocGroupingByCallstack.cpp b/Engine/Source/Developer/TraceInsights/Private/Insights/MemoryProfiler/ViewModels/MemAllocGroupingByCallstack.cpp
--- a/Engine/Source/Developer/TraceInsights/Private/Insights/MemoryProfiler/ViewModels/MemAllocGroupingByCallstack.cpp
+++ b/Engine/Source/Developer/TraceInsights/Private/Insights/MemoryProfiler/ViewModels/MemAllocGroupingByCallstack.cpp
@@ -73,6 +73,7 @@ FMemAllocGroupingByCallstack::FMemAllocGroupingByCallstack(bool bInIsAllocCallst
 	, bIsAllocCallstack(bInIsAllocCallstack)
 	, bIsInverted(bInIsInverted)
 	, bIsGroupingByFunction(bInIsGroupingByFunction)
+	, bShouldSkipFilteredFrames(bIsInverted)
 {
 }
 
@@ -86,6 +87,9 @@ FMemAllocGroupingByCallstack::~FMemAllocGroupingByCallstack()
 
 void FMemAllocGroupingByCallstack::GroupNodes(const TArray<FTableTreeNodePtr>& Nodes, FTableTreeNode& ParentGroup, TWeakPtr<FTable> InParentTable, IAsyncOperationProgress& InAsyncOperationProgress) const
 {
+	const bool bLocalIsGroupingByFunction = bIsGroupingByFunction;
+	const bool bLocalShouldSkipFilteredFrames = bShouldSkipFilteredFrames;
+
 	ParentGroup.ClearChildren();
 
 	TArray<FCallstackGroup*> CallstackGroups;
@@ -135,14 +139,14 @@ void FMemAllocGroupingByCallstack::GroupNodes(const TArray<FTableTreeNodePtr>& N
 					const TraceServices::FStackFrame* Frame = Callstack->Frame(static_cast<uint8>(bIsInverted ? NumFrames - FrameDepth - 1 : FrameDepth));
 					check(Frame != nullptr);
 
-					if (bIsGroupingByFunction)
+					if (bLocalShouldSkipFilteredFrames &&
+						Frame->Symbol->FilterStatus.load() == TraceServices::EResolvedSymbolFilterStatus::Filtered)
 					{
-						// Skip noise for inverted callstack
-						if (bIsInverted && Frame->Symbol->FilterStatus.load() == TraceServices::EResolvedSymbolFilterStatus::Filtered)
-						{
-							continue;
-						}
+						continue;
+					}
 
+					if (bLocalIsGroupingByFunction)
+					{
 						const FName GroupName = GetGroupName(Frame);
 
 						// Merge with parent group, if it has the same name (i.e. same function).
diff --git a/Engine/Source/Developer/TraceInsights/Private/Insights/MemoryProfiler/ViewModels/MemAllocGroupingByCallstack.h b/Engine/Source/Developer/TraceInsights/Private/Insights/MemoryProfiler/ViewModels/MemAllocGroupingByCallstack.h
--- a/Engine/Source/Developer/TraceInsights/Private/Insights/MemoryProfiler/ViewModels/MemAllocGroupingByCallstack.h
+++ b/Engine/Source/Developer/TraceInsights/Private/Insights/MemoryProfiler/ViewModels/MemAllocGroupingByCallstack.h
@@ -81,6 +81,9 @@ public:
 	bool IsGroupingByFunction() const { return bIsGroupingByFunction; }
 	void SetGroupingByFunction(bool bOnOff) { bIsGroupingByFunction = bOnOff; }
 
+	bool ShouldSkipFilteredFrames() const { return bShouldSkipFilteredFrames; }
+	void SetSkipFilteredFrames(bool bOnOff) { bShouldSkipFilteredFrames = bOnOff; }
+
 private:
 	FName GetGroupName(const TraceServices::FStackFrame* Frame) const;
 
@@ -92,7 +95,8 @@ private:
 private:
 	bool bIsAllocCallstack;
 	bool bIsInverted;
-	bool bIsGroupingByFunction;
+	std::atomic<bool> bIsGroupingByFunction;
+	std::atomic<bool> bShouldSkipFilteredFrames;
 };
 
 ////////////////////////////////////////////////////////////////////////////////////////////////////
diff --git a/Engine/Source/Developer/TraceInsights/Private/Insights/MemoryProfiler/Widgets/SMemAllocTableTreeView.cpp b/Engine/Source/Developer/TraceInsights/Private/Insights/MemoryProfiler/Widgets/SMemAllocTableTreeView.cpp
--- a/Engine/Source/Developer/TraceInsights/Private/Insights/MemoryProfiler/Widgets/SMemAllocTableTreeView.cpp
+++ b/Engine/Source/Developer/TraceInsights/Private/Insights/MemoryProfiler/Widgets/SMemAllocTableTreeView.cpp
@@ -872,11 +872,11 @@ void SMemAllocTableTreeView::InitAvailableViewPresets()
 	public:
 		virtual FText GetName() const override
 		{
-			return LOCTEXT("Asset_PresetName", "Asset");
+			return LOCTEXT("Asset_PresetName", "Asset (Package)");
 		}
 		virtual FText GetToolTip() const override
 		{
-			return LOCTEXT("Asset_PresetToolTip", "Asset Breakdown View\nConfigure the tree view to show a breakdown of allocations by their Asset tag.");
+			return LOCTEXT("Asset_PresetToolTip", "Asset (Package) Breakdown View\nConfigure the tree view to show a breakdown of allocations by Package and Asset Name metadata.");
 		}
 		virtual FName GetSortColumn() const override
 		{
@@ -893,11 +893,22 @@ void SMemAllocTableTreeView::InitAvailableViewPresets()
 			check(InAvailableGroupings[0]->Is<FTreeNodeGroupingFlat>());
 			InOutCurrentGroupings.Add(InAvailableGroupings[0]);
 
-			const TSharedPtr<FTreeNodeGrouping>* AssetGrouping = InAvailableGroupings.FindByPredicate(
+			const TSharedPtr<FTreeNodeGrouping>* PackageGrouping = InAvailableGroupings.FindByPredicate(
 				[](TSharedPtr<FTreeNodeGrouping>& Grouping)
 				{
 					return Grouping->Is<FTreeNodeGroupingByPathBreakdown>() &&
-						   Grouping->As<FTreeNodeGroupingByPathBreakdown>().GetColumnId() == FMemAllocTableColumns::AssetColumnId;
+						   Grouping->As<FTreeNodeGroupingByPathBreakdown>().GetColumnId() == FMemAllocTableColumns::PackageColumnId;
+				});
+			if (PackageGrouping)
+			{
+				InOutCurrentGroupings.Add(*PackageGrouping);
+			}
+
+			const TSharedPtr<FTreeNodeGrouping>* AssetGrouping = InAvailableGroupings.FindByPredicate(
+				[](TSharedPtr<FTreeNodeGrouping>& Grouping)
+				{
+					return Grouping->Is<FTreeNodeGroupingByUniqueValueCString>() &&
+						   Grouping->As<FTreeNodeGroupingByUniqueValueCString>().GetColumnId() == FMemAllocTableColumns::AssetColumnId;
 				});
 			if (AssetGrouping)
 			{
@@ -906,11 +917,12 @@ void SMemAllocTableTreeView::InitAvailableViewPresets()
 		}
 		virtual void GetColumnConfigSet(TArray<FTableColumnConfig>& InOutConfigSet) const override
 		{
-			InOutConfigSet.Add({ FTable::GetHierarchyColumnId(),          true, 200.0f });
-			InOutConfigSet.Add({ FMemAllocTableColumns::CountColumnId,    true, 100.0f });
-			InOutConfigSet.Add({ FMemAllocTableColumns::SizeColumnId,     true, 100.0f });
-			InOutConfigSet.Add({ FMemAllocTableColumns::TagColumnId,      true, 120.0f });
-			InOutConfigSet.Add({ FMemAllocTableColumns::AllocFunctionColumnId, true, 400.0f });
+			InOutConfigSet.Add({ FTable::GetHierarchyColumnId(),               true, 200.0f });
+			InOutConfigSet.Add({ FMemAllocTableColumns::CountColumnId,         true, 100.0f });
+			InOutConfigSet.Add({ FMemAllocTableColumns::SizeColumnId,          true, 100.0f });
+			InOutConfigSet.Add({ FMemAllocTableColumns::TagColumnId,           true, 120.0f });
+			InOutConfigSet.Add({ FMemAllocTableColumns::ClassNameColumnId,     true, 120.0f });
+			InOutConfigSet.Add({ FMemAllocTableColumns::AllocFunctionColumnId, true, 300.0f });
 		}
 	};
 	AvailableViewPresets.Add(MakeShared<FAssetViewPreset>());
@@ -927,7 +939,7 @@ void SMemAllocTableTreeView::InitAvailableViewPresets()
 		}
 		virtual FText GetToolTip() const override
 		{
-			return LOCTEXT("ClassName_PresetToolTip", "Class Breakdown View\nConfigure the tree view to show a breakdown of allocations by their Class name.");
+			return LOCTEXT("ClassName_PresetToolTip", "Class Name Breakdown View\nConfigure the tree view to show a breakdown of allocations by Asset's Class Name metadata.");
 		}
 		virtual FName GetSortColumn() const override
 		{
@@ -955,6 +967,17 @@ void SMemAllocTableTreeView::InitAvailableViewPresets()
 				InOutCurrentGroupings.Add(*ClassNameGrouping);
 			}
 
+			const TSharedPtr<FTreeNodeGrouping>* PackageGrouping = InAvailableGroupings.FindByPredicate(
+				[](TSharedPtr<FTreeNodeGrouping>& Grouping)
+				{
+					return Grouping->Is<FTreeNodeGroupingByUniqueValueCString>() &&
+						   Grouping->As<FTreeNodeGroupingByUniqueValueCString>().GetColumnId() == FMemAllocTableColumns::PackageColumnId;
+				});
+			if (PackageGrouping)
+			{
+				InOutCurrentGroupings.Add(*PackageGrouping);
+			}
+
 			const TSharedPtr<FTreeNodeGrouping>* AssetGrouping = InAvailableGroupings.FindByPredicate(
 				[](TSharedPtr<FTreeNodeGrouping>& Grouping)
 				{
@@ -968,10 +991,10 @@ void SMemAllocTableTreeView::InitAvailableViewPresets()
 		}
 		virtual void GetColumnConfigSet(TArray<FTableColumnConfig>& InOutConfigSet) const override
 		{
-			InOutConfigSet.Add({ FTable::GetHierarchyColumnId(),          true, 200.0f });
-			InOutConfigSet.Add({ FMemAllocTableColumns::CountColumnId,    true, 100.0f });
-			InOutConfigSet.Add({ FMemAllocTableColumns::SizeColumnId,     true, 100.0f });
-			InOutConfigSet.Add({ FMemAllocTableColumns::TagColumnId,      true, 120.0f });
+			InOutConfigSet.Add({ FTable::GetHierarchyColumnId(),               true, 200.0f });
+			InOutConfigSet.Add({ FMemAllocTableColumns::CountColumnId,         true, 100.0f });
+			InOutConfigSet.Add({ FMemAllocTableColumns::SizeColumnId,          true, 100.0f });
+			InOutConfigSet.Add({ FMemAllocTableColumns::TagColumnId,           true, 120.0f });
 			InOutConfigSet.Add({ FMemAllocTableColumns::AllocFunctionColumnId, true, 400.0f });
 		}
 	};
diff --git a/Engine/Source/Developer/TraceInsights/Private/Insights/Table/Widgets/STableTreeView.cpp b/Engine/Source/Developer/TraceInsights/Private/Insights/Table/Widgets/STableTreeView.cpp
--- a/Engine/Source/Developer/TraceInsights/Private/Insights/Table/Widgets/STableTreeView.cpp
+++ b/Engine/Source/Developer/TraceInsights/Private/Insights/Table/Widgets/STableTreeView.cpp
@@ -3795,7 +3795,14 @@ void STableTreeView::UpdateBannerText()
 	}
 	else if (Root->GetFilteredChildrenCount() == 0)
 	{
-		TreeViewBannerText = LOCTEXT("HierarchyFilteringZeroResults", "No tree node is matching the current text search.");
+		if (TextFilter->GetRawFilterText().IsEmpty())
+		{
+			TreeViewBannerText = LOCTEXT("TreeViewIsUpdating", "Tree view is updating. Please wait.");
+		}
+		else
+		{
+			TreeViewBannerText = LOCTEXT("HierarchyFilteringZeroResults", "No tree node is matching the current text search.");
+		}
 	}
 }
 
