[Insights] Add contains filter operator for the Timer Name filter in the Quick Find widget. Add new functionality to the filtering system to support more complex filters.

diff --git a/Engine/Source/Developer/TraceInsights/Private/Insights/ViewModels/FilterConfigurator.h b/Engine/Source/Developer/TraceInsights/Private/Insights/ViewModels/FilterConfigurator.h
--- a/Engine/Source/Developer/TraceInsights/Private/Insights/ViewModels/FilterConfigurator.h
+++ b/Engine/Source/Developer/TraceInsights/Private/Insights/ViewModels/FilterConfigurator.h
@@ -27,6 +27,9 @@ public:
 
 	bool IsEmpty() const { return RootNode->GetChildrenCount() == 0; }
 
+	/** Called to update the internal state of some filters */
+	void Update() { RootNode->Update(); }
+
 	virtual bool ApplyFilters(const FFilterContext& Context) const override;
 
 	bool IsKeyUsed(int32 Key) const;
diff --git a/Engine/Source/Developer/TraceInsights/Private/Insights/ViewModels/FilterConfiguratorNode.cpp b/Engine/Source/Developer/TraceInsights/Private/Insights/ViewModels/FilterConfiguratorNode.cpp
--- a/Engine/Source/Developer/TraceInsights/Private/Insights/ViewModels/FilterConfiguratorNode.cpp
+++ b/Engine/Source/Developer/TraceInsights/Private/Insights/ViewModels/FilterConfiguratorNode.cpp
@@ -40,6 +40,10 @@ TSharedRef<FFilterConfiguratorNode> FFilterConfiguratorNode::DeepCopy(const FFil
 	NodeCopy.SelectedFilterOperator = InSourceNode.SelectedFilterOperator;
 	NodeCopy.SelectedFilterGroupOperator = InSourceNode.SelectedFilterGroupOperator;
 	NodeCopy.TextBoxValue = InSourceNode.TextBoxValue;
+	if (InSourceNode.SelectedFilter.IsValid())
+	{
+		NodeCopy.FilterState = InSourceNode.SelectedFilter->BuildFilterState(*InSourceNode.FilterState);
+	}
 
 	NodeCopy.SetExpansion(InSourceNode.IsExpanded());
 
@@ -105,13 +109,20 @@ void FFilterConfiguratorNode::SetAvailableFilters(TSharedPtr<TArray<TSharedPtr<F
 
 void FFilterConfiguratorNode::SetSelectedFilter(TSharedPtr<FFilter> InSelectedFilter)
 {
+	if (SelectedFilter.Get() != InSelectedFilter.Get())
+	{
+		TextBoxValue.Empty();
+	}
+
 	SelectedFilter = InSelectedFilter;
 	if (SelectedFilter.IsValid() && SelectedFilter->GetSupportedOperators()->Num() > 0)
 	{
+		FilterState = SelectedFilter->BuildFilterState();
+
 		SetSelectedFilterOperator(SelectedFilter->GetSupportedOperators()->GetData()[0]);
 
 		AvailableFilterOperators->Empty();
-		SupportedOperatorsArrayPtr AvailableOperators = InSelectedFilter->GetSupportedOperators();
+		SupportedOperatorsArrayConstPtr AvailableOperators = InSelectedFilter->GetSupportedOperators();
 		for (auto& FilterOperator : *AvailableOperators)
 		{
 			AvailableFilterOperators->Add(FilterOperator);
@@ -125,8 +136,8 @@ void FFilterConfiguratorNode::ProcessFilter()
 {
 	if (IsGroup())
 	{
-		TArray<FBaseTreeNodePtr> ChildArray = GetChildrenMutable();
-		for (FBaseTreeNodePtr Child : ChildArray)
+		TArray<FBaseTreeNodePtr>& ChildArray = GetChildrenMutable();
+		for (FBaseTreeNodePtr& Child : ChildArray)
 		{
 			FFilterConfiguratorNodePtr CastedChild = StaticCastSharedPtr<FFilterConfiguratorNode>(Child);
 			CastedChild->ProcessFilter();
@@ -134,62 +145,16 @@ void FFilterConfiguratorNode::ProcessFilter()
 	}
 	else
 	{
-		switch (SelectedFilter->GetDataType())
-		{
-		case EFilterDataType::Double:
-		{
-			if (SelectedFilter->GetConverter().IsValid())
-			{
-				double Value = 0.0;
-				FText Errors;
-				bool Result = SelectedFilter->GetConverter()->Convert(TextBoxValue, Value, Errors);
-				FilterValue.Set<double>(Result ? Value : 0.0);
-			}
-			else
-			{
-				FilterValue.Set<double>(FCString::Atod(*TextBoxValue));
-			}
-			break;
-		}
-		case EFilterDataType::Int64:
-		{
-			if (SelectedFilter->GetConverter().IsValid())
-			{
-				int64 Value = 0;
-				FText Errors;
-				bool Result = SelectedFilter->GetConverter()->Convert(TextBoxValue, Value, Errors);
-				FilterValue.Set<int64>(Result ? Value : 0);
-			}
-			else
-			{
-				if (TextBoxValue.Contains(TEXT("x")))
-				{
-					FilterValue.Set<int64>((int64)FParse::HexNumber64(*TextBoxValue));
-				}
-				else
-				{
-					FilterValue.Set<int64>(FCString::Atoi64(*TextBoxValue));
-				}
-			}
-			break;
-		}
-		case EFilterDataType::String:
+		if (FilterState.IsValid())
 		{
-			FilterValue.Set<FString>(TextBoxValue);
-			break;
-		}
-		case EFilterDataType::StringInt64Pair:
-		{
-			checkf(SelectedFilter->GetConverter().IsValid(), TEXT("StringToInt64Pair filters must have a converter set"));
-			int64 Value = 0;
-			FText Errors;
-			bool Result = SelectedFilter->GetConverter()->Convert(TextBoxValue, Value, Errors);
-			FilterValue.Set<int64>(Result ? Value : -1);
-		}
+			FilterState->SetFilterValue(TextBoxValue);
+			FilterState->Update();
 		}
 	}
 }
 
+////////////////////////////////////////////////////////////////////////////////////////////////////
+
 void FFilterConfiguratorNode::GetUsedKeys(TSet<int32>& KeysUsed) const
 {
 	if (IsGroup())
@@ -242,48 +207,61 @@ bool FFilterConfiguratorNode::ApplyFilters(const FFilterContext& Context) const
 	}
 	else
 	{
-		if (!Context.HasFilterData(SelectedFilter->GetKey()))
-		{
-			// If data is not set for this filter return the value specified in the Context.
-			return Context.GetReturnValueForUnsetFilters();
-		}
+		Ret = FilterState->ApplyFilter(Context);
+	}
 
-		switch (SelectedFilter->GetDataType())
-		{
-		case EFilterDataType::Double:
-		{
-			FFilterOperator<double>* Operator = (FFilterOperator<double>*) SelectedFilterOperator.Get();
-			double Value = 0.0;
-			Context.GetFilterData<double>(SelectedFilter->GetKey(), Value);
+	return Ret;
+}
 
-			Ret = Operator->Apply(Value, FilterValue.Get<double>());
-			break;
-		}
-		case EFilterDataType::Int64:
-		case EFilterDataType::StringInt64Pair:
-		{
-			FFilterOperator<int64>* Operator = (FFilterOperator<int64>*) SelectedFilterOperator.Get();
-			int64 Value = 0;
-			Context.GetFilterData<int64>(SelectedFilter->GetKey(), Value);
+////////////////////////////////////////////////////////////////////////////////////////////////////
 
-			Ret = Operator->Apply(Value, FilterValue.Get<int64>());
-			break;
-		}
-		case EFilterDataType::String:
+void FFilterConfiguratorNode::Update()
+{
+	if (IsGroup())
+	{
+		TArray<FBaseTreeNodePtr>& ChildArray = GetChildrenMutable();
+		for (FBaseTreeNodePtr& Child : ChildArray)
 		{
-			FFilterOperator<FString>* Operator = (FFilterOperator<FString>*) SelectedFilterOperator.Get();
-			FString Value;
-			Context.GetFilterData<FString>(SelectedFilter->GetKey(), Value);
-
-			Ret = Operator->Apply(Value, FilterValue.Get<FString>());
-			break;
+			FFilterConfiguratorNodePtr CastedChild = StaticCastSharedPtr<FFilterConfiguratorNode>(Child);
+			CastedChild->Update();
 		}
-		default:
-			break;
+	}
+	else
+	{
+		if (FilterState.IsValid())
+		{
+			FilterState->Update();
 		}
 	}
+}
 
-	return Ret;
+////////////////////////////////////////////////////////////////////////////////////////////////////
+
+void FFilterConfiguratorNode::SetTextBoxValue(const FString& InValue)
+{
+	TextBoxValue = InValue;
+}
+
+////////////////////////////////////////////////////////////////////////////////////////////////////
+
+void FFilterConfiguratorNode::SetSelectedFilterOperator(TSharedPtr<IFilterOperator> InSelectedFilterOperator)
+{
+	if (FilterState.IsValid())
+	{
+		FilterState->SetSelectedOperator(InSelectedFilterOperator);
+	}
+}
+
+////////////////////////////////////////////////////////////////////////////////////////////////////
+
+TSharedPtr<IFilterOperator> FFilterConfiguratorNode::GetSelectedFilterOperator() const
+{
+	if (FilterState.IsValid())
+	{
+		return FilterState->GetSelectedOperator();
+	}
+
+	return nullptr;
 }
 
 ////////////////////////////////////////////////////////////////////////////////////////////////////
diff --git a/Engine/Source/Developer/TraceInsights/Private/Insights/ViewModels/FilterConfiguratorNode.h b/Engine/Source/Developer/TraceInsights/Private/Insights/ViewModels/FilterConfiguratorNode.h
--- a/Engine/Source/Developer/TraceInsights/Private/Insights/ViewModels/FilterConfiguratorNode.h
+++ b/Engine/Source/Developer/TraceInsights/Private/Insights/ViewModels/FilterConfiguratorNode.h
@@ -63,17 +63,17 @@ public:
 	void SetSelectedFilter(TSharedPtr<FFilter> InSelectedFilter);
 	TSharedPtr<FFilter> GetSelectedFilter() const { return SelectedFilter; }
 
-	void SetSelectedFilterOperator(TSharedPtr<IFilterOperator> InSelectedFilterOperator) { SelectedFilterOperator = InSelectedFilterOperator; }
-	TSharedPtr<IFilterOperator> GetSelectedFilterOperator() const { return SelectedFilterOperator;	}
-	const TArray<TSharedPtr<FFilterGroupOperator>>& GetFilterGroupOperators();
+	void SetSelectedFilterOperator(TSharedPtr<IFilterOperator> InSelectedFilterOperator);
+	TSharedPtr<IFilterOperator> GetSelectedFilterOperator() const;
 
+	const TArray<TSharedPtr<FFilterGroupOperator>>& GetFilterGroupOperators();
 	void SetSelectedFilterGroupOperator(TSharedPtr<FFilterGroupOperator> InSelectedFilterGroupOperator) { SelectedFilterGroupOperator = InSelectedFilterGroupOperator; }
 	TSharedPtr<FFilterGroupOperator> GetSelectedFilterGroupOperator() const { return SelectedFilterGroupOperator; }
 
 	TSharedPtr<TArray<TSharedPtr<IFilterOperator>>> GetAvailableFilterOperators() const {	return AvailableFilterOperators;	}
 
 	const FString& GetTextBoxValue() { return TextBoxValue; }
-	void SetTextBoxValue(const FString& InValue) { TextBoxValue = InValue; }
+	void SetTextBoxValue(const FString& InValue);
 
 	bool ApplyFilters(const class FFilterContext& Context) const;
 
@@ -81,6 +81,8 @@ public:
 
 	void ProcessFilter();
 
+	void Update();
+
 private:
 	FFilterConfiguratorNode& operator=(const FFilterConfiguratorNode& Other);
 
@@ -96,8 +98,7 @@ private:
 
 	FString TextBoxValue;
 
-	// The textbox value, converted to its data type
-	FFilterContext::ContextData FilterValue;
+	TSharedPtr<FFilterState> FilterState;
 };
 
 ////////////////////////////////////////////////////////////////////////////////////////////////////
diff --git a/Engine/Source/Developer/TraceInsights/Private/Insights/ViewModels/Filters.cpp b/Engine/Source/Developer/TraceInsights/Private/Insights/ViewModels/Filters.cpp
--- a/Engine/Source/Developer/TraceInsights/Private/Insights/ViewModels/Filters.cpp
+++ b/Engine/Source/Developer/TraceInsights/Private/Insights/ViewModels/Filters.cpp
@@ -15,8 +15,10 @@
 namespace Insights
 {
 
+INSIGHTS_IMPLEMENT_RTTI(FFilterState)
 INSIGHTS_IMPLEMENT_RTTI(FFilter)
 INSIGHTS_IMPLEMENT_RTTI(FFilterWithSuggestions)
+INSIGHTS_IMPLEMENT_RTTI(FCustomFilter)
 
 ////////////////////////////////////////////////////////////////////////////////////////////////////
 // FFilterStorage
@@ -146,6 +148,119 @@ void FFilterService::UnregisterTabSpawner()
 
 ////////////////////////////////////////////////////////////////////////////////////////////////////
 
+void FFilterState::SetFilterValue(FString InTextValue)
+{
+	if (InTextValue.IsEmpty())
+	{
+		FilterValue = FFilterContext::ContextData();
+	}
+
+	switch (Filter->GetDataType())
+	{
+	case EFilterDataType::Double:
+	{
+		if (Filter->GetConverter().IsValid())
+		{
+			double Value = 0.0;
+			FText Errors;
+			bool Result = Filter->GetConverter()->Convert(InTextValue, Value, Errors);
+			FilterValue.Set<double>(Result ? Value : 0.0);
+		}
+		else
+		{
+			FilterValue.Set<double>(FCString::Atod(*InTextValue));
+		}
+		break;
+	}
+	case EFilterDataType::Int64:
+	{
+		if (Filter->GetConverter().IsValid())
+		{
+			int64 Value = 0;
+			FText Errors;
+			bool Result = Filter->GetConverter()->Convert(InTextValue, Value, Errors);
+			FilterValue.Set<int64>(Result ? Value : 0);
+		}
+		else
+		{
+			if (InTextValue.Contains(TEXT("x")))
+			{
+				FilterValue.Set<int64>((int64)FParse::HexNumber64(*InTextValue));
+			}
+			else
+			{
+				FilterValue.Set<int64>(FCString::Atoi64(*InTextValue));
+			}
+		}
+		break;
+	}
+	case EFilterDataType::String:
+	{
+		FilterValue.Set<FString>(InTextValue);
+		break;
+	}
+	case EFilterDataType::StringInt64Pair:
+	{
+		checkf(Filter->GetConverter().IsValid(), TEXT("StringToInt64Pair filters must have a converter set"));
+		int64 Value = 0;
+		FText Errors;
+		bool Result = Filter->GetConverter()->Convert(InTextValue, Value, Errors);
+		FilterValue.Set<int64>(Result ? Value : -1);
+	}
+	}
+}
+
+////////////////////////////////////////////////////////////////////////////////////////////////////
+
+bool FFilterState::ApplyFilter(const FFilterContext& Context) const
+{
+	if (!Context.HasFilterData(Filter->GetKey()))
+	{
+		// If data is not set for this filter return the value specified in the Context.
+		return Context.GetReturnValueForUnsetFilters();
+	}
+
+	bool Ret = true;
+
+	switch (Filter->GetDataType())
+	{
+	case EFilterDataType::Double:
+	{
+		FFilterOperator<double>* Operator = (FFilterOperator<double>*) SelectedOperator.Get();
+		double Value = 0.0;
+		Context.GetFilterData<double>(Filter->GetKey(), Value);
+
+		Ret = Operator->Apply(Value, FilterValue.Get<double>());
+		break;
+	}
+	case EFilterDataType::Int64:
+	case EFilterDataType::StringInt64Pair:
+	{
+		FFilterOperator<int64>* Operator = (FFilterOperator<int64>*) SelectedOperator.Get();
+		int64 Value = 0;
+		Context.GetFilterData<int64>(Filter->GetKey(), Value);
+
+		Ret = Operator->Apply(Value, FilterValue.Get<int64>());
+		break;
+	}
+	case EFilterDataType::String:
+	{
+		FFilterOperator<FString>* Operator = (FFilterOperator<FString>*) SelectedOperator.Get();
+		FString Value;
+		Context.GetFilterData<FString>(Filter->GetKey(), Value);
+
+		Ret = Operator->Apply(Value, FilterValue.Get<FString>());
+		break;
+	}
+	default:
+		break;
+	}
+
+	return Ret;
+}
+
+////////////////////////////////////////////////////////////////////////////////////////////////////
+
 } // namespace Insights
 
 #undef LOCTEXT_NAMESPACE
\ No newline at end of file
diff --git a/Engine/Source/Developer/TraceInsights/Private/Insights/ViewModels/Filters.h b/Engine/Source/Developer/TraceInsights/Private/Insights/ViewModels/Filters.h
--- a/Engine/Source/Developer/TraceInsights/Private/Insights/ViewModels/Filters.h
+++ b/Engine/Source/Developer/TraceInsights/Private/Insights/ViewModels/Filters.h
@@ -24,6 +24,8 @@
 	Double,
 	String,
 	StringInt64Pair, // Displayed as a string but translates to a Int64 key.
+	
+	Custom, // For complex filters that are implemented as separate classes.
 };
 
 ////////////////////////////////////////////////////////////////////////////////////////////////////
@@ -123,9 +125,86 @@
 
 ////////////////////////////////////////////////////////////////////////////////////////////////////
 
-typedef TSharedPtr<const TArray<TSharedPtr<IFilterOperator>>> SupportedOperatorsArrayPtr;
+class FFilterContext
+{
+public:
+	typedef TVariant<double, int64, FString> ContextData;
 
-class FFilter
+public:
+	template<typename T>
+	void AddFilterData(int32 Key, const T& InData)
+	{
+		ContextData VariantData;
+		VariantData.Set<T>(InData);
+		DataMap.Add(Key, VariantData);
+	}
+
+	template<typename T>
+	void SetFilterData(int32 Key, const T& InData)
+	{
+		DataMap[Key].Set<T>(InData);
+	}
+
+	template<typename T>
+	void GetFilterData(int32 Key, T& OutData) const
+	{
+		const ContextData* Data = DataMap.Find(Key);
+		check(Data);
+
+		check(Data->IsType<T>());
+		OutData = Data->Get<T>();
+	}
+
+	bool HasFilterData(int32 Key) const
+	{
+		return DataMap.Contains(Key);
+	}
+
+	bool GetReturnValueForUnsetFilters() const { return bReturnValueForUnsetFilters; }
+	void SetReturnValueForUnsetFilters(bool InValue) { bReturnValueForUnsetFilters = InValue; }
+
+private:
+	TMap<int32, ContextData> DataMap;
+	bool bReturnValueForUnsetFilters = true;
+};
+
+////////////////////////////////////////////////////////////////////////////////////////////////////
+
+typedef TSharedPtr<const TArray<TSharedPtr<IFilterOperator>>> SupportedOperatorsArrayConstPtr;
+typedef TSharedPtr<TArray<TSharedPtr<IFilterOperator>>> SupportedOperatorsArrayPtr;
+
+////////////////////////////////////////////////////////////////////////////////////////////////////
+
+class FFilter;
+
+class FFilterState
 {
+	INSIGHTS_DECLARE_RTTI_BASE(FFilterState)
+public:
+	FFilterState(TSharedRef<FFilter> InFilter)
+		: Filter(InFilter)
+	{
+	}
+
+	virtual ~FFilterState() {}
+
+	void SetSelectedOperator(TSharedPtr<IFilterOperator> InOperator) { SelectedOperator = InOperator; }
+	TSharedPtr<IFilterOperator> GetSelectedOperator() { return SelectedOperator; }
+
+	virtual void Update() {};
+	virtual bool ApplyFilter(const FFilterContext& Context) const;
+
+	virtual void SetFilterValue(FString InTextValue);
+
+protected:
+	TSharedRef<FFilter> Filter;
+	TSharedPtr<IFilterOperator> SelectedOperator;
+	FFilterContext::ContextData FilterValue;
+};
+
+////////////////////////////////////////////////////////////////////////////////////////////////////
+
+class FFilter : public TSharedFromThis<FFilter>
+{
 	INSIGHTS_DECLARE_RTTI_BASE(FFilter)
 
@@ -149,9 +228,12 @@
 	const FText& GetDesc() const { return Desc; }
 	EFilterDataType GetDataType() const { return DataType; }
 	const TSharedPtr<IFilterValueConverter>& GetConverter() const { return Converter; }
-	SupportedOperatorsArrayPtr GetSupportedOperators() const { return SupportedOperators; }
+	SupportedOperatorsArrayConstPtr GetSupportedOperators() const { return SupportedOperators; }
 
-private:
+	virtual TSharedRef<FFilterState> BuildFilterState() { return MakeShared<FFilterState>(SharedThis(this)); }
+	virtual TSharedRef<FFilterState> BuildFilterState(const FFilterState& Other) { return MakeShared<FFilterState>(Other); }
+
+protected:
 	int32 Key;
 	FText Name;
 	FText Desc;
@@ -244,49 +326,18 @@
 
 ////////////////////////////////////////////////////////////////////////////////////////////////////
 
-class FFilterContext
+class FCustomFilter : public FFilterWithSuggestions
 {
+	INSIGHTS_DECLARE_RTTI(FCustomFilter, FFilterWithSuggestions)
 public:
-	typedef TVariant<double, int64, FString> ContextData;
-
-public:
-	template<typename T>
-	void AddFilterData(int32 Key, const T& InData)
+	FCustomFilter(int32 InKey, FText InName, FText InDesc, EFilterDataType InDataType, TSharedPtr<IFilterValueConverter> InConverter, SupportedOperatorsArrayPtr InSupportedOperators)
+		: FFilterWithSuggestions(InKey, InName, InDesc, InDataType, InConverter, InSupportedOperators)
 	{
-		ContextData VariantData;
-		VariantData.Set<T>(InData);
-		DataMap.Add(Key, VariantData);
 	}
 
-	template<typename T>
-	void SetFilterData(int32 Key, const T& InData)
+	virtual ~FCustomFilter()
 	{
-		DataMap[Key].Set<T>(InData);
 	}
-
-	template<typename T>
-	void GetFilterData(int32 Key, T& OutData) const
-	{
-		const ContextData* Data = DataMap.Find(Key);
-		check(Data);
-
-		check(Data->IsType<T>());
-		OutData = Data->Get<T>();
-	}
-
-	bool HasFilterData(int32 Key) const
-	{
-		return DataMap.Contains(Key);
-	}
-
-	bool GetReturnValueForUnsetFilters() const { return bReturnValueForUnsetFilters; }
-	void SetReturnValueForUnsetFilters(bool InValue) { bReturnValueForUnsetFilters = InValue; }
-
-private:
-	TMap<int32, ContextData> DataMap;
-	bool bReturnValueForUnsetFilters = true;
 };
 
-////////////////////////////////////////////////////////////////////////////////////////////////////
-
 } // namespace Insights
diff --git a/Engine/Source/Developer/TraceInsights/Private/Insights/ViewModels/TimerFilters.cpp b/Engine/Source/Developer/TraceInsights/Private/Insights/ViewModels/TimerFilters.cpp
new file mode 100644
--- /dev/null
+++ b/Engine/Source/Developer/TraceInsights/Private/Insights/ViewModels/TimerFilters.cpp
@@ -0,0 +1,139 @@
+// Copyright Epic Games, Inc. All Rights Reserved.
+
+#include "TimerFilters.h"
+
+#include "Insights/TimingProfilerManager.h"
+#include "Insights/ViewModels/ThreadTimingTrack.h"
+
+#define LOCTEXT_NAMESPACE "Insights::TimerFilters"
+
+namespace Insights
+{
+
+INSIGHTS_IMPLEMENT_RTTI(FTimerNameFilterState)
+INSIGHTS_IMPLEMENT_RTTI(FTimerNameFilter)
+
+////////////////////////////////////////////////////////////////////////////////////////////////////
+// FTimerNameFilterState
+////////////////////////////////////////////////////////////////////////////////////////////////////
+
+void FTimerNameFilterState::Update()
+{
+	if (FilterValue.IsEmpty() || !SelectedOperator.IsValid())
+	{
+		return;
+	}
+
+	TimerIds.Empty();
+	TSharedPtr<const TraceServices::IAnalysisSession> Session = FInsightsManager::Get()->GetSession();
+	if (Session.IsValid() && TraceServices::ReadTimingProfilerProvider(*Session.Get()))
+	{
+		TraceServices::FAnalysisSessionReadScope SessionReadScope(*Session.Get());
+
+		const TraceServices::ITimingProfilerProvider& TimingProfilerProvider = *TraceServices::ReadTimingProfilerProvider(*Session.Get());
+
+		const TraceServices::ITimingProfilerTimerReader* TimerReader;
+		TimingProfilerProvider.ReadTimers([&TimerReader](const TraceServices::ITimingProfilerTimerReader& Out) { TimerReader = &Out; });
+
+		uint32 TimerCount = TimerReader->GetTimerCount();
+		for (uint32 TimerIndex = 0; TimerIndex < TimerCount; ++TimerIndex)
+		{
+			const TraceServices::FTimingProfilerTimer* Timer = TimerReader->GetTimer(TimerIndex);
+			if (Timer && Timer->Name)
+			{
+				if (SelectedOperator->GetKey() == EFilterOperator::Eq)
+				{
+					if (FCString::Stricmp(Timer->Name, *FilterValue) == 0)
+					{
+						TimerIds.Add(Timer->Id);
+					}
+				}
+				else if (SelectedOperator->GetKey() == EFilterOperator::Contains)
+				{
+					if (FCString::Stristr(Timer->Name, *FilterValue))
+					{
+						TimerIds.Add(Timer->Id);
+					}
+				}
+			}
+		}
+	}
+}
+
+////////////////////////////////////////////////////////////////////////////////////////////////////
+
+bool FTimerNameFilterState::ApplyFilter(const FFilterContext& Context) const
+{
+	if (!Context.HasFilterData(static_cast<int32>(Filter->GetKey())))
+	{
+		return Context.GetReturnValueForUnsetFilters();
+	}
+
+	int64 Value;
+	Context.GetFilterData<int64>(Filter->GetKey(), Value);
+
+	return TimerIds.Contains(static_cast<uint32>(Value));
+}
+
+////////////////////////////////////////////////////////////////////////////////////////////////////
+// FTimerNameFilter
+////////////////////////////////////////////////////////////////////////////////////////////////////
+
+FTimerNameFilter::FTimerNameFilter()
+	: FCustomFilter(static_cast<int32>(EFilterField::TimerName),
+		LOCTEXT("TimerName", "Timer Name"),
+		LOCTEXT("TimerName", "Timer Name"),
+		EFilterDataType::Custom,
+		nullptr,
+		nullptr)
+{
+	SupportedOperators = MakeShared<TArray<TSharedPtr<IFilterOperator>>>();
+	SupportedOperators->Add(StaticCastSharedRef<IFilterOperator>(MakeShared<FFilterOperator<int64>>(EFilterOperator::Eq, TEXT("Is"), [](int64 lhs, int64 rhs) { return lhs == rhs; })));
+	SupportedOperators->Add(StaticCastSharedRef<IFilterOperator>(MakeShared<FFilterOperator<int64>>(EFilterOperator::Contains, TEXT("Contains"), [](int64 lhs, int64 rhs) { return lhs == rhs; })));
+
+	SetCallback([this](const FString& Text, TArray<FString>& OutSuggestions)
+		{
+			this->PopulateTimerNameSuggestionList(Text, OutSuggestions);
+		});
+}
+
+////////////////////////////////////////////////////////////////////////////////////////////////////
+
+void FTimerNameFilter::PopulateTimerNameSuggestionList(const FString& Text, TArray<FString>& OutSuggestions)
+{
+	TSharedPtr<const TraceServices::IAnalysisSession> Session = FInsightsManager::Get()->GetSession();
+	if (Session.IsValid() && TraceServices::ReadTimingProfilerProvider(*Session.Get()))
+	{
+		TraceServices::FAnalysisSessionReadScope SessionReadScope(*Session.Get());
+
+		const TraceServices::ITimingProfilerProvider& TimingProfilerProvider = *TraceServices::ReadTimingProfilerProvider(*Session.Get());
+
+		const TraceServices::ITimingProfilerTimerReader* TimerReader;
+		TimingProfilerProvider.ReadTimers([&TimerReader](const TraceServices::ITimingProfilerTimerReader& Out) { TimerReader = &Out; });
+
+		uint32 TimerCount = TimerReader->GetTimerCount();
+		for (uint32 TimerIndex = 0; TimerIndex < TimerCount; ++TimerIndex)
+		{
+			const TraceServices::FTimingProfilerTimer* Timer = TimerReader->GetTimer(TimerIndex);
+			if (Timer && Timer->Name)
+			{
+				if (Text.IsEmpty())
+				{
+					OutSuggestions.Add(Timer->Name);
+					continue;
+				}
+				const TCHAR* FoundString = FCString::Stristr(Timer->Name, *Text);
+				if (FoundString)
+				{
+					OutSuggestions.Add(Timer->Name);
+				}
+			}
+		}
+	}
+}
+
+////////////////////////////////////////////////////////////////////////////////////////////////////
+
+} // namespace Insights
+
+#undef LOCTEXT_NAMESPACE
\ No newline at end of file
diff --git a/Engine/Source/Developer/TraceInsights/Private/Insights/ViewModels/TimerFilters.h b/Engine/Source/Developer/TraceInsights/Private/Insights/ViewModels/TimerFilters.h
new file mode 100644
--- /dev/null
+++ b/Engine/Source/Developer/TraceInsights/Private/Insights/ViewModels/TimerFilters.h
@@ -0,0 +1,56 @@
+// Copyright Epic Games, Inc. All Rights Reserved.
+
+#pragma once
+
+#include "Insights/Common/SimpleRtti.h"
+#include "Insights/ViewModels/Filters.h"
+
+namespace Insights
+{
+
+class FTimerNameFilterState : public FFilterState
+{
+
+	INSIGHTS_DECLARE_RTTI(FTimerNameFilterState, FFilterState)
+public:
+	FTimerNameFilterState(TSharedRef<FFilter> InFilter)
+		: FFilterState(InFilter)
+	{}
+
+	virtual ~FTimerNameFilterState() {}
+
+	virtual void Update() override;
+
+	virtual bool ApplyFilter(const FFilterContext& Context) const override;
+
+	virtual void SetFilterValue(FString InFilterValue) override { FilterValue = InFilterValue; }
+
+private:
+	FString FilterValue;
+	TSet<uint32> TimerIds;
+	bool bShouldUpdate = true;
+};
+
+class FTimerNameFilter : public FCustomFilter
+{
+	INSIGHTS_DECLARE_RTTI(FTimerNameFilter, FCustomFilter)
+
+public:
+	FTimerNameFilter();
+
+	virtual ~FTimerNameFilter() {}
+
+	virtual TSharedRef<FFilterState> BuildFilterState() 
+	{ 
+		return MakeShared<FTimerNameFilterState>(SharedThis(this)); 
+	}
+
+	virtual TSharedRef<FFilterState> BuildFilterState(const FFilterState& Other)
+	{
+		return MakeShared<FTimerNameFilterState>(static_cast<const FTimerNameFilterState&>(Other));
+	}
+
+	void PopulateTimerNameSuggestionList(const FString& Text, TArray<FString>& OutSuggestions);
+};
+
+} // namespace Insights
diff --git a/Engine/Source/Developer/TraceInsights/Private/Insights/Widgets/STimingView.cpp b/Engine/Source/Developer/TraceInsights/Private/Insights/Widgets/STimingView.cpp
--- a/Engine/Source/Developer/TraceInsights/Private/Insights/Widgets/STimingView.cpp
+++ b/Engine/Source/Developer/TraceInsights/Private/Insights/Widgets/STimingView.cpp
@@ -63,6 +63,7 @@
 #include "Insights/ViewModels/RegionsTimingTrack.h"
 #include "Insights/ViewModels/ThreadTimingTrack.h"
 #include "Insights/ViewModels/TimeFilterValueConverter.h"
+#include "Insights/ViewModels/TimerFilters.h"
 #include "Insights/ViewModels/TimeRulerTrack.h"
 #include "Insights/ViewModels/TimingEventSearch.h"
 #include "Insights/ViewModels/TimingGraphTrack.h"
@@ -531,6 +532,8 @@ void STimingView::Tick(const FGeometry& AllottedGeometry, const double InCurrent
 
 	LLM_SCOPE_BYTAG(Insights);
 
+	UpdateFilters();
+
 	ThisGeometry = AllottedGeometry;
 
 	Tooltip.SetFontScale(AllottedGeometry.Scale);
@@ -4960,20 +4963,7 @@ void STimingView::QuickFind_Execute()
 			nullptr,
 			FFilterService::Get()->GetIntegerOperators()));
 
-		TSharedPtr<TArray<TSharedPtr<IFilterOperator>>> EventNameFilterOperators = MakeShared<TArray<TSharedPtr<IFilterOperator>>>();
-		EventNameFilterOperators->Add(StaticCastSharedRef<IFilterOperator>(MakeShared<FFilterOperator<int64>>(EFilterOperator::Eq, TEXT("Is"), [](int64 lhs, int64 rhs) { return lhs == rhs; })));
-		TSharedRef<FFilterWithSuggestions> TimerNameFilter = MakeShared<FFilterWithSuggestions>(
-			static_cast<int32>(EFilterField::TimerName),
-			LOCTEXT("TimerName", "Timer Name"),
-			LOCTEXT("TimerName", "Timer Name"),
-			EFilterDataType::StringInt64Pair,
-			MakeShared<FEventNameFilterValueConverter>(),
-			EventNameFilterOperators);
-		TimerNameFilter->SetCallback([this](const FString& Text, TArray<FString>& OutSuggestions)
-		{
-			this->PopulateTimerNameSuggestionList(Text, OutSuggestions);
-		});
-		NewFilterConfigurator->Add(TimerNameFilter);
+		NewFilterConfigurator->Add(MakeShared<FTimerNameFilter>());
 
 		for (Insights::ITimingViewExtender* Extender : GetExtenders())
 		{
@@ -5235,6 +5225,8 @@ void STimingView::FilterAllTracks()
 {
 	LLM_SCOPE_BYTAG(Insights);
 	FilterConfigurator = MakeShared<Insights::FFilterConfigurator>(*QuickFindVm->GetFilterConfigurator());
+	FilterConfigurator->Update();
+
 	for (auto& Entry : AllTracks)
 	{
 		Entry.Value->SetFilterConfigurator(FilterConfigurator);
@@ -5375,5 +5367,26 @@ void STimingView::SelectEventInstance(uint32 TimerId, ESelectEventType Type, boo
 
 ////////////////////////////////////////////////////////////////////////////////////////////////////
 
+void STimingView::UpdateFilters()
+{
+	if (!bUpdateFilters)
+	{
+		return;
+	}
+
+	if (FInsightsManager::Get()->IsAnalysisComplete())
+	{
+		// This will be the final update.
+		bUpdateFilters = false;
+	}
+
+	if (FilterConfigurator.IsValid())
+	{
+		FilterConfigurator->Update();
+	}
+}
+
+////////////////////////////////////////////////////////////////////////////////////////////////////
+
 #undef INSIGHTS_ACTIVATE_BENCHMARK
 #undef LOCTEXT_NAMESPACE
diff --git a/Engine/Source/Developer/TraceInsights/Private/Insights/Widgets/STimingView.h b/Engine/Source/Developer/TraceInsights/Private/Insights/Widgets/STimingView.h
--- a/Engine/Source/Developer/TraceInsights/Private/Insights/Widgets/STimingView.h
+++ b/Engine/Source/Developer/TraceInsights/Private/Insights/Widgets/STimingView.h
@@ -526,6 +526,11 @@ protected:
 	void EnumerateFilteredTracks(TSharedPtr<Insights::FFilterConfigurator> FilterConfigurator, EnumerateFilteredTracksCallback Callback);
 
 	ETraceFrameType GetFrameTypeToSnapTo();
+	
+	/** The FilterConfigurator needs to be updated so that custom filters work correctly when analysis is still running
+	* and the timer list can change. 
+	 */ 
+	void UpdateFilters();
 
 protected:
 	/** The name of the view. */
@@ -744,6 +749,7 @@ protected:
 	TSharedPtr<Insights::FFilterConfigurator> FilterConfigurator;
 	static uint32 TimingViewId;
 	const FName QuickFindTabId;
+	bool bUpdateFilters = true;
 	
 	// Used only between the creation of the widget and the spawning of the owning tab. When the tab is spawned, we relinquish ownership.
 	TSharedPtr<Insights::SQuickFind> QuickFindWidgetSharedPtr;
