[Insights] Optimization for some Quick Find operations.

diff --git a/Engine/Source/Developer/TraceInsights/Private/Insights/Widgets/STimingView.cpp b/Engine/Source/Developer/TraceInsights/Private/Insights/Widgets/STimingView.cpp
--- a/Engine/Source/Developer/TraceInsights/Private/Insights/Widgets/STimingView.cpp
+++ b/Engine/Source/Developer/TraceInsights/Private/Insights/Widgets/STimingView.cpp
@@ -5116,7 +5116,8 @@ void STimingView::FindPrevEvent()
 	Params.FilterExecutor = QuickFindVm->GetFilterConfigurator();
 	Params.SearchDirection = FTimingEventSearchParameters::ESearchDirection::Backward;
 
-	EnumerateFilteredTracks(QuickFindVm->GetFilterConfigurator(), [&Params, &BestMatchEvent](TSharedPtr<FBaseTimingTrack>& Track)
+	TSharedPtr<const FBaseTimingTrack> PriorityTrack = SelectedEvent.IsValid() ? SelectedEvent->GetTrack().ToSharedPtr() : nullptr;
+	EnumerateFilteredTracks(QuickFindVm->GetFilterConfigurator(), PriorityTrack, [&Params, &BestMatchEvent](TSharedPtr<const FBaseTimingTrack> Track)
 	{
 		if (!Track->IsVisible())
 		{
@@ -5168,7 +5169,8 @@ void STimingView::FindNextEvent()
 	FTimingEventSearchParameters Params(StartTime, std::numeric_limits<double>::max(), ETimingEventSearchFlags::StopAtFirstMatch, EventFilter);
 	Params.FilterExecutor = QuickFindVm->GetFilterConfigurator();
 
-	EnumerateFilteredTracks(QuickFindVm->GetFilterConfigurator(), [&Params, &BestMatchEvent](TSharedPtr<FBaseTimingTrack>& Track)
+	TSharedPtr<const FBaseTimingTrack> PriorityTrack = SelectedEvent.IsValid() ? SelectedEvent->GetTrack().ToSharedPtr() : nullptr;
+	EnumerateFilteredTracks(QuickFindVm->GetFilterConfigurator(), PriorityTrack, [&Params, &BestMatchEvent](TSharedPtr<const FBaseTimingTrack> Track)
 	{
 		if (!Track->IsVisible())
 		{
@@ -5295,12 +5297,32 @@ void STimingView::PopulateTimerNameSuggestionList(const FString& Text, TArray<FS
 
 ////////////////////////////////////////////////////////////////////////////////////////////////////
 
-void STimingView::EnumerateFilteredTracks(TSharedPtr<Insights::FFilterConfigurator> InFilterConfigurator, EnumerateFilteredTracksCallback Callback)
+void STimingView::EnumerateFilteredTracks(TSharedPtr<Insights::FFilterConfigurator> InFilterConfigurator, TSharedPtr<const FBaseTimingTrack> PriorityTrack, EnumerateFilteredTracksCallback Callback)
 {
 	Insights::FFilterContext FilterContext;
 	FilterContext.AddFilterData(static_cast<int32>(EFilterField::TrackName), FString());
+
+	// Call the callback for the PriorityTrack first if it passes the filters.
+	// This is an optimization because in many cases, the next/prev event will be on the same track
+	// and searching this one first will potentially avoid searching all events on other tracks.
+	uint64 SkipId = std::numeric_limits<uint64>::max();
+	if (PriorityTrack.IsValid())
+	{
+		SkipId = PriorityTrack->GetId();
+		FilterContext.SetFilterData(static_cast<int32>(EFilterField::TrackName), PriorityTrack->GetName());
+		if (InFilterConfigurator->ApplyFilters(FilterContext))
+		{
+			Callback(PriorityTrack);
+		}
+	}
+
 	for (auto& Entry : AllTracks)
 	{
+		if (Entry.Value->GetId() == SkipId)
+		{
+			continue;
+		}
+
 		FilterContext.SetFilterData(static_cast<int32>(EFilterField::TrackName), Entry.Value->GetName());
 		if (InFilterConfigurator->ApplyFilters(FilterContext))
 		{
diff --git a/Engine/Source/Developer/TraceInsights/Private/Insights/Widgets/STimingView.h b/Engine/Source/Developer/TraceInsights/Private/Insights/Widgets/STimingView.h
--- a/Engine/Source/Developer/TraceInsights/Private/Insights/Widgets/STimingView.h
+++ b/Engine/Source/Developer/TraceInsights/Private/Insights/Widgets/STimingView.h
@@ -522,8 +522,8 @@ protected:
 	void PopulateTrackSuggestionList(const FString& Text, TArray<FString>& OutSuggestions);
 	void PopulateTimerNameSuggestionList(const FString& Text, TArray<FString>& OutSuggestions);
 
-	typedef TFunctionRef<void(TSharedPtr<FBaseTimingTrack>& Track)> EnumerateFilteredTracksCallback;
-	void EnumerateFilteredTracks(TSharedPtr<Insights::FFilterConfigurator> FilterConfigurator, EnumerateFilteredTracksCallback Callback);
+	typedef TFunctionRef<void(TSharedPtr<const FBaseTimingTrack> Track)> EnumerateFilteredTracksCallback;
+	void EnumerateFilteredTracks(TSharedPtr<Insights::FFilterConfigurator> FilterConfigurator, TSharedPtr<const FBaseTimingTrack> PriorityTrack, EnumerateFilteredTracksCallback Callback);
 
 	ETraceFrameType GetFrameTypeToSnapTo();
 	
