[Insights]
  - Timing View: Added coloring "by [REDACTED_PERSON]" for CPU thread tracks.
  - Timing View: Improved display of the selected event info (in bottom-right corner of the view). Also added the source file of selected timer, when available.
  - Minor spell check fixes.

diff --git a/Engine/Source/Developer/TraceInsights/Private/Insights/TimingProfilerManager.cpp b/Engine/Source/Developer/TraceInsights/Private/Insights/TimingProfilerManager.cpp
--- a/Engine/Source/Developer/TraceInsights/Private/Insights/TimingProfilerManager.cpp
+++ b/Engine/Source/Developer/TraceInsights/Private/Insights/TimingProfilerManager.cpp
@@ -721,14 +721,14 @@ bool FTimingProfilerManager::Exec(const TCHAR* Cmd, FOutputDevice& Ar)
 				}
 				else if (Token.StartsWith(StartTimeToken))
 				{
-					// Default: -startTime=-infinte
+					// Default: -startTime=-infinite
 					// Example: -startTime=10.0
 					Token.RightChopInline(UE_ARRAY_COUNT(StartTimeToken) - 1);
 					Params.IntervalStartTime = atof(TCHAR_TO_ANSI(*Token));
 				}
 				else if (Token.StartsWith(EndTimeToken))
 				{
-					// Default: -endTime=+infinte
+					// Default: -endTime=+infinite
 					// Example: -endTime=20.0
 					Token.RightChopInline(UE_ARRAY_COUNT(EndTimeToken) - 1);
 					Params.IntervalEndTime = atof(TCHAR_TO_ANSI(*Token));
@@ -809,14 +809,14 @@ bool FTimingProfilerManager::Exec(const TCHAR* Cmd, FOutputDevice& Ar)
 				}
 				else if (Token.StartsWith(StartTimeToken))
 				{
-					// Default: -startTime=-infinte
+					// Default: -startTime=-infinite
 					// Example: -startTime=10.0
 					Token.RightChopInline(UE_ARRAY_COUNT(StartTimeToken) - 1);
 					Params.IntervalStartTime = atof(TCHAR_TO_ANSI(*Token));
 				}
 				else if (Token.StartsWith(EndTimeToken))
 				{
-					// Default: -endTime=+infinte
+					// Default: -endTime=+infinite
 					// Example: -endTime=20.0
 					Token.RightChopInline(UE_ARRAY_COUNT(EndTimeToken) - 1);
 					Params.IntervalEndTime = atof(TCHAR_TO_ANSI(*Token));
diff --git a/Engine/Source/Developer/TraceInsights/Private/Insights/TimingProfilerManager.h b/Engine/Source/Developer/TraceInsights/Private/Insights/TimingProfilerManager.h
--- a/Engine/Source/Developer/TraceInsights/Private/Insights/TimingProfilerManager.h
+++ b/Engine/Source/Developer/TraceInsights/Private/Insights/TimingProfilerManager.h
@@ -20,6 +20,7 @@ namespace Insights
 	{
 		ByTimerName,
 		ByTimerId,
+		BySourceFile,
 		ByDuration,
 
 		Count
diff --git a/Engine/Source/Developer/TraceInsights/Private/Insights/ViewModels/ThreadTimingTrack.cpp b/Engine/Source/Developer/TraceInsights/Private/Insights/ViewModels/ThreadTimingTrack.cpp
--- a/Engine/Source/Developer/TraceInsights/Private/Insights/ViewModels/ThreadTimingTrack.cpp
+++ b/Engine/Source/Developer/TraceInsights/Private/Insights/ViewModels/ThreadTimingTrack.cpp
@@ -302,6 +302,9 @@ static void AddTimingEventToBuilder(ITimingEventsTrackDrawStateBuilder& Builder,
 		case Insights::ETimingEventsColoringMode::ByTimerId:
 			EventColor = FTimingEvent::ComputeEventColor(Timer->Id);
 			break;
+		case Insights::ETimingEventsColoringMode::BySourceFile:
+			EventColor = FTimingEvent::ComputeEventColor(Timer->File);
+			break;
 		case Insights::ETimingEventsColoringMode::ByDuration:
 		{
 			const double EventDuration = EventEndTime - EventStartTime;
@@ -310,7 +313,7 @@ static void AddTimingEventToBuilder(ITimingEventsTrackDrawStateBuilder& Builder,
 						 (EventDuration >= 0.0001)   ? 0xFF338833 : // green:  [100us .. 1ms)
 						 (EventDuration >= 0.00001)  ? 0xFF338888 : // cyan:   [10us .. 100us)
 						 (EventDuration >= 0.000001) ? 0xFF333388 : // blue:   [1us .. 10us)
-						                               0xFF888888;  // grey:   < 1us
+						                               0xFF888888;  // gray:   < 1us
 			break;
 		}
 		default:
@@ -1200,12 +1203,26 @@ void FThreadTimingTrack::PostDraw(const ITimingTrackDrawContext& Context) const
 		const TraceServices::FTimingProfilerTimer* Timer = TimerReader->GetTimer(SelectedEvent.GetTimerIndex());
 		if (Timer != nullptr)
 		{
-			FString Str = FString::Printf(TEXT("%s (Incl.: %s, Excl.: %s)"),
-				Timer->Name,
-				*TimeUtils::FormatTimeAuto(SelectedEvent.GetDuration()),
-				*TimeUtils::FormatTimeAuto(SelectedEvent.GetExclusiveTime()));
+			FString TimerName(Timer->Name);
 
-			DrawSelectedEventInfo(Str, Context.GetViewport(), Context.GetDrawContext(), Helper.GetWhiteBrush(), Helper.GetEventFont());
+			const double SelectedEventDuration = SelectedEvent.GetDuration();
+			TStringBuilder<1024> StringBuilder;
+			StringBuilder.Appendf(TEXT("  Incl.: %s"), *TimeUtils::FormatTimeAuto(SelectedEventDuration, 2));
+			if (SelectedEventDuration != std::numeric_limits<double>::infinity())
+			{
+				StringBuilder.Appendf(TEXT("   Excl.: %s"), *TimeUtils::FormatTimeAuto(SelectedEvent.GetExclusiveTime(), 2));
+			}
+			FString StatsText(StringBuilder.ToView());
+
+			if (Timer->File)
+			{
+				FString SourceFile(Timer->File);
+				DrawSelectedEventInfoEx(StatsText, TimerName, FPaths::GetCleanFilename(SourceFile), Context.GetViewport(), Context.GetDrawContext(), Helper.GetWhiteBrush(), Helper.GetEventFont());
+			}
+			else
+			{
+				DrawSelectedEventInfoEx(StatsText, TimerName, FString(), Context.GetViewport(), Context.GetDrawContext(), Helper.GetWhiteBrush(), Helper.GetEventFont());
+			}
 		}
 	}
 }
@@ -1238,13 +1255,15 @@ void FThreadTimingTrack::InitTooltip(FTooltipDrawState& InOutTooltip, const ITim
 		FString TimerName = (Timer != nullptr) ? Timer->Name : TEXT("N/A");
 		InOutTooltip.AddTitle(TimerName);
 
+		const double TooltipEventDuration = TooltipEvent.GetDuration();
+
 		if (ParentTimingEvent.IsValid() && TooltipEvent.GetDepth() > 0)
 		{
 			Timer = TimerReader->GetTimer(ParentTimingEvent->GetTimerIndex());
 			const TCHAR* ParentTimerName = (Timer != nullptr) ? Timer->Name : TEXT("N/A");
 			FNumberFormattingOptions FormattingOptions;
 			FormattingOptions.MaximumFractionalDigits = 2;
-			const FString ValueStr = FString::Printf(TEXT("%s %s"), *FText::AsPercent(TooltipEvent.GetDuration() / ParentTimingEvent->GetDuration(), &FormattingOptions).ToString(), ParentTimerName);
+			const FString ValueStr = FString::Printf(TEXT("%s %s"), *FText::AsPercent(TooltipEventDuration / ParentTimingEvent->GetDuration(), &FormattingOptions).ToString(), ParentTimerName);
 			InOutTooltip.AddNameValueTextLine(TEXT("% of Parent:"), ValueStr);
 		}
 
@@ -1254,15 +1273,15 @@ void FThreadTimingTrack::InitTooltip(FTooltipDrawState& InOutTooltip, const ITim
 			const TCHAR* RootTimerName = (Timer != nullptr) ? Timer->Name : TEXT("N/A");
 			FNumberFormattingOptions FormattingOptions;
 			FormattingOptions.MaximumFractionalDigits = 2;
-			const FString ValueStr = FString::Printf(TEXT("%s %s"), *FText::AsPercent(TooltipEvent.GetDuration() / RootTimingEvent->GetDuration(), &FormattingOptions).ToString(), RootTimerName);
+			const FString ValueStr = FString::Printf(TEXT("%s %s"), *FText::AsPercent(TooltipEventDuration / RootTimingEvent->GetDuration(), &FormattingOptions).ToString(), RootTimerName);
 			InOutTooltip.AddNameValueTextLine(TEXT("% of Root:"), ValueStr);
 		}
 
-		InOutTooltip.AddNameValueTextLine(TEXT("Inclusive Time:"), TimeUtils::FormatTimeAuto(TooltipEvent.GetDuration()));
+		InOutTooltip.AddNameValueTextLine(TEXT("Inclusive Time:"), TimeUtils::FormatTimeAuto(TooltipEventDuration));
 
-		if (TooltipEvent.GetDuration() > 0.0)
+		if (TooltipEventDuration > 0.0 && TooltipEventDuration != std::numeric_limits<double>::infinity())
 		{
-			const double ExclusiveTimePercent = TooltipEvent.GetExclusiveTime() / TooltipEvent.GetDuration();
+			const double ExclusiveTimePercent = TooltipEvent.GetExclusiveTime() / TooltipEventDuration;
 			FNumberFormattingOptions FormattingOptions;
 			FormattingOptions.MaximumFractionalDigits = 2;
 			const FString ExclStr = FString::Printf(TEXT("%s (%s)"), *TimeUtils::FormatTimeAuto(TooltipEvent.GetExclusiveTime()), *FText::AsPercent(ExclusiveTimePercent, &FormattingOptions).ToString());
diff --git a/Engine/Source/Developer/TraceInsights/Private/Insights/ViewModels/TimingEventsTrack.cpp b/Engine/Source/Developer/TraceInsights/Private/Insights/ViewModels/TimingEventsTrack.cpp
--- a/Engine/Source/Developer/TraceInsights/Private/Insights/ViewModels/TimingEventsTrack.cpp
+++ b/Engine/Source/Developer/TraceInsights/Private/Insights/ViewModels/TimingEventsTrack.cpp
@@ -480,4 +480,48 @@ void FTimingEventsTrack::DrawSelectedEventInfo(const FString& InText, const FTim
 
 ////////////////////////////////////////////////////////////////////////////////////////////////////
 
+void FTimingEventsTrack::DrawSelectedEventInfoEx(const FString& InText, const FString& InLeftText, const FString& InTopText, const FTimingTrackViewport& Viewport, const FDrawContext& DrawContext, const FSlateBrush* WhiteBrush, const FSlateFontInfo& Font) const
+{
+	const TSharedRef<FSlateFontMeasure> FontMeasureService = FSlateApplication::Get().GetRenderer()->GetFontMeasureService();
+	const float FontScale = DrawContext.Geometry.Scale;
+
+	const FLinearColor BackgroundColor(0.05f, 0.05f, 0.05f, 1.0f);
+	const FLinearColor TextColor(0.7f, 0.7f, 0.7f, 1.0f);
+	const FLinearColor LeftTextColor(0.9f, 0.9f, 0.5f, 1.0f);
+	const FLinearColor TopTextColor(0.3f, 0.3f, 0.3f, 1.0f);
+
+	const FVector2D Size = FontMeasureService->Measure(InText, Font, FontScale) / FontScale;
+	const float W = static_cast<float>(Size.X);
+	const float H = static_cast<float>(Size.Y);
+	const float X = Viewport.GetWidth() - W - 23.0f;
+	const float Y = Viewport.GetPosY() + Viewport.GetHeight() - H - 18.0f;
+	DrawContext.DrawBox(DrawContext.LayerId, X - 8.0f, Y - 2.0f, W + 16.0f, H + 4.0f, WhiteBrush, BackgroundColor);
+	DrawContext.DrawText(DrawContext.LayerId + 1, X, Y, InText, Font, TextColor);
+
+	if (InLeftText.Len() > 0)
+	{
+		const FVector2D Size2 = FontMeasureService->Measure(InLeftText, Font, FontScale) / FontScale;
+		const float W2 = static_cast<float>(Size2.X);
+		const float H2 = static_cast<float>(Size2.Y);
+		const float X2 = X - W2 - 4.0f;
+		const float Y2 = Y;
+		DrawContext.DrawBox(DrawContext.LayerId, X2 - 8.0f, Y2 - 2.0f, W2 + 16.0f, H2 + 4.0f, WhiteBrush, BackgroundColor);
+		DrawContext.DrawText(DrawContext.LayerId + 1, X2, Y2, InLeftText, Font, LeftTextColor);
+	}
+	if (InTopText.Len() > 0)
+	{
+		const FVector2D Size2 = FontMeasureService->Measure(InTopText, Font, FontScale) / FontScale;
+		const float W2 = static_cast<float>(Size2.X);
+		const float H2 = static_cast<float>(Size2.Y);
+		const float X2 = Viewport.GetWidth() - W2 - 23.0f;
+		const float Y2 = Y - H2 - 4.0f;
+		DrawContext.DrawBox(DrawContext.LayerId, X2 - 8.0f, Y2 - 2.0f, W2 + 16.0f, H2 + 4.0f, WhiteBrush, BackgroundColor);
+		DrawContext.DrawText(DrawContext.LayerId + 1, X2, Y2, InTopText, Font, TopTextColor);
+	}
+
+	DrawContext.LayerId += 2;
+}
+
+////////////////////////////////////////////////////////////////////////////////////////////////////
+
 #undef LOCTEXT_NAMESPACE
diff --git a/Engine/Source/Developer/TraceInsights/Private/Insights/Widgets/STimingView.cpp b/Engine/Source/Developer/TraceInsights/Private/Insights/Widgets/STimingView.cpp
--- a/Engine/Source/Developer/TraceInsights/Private/Insights/Widgets/STimingView.cpp
+++ b/Engine/Source/Developer/TraceInsights/Private/Insights/Widgets/STimingView.cpp
@@ -84,7 +84,7 @@
 
 #define INSIGHTS_ACTIVATE_BENCHMARK 0
 
-// start auto generated ids from a big number (MSB set to 1) to avoid collisions with ids for gpu/cpu tracks based on 32bit timeline index
+// start auto generated ids from a big number (MSB set to 1) to avoid collisions with ids for GPU/CPU tracks based on 32bit timeline index
 uint64 FBaseTimingTrack::IdGenerator = (1ULL << 63);
 
 uint32 STimingView::TimingViewId = 0;
@@ -3884,7 +3884,7 @@ void STimingView::SelectTimingEvent(const TSharedPtr<const ITimingEvent> InEvent
 			if (bBringEventIntoViewVertically)
 			{
 				bBringSelectedEventIntoViewVerticallyOnNextTick = true;
-				// We need the layout to be callculated in one frame, no animations, otherwise the event might not be in view at the end of the animation.
+				// We need the layout to be calculated in one frame, no animations, otherwise the event might not be in view at the end of the animation.
 				Viewport.AddDirtyFlags(ETimingTrackViewportDirtyFlags::VLayoutChanged);
 			}
 		}
@@ -4181,13 +4181,13 @@ TSharedRef<SWidget> STimingView::MakeCompactAutoScrollOptionsMenu()
 			.Value_Lambda([this] { return AutoScrollFrameAlignment; })
 			+ SSegmentedControl<int32>::Slot(-1)
 			.Text(LOCTEXT("None", "None"))
-			.ToolTip(LOCTEXT("AutoScrollNoFrameAlignment_Tooltip", "Disables the frame alignment (when autoscrolling)."))
+			.ToolTip(LOCTEXT("AutoScrollNoFrameAlignment_Tooltip", "Disables the frame alignment (when auto-scrolling)."))
 			+ SSegmentedControl<int32>::Slot((int32)TraceFrameType_Game)
 			.Text(LOCTEXT("Game", "Game"))
-			.ToolTip(LOCTEXT("AutoScrollNoFrameAlignment_Tooltip", "Disables the frame alignment (when autoscrolling)."))
+			.ToolTip(LOCTEXT("AutoScrollNoFrameAlignment_Tooltip", "Disables the frame alignment (when auto-scrolling)."))
 			+ SSegmentedControl<int32>::Slot((int32)TraceFrameType_Rendering)
 			.Text(LOCTEXT("Rendering", "Rendering"))
-			.ToolTip(LOCTEXT("AutoScrollAlignWithGameFrames_Tooltip", "Aligns the viewport's center position with the start time of a Game frame (when autoscrolling)."))
+			.ToolTip(LOCTEXT("AutoScrollAlignWithGameFrames_Tooltip", "Aligns the viewport's center position with the start time of a Game frame (when auto-scrolling)."))
 		);
 
 		CreateCompactMenuLine(MenuBuilder,
@@ -4214,13 +4214,13 @@ TSharedRef<SWidget> STimingView::MakeCompactAutoScrollOptionsMenu()
 				.Value_Lambda([this] { return AutoScrollViewportOffsetPercent; })
 				+ SSegmentedControl<double>::Slot(-0.1)
 				.Text(LOCTEXT("AutoScrollViewportOffset-10", "-10%"))
-				.ToolTip(LOCTEXT("AutoScrollViewportOffset-10_Tooltip", "Sets the viewport offset to -10% (i.e. backward) of the viewport's width (when autoscrolling).\nAvoids flickering as the end of the session will be outside of the viewport."))
+				.ToolTip(LOCTEXT("AutoScrollViewportOffset-10_Tooltip", "Sets the viewport offset to -10% (i.e. backward) of the viewport's width (when auto-scrolling).\nAvoids flickering as the end of the session will be outside of the viewport."))
 				+ SSegmentedControl<double>::Slot(0.0)
 				.Text(LOCTEXT("AutoScrollViewportOffset0", "0"))
-				.ToolTip(LOCTEXT("AutoScrollViewportOffset0_Tooltip", "Sets the viewport offset to 0 (when autoscrolling).\nThe right side of the viewport will correspond to the current session time."))
+				.ToolTip(LOCTEXT("AutoScrollViewportOffset0_Tooltip", "Sets the viewport offset to 0 (when auto-scrolling).\nThe right side of the viewport will correspond to the current session time."))
 				+ SSegmentedControl<double>::Slot(+0.1)
 				.Text(LOCTEXT("AutoScrollViewportOffset+10", "+10%"))
-				.ToolTip(LOCTEXT("AutoScrollViewportOffset+10_Tooltip", "Sets the viewport offset to +10% (i.e. forward) of the viewport's width (when autoscrolling).\nAllows 10% empty space on the right side of the viewport."))
+				.ToolTip(LOCTEXT("AutoScrollViewportOffset+10_Tooltip", "Sets the viewport offset to +10% (i.e. forward) of the viewport's width (when auto-scrolling).\nAllows 10% empty space on the right side of the viewport."))
 			]
 		);
 
@@ -4276,7 +4276,7 @@ TSharedRef<SWidget> STimingView::MakeAutoScrollOptionsMenu()
 	{
 		MenuBuilder.AddMenuEntry(
 			LOCTEXT("AutoScrollNoFrameAlignment", "None"),
-			LOCTEXT("AutoScrollNoFrameAlignment_Tooltip", "Disables the frame alignment (when autoscrolling)."),
+			LOCTEXT("AutoScrollNoFrameAlignment_Tooltip", "Disables the frame alignment (when auto-scrolling)."),
 			FSlateIcon(),
 			FUIAction(FExecuteAction::CreateSP(this, &STimingView::SetAutoScrollFrameAlignment, -1),
 				FCanExecuteAction(),
@@ -4287,7 +4287,7 @@ TSharedRef<SWidget> STimingView::MakeAutoScrollOptionsMenu()
 
 		MenuBuilder.AddMenuEntry(
 			LOCTEXT("AutoScrollAlignWithGameFrames", "Game Frames"),
-			LOCTEXT("AutoScrollAlignWithGameFrames_Tooltip", "Aligns the viewport's center position with the start time of a Game frame (when autoscrolling)."),
+			LOCTEXT("AutoScrollAlignWithGameFrames_Tooltip", "Aligns the viewport's center position with the start time of a Game frame (when auto-scrolling)."),
 			FSlateIcon(),
 			FUIAction(FExecuteAction::CreateSP(this, &STimingView::SetAutoScrollFrameAlignment, (int32)TraceFrameType_Game),
 				FCanExecuteAction(),
@@ -4298,7 +4298,7 @@ TSharedRef<SWidget> STimingView::MakeAutoScrollOptionsMenu()
 
 		MenuBuilder.AddMenuEntry(
 			LOCTEXT("AutoScrollAlignWithRenderingFrames", "Rendering Frames"),
-			LOCTEXT("AutoScrollAlignWithRenderingFrames_Tooltip", "Aligns the viewport's center position with the start time of a Rendering frame (when autoscrolling)."),
+			LOCTEXT("AutoScrollAlignWithRenderingFrames_Tooltip", "Aligns the viewport's center position with the start time of a Rendering frame (when auto-scrolling)."),
 			FSlateIcon(),
 			FUIAction(FExecuteAction::CreateSP(this, &STimingView::SetAutoScrollFrameAlignment, (int32)TraceFrameType_Rendering),
 				FCanExecuteAction(),
@@ -4313,7 +4313,7 @@ TSharedRef<SWidget> STimingView::MakeAutoScrollOptionsMenu()
 	{
 		MenuBuilder.AddMenuEntry(
 			LOCTEXT("AutoScrollViewportOffset-10", "-10%"),
-			LOCTEXT("AutoScrollViewportOffset-10_Tooltip", "Sets the viewport offset to -10% (i.e. backward) of the viewport's width (when autoscrolling).\nAvoids flickering as the end of the session will be outside of the viewport."),
+			LOCTEXT("AutoScrollViewportOffset-10_Tooltip", "Sets the viewport offset to -10% (i.e. backward) of the viewport's width (when auto-scrolling).\nAvoids flickering as the end of the session will be outside of the viewport."),
 			FSlateIcon(),
 			FUIAction(FExecuteAction::CreateSP(this, &STimingView::SetAutoScrollViewportOffset, -0.1),
 				FCanExecuteAction(),
@@ -4324,7 +4324,7 @@ TSharedRef<SWidget> STimingView::MakeAutoScrollOptionsMenu()
 
 		MenuBuilder.AddMenuEntry(
 			LOCTEXT("AutoScrollViewportOffset0", "0"),
-			LOCTEXT("AutoScrollViewportOffset0_Tooltip", "Sets the viewport offset to 0 (when autoscrolling).\nThe right side of the viewport will correspond to the current session time."),
+			LOCTEXT("AutoScrollViewportOffset0_Tooltip", "Sets the viewport offset to 0 (when auto-scrolling).\nThe right side of the viewport will correspond to the current session time."),
 			FSlateIcon(),
 			FUIAction(FExecuteAction::CreateSP(this, &STimingView::SetAutoScrollViewportOffset, 0.0),
 				FCanExecuteAction(),
@@ -4335,7 +4335,7 @@ TSharedRef<SWidget> STimingView::MakeAutoScrollOptionsMenu()
 
 		MenuBuilder.AddMenuEntry(
 			LOCTEXT("AutoScrollViewportOffset+10", "+10%"),
-			LOCTEXT("AutoScrollViewportOffset+10_Tooltip", "Sets the viewport offset to +10% (i.e. forward) of the viewport's width (when autoscrolling).\nAllows 10% empty space on the right side of the viewport."),
+			LOCTEXT("AutoScrollViewportOffset+10_Tooltip", "Sets the viewport offset to +10% (i.e. forward) of the viewport's width (when auto-scrolling).\nAllows 10% empty space on the right side of the viewport."),
 			FSlateIcon(),
 			FUIAction(FExecuteAction::CreateSP(this, &STimingView::SetAutoScrollViewportOffset, +0.1),
 				FCanExecuteAction(),
@@ -4365,7 +4365,7 @@ TSharedRef<SWidget> STimingView::MakeAutoScrollOptionsMenu()
 				.OnTextChanged_Lambda([this](const FText& InText) { SetAutoScrollViewportOffset(atof(TCHAR_TO_ANSI(*InText.ToString())) * 0.01); })
 			],
 			NAME_None,
-			LOCTEXT("AutoScrollViewportOffsetCustom_Tooltip", "Sets a custom value for the viewport offset as percent from viewport's width (when autoscrolling)."),
+			LOCTEXT("AutoScrollViewportOffsetCustom_Tooltip", "Sets a custom value for the viewport offset as percent from viewport's width (when auto-scrolling)."),
 			EUserInterfaceActionType::RadioButton
 		);
 	}
@@ -4835,9 +4835,20 @@ void STimingView::CreateCpuThreadTrackColoringModeMenu(FMenuBuilder& MenuBuilder
 			NAME_None,
 			EUserInterfaceActionType::RadioButton
 		);
+		MenuBuilder.AddMenuEntry(
+			LOCTEXT("CpuThreadTrackColoringMode_BySourceFile", "By Source File"),
+			LOCTEXT("CpuThreadTrackColoringMode_BySourceFile_Desc", "Assign a color to CPU/GPU timing events based on their source file."),
+			FSlateIcon(),
+			FUIAction(
+				FExecuteAction::CreateSP(this, &STimingView::SetCpuThreadTrackColoringMode, Insights::ETimingEventsColoringMode::BySourceFile),
+				FCanExecuteAction(),
+				FIsActionChecked::CreateSP(this, &STimingView::CheckCpuThreadTrackColoringMode, Insights::ETimingEventsColoringMode::BySourceFile)),
+			NAME_None,
+			EUserInterfaceActionType::RadioButton
+		);
 		MenuBuilder.AddMenuEntry(
 			LOCTEXT("CpuThreadTrackColoringMode_ByDuration", "By Duration"),
-			LOCTEXT("CpuThreadTrackColoringMode_ByDuration_Desc", "Assign a color to CPU/GPU timing events based on their duration (inclusive time).\n\t≥ 10ms : red\n\t≥ 1ms : yellow\n\t≥ 100μs : green\n\t≥ 10μs : cyan\n\t≥ 1μs : blue\n\t< 1μs : grey"),
+			LOCTEXT("CpuThreadTrackColoringMode_ByDuration_Desc", "Assign a color to CPU/GPU timing events based on their duration (inclusive time).\n\t≥ 10ms : red\n\t≥ 1ms : yellow\n\t≥ 100μs : green\n\t≥ 10μs : cyan\n\t≥ 1μs : blue\n\t< 1μs : gray"),
 			FSlateIcon(),
 			FUIAction(
 				FExecuteAction::CreateSP(this, &STimingView::SetCpuThreadTrackColoringMode, Insights::ETimingEventsColoringMode::ByDuration),
diff --git a/Engine/Source/Developer/TraceInsights/Public/Insights/ViewModels/TimingEventsTrack.h b/Engine/Source/Developer/TraceInsights/Public/Insights/ViewModels/TimingEventsTrack.h
--- a/Engine/Source/Developer/TraceInsights/Public/Insights/ViewModels/TimingEventsTrack.h
+++ b/Engine/Source/Developer/TraceInsights/Public/Insights/ViewModels/TimingEventsTrack.h
@@ -85,6 +85,7 @@ protected:
 	void DrawMarkers(const ITimingTrackDrawContext& Context, float LineY, float LineH) const;
 
 	void DrawSelectedEventInfo(const FString& InText, const FTimingTrackViewport& Viewport, const FDrawContext& DrawContext, const FSlateBrush* WhiteBrush, const FSlateFontInfo& Font) const;
+	void DrawSelectedEventInfoEx(const FString& InText, const FString& InLeftText, const FString& InTopText, const FTimingTrackViewport& Viewport, const FDrawContext& DrawContext, const FSlateBrush* WhiteBrush, const FSlateFontInfo& Font) const;
 
 	int32 GetHeaderBackgroundLayerId(const ITimingTrackDrawContext& Context) const;
 	int32 GetHeaderTextLayerId(const ITimingTrackDrawContext& Context) const;
@@ -93,7 +94,7 @@ protected:
 
 	virtual bool HasCustomFilter() const { return false; }
 
-	/* Can be overriden to force a max depth for the track. */
+	/* Can be overridden to force a max depth for the track. */
 	virtual int32 GetMaxDepth() const { return -1; }
 
 private:
