[Insights]
  - Session Browser: Added support for multiple selection. Can now delete multiple trace files at once.
  - Session Browser: Added "Copy Trace Id" in context menu. Copies the unique id of the selected trace session.
  - Session Browser: Added "Copy Full Path" in context menu. Copies the full path of the selected trace file.
  - Session Browser: Added "Open Containing Folder" in context menu. Opens the containing folder of the selected trace file.

diff --git a/Engine/Source/Developer/TraceInsights/Private/Insights/StoreService/StoreBrowser.cpp b/Engine/Source/Developer/TraceInsights/Private/Insights/StoreService/StoreBrowser.cpp
--- a/Engine/Source/Developer/TraceInsights/Private/Insights/StoreService/StoreBrowser.cpp
+++ b/Engine/Source/Developer/TraceInsights/Private/Insights/StoreService/StoreBrowser.cpp
@@ -292,7 +292,6 @@ void FStoreBrowser::UpdateTraces()
 							FStoreBrowserTraceInfo& Trace = *TracePtr;
 
 							Trace.TraceId = TraceId;
-							Trace.TraceIndex = TraceIndex;
 
 							const FUtf8StringView Utf8NameView = TraceInfo->GetName();
 							Trace.Name = FString(Utf8NameView);
@@ -303,7 +302,7 @@ void FStoreBrowser::UpdateTraces()
 							}
 							else
 							{
-								// Fallback for older versions of UTS which didn't write uri
+								// Fallback for older versions of UTS which didn't write URI.
 								Trace.Uri = FPaths::SetExtension(FPaths::Combine(StoreDirectory, Trace.Name), TEXT(".utrace"));
 								FPaths::MakePlatformFilename(Trace.Uri);
 							}
diff --git a/Engine/Source/Developer/TraceInsights/Private/Insights/StoreService/StoreBrowser.h b/Engine/Source/Developer/TraceInsights/Private/Insights/StoreService/StoreBrowser.h
--- a/Engine/Source/Developer/TraceInsights/Private/Insights/StoreService/StoreBrowser.h
+++ b/Engine/Source/Developer/TraceInsights/Private/Insights/StoreService/StoreBrowser.h
@@ -25,7 +25,6 @@ namespace Insights
 struct FStoreBrowserTraceInfo
 {
 	uint32 TraceId = 0;
-	int32 TraceIndex = -1;
 
 	uint64 ChangeSerial = 0;
 
diff --git a/Engine/Source/Developer/TraceInsights/Private/Insights/Widgets/SStartPageWindow.cpp b/Engine/Source/Developer/TraceInsights/Private/Insights/Widgets/SStartPageWindow.cpp
--- a/Engine/Source/Developer/TraceInsights/Private/Insights/Widgets/SStartPageWindow.cpp
+++ b/Engine/Source/Developer/TraceInsights/Private/Insights/Widgets/SStartPageWindow.cpp
@@ -7,6 +7,7 @@
 #include "Framework/MetaData/DriverMetaData.h"
 #include "Framework/MultiBox/MultiBoxBuilder.h"
 #include "HAL/FileManager.h"
+#include "HAL/PlatformApplicationMisc.h"
 #include "HAL/PlatformProcess.h"
 #include "Internationalization/Text.h"
 #include "IPAddress.h"
@@ -328,7 +329,7 @@ public:
 		UE_LOG(TraceInsights, Verbose, TEXT("[TraceStore] Renamed utrace file (\"%s\")."), *NewTraceFile);
 		Trace.Name = FText::FromString(NewTraceName);
 		Trace.Uri = FText::FromString(NewTraceFile);
-		Trace.bWasJustRenamed = true; // cannot be open until its TraceId is not updated
+		Trace.TraceId = FTraceViewModel::InvalidTraceId; // cannot be open until its TraceId is updated
 		Trace.ChangeSerial = 0; // to force update
 
 		const FString CacheFile = FPaths::ChangeExtension(TraceFile, TEXT("ucache"));
@@ -352,7 +353,7 @@ public:
 		TSharedPtr<FTraceViewModel> TracePin = WeakTrace.Pin();
 		if (TracePin.IsValid())
 		{
-			const FString TraceIdStr = FString::Printf(TEXT("%d (0x%08X)"), TracePin->TraceIndex, TracePin->TraceId);
+			const FString TraceIdStr = FString::Printf(TEXT("0x%X"), TracePin->TraceId);
 			return FText::FromString(TraceIdStr);
 		}
 		else
@@ -379,7 +380,7 @@ public:
 		TSharedPtr<FTraceViewModel> TracePin = WeakTrace.Pin();
 		if (TracePin.IsValid())
 		{
-			if (TracePin->bWasJustRenamed)
+			if (TracePin->TraceId == FTraceViewModel::InvalidTraceId)
 			{
 				return FSlateColor(EStyleColor::White25);
 			}
@@ -1001,7 +1002,7 @@ void STraceStoreWindow::Construct(const FArguments& InArgs)
 			]
 		]
 
-		// Overlay for fake splashscreen.
+		// Overlay for fake splash-screen.
 		+ SOverlay::Slot()
 		.HAlign(HAlign_Fill)
 		.VAlign(VAlign_Fill)
@@ -1116,8 +1117,7 @@ TSharedRef<SWidget> STraceStoreWindow::ConstructFiltersToolbar()
 				.IsEnabled_Lambda([this]() { return TraceViewModels.Num() > 0; })
 				.OnTextChanged(this, &STraceStoreWindow::FilterByNameSearchBox_OnTextChanged)
 				.DelayChangeNotificationsWhileTyping(true)
-			]
-		);
+			]);
 
 		// Filter by Platform
 		ToolbarBuilder.AddComboButton(
@@ -1126,8 +1126,7 @@ TSharedRef<SWidget> STraceStoreWindow::ConstructFiltersToolbar()
 			LOCTEXT("FilterByPlatformText", "Platform"),
 			LOCTEXT("FilterByPlatformToolTip", "Filters the list of trace sessions by platform."),
 			FSlateIcon(FInsightsStyle::GetStyleSetName(), "Icons.Filter.ToolBar"),
-			false
-		);
+			false);
 
 		// Filter by AppName
 		ToolbarBuilder.AddComboButton(
@@ -1136,8 +1135,7 @@ TSharedRef<SWidget> STraceStoreWindow::ConstructFiltersToolbar()
 			LOCTEXT("FilterByAppNameText", "App Name"),
 			LOCTEXT("FilterByAppNameToolTip", "Filters the list of trace sessions by application name."),
 			FSlateIcon(FInsightsStyle::GetStyleSetName(), "Icons.Filter.ToolBar"),
-			false
-		);
+			false);
 
 		// Filter by Build Config
 		ToolbarBuilder.AddComboButton(
@@ -1146,8 +1144,7 @@ TSharedRef<SWidget> STraceStoreWindow::ConstructFiltersToolbar()
 			LOCTEXT("FilterByBuildConfigText", "Config"),
 			LOCTEXT("FilterByBuildConfigToolTip", "Filters the list of trace sessions by build configuration."),
 			FSlateIcon(FInsightsStyle::GetStyleSetName(), "Icons.Filter.ToolBar"),
-			false
-		);
+			false);
 
 		// Filter by Build Target
 		ToolbarBuilder.AddComboButton(
@@ -1156,8 +1153,7 @@ TSharedRef<SWidget> STraceStoreWindow::ConstructFiltersToolbar()
 			LOCTEXT("FilterByBuildTargetText", "Target"),
 			LOCTEXT("FilterByBuildTargetToolTip", "Filters the list of trace sessions by build target."),
 			FSlateIcon(FInsightsStyle::GetStyleSetName(), "Icons.Filter.ToolBar"),
-			false
-		);
+			false);
 
 		// Filter by Branch
 		ToolbarBuilder.AddComboButton(
@@ -1166,8 +1162,7 @@ TSharedRef<SWidget> STraceStoreWindow::ConstructFiltersToolbar()
 			LOCTEXT("FilterByBranchText", "Branch"),
 			LOCTEXT("FilterByBranchToolTip", "Filters the list of trace sessions by branch."),
 			FSlateIcon(FInsightsStyle::GetStyleSetName(), "Icons.Filter.ToolBar"),
-			false
-		);
+			false);
 	}
 	ToolbarBuilder.EndSection();
 
@@ -1183,8 +1178,7 @@ TSharedRef<SWidget> STraceStoreWindow::ConstructFiltersToolbar()
 			[
 				SNew(STextBlock)
 				.Text(this, &STraceStoreWindow::GetFilterStatsText)
-			]
-		);
+			]);
 	}
 	RightSideToolbarBuilder.EndSection();
 
@@ -1214,7 +1208,7 @@ TSharedRef<SWidget> STraceStoreWindow::ConstructSessionsPanel()
 	TSharedRef<SWidget> Widget = SAssignNew(TraceListView, SListView<TSharedPtr<FTraceViewModel>>)
 		.IsFocusable(true)
 		.ItemHeight(20.0f)
-		.SelectionMode(ESelectionMode::Single)
+		.SelectionMode(ESelectionMode::Multi)
 		.OnSelectionChanged(this, &STraceStoreWindow::TraceList_OnSelectionChanged)
 		.OnMouseButtonDoubleClick(this, &STraceStoreWindow::TraceList_OnMouseButtonDoubleClick)
 		.ListItemsSource(&FilteredTraceViewModels)
@@ -1496,8 +1490,7 @@ FReply STraceDirectoryItem::OnModifyStore()
 			FSlateApplication::Get().FindBestParentWindowHandleForDialogs(AsShared()),
 			Title,
 			CurrentStoreDirectory,
-			SelectedDirectory
-		);
+			SelectedDirectory);
 
 		const bool bIsWatchDir = Window->WatchDirectoriesModel.FindByPredicate([&](const auto& Directory){ return FPathViews::Equals(SelectedDirectory, Directory->Path);}) != nullptr;
 		const bool bIsCurrentStoreDir = FPathViews::Equals(SelectedDirectory, CurrentStoreDirectory);
@@ -1974,35 +1967,61 @@ TSharedPtr<SWidget> STraceStoreWindow::TraceList_GetMenuContent()
 
 	MenuBuilder.BeginSection("Misc");
 	{
-		FUIAction Action_Rename
-		(
-			FExecuteAction::CreateSP(this, &STraceStoreWindow::RenameTraceFile),
-			FCanExecuteAction::CreateSP(this, &STraceStoreWindow::CanEditTraceFile)
-		);
-		MenuBuilder.AddMenuEntry
-		(
-			LOCTEXT("ContextMenu_Rename", "Rename... \t\tF2"),
-			LOCTEXT("ContextMenu_Rename_ToolTip", "Rename the selected utrace file."),
-			FSlateIcon(FInsightsStyle::GetStyleSetName(), "Icons.Rename"),
-			Action_Rename,
+		{
+			FMenuEntryParams MenuEntry;
+			MenuEntry.LabelOverride = LOCTEXT("ContextMenu_Rename", "Rename...");
+			MenuEntry.InputBindingOverride = LOCTEXT("ContextMenu_Rename_InputBinding", "F2");
+			MenuEntry.ToolTipOverride = LOCTEXT("ContextMenu_Rename_ToolTip", "Starts renaming of the selected trace file.");
+			MenuEntry.IconOverride = FSlateIcon(FInsightsStyle::GetStyleSetName(), "Icons.Rename");
+			MenuEntry.DirectActions = FUIAction(
+				FExecuteAction::CreateSP(this, &STraceStoreWindow::RenameSelectedTrace),
+				FCanExecuteAction::CreateSP(this, &STraceStoreWindow::CanRenameSelectedTrace));
+			MenuEntry.UserInterfaceActionType = EUserInterfaceActionType::Button;
+			MenuBuilder.AddMenuEntry(MenuEntry);
+		}
+		{
+			FMenuEntryParams MenuEntry;
+			MenuEntry.LabelOverride = LOCTEXT("ContextMenu_Delete", "Delete");
+			MenuEntry.InputBindingOverride = LOCTEXT("ContextMenu_Delete_InputBinding", "Del");
+			MenuEntry.ToolTipOverride = LOCTEXT("ContextMenu_Delete_ToolTip", "Deletes the selected trace files.");
+			MenuEntry.IconOverride = FSlateIcon(FInsightsStyle::GetStyleSetName(), "Icons.Delete");
+			MenuEntry.DirectActions = FUIAction(
+				FExecuteAction::CreateSP(this, &STraceStoreWindow::DeleteSelectedTraces),
+				FCanExecuteAction::CreateSP(this, &STraceStoreWindow::CanDeleteSelectedTraces));
+			MenuEntry.UserInterfaceActionType = EUserInterfaceActionType::Button;
+			MenuBuilder.AddMenuEntry(MenuEntry);
+		}
+		MenuBuilder.AddSeparator();
+
+		MenuBuilder.AddMenuEntry(
+			LOCTEXT("ContextMenu_CopyTraceId", "Copy Trace Id"),
+			LOCTEXT("ContextMenu_CopyTraceId_ToolTip", "Copies the unique id of the selected trace session."),
+			FSlateIcon(FAppStyle::Get().GetStyleSetName(), "GenericCommands.Copy"),
+			FUIAction(
+				FExecuteAction::CreateSP(this, &STraceStoreWindow::CopyTraceId),
+				FCanExecuteAction::CreateSP(this, &STraceStoreWindow::CanCopyTraceId)),
 			NAME_None,
-			EUserInterfaceActionType::Button
-		);
+			EUserInterfaceActionType::Button);
 
-		FUIAction Action_Delete
-		(
-			FExecuteAction::CreateSP(this, &STraceStoreWindow::DeleteTraceFile),
-			FCanExecuteAction::CreateSP(this, &STraceStoreWindow::CanEditTraceFile)
-		);
-		MenuBuilder.AddMenuEntry
-		(
-			LOCTEXT("ContextMenu_Delete", "Delete \t\t\tDel"),
-			LOCTEXT("ContextMenu_Delete_ToolTip", "Delete the selected utrace file."),
-			FSlateIcon(FInsightsStyle::GetStyleSetName(), "Icons.Delete"),
-			Action_Delete,
+		MenuBuilder.AddMenuEntry(
+			LOCTEXT("ContextMenu_CopyUri", "Copy Full Path"),
+			LOCTEXT("ContextMenu_CopyUri_ToolTip", "Copies the full path of the selected trace file."),
+			FSlateIcon(FAppStyle::Get().GetStyleSetName(), "GenericCommands.Copy"),
+			FUIAction(
+				FExecuteAction::CreateSP(this, &STraceStoreWindow::CopyFullPath),
+				FCanExecuteAction::CreateSP(this, &STraceStoreWindow::CanCopyFullPath)),
 			NAME_None,
-			EUserInterfaceActionType::Button
-		);
+			EUserInterfaceActionType::Button);
+
+		MenuBuilder.AddMenuEntry(
+			LOCTEXT("ContextMenu_OpenContainingFolder", "Open Containing Folder"),
+			LOCTEXT("ContextMenu_OpenContainingFolder_ToolTip", "Opens the containing folder of the selected trace file."),
+			FSlateIcon(FAppStyle::Get().GetStyleSetName(), "Icons.FolderOpen"),
+			FUIAction(
+				FExecuteAction::CreateSP(this, &STraceStoreWindow::OpenContainingFolder),
+				FCanExecuteAction::CreateSP(this, &STraceStoreWindow::CanOpenContainingFolder)),
+			NAME_None,
+			EUserInterfaceActionType::Button);
 	}
 
 	return MenuBuilder.MakeWidget();
@@ -2010,27 +2029,32 @@ TSharedPtr<SWidget> STraceStoreWindow::TraceList_GetMenuContent()
 
 ////////////////////////////////////////////////////////////////////////////////////////////////////
 
-bool STraceStoreWindow::CanEditTraceFile() const
+bool STraceStoreWindow::CanRenameSelectedTrace() const
 {
-	return CanChangeStoreSettings()
-		&& SelectedTrace.IsValid()
-		&& !SelectedTrace->bIsLive
-		&& !SelectedTrace->bWasJustRenamed;
+	if (!CanChangeStoreSettings())
+	{
+		return false;
+	}
+	TSharedPtr<FTraceViewModel> SelectedTrace = GetSingleSelectedTrace();
+	return SelectedTrace.IsValid()
+		&& SelectedTrace->TraceId != FTraceViewModel::InvalidTraceId
+		&& !SelectedTrace->bIsLive;
 }
 
 ////////////////////////////////////////////////////////////////////////////////////////////////////
 
-void STraceStoreWindow::RenameTraceFile()
+void STraceStoreWindow::RenameSelectedTrace()
 {
-	if (!CanEditTraceFile())
+	FSlateApplication::Get().CloseToolTip();
+
+	if (!CanRenameSelectedTrace())
 	{
 		return;
 	}
 
+	TSharedPtr<FTraceViewModel> SelectedTrace = GetSingleSelectedTrace();
 	SelectedTrace->bIsRenaming = true;
 
-	FSlateApplication::Get().CloseToolTip();
-
 	TSharedPtr<SEditableTextBox> RenameTextBox = SelectedTrace->RenameTextBox.Pin();
 	if (RenameTextBox.IsValid())
 	{
@@ -2040,23 +2064,58 @@ void STraceStoreWindow::RenameTraceFile()
 
 ////////////////////////////////////////////////////////////////////////////////////////////////////
 
-void STraceStoreWindow::DeleteTraceFile()
+bool STraceStoreWindow::CanDeleteSelectedTraces() const
 {
-	if (!CanEditTraceFile())
+	if (!CanChangeStoreSettings() ||
+		TraceListView->GetNumItemsSelected() == 0)
 	{
-		return;
+		return false;
+	}
+	TArray<TSharedPtr<FTraceViewModel>> SelectedTraces = TraceListView->GetSelectedItems();
+	for (const TSharedPtr<FTraceViewModel>& SelectedTrace : SelectedTraces)
+	{
+		if (SelectedTrace->TraceId == FTraceViewModel::InvalidTraceId ||
+			SelectedTrace->bIsLive)
+		{
+			return false;
+		}
 	}
+	return true;
+}
 
+////////////////////////////////////////////////////////////////////////////////////////////////////
+
+void STraceStoreWindow::DeleteSelectedTraces()
+{
 	FSlateApplication::Get().CloseToolTip();
 
-	FString TraceFile = SelectedTrace->Uri.ToString();
+	if (!CanDeleteSelectedTraces())
+	{
+		return;
+	}
+
+	TArray<TSharedPtr<FTraceViewModel>> TracesToDelete = TraceListView->GetSelectedItems();
+	if (TracesToDelete.Num() == 0)
+	{
+		return;
+	}
 
 	if (bIsDeleteTraceConfirmWindowVisible)
 	{
 		//TODO: Make a custom OkCancel modal dialog. See FSlateApplication::Get().AddModalWindow(..).
 		FText Title = LOCTEXT("ConfirmToDeleteTraceFile_Title", "Unreal Insights");
-		FText ConfirmMessage = FText::Format(LOCTEXT("ConfirmToDeleteTraceFile", "You are about to delete the utrace file:\n{0}\n\nPress Ok to continue."),
-			FText::FromString(TraceFile));
+		TStringBuilder<2048> TraceFilesToDelete;
+		for (int32 TraceIndex = 0; TraceIndex < TracesToDelete.Num() && TraceIndex < 3; ++TraceIndex)
+		{
+			TraceFilesToDelete.Append(TracesToDelete[TraceIndex]->Uri.ToString());
+			TraceFilesToDelete.Append(TEXT("\n"));
+		}
+		if (TracesToDelete.Num() > 3)
+		{
+			TraceFilesToDelete.Append(TEXT("...\n"));
+		}
+		FText ConfirmMessage = FText::Format(LOCTEXT("ConfirmToDeleteTraceFile", "You are about to delete {0} trace {0}|plural(one=file,other=files):\n\n{1}\nPress OK to continue."),
+			TracesToDelete.Num(), FText::FromStringView(TraceFilesToDelete.ToView()));
 		EAppReturnType::Type OkToDelete = FMessageDialog::Open(EAppMsgType::OkCancel, ConfirmMessage, Title);
 		if (OkToDelete == EAppReturnType::Cancel)
 		{
@@ -2064,35 +2123,176 @@ void STraceStoreWindow::DeleteTraceFile()
 		}
 	}
 
-	FString TraceName = SelectedTrace->Name.ToString();
-	UE_LOG(TraceInsights, Log, TEXT("[TraceStore] Deleting \"%s\"..."), *TraceName);
-
-	if (FPaths::FileExists(TraceFile) && IFileManager::Get().Delete(*TraceFile))
+	// Find an unselected item (close to last selected one).
+	int32 TraceIndexToSelect = -1;
+	for (int32 TraceIndex = 0; TraceIndex < TracesToDelete.Num(); ++TraceIndex)
 	{
-		UE_LOG(TraceInsights, Verbose, TEXT("[TraceStore] Deleted utrace file (\"%s\")."), *TraceFile);
-
-		FString CacheFile = FPaths::ChangeExtension(TraceFile, TEXT("ucache"));
-		if (FPaths::FileExists(CacheFile))
+		FTraceViewModel* TraceVM = TracesToDelete[TraceIndex].Get();
+		int32 FilteredTraceIndex = FilteredTraceViewModels.IndexOfByPredicate([TraceVM](const TSharedPtr<FTraceViewModel>& VM) { return VM.Get() == TraceVM; });
+		if (FilteredTraceIndex + 1 >= 0 &&
+			FilteredTraceIndex + 1 < FilteredTraceViewModels.Num() &&
+			!TraceListView->IsItemSelected(FilteredTraceViewModels[FilteredTraceIndex + 1]))
 		{
-			if (IFileManager::Get().Delete(*CacheFile))
+			if (FilteredTraceIndex + 1> TraceIndexToSelect)
 			{
-				UE_LOG(TraceInsights, Verbose, TEXT("[TraceStore] Deleted ucache file (\"%s\")."), *CacheFile);
+				TraceIndexToSelect = FilteredTraceIndex + 1;
 			}
 		}
+		else
+		if (FilteredTraceIndex - 1 >= 0 &&
+			FilteredTraceIndex - 1 < FilteredTraceViewModels.Num() &&
+			!TraceListView->IsItemSelected(FilteredTraceViewModels[FilteredTraceIndex - 1]))
+		{
+			if (FilteredTraceIndex - 1 > TraceIndexToSelect)
+			{
+				TraceIndexToSelect = FilteredTraceIndex - 1;
+			}
+		}
+	}
+	TSharedPtr<FTraceViewModel> TraceToSelect = (TraceIndexToSelect >= 0) ? FilteredTraceViewModels[TraceIndexToSelect] : nullptr;
 
-		TraceViewModels.Remove(SelectedTrace);
-		TraceViewModelMap.Remove(SelectedTrace->TraceId);
-		OnTraceListChanged();
+	// Delete traces.
+	int32 NumDeletedTraces = 0;
+	for (int32 TraceIndex = 0; TraceIndex < TracesToDelete.Num(); ++TraceIndex)
+	{
+		const TSharedPtr<FTraceViewModel>& TraceToDelete = TracesToDelete[TraceIndex];
+		if (DeleteTrace(TraceToDelete))
+		{
+			++NumDeletedTraces;
+
+			TraceListView->SetItemSelection(TraceToDelete, false);
 
-		FText Message = FText::Format(LOCTEXT("DeleteSuccessFmt", "Deleted \"{0}\"."), FText::FromString(TraceName));
+			FilteredTraceViewModels.Remove(TraceToDelete);
+
+			TraceViewModels.Remove(TraceToDelete);
+			TraceViewModelMap.Remove(TraceToDelete->TraceId);
+		}
+	}
+
+	if (NumDeletedTraces == TracesToDelete.Num())
+	{
+		FText Message = FText::Format(LOCTEXT("DeleteSuccessFmt", "Successfully deleted {0} trace {0}|plural(one=file,other=files)."), NumDeletedTraces);
 		ShowSuccessMessage(Message);
+
+		// Set new selection.
+		if (TraceToSelect.IsValid())
+		{
+			TraceListView->SetItemSelection(TraceToSelect, true);
+			bIsUserSelectedTrace = true;
+		}
 	}
 	else
 	{
-		UE_LOG(TraceInsights, Warning, TEXT("[TraceStore] Failed to delete utrace file (\"%s\")!"), *TraceFile);
+		FText Message = FText::Format(LOCTEXT("FailedToDeleteAllTracesFmt", "Deleted {0} trace {0}|plural(one=file,other=files). Failed to delete {1} trace {1}|plural(one=file,other=files)!"),
+			NumDeletedTraces, TracesToDelete.Num() - NumDeletedTraces);
+		ShowFailMessage(Message);
+	}
+
+	OnTraceListChanged();
+}
+
+////////////////////////////////////////////////////////////////////////////////////////////////////
+
+bool STraceStoreWindow::DeleteTrace(const TSharedPtr<FTraceViewModel>& TraceToDelete)
+{
+	FString TraceName = TraceToDelete->Name.ToString();
+	UE_LOG(TraceInsights, Log, TEXT("[TraceStore] Deleting \"%s\"..."), *TraceName);
+
+	if (TraceToDelete->bIsLive)
+	{
+		FText Message = FText::Format(LOCTEXT("CannotDeleteLiveTraceFmt", "Cannot delete a live trace (\"{0}\")!"), FText::FromString(TraceName));
+		ShowFailMessage(Message);
+		return false;
+	}
+
+	FString TraceFile = TraceToDelete->Uri.ToString();
+	if (!FPaths::FileExists(TraceFile) || !IFileManager::Get().Delete(*TraceFile))
+	{
+		UE_LOG(TraceInsights, Warning, TEXT("[TraceStore] Failed to delete trace file (\"%s\")!"), *TraceFile);
 
 		FText Message = FText::Format(LOCTEXT("DeleteFailFmt", "Failed to delete \"{0}\"!"), FText::FromString(TraceName));
 		ShowFailMessage(Message);
+		return false;
+	}
+
+	UE_LOG(TraceInsights, Verbose, TEXT("[TraceStore] Deleted utrace file (\"%s\")."), *TraceFile);
+
+	FString CacheFile = FPaths::ChangeExtension(TraceFile, TEXT("ucache"));
+	if (FPaths::FileExists(CacheFile))
+	{
+		if (IFileManager::Get().Delete(*CacheFile))
+		{
+			UE_LOG(TraceInsights, Verbose, TEXT("[TraceStore] Deleted ucache file (\"%s\")."), *CacheFile);
+		}
+	}
+
+	return true;
+}
+
+////////////////////////////////////////////////////////////////////////////////////////////////////
+
+bool STraceStoreWindow::CanCopyTraceId() const
+{
+	TSharedPtr<FTraceViewModel> SelectedTrace = GetSingleSelectedTrace();
+	return SelectedTrace.IsValid()
+		&& SelectedTrace->TraceId != FTraceViewModel::InvalidTraceId;
+}
+
+////////////////////////////////////////////////////////////////////////////////////////////////////
+
+void STraceStoreWindow::CopyTraceId()
+{
+	if (CanCopyTraceId())
+	{
+		TSharedPtr<FTraceViewModel> SelectedTrace = GetSingleSelectedTrace();
+		FString ClipboardText = FString::Printf(TEXT("0x%X"), SelectedTrace->TraceId);
+		FPlatformApplicationMisc::ClipboardCopy(*ClipboardText);
+	}
+}
+
+////////////////////////////////////////////////////////////////////////////////////////////////////
+
+bool STraceStoreWindow::CanCopyFullPath() const
+{
+	TSharedPtr<FTraceViewModel> SelectedTrace = GetSingleSelectedTrace();
+	return SelectedTrace.IsValid()
+		&& SelectedTrace->TraceId != FTraceViewModel::InvalidTraceId;
+}
+
+////////////////////////////////////////////////////////////////////////////////////////////////////
+
+void STraceStoreWindow::CopyFullPath()
+{
+	if (CanCopyFullPath())
+	{
+		TSharedPtr<FTraceViewModel> SelectedTrace = GetSingleSelectedTrace();
+		FPlatformApplicationMisc::ClipboardCopy(*SelectedTrace->Uri.ToString());
+	}
+}
+
+////////////////////////////////////////////////////////////////////////////////////////////////////
+
+bool STraceStoreWindow::CanOpenContainingFolder() const
+{
+	if (!CanChangeStoreSettings())
+	{
+		return false;
+	}
+	TSharedPtr<FTraceViewModel> SelectedTrace = GetSingleSelectedTrace();
+	return SelectedTrace.IsValid()
+		&& SelectedTrace->TraceId != FTraceViewModel::InvalidTraceId;
+}
+
+////////////////////////////////////////////////////////////////////////////////////////////////////
+
+void STraceStoreWindow::OpenContainingFolder()
+{
+	FSlateApplication::Get().CloseToolTip();
+
+	if (CanOpenContainingFolder())
+	{
+		TSharedPtr<FTraceViewModel> SelectedTrace = GetSingleSelectedTrace();
+		FPlatformProcess::ExploreFolder(*SelectedTrace->Uri.ToString());
 	}
 }
 
@@ -2271,7 +2471,7 @@ void STraceStoreWindow::RefreshTraceList()
 		if (NewTracesChangeSerial != TracesChangeSerial || bSettingsChanged)
 		{
 			TracesChangeSerial = NewTracesChangeSerial;
-			//UE_LOG(TraceInsights, Log, TEXT("[TraceStore] Synching the trace list with StoreBrowser..."));
+			//UE_LOG(TraceInsights, Log, TEXT("[TraceStore] Syncing the trace list with StoreBrowser..."));
 
 			const TArray<TSharedPtr<Insights::FStoreBrowserTraceInfo>>& InTraces = StoreBrowser->GetTraces();
 			const TMap<uint32, TSharedPtr<Insights::FStoreBrowserTraceInfo>>& InTraceMap = StoreBrowser->GetTraceMap();
@@ -2359,10 +2559,8 @@ bool STraceStoreWindow::IsConnected() const
 void STraceStoreWindow::UpdateTrace(FTraceViewModel& InOutTrace, const Insights::FStoreBrowserTraceInfo& InSourceTrace)
 {
 	InOutTrace.TraceId = InSourceTrace.TraceId;
-	InOutTrace.bWasJustRenamed = false;
 
 	InOutTrace.ChangeSerial = InSourceTrace.ChangeSerial;
-	InOutTrace.TraceIndex = InSourceTrace.TraceIndex;
 
 	InOutTrace.Name = FText::FromString(InSourceTrace.Name);
 	InOutTrace.Uri = FText::FromString(InSourceTrace.Uri);
@@ -2437,22 +2635,26 @@ void STraceStoreWindow::UpdateTraceListView()
 		return;
 	}
 
-	TSharedPtr<FTraceViewModel> NewSelectedTrace;
-	if (bIsUserSelectedTrace && SelectedTrace)
+	TArray<TSharedPtr<FTraceViewModel>> NewSelectedTraces;
+	if (bIsUserSelectedTrace)
 	{
-		// Identify the previously selected trace (if still available) to ensure selection remains unchanged.
-		TSharedPtr<FTraceViewModel>* FoundNewTrace = TraceViewModelMap.Find(SelectedTrace->TraceId);
-		if (!FoundNewTrace)
+		// Identify the previously selected traces (if still available) to ensure selection remains unchanged.
+		TArray<TSharedPtr<FTraceViewModel>> SelectedTraces = TraceListView->GetSelectedItems();
+		for (const TSharedPtr<FTraceViewModel>& SelectedTrace : SelectedTraces)
 		{
-			FoundNewTrace = TraceViewModels.FindByPredicate([this](const TSharedPtr<FTraceViewModel>& Trace) { return Trace->Uri.EqualTo(SelectedTrace->Uri); });
-		}
-		if (!FoundNewTrace)
-		{
-			FoundNewTrace = TraceViewModels.FindByPredicate([this](const TSharedPtr<FTraceViewModel>& Trace) { return Trace->Name.EqualTo(SelectedTrace->Name); });
-		}
-		if (FoundNewTrace)
-		{
-			NewSelectedTrace = *FoundNewTrace;
+			TSharedPtr<FTraceViewModel>* FoundNewTrace = TraceViewModelMap.Find(SelectedTrace->TraceId);
+			if (!FoundNewTrace)
+			{
+				FoundNewTrace = TraceViewModels.FindByPredicate([SelectedTrace](const TSharedPtr<FTraceViewModel>& Trace) { return Trace->Uri.EqualTo(SelectedTrace->Uri); });
+			}
+			if (!FoundNewTrace)
+			{
+				FoundNewTrace = TraceViewModels.FindByPredicate([SelectedTrace](const TSharedPtr<FTraceViewModel>& Trace) { return Trace->Name.EqualTo(SelectedTrace->Name); });
+			}
+			if (FoundNewTrace)
+			{
+				NewSelectedTraces.Add(*FoundNewTrace);
+			}
 		}
 	}
 
@@ -2462,20 +2664,20 @@ void STraceStoreWindow::UpdateTraceListView()
 	TraceListView->RebuildList();
 
 	// If no selection...
-	if (!NewSelectedTrace.IsValid() && FilteredTraceViewModels.Num() > 0)
+	if (NewSelectedTraces.Num() == 0 && FilteredTraceViewModels.Num() > 0)
 	{
 		if ((SortColumn == TraceStoreColumns::Date && SortMode == EColumnSortMode::Ascending) ||
 			(SortColumn == TraceStoreColumns::Status && SortMode == EColumnSortMode::Ascending))
 		{
 			// Auto select the last (newest) trace.
-			NewSelectedTrace = FilteredTraceViewModels.Last();
+			NewSelectedTraces.Add(FilteredTraceViewModels.Last());
 			DistanceFromTop = 1.0;
 			DistanceFromBottom = 0.0; // scroll to bottom
 		}
 		else
 		{
 			// Auto select the first trace.
-			NewSelectedTrace = FilteredTraceViewModels[0];
+			NewSelectedTraces.Add(FilteredTraceViewModels[0]);
 			DistanceFromTop = 0.0; // scroll to top
 			DistanceFromBottom = 1.0;
 		}
@@ -2491,22 +2693,28 @@ void STraceStoreWindow::UpdateTraceListView()
 	}
 
 	// Restore selection.
-	if (NewSelectedTrace.IsValid())
+	if (NewSelectedTraces.Num() > 0)
 	{
 		TraceListView->ClearSelection();
-		TraceListView->SetItemSelection(NewSelectedTrace, true);
+		TraceListView->SetItemSelection(NewSelectedTraces, true);
 	}
 }
 
 ////////////////////////////////////////////////////////////////////////////////////////////////////
 
+TSharedPtr<FTraceViewModel> STraceStoreWindow::GetSingleSelectedTrace() const
+{
+	return (TraceListView->GetNumItemsSelected() == 1) ? TraceListView->GetSelectedItems()[0] : nullptr;
+}
+
+////////////////////////////////////////////////////////////////////////////////////////////////////
+
 void STraceStoreWindow::TraceList_OnSelectionChanged(TSharedPtr<FTraceViewModel> TraceSession, ESelectInfo::Type SelectInfo)
 {
 	if (SelectInfo != ESelectInfo::Direct)
 	{
 		bIsUserSelectedTrace = true;
 	}
-	SelectedTrace = TraceSession;
 }
 
 ////////////////////////////////////////////////////////////////////////////////////////////////////
@@ -2682,12 +2890,12 @@ FReply STraceStoreWindow::OnKeyDown(const FGeometry& MyGeometry, const FKeyEvent
 	}
 	else if (InKeyEvent.GetKey() == EKeys::F2)
 	{
-		RenameTraceFile();
+		RenameSelectedTrace();
 		return FReply::Handled();
 	}
 	else if (InKeyEvent.GetKey() == EKeys::Delete)
 	{
-		DeleteTraceFile();
+		DeleteSelectedTraces();
 		return FReply::Handled();
 	}
 
@@ -2757,15 +2965,16 @@ FReply STraceStoreWindow::OnDrop(const FGeometry& MyGeometry, const FDragDropEve
 
 bool STraceStoreWindow::Open_IsEnabled() const
 {
-	return FilteredTraceViewModels.Num() > 0
-		&& SelectedTrace.IsValid()
-		&& !SelectedTrace->bWasJustRenamed;
+	TSharedPtr<FTraceViewModel> SelectedTrace = GetSingleSelectedTrace();
+	return SelectedTrace.IsValid()
+		&& SelectedTrace->TraceId != FTraceViewModel::InvalidTraceId;
 }
 
 ////////////////////////////////////////////////////////////////////////////////////////////////////
 
 FReply STraceStoreWindow::Open_OnClicked()
 {
+	TSharedPtr<FTraceViewModel> SelectedTrace = GetSingleSelectedTrace();
 	OpenTraceSession(SelectedTrace);
 	return FReply::Handled();
 }
@@ -2803,7 +3012,8 @@ void STraceStoreWindow::OpenTraceFile(const FString& InTraceFile)
 
 void STraceStoreWindow::OpenTraceSession(TSharedPtr<FTraceViewModel> InTraceSession)
 {
-	if (InTraceSession.IsValid() && !InTraceSession->bWasJustRenamed)
+	if (InTraceSession.IsValid() &&
+		InTraceSession->TraceId != FTraceViewModel::InvalidTraceId)
 	{
 		OpenTraceSession(InTraceSession->TraceId);
 	}
@@ -2865,28 +3075,25 @@ TSharedRef<SWidget> STraceStoreWindow::MakeTraceListMenu()
 	{
 		MenuBuilder.AddMenuEntry(
 			LOCTEXT("OpenFileButtonLabel", "Open Trace File..."),
-			LOCTEXT("OpenFileButtonTooltip", "Start analysis for a specified trace file."),
+			LOCTEXT("OpenFileButtonTooltip", "Starts analysis for a specified trace file."),
 			FSlateIcon(FAppStyle::Get().GetStyleSetName(), "Icons.FolderOpen"),
 			FUIAction(FExecuteAction::CreateSP(this, &STraceStoreWindow::OpenTraceFile)),
 			NAME_None,
-			EUserInterfaceActionType::Button
-		);
+			EUserInterfaceActionType::Button);
 		MenuBuilder.AddMenuEntry(
 			LOCTEXT("ImportTableButtonLabel", "Import Table..."),
-			LOCTEXT("ImportTableButtonTooltip", "Open .csv or .tsv file."),
+			LOCTEXT("ImportTableButtonTooltip", "Opens .csv or .tsv file."),
 			FSlateIcon(FAppStyle::Get().GetStyleSetName(), "Icons.FolderOpen"),
 			FUIAction(FExecuteAction::CreateLambda([]{ Insights::FTableImportTool::Get()->StartImportProcess(); })),
 			NAME_None,
-			EUserInterfaceActionType::Button
-		);
+			EUserInterfaceActionType::Button);
 		MenuBuilder.AddMenuEntry(
 			LOCTEXT("DiffTablesButtonLabel", "Diff Tables..."),
-			LOCTEXT("DiffTablesButtonTooltip", "Open two table files in diff mode."),
+			LOCTEXT("DiffTablesButtonTooltip", "Opens two table files in diff mode."),
 			FSlateIcon(FAppStyle::Get().GetStyleSetName(), "Icons.FolderOpen"),
 			FUIAction(FExecuteAction::CreateLambda([]{ Insights::FTableImportTool::Get()->StartDiffProcess(); })),
 			NAME_None,
-			EUserInterfaceActionType::Button
-		);
+			EUserInterfaceActionType::Button);
 	}
 
 	MenuBuilder.EndSection();
@@ -2916,8 +3123,7 @@ TSharedRef<SWidget> STraceStoreWindow::MakeTraceListMenu()
 				FSlateIcon(),
 				FUIAction(FExecuteAction::CreateSP(this, &STraceStoreWindow::OpenTraceSession, Trace.TraceId)),
 				NAME_None,
-				EUserInterfaceActionType::Button
-			);
+				EUserInterfaceActionType::Button);
 		}
 	}
 	MenuBuilder.EndSection();
@@ -2935,8 +3141,7 @@ TSharedRef<SWidget> STraceStoreWindow::MakeTraceListMenu()
 					FGlobalTabmanager::Get()->TryInvokeTab(FInsightsManagerTabs::AutomationWindowTabId);
 				})),
 			NAME_None,
-					EUserInterfaceActionType::Button
-					);
+			EUserInterfaceActionType::Button);
 	}
 
 	// Enable Automation Tests Option.
@@ -2957,8 +3162,7 @@ TSharedRef<SWidget> STraceStoreWindow::MakeTraceListMenu()
 			FSlateIcon(FInsightsStyle::GetStyleSetName(), "Icons.TestAutomation"),
 			ToogleAutomationTestsAction,
 			NAME_None,
-			EUserInterfaceActionType::ToggleButton
-		);
+			EUserInterfaceActionType::ToggleButton);
 	}
 
 	// Enable Debug Tools Option.
@@ -2979,8 +3183,7 @@ TSharedRef<SWidget> STraceStoreWindow::MakeTraceListMenu()
 			FSlateIcon(FInsightsStyle::GetStyleSetName(), "Icons.Debug"),
 			ToogleDebugToolsAction,
 			NAME_None,
-			EUserInterfaceActionType::ToggleButton
-		);
+			EUserInterfaceActionType::ToggleButton);
 	}
 
 #if !UE_BUILD_SHIPPING
@@ -2998,8 +3201,7 @@ TSharedRef<SWidget> STraceStoreWindow::MakeTraceListMenu()
 			FSlateIcon(FInsightsStyle::GetStyleSetName(), "Icons.Test"),
 			OpenStarshipSuiteAction,
 			NAME_None,
-			EUserInterfaceActionType::Button
-		);
+			EUserInterfaceActionType::Button);
 	}
 #endif // !UE_BUILD_SHIPPING
 
@@ -3259,8 +3461,7 @@ FReply STraceStoreWindow::AddWatchDir_Clicked()
 			FSlateApplication::Get().FindBestParentWindowHandleForDialogs(AsShared()),
 			Title,
 			CurrentStoreDirectory,
-			SelectedDirectory
-		);
+			SelectedDirectory);
 
 		if (bHasSelected && !FPathViews::Equals(SelectedDirectory, CurrentStoreDirectory))
 		{
@@ -3365,8 +3566,7 @@ void TTraceSetFilter<TSetType>::BuildMenu(FMenuBuilder& InMenuBuilder, STraceSto
 			FSlateIcon(),
 			Action,
 			NAME_None,
-			EUserInterfaceActionType::ToggleButton
-		);
+			EUserInterfaceActionType::ToggleButton);
 	}
 
 	InMenuBuilder.AddSeparator();
@@ -3413,8 +3613,7 @@ void TTraceSetFilter<TSetType>::BuildMenu(FMenuBuilder& InMenuBuilder, STraceSto
 			FSlateIcon(),
 			Action,
 			NAME_None,
-			EUserInterfaceActionType::ToggleButton
-		);
+			EUserInterfaceActionType::ToggleButton);
 	}
 }
 
diff --git a/Engine/Source/Developer/TraceInsights/Private/Insights/Widgets/SStartPageWindow.h b/Engine/Source/Developer/TraceInsights/Private/Insights/Widgets/SStartPageWindow.h
--- a/Engine/Source/Developer/TraceInsights/Private/Insights/Widgets/SStartPageWindow.h
+++ b/Engine/Source/Developer/TraceInsights/Private/Insights/Widgets/SStartPageWindow.h
@@ -68,8 +68,9 @@ ENUM_CLASS_FLAGS(ETraceDirOperations);
 
 struct FTraceViewModel
 {
-	uint32 TraceId = 0;
-	int32 TraceIndex = -1; // debug
+	static constexpr uint32 InvalidTraceId = 0;
+
+	uint32 TraceId = InvalidTraceId;
 
 	uint64 ChangeSerial = 0;
 
@@ -91,7 +92,6 @@ struct FTraceViewModel
 
 	bool bIsMetadataUpdated = false;
 	bool bIsRenaming = false;
-	bool bWasJustRenamed = false;
 	bool bIsLive = false;
 	uint32 IpAddress = 0;
 
@@ -296,9 +296,22 @@ private:
 
 	TSharedPtr<SWidget> TraceList_GetMenuContent();
 
-	bool CanEditTraceFile() const;
-	void RenameTraceFile();
-	void DeleteTraceFile();
+	bool CanRenameSelectedTrace() const;
+	void RenameSelectedTrace();
+
+	bool CanDeleteSelectedTraces() const;
+	void DeleteSelectedTraces();
+	bool DeleteTrace(const TSharedPtr<FTraceViewModel>& TraceToDelete);
+
+	bool CanCopyTraceId() const;
+	void CopyTraceId();
+
+	bool CanCopyFullPath() const;
+	void CopyFullPath();
+
+	bool CanOpenContainingFolder() const;
+	void OpenContainingFolder();
+
 	bool HasAnyLiveTrace() const;
 
 	//////////////////////////////////////////////////
@@ -354,6 +367,8 @@ private:
 	void UpdateTrace(FTraceViewModel& InOutTrace, const Insights::FStoreBrowserTraceInfo& InSourceTrace);
 	void OnTraceListChanged();
 
+	TSharedPtr<FTraceViewModel> GetSingleSelectedTrace() const;
+
 	void TraceList_OnSelectionChanged(TSharedPtr<FTraceViewModel> InTrace, ESelectInfo::Type SelectInfo);
 	void TraceList_OnMouseButtonDoubleClick(TSharedPtr<FTraceViewModel> InTrace);
 
@@ -545,7 +560,6 @@ private:
 	TSharedPtr<STableViewBase> WatchDirsListView;
 	TSharedPtr<SListView<TSharedPtr<FTraceViewModel>>> TraceListView;
 
-	TSharedPtr<FTraceViewModel> SelectedTrace;
 	bool bIsUserSelectedTrace = false;
 
 	//////////////////////////////////////////////////
