[Insights] Improved API for creating the STimingView tooltips to accept FStringView instead of FString. This reduces the number of intermediary string copies and unnecessary memory allocations. Updated the code for creating tooltips for thread timing tracks (including the "tasks" tooltips).

diff --git a/Engine/Source/Developer/TraceInsights/Private/Insights/ViewModels/ThreadTimingTrack.cpp b/Engine/Source/Developer/TraceInsights/Private/Insights/ViewModels/ThreadTimingTrack.cpp
--- a/Engine/Source/Developer/TraceInsights/Private/Insights/ViewModels/ThreadTimingTrack.cpp
+++ b/Engine/Source/Developer/TraceInsights/Private/Insights/ViewModels/ThreadTimingTrack.cpp
@@ -1252,7 +1252,7 @@ void FThreadTimingTrack::InitTooltip(FTooltipDrawState& InOutTooltip, const ITim
 		TimingProfilerProvider.ReadTimers([&TimerReader](const TraceServices::ITimingProfilerTimerReader& Out) { TimerReader = &Out; });
 
 		const TraceServices::FTimingProfilerTimer* Timer = TimerReader->GetTimer(TooltipEvent.GetTimerIndex());
-		FString TimerName = (Timer != nullptr) ? Timer->Name : TEXT("N/A");
+		const TCHAR* TimerName = (Timer != nullptr) ? Timer->Name : TEXT("N/A");
 		InOutTooltip.AddTitle(TimerName);
 
 		const double TooltipEventDuration = TooltipEvent.GetDuration();
@@ -1264,7 +1264,7 @@ void FThreadTimingTrack::InitTooltip(FTooltipDrawState& InOutTooltip, const ITim
 			FNumberFormattingOptions FormattingOptions;
 			FormattingOptions.MaximumFractionalDigits = 2;
 			const FString ValueStr = FString::Printf(TEXT("%s %s"), *FText::AsPercent(TooltipEventDuration / ParentTimingEvent->GetDuration(), &FormattingOptions).ToString(), ParentTimerName);
-			InOutTooltip.AddNameValueTextLine(TEXT("% of Parent:"), ValueStr);
+			InOutTooltip.AddNameValueTextLine(TEXTVIEW("% of Parent:"), ValueStr);
 		}
 
 		if (RootTimingEvent.IsValid() && TooltipEvent.GetDepth() > 1)
@@ -1274,21 +1274,21 @@ void FThreadTimingTrack::InitTooltip(FTooltipDrawState& InOutTooltip, const ITim
 			FNumberFormattingOptions FormattingOptions;
 			FormattingOptions.MaximumFractionalDigits = 2;
 			const FString ValueStr = FString::Printf(TEXT("%s %s"), *FText::AsPercent(TooltipEventDuration / RootTimingEvent->GetDuration(), &FormattingOptions).ToString(), RootTimerName);
-			InOutTooltip.AddNameValueTextLine(TEXT("% of Root:"), ValueStr);
+			InOutTooltip.AddNameValueTextLine(TEXTVIEW("% of Root:"), ValueStr);
 		}
 
-		InOutTooltip.AddNameValueTextLine(TEXT("Inclusive Time:"), TimeUtils::FormatTimeAuto(TooltipEventDuration));
+		InOutTooltip.AddNameValueTextLine(TEXTVIEW("Inclusive Time:"), TimeUtils::FormatTimeAuto(TooltipEventDuration, 2));
 
 		if (TooltipEventDuration > 0.0 && TooltipEventDuration != std::numeric_limits<double>::infinity())
 		{
 			const double ExclusiveTimePercent = TooltipEvent.GetExclusiveTime() / TooltipEventDuration;
 			FNumberFormattingOptions FormattingOptions;
 			FormattingOptions.MaximumFractionalDigits = 2;
-			const FString ExclStr = FString::Printf(TEXT("%s (%s)"), *TimeUtils::FormatTimeAuto(TooltipEvent.GetExclusiveTime()), *FText::AsPercent(ExclusiveTimePercent, &FormattingOptions).ToString());
-			InOutTooltip.AddNameValueTextLine(TEXT("Exclusive Time:"), ExclStr);
+			const FString ExclStr = FString::Printf(TEXT("%s (%s)"), *TimeUtils::FormatTimeAuto(TooltipEvent.GetExclusiveTime(), 2), *FText::AsPercent(ExclusiveTimePercent, &FormattingOptions).ToString());
+			InOutTooltip.AddNameValueTextLine(TEXTVIEW("Exclusive Time:"), ExclStr);
 		}
 
-		InOutTooltip.AddNameValueTextLine(TEXT("Depth:"), FString::Printf(TEXT("%d"), TooltipEvent.GetDepth()));
+		InOutTooltip.AddNameValueTextLine(TEXTVIEW("Depth:"), FString::Printf(TEXT("%d"), TooltipEvent.GetDepth()));
 
 		TArrayView<const uint8> Metadata = TimerReader->GetMetadata(TooltipEvent.GetTimerIndex());
 		if (Metadata.Num() > 0)
@@ -1297,7 +1297,6 @@ void FThreadTimingTrack::InitTooltip(FTooltipDrawState& InOutTooltip, const ITim
 		}
 
 		// tasks
-
 		const TraceServices::ITasksProvider* TasksProvider = TraceServices::ReadTasksProvider(*Session.Get());
 		if (TasksProvider != nullptr)
 		{
@@ -1313,6 +1312,16 @@ void FThreadTimingTrack::InitTooltip(FTooltipDrawState& InOutTooltip, const ITim
 				ENamedThreads::Type ThreadInfo = (ENamedThreads::Type)Task.ThreadToExecuteOn;
 				ENamedThreads::Type ThreadIndex = ENamedThreads::GetThreadIndex(ThreadInfo);
 
+				auto FormatTaskTimestamp = [](double Timestamp) -> FString
+				{
+					return (Timestamp != TraceServices::FTaskInfo::InvalidTimestamp) ? FString::SanitizeFloat(Timestamp) : TEXT("[not set]");
+				};
+
+				auto FormatTaskTime = [](double Time) -> FString
+				{
+					return TimeUtils::FormatTimeAuto(Time, 2);
+				};
+
 				auto GetTrackName = [this](uint32 InThreadId) -> FString
 				{
 					TSharedPtr<FCpuTimingTrack> Track = SharedState.GetCpuTrack(InThreadId);
@@ -1327,65 +1336,76 @@ void FThreadTimingTrack::InitTooltip(FTooltipDrawState& InOutTooltip, const ITim
 					const TCHAR* ThreadPriStrs[] = { TEXT("Normal"), TEXT("High"), TEXT("Low") };
 					const TCHAR* ThreadPri = ensure(ThreadPriIndex >= 0 && ThreadPriIndex < 3) ? ThreadPriStrs[ThreadPriIndex] : TEXT("Unknown");
 
-					InOutTooltip.AddTextLine(FString::Printf(TEXT("%s Pri task on %s Pri worker (%s)"), TaskPri, ThreadPri, *GetTrackName(Task.StartedThreadId)), FLinearColor::Green);
+					InOutTooltip.AddTextLine(
+						FString::Printf(TEXT("%s Pri task on %s Pri worker (%s)"), TaskPri, ThreadPri, *GetTrackName(Task.StartedThreadId)),
+						FLinearColor::Green);
 				}
 				else
 				{
 					const TCHAR* QueueStr = ENamedThreads::GetQueueIndex(ThreadInfo) == ENamedThreads::MainQueue ? TEXT("Main") : TEXT("Local");
-					InOutTooltip.AddTextLine(FString::Printf(TEXT("%s (%s queue)"), *GetTrackName(Task.StartedThreadId), QueueStr), FLinearColor::Green);
+					InOutTooltip.AddTextLine(
+						FString::Printf(TEXT("%s (%s queue)"), *GetTrackName(Task.StartedThreadId), QueueStr),
+						FLinearColor::Green);
 				}
 
-				InOutTooltip.AddNameValueTextLine(TEXT("Created:"), FString::Printf(TEXT("%f on %s"), Task.CreatedTimestamp, *GetTrackName(Task.CreatedThreadId)));
+				InOutTooltip.AddNameValueTextLine(TEXTVIEW("Created:"), FString::Printf(TEXT("%s on %s"),
+					*FormatTaskTimestamp(Task.CreatedTimestamp),
+					*GetTrackName(Task.CreatedThreadId)));
 
-				InOutTooltip.AddNameValueTextLine(TEXT("Launched:"), FString::Printf(TEXT("%f (+%s) on %s"), Task.LaunchedTimestamp, *TimeUtils::FormatTimeAuto(Task.LaunchedTimestamp - Task.CreatedTimestamp), *GetTrackName(Task.LaunchedThreadId)));
+				InOutTooltip.AddNameValueTextLine(TEXTVIEW("Launched:"), FString::Printf(TEXT("%s (+%s) on %s"),
+					*FormatTaskTimestamp(Task.LaunchedTimestamp),
+					*FormatTaskTime(Task.LaunchedTimestamp - Task.CreatedTimestamp),
+					*GetTrackName(Task.LaunchedThreadId)));
 
-				InOutTooltip.AddNameValueTextLine(TEXT("Scheduled:"), FString::Printf(TEXT("%f (+%s) on %s"), Task.ScheduledTimestamp, *TimeUtils::FormatTimeAuto(Task.ScheduledTimestamp - Task.LaunchedTimestamp), *GetTrackName(Task.ScheduledThreadId)));
+				InOutTooltip.AddNameValueTextLine(TEXTVIEW("Scheduled:"), FString::Printf(TEXT("%s (+%s) on %s"),
+					*FormatTaskTimestamp(Task.ScheduledTimestamp),
+					*FormatTaskTime(Task.ScheduledTimestamp - Task.LaunchedTimestamp),
+					*GetTrackName(Task.ScheduledThreadId)));
+
+				InOutTooltip.AddNameValueTextLine(TEXTVIEW("Started:"), FString::Printf(TEXT("%s (+%s)"),
+					*FormatTaskTimestamp(Task.StartedTimestamp),
+					*FormatTaskTime(Task.StartedTimestamp - Task.ScheduledTimestamp)));
 
-				InOutTooltip.AddNameValueTextLine(TEXT("Started:"), FString::Printf(TEXT("%f (+%s)"), Task.StartedTimestamp, *TimeUtils::FormatTimeAuto(Task.StartedTimestamp - Task.ScheduledTimestamp)));
 				if (Task.FinishedTimestamp != TraceServices::FTaskInfo::InvalidTimestamp)
 				{
-					InOutTooltip.AddNameValueTextLine(TEXT("Finished:"), FString::Printf(TEXT("%f (+%s)"), Task.FinishedTimestamp, *TimeUtils::FormatTimeAuto(Task.FinishedTimestamp - Task.StartedTimestamp)));
+					InOutTooltip.AddNameValueTextLine(TEXTVIEW("Finished:"), FString::Printf(TEXT("%s (+%s)"),
+						*FormatTaskTimestamp(Task.FinishedTimestamp),
+						*FormatTaskTime(Task.FinishedTimestamp - Task.StartedTimestamp)));
 
 					if (Task.CompletedTimestamp != TraceServices::FTaskInfo::InvalidTimestamp)
 					{
-						InOutTooltip.AddNameValueTextLine(TEXT("Completed:"), FString::Printf(TEXT("%f (+%s) on %s"), Task.CompletedTimestamp, *TimeUtils::FormatTimeAuto(Task.CompletedTimestamp - Task.FinishedTimestamp), *GetTrackName(Task.CompletedThreadId)));
+						InOutTooltip.AddNameValueTextLine(TEXTVIEW("Completed:"), FString::Printf(TEXT("%s (+%s) on %s"),
+							*FormatTaskTimestamp(Task.CompletedTimestamp),
+							*FormatTaskTime(Task.CompletedTimestamp - Task.FinishedTimestamp),
+							*GetTrackName(Task.CompletedThreadId)));
 
 						if (Task.DestroyedTimestamp != TraceServices::FTaskInfo::InvalidTimestamp)
 						{
-							InOutTooltip.AddNameValueTextLine(TEXT("Destroyed:"), FString::Printf(TEXT("%f (+%s) on %s"), Task.DestroyedTimestamp, *TimeUtils::FormatTimeAuto(Task.DestroyedTimestamp - Task.CompletedTimestamp), *GetTrackName(Task.DestroyedThreadId)));
+							InOutTooltip.AddNameValueTextLine(TEXTVIEW("Destroyed:"), FString::Printf(TEXT("%s (+%s) on %s"),
+								*FormatTaskTimestamp(Task.DestroyedTimestamp),
+								*FormatTaskTime(Task.DestroyedTimestamp - Task.CompletedTimestamp),
+								*GetTrackName(Task.DestroyedThreadId)));
 						}
 					}
 				}
-				InOutTooltip.AddNameValueTextLine(TEXT("Prerequisite tasks:"), FString::Printf(TEXT("%d"), Task.Prerequisites.Num()));
-				InOutTooltip.AddNameValueTextLine(TEXT("Subsequent tasks:"), FString::Printf(TEXT("%d"), Task.Subsequents.Num()));
-				InOutTooltip.AddNameValueTextLine(TEXT("Parent tasks:"), FString::Printf(TEXT("%d"), Task.ParentTasks.Num()));
-				InOutTooltip.AddNameValueTextLine(TEXT("Nested tasks:"), FString::Printf(TEXT("%d"), Task.NestedTasks.Num()));
+
+				InOutTooltip.AddNameValueTextLine(TEXTVIEW("Prerequisite tasks:"), FString::Printf(TEXT("%d"), Task.Prerequisites.Num()));
+				InOutTooltip.AddNameValueTextLine(TEXTVIEW("Subsequent tasks:"), FString::Printf(TEXT("%d"), Task.Subsequents.Num()));
+				InOutTooltip.AddNameValueTextLine(TEXTVIEW("Parent tasks:"), FString::Printf(TEXT("%d"), Task.ParentTasks.Num()));
+				InOutTooltip.AddNameValueTextLine(TEXTVIEW("Nested tasks:"), FString::Printf(TEXT("%d"), Task.NestedTasks.Num()));
 			};
 
-			do // info about a task
+			// info about a task
+			const TraceServices::FTaskInfo* Task = TasksProvider->TryGetTask(ThreadId, TooltipEvent.GetStartTime());
+			if (Task != nullptr && Task->FinishedTimestamp >= TooltipEvent.GetEndTime())
 			{
-				const TraceServices::FTaskInfo* Task = TasksProvider->TryGetTask(ThreadId, TooltipEvent.GetStartTime());
-				if (Task == nullptr)
-				{
-					break;
-				}
-
-				if (Task->FinishedTimestamp < TooltipEvent.GetEndTime())
-				{
-					break;
-				}
-
 				AddTaskInfo(*Task);
-			} while (false);
+			}
 
-			do // info about blocking
+			// info about blocking
+			const TraceServices::FWaitingForTasks* Waiting = TasksProvider->TryGetWaiting(TimerName, ThreadId, TooltipEvent.GetStartTime());
+			if (Waiting != nullptr && Waiting->Tasks.Num() > 0)
 			{
-				const TraceServices::FWaitingForTasks* Waiting = TasksProvider->TryGetWaiting(*TimerName, ThreadId, TooltipEvent.GetStartTime());
-				if (Waiting == nullptr || Waiting->Tasks.Num() == 0)
-				{
-					break;
-				}
-
 				InOutTooltip.AddTextLine(TEXT("-------- Waiting for tasks --------"), FLinearColor::Red);
 				constexpr int32 NumIdsOnRow = 4;
 				TStringBuilder<1024> StringBuilder;
@@ -1401,7 +1421,7 @@ void FThreadTimingTrack::InitTooltip(FTooltipDrawState& InOutTooltip, const ITim
 				{
 					StringBuilder.RemoveSuffix(2);
 				}
-				InOutTooltip.AddNameValueTextLine(FString::Printf(TEXT("Tasks[%d]:"), Waiting->Tasks.Num()), StringBuilder.ToString());
+				InOutTooltip.AddNameValueTextLine(FString::Printf(TEXT("Tasks[%d]:"), Waiting->Tasks.Num()), StringBuilder.ToView());
 				StringBuilder.Reset();
 
 				// Add the rest of the lines with an empty name so they appear as a multi line value.
@@ -1411,39 +1431,46 @@ void FThreadTimingTrack::InitTooltip(FTooltipDrawState& InOutTooltip, const ITim
 
 					if ((Index + 1) % NumIdsOnRow == 0)
 					{
-						InOutTooltip.AddNameValueTextLine(TEXT(""), StringBuilder.ToString());
+						InOutTooltip.AddNameValueTextLine(TEXT(""), StringBuilder.ToView());
 						StringBuilder.Reset();
 					}
 				}
-
 				if (StringBuilder.Len() > 1)
 				{
 					StringBuilder.RemoveSuffix(2);
-					InOutTooltip.AddNameValueTextLine(TEXT(""), StringBuilder.ToString());
+					InOutTooltip.AddNameValueTextLine(TEXT(""), StringBuilder.ToView());
 				}
 
-				InOutTooltip.AddNameValueTextLine(TEXT("Started waiting:"), FString::Printf(TEXT("%f"), Waiting->StartedTimestamp));
-				InOutTooltip.AddNameValueTextLine(TEXT("Finished waiting:"),
-					FString::Printf(TEXT("%s (+%s)"),
-						Waiting->FinishedTimestamp == TraceServices::FTaskInfo::InvalidTimestamp ? TEXT("[not set]") : *FString::SanitizeFloat(Waiting->FinishedTimestamp),
-						*TimeUtils::FormatTimeAuto(Waiting->FinishedTimestamp - Waiting->StartedTimestamp)));
+				InOutTooltip.AddNameValueTextLine(TEXT("Started waiting:"),
+					FString::Printf(TEXT("%s"), *FString::SanitizeFloat(Waiting->StartedTimestamp)));
+
+				if (Waiting->FinishedTimestamp != TraceServices::FTaskInfo::InvalidTimestamp)
+				{
+					InOutTooltip.AddNameValueTextLine(TEXT("Finished waiting:"),
+						FString::Printf(TEXT("%s (+%s)"),
+							*FString::SanitizeFloat(Waiting->FinishedTimestamp),
+							*TimeUtils::FormatTimeAuto(Waiting->FinishedTimestamp - Waiting->StartedTimestamp, 2)));
+				}
+				else
+				{
+					InOutTooltip.AddNameValueTextLine(TEXT("Finished waiting:"), TEXT("[not set]"));
+				}
 
 				const int32 MaxWaitedTasksToList = 5;
 				int32 NumTasksToList = FMath::Min(Waiting->Tasks.Num(), MaxWaitedTasksToList);
-				for (int32 i = 0; i != NumTasksToList; ++i)
+				for (int32 TaskIndex = 0; TaskIndex != NumTasksToList; ++TaskIndex)
 				{
-					const TraceServices::FTaskInfo* Task = TasksProvider->TryGetTask(Waiting->Tasks[i]);
-					if (Task != nullptr)
+					const TraceServices::FTaskInfo* WaitedTask = TasksProvider->TryGetTask(Waiting->Tasks[TaskIndex]);
+					if (WaitedTask != nullptr)
 					{
-						AddTaskInfo(*Task);
+						AddTaskInfo(*WaitedTask);
 					}
 				}
 				if (NumTasksToList < Waiting->Tasks.Num())
 				{
 					InOutTooltip.AddTextLine(TEXT("[...]"), FLinearColor::Green);
 				}
-
-			} while (false);
+			}
 		}
 	}
 	else if (ChildTrack.IsValid())
diff --git a/Engine/Source/Developer/TraceInsights/Private/Insights/ViewModels/TooltipDrawState.cpp b/Engine/Source/Developer/TraceInsights/Private/Insights/ViewModels/TooltipDrawState.cpp
--- a/Engine/Source/Developer/TraceInsights/Private/Insights/ViewModels/TooltipDrawState.cpp
+++ b/Engine/Source/Developer/TraceInsights/Private/Insights/ViewModels/TooltipDrawState.cpp
@@ -63,63 +63,63 @@ void FTooltipDrawState::ResetContent()
 
 ////////////////////////////////////////////////////////////////////////////////////////////////////
 
-void FTooltipDrawState::AddTitle(const FString& Title)
+void FTooltipDrawState::AddTitle(FStringView Title)
 {
 	const TSharedRef<FSlateFontMeasure> FontMeasureService = FSlateApplication::Get().GetRenderer()->GetFontMeasureService();
 
 	const FVector2D TextSize = FontMeasureService->Measure(Title, Font, FontScale) / FontScale;
-	Texts.Add({ BorderX, NewLineY, TextSize, Title, DefaultTitleColor, FDrawTextType::Title });
+	Texts.Add({ BorderX, NewLineY, TextSize, FString(Title), DefaultTitleColor, FDrawTextType::Title });
 
 	NewLineY += DefaultTitleHeight;
 }
 
 ////////////////////////////////////////////////////////////////////////////////////////////////////
 
-void FTooltipDrawState::AddTitle(const FString& Title, const FLinearColor& Color)
+void FTooltipDrawState::AddTitle(FStringView Title, const FLinearColor& Color)
 {
 	const TSharedRef<FSlateFontMeasure> FontMeasureService = FSlateApplication::Get().GetRenderer()->GetFontMeasureService();
 
 	const FVector2D TextSize = FontMeasureService->Measure(Title, Font, FontScale) / FontScale;
-	Texts.Add({ BorderX, NewLineY, TextSize, Title, Color, FDrawTextType::Misc });
+	Texts.Add({ BorderX, NewLineY, TextSize, FString(Title), Color, FDrawTextType::Misc });
 
 	NewLineY += DefaultTitleHeight;
 }
 
 ////////////////////////////////////////////////////////////////////////////////////////////////////
 
-void FTooltipDrawState::AddNameValueTextLine(const FString& Name, const FString& Value)
+void FTooltipDrawState::AddNameValueTextLine(FStringView Name, FStringView Value)
 {
 	const TSharedRef<FSlateFontMeasure> FontMeasureService = FSlateApplication::Get().GetRenderer()->GetFontMeasureService();
 
 	const FVector2D NameTextSize = FontMeasureService->Measure(Name, Font, FontScale) / FontScale;
-	Texts.Add({ 0.0f, NewLineY, NameTextSize, Name, DefaultNameColor, FDrawTextType::Name });
+	Texts.Add({ 0.0f, NewLineY, NameTextSize, FString(Name), DefaultNameColor, FDrawTextType::Name });
 
 	const FVector2D ValueTextSize = FontMeasureService->Measure(Value, Font, FontScale) / FontScale;
-	Texts.Add({ 0.0f, NewLineY, ValueTextSize, Value, DefaultValueColor, FDrawTextType::Value });
+	Texts.Add({ 0.0f, NewLineY, ValueTextSize, FString(Value), DefaultValueColor, FDrawTextType::Value });
 
 	NewLineY += DefaultLineHeight;
 }
 
 ////////////////////////////////////////////////////////////////////////////////////////////////////
 
-void FTooltipDrawState::AddTextLine(const FString& Text, const FLinearColor& Color)
+void FTooltipDrawState::AddTextLine(FStringView Text, const FLinearColor& Color)
 {
 	const TSharedRef<FSlateFontMeasure> FontMeasureService = FSlateApplication::Get().GetRenderer()->GetFontMeasureService();
 
 	const FVector2D TextSize = FontMeasureService->Measure(Text, Font, FontScale) / FontScale;
-	Texts.Add({ BorderX, NewLineY, TextSize, Text, Color, FDrawTextType::Misc });
+	Texts.Add({ BorderX, NewLineY, TextSize, FString(Text), Color, FDrawTextType::Misc });
 
 	NewLineY += DefaultLineHeight;
 }
 
 ////////////////////////////////////////////////////////////////////////////////////////////////////
 
-void FTooltipDrawState::AddTextLine(const float X, const float Y, const FString& Text, const FLinearColor& Color)
+void FTooltipDrawState::AddTextLine(const float X, const float Y, FStringView Text, const FLinearColor& Color)
 {
 	const TSharedRef<FSlateFontMeasure> FontMeasureService = FSlateApplication::Get().GetRenderer()->GetFontMeasureService();
 
 	const FVector2D TextSize = FontMeasureService->Measure(Text, Font, FontScale) / FontScale;
-	Texts.Add({ X, Y, TextSize, Text, Color, FDrawTextType::Misc });
+	Texts.Add({ X, Y, TextSize, FString(Text), Color, FDrawTextType::Misc });
 }
 
 ////////////////////////////////////////////////////////////////////////////////////////////////////
diff --git a/Engine/Source/Developer/TraceInsights/Public/Insights/ViewModels/TooltipDrawState.h b/Engine/Source/Developer/TraceInsights/Public/Insights/ViewModels/TooltipDrawState.h
--- a/Engine/Source/Developer/TraceInsights/Public/Insights/ViewModels/TooltipDrawState.h
+++ b/Engine/Source/Developer/TraceInsights/Public/Insights/ViewModels/TooltipDrawState.h
@@ -55,11 +55,11 @@ public:
 	void Reset();
 
 	void ResetContent();
-	void AddTitle(const FString& Title);
-	void AddTitle(const FString& Title, const FLinearColor& Color);
-	void AddNameValueTextLine(const FString& Name, const FString& Value);
-	void AddTextLine(const FString& Text, const FLinearColor& Color);
-	void AddTextLine(const float X, const float Y, const FString& Text, const FLinearColor& Color);
+	void AddTitle(FStringView Title);
+	void AddTitle(FStringView Title, const FLinearColor& Color);
+	void AddNameValueTextLine(FStringView Name, FStringView Value);
+	void AddTextLine(FStringView Text, const FLinearColor& Color);
+	void AddTextLine(const float X, const float Y, FStringView Text, const FLinearColor& Color);
 	void UpdateLayout(); // updates ValueOffsetX and DesiredSize
 
 	const FLinearColor& GetBackgroundColor() const { return BackgroundColor; }
