[Insights]
  - Timing View: Fixed fill color for graph series.
  - Timing View: Increased tickness of lines for graph series based on scale factor and enabled antialias.

diff --git a/Engine/Source/Developer/TraceInsights/Private/Insights/ViewModels/GraphTrack.cpp b/Engine/Source/Developer/TraceInsights/Private/Insights/ViewModels/GraphTrack.cpp
--- a/Engine/Source/Developer/TraceInsights/Private/Insights/ViewModels/GraphTrack.cpp
+++ b/Engine/Source/Developer/TraceInsights/Private/Insights/ViewModels/GraphTrack.cpp
@@ -251,6 +251,7 @@ void FGraphTrack::DrawSeries(const FGraphSeries& Series, FDrawContext& DrawConte
 
 		FVector2D AtlasOffset = ResourceProxy ? FVector2D(ResourceProxy->StartUV) : FVector2D(0.0, 0.0);
 		FVector2D AtlasUVSize = ResourceProxy ? FVector2D(ResourceProxy->SizeUV) : FVector2D(1.0, 1.0);
+		FVector2f UV(AtlasOffset + FVector2D(0.0, 0.0) * AtlasUVSize);
 
 		const FVector2D Size = DrawContext.Geometry.GetLocalSize();
 
@@ -258,6 +259,8 @@ void FGraphTrack::DrawSeries(const FGraphSeries& Series, FDrawContext& DrawConte
 
 		FColor FillColor = Series.FillColor.ToFColor(true);
 
+		const double BaselineY = Series.GetBaselineY();
+
 		for (int32 BatchIndex = 0; BatchIndex < Series.LinePoints.Num(); ++BatchIndex)
 		{
 			const TArray<FVector2D>& LinePoints = Series.LinePoints[BatchIndex];
@@ -268,11 +271,6 @@ void FGraphTrack::DrawSeries(const FGraphSeries& Series, FDrawContext& DrawConte
 			Indices.Reserve(LinePoints.Num() * 6);
 			Verts.Reserve(LinePoints.Num() * 2);
 
-			const double TopV = 0.0;
-			const double BottomV = GetHeight() / Size.Y;
-			const double BaselineY = Series.GetBaselineY();
-			const double BaselineV = FMath::Clamp<double>(BaselineY / Size.Y, TopV, BottomV);
-
 			int32 PrevSide = 0;
 
 			for (int32 PointIndex = 0; PointIndex < LinePoints.Num(); ++PointIndex)
@@ -285,59 +283,52 @@ void FGraphTrack::DrawSeries(const FGraphSeries& Series, FDrawContext& DrawConte
 				{
 					// Compute intersection point.
 					const FVector2D& PrevLinePoint = LinePoints[PointIndex - 1];
-					const double X = PrevLinePoint.X + (LinePoint.X - PrevLinePoint.X) / ((BaselineY - LinePoint.Y) / (PrevLinePoint.Y - BaselineY) + 1.0);
+					const double Delta = (PrevLinePoint.Y - BaselineY) / (PrevLinePoint.Y - LinePoint.Y);
+					const double X = PrevLinePoint.X + (LinePoint.X - PrevLinePoint.X) * Delta;
 
 					// Add an intersection point vertex.
-					FVector2D UV(X / Size.X, BaselineV);
-					Verts.Add(FSlateVertex::Make<ESlateVertexRounding::Disabled>(RenderTransform, FVector2f(UV * Size), FVector2f(AtlasOffset + UV * AtlasUVSize), FillColor));
+					Verts.Add(FSlateVertex::Make<ESlateVertexRounding::Disabled>(RenderTransform,
+						FVector2f((float)X, (float)BaselineY), UV, FillColor));
 
 					// Add a value point vertex.
-					UV.X = LinePoint.X / Size.X;
-					UV.Y = TopV + LinePoint.Y / Size.Y;
-					Verts.Add(FSlateVertex::Make<ESlateVertexRounding::Disabled>(RenderTransform, FVector2f(UV * Size), FVector2f(AtlasOffset + UV * AtlasUVSize), FillColor));
+					Verts.Add(FSlateVertex::Make<ESlateVertexRounding::Disabled>(RenderTransform,
+						FVector2f((float)LinePoint.X, (float)LinePoint.Y), UV, FillColor));
 
 					// Add a baseline vertex.
-					UV.Y = BaselineV;
-					Verts.Add(FSlateVertex::Make<ESlateVertexRounding::Disabled>(RenderTransform, FVector2f(UV * Size), FVector2f(AtlasOffset + FVector2D(UV.X, 0.5) * AtlasUVSize), FillColor));
-
-					int32 Index0 = Verts.Num() - 5;
-					int32 Index1 = Verts.Num() - 4;
-					int32 Index2 = Verts.Num() - 3;
-					int32 Index3 = Verts.Num() - 2;
-					int32 Index4 = Verts.Num() - 1;
-
-					Indices.Add(Index0);
-					Indices.Add(Index1);
-					Indices.Add(Index2);
-
-					Indices.Add(Index2);
-					Indices.Add(Index3);
-					Indices.Add(Index4);
+					Verts.Add(FSlateVertex::Make<ESlateVertexRounding::Disabled>(RenderTransform,
+						FVector2f((float)LinePoint.X, (float)BaselineY), UV, FillColor));
+
+					SlateIndex NumVerts = (SlateIndex)Verts.Num();
+					check(NumVerts >= 5);
+
+					Indices.Add(NumVerts - 5);
+					Indices.Add(NumVerts - 4);
+					Indices.Add(NumVerts - 3);
+
+					Indices.Add(NumVerts - 3);
+					Indices.Add(NumVerts - 2);
+					Indices.Add(NumVerts - 1);
 				}
 				else
 				{
 					// Add a value point vertex.
-					FVector2D UV(LinePoint.X / Size.X, TopV + LinePoint.Y / Size.Y);
-					Verts.Add(FSlateVertex::Make<ESlateVertexRounding::Disabled>(RenderTransform, FVector2f(UV * Size), FVector2f(AtlasOffset + UV * AtlasUVSize), FillColor));
+					Verts.Add(FSlateVertex::Make<ESlateVertexRounding::Disabled>(RenderTransform,
+						FVector2f((float)LinePoint.X, (float)LinePoint.Y), UV, FillColor));
 
 					// Add a baseline vertex.
-					UV.Y = BaselineV;
-					Verts.Add(FSlateVertex::Make<ESlateVertexRounding::Disabled>(RenderTransform, FVector2f(UV * Size), FVector2f(AtlasOffset + FVector2D(UV.X, 0.5) * AtlasUVSize), FillColor));
+					Verts.Add(FSlateVertex::Make<ESlateVertexRounding::Disabled>(RenderTransform,
+						FVector2f((float)LinePoint.X, (float)BaselineY), UV, FillColor));
 
-					if (Verts.Num() >= 4)
+					SlateIndex NumVerts = (SlateIndex)Verts.Num();
+					if (NumVerts >= 4)
 					{
-						int32 Index0 = Verts.Num() - 4;
-						int32 Index1 = Verts.Num() - 3;
-						int32 Index2 = Verts.Num() - 2;
-						int32 Index3 = Verts.Num() - 1;
-
-						Indices.Add(Index0);
-						Indices.Add(Index1);
-						Indices.Add(Index2);
-
-						Indices.Add(Index2);
-						Indices.Add(Index1);
-						Indices.Add(Index3);
+						Indices.Add(NumVerts - 4);
+						Indices.Add(NumVerts - 3);
+						Indices.Add(NumVerts - 2);
+
+						Indices.Add(NumVerts - 2);
+						Indices.Add(NumVerts - 3);
+						Indices.Add(NumVerts - 1);
 					}
 				}
 				PrevSide = CrtSide;
@@ -362,19 +353,23 @@ void FGraphTrack::DrawSeries(const FGraphSeries& Series, FDrawContext& DrawConte
 	if (IsAnyOptionEnabled(EGraphOptions::ShowLines))
 	{
 		// Find scale to get the size of a pixel
-		const float InvScale = 1.0f / DrawContext.Geometry.GetAccumulatedLayoutTransform().GetScale();
+		const float Scale = DrawContext.Geometry.GetAccumulatedLayoutTransform().GetScale();
+
 		FPaintGeometry LineGeo = Geo;
-		LineGeo.AppendTransform(FSlateLayoutTransform(FVector2D(0.5f * InvScale, 0.5f * InvScale)));
+		LineGeo.AppendTransform(FSlateLayoutTransform(FVector2D(0.5f / Scale, 0.5f / Scale)));
 
 		// Disable pixel snapping here so lines line up with boxes/polys correctly.
 		const ESlateDrawEffect LineDrawEffects = DrawContext.DrawEffects | ESlateDrawEffect::NoPixelSnapping;
 
+		constexpr bool bAntialias = true;
+		const float Thickness = FMath::Max(1.0f, Scale);
+
 		for (int32 BatchIndex = 0; BatchIndex < Series.LinePoints.Num(); ++BatchIndex)
 		{
 			const TArray<FVector2D>& LinePoints = Series.LinePoints[BatchIndex];
 			if (LinePoints.Num() > 0)
 			{
-				FSlateDrawElement::MakeLines(DrawContext.ElementList, DrawContext.LayerId, LineGeo, LinePoints, LineDrawEffects, Series.Color, false, 1.0f);
+				FSlateDrawElement::MakeLines(DrawContext.ElementList, DrawContext.LayerId, LineGeo, LinePoints, LineDrawEffects, Series.Color, bAntialias, Thickness);
 			}
 		}
 		DrawContext.LayerId++;
