[Insights] Timing View: Disabled "% of Root" and/or "% of Parent" in tooltips of timing events when Root/Parent timing event has zero or infinite duration.

diff --git a/Engine/Source/Developer/TraceInsights/Private/Insights/ViewModels/ThreadTimingTrack.cpp b/Engine/Source/Developer/TraceInsights/Private/Insights/ViewModels/ThreadTimingTrack.cpp
--- a/Engine/Source/Developer/TraceInsights/Private/Insights/ViewModels/ThreadTimingTrack.cpp
+++ b/Engine/Source/Developer/TraceInsights/Private/Insights/ViewModels/ThreadTimingTrack.cpp
@@ -1257,7 +1257,10 @@ void FThreadTimingTrack::InitTooltip(FTooltipDrawState& InOutTooltip, const ITim
 
 		const double TooltipEventDuration = TooltipEvent.GetDuration();
 
-		if (ParentTimingEvent.IsValid() && TooltipEvent.GetDepth() > 0)
+		if (TooltipEvent.GetDepth() > 0 &&
+			ParentTimingEvent.IsValid() &&
+			ParentTimingEvent->GetDuration() > 0.0 &&
+			ParentTimingEvent->GetDuration() != std::numeric_limits<double>::infinity())
 		{
 			Timer = TimerReader->GetTimer(ParentTimingEvent->GetTimerIndex());
 			const TCHAR* ParentTimerName = (Timer != nullptr) ? Timer->Name : TEXT("N/A");
@@ -1267,7 +1270,10 @@ void FThreadTimingTrack::InitTooltip(FTooltipDrawState& InOutTooltip, const ITim
 			InOutTooltip.AddNameValueTextLine(TEXTVIEW("% of Parent:"), ValueStr);
 		}
 
-		if (RootTimingEvent.IsValid() && TooltipEvent.GetDepth() > 1)
+		if (TooltipEvent.GetDepth() > 1 &&
+			RootTimingEvent.IsValid() &&
+			RootTimingEvent->GetDuration() > 0.0 &&
+			RootTimingEvent->GetDuration() != std::numeric_limits<double>::infinity())
 		{
 			Timer = TimerReader->GetTimer(RootTimingEvent->GetTimerIndex());
 			const TCHAR* RootTimerName = (Timer != nullptr) ? Timer->Name : TEXT("N/A");
@@ -1279,7 +1285,8 @@ void FThreadTimingTrack::InitTooltip(FTooltipDrawState& InOutTooltip, const ITim
 
 		InOutTooltip.AddNameValueTextLine(TEXTVIEW("Inclusive Time:"), TimeUtils::FormatTimeAuto(TooltipEventDuration, 2));
 
-		if (TooltipEventDuration > 0.0 && TooltipEventDuration != std::numeric_limits<double>::infinity())
+		if (TooltipEventDuration > 0.0 &&
+			TooltipEventDuration != std::numeric_limits<double>::infinity())
 		{
 			const double ExclusiveTimePercent = TooltipEvent.GetExclusiveTime() / TooltipEventDuration;
 			FNumberFormattingOptions FormattingOptions;
