Fix a crash when running Insights.Hub.SessionBrowser test when no top window exists

#rnx

diff --git a/Engine/Source/Developer/TraceInsights/Private/Insights/Tests/FunctionalTests/SessionBrowserTests.cpp b/Engine/Source/Developer/TraceInsights/Private/Insights/Tests/FunctionalTests/SessionBrowserTests.cpp
--- a/Engine/Source/Developer/TraceInsights/Private/Insights/Tests/FunctionalTests/SessionBrowserTests.cpp
+++ b/Engine/Source/Developer/TraceInsights/Private/Insights/Tests/FunctionalTests/SessionBrowserTests.cpp
@@ -26,7 +26,7 @@ void FAutomationDriverUnrealInsightsSessionBrowserTest::Define()
 	BeforeEach([this]() {
 		AutomationWindow = FSlateApplication::Get().GetActiveTopLevelWindow();
 		const FString AutomationWindowName = TEXT("Automation");
-		if (AutomationWindow->GetTitle().ToString().Contains(AutomationWindowName))
+		if (AutomationWindow && AutomationWindow->GetTitle().ToString().Contains(AutomationWindowName))
 		{
 			AutomationWindow->Minimize();
 		}
