[Insights] Add metadata filter to the Filter Configurator. Add support for filters with custom UI.

diff --git a/Engine/Source/Developer/TraceInsights/Private/Insights/ViewModels/FilterConfiguratorNode.cpp b/Engine/Source/Developer/TraceInsights/Private/Insights/ViewModels/FilterConfiguratorNode.cpp
--- a/Engine/Source/Developer/TraceInsights/Private/Insights/ViewModels/FilterConfiguratorNode.cpp
+++ b/Engine/Source/Developer/TraceInsights/Private/Insights/ViewModels/FilterConfiguratorNode.cpp
@@ -109,16 +109,19 @@ void FFilterConfiguratorNode::SetAvailableFilters(TSharedPtr<TArray<TSharedPtr<F
 
 void FFilterConfiguratorNode::SetSelectedFilter(TSharedPtr<FFilter> InSelectedFilter)
 {
-	if (SelectedFilter.Get() != InSelectedFilter.Get())
+	if (SelectedFilter.Get() == InSelectedFilter.Get())
 	{
-		TextBoxValue.Empty();
+		return;
 	}
 
+	TextBoxValue.Empty();
 	SelectedFilter = InSelectedFilter;
-	if (SelectedFilter.IsValid() && SelectedFilter->GetSupportedOperators()->Num() > 0)
-	{
-		FilterState = SelectedFilter->BuildFilterState();
+	check(SelectedFilter.IsValid());
+
+	FilterState = SelectedFilter->BuildFilterState();
 
+	if (SelectedFilter->GetSupportedOperators()->Num() > 0)
+	{
 		SetSelectedFilterOperator(SelectedFilter->GetSupportedOperators()->GetData()[0]);
 
 		AvailableFilterOperators->Empty();
diff --git a/Engine/Source/Developer/TraceInsights/Private/Insights/ViewModels/FilterConfiguratorNode.h b/Engine/Source/Developer/TraceInsights/Private/Insights/ViewModels/FilterConfiguratorNode.h
--- a/Engine/Source/Developer/TraceInsights/Private/Insights/ViewModels/FilterConfiguratorNode.h
+++ b/Engine/Source/Developer/TraceInsights/Private/Insights/ViewModels/FilterConfiguratorNode.h
@@ -83,6 +83,8 @@ public:
 
 	void Update();
 
+	TSharedPtr<FFilterState> GetSelectedFilterState() { return FilterState; }
+
 private:
 	FFilterConfiguratorNode& operator=(const FFilterConfiguratorNode& Other);
 
diff --git a/Engine/Source/Developer/TraceInsights/Private/Insights/ViewModels/Filters.h b/Engine/Source/Developer/TraceInsights/Private/Insights/ViewModels/Filters.h
--- a/Engine/Source/Developer/TraceInsights/Private/Insights/ViewModels/Filters.h
+++ b/Engine/Source/Developer/TraceInsights/Private/Insights/ViewModels/Filters.h
@@ -11,6 +11,7 @@
 
 class FSpawnTabArgs;
 class SDockTab;
+class SHorizontalBox;
 class SWidget;
 
 namespace Insights
@@ -196,6 +197,9 @@ public:
 
 	virtual void SetFilterValue(FString InTextValue);
 
+	virtual bool HasCustomUI() const { return false; }
+	virtual void AddCustomUI(TSharedRef<SHorizontalBox> LeftBox) {}
+
 protected:
 	TSharedRef<FFilter> Filter;
 	TSharedPtr<IFilterOperator> SelectedOperator;
diff --git a/Engine/Source/Developer/TraceInsights/Private/Insights/ViewModels/ThreadTimingTrack.cpp b/Engine/Source/Developer/TraceInsights/Private/Insights/ViewModels/ThreadTimingTrack.cpp
--- a/Engine/Source/Developer/TraceInsights/Private/Insights/ViewModels/ThreadTimingTrack.cpp
+++ b/Engine/Source/Developer/TraceInsights/Private/Insights/ViewModels/ThreadTimingTrack.cpp
@@ -1112,6 +1112,14 @@ void FThreadTimingTrack::BuildFilteredDrawState(ITimingEventsTrackDrawStateBuild
 					TArray<FFilterContext> FilterContexts;
 
 					TraceServices::ITimeline<TraceServices::FTimingProfilerEvent>::EnumerateAsyncParams Params;
+					constexpr uint32 LargeTimelineThreshold = 50 * 1000 * 1000;
+					if (Timeline.GetEventCount() > LargeTimelineThreshold)
+					{
+						if (FilterConfigurator->IsKeyUsed(static_cast<int32>(EFilterField::Metadata)))
+						{
+							Params.MaxOccupancy = 0.75; // This filter can be slow so reduce occupancy to avoid starvation.
+						}
+					}
 					Params.IntervalStart = Viewport.GetStartTime();
 					Params.IntervalEnd = Viewport.GetEndTime();
 
@@ -1130,6 +1138,7 @@ void FThreadTimingTrack::BuildFilteredDrawState(ITimingEventsTrackDrawStateBuild
 							Context.AddFilterData<FString>(static_cast<int32>(EFilterField::TrackName), this->GetName());
 							Context.AddFilterData<int64>(static_cast<int32>(EFilterField::TimerId), 0);
 							Context.AddFilterData<int64>(static_cast<int32>(EFilterField::TimerName), 0);
+							Context.AddFilterData<int64>(static_cast<int32>(EFilterField::Metadata), 0);
 						}
 					};
 					Params.Callback = [this, &Builder, TimerReader, &FilteredEvents, &FilterContexts](double StartTime, double EndTime, uint32 Depth, const TraceServices::FTimingProfilerEvent& Event, uint32 TaskIndex)
@@ -1144,6 +1153,7 @@ void FThreadTimingTrack::BuildFilteredDrawState(ITimingEventsTrackDrawStateBuild
 							// The TimerName filter also translates to the numeric Id for performance reasons.
 							Context.SetFilterData<int64>(static_cast<int32>(EFilterField::TimerId), Timer->Id);
 							Context.SetFilterData<int64>(static_cast<int32>(EFilterField::TimerName), Timer->Id);
+							Context.SetFilterData<int64>(static_cast<int32>(EFilterField::Metadata), Event.TimerIndex);
 
 							if (FilterConfigurator->ApplyFilters(Context))
 							{
@@ -1785,6 +1795,7 @@ bool FThreadTimingTrack::FindTimingProfilerEvent(const FTimingEventSearchParamet
 	FilterConfiguratorContext.AddFilterData<FString>(static_cast<int32>(EFilterField::TrackName), this->GetName());
 	FilterConfiguratorContext.AddFilterData<int64>(static_cast<int32>(EFilterField::TimerId), 0);
 	FilterConfiguratorContext.AddFilterData<int64>(static_cast<int32>(EFilterField::TimerName), 0);
+	FilterConfiguratorContext.AddFilterData<int64>(static_cast<int32>(EFilterField::Metadata), 0);
 
 	return TTimingEventSearch<TraceServices::FTimingProfilerEvent>::Search(
 		InParameters,
@@ -1846,8 +1857,8 @@ bool FThreadTimingTrack::FindTimingProfilerEvent(const FTimingEventSearchParamet
 						FilterConfiguratorContext.SetFilterData<double>(static_cast<int32>(EFilterField::EndTime), EventEndTime);
 						FilterConfiguratorContext.SetFilterData<double>(static_cast<int32>(EFilterField::Duration), EventEndTime - EventStartTime);
 						FilterConfiguratorContext.SetFilterData<int64>(static_cast<int32>(EFilterField::TimerId), Timer->Id);
-						// The TimerName filter also translates to the numeric Id for performance reasons.
 						FilterConfiguratorContext.SetFilterData<int64>(static_cast<int32>(EFilterField::TimerName), Timer->Id);
+						FilterConfiguratorContext.SetFilterData<int64>(static_cast<int32>(EFilterField::Metadata), Event.TimerIndex);
 					}
 					return InParameters.FilterExecutor->ApplyFilters(FilterConfiguratorContext);
 				}
diff --git a/Engine/Source/Developer/TraceInsights/Private/Insights/ViewModels/ThreadTimingTrack.h b/Engine/Source/Developer/TraceInsights/Private/Insights/ViewModels/ThreadTimingTrack.h
--- a/Engine/Source/Developer/TraceInsights/Private/Insights/ViewModels/ThreadTimingTrack.h
+++ b/Engine/Source/Developer/TraceInsights/Private/Insights/ViewModels/ThreadTimingTrack.h
@@ -134,6 +134,7 @@ enum class EFilterField : int32
 	TimerName = 5,
 	CoreEventName = 6,
 	RegionName = 7,
+	Metadata = 8,
 };
 
 ////////////////////////////////////////////////////////////////////////////////////////////////////
diff --git a/Engine/Source/Developer/TraceInsights/Private/Insights/ViewModels/TimerFilters.cpp b/Engine/Source/Developer/TraceInsights/Private/Insights/ViewModels/TimerFilters.cpp
--- a/Engine/Source/Developer/TraceInsights/Private/Insights/ViewModels/TimerFilters.cpp
+++ b/Engine/Source/Developer/TraceInsights/Private/Insights/ViewModels/TimerFilters.cpp
@@ -2,6 +2,10 @@
 
 #include "TimerFilters.h"
 
+#include "CborReader.h"
+#include "Widgets/Input/SEditableTextBox.h"
+#include "Widgets/SBoxPanel.h"
+
 #include "Insights/TimingProfilerManager.h"
 #include "Insights/ViewModels/ThreadTimingTrack.h"
 
@@ -132,6 +136,488 @@ void FTimerNameFilter::PopulateTimerNameSuggestionList(const FString& Text, TArr
 	}
 }
 
+////////////////////////////////////////////////////////////////////////////////////////////////////
+// FMetadataFilterState
+////////////////////////////////////////////////////////////////////////////////////////////////////
+
+INSIGHTS_IMPLEMENT_RTTI(FMetadataFilterState)
+
+FMetadataFilterState::FMetadataFilterState(TSharedRef<FFilter> InFilter)
+	: FFilterState(InFilter)
+{
+	AvailableDataTypes.Add(MakeShared<FMetadataFilterDataTypeEntry>(EMetadataFilterDataType::Int, LOCTEXT("IntDataType", "Int")));
+	AvailableDataTypes.Add(MakeShared<FMetadataFilterDataTypeEntry>(EMetadataFilterDataType::Double, LOCTEXT("DoubleDataType", "Double")));
+	AvailableDataTypes.Add(MakeShared<FMetadataFilterDataTypeEntry>(EMetadataFilterDataType::Bool, LOCTEXT("BoolDataType", "Bool")));
+	AvailableDataTypes.Add(MakeShared<FMetadataFilterDataTypeEntry>(EMetadataFilterDataType::String, LOCTEXT("StringDataType", "String")));
+
+	SelectedDataType = AvailableDataTypes[0];
+	DataType_OnSelectionChanged(SelectedDataType, ESelectInfo::Type::Direct);
+
+	BoolOperators.Add((MakeShared<FFilterOperator<bool>>(EFilterOperator::Eq, TEXT("IS"), [](bool lhs, bool rhs) { return lhs == rhs; })));
+}
+
+////////////////////////////////////////////////////////////////////////////////////////////////////
+
+void FMetadataFilterState::AddCustomUI(TSharedRef<SHorizontalBox> Box)
+{
+	Box->AddSlot()
+		.AutoWidth()
+		.VAlign(VAlign_Center)
+		.HAlign(HAlign_Left)
+		.Padding(FMargin(4.0f, 0.0f, 4.0f, 0.0f))
+		[
+			SNew(SEditableTextBox)
+			.MinDesiredWidth(200.0f)
+			.OnTextCommitted(this, &FMetadataFilterState::OnKeyTextBoxValueCommitted)
+			.Text(this, &FMetadataFilterState::GetKeyTextBoxValue)
+			.ToolTipText(LOCTEXT("MetadataKeyTooltipText", "The key value for the metadata."))
+			.HintText(LOCTEXT("MetadataKeyHint", "The key value for the metadata."))
+		];
+
+	Box->AddSlot()
+		.HAlign(HAlign_Left)
+		.VAlign(VAlign_Center)
+		.AutoWidth()
+		.Padding(FMargin(0.0f, 2.0f))
+		[
+			SNew(SComboBox<TSharedPtr<FMetadataFilterDataTypeEntry>>)
+			.OptionsSource(&AvailableDataTypes)
+			.OnSelectionChanged(this, &FMetadataFilterState::DataType_OnSelectionChanged)
+			.OnGenerateWidget(this, &FMetadataFilterState::DataType_OnGenerateWidget)
+			[
+				SNew(STextBlock)
+				.Text(this, &FMetadataFilterState::DataType_GetSelectionText)
+			]
+		];
+
+		Box->AddSlot()
+			.HAlign(HAlign_Left)
+			.VAlign(VAlign_Center)
+			.FillWidth(1.0)
+			.Padding(FMargin(4.0f, 0.0f, 0.0f, 0.0f))
+			.AutoWidth()
+			[
+				SAssignNew(OperatorComboBox, SComboBox<TSharedPtr<IFilterOperator>>)
+				.OptionsSource(&AvailableOperators)
+				.OnSelectionChanged(this, &FMetadataFilterState::AvailableOperators_OnSelectionChanged)
+				.OnGenerateWidget(this, &FMetadataFilterState::AvailableOperators_OnGenerateWidget)
+				[
+					SNew(STextBlock)
+					.Text(this, &FMetadataFilterState::AvailableOperators_GetSelectionText)
+				]
+			];
+
+		Box->AddSlot()
+		.AutoWidth()
+		.VAlign(VAlign_Center)
+		.HAlign(HAlign_Left)
+		.FillWidth(1.0)
+		.Padding(FMargin(4.0f, 0.0f, 4.0f, 0.0f))
+		[
+			SNew(SEditableTextBox)
+			.MinDesiredWidth(200.0f)
+			.OnTextCommitted(this, &FMetadataFilterState::OnTermTextBoxValueCommitted)
+			.Text(this, &FMetadataFilterState::GetTermTextBoxValue)
+			.ToolTipText(LOCTEXT("MetadataValueTooltipText", "The value of the metadata field."))
+			.HintText(LOCTEXT("MetadataValueHint", "The value for the metadata field."))
+		];
+}
+
+////////////////////////////////////////////////////////////////////////////////////////////////////
+
+void FMetadataFilterState::Update()
+{
+	bShowAllMetadataEvents = false;
+	if (Key.Equals(TEXT("*")))
+	{
+		bShowAllMetadataEvents = true;
+	}
+
+	ConvertedData = ConvertedDataVariant();
+	
+	ensure(SelectedDataType.IsValid());
+	switch (SelectedDataType->Type)
+	{
+	case EMetadataFilterDataType::Bool:
+	{
+		ConvertedData.Set<bool>(FCString::ToBool(*Term));
+		break;
+	}
+	case EMetadataFilterDataType::Int:
+	{
+		ConvertedData.Set<int64>(FCString::Atoi64(*Term));
+		break;
+	}
+	case EMetadataFilterDataType::Double:
+	{
+		ConvertedData.Set<double>(FCString::Atod(*Term));
+		break;
+	}
+	}
+
+	TSharedPtr<const TraceServices::IAnalysisSession> Session = FInsightsManager::Get()->GetSession();
+	if (Session.IsValid() && TraceServices::ReadTimingProfilerProvider(*Session.Get()))
+	{
+		TraceServices::FAnalysisSessionReadScope SessionReadScope(*Session.Get());
+
+		const TraceServices::ITimingProfilerProvider& TimingProfilerProvider = *TraceServices::ReadTimingProfilerProvider(*Session.Get());
+
+		// We can cache the timer reader because ApplyFilters needs to be called from inside a session lock. Acquiring the lock
+		// in ApplyFilter is too expensive.
+		TimingProfilerProvider.ReadTimers([this](const TraceServices::ITimingProfilerTimerReader& Out) { this->TimerReader = &Out; });
+	}
+}
+
+////////////////////////////////////////////////////////////////////////////////////////////////////
+
+bool FMetadataFilterState::ApplyFilter(const FFilterContext& Context) const
+{
+	if (!Context.HasFilterData(static_cast<int32>(Filter->GetKey())))
+	{
+		return Context.GetReturnValueForUnsetFilters();
+	}
+
+	int64 TimerIndex;
+	Context.GetFilterData<int64>(Filter->GetKey(), TimerIndex);
+
+	TArrayView<const uint8> Metadata = TimerReader->GetMetadata(static_cast<uint32>(TimerIndex));
+	if (Metadata.Num() > 0)
+	{
+		if (!bShowAllMetadataEvents)
+		{
+			return ApplyFilterToMetadata(Metadata);
+		}
+
+		return true;
+	}
+
+	return false;
+}
+
+////////////////////////////////////////////////////////////////////////////////////////////////////
+
+bool FMetadataFilterState::ApplyFilterToMetadata(TArrayView<const uint8>& Metadata) const
+{
+	FMemoryReaderView MemoryReader(Metadata);
+	FCborReader CborReader(&MemoryReader, ECborEndianness::StandardCompliant);
+	FCborContext Context;
+
+	if (!CborReader.ReadNext(Context) || Context.MajorType() != ECborCode::Map)
+	{
+		return false;
+	}
+
+	ensure(SelectedDataType);
+
+	while (true)
+	{
+		// Read key
+		if (!CborReader.ReadNext(Context) || !Context.IsString())
+		{
+			break;
+		}
+
+		FString CurrentKey(static_cast<int32>(Context.AsLength()), Context.AsCString());
+
+		if (!CborReader.ReadNext(Context))
+		{
+			break;
+		}
+
+		if (Context.IsFiniteContainer())
+		{
+			CborReader.SkipContainer(ECborCode::Array);
+		}
+
+		if (!CurrentKey.Equals(Key, ESearchCase::IgnoreCase))
+		{
+			continue;
+		}
+
+		switch (Context.MajorType())
+		{
+		case ECborCode::Int:
+		case ECborCode::Uint:
+		{
+			if (SelectedDataType->Type != EMetadataFilterDataType::Int)
+			{
+				continue;
+			}
+			ensure(ConvertedData.IsType<int64>());
+
+			int64 MetadataValue = static_cast<int64>(Context.AsUInt());
+			int64 InputValue = ConvertedData.Get<int64>();
+
+			FFilterOperator<int64>* Operator = (FFilterOperator<int64>*) SelectedOperator.Get();
+			if (Operator->Apply(MetadataValue, InputValue))
+			{
+				return true;
+			}
+			
+			continue;
+		}
+
+		case ECborCode::TextString:
+		{
+			if (SelectedDataType->Type != EMetadataFilterDataType::String)
+			{
+				continue;
+			}
+
+			FString Value = Context.AsString();
+			FFilterOperator<FString>* Operator = (FFilterOperator<FString>*) SelectedOperator.Get();
+			if (Operator->Apply(Value, Term))
+			{
+				return true;
+			}
+
+			continue;
+		}
+
+		case ECborCode::ByteString:
+		{
+			if (SelectedDataType->Type != EMetadataFilterDataType::String)
+			{
+				continue;
+			}
+
+			FAnsiStringView Value(Context.AsCString(), static_cast<int32>(Context.AsLength()));
+			FString ValueStr(Value);
+
+			FFilterOperator<FString>* Operator = (FFilterOperator<FString>*) SelectedOperator.Get();
+			if (Operator->Apply(ValueStr, Term))
+			{
+				return true;
+			}
+
+			continue;
+		}
+		}
+
+		if (Context.RawCode() == (ECborCode::Prim | ECborCode::Value_4Bytes))
+		{
+			if (SelectedDataType->Type != EMetadataFilterDataType::Double)
+			{
+				continue;
+			}
+			ensure(ConvertedData.IsType<double>());
+
+			double Value = static_cast<double>(Context.AsFloat());
+			double InputValue = ConvertedData.Get<double>();
+
+			FFilterOperator<double>* Operator = (FFilterOperator<double>*) SelectedOperator.Get();
+			if (Operator->Apply(Value, InputValue))
+			{
+				return true;
+			}
+
+			continue;
+		}
+
+		if (Context.RawCode() == (ECborCode::Prim | ECborCode::Value_8Bytes))
+		{
+			if (SelectedDataType->Type != EMetadataFilterDataType::Double)
+			{
+				continue;
+			}
+			ensure(ConvertedData.IsType<double>());
+
+			double Value = Context.AsDouble();
+			double InputValue = ConvertedData.Get<double>();
+
+			FFilterOperator<double>* Operator = (FFilterOperator<double>*) SelectedOperator.Get();
+			if (Operator->Apply(Value, InputValue))
+			{
+				return true;
+			}
+
+			continue;
+		}
+
+		if (Context.RawCode() == (ECborCode::Prim | ECborCode::False))
+		{
+			if (SelectedDataType->Type != EMetadataFilterDataType::Bool)
+			{
+				continue;
+			}
+			ensure(ConvertedData.IsType<bool>());
+
+			FFilterOperator<bool>* Operator = (FFilterOperator<bool>*) SelectedOperator.Get();
+			if (Operator->Apply(ConvertedData.Get<bool>(), false))
+			{
+				return true;
+			}
+
+			continue;
+		}
+
+		if (Context.RawCode() == (ECborCode::Prim | ECborCode::True))
+		{
+			if (SelectedDataType->Type != EMetadataFilterDataType::Bool)
+			{
+				continue;
+			}
+			ensure(ConvertedData.IsType<bool>());
+
+			FFilterOperator<bool>* Operator = (FFilterOperator<bool>*) SelectedOperator.Get();
+			if (Operator->Apply(ConvertedData.Get<bool>(), true))
+			{
+				return true;
+			}
+
+			continue;
+		}
+	}
+
+	return false;
+}
+
+////////////////////////////////////////////////////////////////////////////////////////////////////
+
+FText FMetadataFilterState::GetKeyTextBoxValue() const
+{
+	return FText::FromString(Key);
+}
+
+////////////////////////////////////////////////////////////////////////////////////////////////////
+
+void FMetadataFilterState::OnKeyTextBoxValueCommitted(const FText& InNewText, ETextCommit::Type InTextCommit)
+{
+	Key = InNewText.ToString();
+}
+
+////////////////////////////////////////////////////////////////////////////////////////////////////
+
+TSharedRef<SWidget> FMetadataFilterState::DataType_OnGenerateWidget(TSharedPtr<FMetadataFilterDataTypeEntry> InDataType)
+{
+	TSharedRef<SHorizontalBox> Widget = SNew(SHorizontalBox);
+	Widget->AddSlot()
+		.AutoWidth()
+		[
+			SNew(STextBlock)
+			.Text(InDataType->Name)
+			.Margin(2.0f)
+		];
+
+	return Widget;
+}
+
+////////////////////////////////////////////////////////////////////////////////////////////////////
+
+void FMetadataFilterState::DataType_OnSelectionChanged(TSharedPtr<FMetadataFilterDataTypeEntry> InDataType, ESelectInfo::Type SelectInfo)
+{
+	SelectedDataType = InDataType;
+
+	AvailableOperators.Empty();
+	SelectedOperator = nullptr;
+
+	switch (SelectedDataType->Type)
+	{
+	case EMetadataFilterDataType::Bool:
+	{
+		AvailableOperators.Insert(BoolOperators, 0);
+		break;
+	}
+	case EMetadataFilterDataType::Int:
+	{
+		AvailableOperators.Insert(*FFilterService::Get()->GetIntegerOperators(), 0);
+		break;
+	}
+	case EMetadataFilterDataType::Double:
+	{
+		AvailableOperators.Insert(*FFilterService::Get()->GetDoubleOperators(), 0);
+		break;
+	}
+	case EMetadataFilterDataType::String:
+	{
+		AvailableOperators.Insert(*FFilterService::Get()->GetStringOperators(), 0);
+		break;
+	}
+	}
+
+	if (AvailableOperators.Num() > 0)
+	{
+		AvailableOperators_OnSelectionChanged(AvailableOperators[0], ESelectInfo::Type::Direct);
+
+		if (OperatorComboBox.IsValid())
+		{
+			OperatorComboBox->RefreshOptions();
+			OperatorComboBox->SetSelectedItem(AvailableOperators[0]);
+		}
+	}
+}
+
+////////////////////////////////////////////////////////////////////////////////////////////////////
+
+FText FMetadataFilterState::DataType_GetSelectionText() const
+{
+	return SelectedDataType->Name;
+}
+
+////////////////////////////////////////////////////////////////////////////////////////////////////
+
+TSharedRef<SWidget> FMetadataFilterState::AvailableOperators_OnGenerateWidget(TSharedPtr<IFilterOperator> InOperator)
+{
+	TSharedRef<SHorizontalBox> Widget = SNew(SHorizontalBox);
+	Widget->AddSlot()
+		.AutoWidth()
+		[
+			SNew(STextBlock)
+			.Text(FText::FromString(InOperator->GetName()))
+			.Margin(2.0f)
+		];
+
+	return Widget;
+}
+
+////////////////////////////////////////////////////////////////////////////////////////////////////
+
+void FMetadataFilterState::AvailableOperators_OnSelectionChanged(TSharedPtr<IFilterOperator> InOperator, ESelectInfo::Type SelectInfo)
+{
+	if (InOperator)
+	{
+		SelectedOperator = InOperator;
+	}
+}
+
+////////////////////////////////////////////////////////////////////////////////////////////////////
+
+FText FMetadataFilterState::AvailableOperators_GetSelectionText() const
+{
+	return FText::FromString(SelectedOperator.IsValid() ? SelectedOperator->GetName(): TEXT(""));
+}
+
+////////////////////////////////////////////////////////////////////////////////////////////////////
+
+FText FMetadataFilterState::GetTermTextBoxValue() const
+{
+	return FText::FromString(Term);
+}
+
+////////////////////////////////////////////////////////////////////////////////////////////////////
+
+void FMetadataFilterState::OnTermTextBoxValueCommitted(const FText& InNewText, ETextCommit::Type InTextCommit)
+{
+	Term = InNewText.ToString();
+}
+
+////////////////////////////////////////////////////////////////////////////////////////////////////
+// FMetadataFilter
+////////////////////////////////////////////////////////////////////////////////////////////////////
+
+INSIGHTS_IMPLEMENT_RTTI(FMetadataFilter)
+
+FMetadataFilter::FMetadataFilter()
+	: FFilter(static_cast<int32>(EFilterField::Metadata),
+		LOCTEXT("MetadataFilterName", "Metadata"),
+		LOCTEXT("MetadataFilterDesc", "A filter for timing event metadata."),
+		EFilterDataType::Custom,
+		nullptr,
+		nullptr)
+{
+	SupportedOperators = MakeShared<TArray<TSharedPtr<IFilterOperator>>>();
+}
+
 ////////////////////////////////////////////////////////////////////////////////////////////////////
 
 } // namespace Insights
diff --git a/Engine/Source/Developer/TraceInsights/Private/Insights/ViewModels/TimerFilters.h b/Engine/Source/Developer/TraceInsights/Private/Insights/ViewModels/TimerFilters.h
--- a/Engine/Source/Developer/TraceInsights/Private/Insights/ViewModels/TimerFilters.h
+++ b/Engine/Source/Developer/TraceInsights/Private/Insights/ViewModels/TimerFilters.h
@@ -2,9 +2,17 @@
 
 #pragma once
 
+#include "Types/SlateEnums.h"
+#include "Widgets/Input/SComboBox.h"
+
 #include "Insights/Common/SimpleRtti.h"
 #include "Insights/ViewModels/Filters.h"
 
+namespace TraceServices
+{
+	class ITimingProfilerTimerReader;
+}
+
 namespace Insights
 {
 
@@ -53,4 +61,99 @@ public:
 	void PopulateTimerNameSuggestionList(const FString& Text, TArray<FString>& OutSuggestions);
 };
 
+enum class EMetadataFilterDataType
+{
+	Bool = 1,
+	Int = 2,
+	Double = 3,
+	String = 4,
+};
+
+struct FMetadataFilterDataTypeEntry
+{
+	FMetadataFilterDataTypeEntry(EMetadataFilterDataType InType, FText InName)
+	{
+		Type = InType;
+		Name = InName;
+	}
+
+	EMetadataFilterDataType Type;
+	FText Name;
+};
+
+class FMetadataFilterState : public FFilterState, public TSharedFromThis<FMetadataFilterState>
+{
+	INSIGHTS_DECLARE_RTTI(FMetadataFilterState, FFilterState)
+
+public:
+	FMetadataFilterState(TSharedRef<FFilter> InFilter);
+
+	virtual ~FMetadataFilterState() {}
+
+	virtual void Update() override;
+
+	virtual bool ApplyFilter(const FFilterContext& Context) const override;
+
+	virtual void SetFilterValue(FString InFilterValue) override {}
+
+	virtual bool HasCustomUI() const override { return true; }
+	virtual void AddCustomUI(TSharedRef<SHorizontalBox> LeftBox) override;
+
+private:
+	FText GetKeyTextBoxValue() const;
+	void OnKeyTextBoxValueCommitted(const FText& InNewText, ETextCommit::Type InTextCommit);
+
+	TSharedRef<SWidget> DataType_OnGenerateWidget(TSharedPtr<FMetadataFilterDataTypeEntry> InDataType);
+	void DataType_OnSelectionChanged(TSharedPtr<FMetadataFilterDataTypeEntry> InDataType, ESelectInfo::Type SelectInfo);
+	FText DataType_GetSelectionText() const;
+
+	TSharedRef<SWidget> AvailableOperators_OnGenerateWidget(TSharedPtr<IFilterOperator> InOperator);
+	void AvailableOperators_OnSelectionChanged(TSharedPtr<IFilterOperator> InOperator, ESelectInfo::Type SelectInfo);
+	FText AvailableOperators_GetSelectionText() const;
+
+	FText GetTermTextBoxValue() const;
+	void OnTermTextBoxValueCommitted(const FText& InNewText, ETextCommit::Type InTextCommit);
+
+	bool ApplyFilterToMetadata(TArrayView<const uint8>& Metadata) const;
+
+private:
+	FString Key;
+	FString Term;
+
+	typedef TVariant<double, int64, bool> ConvertedDataVariant;
+	ConvertedDataVariant ConvertedData;
+
+	bool bShowAllMetadataEvents = false;
+	
+	TArray<TSharedPtr<FMetadataFilterDataTypeEntry>> AvailableDataTypes;
+	TSharedPtr<FMetadataFilterDataTypeEntry> SelectedDataType;
+	
+	TArray<TSharedPtr<IFilterOperator>> AvailableOperators;
+	TArray<TSharedPtr<IFilterOperator>> BoolOperators;
+
+	TSharedPtr<SComboBox<TSharedPtr<IFilterOperator>>> OperatorComboBox;
+
+	const TraceServices::ITimingProfilerTimerReader* TimerReader;
+};
+
+class FMetadataFilter : public FFilter
+{
+	INSIGHTS_DECLARE_RTTI(FMetadataFilter, FFilter)
+
+public:
+	FMetadataFilter();
+
+	virtual ~FMetadataFilter() {}
+
+	virtual TSharedRef<FFilterState> BuildFilterState()
+	{
+		return MakeShared<FMetadataFilterState>(SharedThis(this));
+	}
+
+	virtual TSharedRef<FFilterState> BuildFilterState(const FFilterState& Other)
+	{
+		return MakeShared<FMetadataFilterState>(static_cast<const FMetadataFilterState&>(Other));
+	}
+};
+
 } // namespace Insights
diff --git a/Engine/Source/Developer/TraceInsights/Private/Insights/Widgets/SFilterConfiguratorRow.cpp b/Engine/Source/Developer/TraceInsights/Private/Insights/Widgets/SFilterConfiguratorRow.cpp
--- a/Engine/Source/Developer/TraceInsights/Private/Insights/Widgets/SFilterConfiguratorRow.cpp
+++ b/Engine/Source/Developer/TraceInsights/Private/Insights/Widgets/SFilterConfiguratorRow.cpp
@@ -53,80 +53,90 @@ TSharedRef<SWidget> SFilterConfiguratorRow::GenerateWidgetForColumn(const FName&
 				.ShouldDrawWires(true)
 			];
 
-		// Filter combo box
-		LeftBox->AddSlot()
-			.HAlign(HAlign_Left)
-			.VAlign(VAlign_Center)
-			.AutoWidth()
-			.Padding(FMargin(0.0f, 2.0f))
-			[
-				SAssignNew(FilterTypeComboBox, SComboBox<TSharedPtr<FFilter>>)
-				.OptionsSource(GetAvailableFilters())
-				.OnSelectionChanged(this, &SFilterConfiguratorRow::AvailableFilters_OnSelectionChanged)
-				.OnGenerateWidget(this, &SFilterConfiguratorRow::AvailableFilters_OnGenerateWidget)
-				[
-					SNew(STextBlock)
-					.Text(this, &SFilterConfiguratorRow::AvailableFilters_GetSelectionText)
-				]
-			];
-
-		// Operator combo box
-		LeftBox->AddSlot()
-			.HAlign(HAlign_Left)
-			.VAlign(VAlign_Center)
-			.FillWidth(1.0)
-			.Padding(FMargin(4.0f, 0.0f, 0.0f, 0.0f))
-			.AutoWidth()
-			[
-				SAssignNew(FilterOperatorComboBox, SComboBox<TSharedPtr<IFilterOperator>>)
-				.OptionsSource(GetAvailableFilterOperators())
-				.OnSelectionChanged(this, &SFilterConfiguratorRow::AvailableFilterOperators_OnSelectionChanged)
-				.OnGenerateWidget(this, &SFilterConfiguratorRow::AvailableFilterOperators_OnGenerateWidget)
-				[
-					SNew(STextBlock)
-					.Text(this, &SFilterConfiguratorRow::AvailableFilterOperators_GetSelectionText)
-				]
-			];
-
-		if (FilterConfiguratorNodePtr->GetSelectedFilter()->Is<FFilterWithSuggestions>())
+		TSharedPtr<FFilterState> FilterState = FilterConfiguratorNodePtr->GetSelectedFilterState();
+		check(FilterState.IsValid());
+		
+		if (!FilterState->HasCustomUI())
 		{
-			SuggestionTextBoxValue = FilterConfiguratorNodePtr->GetTextBoxValue();
-
+			// Filter combo box
 			LeftBox->AddSlot()
-				.AutoWidth()
-				.VAlign(VAlign_Center)
 				.HAlign(HAlign_Left)
-				.Padding(FMargin(4.0f, 0.0f, 4.0f, 0.0f))
+				.VAlign(VAlign_Center)
+				.AutoWidth()
+				.Padding(FMargin(0.0f, 2.0f))
 				[
-					SNew(SSuggestionTextBox)
-					.MinDesiredWidth(300.0f)
-					.ForegroundColor(FSlateColor::UseForeground())
-					.OnTextCommitted(this, &SFilterConfiguratorRow::SuggestionTextBox_OnValueCommitted)
-					.OnTextChanged(this, &SFilterConfiguratorRow::SuggestionTextBox_OnValueChanged)
-					.Text(this, &SFilterConfiguratorRow::SuggestionTextBox_GetValue)
-					.ToolTipText(this, &SFilterConfiguratorRow::GetTextBoxTooltipText)
-					.HintText(this, &SFilterConfiguratorRow::GetTextBoxHintText)
-					.OnShowingSuggestions(FOnShowingSuggestions::CreateSP(this, &SFilterConfiguratorRow::SuggestionTextBox_GetSuggestions))
-					.OnShowingHistory(FOnShowingHistory::CreateSP(this, &SFilterConfiguratorRow::SuggestionTextBox_GetHistory))
+					SAssignNew(FilterTypeComboBox, SComboBox<TSharedPtr<FFilter>>)
+					.OptionsSource(GetAvailableFilters())
+					.OnSelectionChanged(this, &SFilterConfiguratorRow::AvailableFilters_OnSelectionChanged)
+					.OnGenerateWidget(this, &SFilterConfiguratorRow::AvailableFilters_OnGenerateWidget)
+					[
+						SNew(STextBlock)
+						.Text(this, &SFilterConfiguratorRow::AvailableFilters_GetSelectionText)
+					]
 				];
-		}
-		else
-		{
+
+			// Operator combo box
 			LeftBox->AddSlot()
-				.AutoWidth()
-				.VAlign(VAlign_Center)
 				.HAlign(HAlign_Left)
+				.VAlign(VAlign_Center)
 				.FillWidth(1.0)
-				.Padding(FMargin(4.0f, 0.0f, 4.0f, 0.0f))
+				.Padding(FMargin(4.0f, 0.0f, 0.0f, 0.0f))
+				.AutoWidth()
 				[
-					SNew(SEditableTextBox)
-					.MinDesiredWidth(50.0f)
-					.OnTextCommitted(this, &SFilterConfiguratorRow::OnTextBoxValueCommitted)
-					.Text(this, &SFilterConfiguratorRow::GetTextBoxValue)
-					.ToolTipText(this, &SFilterConfiguratorRow::GetTextBoxTooltipText)
-					.HintText(this, &SFilterConfiguratorRow::GetTextBoxHintText)
-					.OnVerifyTextChanged(this, &SFilterConfiguratorRow::TextBox_OnVerifyTextChanged)
+					SAssignNew(FilterOperatorComboBox, SComboBox<TSharedPtr<IFilterOperator>>)
+					.OptionsSource(GetAvailableFilterOperators())
+					.OnSelectionChanged(this, &SFilterConfiguratorRow::AvailableFilterOperators_OnSelectionChanged)
+					.OnGenerateWidget(this, &SFilterConfiguratorRow::AvailableFilterOperators_OnGenerateWidget)
+					[
+						SNew(STextBlock)
+						.Text(this, &SFilterConfiguratorRow::AvailableFilterOperators_GetSelectionText)
+					]
 				];
+
+			if (FilterConfiguratorNodePtr->GetSelectedFilter()->Is<FFilterWithSuggestions>())
+			{
+				SuggestionTextBoxValue = FilterConfiguratorNodePtr->GetTextBoxValue();
+
+				LeftBox->AddSlot()
+					.AutoWidth()
+					.VAlign(VAlign_Center)
+					.HAlign(HAlign_Left)
+					.Padding(FMargin(4.0f, 0.0f, 4.0f, 0.0f))
+					[
+						SNew(SSuggestionTextBox)
+						.MinDesiredWidth(300.0f)
+						.ForegroundColor(FSlateColor::UseForeground())
+						.OnTextCommitted(this, &SFilterConfiguratorRow::SuggestionTextBox_OnValueCommitted)
+						.OnTextChanged(this, &SFilterConfiguratorRow::SuggestionTextBox_OnValueChanged)
+						.Text(this, &SFilterConfiguratorRow::SuggestionTextBox_GetValue)
+						.ToolTipText(this, &SFilterConfiguratorRow::GetTextBoxTooltipText)
+						.HintText(this, &SFilterConfiguratorRow::GetTextBoxHintText)
+						.OnShowingSuggestions(FOnShowingSuggestions::CreateSP(this, &SFilterConfiguratorRow::SuggestionTextBox_GetSuggestions))
+						.OnShowingHistory(FOnShowingHistory::CreateSP(this, &SFilterConfiguratorRow::SuggestionTextBox_GetHistory))
+					];
+			}
+			else
+			{
+				LeftBox->AddSlot()
+					.AutoWidth()
+					.VAlign(VAlign_Center)
+					.HAlign(HAlign_Left)
+					.FillWidth(1.0)
+					.Padding(FMargin(4.0f, 0.0f, 4.0f, 0.0f))
+					[
+						SNew(SEditableTextBox)
+						.MinDesiredWidth(50.0f)
+						.OnTextCommitted(this, &SFilterConfiguratorRow::OnTextBoxValueCommitted)
+						.Text(this, &SFilterConfiguratorRow::GetTextBoxValue)
+						.ToolTipText(this, &SFilterConfiguratorRow::GetTextBoxTooltipText)
+						.HintText(this, &SFilterConfiguratorRow::GetTextBoxHintText)
+						.OnVerifyTextChanged(this, &SFilterConfiguratorRow::TextBox_OnVerifyTextChanged)
+					];
+			}
+		}
+		else
+		{
+			FilterState->AddCustomUI(LeftBox);
 		}
 
 		RightBox->AddSlot()
diff --git a/Engine/Source/Developer/TraceInsights/Private/Insights/Widgets/STimingView.cpp b/Engine/Source/Developer/TraceInsights/Private/Insights/Widgets/STimingView.cpp
--- a/Engine/Source/Developer/TraceInsights/Private/Insights/Widgets/STimingView.cpp
+++ b/Engine/Source/Developer/TraceInsights/Private/Insights/Widgets/STimingView.cpp
@@ -4975,6 +4975,7 @@ void STimingView::QuickFind_Execute()
 			FFilterService::Get()->GetIntegerOperators()));
 
 		NewFilterConfigurator->Add(MakeShared<FTimerNameFilter>());
+		NewFilterConfigurator->Add(MakeShared<FMetadataFilter>());
 
 		for (Insights::ITimingViewExtender* Extender : GetExtenders())
 		{
diff --git a/Engine/Source/Developer/TraceServices/Public/Model/MonotonicTimeline.h b/Engine/Source/Developer/TraceServices/Public/Model/MonotonicTimeline.h
--- a/Engine/Source/Developer/TraceServices/Public/Model/MonotonicTimeline.h
+++ b/Engine/Source/Developer/TraceServices/Public/Model/MonotonicTimeline.h
@@ -663,7 +663,14 @@ public:
 		check(LastScopePageIndex >= FirstScopePageIndex);
 		uint32 NumPages = FMath::Max(static_cast<uint32>(LastScopePageIndex - FirstScopePageIndex), 1u);
 		uint32 NumThreads = GThreadPool->GetNumThreads();
+
+		if (EnumerateAsyncParams.MaxOccupancy > 0)
+		{
+			NumThreads = FMath::Max(static_cast<uint32>(NumThreads * EnumerateAsyncParams.MaxOccupancy), 1u);
+		}
+
 		uint32 NumTasks = FMath::Min(NumThreads, NumPages);
+
 		uint32 PagesPerTask = NumPages / NumTasks;
 		uint32 RemainingPages = NumPages % NumTasks; // The remaining pages will be split between the first tasks
 
diff --git a/Engine/Source/Developer/TraceServices/Public/TraceServices/Containers/Timelines.h b/Engine/Source/Developer/TraceServices/Public/TraceServices/Containers/Timelines.h
--- a/Engine/Source/Developer/TraceServices/Public/TraceServices/Containers/Timelines.h
+++ b/Engine/Source/Developer/TraceServices/Public/TraceServices/Containers/Timelines.h
@@ -42,6 +42,7 @@ public:
 		double IntervalStart = 0.0;
 		double IntervalEnd = -1.0;
 		double Resolution = 0;
+		double MaxOccupancy = 0.875; // Leave 1/8 threads free to avoid starvation.
 		EEventSortOrder SortOrder = EEventSortOrder::ByEndTime;
 		typename ITimeline<EventType>::TaskInfoCallback SetupCallback;
 		typename ITimeline<EventType>::AsyncEventRangeCallback Callback;
