[Insights]
  - Fixed crash when excuting TimingInsights.ExportTimerStatistics command with -region=*. In Editor traces, there are many (200000+) timing regions and some have names containing invalid characters for a file name. Using region name to make the output csv file fails. Fixed the issue by replacing all invalid characters with "_".
  - Fixed the slow enumeration of regions for TimingInsights.ExportTimerStatistics command.
  - Limited the number of export files to 100 for a single region (like "<SlowTask>") and to 10000 for the total number of csv files exported.

diff --git a/Engine/Source/Developer/TraceInsights/Private/Insights/ViewModels/TimingExporter.cpp b/Engine/Source/Developer/TraceInsights/Private/Insights/ViewModels/TimingExporter.cpp
--- a/Engine/Source/Developer/TraceInsights/Private/Insights/ViewModels/TimingExporter.cpp
+++ b/Engine/Source/Developer/TraceInsights/Private/Insights/ViewModels/TimingExporter.cpp
@@ -554,105 +554,162 @@ int32 FTimingExporter::ExportTimingEventsAsText(const FString& Filename, FExport
 
 ////////////////////////////////////////////////////////////////////////////////////////////////////
 
-int32 FTimingExporter::ExportTimerStatisticsAsText(const FString& Filename, FExportTimerStatisticsParams& Params) const
+int32 FTimingExporter::ExportTimerStatisticsAsTextByRegions(const FString& Filename, FExportTimerStatisticsParams& Params) const
 {
-	TraceServices::ITable<TraceServices::FTimingProfilerAggregatedStats>* StatsTable;
-
-	if (!Params.Region.IsEmpty())
+	class FRegionNameSpec
 	{
-		class FRegionNameSpec
+	public:
+		FRegionNameSpec(FString& InNamePatternList)
 		{
-		public:
-			FRegionNameSpec(FString& InNamePatternList)
-			{
-				InNamePatternList.ParseIntoArray(NamePatterns, TEXT(","), true);
-			}
+			InNamePatternList.ParseIntoArray(NamePatterns, TEXT(","), true);
+		}
 
-			bool Match(const FString& InRegionName)
+		bool Match(const FString& InRegionName)
+		{
+			for (const FString& NamePattern : NamePatterns)
 			{
-				for (const FString& NamePattern : NamePatterns)
+				if (InRegionName.MatchesWildcard(NamePattern))
 				{
-					if (InRegionName.MatchesWildcard(NamePattern))
-					{
-						return true;
-					}
+					return true;
 				}
-				return false;
 			}
+			return false;
+		}
 
-		private:
-			TArray<FString> NamePatterns;
-		};
-		FRegionNameSpec RegionNameSpec(Params.Region);
+	private:
+		TArray<FString> NamePatterns;
+	};
+	FRegionNameSpec RegionNameSpec(Params.Region);
 
-		TMap<FString, TraceServices::FTimeRegion> Regions;
+	struct FTimeRegionInterval
+	{
+		double StartTime;
+		double EndTime;
+	};
+	struct FTimeRegionGroup
+	{
+		TArray<FTimeRegionInterval> Intervals;
+	};
+	TMap<FString, FTimeRegionGroup> RegionGroups;
 
-		// Detect regions
-		{
-			const TraceServices::IRegionProvider& RegionProvider = TraceServices::ReadRegionProvider(Session);
-			TraceServices::FProviderReadScopeLock RegionProviderScopedLock(RegionProvider);
+	// Detect regions
+	int32 RegionCount = 0;
+	FStopwatch DetectRegionsStopwatch;
+	DetectRegionsStopwatch.Start();
+	{
+		const TraceServices::IRegionProvider& RegionProvider = TraceServices::ReadRegionProvider(Session);
+		TraceServices::FProviderReadScopeLock RegionProviderScopedLock(RegionProvider);
 
-			UE_LOG(TraceInsights, Log, TEXT("Looking for regions: '%s'"), *Params.Region);
+		UE_LOG(TraceInsights, Log, TEXT("Looking for regions: '%s'"), *Params.Region);
 
-			RegionProvider.EnumerateRegions(0.0, std::numeric_limits<double>::max(),
-				[&RegionNameSpec, &Regions](const TraceServices::FTimeRegion& InRegion) -> bool
+		RegionProvider.EnumerateRegions(0.0, std::numeric_limits<double>::max(),
+			[&RegionCount, &RegionNameSpec, &RegionGroups](const TraceServices::FTimeRegion& InRegion) -> bool
+			{
+				if (!RegionNameSpec.Match(InRegion.Text))
 				{
-					if (!RegionNameSpec.Match(InRegion.Text))
-					{
-						return true;
-					}
-
-					// Handle duplicate region names, individual regions may appear multiple times
-					// we append numbers to allow for unique export filenames.
-					FString RegionName = InRegion.Text;
-					int32 Index = 2;
-
-					while (true)
-					{
-						const TraceServices::FTimeRegion* ExistingRegion = Regions.Find(RegionName);
+					return true;
+				}
 
-						if (!ExistingRegion)
-						{
-							Regions.Add(RegionName, InRegion);
-							break;
-						}
+				// Handle duplicate region names, individual regions may appear multiple times
+				// we append numbers to allow for unique export filenames.
+				FString RegionName = InRegion.Text;
+				FTimeRegionGroup* ExistingRegionGroup = RegionGroups.Find(RegionName);
+				if (!ExistingRegionGroup)
+				{
+					ExistingRegionGroup = &RegionGroups.Add(RegionName, FTimeRegionGroup{});
+				}
+				ExistingRegionGroup->Intervals.Add(FTimeRegionInterval{ InRegion.BeginTime, InRegion.EndTime });
+				++RegionCount;
+				return true;
+			});
+	}
+	DetectRegionsStopwatch.Stop();
+	const double DetectRegionsTime = DetectRegionsStopwatch.GetAccumulatedTime();
+	UE_LOG(TraceInsights, Display, TEXT("Detected %d regions in %.3fs."), RegionCount, DetectRegionsTime);
 
-						RegionName = FString::Printf(TEXT("%s%d"), *RegionName, Index++);
-					}
+	if (RegionGroups.Num() == 0)
+	{
+		UE_LOG(TraceInsights, Error, TEXT("Unable to find any region with name pattern '%s'."), *Params.Region);
+		return -1;
+	}
 
-					return true;
-				});
-		}
+	FStopwatch Stopwatch;
+	Stopwatch.Start();
 
-		if (Regions.Num() == 0)
+	// Export timing statistics for each region.
+	constexpr int32 MaxIntervalsPerRegion = 100;
+	constexpr int32 MaxExportedRegions = 10000;
+	int32 ExportedRegionCount = 0;
+	for (auto& KV : RegionGroups)
+	{
+		FString RegionName = KV.Key;
+		const FString InvalidFileSystemChars = FPaths::GetInvalidFileSystemChars();
+		for (int32 CharIndex = 0; CharIndex < InvalidFileSystemChars.Len(); CharIndex++)
 		{
-			UE_LOG(TraceInsights, Error, TEXT("Unable to find any region with name pattern '%s'."), *Params.Region);
-			return -1;
+			FString Char = FString().AppendChar(InvalidFileSystemChars[CharIndex]);
+			RegionName.ReplaceInline(*Char, TEXT("_"));
 		}
+		RegionName.TrimStartAndEndInline();
 
-		FStopwatch Stopwatch;
-		Stopwatch.Start();
-
-		// Export timing statistics for each region.
-		for (auto& KV : Regions)
+		int32 IntervalIndex = 0;
+		for (const FTimeRegionInterval& Interval : KV.Value.Intervals)
 		{
-			FString RegionFilename = Filename.Replace(TEXT("*"), *KV.Key);
+			FString RegionFilename;
+			if (IntervalIndex == 0)
+			{
+				RegionFilename = Filename.Replace(TEXT("*"), *RegionName);
+			}
+			else
+			{
+				FString UniqueRegionName = FString::Printf(TEXT("%s_%d"), *RegionName, IntervalIndex);
+				RegionFilename = Filename.Replace(TEXT("*"), *UniqueRegionName);
+			}
+			++IntervalIndex;
+
 			FExportTimerStatisticsParams RegionParams = Params;
 			RegionParams.Region.Reset();
-			RegionParams.IntervalStartTime = KV.Value.BeginTime;
-			RegionParams.IntervalEndTime = KV.Value.EndTime;
+			RegionParams.IntervalStartTime = Interval.StartTime;
+			RegionParams.IntervalEndTime = Interval.EndTime;
 			UE_LOG(TraceInsights, Display, TEXT("Exporting timing statistics for region '%s' [%f .. %f] to '%s'"), *KV.Key, RegionParams.IntervalStartTime, RegionParams.IntervalEndTime, *RegionFilename);
 			ExportTimerStatisticsAsText(RegionFilename, RegionParams);
+
+			++ExportedRegionCount;
+
+			// Avoid writing too many files...
+			if (IntervalIndex >= MaxIntervalsPerRegion)
+			{
+				UE_LOG(TraceInsights, Error, TEXT("Too many intervals for region '%s'! Exporting timing statistics to separate file per interval for this region is not allowed to continue."), *KV.Key);
+				break;
+			}
+			if (ExportedRegionCount >= MaxExportedRegions)
+			{
+				UE_LOG(TraceInsights, Error, TEXT("Too many regions! Exporting timing statistics to separate file per region is not allowed to continue."));
+				break;
+			}
 		}
 
-		Stopwatch.Stop();
-		const double TotalTime = Stopwatch.GetAccumulatedTime();
-		UE_LOG(TraceInsights, Log, TEXT("Exported timing statistics for %d regions in %.3fs."), Regions.Num(), TotalTime);
-		return Regions.Num();
+		if (ExportedRegionCount >= MaxExportedRegions)
+		{
+			break;
+		}
 	}
 
-	check(Params.Region.IsEmpty());
+	Stopwatch.Stop();
+	const double TotalTime = Stopwatch.GetAccumulatedTime();
+	UE_LOG(TraceInsights, Log, TEXT("Exported timing statistics for %d regions in %.3fs."), ExportedRegionCount, TotalTime);
+	return ExportedRegionCount;
+}
+
+////////////////////////////////////////////////////////////////////////////////////////////////////
 
+int32 FTimingExporter::ExportTimerStatisticsAsText(const FString& Filename, FExportTimerStatisticsParams& Params) const
+{
+	if (!Params.Region.IsEmpty())
+	{
+		return ExportTimerStatisticsAsTextByRegions(Filename, Params);
+	}
+
+	TraceServices::ITable<TraceServices::FTimingProfilerAggregatedStats>* StatsTable;
 	{
 		TraceServices::FAnalysisSessionReadScope SessionReadScope(Session);
 
diff --git a/Engine/Source/Developer/TraceInsights/Private/Insights/ViewModels/TimingExporter.h b/Engine/Source/Developer/TraceInsights/Private/Insights/ViewModels/TimingExporter.h
--- a/Engine/Source/Developer/TraceInsights/Private/Insights/ViewModels/TimingExporter.h
+++ b/Engine/Source/Developer/TraceInsights/Private/Insights/ViewModels/TimingExporter.h
@@ -205,6 +205,7 @@ private:
 	void ExportTimingEvents_InitColumns() const;
 	void ExportTimingEvents_WriteHeader(FExportTimingEventsInternalParams& Params) const;
 	int32 ExportTimingEvents_WriteEvents(FExportTimingEventsInternalParams& Params) const;
+	int32 ExportTimerStatisticsAsTextByRegions(const FString& Filename, FExportTimerStatisticsParams& Params) const;
 
 private:
 	const TraceServices::IAnalysisSession& Session;
