Fix static analysis error.

diff --git a/Engine/Source/Developer/TraceInsights/Private/Insights/ViewModels/ThreadTimingTrack.cpp b/Engine/Source/Developer/TraceInsights/Private/Insights/ViewModels/ThreadTimingTrack.cpp
--- a/Engine/Source/Developer/TraceInsights/Private/Insights/ViewModels/ThreadTimingTrack.cpp
+++ b/Engine/Source/Developer/TraceInsights/Private/Insights/ViewModels/ThreadTimingTrack.cpp
@@ -1117,7 +1117,7 @@ void FThreadTimingTrack::BuildFilteredDrawState(ITimingEventsTrackDrawStateBuild
 					{
 						if (FilterConfigurator->IsKeyUsed(static_cast<int32>(EFilterField::Metadata)))
 						{
-							Params.MaxOccupancy = 0.75; // This filter can be slow so reduce occupancy to avoid starvation.
+							Params.MaxOccupancy = 0.75f; // This filter can be slow so reduce occupancy to avoid starvation.
 						}
 					}
 					Params.IntervalStart = Viewport.GetStartTime();
diff --git a/Engine/Source/Developer/TraceServices/Public/Model/MonotonicTimeline.h b/Engine/Source/Developer/TraceServices/Public/Model/MonotonicTimeline.h
--- a/Engine/Source/Developer/TraceServices/Public/Model/MonotonicTimeline.h
+++ b/Engine/Source/Developer/TraceServices/Public/Model/MonotonicTimeline.h
@@ -664,7 +664,7 @@ public:
 		uint32 NumPages = FMath::Max(static_cast<uint32>(LastScopePageIndex - FirstScopePageIndex), 1u);
 		uint32 NumThreads = GThreadPool->GetNumThreads();
 
-		if (EnumerateAsyncParams.MaxOccupancy > 0)
+		if (EnumerateAsyncParams.MaxOccupancy > 0.0f)
 		{
 			NumThreads = FMath::Max(static_cast<uint32>(NumThreads * EnumerateAsyncParams.MaxOccupancy), 1u);
 		}
