Store browser caches used ports

diff --git a/Engine/Source/Developer/TraceInsights/Private/Insights/StoreService/StoreBrowser.cpp b/Engine/Source/Developer/TraceInsights/Private/Insights/StoreService/StoreBrowser.cpp
--- a/Engine/Source/Developer/TraceInsights/Private/Insights/StoreService/StoreBrowser.cpp
+++ b/Engine/Source/Developer/TraceInsights/Private/Insights/StoreService/StoreBrowser.cpp
@@ -216,31 +216,39 @@ void FStoreBrowser::UpdateTraces()
 		TArray<FString> NewWatchDirectories;
 
 		{
-			// Get Trace Store status.
-			FScopeLock StoreClientLock(&GetStoreClientCriticalSection());
-			const UE::Trace::FStoreClient::FStatus* Status = StoreClient->GetStatus();
-			if (Status)
+			const UE::Trace::FStoreClient::FStatus* Status = nullptr;
 			{
-				NewStoreChangeSerial = Status->GetChangeSerial();
-				NewStoreSettingsChangeSerial = Status->GetSettingsSerial();
-				if (StoreSettingsChangeSerial != NewStoreSettingsChangeSerial)
+				// Get Trace Store status.
+				FScopeLock StoreClientLock(&GetStoreClientCriticalSection());
+				Status = StoreClient->GetStatus();
+				if (Status)
 				{
-					NewStoreDirectory = FString(Status->GetStoreDir());
-					Status->GetWatchDirectories(NewWatchDirectories);
+					NewStoreChangeSerial = Status->GetChangeSerial();
+					NewStoreSettingsChangeSerial = Status->GetSettingsSerial();
+					if (StoreSettingsChangeSerial != NewStoreSettingsChangeSerial)
+					{
+						NewStoreDirectory = FString(Status->GetStoreDir());
+						Status->GetWatchDirectories(NewWatchDirectories);
+					}
 				}
 			}
-		}
 
-		if (StoreSettingsChangeSerial != NewStoreSettingsChangeSerial)
-		{
-			StoreSettingsChangeSerial = NewStoreSettingsChangeSerial;
+			if (StoreSettingsChangeSerial != NewStoreSettingsChangeSerial)
+			{
+				StoreSettingsChangeSerial = NewStoreSettingsChangeSerial;
 
-			// Update settings.
-			FScopeLock Lock(&SettingsCriticalSection);
-			SettingsChangeSerial++;
-			StoreDirectory = NewStoreDirectory;
-			WatchDirectories.Empty();
-			WatchDirectories = NewWatchDirectories;
+				// Update settings.
+				FScopeLock Lock(&SettingsCriticalSection);
+				SettingsChangeSerial++;
+				StoreDirectory = NewStoreDirectory;
+				WatchDirectories.Empty();
+				WatchDirectories = NewWatchDirectories;
+				if (Status)
+				{
+					StorePort = Status->GetStorePort();
+					RecorderPort = Status->GetRecorderPort();
+				}
+			}
 		}
 
 		if (StoreChangeSerial != NewStoreChangeSerial)
diff --git a/Engine/Source/Developer/TraceInsights/Private/Insights/StoreService/StoreBrowser.h b/Engine/Source/Developer/TraceInsights/Private/Insights/StoreService/StoreBrowser.h
--- a/Engine/Source/Developer/TraceInsights/Private/Insights/StoreService/StoreBrowser.h
+++ b/Engine/Source/Developer/TraceInsights/Private/Insights/StoreService/StoreBrowser.h
@@ -101,6 +101,8 @@ public:
 	uint32 GetSettingsChangeSerial() const { check(bSettingsLocked); return SettingsChangeSerial; };
 	const FString& GetHost() const { check(bSettingsLocked); return Host; }
 	const FString& GetVersion() const { check(bSettingsLocked); return Version; }
+	uint32 GetRecorderPort() const {check(bSettingsLocked); return RecorderPort; }
+	uint32 GetStorePort() const { check(bSettingsLocked); return StorePort; }
 	const FString& GetStoreDirectory() const { check(bSettingsLocked); return StoreDirectory; }
 	const TArray<FString>& GetWatchDirectories() const { check(bSettingsLocked); return WatchDirectories; }
 
@@ -139,6 +141,8 @@ private:
 	uint32 SettingsChangeSerial = 0;
 	FString Host;
 	FString Version;
+	uint32 StorePort;
+	uint32 RecorderPort;
 	FString StoreDirectory;
 	TArray<FString> WatchDirectories;
 
