[Insights] Fix several bugs in the filtering system.

diff --git a/Engine/Source/Developer/TraceInsights/Private/Insights/Table/Widgets/STableTreeView.cpp b/Engine/Source/Developer/TraceInsights/Private/Insights/Table/Widgets/STableTreeView.cpp
--- a/Engine/Source/Developer/TraceInsights/Private/Insights/Table/Widgets/STableTreeView.cpp
+++ b/Engine/Source/Developer/TraceInsights/Private/Insights/Table/Widgets/STableTreeView.cpp
@@ -1140,7 +1140,7 @@ void STableTreeView::UpdateFilterContext(const FFilterConfigurator& InFilterConf
 {
 	for (const TSharedRef<FTableColumn>& Column : Table->GetColumns())
 	{
-		if (!Column->CanBeFiltered() || !InFilterConfigurator.IsKeyUsed(Column->GetIndex()))
+		if (!Column->CanBeFiltered() || !InFilterConfigurator.IsKeyUsed(Column->GetIndex()) || !Column->GetValue(InNode).IsSet())
 		{
 			continue;
 		}
diff --git a/Engine/Source/Developer/TraceInsights/Private/Insights/ViewModels/FilterConfiguratorNode.cpp b/Engine/Source/Developer/TraceInsights/Private/Insights/ViewModels/FilterConfiguratorNode.cpp
--- a/Engine/Source/Developer/TraceInsights/Private/Insights/ViewModels/FilterConfiguratorNode.cpp
+++ b/Engine/Source/Developer/TraceInsights/Private/Insights/ViewModels/FilterConfiguratorNode.cpp
@@ -37,13 +37,8 @@ TSharedRef<FFilterConfiguratorNode> FFilterConfiguratorNode::DeepCopy(const FFil
 	NodeCopy.AvailableFilters = InSourceNode.AvailableFilters;
 	NodeCopy.AvailableFilterOperators = InSourceNode.AvailableFilterOperators;
 	NodeCopy.SelectedFilter = InSourceNode.SelectedFilter;
-	NodeCopy.SelectedFilterOperator = InSourceNode.SelectedFilterOperator;
 	NodeCopy.SelectedFilterGroupOperator = InSourceNode.SelectedFilterGroupOperator;
 	NodeCopy.TextBoxValue = InSourceNode.TextBoxValue;
-	if (InSourceNode.SelectedFilter.IsValid())
-	{
-		NodeCopy.FilterState = InSourceNode.SelectedFilter->BuildFilterState(*InSourceNode.FilterState);
-	}
 
 	NodeCopy.SetExpansion(InSourceNode.IsExpanded());
 
@@ -58,6 +53,10 @@ TSharedRef<FFilterConfiguratorNode> FFilterConfiguratorNode::DeepCopy(const FFil
 			NodeCopy.AddChildAndSetParent(ChildCopy);
 		}
 	}
+	else
+	{
+		NodeCopy.FilterState = InSourceNode.FilterState->DeepCopy();
+	}
 
 	return NodeCopyPtr;
 }
@@ -67,9 +66,24 @@ TSharedRef<FFilterConfiguratorNode> FFilterConfiguratorNode::DeepCopy(const FFil
 bool FFilterConfiguratorNode::operator==(const FFilterConfiguratorNode& Other) const
 {
 	bool bIsEqual = true;
+	
+	if (!IsGroup())
+	{
+		check(FilterState.IsValid());
+		check(Other.FilterState.IsValid());
+
+		if (FilterState->GetTypeName() == Other.FilterState->GetTypeName())
+		{
+			bIsEqual &= FilterState->Equals(*Other.FilterState);
+		}
+		else
+		{
+			return false;
+		}
+	}
+
 	bIsEqual &= AvailableFilters.Get() == Other.AvailableFilters.Get();
 	bIsEqual &= SelectedFilter.Get() == Other.SelectedFilter.Get();
-	bIsEqual &= SelectedFilterOperator.Get() == Other.SelectedFilterOperator.Get();
 	bIsEqual &= SelectedFilterGroupOperator.Get() == Other.SelectedFilterGroupOperator.Get();
 	bIsEqual &= TextBoxValue == Other.TextBoxValue;
 	bIsEqual &= GetChildrenCount() == Other.GetChildrenCount();
@@ -88,7 +102,7 @@ bool FFilterConfiguratorNode::operator==(const FFilterConfiguratorNode& Other) c
 
 ////////////////////////////////////////////////////////////////////////////////////////////////////
 
-const TArray<TSharedPtr<FFilterGroupOperator>>& FFilterConfiguratorNode::GetFilterGroupOperators()
+const TArray<TSharedPtr<FFilterGroupOperator>>& FFilterConfiguratorNode::GetFilterGroupOperators() const
 {
 	return FFilterService::Get()->GetFilterGroupOperators();
 }
@@ -257,7 +271,7 @@ void FFilterConfiguratorNode::SetSelectedFilterOperator(TSharedPtr<IFilterOperat
 
 ////////////////////////////////////////////////////////////////////////////////////////////////////
 
-TSharedPtr<IFilterOperator> FFilterConfiguratorNode::GetSelectedFilterOperator() const
+TSharedPtr<const IFilterOperator> FFilterConfiguratorNode::GetSelectedFilterOperator() const
 {
 	if (FilterState.IsValid())
 	{
diff --git a/Engine/Source/Developer/TraceInsights/Private/Insights/ViewModels/FilterConfiguratorNode.h b/Engine/Source/Developer/TraceInsights/Private/Insights/ViewModels/FilterConfiguratorNode.h
--- a/Engine/Source/Developer/TraceInsights/Private/Insights/ViewModels/FilterConfiguratorNode.h
+++ b/Engine/Source/Developer/TraceInsights/Private/Insights/ViewModels/FilterConfiguratorNode.h
@@ -64,9 +64,9 @@ public:
 	TSharedPtr<FFilter> GetSelectedFilter() const { return SelectedFilter; }
 
 	void SetSelectedFilterOperator(TSharedPtr<IFilterOperator> InSelectedFilterOperator);
-	TSharedPtr<IFilterOperator> GetSelectedFilterOperator() const;
+	TSharedPtr<const IFilterOperator> GetSelectedFilterOperator() const;
 
-	const TArray<TSharedPtr<FFilterGroupOperator>>& GetFilterGroupOperators();
+	const TArray<TSharedPtr<FFilterGroupOperator>>& GetFilterGroupOperators() const;
 	void SetSelectedFilterGroupOperator(TSharedPtr<FFilterGroupOperator> InSelectedFilterGroupOperator) { SelectedFilterGroupOperator = InSelectedFilterGroupOperator; }
 	TSharedPtr<FFilterGroupOperator> GetSelectedFilterGroupOperator() const { return SelectedFilterGroupOperator; }
 
@@ -92,8 +92,6 @@ private:
 
 	TSharedPtr<FFilter> SelectedFilter;
 
-	TSharedPtr<IFilterOperator> SelectedFilterOperator;
-
 	TSharedPtr<FFilterGroupOperator> SelectedFilterGroupOperator;
 
 	TSharedPtr<TArray<TSharedPtr<IFilterOperator>>> AvailableFilterOperators;
diff --git a/Engine/Source/Developer/TraceInsights/Private/Insights/ViewModels/Filters.cpp b/Engine/Source/Developer/TraceInsights/Private/Insights/ViewModels/Filters.cpp
--- a/Engine/Source/Developer/TraceInsights/Private/Insights/ViewModels/Filters.cpp
+++ b/Engine/Source/Developer/TraceInsights/Private/Insights/ViewModels/Filters.cpp
@@ -27,24 +27,27 @@ INSIGHTS_IMPLEMENT_RTTI(FCustomFilter)
 FFilterStorage::FFilterStorage()
 {
 	DoubleOperators = MakeShared<TArray<TSharedPtr<IFilterOperator>>>();
-	DoubleOperators->Add(StaticCastSharedRef<IFilterOperator>(MakeShared<FFilterOperator<double>>(EFilterOperator::Lt, TEXT("<"), std::less<>{})));
-	DoubleOperators->Add(StaticCastSharedRef<IFilterOperator>(MakeShared<FFilterOperator<double>>(EFilterOperator::Lte, TEXT("\u2264"), std::less_equal<>())));
-	DoubleOperators->Add(StaticCastSharedRef<IFilterOperator>(MakeShared<FFilterOperator<double>>(EFilterOperator::Eq, TEXT("="), std::equal_to<>())));
-	DoubleOperators->Add(StaticCastSharedRef<IFilterOperator>(MakeShared<FFilterOperator<double>>(EFilterOperator::Gt, TEXT(">"), std::greater<>())));
-	DoubleOperators->Add(StaticCastSharedRef<IFilterOperator>(MakeShared<FFilterOperator<double>>(EFilterOperator::Gte, TEXT("\u2265"), std::greater_equal<>())));
+	DoubleOperators->Add(MakeShared<FFilterOperator<double>>(EFilterOperator::Eq, TEXT("="), std::equal_to<>()));
+	DoubleOperators->Add(MakeShared<FFilterOperator<double>>(EFilterOperator::NotEq, TEXT("!="), std::not_equal_to<>()));
+	DoubleOperators->Add(MakeShared<FFilterOperator<double>>(EFilterOperator::Lt, TEXT("<"), std::less<>{}));
+	DoubleOperators->Add(MakeShared<FFilterOperator<double>>(EFilterOperator::Lte, TEXT("\u2264"), std::less_equal<>()));
+	DoubleOperators->Add(MakeShared<FFilterOperator<double>>(EFilterOperator::Gt, TEXT(">"), std::greater<>()));
+	DoubleOperators->Add(MakeShared<FFilterOperator<double>>(EFilterOperator::Gte, TEXT("\u2265"), std::greater_equal<>()));
 
 	IntegerOperators = MakeShared<TArray<TSharedPtr<IFilterOperator>>>();
-	IntegerOperators->Add(StaticCastSharedRef<IFilterOperator>(MakeShared<FFilterOperator<int64>>(EFilterOperator::Lt, TEXT("<"), std::less<>{})));
-	IntegerOperators->Add(StaticCastSharedRef<IFilterOperator>(MakeShared<FFilterOperator<int64>>(EFilterOperator::Lte, TEXT("\u2264"), std::less_equal<>())));
-	IntegerOperators->Add(StaticCastSharedRef<IFilterOperator>(MakeShared<FFilterOperator<int64>>(EFilterOperator::Eq, TEXT("="), std::equal_to<>())));
-	IntegerOperators->Add(StaticCastSharedRef<IFilterOperator>(MakeShared<FFilterOperator<int64>>(EFilterOperator::Gt, TEXT(">"), std::greater<>())));
-	IntegerOperators->Add(StaticCastSharedRef<IFilterOperator>(MakeShared<FFilterOperator<int64>>(EFilterOperator::Gte, TEXT("\u2265"), std::greater_equal<>())));
+	IntegerOperators->Add(MakeShared<FFilterOperator<int64>>(EFilterOperator::Eq, TEXT("="), std::equal_to<>()));
+	IntegerOperators->Add(MakeShared<FFilterOperator<int64>>(EFilterOperator::NotEq, TEXT("!="), std::not_equal_to<>()));
+	IntegerOperators->Add(MakeShared<FFilterOperator<int64>>(EFilterOperator::Lt, TEXT("<"), std::less<>{}));
+	IntegerOperators->Add(MakeShared<FFilterOperator<int64>>(EFilterOperator::Lte, TEXT("\u2264"), std::less_equal<>()));
+	IntegerOperators->Add(MakeShared<FFilterOperator<int64>>(EFilterOperator::Eq, TEXT("="), std::equal_to<>()));
+	IntegerOperators->Add(MakeShared<FFilterOperator<int64>>(EFilterOperator::Gt, TEXT(">"), std::greater<>()));
+	IntegerOperators->Add(MakeShared<FFilterOperator<int64>>(EFilterOperator::Gte, TEXT("\u2265"), std::greater_equal<>()));
 
 	StringOperators = MakeShared<TArray<TSharedPtr<IFilterOperator>>>();
-	StringOperators->Add(StaticCastSharedRef<IFilterOperator>(MakeShared<FFilterOperator<FString>>(EFilterOperator::Eq, TEXT("IS"), [](const FString& lhs, const FString& rhs) { return lhs.Equals(rhs); })));
-	StringOperators->Add(StaticCastSharedRef<IFilterOperator>(MakeShared<FFilterOperator<FString>>(EFilterOperator::NotEq, TEXT("IS NOT"), [](const FString& lhs, const FString& rhs) { return !lhs.Equals(rhs); })));
-	StringOperators->Add(StaticCastSharedRef<IFilterOperator>(MakeShared<FFilterOperator<FString>>(EFilterOperator::Contains, TEXT("CONTAINS"), [](const FString& lhs, const FString& rhs) { return lhs.Contains(rhs); })));
-	StringOperators->Add(StaticCastSharedRef<IFilterOperator>(MakeShared<FFilterOperator<FString>>(EFilterOperator::NotContains, TEXT("NOT CONTAINS"), [](const FString& lhs, const FString& rhs) { return !lhs.Contains(rhs); })));
+	StringOperators->Add(MakeShared<FFilterOperator<FString>>(EFilterOperator::Eq, TEXT("IS"), [](const FString& lhs, const FString& rhs) { return lhs.Equals(rhs); }));
+	StringOperators->Add(MakeShared<FFilterOperator<FString>>(EFilterOperator::NotEq, TEXT("IS NOT"), [](const FString& lhs, const FString& rhs) { return !lhs.Equals(rhs); }));
+	StringOperators->Add(MakeShared<FFilterOperator<FString>>(EFilterOperator::Contains, TEXT("CONTAINS"), [](const FString& lhs, const FString& rhs) { return lhs.Contains(rhs); }));
+	StringOperators->Add(MakeShared<FFilterOperator<FString>>(EFilterOperator::NotContains, TEXT("NOT CONTAINS"), [](const FString& lhs, const FString& rhs) { return !lhs.Contains(rhs); }));
 
 	FilterGroupOperators.Add(MakeShared<FFilterGroupOperator>(EFilterGroupOperator::And, LOCTEXT("AllOf", "All Of (AND)"), LOCTEXT("AllOfDesc", "All of the children must be true for the group to return true. Equivalent to an AND operation.")));
 	FilterGroupOperators.Add(MakeShared<FFilterGroupOperator>(EFilterGroupOperator::Or, LOCTEXT("AnyOf", "Any Of (OR)"), LOCTEXT("AnyOfDesc", "Any of the children must be true for the group to return true. Equivalent to an OR operation.")));
@@ -261,6 +264,20 @@ bool FFilterState::ApplyFilter(const FFilterContext& Context) const
 
 ////////////////////////////////////////////////////////////////////////////////////////////////////
 
+bool FFilterState::Equals(const FFilterState& Other) const
+{
+	return SelectedOperator == Other.GetSelectedOperator();
+}
+
+////////////////////////////////////////////////////////////////////////////////////////////////////
+
+TSharedRef<FFilterState> FFilterState::DeepCopy() const
+{
+	return MakeShared<FFilterState>(*this);
+}
+
+////////////////////////////////////////////////////////////////////////////////////////////////////
+
 } // namespace Insights
 
 #undef LOCTEXT_NAMESPACE
\ No newline at end of file
diff --git a/Engine/Source/Developer/TraceInsights/Private/Insights/ViewModels/Filters.h b/Engine/Source/Developer/TraceInsights/Private/Insights/ViewModels/Filters.h
--- a/Engine/Source/Developer/TraceInsights/Private/Insights/ViewModels/Filters.h
+++ b/Engine/Source/Developer/TraceInsights/Private/Insights/ViewModels/Filters.h
@@ -190,7 +190,7 @@ public:
 	virtual ~FFilterState() {}
 
 	void SetSelectedOperator(TSharedPtr<IFilterOperator> InOperator) { SelectedOperator = InOperator; }
-	TSharedPtr<IFilterOperator> GetSelectedOperator() { return SelectedOperator; }
+	TSharedPtr<const IFilterOperator> GetSelectedOperator() const { return SelectedOperator; }
 
 	virtual void Update() {};
 	virtual bool ApplyFilter(const FFilterContext& Context) const;
@@ -200,6 +200,9 @@ public:
 	virtual bool HasCustomUI() const { return false; }
 	virtual void AddCustomUI(TSharedRef<SHorizontalBox> LeftBox) {}
 
+	virtual bool Equals(const FFilterState& Other) const;
+	virtual TSharedRef<FFilterState> DeepCopy() const;
+
 protected:
 	TSharedRef<FFilter> Filter;
 	TSharedPtr<IFilterOperator> SelectedOperator;
diff --git a/Engine/Source/Developer/TraceInsights/Private/Insights/ViewModels/TimerFilters.cpp b/Engine/Source/Developer/TraceInsights/Private/Insights/ViewModels/TimerFilters.cpp
--- a/Engine/Source/Developer/TraceInsights/Private/Insights/ViewModels/TimerFilters.cpp
+++ b/Engine/Source/Developer/TraceInsights/Private/Insights/ViewModels/TimerFilters.cpp
@@ -79,6 +79,28 @@ bool FTimerNameFilterState::ApplyFilter(const FFilterContext& Context) const
 	return TimerIds.Contains(static_cast<uint32>(Value));
 }
 
+////////////////////////////////////////////////////////////////////////////////////////////////////
+
+bool FTimerNameFilterState::Equals(const FFilterState& Other) const
+{
+	if (this->GetTypeName() != Other.GetTypeName())
+	{
+		return false;
+	}
+
+	const FTimerNameFilterState& OtherTimerNameFilter = StaticCast<const FTimerNameFilterState&>(Other);
+	return FilterValue == OtherTimerNameFilter.FilterValue;
+}
+
+////////////////////////////////////////////////////////////////////////////////////////////////////
+
+TSharedRef<FFilterState> FTimerNameFilterState::DeepCopy() const
+{
+	TSharedRef<FTimerNameFilterState> Copy = MakeShared<FTimerNameFilterState>(*this);
+
+	return Copy;
+}
+
 ////////////////////////////////////////////////////////////////////////////////////////////////////
 // FTimerNameFilter
 ////////////////////////////////////////////////////////////////////////////////////////////////////
@@ -601,6 +623,35 @@ void FMetadataFilterState::OnTermTextBoxValueCommitted(const FText& InNewText, E
 	Term = InNewText.ToString();
 }
 
+////////////////////////////////////////////////////////////////////////////////////////////////////
+
+bool FMetadataFilterState::Equals(const FFilterState& Other) const
+{
+	if (this->GetTypeName() != Other.GetTypeName())
+	{
+		return false;
+	}
+
+	bool bIsEqual = true;
+
+	const FMetadataFilterState& OtherMetadataFilter = StaticCast<const FMetadataFilterState&>(Other);
+	
+	bIsEqual &= Key.Equals(OtherMetadataFilter.Key);
+	bIsEqual &= Term.Equals(OtherMetadataFilter.Term);
+	bIsEqual &= SelectedDataType == OtherMetadataFilter.SelectedDataType;
+
+	return bIsEqual;
+}
+
+////////////////////////////////////////////////////////////////////////////////////////////////////
+
+TSharedRef<FFilterState> FMetadataFilterState::DeepCopy() const
+{
+	TSharedRef<FMetadataFilterState> Copy = MakeShared<FMetadataFilterState>(*this);
+	
+	return Copy;
+}
+
 ////////////////////////////////////////////////////////////////////////////////////////////////////
 // FMetadataFilter
 ////////////////////////////////////////////////////////////////////////////////////////////////////
diff --git a/Engine/Source/Developer/TraceInsights/Private/Insights/ViewModels/TimerFilters.h b/Engine/Source/Developer/TraceInsights/Private/Insights/ViewModels/TimerFilters.h
--- a/Engine/Source/Developer/TraceInsights/Private/Insights/ViewModels/TimerFilters.h
+++ b/Engine/Source/Developer/TraceInsights/Private/Insights/ViewModels/TimerFilters.h
@@ -25,6 +25,12 @@ public:
 		: FFilterState(InFilter)
 	{}
 
+	FTimerNameFilterState(const FTimerNameFilterState& Other)
+		: FFilterState(Other)
+	{
+		FilterValue = Other.FilterValue;
+	}
+
 	virtual ~FTimerNameFilterState() {}
 
 	virtual void Update() override;
@@ -33,10 +39,12 @@ public:
 
 	virtual void SetFilterValue(FString InFilterValue) override { FilterValue = InFilterValue; }
 
+	virtual bool Equals(const FFilterState& Other) const;
+	virtual TSharedRef<FFilterState> DeepCopy() const;
+
 private:
 	FString FilterValue;
 	TSet<uint32> TimerIds;
-	bool bShouldUpdate = true;
 };
 
 class FTimerNameFilter : public FCustomFilter
@@ -99,6 +107,9 @@ public:
 	virtual bool HasCustomUI() const override { return true; }
 	virtual void AddCustomUI(TSharedRef<SHorizontalBox> LeftBox) override;
 
+	virtual bool Equals(const FFilterState& Other) const override ;
+	virtual TSharedRef<FFilterState> DeepCopy() const override ;
+
 private:
 	FText GetKeyTextBoxValue() const;
 	void OnKeyTextBoxValueCommitted(const FText& InNewText, ETextCommit::Type InTextCommit);
diff --git a/Engine/Source/Developer/TraceInsights/Private/Insights/Widgets/SQuickFind.cpp b/Engine/Source/Developer/TraceInsights/Private/Insights/Widgets/SQuickFind.cpp
--- a/Engine/Source/Developer/TraceInsights/Private/Insights/Widgets/SQuickFind.cpp
+++ b/Engine/Source/Developer/TraceInsights/Private/Insights/Widgets/SQuickFind.cpp
@@ -242,7 +242,6 @@ FReply SQuickFind::FindLast_OnClicked()
 
 FReply SQuickFind::FilterAll_OnClicked()
 {
-	QuickFindViewModel->GetFilterConfigurator()->GetRootNode()->ProcessFilter();
 	QuickFindViewModel->GetOnFilterAllEvent().Broadcast();
 
 	return FReply::Handled();
diff --git a/Engine/Source/Developer/TraceInsights/Private/Insights/Widgets/STimingView.cpp b/Engine/Source/Developer/TraceInsights/Private/Insights/Widgets/STimingView.cpp
--- a/Engine/Source/Developer/TraceInsights/Private/Insights/Widgets/STimingView.cpp
+++ b/Engine/Source/Developer/TraceInsights/Private/Insights/Widgets/STimingView.cpp
@@ -5239,7 +5239,6 @@ void STimingView::FilterAllTracks()
 {
 	LLM_SCOPE_BYTAG(Insights);
 	FilterConfigurator = MakeShared<Insights::FFilterConfigurator>(*QuickFindVm->GetFilterConfigurator());
-	FilterConfigurator->Update();
 
 	for (auto& Entry : AllTracks)
 	{
