[Insights] Added thread name in tooltip for a task event.

diff --git a/Engine/Source/Developer/TraceInsights/Private/Insights/TaskGraphProfiler/ViewModels/TaskTimingTrack.cpp b/Engine/Source/Developer/TraceInsights/Private/Insights/TaskGraphProfiler/ViewModels/TaskTimingTrack.cpp
--- a/Engine/Source/Developer/TraceInsights/Private/Insights/TaskGraphProfiler/ViewModels/TaskTimingTrack.cpp
+++ b/Engine/Source/Developer/TraceInsights/Private/Insights/TaskGraphProfiler/ViewModels/TaskTimingTrack.cpp
@@ -6,6 +6,7 @@
 #include "Framework/Commands/UICommandList.h"
 #include "Framework/MultiBox/MultiBoxBuilder.h"
 #include "TraceServices/Model/TasksProfiler.h"
+#include "TraceServices/Model/Threads.h"
 
 // Insights
 #include "Insights/Common/TimeUtils.h"
@@ -1052,7 +1053,16 @@ void FTaskTimingTrack::InitTooltip(FTooltipDrawState& InOutTooltip, const ITimin
 	case ETaskEventType::NestedStarted:
 	case ETaskEventType::SubsequentStarted:
 	{
-		InOutTooltip.AddNameValueTextLine(TEXT("Executed Thread Id:"), FString::Printf(TEXT("%d"), Task->StartedThreadId));
+		const TraceServices::IThreadProvider& ThreadProvider = TraceServices::ReadThreadProvider(*Session.Get());
+		const TCHAR* StartedThreadName = ThreadProvider.GetThreadName(Task->StartedThreadId);
+		if (StartedThreadName)
+		{
+			InOutTooltip.AddNameValueTextLine(TEXT("Executed Thread Id:"), FString::Printf(TEXT("%d (%s)"), Task->StartedThreadId, StartedThreadName));
+		}
+		else
+		{
+			InOutTooltip.AddNameValueTextLine(TEXT("Executed Thread Id:"), FString::Printf(TEXT("%d (<unknown>)"), Task->StartedThreadId));
+		}
 		break;
 	}
 	case ETaskEventType::Launched:
