[Insights] Fix Copy to clipboard for a selected CPU/GPU timing event does not work if Timers panel is hidden.

diff --git a/Engine/Source/Developer/TraceInsights/Private/Insights/ViewModels/ThreadTimingTrack.cpp b/Engine/Source/Developer/TraceInsights/Private/Insights/ViewModels/ThreadTimingTrack.cpp
--- a/Engine/Source/Developer/TraceInsights/Private/Insights/ViewModels/ThreadTimingTrack.cpp
+++ b/Engine/Source/Developer/TraceInsights/Private/Insights/ViewModels/ThreadTimingTrack.cpp
@@ -12,6 +12,7 @@
 #include "TraceServices/Model/LoadTimeProfiler.h"
 #include "TraceServices/Model/TasksProfiler.h"
 #include "TraceServices/Model/Threads.h"
+#include "TraceServices/Model/TimingProfiler.h"
 #include "Async/TaskGraphInterfaces.h"
 
 // Insights
@@ -1683,24 +1684,26 @@ void FThreadTimingTrack::OnClipboardCopyEvent(const ITimingEvent& InSelectedEven
 	{
 		const FThreadTrackEvent& TrackEvent = InSelectedEvent.As<FThreadTrackEvent>();
 
-		FTimerNodePtr TimerNodePtr = FTimingProfilerManager::Get()->GetTimerNode(TrackEvent.GetTimerId());
-		if (TimerNodePtr)
+		TSharedPtr<const TraceServices::IAnalysisSession> Session = FInsightsManager::Get()->GetSession();
+		check(Session.IsValid());
+
+		TraceServices::FAnalysisSessionReadScope SessionReadScope(*Session.Get());
+		const TraceServices::ITimingProfilerProvider& TimingProfilerProvider = *TraceServices::ReadTimingProfilerProvider(*Session.Get());
+		const TraceServices::ITimingProfilerTimerReader* TimerReader = nullptr;
+		TimingProfilerProvider.ReadTimers([&TimerReader](const TraceServices::ITimingProfilerTimerReader& Out) { TimerReader = &Out; });
+		check(TimerReader);
+
+		const TraceServices::FTimingProfilerTimer* TimerPtr = TimerReader->GetTimer(TrackEvent.GetTimerIndex());
+
+		if (TimerPtr)
 		{
-			FString EventName = TimerNodePtr->GetName().ToString();
+			FString EventName(TimerPtr->Name);
 
 			FTimingEventsTrackDrawStateBuilder::AppendDurationToEventName(EventName, TrackEvent.GetDuration());
 
 			const uint32 TimerIndex = TrackEvent.GetTimerIndex();
 			if (int32(TimerIndex) < 0) // has metadata?
 			{
-				TSharedPtr<const TraceServices::IAnalysisSession> Session = FInsightsManager::Get()->GetSession();
-				check(Session.IsValid());
-
-				TraceServices::FAnalysisSessionReadScope SessionReadScope(*Session.Get());
-				const TraceServices::ITimingProfilerProvider& TimingProfilerProvider = *TraceServices::ReadTimingProfilerProvider(*Session.Get());
-				const TraceServices::ITimingProfilerTimerReader* TimerReader;
-				TimingProfilerProvider.ReadTimers([&TimerReader](const TraceServices::ITimingProfilerTimerReader& Out) { TimerReader = &Out; });
-
 				TArrayView<const uint8> Metadata = TimerReader->GetMetadata(TimerIndex);
 				if (Metadata.Num() > 0)
 				{
