[Insights] Session Browser: Added filtering by [REDACTED_PERSON] (empty/small/medium/large) and by Status (live/offline).

diff --git a/Engine/Source/Developer/TraceInsights/Private/Insights/Widgets/SStartPageWindow.cpp b/Engine/Source/Developer/TraceInsights/Private/Insights/Widgets/SStartPageWindow.cpp
--- a/Engine/Source/Developer/TraceInsights/Private/Insights/Widgets/SStartPageWindow.cpp
+++ b/Engine/Source/Developer/TraceInsights/Private/Insights/Widgets/SStartPageWindow.cpp
@@ -1163,6 +1163,24 @@ TSharedRef<SWidget> STraceStoreWindow::ConstructFiltersToolbar()
 			LOCTEXT("FilterByBranchToolTip", "Filters the list of trace sessions by branch."),
 			FSlateIcon(FInsightsStyle::GetStyleSetName(), "Icons.Filter.ToolBar"),
 			false);
+
+		// Filter by Size
+		ToolbarBuilder.AddComboButton(
+			FUIAction(),
+			FOnGetContent::CreateSP(this, &STraceStoreWindow::MakeSizeFilterMenu),
+			LOCTEXT("FilterBySizeText", "Size"),
+			LOCTEXT("FilterBySizeToolTip", "Filters the list of trace sessions by size."),
+			FSlateIcon(FInsightsStyle::GetStyleSetName(), "Icons.Filter.ToolBar"),
+			false);
+
+		// Filter by Status
+		ToolbarBuilder.AddComboButton(
+			FUIAction(),
+			FOnGetContent::CreateSP(this, &STraceStoreWindow::MakeStatusFilterMenu),
+			LOCTEXT("FilterByStatusText", "Status"),
+			LOCTEXT("FilterByStatusToolTip", "Filters the list of trace sessions by status.."),
+			FSlateIcon(FInsightsStyle::GetStyleSetName(), "Icons.Filter.ToolBar"),
+			false);
 	}
 	ToolbarBuilder.EndSection();
 
@@ -1227,6 +1245,7 @@ TSharedRef<SWidget> STraceStoreWindow::ConstructSessionsPanel()
 			[
 				SNew(SBox)
 				.VAlign(VAlign_Center)
+				.MinDesiredHeight(24.0f)
 				[
 					SNew(STextBlock)
 					.Text(LOCTEXT("NameColumn", "Name"))
@@ -1243,6 +1262,7 @@ TSharedRef<SWidget> STraceStoreWindow::ConstructSessionsPanel()
 			[
 				SNew(SBox)
 				.VAlign(VAlign_Center)
+				.MinDesiredHeight(24.0f)
 				[
 					SNew(STextBlock)
 					.Text(LOCTEXT("PlatformColumn", "Platform"))
@@ -1259,6 +1279,7 @@ TSharedRef<SWidget> STraceStoreWindow::ConstructSessionsPanel()
 			[
 				SNew(SBox)
 				.VAlign(VAlign_Center)
+				.MinDesiredHeight(24.0f)
 				[
 					SNew(STextBlock)
 					.Text(LOCTEXT("AppNameColumn", "App Name"))
@@ -1275,6 +1296,7 @@ TSharedRef<SWidget> STraceStoreWindow::ConstructSessionsPanel()
 			[
 				SNew(SBox)
 				.VAlign(VAlign_Center)
+				.MinDesiredHeight(24.0f)
 				[
 					SNew(STextBlock)
 					.Text(LOCTEXT("BuildConfigColumn", "Build Config"))
@@ -1291,6 +1313,7 @@ TSharedRef<SWidget> STraceStoreWindow::ConstructSessionsPanel()
 			[
 				SNew(SBox)
 				.VAlign(VAlign_Center)
+				.MinDesiredHeight(24.0f)
 				[
 					SNew(STextBlock)
 					.Text(LOCTEXT("BuildTargetColumn", "Build Target"))
@@ -1299,22 +1322,42 @@ TSharedRef<SWidget> STraceStoreWindow::ConstructSessionsPanel()
 			]
 
 			+ SHeaderRow::Column(TraceStoreColumns::Size)
-			.DefaultLabel(LOCTEXT("SizeColumn", "File Size"))
 			.FixedWidth(100.0f)
 			.HAlignHeader(HAlign_Right)
 			.HAlignCell(HAlign_Right)
 			.InitialSortMode(EColumnSortMode::Descending)
 			.SortMode(this, &STraceStoreWindow::GetSortModeForColumn, TraceStoreColumns::Size)
 			.OnSort(this, &STraceStoreWindow::OnSortModeChanged)
+			.OnGetMenuContent(this, &STraceStoreWindow::MakeSizeColumnHeaderMenu)
+			[
+				SNew(SBox)
+				.VAlign(VAlign_Center)
+				.MinDesiredHeight(24.0f)
+				[
+					SNew(STextBlock)
+					.Text(LOCTEXT("SizeColumn", "File Size"))
+					.ColorAndOpacity_Lambda([this] { return FilterBySize->IsEmpty() ? FLinearColor(0.5f, 0.5f, 0.5f, 1.0f) : FLinearColor(0.3f, 0.75f, 1.0f, 1.0f); })
+				]
+			]
 
 			+ SHeaderRow::Column(TraceStoreColumns::Status)
-			.DefaultLabel(LOCTEXT("StatusColumn", "Status"))
 			.FixedWidth(60.0f)
 			.HAlignHeader(HAlign_Right)
 			.HAlignCell(HAlign_Right)
 			.InitialSortMode(EColumnSortMode::Ascending)
 			.SortMode(this, &STraceStoreWindow::GetSortModeForColumn, TraceStoreColumns::Status)
 			.OnSort(this, &STraceStoreWindow::OnSortModeChanged)
+			.OnGetMenuContent(this, &STraceStoreWindow::MakeStatusColumnHeaderMenu)
+			[
+				SNew(SBox)
+				.VAlign(VAlign_Center)
+				.MinDesiredHeight(24.0f)
+				[
+					SNew(STextBlock)
+					.Text(LOCTEXT("StatusColumn", "Status"))
+					.ColorAndOpacity_Lambda([this] { return FilterByStatus->IsEmpty() ? FLinearColor(0.5f, 0.5f, 0.5f, 1.0f) : FLinearColor(0.3f, 0.75f, 1.0f, 1.0f); })
+				]
+			]
 		);
 
 	return Widget;
@@ -3218,7 +3261,7 @@ TSharedRef<SWidget> STraceStoreWindow::MakePlatformColumnHeaderMenu()
 
 	FMenuBuilder MenuBuilder(/*bInShouldCloseWindowAfterMenuSelection=*/true, nullptr);
 
-	MenuBuilder.BeginSection("Filter", LOCTEXT("MenuSection_PlatformFilter", "Platform Filter"));
+	MenuBuilder.BeginSection("Filter", LOCTEXT("ColumnHeaderMenuSection_PlatformFilter", "Platform Filter"));
 	BuildPlatformFilterSubMenu(MenuBuilder);
 	MenuBuilder.EndSection();
 
@@ -3255,7 +3298,7 @@ TSharedRef<SWidget> STraceStoreWindow::MakeAppNameColumnHeaderMenu()
 
 	FMenuBuilder MenuBuilder(/*bInShouldCloseWindowAfterMenuSelection=*/true, nullptr);
 
-	MenuBuilder.BeginSection("Filter", LOCTEXT("MenuSection_AppNameFilter", "App Name Filter"));
+	MenuBuilder.BeginSection("Filter", LOCTEXT("ColumnHeaderMenuSection_AppNameFilter", "App Name Filter"));
 	BuildAppNameFilterSubMenu(MenuBuilder);
 	MenuBuilder.EndSection();
 
@@ -3292,7 +3335,7 @@ TSharedRef<SWidget> STraceStoreWindow::MakeBuildConfigColumnHeaderMenu()
 
 	FMenuBuilder MenuBuilder(/*bInShouldCloseWindowAfterMenuSelection=*/true, nullptr);
 
-	MenuBuilder.BeginSection("Filter", LOCTEXT("MenuSection_BuildConfigFilter", "Build Config Filter"));
+	MenuBuilder.BeginSection("Filter", LOCTEXT("ColumnHeaderMenuSection_BuildConfigFilter", "Build Config Filter"));
 	BuildBuildConfigFilterSubMenu(MenuBuilder);
 	MenuBuilder.EndSection();
 
@@ -3329,7 +3372,7 @@ TSharedRef<SWidget> STraceStoreWindow::MakeBuildTargetColumnHeaderMenu()
 
 	FMenuBuilder MenuBuilder(/*bInShouldCloseWindowAfterMenuSelection=*/true, nullptr);
 
-	MenuBuilder.BeginSection("Filter", LOCTEXT("MenuSection_BuildTargetFilter", "Build Target Filter"));
+	MenuBuilder.BeginSection("Filter", LOCTEXT("ColumnHeaderMenuSection_BuildTargetFilter", "Build Target Filter"));
 	BuildBuildTargetFilterSubMenu(MenuBuilder);
 	MenuBuilder.EndSection();
 
@@ -3360,6 +3403,21 @@ void STraceStoreWindow::BuildBuildTargetFilterSubMenu(FMenuBuilder& InMenuBuilde
 
 ////////////////////////////////////////////////////////////////////////////////////////////////////
 
+TSharedRef<SWidget> STraceStoreWindow::MakeBranchColumnHeaderMenu()
+{
+	FSlateApplication::Get().CloseToolTip();
+
+	FMenuBuilder MenuBuilder(/*bInShouldCloseWindowAfterMenuSelection=*/true, nullptr);
+
+	MenuBuilder.BeginSection("Filter", LOCTEXT("ColumnHeaderMenuSection_BranchFilter", "Branch Filter"));
+	BuildBranchFilterSubMenu(MenuBuilder);
+	MenuBuilder.EndSection();
+
+	return MenuBuilder.MakeWidget();
+}
+
+////////////////////////////////////////////////////////////////////////////////////////////////////
+
 TSharedRef<SWidget> STraceStoreWindow::MakeBranchFilterMenu()
 {
 	FSlateApplication::Get().CloseToolTip();
@@ -3382,6 +3440,80 @@ void STraceStoreWindow::BuildBranchFilterSubMenu(FMenuBuilder& InMenuBuilder)
 
 ////////////////////////////////////////////////////////////////////////////////////////////////////
 
+TSharedRef<SWidget> STraceStoreWindow::MakeSizeColumnHeaderMenu()
+{
+	FSlateApplication::Get().CloseToolTip();
+
+	FMenuBuilder MenuBuilder(/*bInShouldCloseWindowAfterMenuSelection=*/true, nullptr);
+
+	MenuBuilder.BeginSection("Filter", LOCTEXT("ColumnHeaderMenuSection_SizeFilter", "Size Filter"));
+	BuildSizeFilterSubMenu(MenuBuilder);
+	MenuBuilder.EndSection();
+
+	return MenuBuilder.MakeWidget();
+}
+
+////////////////////////////////////////////////////////////////////////////////////////////////////
+
+TSharedRef<SWidget> STraceStoreWindow::MakeSizeFilterMenu()
+{
+	FSlateApplication::Get().CloseToolTip();
+
+	FMenuBuilder MenuBuilder(/*bInShouldCloseWindowAfterMenuSelection=*/true, nullptr);
+
+	MenuBuilder.BeginSection("Filter", LOCTEXT("MenuSection_SizeFilter", "Size Filter"));
+	BuildSizeFilterSubMenu(MenuBuilder);
+	MenuBuilder.EndSection();
+
+	return MenuBuilder.MakeWidget();
+}
+
+////////////////////////////////////////////////////////////////////////////////////////////////////
+
+void STraceStoreWindow::BuildSizeFilterSubMenu(FMenuBuilder& InMenuBuilder)
+{
+	FilterBySize->BuildMenu(InMenuBuilder, *this);
+}
+
+////////////////////////////////////////////////////////////////////////////////////////////////////
+
+TSharedRef<SWidget> STraceStoreWindow::MakeStatusColumnHeaderMenu()
+{
+	FSlateApplication::Get().CloseToolTip();
+
+	FMenuBuilder MenuBuilder(/*bInShouldCloseWindowAfterMenuSelection=*/true, nullptr);
+
+	MenuBuilder.BeginSection("Filter", LOCTEXT("ColumnHeaderMenuSection_StatusFilter", "Status Filter"));
+	BuildStatusFilterSubMenu(MenuBuilder);
+	MenuBuilder.EndSection();
+
+	return MenuBuilder.MakeWidget();
+}
+
+////////////////////////////////////////////////////////////////////////////////////////////////////
+
+TSharedRef<SWidget> STraceStoreWindow::MakeStatusFilterMenu()
+{
+	FSlateApplication::Get().CloseToolTip();
+
+	FMenuBuilder MenuBuilder(/*bInShouldCloseWindowAfterMenuSelection=*/true, nullptr);
+
+	MenuBuilder.BeginSection("Filter", LOCTEXT("MenuSection_StatusFilter", "Status Filter"));
+	BuildStatusFilterSubMenu(MenuBuilder);
+	MenuBuilder.EndSection();
+
+	return MenuBuilder.MakeWidget();
+}
+
+////////////////////////////////////////////////////////////////////////////////////////////////////
+
+void STraceStoreWindow::BuildStatusFilterSubMenu(FMenuBuilder& InMenuBuilder)
+{
+	FilterByStatus->BuildMenu(InMenuBuilder, *this);
+}
+
+////////////////////////////////////////////////////////////////////////////////////////////////////
+
 FText STraceStoreWindow::GetTraceStoreDirectory() const
 {
 	return FText::FromString(FPaths::ConvertRelativePathToFull(GetStoreDirectory()));
@@ -3571,7 +3703,14 @@ void TTraceSetFilter<TSetType>::BuildMenu(FMenuBuilder& InMenuBuilder, STraceSto
 
 	InMenuBuilder.AddSeparator();
 
+	TArray<TSetType> DefaultValues;
+	AddDefaultValues(DefaultValues);
+
 	TMap<TSetType, uint32> AllUniqueValues;
+	for (const TSetType& Value : DefaultValues)
+	{
+		AllUniqueValues.Add(Value, 0);
+	}
 	for (const TSharedPtr<FTraceViewModel>& Trace : Window.GetAllAvailableTraces())
 	{
 		TSetType Value = GetFilterValueForTrace(*Trace);
@@ -3617,6 +3756,8 @@ void TTraceSetFilter<TSetType>::BuildMenu(FMenuBuilder& InMenuBuilder, STraceSto
 	}
 }
 
+////////////////////////////////////////////////////////////////////////////////////////////////////
+
 template<typename TSetType>
 TTraceSetFilter<TSetType>::TTraceSetFilter()
 {
@@ -3624,26 +3765,36 @@ TTraceSetFilter<TSetType>::TTraceSetFilter()
 	UndefinedValueLabel = LOCTEXT("UndefinedValueLabel", "N/A");
 }
 
+////////////////////////////////////////////////////////////////////////////////////////////////////
+
 FTraceFilterByPlatform::FTraceFilterByPlatform()
 {
 	ToggleAllActionTooltip = LOCTEXT("FilterByPlatform_ToggleAll_Tooltip", "Shows or hides traces for all platforms.");
 }
 
+////////////////////////////////////////////////////////////////////////////////////////////////////
+
 FTraceFilterByAppName::FTraceFilterByAppName()
 {
 	ToggleAllActionTooltip = LOCTEXT("FilterByAppName_ToggleAll_Tooltip", "Shows or hides traces for all app names.");
 }
 
+////////////////////////////////////////////////////////////////////////////////////////////////////
+
 FTraceFilterByBuildConfig::FTraceFilterByBuildConfig()
 {
 	ToggleAllActionTooltip = LOCTEXT("FilterByBuildConfig_ToggleAll_Tooltip", "Shows or hides traces for all build configurations.");
 }
 
+////////////////////////////////////////////////////////////////////////////////////////////////////
+
 FTraceFilterByBuildTarget::FTraceFilterByBuildTarget()
 {
 	ToggleAllActionTooltip = LOCTEXT("FilterByBuildTarget_ToggleAll_Tooltip", "Shows or hides traces for all build targets.");
 }
 
+////////////////////////////////////////////////////////////////////////////////////////////////////
+
 FTraceFilterByBranch::FTraceFilterByBranch()
 {
 	ToggleAllActionTooltip = LOCTEXT("FilterByBranch_ToggleAll_Tooltip", "Shows or hides traces for all branches.");
@@ -3651,6 +3802,66 @@ FTraceFilterByBranch::FTraceFilterByBranch()
 
 ////////////////////////////////////////////////////////////////////////////////////////////////////
 
+FTraceFilterBySize::FTraceFilterBySize()
+{
+	ToggleAllActionTooltip = LOCTEXT("FilterBySize_ToggleAll_Tooltip", "Shows or hides traces of all sizes.");
+}
+
+////////////////////////////////////////////////////////////////////////////////////////////////////
+
+void FTraceFilterBySize::AddDefaultValues(TArray<uint8>& InOutDefaultValues) const
+{
+	for (uint8 SizeCategory = 0; SizeCategory < (uint8)ESizeCategory::InvalidOrMax; ++SizeCategory)
+	{
+		InOutDefaultValues.Add(SizeCategory);
+	}
+}
+
+////////////////////////////////////////////////////////////////////////////////////////////////////
+
+FText FTraceFilterBySize::ValueToText(const uint8 InValue) const
+{
+	switch ((ESizeCategory)InValue)
+	{
+		case ESizeCategory::Empty:  return LOCTEXT("FilterBySize_Empty",  "Empty (0 bytes)");
+		case ESizeCategory::Small:  return LOCTEXT("FilterBySize_Small",  "Small (< 1 MiB)");
+		case ESizeCategory::Medium: return LOCTEXT("FilterBySize_Medium", "Medium (< 1 GiB)");
+		case ESizeCategory::Large:  return LOCTEXT("FilterBySize_Large",  "Large (\u2265 1 GiB)");
+		default:                    return UndefinedValueLabel;
+	}
+}
+
+////////////////////////////////////////////////////////////////////////////////////////////////////
+
+FTraceFilterByStatus::FTraceFilterByStatus()
+{
+	ToggleAllActionTooltip = LOCTEXT("FilterByStatus_ToggleAll_Tooltip", "Shows or hides all traces.");
+}
+
+////////////////////////////////////////////////////////////////////////////////////////////////////
+
+void FTraceFilterByStatus::AddDefaultValues(TArray<bool>& InOutDefaultValues) const
+{
+	InOutDefaultValues.Add(false);
+	InOutDefaultValues.Add(true);
+}
+
+////////////////////////////////////////////////////////////////////////////////////////////////////
+
+FText FTraceFilterByStatus::ValueToText(const bool InValue) const
+{
+	if (InValue)
+	{
+		return LOCTEXT("FilterByStatus_Live", "LIVE");
+	}
+	else
+	{
+		return LOCTEXT("FilterByStatus_Offline", "Offline");
+	}
+}
+
+////////////////////////////////////////////////////////////////////////////////////////////////////
+
 void STraceStoreWindow::FilterByNameSearchBox_OnTextChanged(const FText& InFilterText)
 {
 	FilterByName->SetRawFilterText(InFilterText);
@@ -3697,6 +3908,12 @@ void STraceStoreWindow::CreateFilters()
 
 	FilterByBranch = MakeShared<FTraceFilterByBranch>();
 	Filters->Add(FilterByBranch);
+
+	FilterBySize = MakeShared<FTraceFilterBySize>();
+	Filters->Add(FilterBySize);
+
+	FilterByStatus = MakeShared<FTraceFilterByStatus>();
+	Filters->Add(FilterByStatus);
 }
 
 ////////////////////////////////////////////////////////////////////////////////////////////////////
@@ -3724,7 +3941,9 @@ void STraceStoreWindow::UpdateFiltering()
 		FilterByAppName->IsEmpty() &&
 		FilterByBuildConfig->IsEmpty() &&
 		FilterByBuildTarget->IsEmpty() &&
-		FilterByBranch->IsEmpty())
+		FilterByBranch->IsEmpty() &&
+		FilterBySize->IsEmpty() &&
+		FilterByStatus->IsEmpty())
 	{
 		// No filtering.
 		FilteredTraceViewModels = TraceViewModels;
diff --git a/Engine/Source/Developer/TraceInsights/Private/Insights/Widgets/SStartPageWindow.h b/Engine/Source/Developer/TraceInsights/Private/Insights/Widgets/SStartPageWindow.h
--- a/Engine/Source/Developer/TraceInsights/Private/Insights/Widgets/SStartPageWindow.h
+++ b/Engine/Source/Developer/TraceInsights/Private/Insights/Widgets/SStartPageWindow.h
@@ -146,6 +146,7 @@ public:
 	virtual void BuildMenu(FMenuBuilder& InMenuBuilder, class STraceStoreWindow& InWindow);
 
 protected:
+	virtual void AddDefaultValues(TArray<TSetType>& InOutDefaultValues) const { }
 	virtual TSetType GetFilterValueForTrace(const FTraceViewModel& InTrace) const = 0;
 	virtual FText ValueToText(const TSetType Value) const = 0;
 
@@ -242,6 +243,64 @@ protected:
 	}
 };
 
+class FTraceFilterBySize : public TTraceSetFilter<uint8>
+{
+public:
+	enum class ESizeCategory : uint8
+	{
+		Empty,  // 0 bytes
+		Small,  // < 1 MiB
+		Medium, // < 1 GiB
+		Large,  // >= 1 GiB
+
+		InvalidOrMax
+	};
+
+public:
+	FTraceFilterBySize();
+
+protected:
+	virtual void AddDefaultValues(TArray<uint8>& InOutDefaultValues) const override;
+
+	virtual uint8 GetFilterValueForTrace(const FTraceViewModel& InTrace) const override
+	{
+		if (InTrace.Size == 0)
+		{
+			return (uint8)ESizeCategory::Empty;
+		}
+		else if (InTrace.Size < 1024ull * 1024ull)
+		{
+			return (uint8)ESizeCategory::Small;
+		}
+		else if (InTrace.Size < 1024ull * 1024ull * 1024ull)
+		{
+			return (uint8)ESizeCategory::Medium;
+		}
+		else
+		{
+			return (uint8)ESizeCategory::Large;
+		}
+	}
+
+	virtual FText ValueToText(const uint8 InValue) const override;
+};
+
+class FTraceFilterByStatus : public TTraceSetFilter<bool>
+{
+public:
+	FTraceFilterByStatus();
+
+protected:
+	virtual void AddDefaultValues(TArray<bool>& InOutDefaultValues) const override;
+
+	virtual bool GetFilterValueForTrace(const FTraceViewModel& InTrace) const override
+	{
+		return InTrace.bIsLive;
+	}
+
+	virtual FText ValueToText(const bool InValue) const override;
+};
+
 ////////////////////////////////////////////////////////////////////////////////////////////////////
 
 /** Implements the Trace Store window. */
@@ -357,9 +416,18 @@ private:
 	TSharedRef<SWidget> MakeBuildTargetFilterMenu();
 	void BuildBuildTargetFilterSubMenu(FMenuBuilder& InMenuBuilder);
 
+	TSharedRef<SWidget> MakeBranchColumnHeaderMenu();
 	TSharedRef<SWidget> MakeBranchFilterMenu();
 	void BuildBranchFilterSubMenu(FMenuBuilder& InMenuBuilder);
 
+	TSharedRef<SWidget> MakeSizeColumnHeaderMenu();
+	TSharedRef<SWidget> MakeSizeFilterMenu();
+	void BuildSizeFilterSubMenu(FMenuBuilder& InMenuBuilder);
+
+	TSharedRef<SWidget> MakeStatusColumnHeaderMenu();
+	TSharedRef<SWidget> MakeStatusFilterMenu();
+	void BuildStatusFilterSubMenu(FMenuBuilder& InMenuBuilder);
+
 	FReply RefreshTraces_OnClicked();
 	FSlateColor GetColorByPath(const FString& Uri);
 
@@ -576,6 +644,8 @@ private:
 	TSharedPtr<FTraceFilterByBuildConfig> FilterByBuildConfig;
 	TSharedPtr<FTraceFilterByBuildTarget> FilterByBuildTarget;
 	TSharedPtr<FTraceFilterByBranch> FilterByBranch;
+	TSharedPtr<FTraceFilterBySize> FilterBySize;
+	TSharedPtr<FTraceFilterByStatus> FilterByStatus;
 
 	bool bFilterStatsTextIsDirty = true;
 	FText FilterStatsText;
