Unreal Insights : delete trace is now allowed as soon as at least one deletable (i.e. not-LIVE) trace is selected so that a user can "select all" and then delete all non-LIVE traces

diff --git a/Engine/Source/Developer/TraceInsights/Private/Insights/Widgets/SStartPageWindow.cpp b/Engine/Source/Developer/TraceInsights/Private/Insights/Widgets/SStartPageWindow.cpp
--- a/Engine/Source/Developer/TraceInsights/Private/Insights/Widgets/SStartPageWindow.cpp
+++ b/Engine/Source/Developer/TraceInsights/Private/Insights/Widgets/SStartPageWindow.cpp
@@ -2114,16 +2114,17 @@ bool STraceStoreWindow::CanDeleteSelectedTraces() const
 	{
 		return false;
 	}
+
 	TArray<TSharedPtr<FTraceViewModel>> SelectedTraces = TraceListView->GetSelectedItems();
 	for (const TSharedPtr<FTraceViewModel>& SelectedTrace : SelectedTraces)
 	{
-		if (SelectedTrace->TraceId == FTraceViewModel::InvalidTraceId ||
-			SelectedTrace->bIsLive)
+		if ((SelectedTrace->TraceId != FTraceViewModel::InvalidTraceId) && 
+			!SelectedTrace->bIsLive)
 		{
-			return false;
+			return true;
 		}
 	}
-	return true;
+	return false;
 }
 
 ////////////////////////////////////////////////////////////////////////////////////////////////////
@@ -2138,6 +2139,8 @@ void STraceStoreWindow::DeleteSelectedTraces()
 	}
 
 	TArray<TSharedPtr<FTraceViewModel>> TracesToDelete = TraceListView->GetSelectedItems();
+	// Filter the traces that can actually be deleted : 
+	TracesToDelete.RemoveAll([](const TSharedPtr<FTraceViewModel>& InTrace) { return (InTrace->TraceId == FTraceViewModel::InvalidTraceId) || InTrace->bIsLive; });
 	if (TracesToDelete.Num() == 0)
 	{
 		return;
