[Insights] Remove the Filter Track menu option from the track context menu because it was replaced by the Quick Find filter functionality.

diff --git a/Engine/Source/Developer/TraceInsights/Private/Insights/ViewModels/ThreadTimingTrack.cpp b/Engine/Source/Developer/TraceInsights/Private/Insights/ViewModels/ThreadTimingTrack.cpp
--- a/Engine/Source/Developer/TraceInsights/Private/Insights/ViewModels/ThreadTimingTrack.cpp
+++ b/Engine/Source/Developer/TraceInsights/Private/Insights/ViewModels/ThreadTimingTrack.cpp
@@ -899,10 +899,6 @@ INSIGHTS_IMPLEMENT_RTTI(FThreadTimingTrack)
 
 FThreadTimingTrack::~FThreadTimingTrack()
 {
-	if (FilterConfigurator.IsValid())
-	{
-		FilterConfigurator->GetOnChangesCommittedEvent().Remove(OnFilterChangesCommittedHandle);
-	}
 }
 
 ////////////////////////////////////////////////////////////////////////////////////////////////////
@@ -1751,22 +1747,6 @@ void FThreadTimingTrack::BuildContextMenu(FMenuBuilder& MenuBuilder)
 	{
 		ChildTrack->BuildContextMenu(MenuBuilder);
 	}
-
-	MenuBuilder.BeginSection("TimingEvents", LOCTEXT("ContextMenu_Section_TimingEvents", "Timing Events"));
-	{
-		FExecuteAction FilterTrackAction;
-		FilterTrackAction.BindSP(this, &FThreadTimingTrack::OnFilterTrackClicked);
-
-		MenuBuilder.AddMenuEntry(
-			LOCTEXT("FilterTrack", "Filter Track..."),
-			FText(),
-			FSlateIcon(FAppStyle::Get().GetStyleSetName(), "Icons.Filter"),
-			FUIAction(FilterTrackAction, FCanExecuteAction::CreateLambda([]() { return true; })),
-			NAME_None,
-			EUserInterfaceActionType::Button
-		);
-	}
-	MenuBuilder.EndSection();
 }
 
 ////////////////////////////////////////////////////////////////////////////////////////////////////
@@ -1925,64 +1905,6 @@ bool FThreadTimingTrack::TimerIndexToTimerId(uint32 InTimerIndex, uint32& OutTim
 
 ////////////////////////////////////////////////////////////////////////////////////////////////////
 
-void FThreadTimingTrack::OnFilterTrackClicked()
-{
-	LLM_SCOPE_BYTAG(Insights);
-
-	if (!FilterConfigurator.IsValid())
-	{
-		FilterConfigurator = MakeShared<FFilterConfigurator>();
-
-		FilterConfigurator->Add(MakeShared<FFilter>(
-			static_cast<int32>(EFilterField::StartTime),
-			LOCTEXT("StartTime", "Start Time"),
-			LOCTEXT("StartTime", "Start Time"),
-			EFilterDataType::Double,
-			nullptr,
-			FFilterService::Get()->GetDoubleOperators()));
-
-		FilterConfigurator->Add(MakeShared<FFilter>(
-			static_cast<int32>(EFilterField::EndTime),
-			LOCTEXT("EndTime", "End Time"),
-			LOCTEXT("EndTime", "End Time"),
-			EFilterDataType::Double,
-			nullptr,
-			FFilterService::Get()->GetDoubleOperators()));
-
-		FilterConfigurator->Add(MakeShared<FFilter>(
-			static_cast<int32>(EFilterField::Duration),
-			LOCTEXT("Duration", "Duration"),
-			LOCTEXT("Duration", "Duration"),
-			EFilterDataType::Double,
-			nullptr,
-			FFilterService::Get()->GetDoubleOperators()));
-
-		FilterConfigurator->Add(MakeShared<FFilter>(
-			static_cast<int32>(EFilterField::TimerId),
-			LOCTEXT("TimerId", "Timer Id"),
-			LOCTEXT("TimerId", "Timer Id"),
-			EFilterDataType::Int64,
-			nullptr,
-			FFilterService::Get()->GetIntegerOperators()));
-	}
-	else
-	{
-		FilterConfigurator->GetOnChangesCommittedEvent().Remove(OnFilterChangesCommittedHandle);
-
-		// Make a copy, so it will not affect other tracks that shares same filter.
-		FilterConfigurator = MakeShared<FFilterConfigurator>(*FilterConfigurator);
-	}
-
-	OnFilterChangesCommittedHandle = FilterConfigurator->GetOnChangesCommittedEvent().AddLambda([this]()
-		{
-			this->SetDirtyFlag();
-		});
-
-	FFilterService::Get()->CreateFilterConfiguratorWidget(FilterConfigurator);
-}
-
-////////////////////////////////////////////////////////////////////////////////////////////////////
-
 bool FThreadTimingTrack::HasCustomFilter() const
 {
 	return FilterConfigurator.IsValid() && !FilterConfigurator->IsEmpty();
diff --git a/Engine/Source/Developer/TraceInsights/Private/Insights/ViewModels/ThreadTimingTrack.h b/Engine/Source/Developer/TraceInsights/Private/Insights/ViewModels/ThreadTimingTrack.h
--- a/Engine/Source/Developer/TraceInsights/Private/Insights/ViewModels/ThreadTimingTrack.h
+++ b/Engine/Source/Developer/TraceInsights/Private/Insights/ViewModels/ThreadTimingTrack.h
@@ -206,8 +206,6 @@ private:
 						  TSharedPtr<FThreadTrackEvent>& OutParentTimingEvent,
 						  TSharedPtr<FThreadTrackEvent>& OutRootTimingEvent) const;
 
-	void OnFilterTrackClicked();
-
 	static void CreateFThreadTrackEventFromInfo(const TimelineEventInfo& InEventInfo, const TSharedRef<const FBaseTimingTrack> InTrack, int32 InDepth, TSharedPtr<FThreadTrackEvent> &OutTimingEvent);
 	static bool TimerIndexToTimerId(uint32 InTimerIndex, uint32 & OutTimerId);
 
@@ -219,7 +217,6 @@ private:
 	FThreadTimingSharedState& SharedState;
 
 	TSharedPtr<Insights::FFilterConfigurator> FilterConfigurator;
-	FDelegateHandle OnFilterChangesCommittedHandle;
 
 	// Search cache
 	mutable TTimingEventSearchCache<TraceServices::FTimingProfilerEvent> SearchCache;
