[Insights] Fixed support for execute commands from a response file that has spaces in the path and/or filename.

diff --git a/Engine/Source/Developer/TraceInsights/Private/Insights/InsightsManager.cpp b/Engine/Source/Developer/TraceInsights/Private/Insights/InsightsManager.cpp
--- a/Engine/Source/Developer/TraceInsights/Private/Insights/InsightsManager.cpp
+++ b/Engine/Source/Developer/TraceInsights/Private/Insights/InsightsManager.cpp
@@ -1218,10 +1218,12 @@ void FInsightsManager::OnSessionAnalysisCompleted()
 
 bool FInsightsManager::Exec(const TCHAR* Cmd, FOutputDevice& Ar)
 {
-	FString ResponseFile;
-	if (FParse::Value(Cmd, TEXT("@="), ResponseFile))
+	// @=ResponseFile
+	if (Cmd != nullptr &&
+		Cmd[0] == TEXT('@') &&
+		Cmd[1] == TEXT('='))
 	{
-		HandleResponseFileCmd(*ResponseFile, Ar);
+		HandleResponseFileCmd(Cmd + 2, Ar);
 		return true;
 	}
 
@@ -1232,12 +1234,12 @@ bool FInsightsManager::Exec(const TCHAR* Cmd, FOutputDevice& Ar)
 
 bool FInsightsManager::HandleResponseFileCmd(const TCHAR* ResponseFile, FOutputDevice& Ar)
 {
-	Ar.Logf(TEXT("Executing commands using response file (%s)..."), ResponseFile);
+	Ar.Logf(TEXT("Executing commands using response file (\"%s\")..."), ResponseFile);
 
 	FString Contents;
 	if (!FFileHelper::LoadFileToString(Contents, &IPlatformFile::GetPlatformPhysical(), ResponseFile))
 	{
-		Ar.Logf(ELogVerbosity::Error, TEXT("Failed to open the response file (%s)."), ResponseFile);
+		Ar.Logf(ELogVerbosity::Error, TEXT("Failed to open the response file (\"%s\")."), ResponseFile);
 		return false;
 	}
 
@@ -1256,7 +1258,7 @@ bool FInsightsManager::HandleResponseFileCmd(const TCHAR* ResponseFile, FOutputD
 		uint32 EndOfLine = 0;
 		if (*CrtPos == TEXT('\r'))
 		{
-			if (*(CrtPos + 1) == '\n')
+			if (*(CrtPos + 1) == TEXT('\n'))
 			{
 				EndOfLine = 2;
 			}
