Fix potential data race in store browser

diff --git a/Engine/Source/Developer/TraceInsights/Private/Insights/StoreService/StoreBrowser.cpp b/Engine/Source/Developer/TraceInsights/Private/Insights/StoreService/StoreBrowser.cpp
--- a/Engine/Source/Developer/TraceInsights/Private/Insights/StoreService/StoreBrowser.cpp
+++ b/Engine/Source/Developer/TraceInsights/Private/Insights/StoreService/StoreBrowser.cpp
@@ -214,41 +214,38 @@ void FStoreBrowser::UpdateTraces()
 		uint32 NewStoreSettingsChangeSerial = 0;
 		FString NewStoreDirectory;
 		TArray<FString> NewWatchDirectories;
+		TOptional<uint32> NewStorePort, NewRecorderPort;
 
 		{
-			const UE::Trace::FStoreClient::FStatus* Status = nullptr;
+			// Get Trace Store status.
+			FScopeLock StoreClientLock(&GetStoreClientCriticalSection());
+			const UE::Trace::FStoreClient::FStatus* Status = StoreClient->GetStatus();
+			if (Status)
 			{
-				// Get Trace Store status.
-				FScopeLock StoreClientLock(&GetStoreClientCriticalSection());
-				Status = StoreClient->GetStatus();
-				if (Status)
+				NewStoreChangeSerial = Status->GetChangeSerial();
+				NewStoreSettingsChangeSerial = Status->GetSettingsSerial();
+				if (StoreSettingsChangeSerial != NewStoreSettingsChangeSerial)
 				{
-					NewStoreChangeSerial = Status->GetChangeSerial();
-					NewStoreSettingsChangeSerial = Status->GetSettingsSerial();
-					if (StoreSettingsChangeSerial != NewStoreSettingsChangeSerial)
-					{
-						NewStoreDirectory = FString(Status->GetStoreDir());
-						Status->GetWatchDirectories(NewWatchDirectories);
-					}
+					NewStoreDirectory = FString(Status->GetStoreDir());
+					Status->GetWatchDirectories(NewWatchDirectories);
+					NewStorePort = Status->GetStorePort();
+					NewRecorderPort = Status->GetRecorderPort();
 				}
 			}
+		}
 
-			if (StoreSettingsChangeSerial != NewStoreSettingsChangeSerial)
-			{
-				StoreSettingsChangeSerial = NewStoreSettingsChangeSerial;
+		if (StoreSettingsChangeSerial != NewStoreSettingsChangeSerial)
+		{
+			StoreSettingsChangeSerial = NewStoreSettingsChangeSerial;
 
-				// Update settings.
-				FScopeLock Lock(&SettingsCriticalSection);
-				SettingsChangeSerial++;
-				StoreDirectory = NewStoreDirectory;
-				WatchDirectories.Empty();
-				WatchDirectories = NewWatchDirectories;
-				if (Status)
-				{
-					StorePort = Status->GetStorePort();
-					RecorderPort = Status->GetRecorderPort();
-				}
-			}
+			// Update settings.
+			FScopeLock Lock(&SettingsCriticalSection);
+			SettingsChangeSerial++;
+			StoreDirectory = NewStoreDirectory;
+			WatchDirectories.Empty();
+			WatchDirectories = NewWatchDirectories;
+			StorePort = NewStorePort.Get(StorePort);
+			RecorderPort = NewRecorderPort.Get(RecorderPort);
 		}
 
 		if (StoreChangeSerial != NewStoreChangeSerial)
