[Insights] Fixed export as CSV for generic (tree) tables (like the ones from Asset Loading Insights) to write time and memory size values with unified units (i.e. double value [seconds] for time value; integer value [bytes] for a memory size value).

diff --git a/Engine/Source/Developer/TraceInsights/Private/Insights/Table/ViewModels/Table.cpp b/Engine/Source/Developer/TraceInsights/Private/Insights/Table/ViewModels/Table.cpp
--- a/Engine/Source/Developer/TraceInsights/Private/Insights/Table/ViewModels/Table.cpp
+++ b/Engine/Source/Developer/TraceInsights/Private/Insights/Table/ViewModels/Table.cpp
@@ -194,7 +194,7 @@ void FTable::GetVisibleColumnsData(const TArray<FBaseTreeNodePtr>& InNodes, cons
 			{
 				OutData += Separator;
 			}
-			FString Value = ColumnRef->GetValueAsText(*Node).ToString().ReplaceCharWithEscapedChar();
+			FString Value = ColumnRef->GetValueAsSerializableString(*Node).ReplaceCharWithEscapedChar();
 			int32 CharIndex;
 			if (Value.FindChar(Separator, CharIndex))
 			{
diff --git a/Engine/Source/Developer/TraceInsights/Private/Insights/Table/ViewModels/TableColumn.cpp b/Engine/Source/Developer/TraceInsights/Private/Insights/Table/ViewModels/TableColumn.cpp
--- a/Engine/Source/Developer/TraceInsights/Private/Insights/Table/ViewModels/TableColumn.cpp
+++ b/Engine/Source/Developer/TraceInsights/Private/Insights/Table/ViewModels/TableColumn.cpp
@@ -102,6 +102,10 @@ inline void TrimEndingZero(FString& Str)
 FString FTableColumn::GetValueAsSerializableString(const FBaseTreeNode& InNode) const
 {
 	const TOptional<FTableCellValue> Value = GetValue(InNode);
+	if (!Value.IsSet())
+	{
+		return FString();
+	}
 	switch (Value->DataType)
 	{
 		case ETableCellDataType::Double:
