Server control menu

* Add new server control type that is able to start/stop server as well as controlling sponsor mode.
* Add new controller in session browser. One control for the local server is always present, and if connected to a remote server an additional control is created.
* Add port details in tooltip for the currently connected server.

diff --git a/Engine/Plugins/TraceUtilities/Source/EditorTraceUtilities/EditorTraceUtilities.Build.cs b/Engine/Plugins/TraceUtilities/Source/EditorTraceUtilities/EditorTraceUtilities.Build.cs
--- a/Engine/Plugins/TraceUtilities/Source/EditorTraceUtilities/EditorTraceUtilities.Build.cs
+++ b/Engine/Plugins/TraceUtilities/Source/EditorTraceUtilities/EditorTraceUtilities.Build.cs
@@ -30,6 +30,7 @@ public class EditorTraceUtilities : ModuleRules
 				"SlateCore",
 				"ToolMenus",
 				"TraceAnalysis",
+				"TraceInsights",
 				"TraceLog",
 				"UATHelper",
 			}
diff --git a/Engine/Plugins/TraceUtilities/Source/EditorTraceUtilities/Private/InsightsStatusBar/EditorTraceUtilitiesStyle.h b/Engine/Plugins/TraceUtilities/Source/EditorTraceUtilities/Private/InsightsStatusBar/EditorTraceUtilitiesStyle.h
--- a/Engine/Plugins/TraceUtilities/Source/EditorTraceUtilities/Private/InsightsStatusBar/EditorTraceUtilitiesStyle.h
+++ b/Engine/Plugins/TraceUtilities/Source/EditorTraceUtilities/Private/InsightsStatusBar/EditorTraceUtilitiesStyle.h
@@ -48,6 +48,9 @@ public:
 
 		Set("Icons.Screenshot.Menu", new CORE_IMAGE_BRUSH_SVG("../Editor/Slate/Starship/Common/HighResolutionScreenshot", Icon16x16));
 		Set("Icons.Bookmark.Menu", new CORE_IMAGE_BRUSH_SVG("../Editor/Slate/Starship/Common/Bookmarks", Icon16x16));
+		
+		Set("Icons.TraceServerStart", new CORE_IMAGE_BRUSH_SVG("Starship/Common/play", Icon16x16, FStyleColors::AccentGreen));
+		Set("Icons.TraceServerStop", new CORE_IMAGE_BRUSH_SVG("Starship/Common/close", Icon16x16, FStyleColors::AccentRed));
 
 		FSlateStyleRegistry::RegisterSlateStyle(*this);
 	}
diff --git a/Engine/Plugins/TraceUtilities/Source/EditorTraceUtilities/Private/InsightsStatusBar/SInsightsStatusBar.cpp b/Engine/Plugins/TraceUtilities/Source/EditorTraceUtilities/Private/InsightsStatusBar/SInsightsStatusBar.cpp
--- a/Engine/Plugins/TraceUtilities/Source/EditorTraceUtilities/Private/InsightsStatusBar/SInsightsStatusBar.cpp
+++ b/Engine/Plugins/TraceUtilities/Source/EditorTraceUtilities/Private/InsightsStatusBar/SInsightsStatusBar.cpp
@@ -29,6 +29,7 @@
 #include "Trace/StoreClient.h"
 #include "Trace/Trace.h"
 #include "UnrealInsightsLauncher.h"
+#include "Insights/Widgets/STraceServerControl.h"
 #include "Widgets/SBoxPanel.h"
 #include "Widgets/SOverlay.h"
 #include "Widgets/Images/SImage.h"
@@ -248,6 +249,11 @@ TSharedRef<SWidget> SInsightsStatusBarWidget::MakeTraceMenu()
 {
 	LiveSessionTracker->StartQuery();
 
+	if (ServerControls.IsEmpty())
+	{
+		ServerControls.Emplace(TEXT("127.0.0.1"), 0, FEditorTraceUtilitiesStyle::Get().GetStyleSetName());
+	}
+
 	FMenuBuilder MenuBuilder(true, CommandList.ToSharedRef());
 
 	MenuBuilder.BeginSection("TraceData", LOCTEXT("TraceMenu_Section_Data", "Trace Data"));
@@ -402,6 +408,20 @@ TSharedRef<SWidget> SInsightsStatusBarWidget::MakeTraceMenu()
 
 	MenuBuilder.BeginSection("Insights", LOCTEXT("TraceMenu_Section_Insights", "Insights"));
 	{
+		MenuBuilder.AddSubMenu(
+			LOCTEXT("ServerControlLabel", "UnrealTraceServer"),
+			LOCTEXT("ServerControlTooltip", "Control UnrealTraceServer instances"),
+			FNewMenuDelegate::CreateLambda([this](FMenuBuilder& MenuBuilder)
+			{
+				for (auto& ServerControl : ServerControls)
+				{
+					 ServerControl.MakeMenu(MenuBuilder);
+				}
+			}),
+			false,
+			FSlateIcon(FAppStyle::Get().GetStyleSetName(), ("Icons.Server"))
+		);
+		
 		MenuBuilder.AddMenuEntry(
 			LOCTEXT("OpenInsightsLabel", "Unreal Insights (Session Browser)"),
 			LOCTEXT("OpenInsightsTooltip", "Launch the Unreal Insights Session Browser."),
diff --git a/Engine/Plugins/TraceUtilities/Source/EditorTraceUtilities/Private/InsightsStatusBar/SInsightsStatusBar.h b/Engine/Plugins/TraceUtilities/Source/EditorTraceUtilities/Private/InsightsStatusBar/SInsightsStatusBar.h
--- a/Engine/Plugins/TraceUtilities/Source/EditorTraceUtilities/Private/InsightsStatusBar/SInsightsStatusBar.h
+++ b/Engine/Plugins/TraceUtilities/Source/EditorTraceUtilities/Private/InsightsStatusBar/SInsightsStatusBar.h
@@ -11,6 +11,7 @@
 class FLiveSessionTracker;
 class FMenuBuilder;
 class FUICommandList;
+class STraceServerControl;
 
 TSharedRef<SWidget> CreateInsightsStatusBarWidget();
 
@@ -159,6 +160,8 @@ private:
 	TSharedPtr<FLiveSessionTracker> LiveSessionTracker;
 
 	TSharedPtr<FUICommandList> CommandList;
+	
+	TArray<STraceServerControl> ServerControls;
 
 	TArray<TSharedPtr<FTraceFileInfo>> Traces;
 	FName LogListingName;
diff --git a/Engine/Source/Developer/TraceInsights/Private/Insights/InsightsStyle.cpp b/Engine/Source/Developer/TraceInsights/Private/Insights/InsightsStyle.cpp
--- a/Engine/Source/Developer/TraceInsights/Private/Insights/InsightsStyle.cpp
+++ b/Engine/Source/Developer/TraceInsights/Private/Insights/InsightsStyle.cpp
@@ -159,6 +159,8 @@ void FInsightsStyle::FStyle::Initialize()
 	Set("Icons.Launcher", new TODO_IMAGE_BRUSH(Icon16x16));
 	Set("Icons.UTrace", new IMAGE_BRUSH_SVG("UTrace", Icon16x16));
 	Set("Icons.UTraceStack", new IMAGE_BRUSH_SVG("UTrace", Icon16x16));
+	Set("Icons.TraceServerStart", new CORE_IMAGE_BRUSH_SVG("Starship/Common/play", Icon16x16, FStyleColors::AccentGreen));
+	Set("Icons.TraceServerStop", new CORE_IMAGE_BRUSH_SVG("Starship/Common/close", Icon16x16, FStyleColors::AccentRed));
 
 	//////////////////////////////////////////////////
 	// Timing Insights
diff --git a/Engine/Source/Developer/TraceInsights/Private/Insights/Widgets/SStartPageWindow.cpp b/Engine/Source/Developer/TraceInsights/Private/Insights/Widgets/SStartPageWindow.cpp
--- a/Engine/Source/Developer/TraceInsights/Private/Insights/Widgets/SStartPageWindow.cpp
+++ b/Engine/Source/Developer/TraceInsights/Private/Insights/Widgets/SStartPageWindow.cpp
@@ -903,6 +903,9 @@ STraceStoreWindow::STraceStoreWindow()
 
 	SortColumn = TraceStoreColumns::Date;
 	SortMode = EColumnSortMode::Ascending;
+
+	// Add controls for the local server
+	ServerControls.Emplace(TEXT("127.0.0.1"), 0, FAppStyle::Get().GetStyleSetName());
 }
 
 ////////////////////////////////////////////////////////////////////////////////////////////////////
@@ -1960,7 +1963,7 @@ FText STraceStoreWindow::GetConnectionStatusTooltip() const
 {
 	using Insights::FStoreBrowser;
 
-	static FText Connected    = LOCTEXT("Connected",    "Connected to the trace server.\nServer version: {0}");
+	static FText Connected    = LOCTEXT("Connected",    "Connected to the trace server.\nServer version: {0}\nRecorder port: {1}, Store port: {2}");
 	static FText NotConnected = LOCTEXT("NoConnection", "Unable to connect to trace server.");
 	static FText Connecting   = LOCTEXT("Connecting",   "Trying to connect to trace server.");
 	static FText Disconnected = LOCTEXT("Disconnected", "Connection to trace server has been lost. Attempting to reconnect in {0} seconds.");
@@ -1973,8 +1976,14 @@ FText STraceStoreWindow::GetConnectionStatusTooltip() const
 		{
 			StoreBrowser->LockSettings();
 			FText Version = FText::FromString(StoreBrowser->GetVersion());
+			const uint32 RecorderPort = StoreBrowser->GetRecorderPort();
+			const uint32 StorePort = StoreBrowser->GetStorePort();	
 			StoreBrowser->UnlockSettings();
-			return FText::Format(Connected, Version);
+			return FText::Format(Connected,
+				Version,
+				FText::AsNumber(RecorderPort, &FNumberFormattingOptions::DefaultNoGrouping()),
+				FText::AsNumber(StorePort,  &FNumberFormattingOptions::DefaultNoGrouping())
+			);
 		}
 
 		case FStoreBrowser::EConnectionStatus::NoConnection:
@@ -2457,6 +2466,14 @@ void STraceStoreWindow::RefreshTraceList()
 		{
 			SettingsChangeSerial = NewSettingsChangeSerial;
 
+			// Add remote server controls. It's not possible to change server
+			// address on the fly so we can expect that there cannot be more than
+			// two entries (the local and possibly a currently connected remote server)
+			if (!StoreBrowser->GetHost().Equals(TEXT("127.0.0.1")) && ServerControls.Num() == 1)
+			{
+				 ServerControls.Emplace(*StoreBrowser->GetHost(), StoreBrowser->GetStorePort(), FAppStyle::Get().GetStyleSetName());
+			}
+			
 			// Update the host text
 			if (StoreHostTextBox)
 			{
@@ -3174,6 +3191,24 @@ TSharedRef<SWidget> STraceStoreWindow::MakeTraceListMenu()
 	}
 	MenuBuilder.EndSection();
 
+	MenuBuilder.BeginSection("Control", LOCTEXT("TraceListMenu_Section_Control", "Control"));
+	{
+		MenuBuilder.AddSubMenu(
+			LOCTEXT("ServerControlLabel", "UnrealTraceServer"),
+			LOCTEXT("ServerControlTooltip", "Control UnrealTraceServer instances"),
+			FNewMenuDelegate::CreateLambda([this](FMenuBuilder& MenuBuilder)
+			{
+				 for (auto& ServerControl : ServerControls)
+				 {
+					   ServerControl.MakeMenu(MenuBuilder);
+				 }
+			}),
+			false,
+			FSlateIcon(FAppStyle::Get().GetStyleSetName(), ("Icons.Server"))
+		);
+	}
+	MenuBuilder.EndSection();
+	
 	MenuBuilder.BeginSection("DebugOptions", LOCTEXT("TraceListMenu_Section_DebugOptions", "Debug Options"));
 
 	if (FGlobalTabmanager::Get()->HasTabSpawner(FInsightsManagerTabs::AutomationWindowTabId))
diff --git a/Engine/Source/Developer/TraceInsights/Private/Insights/Widgets/SStartPageWindow.h b/Engine/Source/Developer/TraceInsights/Private/Insights/Widgets/SStartPageWindow.h
--- a/Engine/Source/Developer/TraceInsights/Private/Insights/Widgets/SStartPageWindow.h
+++ b/Engine/Source/Developer/TraceInsights/Private/Insights/Widgets/SStartPageWindow.h
@@ -2,6 +2,7 @@
 
 #pragma once
 
+#include "Insights/Widgets/STraceServerControl.h"
 #include "Async/TaskGraphInterfaces.h"
 #include "Containers/Ticker.h"
 #include "Misc/FilterCollection.h"
@@ -680,6 +681,8 @@ private:
 	FTickerDelegate OnTick;
 	FTSTicker::FDelegateHandle OnTickHandle;
 
+	TArray<STraceServerControl> ServerControls;
+
 #if PLATFORM_WINDOWS
 	void* AutoConnectEvent = nullptr;
 #endif
diff --git a/Engine/Source/Developer/TraceInsights/Private/Insights/Widgets/STraceServerControl.cpp b/Engine/Source/Developer/TraceInsights/Private/Insights/Widgets/STraceServerControl.cpp
new file mode 100644
--- /dev/null
+++ b/Engine/Source/Developer/TraceInsights/Private/Insights/Widgets/STraceServerControl.cpp
@@ -0,0 +1,290 @@
+﻿// Copyright Epic Games, Inc. All Rights Reserved.
+
+#include "Insights/Widgets/STraceServerControl.h"
+
+#include "SlateOptMacros.h"
+#include "Framework/MultiBox/MultiBoxBuilder.h"
+#include "Insights/Log.h"
+#include "ProfilingDebugging/TraceAuxiliary.h"
+#include "Trace/StoreClient.h"
+
+/////////////////////////////////////////////////////////////////////////////////////////
+
+#define LOCTEXT_NAMESPACE "STraceServerControl"
+
+// Number of times to attempt state change before failing
+constexpr uint32 GStateChangeRetries = 6;
+// Number of times to attempt reconnection while starting server
+constexpr uint32 GStartConnectAttempts = 25;
+// Number of seconds between each reconnection attempt
+constexpr float GStartConnectFrequencySeconds = 0.25;
+
+/////////////////////////////////////////////////////////////////////////////////////////
+
+STraceServerControl::STraceServerControl(const TCHAR* InHost, uint32 InPort, FName InStyleSet)
+	: Host(InHost)
+	, Port(InPort)
+	, StyleSet(InStyleSet)
+{
+	bIsLocalHost = Host.Equals(TEXT("127.0.0.1")) || Host.Equals(TEXT("localhost"));
+}
+
+/////////////////////////////////////////////////////////////////////////////////////////
+
+BEGIN_SLATE_FUNCTION_BUILD_OPTIMIZATION
+
+void STraceServerControl::MakeMenu(FMenuBuilder& Builder)
+{
+	// Create client
+	if (!Client.IsValid())
+	{
+		Client = TUniquePtr<UE::Trace::FStoreClient>(UE::Trace::FStoreClient::Connect(*Host, Port));
+		if (Client)
+		{
+			ChangeState(EState::NotConnected, EState::Connected);
+		}
+	}
+
+	// If connected kick off status and version check
+	TriggerStatusUpdate();
+	
+	if (bIsLocalHost)
+	{
+		Builder.BeginSection("LocalTraceServer", LOCTEXT("Section_LocalServer", "Local Trace Server"));
+	}
+	else
+	{
+		Builder.BeginSection("TraceServer", FText::Format(LOCTEXT("Section_RemoteServer", "Remote Trace Server '{0}'"), FText::FromString(Host)));
+	}
+	Builder.AddMenuEntry(
+		TAttribute<FText>::CreateLambda([this] { FScopeLock _(&StringsLock); return FText::FromString(StatusString.IsEmpty() ? TEXT("Not running") : StatusString); }),
+		FText::GetEmpty(),
+		FSlateIcon(),
+		FUIAction(FExecuteAction(), FCanExecuteAction::CreateLambda([]() { return false; })),
+		NAME_None,
+		EUserInterfaceActionType::Button
+	);
+	if (bIsLocalHost)
+	{
+		Builder.AddMenuEntry(
+			 LOCTEXT("ServerControlSponsoredLabel", "Sponsored Mode"),
+			 LOCTEXT("ServerControlSponsoredTooltip", "In sponsored mode the server only runs as long as local processes that uses it are alive."),
+			 FSlateIcon(), //?
+			 FUIAction(
+				  FExecuteAction::CreateRaw(this, &STraceServerControl::OnSponsored_Changed),
+				  FCanExecuteAction::CreateRaw(this, &STraceServerControl::AreControlsEnabled),
+				  FGetActionCheckState::CreateLambda([this](){ return IsSponsored() ? ECheckBoxState::Checked : ECheckBoxState::Unchecked; })
+			 ),
+			 NAME_None,
+			 EUserInterfaceActionType::ToggleButton
+		);
+		
+		Builder.AddMenuEntry(
+			 LOCTEXT("ServerControlStartLabel", "Start"),
+			 LOCTEXT("ServerControlStartTooltop", "Start Server"),
+			 FSlateIcon(StyleSet, "Icons.TraceServerStart"),
+			 FUIAction(
+				  FExecuteAction::CreateRaw(this, &STraceServerControl::OnStart_Clicked),
+				  FCanExecuteAction::CreateRaw(this, &STraceServerControl::CanServerBeStarted)
+			 ),
+			 NAME_None,
+			 EUserInterfaceActionType::Button
+		);
+
+		Builder.AddMenuEntry(
+			 LOCTEXT("ServerControlStopLabel", "Stop"),
+			 LOCTEXT("ServerControlStopTooltop", "Stop Server. Any running traces will be cancelled."),
+			 FSlateIcon(StyleSet, "Icons.TraceServerStop"),
+			 FUIAction(
+				  FExecuteAction::CreateRaw(this, &STraceServerControl::OnStop_Clicked),
+				  FCanExecuteAction::CreateRaw(this, &STraceServerControl::CanServerBeStopped)
+			 ),
+			 NAME_None,
+			 EUserInterfaceActionType::Button
+		);
+	}
+	Builder.EndSection();
+}
+
+END_SLATE_FUNCTION_BUILD_OPTIMIZATION
+
+/////////////////////////////////////////////////////////////////////////////////////////
+
+const TCHAR* LexState(STraceServerControl::EState state)
+{
+	switch(state)
+	{
+	case STraceServerControl::EState::NotConnected: return TEXT("NotConnected");
+	case STraceServerControl::EState::Connecting: return TEXT("Connecting");
+	case STraceServerControl::EState::Connected: return TEXT("Connected");
+	case STraceServerControl::EState::CheckStatus: return TEXT("CheckStatus");
+	case STraceServerControl::EState::Command: return TEXT("Command");
+	default: return TEXT("Unknown");
+	}
+}
+
+/////////////////////////////////////////////////////////////////////////////////////////
+
+bool STraceServerControl::ChangeState(EState Expected, EState ChangeTo, uint32 Attempts)
+{
+	while (Attempts-- > 0)
+	{
+		const bool bChanged = State.compare_exchange_strong(Expected, ChangeTo);
+		if (bChanged)
+		{
+			UE_LOG(LogTemp, VeryVerbose, TEXT("Changing state from '%s' -> '%s'"), LexState(Expected), LexState(ChangeTo))
+			return true;
+		}
+		UE_LOG(LogTemp, VeryVerbose, TEXT("Busy wait for '%s'..."), LexState(Expected));
+		FPlatformProcess::Sleep(0.25);
+	}
+	return false;
+}
+
+/////////////////////////////////////////////////////////////////////////////////////////
+
+void STraceServerControl::TriggerStatusUpdate()
+{
+	if (State.load(std::memory_order_relaxed) == EState::Connected && ChangeState(EState::Connected, EState::CheckStatus))
+	{
+		FFunctionGraphTask::CreateAndDispatchWhenReady([this]
+		{
+			UpdateStatus();
+			ChangeState(EState::CheckStatus, EState::Connected);
+		});
+	}
+}
+
+/////////////////////////////////////////////////////////////////////////////////////////
+
+void STraceServerControl::UpdateStatus()
+{
+	const UE::Trace::FStoreClient::FStatus* Status = Client->GetStatus();
+
+	// Check current status
+	const bool bServerIsRunning = Status != nullptr;
+	bCanServerBeStarted.store(!bServerIsRunning, std::memory_order_relaxed);
+	bCanServerBeStopped.store(bServerIsRunning, std::memory_order_relaxed);
+
+	if (!bServerIsRunning)
+	{
+		ResetStatus();
+		Client.Reset();
+		ChangeState(EState::CheckStatus, EState::NotConnected);
+		return;
+	}
+	
+	TStringBuilder<64> PortsStringBuilder;
+	if (Status)
+	{
+		bSponsored.store(Status->GetSponsored());
+		PortsStringBuilder << TEXT("Recorder Port: ") << Status->GetRecorderPort()
+		<< TEXT(", Store Port: ") << Status->GetStorePort();
+	}
+
+	// If not previously checked, also query version information
+	if (StatusString.IsEmpty())
+	{
+		const UE::Trace::FStoreClient::FVersion* Version = Client->GetVersion();
+		if (Version)
+		{
+			TStringBuilder<64> StatusStringBuilder;
+			StatusStringBuilder << TEXT("Version: ") << Version->GetMajorVersion() << TEXT(".") << Version->GetMinorVersion();
+			// Only print configuration if it's not release
+			if (!Version->GetConfiguration().Equals(UTF8TEXT("Release")))
+			{
+				StatusStringBuilder << TEXT(" (") << Version->GetConfiguration() << TEXT(")");
+			}
+			StatusStringBuilder << TEXT(", ") << PortsStringBuilder;
+			{
+				FScopeLock Lock(&StringsLock);
+				StatusString = StatusStringBuilder.ToString();
+			}
+		}
+	}
+}
+
+/////////////////////////////////////////////////////////////////////////////////////////
+
+void STraceServerControl::ResetStatus()
+{
+	FScopeLock Lock(&StringsLock);
+	StatusString.Empty();
+}
+
+/////////////////////////////////////////////////////////////////////////////////////////
+
+void STraceServerControl::OnStart_Clicked()
+{
+	FGraphEventRef CommandTask = FFunctionGraphTask::CreateAndDispatchWhenReady([this]
+	{
+		if (ChangeState(EState::NotConnected, EState::Command, GStateChangeRetries))
+		{
+			FTraceServerControls::Start();
+			ChangeState(EState::Command, EState::Connecting);
+		}
+		else
+		{
+			UE_LOG(TraceInsights, Warning, TEXT("Failed to start server."));
+		}
+	});
+	FFunctionGraphTask::CreateAndDispatchWhenReady([this]
+	{
+		uint32 Attempts = 0;
+		while (!Client && Attempts++ < GStartConnectAttempts)
+		{
+			Client = TUniquePtr<UE::Trace::FStoreClient>(UE::Trace::FStoreClient::Connect(TEXT("127.0.0.1")));
+			FPlatformProcess::Sleep(GStartConnectFrequencySeconds);
+		}
+		if (Client)
+		{
+			ChangeState(EState::Connecting, EState::Connected);
+		}
+		else
+		{
+			UE_LOG(LogTemp, Warning, TEXT("Failed to connect to store"));
+		}
+	}, TStatId(), {CommandTask});
+}
+
+/////////////////////////////////////////////////////////////////////////////////////////
+
+void STraceServerControl::OnStop_Clicked()
+{
+	FFunctionGraphTask::CreateAndDispatchWhenReady([this]
+	{
+		if (ChangeState(EState::Connected, EState::Command, GStateChangeRetries))
+		{
+			FTraceServerControls::Stop();
+			Client.Reset();
+			ResetStatus();
+			ChangeState(EState::Command, EState::NotConnected);
+		}
+		else
+		{
+			UE_LOG(TraceInsights, Warning, TEXT("Failed to stop server."));
+		}
+	});
+}
+
+/////////////////////////////////////////////////////////////////////////////////////////
+
+void STraceServerControl::OnSponsored_Changed()
+{
+	FFunctionGraphTask::CreateAndDispatchWhenReady([this]
+	{
+		if (ChangeState(EState::Connected, EState::Command, GStateChangeRetries))
+		{
+			const bool Success = Client->SetSponsored(!IsSponsored());
+			ChangeState(EState::Command, EState::Connected);
+			if (Success)
+			{
+				TriggerStatusUpdate();
+			}
+		}
+		else
+		{
+			UE_LOG(TraceInsights, Warning, TEXT("Failed to set sponsored mode."));
+		}
+	});
+}
diff --git a/Engine/Source/Developer/TraceInsights/Public/Insights/Widgets/STraceServerControl.h b/Engine/Source/Developer/TraceInsights/Public/Insights/Widgets/STraceServerControl.h
new file mode 100644
--- /dev/null
+++ b/Engine/Source/Developer/TraceInsights/Public/Insights/Widgets/STraceServerControl.h
@@ -0,0 +1,67 @@
+﻿// Copyright Epic Games, Inc. All Rights Reserved.
+
+#pragma once
+
+#include <atomic>
+#include "CoreFwd.h"
+#include "HAL/CriticalSection.h"
+#include "Templates/UniquePtr.h"
+#include "UObject/NameTypes.h"
+#include "UObject/UnrealNames.h"
+
+class FMenuBuilder;
+
+namespace UE::Trace
+{
+	class FStoreClient;
+}
+
+class TRACEINSIGHTS_API STraceServerControl
+{
+public:
+	STraceServerControl(const TCHAR* Host, uint32 Port = 0, FName StyleSet = NAME_None);
+	~STraceServerControl() = default;
+
+	void MakeMenu(FMenuBuilder& Builder);
+
+private:
+	enum class EState : uint8
+	{
+		NotConnected,
+		Connecting,
+		Connected,
+		CheckStatus,
+		Command
+	};
+
+	bool ChangeState(EState Expected, EState ChangeTo, uint32 Attempts = 1);
+
+	void TriggerStatusUpdate();
+	void UpdateStatus();
+	void ResetStatus();
+
+	bool CanServerBeStarted() const { return bIsLocalHost && State.load(std::memory_order_relaxed) == EState::NotConnected; }
+	bool CanServerBeStopped() const { return bIsLocalHost && State.load(std::memory_order_relaxed) == EState::Connected; }
+	bool AreControlsEnabled() const { return bIsLocalHost && State.load(std::memory_order_relaxed) == EState::Connected; }
+	bool IsSponsored() const { return bSponsored.load(std::memory_order_relaxed); }
+
+	void OnStart_Clicked();
+	void OnStop_Clicked();
+	void OnSponsored_Changed();
+
+	std::atomic<EState> State = EState::NotConnected;
+
+	std::atomic<bool> bCanServerBeStarted = false;
+	std::atomic<bool> bCanServerBeStopped = false;
+	std::atomic<bool> bSponsored = false;
+	FCriticalSection StringsLock;
+	FString StatusString;
+
+	FString Host;
+	uint32 Port;
+	FName StyleSet;
+	bool bIsLocalHost;
+	TUniquePtr<UE::Trace::FStoreClient> Client;
+	
+	friend const TCHAR* LexState(EState);
+};
