Fix fwd declare issue

diff --git a/Engine/Source/Developer/TraceInsights/Public/Insights/Widgets/STraceServerControl.h b/Engine/Source/Developer/TraceInsights/Public/Insights/Widgets/STraceServerControl.h
--- a/Engine/Source/Developer/TraceInsights/Public/Insights/Widgets/STraceServerControl.h
+++ b/Engine/Source/Developer/TraceInsights/Public/Insights/Widgets/STraceServerControl.h
@@ -6,16 +6,12 @@
 #include "CoreFwd.h"
 #include "HAL/CriticalSection.h"
 #include "Templates/UniquePtr.h"
+#include "Trace/StoreClient.h"
 #include "UObject/NameTypes.h"
 #include "UObject/UnrealNames.h"
 
 class FMenuBuilder;
 
-namespace UE::Trace
-{
-	class FStoreClient;
-}
-
 class TRACEINSIGHTS_API STraceServerControl
 {
 public:
