Fix for configurations that doesn't have UTS launch controls.

diff --git a/Engine/Source/Developer/TraceInsights/Private/Insights/Widgets/STraceServerControl.cpp b/Engine/Source/Developer/TraceInsights/Private/Insights/Widgets/STraceServerControl.cpp
--- a/Engine/Source/Developer/TraceInsights/Private/Insights/Widgets/STraceServerControl.cpp
+++ b/Engine/Source/Developer/TraceInsights/Private/Insights/Widgets/STraceServerControl.cpp
@@ -220,7 +220,9 @@ void STraceServerControl::OnStart_Clicked()
 	{
 		if (ChangeState(EState::NotConnected, EState::Command, GStateChangeRetries))
 		{
+#if UE_TRACE_SERVER_CONTROLS_ENABLED
 			FTraceServerControls::Start();
+#endif
 			ChangeState(EState::Command, EState::Connecting);
 		}
 		else
@@ -255,7 +257,9 @@ void STraceServerControl::OnStop_Clicked()
 	{
 		if (ChangeState(EState::Connected, EState::Command, GStateChangeRetries))
 		{
+#if UE_TRACE_SERVER_CONTROLS_ENABLED
 			FTraceServerControls::Stop();
+#endif
 			Client.Reset();
 			ResetStatus();
 			ChangeState(EState::Command, EState::NotConnected);
@@ -288,3 +292,5 @@ void STraceServerControl::OnSponsored_Changed()
 		}
 	});
 }
+
+#undef LOCTEXT_NAMESPACE
\ No newline at end of file
