[Insights] Support for basic filters in the Asset Loading Track.

diff --git a/Engine/Source/Developer/TraceInsights/Private/Insights/ContextSwitches/ViewModels/CpuCoreTimingTrack.cpp b/Engine/Source/Developer/TraceInsights/Private/Insights/ContextSwitches/ViewModels/CpuCoreTimingTrack.cpp
--- a/Engine/Source/Developer/TraceInsights/Private/Insights/ContextSwitches/ViewModels/CpuCoreTimingTrack.cpp
+++ b/Engine/Source/Developer/TraceInsights/Private/Insights/ContextSwitches/ViewModels/CpuCoreTimingTrack.cpp
@@ -13,7 +13,6 @@
 #include "Insights/InsightsManager.h"
 #include "Insights/ViewModels/FilterConfigurator.h"
 #include "Insights/ViewModels/Filters.h"
-#include "Insights/ViewModels/ThreadTimingTrack.h"
 #include "Insights/ViewModels/TimingTrackViewport.h"
 #include "Insights/ViewModels/TimingViewDrawHelper.h"
 #include "Insights/ViewModels/TooltipDrawState.h"
diff --git a/Engine/Source/Developer/TraceInsights/Private/Insights/CookProfiler/Widgets/SPackageTableTreeView.cpp b/Engine/Source/Developer/TraceInsights/Private/Insights/CookProfiler/Widgets/SPackageTableTreeView.cpp
--- a/Engine/Source/Developer/TraceInsights/Private/Insights/CookProfiler/Widgets/SPackageTableTreeView.cpp
+++ b/Engine/Source/Developer/TraceInsights/Private/Insights/CookProfiler/Widgets/SPackageTableTreeView.cpp
@@ -16,7 +16,6 @@
 #include "Insights/Table/ViewModels/TreeNodeGrouping.h"
 #include "Insights/TimingProfilerManager.h"
 #include "Insights/ViewModels/FilterConfigurator.h"
-#include "Insights/ViewModels/ThreadTimingTrack.h"
 #include "Insights/Widgets/STimingProfilerWindow.h"
 #include "Insights/Widgets/STimingView.h"
 
diff --git a/Engine/Source/Developer/TraceInsights/Private/Insights/ViewModels/Filters.h b/Engine/Source/Developer/TraceInsights/Private/Insights/ViewModels/Filters.h
--- a/Engine/Source/Developer/TraceInsights/Private/Insights/ViewModels/Filters.h
+++ b/Engine/Source/Developer/TraceInsights/Private/Insights/ViewModels/Filters.h
@@ -45,6 +45,21 @@ enum class EFilterOperator : uint8
 
 ////////////////////////////////////////////////////////////////////////////////////////////////////
 
+enum class EFilterField : int32
+{
+	StartTime = 0,
+	EndTime = 1,
+	Duration = 2,
+	TrackName = 3,
+	TimerId = 4,
+	TimerName = 5,
+	CoreEventName = 6,
+	RegionName = 7,
+	Metadata = 8,
+};
+
+////////////////////////////////////////////////////////////////////////////////////////////////////
+
 class IFilterOperator
 {
 public:
diff --git a/Engine/Source/Developer/TraceInsights/Private/Insights/ViewModels/LoadingTimingTrack.cpp b/Engine/Source/Developer/TraceInsights/Private/Insights/ViewModels/LoadingTimingTrack.cpp
--- a/Engine/Source/Developer/TraceInsights/Private/Insights/ViewModels/LoadingTimingTrack.cpp
+++ b/Engine/Source/Developer/TraceInsights/Private/Insights/ViewModels/LoadingTimingTrack.cpp
@@ -15,6 +15,8 @@
 #include "Insights/InsightsManager.h"
 #include "Insights/InsightsStyle.h"
 #include "Insights/ITimingViewSession.h"
+#include "Insights/ViewModels/Filters.h"
+#include "Insights/ViewModels/FilterConfigurator.h"
 #include "Insights/ViewModels/TimingEvent.h"
 #include "Insights/ViewModels/TimingTrackViewport.h"
 #include "Insights/ViewModels/TooltipDrawState.h"
@@ -334,6 +336,56 @@ void FLoadingTimingTrack::BuildDrawState(ITimingEventsTrackDrawStateBuilder& Bui
 
 ////////////////////////////////////////////////////////////////////////////////////////////////////
 
+void FLoadingTimingTrack::BuildFilteredDrawState(ITimingEventsTrackDrawStateBuilder& Builder, const ITimingTrackUpdateContext& Context)
+{
+	if (HasCustomFilter())
+	{
+		using namespace Insights;
+
+		FFilterContext FilterConfiguratorContext;
+		FilterConfiguratorContext.SetReturnValueForUnsetFilters(false);
+		FilterConfiguratorContext.AddFilterData<double>(static_cast<int32>(EFilterField::StartTime), 0.0f);
+		FilterConfiguratorContext.AddFilterData<double>(static_cast<int32>(EFilterField::EndTime), 0.0f);
+		FilterConfiguratorContext.AddFilterData<double>(static_cast<int32>(EFilterField::Duration), 0.0f);
+		FilterConfiguratorContext.AddFilterData<FString>(static_cast<int32>(EFilterField::TrackName), this->GetName());
+
+		TSharedPtr<const TraceServices::IAnalysisSession> Session = FInsightsManager::Get()->GetSession();
+		if (Session.IsValid() && TraceServices::ReadLoadTimeProfilerProvider(*Session.Get()))
+		{
+			TraceServices::FAnalysisSessionReadScope SessionReadScope(*Session.Get());
+
+			const TraceServices::ILoadTimeProfilerProvider& LoadTimeProfilerProvider = *TraceServices::ReadLoadTimeProfilerProvider(*Session.Get());
+
+			const FTimingTrackViewport& Viewport = Context.GetViewport();
+
+			LoadTimeProfilerProvider.ReadTimeline(TimelineIndex, [this, &Builder, &Viewport, &FilterConfiguratorContext](const TraceServices::ILoadTimeProfilerProvider::CpuTimeline& Timeline)
+			{
+				Timeline.EnumerateEvents(Viewport.GetStartTime(), Viewport.GetEndTime(), [this, &Builder, &FilterConfiguratorContext](double StartTime, double EndTime, uint32 Depth, const TraceServices::FLoadTimeProfilerCpuEvent& Event)
+				{
+					if (Event.Package)
+					{
+						FilterConfiguratorContext.SetFilterData<double>(static_cast<int32>(EFilterField::StartTime), StartTime);
+						FilterConfiguratorContext.SetFilterData<double>(static_cast<int32>(EFilterField::EndTime), EndTime);
+						FilterConfiguratorContext.SetFilterData<double>(static_cast<int32>(EFilterField::Duration), EndTime - StartTime);
+
+						if (FilterConfigurator->ApplyFilters(FilterConfiguratorContext))
+						{
+							const TCHAR* Name = SharedState.GetEventName(Depth, Event);
+							const uint64 Type = static_cast<uint64>(Event.EventType);
+							const uint32 Color = 0;
+							Builder.AddEvent(StartTime, EndTime, Depth, Name, Type, Color);
+						}
+					}
+
+					return TraceServices::EEventEnumerate::Continue;
+				});
+			});
+		}
+	}
+}
+
+////////////////////////////////////////////////////////////////////////////////////////////////////
+
 void FLoadingTimingTrack::InitTooltip(FTooltipDrawState& InOutTooltip, const ITimingEvent& InTooltipEvent) const
 {
 	if (InTooltipEvent.CheckTrack(this) && InTooltipEvent.Is<FTimingEvent>())
@@ -402,6 +454,15 @@ const TSharedPtr<const ITimingEvent> FLoadingTimingTrack::SearchEvent(const FTim
 
 bool FLoadingTimingTrack::FindLoadTimeProfilerCpuEvent(const FTimingEventSearchParameters& InParameters, TFunctionRef<void(double, double, uint32, const TraceServices::FLoadTimeProfilerCpuEvent&)> InFoundPredicate) const
 {
+	using namespace Insights;
+
+	FFilterContext FilterConfiguratorContext;
+	FilterConfiguratorContext.SetReturnValueForUnsetFilters(false);
+	FilterConfiguratorContext.AddFilterData<double>(static_cast<int32>(EFilterField::StartTime), 0.0f);
+	FilterConfiguratorContext.AddFilterData<double>(static_cast<int32>(EFilterField::EndTime), 0.0f);
+	FilterConfiguratorContext.AddFilterData<double>(static_cast<int32>(EFilterField::Duration), 0.0f);
+	FilterConfiguratorContext.AddFilterData<FString>(static_cast<int32>(EFilterField::TrackName), this->GetName());
+
 	return TTimingEventSearch<TraceServices::FLoadTimeProfilerCpuEvent>::Search(
 		InParameters,
 
@@ -418,7 +479,7 @@ bool FLoadingTimingTrack::FindLoadTimeProfilerCpuEvent(const FTimingEventSearchP
 
 					LoadTimeProfilerProvider.ReadTimeline(TimelineIndex, [&InContext](const TraceServices::ILoadTimeProfilerProvider::CpuTimeline& Timeline)
 					{
-						Timeline.EnumerateEvents(InContext.GetParameters().StartTime, InContext.GetParameters().EndTime, [&InContext](double EventStartTime, double EventEndTime, uint32 EventDepth, const TraceServices::FLoadTimeProfilerCpuEvent& Event)
+						auto Callback = [&InContext](double EventStartTime, double EventEndTime, uint32 EventDepth, const TraceServices::FLoadTimeProfilerCpuEvent& Event)
 						{
 							if (Event.Package)
 							{
@@ -429,16 +490,70 @@ bool FLoadingTimingTrack::FindLoadTimeProfilerCpuEvent(const FTimingEventSearchP
 							{
 								return TraceServices::EEventEnumerate::Continue;
 							}
-						});
+						};
+
+						if (InContext.GetParameters().SearchDirection == FTimingEventSearchParameters::ESearchDirection::Forward)
+						{
+							Timeline.EnumerateEvents(InContext.GetParameters().StartTime, InContext.GetParameters().EndTime, Callback);
+						}
+						else
+						{
+							Timeline.EnumerateEventsBackwards(InContext.GetParameters().EndTime, InContext.GetParameters().StartTime, Callback);
+						}
 					});
 				}
 			}
 		},
 
+		[&FilterConfiguratorContext, &InParameters](double EventStartTime, double EventEndTime, uint32 EventDepth, const TraceServices::FLoadTimeProfilerCpuEvent& Event)
+		{
+			if (!InParameters.FilterExecutor.IsValid())
+			{
+				return true;
+			}
+
+			TSharedPtr<const TraceServices::IAnalysisSession> Session = FInsightsManager::Get()->GetSession();
+			if (Session.IsValid())
+			{
+				TraceServices::FAnalysisSessionReadScope SessionReadScope(*Session.Get());
+
+				if (TraceServices::ReadTimingProfilerProvider(*Session.Get()))
+				{
+					FilterConfiguratorContext.SetFilterData<double>(static_cast<int32>(EFilterField::StartTime), EventStartTime);
+					FilterConfiguratorContext.SetFilterData<double>(static_cast<int32>(EFilterField::EndTime), EventEndTime);
+					FilterConfiguratorContext.SetFilterData<double>(static_cast<int32>(EFilterField::Duration), EventEndTime - EventStartTime);
+
+					return InParameters.FilterExecutor->ApplyFilters(FilterConfiguratorContext);
+				}
+			}
+
+			return false;
+		},
+
 		[&InFoundPredicate](double InFoundStartTime, double InFoundEndTime, uint32 InFoundDepth, const TraceServices::FLoadTimeProfilerCpuEvent& InEvent)
 		{
 			InFoundPredicate(InFoundStartTime, InFoundEndTime, InFoundDepth, InEvent);
-		});
+		},
+
+		TTimingEventSearch<TraceServices::FLoadTimeProfilerCpuEvent>::NoMatch);
+}
+
+////////////////////////////////////////////////////////////////////////////////////////////////////
+
+bool FLoadingTimingTrack::HasCustomFilter() const
+{
+	return FilterConfigurator.IsValid() && !FilterConfigurator->IsEmpty();
+}
+
+////////////////////////////////////////////////////////////////////////////////////////////////////
+
+void FLoadingTimingTrack::SetFilterConfigurator(TSharedPtr<Insights::FFilterConfigurator> InFilterConfigurator)
+{
+	if (FilterConfigurator != InFilterConfigurator)
+	{
+		FilterConfigurator = InFilterConfigurator;
+		SetDirtyFlag();
+	}
 }
 
 ////////////////////////////////////////////////////////////////////////////////////////////////////
diff --git a/Engine/Source/Developer/TraceInsights/Private/Insights/ViewModels/LoadingTimingTrack.h b/Engine/Source/Developer/TraceInsights/Private/Insights/ViewModels/LoadingTimingTrack.h
--- a/Engine/Source/Developer/TraceInsights/Private/Insights/ViewModels/LoadingTimingTrack.h
+++ b/Engine/Source/Developer/TraceInsights/Private/Insights/ViewModels/LoadingTimingTrack.h
@@ -104,17 +104,25 @@ public:
 
 	virtual void BuildDrawState(ITimingEventsTrackDrawStateBuilder& Builder, const ITimingTrackUpdateContext& Context) override;
 
+	virtual void BuildFilteredDrawState(ITimingEventsTrackDrawStateBuilder& Builder, const ITimingTrackUpdateContext& Context) override;
+
 	virtual void InitTooltip(FTooltipDrawState& InOutTooltip, const ITimingEvent& InTooltipEvent) const override;
 
 	virtual const TSharedPtr<const ITimingEvent> SearchEvent(const FTimingEventSearchParameters& InSearchParameters) const override;
 
+	virtual void SetFilterConfigurator(TSharedPtr<Insights::FFilterConfigurator> InFilterConfigurator) override;
+
 protected:
 	// Helper function to find an event given search parameters
 	bool FindLoadTimeProfilerCpuEvent(const FTimingEventSearchParameters& InParameters, TFunctionRef<void(double, double, uint32, const TraceServices::FLoadTimeProfilerCpuEvent&)> InFoundPredicate) const;
 
+	virtual bool HasCustomFilter() const override;
+
 protected:
 	FLoadingSharedState& SharedState;
 	uint32 TimelineIndex;
+
+	TSharedPtr<Insights::FFilterConfigurator> FilterConfigurator;
 };
 
 ////////////////////////////////////////////////////////////////////////////////////////////////////
diff --git a/Engine/Source/Developer/TraceInsights/Private/Insights/ViewModels/RegionsTimingTrack.cpp b/Engine/Source/Developer/TraceInsights/Private/Insights/ViewModels/RegionsTimingTrack.cpp
--- a/Engine/Source/Developer/TraceInsights/Private/Insights/ViewModels/RegionsTimingTrack.cpp
+++ b/Engine/Source/Developer/TraceInsights/Private/Insights/ViewModels/RegionsTimingTrack.cpp
@@ -19,7 +19,6 @@
 #include "Insights/TimingProfilerCommon.h"
 #include "Insights/ViewModels/FilterConfigurator.h"
 #include "Insights/ViewModels/Filters.h"
-#include "Insights/ViewModels/ThreadTimingTrack.h"
 #include "Insights/ViewModels/TimingTrackViewport.h"
 #include "Insights/ViewModels/TimingViewDrawHelper.h"
 #include "Insights/Widgets/STimingView.h"
diff --git a/Engine/Source/Developer/TraceInsights/Private/Insights/ViewModels/ThreadTimingTrack.h b/Engine/Source/Developer/TraceInsights/Private/Insights/ViewModels/ThreadTimingTrack.h
--- a/Engine/Source/Developer/TraceInsights/Private/Insights/ViewModels/ThreadTimingTrack.h
+++ b/Engine/Source/Developer/TraceInsights/Private/Insights/ViewModels/ThreadTimingTrack.h
@@ -124,21 +124,6 @@ private:
 
 ////////////////////////////////////////////////////////////////////////////////////////////////////
 
-enum class EFilterField : int32
-{
-	StartTime = 0,
-	EndTime = 1,
-	Duration = 2,
-	TrackName = 3,
-	TimerId = 4,
-	TimerName = 5,
-	CoreEventName = 6,
-	RegionName = 7,
-	Metadata = 8,
-};
-
-////////////////////////////////////////////////////////////////////////////////////////////////////
-
 class FThreadTimingTrack : public FTimingEventsTrack
 {
 	INSIGHTS_DECLARE_RTTI(FThreadTimingTrack, FTimingEventsTrack)
@@ -190,7 +175,7 @@ public:
 
 	int32 GetDepthAt(double Time) const;
 
-	virtual void SetFilterConfigurator(TSharedPtr<Insights::FFilterConfigurator> InFilterConfigurator);
+	virtual void SetFilterConfigurator(TSharedPtr<Insights::FFilterConfigurator> InFilterConfigurator) override;
 
 	TSharedPtr<const ITimingEvent> FindMaxEventInstance(uint32 TimerId, double StartTime, double EndTime) const;
 	TSharedPtr<const ITimingEvent> FindMinEventInstance(uint32 TimerId, double StartTime, double EndTime) const;
diff --git a/Engine/Source/Developer/TraceInsights/Private/Insights/Widgets/STimingView.cpp b/Engine/Source/Developer/TraceInsights/Private/Insights/Widgets/STimingView.cpp
--- a/Engine/Source/Developer/TraceInsights/Private/Insights/Widgets/STimingView.cpp
+++ b/Engine/Source/Developer/TraceInsights/Private/Insights/Widgets/STimingView.cpp
@@ -5310,6 +5310,8 @@ void STimingView::PopulateTimerNameSuggestionList(const FString& Text, TArray<FS
 
 void STimingView::EnumerateFilteredTracks(TSharedPtr<Insights::FFilterConfigurator> InFilterConfigurator, TSharedPtr<const FBaseTimingTrack> PriorityTrack, EnumerateFilteredTracksCallback Callback)
 {
+	using namespace Insights;
+
 	Insights::FFilterContext FilterContext;
 	FilterContext.AddFilterData(static_cast<int32>(EFilterField::TrackName), FString());
 
