[Insights] Added validation for renaming a utrace file with an empty name.

diff --git a/Engine/Source/Developer/TraceInsights/Private/Insights/Widgets/SStartPageWindow.cpp b/Engine/Source/Developer/TraceInsights/Private/Insights/Widgets/SStartPageWindow.cpp
--- a/Engine/Source/Developer/TraceInsights/Private/Insights/Widgets/SStartPageWindow.cpp
+++ b/Engine/Source/Developer/TraceInsights/Private/Insights/Widgets/SStartPageWindow.cpp
@@ -272,7 +272,7 @@ public:
 		const FString TraceName = Trace.Name.ToString();
 		const FString NewTraceName = InText.ToString().TrimStartAndEnd();
 
-		if (NewTraceName == TraceName)
+		if (NewTraceName == TraceName || NewTraceName.IsEmpty())
 		{
 			return;
 		}
