[Insights]
  - Callers and Callees: Added Average (Inclusive and Exclusive) Time columns for Callers and Callees tree views.
  - Callers and Callees: Fixed tooltip to hide stats that are not computed (Min/Max/Median).

diff --git a/Engine/Source/Developer/TraceInsights/Private/Insights/ViewModels/TimersViewColumnFactory.cpp b/Engine/Source/Developer/TraceInsights/Private/Insights/ViewModels/TimersViewColumnFactory.cpp
--- a/Engine/Source/Developer/TraceInsights/Private/Insights/ViewModels/TimersViewColumnFactory.cpp
+++ b/Engine/Source/Developer/TraceInsights/Private/Insights/ViewModels/TimersViewColumnFactory.cpp
@@ -72,7 +72,9 @@ void FTimersViewColumnFactory::CreateTimerTreeViewColumns(TArray<TSharedRef<Insi
 	Columns.Add(CreateTypeColumn());
 	Columns.Add(CreateInstanceCountColumn());
 	Columns.Add(CreateTotalInclusiveTimeColumn());
+	Columns.Add(CreateAverageInclusiveTimeColumn());
 	Columns.Add(CreateTotalExclusiveTimeColumn());
+	Columns.Add(CreateAverageExclusiveTimeColumn());
 }
 
 ////////////////////////////////////////////////////////////////////////////////////////////////////
@@ -583,7 +585,7 @@ TSharedRef<Insights::FTableColumn> FTimersViewColumnFactory::CreateMaxExclusiveT
 	Column.SetDescription(ETraceFrameType::TraceFrameType_Count, LOCTEXT("MaxExclusiveTime_ColumnDesc", "Maximum exclusive duration of selected timer's instances, in milliseconds"));
 	Column.SetDescription(ETraceFrameType::TraceFrameType_Game, LOCTEXT("MaxExclusiveTime_GameFrameColumnDesc", "Game Frame Maximum Exclusive Duration.\nExclusive duration is computed for a single frame as the sum of exclusive duration of all instances of the timer in the respective frame.\nThe maximum is selected from these per-frame Exclusive durations. Unit is miliseconds."));
 	Column.SetDescription(ETraceFrameType::TraceFrameType_Rendering, LOCTEXT("MaxExclusiveTime_RenderingFrameColumnDesc", "Rendering Frame Maximum Exclusive Duration.\nExclusive duration is computed for a single frame as the sum of exclusive duration of all instances of the timer in the respective frame. The maximum is selected from these per-frame Exclusive durations.\nUnit is miliseconds."));
-	
+
 	Column.SetFlags(ETableColumnFlags::CanBeHidden |
 					//ETableColumnFlags::ShouldBeVisible |
 					ETableColumnFlags::CanBeFiltered);
diff --git a/Engine/Source/Developer/TraceInsights/Private/Insights/Widgets/STimerTreeView.cpp b/Engine/Source/Developer/TraceInsights/Private/Insights/Widgets/STimerTreeView.cpp
--- a/Engine/Source/Developer/TraceInsights/Private/Insights/Widgets/STimerTreeView.cpp
+++ b/Engine/Source/Developer/TraceInsights/Private/Insights/Widgets/STimerTreeView.cpp
@@ -824,7 +824,7 @@ void STimerTreeView::ShowColumn(const FName ColumnId)
 	ColumnArgs
 		.ColumnId(Column.GetId())
 		.DefaultLabel(Column.GetShortName())
-		.ToolTip(STimersViewTooltip::GetColumnTooltip(Column))
+		.ToolTip(STimersViewTooltip::GetColumnTooltipForMode(Column, ETraceFrameType::TraceFrameType_Count))
 		.HAlignHeader(Column.GetHorizontalAlignment())
 		.VAlignHeader(VAlign_Center)
 		.HAlignCell(HAlign_Fill)
@@ -1035,6 +1035,15 @@ FTimerNodePtr STimerTreeView::CreateTimerNodeRec(const TraceServices::FTimingPro
 	AggregatedStats.InstanceCount = Node.Count;
 	AggregatedStats.TotalInclusiveTime = Node.InclusiveTime;
 	AggregatedStats.TotalExclusiveTime = Node.ExclusiveTime;
+	AggregatedStats.AverageInclusiveTime = Node.Count != 0 ? Node.InclusiveTime / (double)Node.Count : 0.0;
+	AggregatedStats.AverageExclusiveTime = Node.Count != 0 ? Node.ExclusiveTime / (double)Node.Count : 0.0;
+	constexpr double NanTimeValue = std::numeric_limits<double>::quiet_NaN();
+	AggregatedStats.MinInclusiveTime = NanTimeValue;
+	AggregatedStats.MinExclusiveTime = NanTimeValue;
+	AggregatedStats.MaxInclusiveTime = NanTimeValue;
+	AggregatedStats.MaxExclusiveTime = NanTimeValue;
+	AggregatedStats.MedianInclusiveTime = NanTimeValue;
+	AggregatedStats.MedianExclusiveTime = NanTimeValue;
 	TimerNodePtr->SetAggregatedStats(AggregatedStats);
 
 	for (const TraceServices::FTimingProfilerButterflyNode* ChildNodePtr : Node.Children)
diff --git a/Engine/Source/Developer/TraceInsights/Private/Insights/Widgets/STimersViewTooltip.cpp b/Engine/Source/Developer/TraceInsights/Private/Insights/Widgets/STimersViewTooltip.cpp
--- a/Engine/Source/Developer/TraceInsights/Private/Insights/Widgets/STimersViewTooltip.cpp
+++ b/Engine/Source/Developer/TraceInsights/Private/Insights/Widgets/STimersViewTooltip.cpp
@@ -22,6 +22,8 @@
 #include "Insights/Widgets/STimersView.h"
 #include "Insights/Widgets/STimingProfilerWindow.h"
 
+#include <cmath>
+
 ////////////////////////////////////////////////////////////////////////////////////////////////////
 
 #define LOCTEXT_NAMESPACE "STimersView"
@@ -127,28 +129,45 @@ TSharedPtr<SToolTip> STimersViewTooltip::GetColumnTooltip(const Insights::FTable
 
 ////////////////////////////////////////////////////////////////////////////////////////////////////
 
-TSharedPtr<SToolTip> STimersViewTooltip::GetRowTooltip(const TSharedPtr<FTimerNode> TimerNodePtr)
+TSharedPtr<SToolTip> STimersViewTooltip::GetColumnTooltipForMode(const Insights::FTableColumn& Column, ETraceFrameType InAggregationMode)
 {
-	const TraceServices::FTimingProfilerAggregatedStats& Stats = TimerNodePtr->GetAggregatedStats();
+	const FTimersTableColumn& TimersColumn = static_cast<const FTimersTableColumn&>(Column);
+	FText Description = TimersColumn.GetDescription(InAggregationMode);
 
-	const FText InstanceCountText = FText::AsNumber(Stats.InstanceCount);
+	TSharedPtr<SToolTip> ColumnTooltip =
+		SNew(SToolTip)
+		[
+			SNew(SVerticalBox)
 
-	const int32 NumDigits = 5;
+			+ SVerticalBox::Slot()
+			.AutoHeight()
+			.Padding(2.0f)
+			[
+				SNew(STextBlock)
+				.Text(Column.GetTitleName())
+				.TextStyle(FInsightsStyle::Get(), TEXT("TreeTable.TooltipBold"))
+			]
 
-	TCHAR FormatString[32];
-	FCString::Snprintf(FormatString, sizeof(FormatString), TEXT("%%.%dfs (%%s)"), NumDigits);
+			+ SVerticalBox::Slot()
+			.AutoHeight()
+			.Padding(2.0f)
+			[
+				SNew(STextBlock)
+				.Text(Description)
+				.TextStyle(FInsightsStyle::Get(), TEXT("TreeTable.Tooltip"))
+			]
+		];
+
+	return ColumnTooltip;
+}
+
+////////////////////////////////////////////////////////////////////////////////////////////////////
 
-	const FText TotalInclusiveTimeText = FText::FromString(FString::Printf(FormatString, Stats.TotalInclusiveTime, *TimeUtils::FormatTimeAuto(Stats.TotalInclusiveTime, 2)));
-	const FText MinInclusiveTimeText = FText::FromString(TimeUtils::FormatTimeMs(Stats.MinInclusiveTime, NumDigits, true));
-	const FText MaxInclusiveTimeText = FText::FromString(TimeUtils::FormatTimeMs(Stats.MaxInclusiveTime, NumDigits, true));
-	const FText AvgInclusiveTimeText = FText::FromString(TimeUtils::FormatTimeMs(Stats.AverageInclusiveTime, NumDigits, true));
-	const FText MedInclusiveTimeText = FText::FromString(TimeUtils::FormatTimeMs(Stats.MedianInclusiveTime, NumDigits, true));
+TSharedPtr<SToolTip> STimersViewTooltip::GetRowTooltip(const TSharedPtr<FTimerNode> TimerNodePtr)
+{
+	const TraceServices::FTimingProfilerAggregatedStats& Stats = TimerNodePtr->GetAggregatedStats();
 
-	const FText TotalExclusiveTimeText = FText::FromString(FString::Printf(FormatString, Stats.TotalExclusiveTime, *TimeUtils::FormatTimeAuto(Stats.TotalExclusiveTime, 2)));
-	const FText MinExclusiveTimeText = FText::FromString(TimeUtils::FormatTimeMs(Stats.MinExclusiveTime, NumDigits, true));
-	const FText MaxExclusiveTimeText = FText::FromString(TimeUtils::FormatTimeMs(Stats.MaxExclusiveTime, NumDigits, true));
-	const FText AvgExclusiveTimeText = FText::FromString(TimeUtils::FormatTimeMs(Stats.AverageExclusiveTime, NumDigits, true));
-	const FText MedExclusiveTimeText = FText::FromString(TimeUtils::FormatTimeMs(Stats.MedianExclusiveTime, NumDigits, true));
+	const FText InstanceCountText = FText::AsNumber(Stats.InstanceCount);
 
 	TSharedPtr<SGridPanel> GridPanel;
 	TSharedPtr<SHorizontalBox> HBox;
@@ -344,12 +363,41 @@ TSharedPtr<SToolTip> STimersViewTooltip::GetRowTooltip(const TSharedPtr<FTimerNo
 			]
 		];
 
+	const int32 NumDigits = 5;
 	int32 Row = 1;
-	AddStatsRow(GridPanel, Row, LOCTEXT("TT_TotalTime",   "Total Time:"),   TotalInclusiveTimeText, TotalExclusiveTimeText);
-	AddStatsRow(GridPanel, Row, LOCTEXT("TT_MaxTime",     "Max Time:"),     MaxInclusiveTimeText,   MaxExclusiveTimeText);
-	AddStatsRow(GridPanel, Row, LOCTEXT("TT_AverageTime", "Average Time:"), AvgInclusiveTimeText,   AvgExclusiveTimeText);
-	AddStatsRow(GridPanel, Row, LOCTEXT("TT_MedianTime",  "Median Time:"),  MedInclusiveTimeText,   MedExclusiveTimeText);
-	AddStatsRow(GridPanel, Row, LOCTEXT("TT_MinTime",     "Min Time:"),     MinInclusiveTimeText,   MinExclusiveTimeText);
+
+	if (true)
+	{
+		TCHAR FormatString[32];
+		FCString::Snprintf(FormatString, sizeof(FormatString), TEXT("%%.%dfs (%%s)"), NumDigits);
+		const FText TotalInclusiveTimeText = FText::FromString(FString::Printf(FormatString, Stats.TotalInclusiveTime, *TimeUtils::FormatTimeAuto(Stats.TotalInclusiveTime, 2)));
+		const FText TotalExclusiveTimeText = FText::FromString(FString::Printf(FormatString, Stats.TotalExclusiveTime, *TimeUtils::FormatTimeAuto(Stats.TotalExclusiveTime, 2)));
+		AddStatsRow(GridPanel, Row, LOCTEXT("TT_TotalTime", "Total Time:"), TotalInclusiveTimeText, TotalExclusiveTimeText);
+	}
+	if (!std::isnan(Stats.MaxInclusiveTime) || !std::isnan(Stats.MaxExclusiveTime))
+	{
+		const FText MaxInclusiveTimeText = FText::FromString(TimeUtils::FormatTimeMs(Stats.MaxInclusiveTime, NumDigits, true));
+		const FText MaxExclusiveTimeText = FText::FromString(TimeUtils::FormatTimeMs(Stats.MaxExclusiveTime, NumDigits, true));
+		AddStatsRow(GridPanel, Row, LOCTEXT("TT_MaxTime", "Max Time:"), MaxInclusiveTimeText, MaxExclusiveTimeText);
+	}
+	if (true)
+	{
+		const FText AvgInclusiveTimeText = FText::FromString(TimeUtils::FormatTimeMs(Stats.AverageInclusiveTime, NumDigits, true));
+		const FText AvgExclusiveTimeText = FText::FromString(TimeUtils::FormatTimeMs(Stats.AverageExclusiveTime, NumDigits, true));
+		AddStatsRow(GridPanel, Row, LOCTEXT("TT_AverageTime", "Average Time:"), AvgInclusiveTimeText,   AvgExclusiveTimeText);
+	}
+	if (!std::isnan(Stats.MedianInclusiveTime) || !std::isnan(Stats.MedianExclusiveTime))
+	{
+		const FText MedInclusiveTimeText = FText::FromString(TimeUtils::FormatTimeMs(Stats.MedianInclusiveTime, NumDigits, true));
+		const FText MedExclusiveTimeText = FText::FromString(TimeUtils::FormatTimeMs(Stats.MedianExclusiveTime, NumDigits, true));
+		AddStatsRow(GridPanel, Row, LOCTEXT("TT_MedianTime",  "Median Time:"),  MedInclusiveTimeText,   MedExclusiveTimeText);
+	}
+	if (!std::isnan(Stats.MinInclusiveTime) || !std::isnan(Stats.MinExclusiveTime))
+	{
+		const FText MinInclusiveTimeText = FText::FromString(TimeUtils::FormatTimeMs(Stats.MinInclusiveTime, NumDigits, true));
+		const FText MinExclusiveTimeText = FText::FromString(TimeUtils::FormatTimeMs(Stats.MinExclusiveTime, NumDigits, true));
+		AddStatsRow(GridPanel, Row, LOCTEXT("TT_MinTime", "Min Time:"), MinInclusiveTimeText, MinExclusiveTimeText);
+	}
 
 	return TableCellTooltip;
 }
diff --git a/Engine/Source/Developer/TraceInsights/Private/Insights/Widgets/STimersViewTooltip.h b/Engine/Source/Developer/TraceInsights/Private/Insights/Widgets/STimersViewTooltip.h
--- a/Engine/Source/Developer/TraceInsights/Private/Insights/Widgets/STimersViewTooltip.h
+++ b/Engine/Source/Developer/TraceInsights/Private/Insights/Widgets/STimersViewTooltip.h
@@ -26,6 +26,7 @@ public:
 
 	static TSharedPtr<SToolTip> GetTableTooltip(const Insights::FTable& Table);
 	static TSharedPtr<SToolTip> GetColumnTooltip(const Insights::FTableColumn& Column);
+	static TSharedPtr<SToolTip> GetColumnTooltipForMode(const Insights::FTableColumn& Column, ETraceFrameType InAggregationMode);
 	static TSharedPtr<SToolTip> GetRowTooltip(const TSharedPtr<FTimerNode> TreeNodePtr);
 
 	static bool GetSource(const TSharedPtr<FTimerNode> TreeNodePtr, FText& OutSourcePrefix, FText& OutSourceSuffix);
