[Insights] Implemented Toggle for timer series in the Main Graph View when Click + Shift  is used on a timer in the Timing View.
PR #11269

diff --git a/Engine/Source/Developer/TraceInsights/Private/Insights/TimingProfilerManager.cpp b/Engine/Source/Developer/TraceInsights/Private/Insights/TimingProfilerManager.cpp
--- a/Engine/Source/Developer/TraceInsights/Private/Insights/TimingProfilerManager.cpp
+++ b/Engine/Source/Developer/TraceInsights/Private/Insights/TimingProfilerManager.cpp
@@ -516,6 +516,22 @@ void FTimingProfilerManager::SetSelectedTimer(uint32 InTimerId)
 
 ////////////////////////////////////////////////////////////////////////////////////////////////////
 
+void FTimingProfilerManager::ToggleTimingViewMainGraphEventSeries(uint32 InTimerId)
+{
+	FTimerNodePtr NodePtr = GetTimerNode(InTimerId);
+	TSharedPtr<STimingProfilerWindow> Wnd = GetProfilerWindow();
+	if (Wnd && NodePtr)
+	{
+		TSharedPtr<STimersView> TimersView = Wnd->GetTimersView();
+		if (TimersView)
+		{
+			TimersView->ToggleTimingViewMainGraphEventSeries(NodePtr);
+		}
+	}
+}
+
+////////////////////////////////////////////////////////////////////////////////////////////////////
+
 void FTimingProfilerManager::OnThreadFilterChanged()
 {
 	UpdateCallersAndCallees();
diff --git a/Engine/Source/Developer/TraceInsights/Private/Insights/TimingProfilerManager.h b/Engine/Source/Developer/TraceInsights/Private/Insights/TimingProfilerManager.h
--- a/Engine/Source/Developer/TraceInsights/Private/Insights/TimingProfilerManager.h
+++ b/Engine/Source/Developer/TraceInsights/Private/Insights/TimingProfilerManager.h
@@ -138,6 +138,7 @@ public:
 	FTimerNodePtr GetTimerNode(uint32 TimerId) const;
 	uint32 GetSelectedTimer() const { return SelectedTimerId; }
 	void SetSelectedTimer(uint32 TimerId);
+	void ToggleTimingViewMainGraphEventSeries(uint32 InTimerId);
 
 	void OnThreadFilterChanged();
 
diff --git a/Engine/Source/Developer/TraceInsights/Private/Insights/Widgets/STimersView.cpp b/Engine/Source/Developer/TraceInsights/Private/Insights/Widgets/STimersView.cpp
--- a/Engine/Source/Developer/TraceInsights/Private/Insights/Widgets/STimersView.cpp
+++ b/Engine/Source/Developer/TraceInsights/Private/Insights/Widgets/STimersView.cpp
@@ -845,7 +845,7 @@ void STimersView::TreeView_BuildPlotTimerMenu(FMenuBuilder& MenuBuilder)
 
 		if (SelectedNode.IsValid() &&
 			SelectedNode->GetType() != ETimerNodeType::Group &&
-			IsSeriesInTimingViewMainGraph(SelectedNode))
+			IsInstanceSeriesInTimingViewMainGraph(SelectedNode))
 		{
 			MenuBuilder.AddMenuEntry
 			(
@@ -1521,25 +1521,7 @@ void STimersView::TreeView_OnMouseButtonDoubleClick(FTimerNodePtr NodePtr)
 		}
 		else
 		{
-			switch (ModeFrameType)
-			{
-				case ETraceFrameType::TraceFrameType_Count:
-				{
-					// Instance graph
-					ToggleTimingViewMainGraphEventSeries(NodePtr);
-					break;
-				}
-				case ETraceFrameType::TraceFrameType_Game:
-				case ETraceFrameType::TraceFrameType_Rendering:
-				{
-					ToggleTimingViewMainGraphEventFrameStatsSeries(NodePtr, ModeFrameType);
-					break;
-				}
-				default:
-				{
-					ensure(0);
-				}
-			}
+			ToggleTimingViewMainGraphEventSeries(NodePtr);
 		}
 	}
 }
@@ -2639,7 +2621,32 @@ TSharedPtr<SFrameTrack> STimersView::GetFrameTrack() const
 
 ////////////////////////////////////////////////////////////////////////////////////////////////////
 
-void STimersView::ToggleGraphSeries(TSharedRef<FTimingGraphTrack> GraphTrack, FTimerNodeRef NodePtr) const
+void STimersView::ToggleTimingViewMainGraphEventSeries(FTimerNodePtr TimerNode) const
+{
+	switch (ModeFrameType)
+	{
+	case ETraceFrameType::TraceFrameType_Count:
+	{
+		// Instance graph
+		ToggleTimingViewMainGraphEventInstanceSeries(TimerNode);
+		break;
+	}
+	case ETraceFrameType::TraceFrameType_Game:
+	case ETraceFrameType::TraceFrameType_Rendering:
+	{
+		ToggleTimingViewMainGraphEventFrameStatsSeries(TimerNode, ModeFrameType);
+		break;
+	}
+	default:
+	{
+		ensure(0);
+	}
+	}
+}
+
+////////////////////////////////////////////////////////////////////////////////////////////////////
+
+void STimersView::ToggleGraphInstanceSeries(TSharedRef<FTimingGraphTrack> GraphTrack, FTimerNodeRef NodePtr) const
 {
 	const uint32 TimerId = NodePtr->GetTimerId();
 	TSharedPtr<FTimingGraphSeries> Series = GraphTrack->GetTimerSeries(TimerId);
@@ -2661,7 +2668,7 @@ void STimersView::ToggleGraphSeries(TSharedRef<FTimingGraphTrack> GraphTrack, FT
 
 ////////////////////////////////////////////////////////////////////////////////////////////////////
 
-bool STimersView::IsSeriesInTimingViewMainGraph(FTimerNodePtr TimerNode) const
+bool STimersView::IsInstanceSeriesInTimingViewMainGraph(FTimerNodePtr TimerNode) const
 {
 	TSharedPtr<FTimingGraphTrack> GraphTrack = GetTimingViewMainGraphTrack();
 
@@ -2678,12 +2685,12 @@ bool STimersView::IsSeriesInTimingViewMainGraph(FTimerNodePtr TimerNode) const
 
 ////////////////////////////////////////////////////////////////////////////////////////////////////
 
-void STimersView::ToggleTimingViewMainGraphEventSeries(FTimerNodePtr TimerNode) const
+void STimersView::ToggleTimingViewMainGraphEventInstanceSeries(FTimerNodePtr TimerNode) const
 {
 	TSharedPtr<FTimingGraphTrack> GraphTrack = GetTimingViewMainGraphTrack();
 	if (GraphTrack.IsValid())
 	{
-		ToggleGraphSeries(GraphTrack.ToSharedRef(), TimerNode.ToSharedRef());
+		ToggleGraphInstanceSeries(GraphTrack.ToSharedRef(), TimerNode.ToSharedRef());
 	}
 }
 
diff --git a/Engine/Source/Developer/TraceInsights/Private/Insights/Widgets/STimersView.h b/Engine/Source/Developer/TraceInsights/Private/Insights/Widgets/STimersView.h
--- a/Engine/Source/Developer/TraceInsights/Private/Insights/Widgets/STimersView.h
+++ b/Engine/Source/Developer/TraceInsights/Private/Insights/Widgets/STimersView.h
@@ -93,6 +93,8 @@ public:
 
 	void OnTimingViewTrackListChanged();
 
+	void ToggleTimingViewMainGraphEventSeries(FTimerNodePtr TimerNode) const;
+
 private:
 	void InitCommandList();
 
@@ -317,9 +319,9 @@ private:
 	TSharedPtr<FTimingGraphTrack> GetTimingViewMainGraphTrack() const;
 	TSharedPtr<SFrameTrack> GetFrameTrack() const;
 
-	void ToggleGraphSeries(TSharedRef<FTimingGraphTrack> GraphTrack, FTimerNodeRef NodePtr) const;
-	bool IsSeriesInTimingViewMainGraph(FTimerNodePtr TimerNode) const;
-	void ToggleTimingViewMainGraphEventSeries(FTimerNodePtr TimerNode) const;
+	void ToggleGraphInstanceSeries(TSharedRef<FTimingGraphTrack> GraphTrack, FTimerNodeRef NodePtr) const;
+	bool IsInstanceSeriesInTimingViewMainGraph(FTimerNodePtr TimerNode) const;
+	void ToggleTimingViewMainGraphEventInstanceSeries(FTimerNodePtr TimerNode) const;
 
 	void ToggleGraphFrameStatsSeries(TSharedRef<FTimingGraphTrack> GraphTrack, FTimerNodeRef NodePtr, ETraceFrameType FrameType) const;
 	bool IsFrameStatsSeriesInTimingViewMainGraph(FTimerNodePtr TimerNode, ETraceFrameType FrameType) const;
diff --git a/Engine/Source/Developer/TraceInsights/Private/Insights/Widgets/STimingView.cpp b/Engine/Source/Developer/TraceInsights/Private/Insights/Widgets/STimingView.cpp
--- a/Engine/Source/Developer/TraceInsights/Private/Insights/Widgets/STimingView.cpp
+++ b/Engine/Source/Developer/TraceInsights/Private/Insights/Widgets/STimingView.cpp
@@ -74,10 +74,10 @@
 #include "Insights/Widgets/STimingProfilerWindow.h"
 #include "Insights/Widgets/STimingViewTrackList.h"
 #include "Insights/Widgets/SQuickFind.h"
+#include "Insights/ViewModels/ThreadTrackEvent.h"
 
 #include <limits>
 
-
 ////////////////////////////////////////////////////////////////////////////////////////////////////
 
 #define LOCTEXT_NAMESPACE "STimingView"
@@ -2190,6 +2190,11 @@ FReply STimingView::OnMouseButtonUp(const FGeometry& MyGeometry, const FPointerE
 				SelectHoveredTimingTrack();
 				SelectHoveredTimingEvent();
 
+				if (MouseEvent.GetModifierKeys().IsShiftDown())
+				{
+					ToggleGraphSeries(HoveredEvent);
+				}
+
 				// When clicking on an empty space...
 				if (!SelectedEvent.IsValid())
 				{
@@ -3459,7 +3464,7 @@ void STimingView::SelectTimeInterval(double IntervalStartTime, double IntervalDu
 	SelectionStartTime = IntervalStartTime;
 	SelectionEndTime = IntervalStartTime + IntervalDuration;
 
-	if (GetFrameTypeToSnapTo() != ETraceFrameType::TraceFrameType_Count)
+	if (GetFrameTypeToSnapTo() != ETraceFrameType::TraceFrameType_Count && IsInTimingProfiler())
 	{
 		SnapToFrameBound(SelectionStartTime, SelectionEndTime);
 	}
@@ -3895,6 +3900,18 @@ void STimingView::SelectTimingEvent(const TSharedPtr<const ITimingEvent> InEvent
 
 ////////////////////////////////////////////////////////////////////////////////////////////////////
 
+void STimingView::ToggleGraphSeries(const TSharedPtr<const ITimingEvent> InEvent)
+{
+	if(InEvent.Get() && InEvent.Get()->Is<FThreadTrackEvent>() && IsInTimingProfiler())
+	{
+		const FThreadTrackEvent& TrackEvent = InEvent.Get()->As<FThreadTrackEvent>();
+
+		FTimingProfilerManager::Get()->ToggleTimingViewMainGraphEventSeries(TrackEvent.GetTimerId());
+	}
+}
+
+////////////////////////////////////////////////////////////////////////////////////////////////////
+
 void STimingView::SelectLeftTimingEvent()
 {
 	if (SelectedEvent.IsValid())
@@ -5423,5 +5440,12 @@ void STimingView::UpdateFilters()
 
 ////////////////////////////////////////////////////////////////////////////////////////////////////
 
+bool STimingView::IsInTimingProfiler()
+{
+	return GetName() == FInsightsManagerTabs::TimingProfilerTabId;
+}
+
+////////////////////////////////////////////////////////////////////////////////////////////////////
+
 #undef INSIGHTS_ACTIVATE_BENCHMARK
 #undef LOCTEXT_NAMESPACE
diff --git a/Engine/Source/Developer/TraceInsights/Private/Insights/Widgets/STimingView.h b/Engine/Source/Developer/TraceInsights/Private/Insights/Widgets/STimingView.h
--- a/Engine/Source/Developer/TraceInsights/Private/Insights/Widgets/STimingView.h
+++ b/Engine/Source/Developer/TraceInsights/Private/Insights/Widgets/STimingView.h
@@ -372,6 +372,7 @@ public:
 
 	void SelectTimingTrack(const TSharedPtr<FBaseTimingTrack> InTrack, bool bBringTrackIntoView);
 	void SelectTimingEvent(const TSharedPtr<const ITimingEvent> InEvent, bool bBringEventIntoViewHorizontally, bool bBringEventIntoViewVertically = false);
+	void ToggleGraphSeries(const TSharedPtr<const ITimingEvent> InEvent);
 
 	const TSharedPtr<ITimingEventFilter> GetEventFilter() const { return TimingEventFilter; }
 	void SetEventFilter(const TSharedPtr<ITimingEventFilter> InEventFilter);
@@ -532,6 +533,8 @@ protected:
 	 */ 
 	void UpdateFilters();
 
+	bool IsInTimingProfiler();
+
 protected:
 	/** The name of the view. */
 	FName ViewName;
