[Insights] Minor changes to the "Unreal Trace Server" control menu (section name, tooltip).

diff --git a/Engine/Plugins/TraceUtilities/Source/EditorTraceUtilities/Private/InsightsStatusBar/SInsightsStatusBar.cpp b/Engine/Plugins/TraceUtilities/Source/EditorTraceUtilities/Private/InsightsStatusBar/SInsightsStatusBar.cpp
--- a/Engine/Plugins/TraceUtilities/Source/EditorTraceUtilities/Private/InsightsStatusBar/SInsightsStatusBar.cpp
+++ b/Engine/Plugins/TraceUtilities/Source/EditorTraceUtilities/Private/InsightsStatusBar/SInsightsStatusBar.cpp
@@ -410,19 +410,19 @@ TSharedRef<SWidget> SInsightsStatusBarWidget::MakeTraceMenu()
 	MenuBuilder.BeginSection("Insights", LOCTEXT("TraceMenu_Section_Insights", "Insights"));
 	{
 		MenuBuilder.AddSubMenu(
-			LOCTEXT("ServerControlLabel", "UnrealTraceServer"),
-			LOCTEXT("ServerControlTooltip", "Control UnrealTraceServer instances"),
+			LOCTEXT("ServerControlLabel", "Unreal Trace Server"),
+			LOCTEXT("ServerControlTooltip", "Info and controls for the Unreal Trace Server instances"),
 			FNewMenuDelegate::CreateLambda([this](FMenuBuilder& MenuBuilder)
-			{
-				for (auto& ServerControl : ServerControls)
 				{
-					 ServerControl.MakeMenu(MenuBuilder);
-				}
-			}),
+					for (auto& ServerControl : ServerControls)
+					{
+						ServerControl.MakeMenu(MenuBuilder);
+					}
+				}),
 			false,
-			FSlateIcon(FAppStyle::Get().GetStyleSetName(), ("Icons.Server"))
+			FSlateIcon(FAppStyle::Get().GetStyleSetName(), "Icons.Server")
 		);
-		
+
 		MenuBuilder.AddMenuEntry(
 			LOCTEXT("OpenInsightsLabel", "Unreal Insights (Session Browser)"),
 			LOCTEXT("OpenInsightsTooltip", "Launch the Unreal Insights Session Browser."),
@@ -512,7 +512,7 @@ FText SInsightsStatusBarWidget::GetTitleToolTipText() const
 	FTextBuilder DescBuilder;
 
 	const FString Dest = FTraceAuxiliary::GetTraceDestinationString();
-	
+
 	if (*Dest != 0)
 	{
 		DescBuilder.AppendLineFormat(LOCTEXT("TracingToText", "Tracing to: {0}"), FText::FromString(Dest));
@@ -525,7 +525,7 @@ FText SInsightsStatusBarWidget::GetTitleToolTipText() const
 	{
 		DescBuilder.AppendLine(LOCTEXT("NotTracingText","Not currently tracing."));
 	}
-		
+
 	return DescBuilder.ToText();
 }
 
@@ -1046,7 +1046,7 @@ void SInsightsStatusBarWidget::OpenTrace(int32 Index)
 				const FLiveSessionsMap& Sessions = LiveSessionTracker->GetLiveSessions();
 				FString FileName = FPaths::GetBaseFilename(Traces[Index]->FilePath);
 				const uint32* TraceId = Sessions.Find(FileName);
-				
+
 				if (TraceId)
 				{
 					FUnrealInsightsLauncher::Get()->OpenRemoteTrace(TEXT("localhost"), uint16(LiveSessionTracker->GetStorePort()), *TraceId);
diff --git a/Engine/Source/Developer/TraceInsights/Private/Insights/Widgets/SStartPageWindow.cpp b/Engine/Source/Developer/TraceInsights/Private/Insights/Widgets/SStartPageWindow.cpp
--- a/Engine/Source/Developer/TraceInsights/Private/Insights/Widgets/SStartPageWindow.cpp
+++ b/Engine/Source/Developer/TraceInsights/Private/Insights/Widgets/SStartPageWindow.cpp
@@ -1977,7 +1977,7 @@ FText STraceStoreWindow::GetConnectionStatusTooltip() const
 			StoreBrowser->LockSettings();
 			FText Version = FText::FromString(StoreBrowser->GetVersion());
 			const uint32 RecorderPort = StoreBrowser->GetRecorderPort();
-			const uint32 StorePort = StoreBrowser->GetStorePort();	
+			const uint32 StorePort = StoreBrowser->GetStorePort();
 			StoreBrowser->UnlockSettings();
 			return FText::Format(Connected,
 				Version,
@@ -2127,7 +2127,7 @@ bool STraceStoreWindow::CanDeleteSelectedTraces() const
 	TArray<TSharedPtr<FTraceViewModel>> SelectedTraces = TraceListView->GetSelectedItems();
 	for (const TSharedPtr<FTraceViewModel>& SelectedTrace : SelectedTraces)
 	{
-		if ((SelectedTrace->TraceId != FTraceViewModel::InvalidTraceId) && 
+		if ((SelectedTrace->TraceId != FTraceViewModel::InvalidTraceId) &&
 			!SelectedTrace->bIsLive)
 		{
 			return true;
@@ -2148,7 +2148,7 @@ void STraceStoreWindow::DeleteSelectedTraces()
 	}
 
 	TArray<TSharedPtr<FTraceViewModel>> TracesToDelete = TraceListView->GetSelectedItems();
-	// Filter the traces that can actually be deleted : 
+	// Filter the traces that can actually be deleted :
 	TracesToDelete.RemoveAll([](const TSharedPtr<FTraceViewModel>& InTrace) { return (InTrace->TraceId == FTraceViewModel::InvalidTraceId) || InTrace->bIsLive; });
 	if (TracesToDelete.Num() == 0)
 	{
@@ -2473,7 +2473,7 @@ void STraceStoreWindow::RefreshTraceList()
 			{
 				 ServerControls.Emplace(*StoreBrowser->GetHost(), StoreBrowser->GetStorePort(), FAppStyle::Get().GetStyleSetName());
 			}
-			
+
 			// Update the host text
 			if (StoreHostTextBox)
 			{
@@ -3191,24 +3191,23 @@ TSharedRef<SWidget> STraceStoreWindow::MakeTraceListMenu()
 	}
 	MenuBuilder.EndSection();
 
-	MenuBuilder.BeginSection("Control", LOCTEXT("TraceListMenu_Section_Control", "Control"));
+	MenuBuilder.BeginSection("UnrealTraceServer", LOCTEXT("TraceListMenu_Section_Server", "Server"));
 	{
 		MenuBuilder.AddSubMenu(
-			LOCTEXT("ServerControlLabel", "UnrealTraceServer"),
-			LOCTEXT("ServerControlTooltip", "Control UnrealTraceServer instances"),
+			LOCTEXT("ServerControlLabel", "Unreal Trace Server"),
+			LOCTEXT("ServerControlTooltip", "Info and controls for the Unreal Trace Server instances"),
 			FNewMenuDelegate::CreateLambda([this](FMenuBuilder& MenuBuilder)
-			{
-				 for (auto& ServerControl : ServerControls)
-				 {
-					   ServerControl.MakeMenu(MenuBuilder);
-				 }
-			}),
+				{
+					for (auto& ServerControl : ServerControls)
+					{
+						ServerControl.MakeMenu(MenuBuilder);
+					}
+				}),
 			false,
-			FSlateIcon(FAppStyle::Get().GetStyleSetName(), ("Icons.Server"))
-		);
+			FSlateIcon(FAppStyle::Get().GetStyleSetName(), "Icons.Server"));
 	}
 	MenuBuilder.EndSection();
-	
+
 	MenuBuilder.BeginSection("DebugOptions", LOCTEXT("TraceListMenu_Section_DebugOptions", "Debug Options"));
 
 	if (FGlobalTabmanager::Get()->HasTabSpawner(FInsightsManagerTabs::AutomationWindowTabId))
