[Insights]
  - Callers & Callees: Added "Child Instance Count" column. It shows the total number of timing event instances of the child timers (callers or callees).
  - Callers & Callees: Added the number of children nodes to each timer node in the tree.

diff --git a/Engine/Source/Developer/TraceInsights/Private/Insights/ViewModels/TimerNode.cpp b/Engine/Source/Developer/TraceInsights/Private/Insights/ViewModels/TimerNode.cpp
--- a/Engine/Source/Developer/TraceInsights/Private/Insights/ViewModels/TimerNode.cpp
+++ b/Engine/Source/Developer/TraceInsights/Private/Insights/ViewModels/TimerNode.cpp
@@ -17,8 +17,8 @@ const FName FTimerNode::CpuGroup(TEXT("CPU"));
 
 ////////////////////////////////////////////////////////////////////////////////////////////////////
 
-FTimerNode::FTimerNode(uint32 InTimerId, const TCHAR* InName, ETimerNodeType InType)
-	: FBaseTreeNode(FName(InName), true) // Timer nodes are groups (for callers/callees use case)
+FTimerNode::FTimerNode(uint32 InTimerId, const TCHAR* InName, ETimerNodeType InType, bool bInIsGroup)
+	: FBaseTreeNode(FName(InName), bInIsGroup)
 	, TimerId(InTimerId)
 	, MetaGroupName(InType == ETimerNodeType::CpuScope ? CpuGroup : InType == ETimerNodeType::GpuScope ? GpuGroup : NAME_None)
 	, Type(InType)
diff --git a/Engine/Source/Developer/TraceInsights/Private/Insights/ViewModels/TimerNode.h b/Engine/Source/Developer/TraceInsights/Private/Insights/ViewModels/TimerNode.h
--- a/Engine/Source/Developer/TraceInsights/Private/Insights/ViewModels/TimerNode.h
+++ b/Engine/Source/Developer/TraceInsights/Private/Insights/ViewModels/TimerNode.h
@@ -61,7 +61,7 @@ public:
 
 public:
 	/** Initialization constructor for the timer node. */
-	FTimerNode(uint32 InTimerId, const TCHAR* InName, ETimerNodeType InType);
+	FTimerNode(uint32 InTimerId, const TCHAR* InName, ETimerNodeType InType, bool bInIsGroup);
 
 	/** Initialization constructor for the group node. */
 	explicit FTimerNode(const FName InGroupName);
diff --git a/Engine/Source/Developer/TraceInsights/Private/Insights/ViewModels/TimersViewColumnFactory.cpp b/Engine/Source/Developer/TraceInsights/Private/Insights/ViewModels/TimersViewColumnFactory.cpp
--- a/Engine/Source/Developer/TraceInsights/Private/Insights/ViewModels/TimersViewColumnFactory.cpp
+++ b/Engine/Source/Developer/TraceInsights/Private/Insights/ViewModels/TimersViewColumnFactory.cpp
@@ -19,6 +19,7 @@ const FName FTimersViewColumns::NameColumnID(TEXT("Name")); // TEXT("_Hierarchy"
 const FName FTimersViewColumns::MetaGroupNameColumnID(TEXT("MetaGroupName"));
 const FName FTimersViewColumns::TypeColumnID(TEXT("Type"));
 const FName FTimersViewColumns::InstanceCountColumnID(TEXT("Count"));
+const FName FTimersViewColumns::ChildInstanceCountColumnID(TEXT("ChildCount"));
 
 // Inclusive Time columns
 const FName FTimersViewColumns::TotalInclusiveTimeColumnID(TEXT("TotalInclTime"));
@@ -71,6 +72,7 @@ void FTimersViewColumnFactory::CreateTimerTreeViewColumns(TArray<TSharedRef<Insi
 	Columns.Add(CreateNameColumn());
 	Columns.Add(CreateTypeColumn());
 	Columns.Add(CreateInstanceCountColumn());
+	Columns.Add(CreateChildInstanceCountColumn());
 	Columns.Add(CreateTotalInclusiveTimeColumn());
 	Columns.Add(CreateAverageInclusiveTimeColumn());
 	Columns.Add(CreateTotalExclusiveTimeColumn());
@@ -227,7 +229,7 @@ TSharedRef<Insights::FTableColumn> FTimersViewColumnFactory::CreateInstanceCount
 	Column.SetShortName(LOCTEXT("InstanceCount_ColumnName", "Count"));
 	Column.SetTitleName(LOCTEXT("InstanceCount_ColumnTitle", "Instance Count"));
 
-	FText Description = LOCTEXT("InstanceCount_ColumnDesc", "Number of selected timer's instances");
+	FText Description = LOCTEXT("InstanceCount_ColumnDesc", "Number of timing event instances of the selected timer");
 	Column.SetDescription(ETraceFrameType::TraceFrameType_Count, Description);
 	Column.SetDescription(ETraceFrameType::TraceFrameType_Game, Description);
 	Column.SetDescription(ETraceFrameType::TraceFrameType_Rendering, Description);
@@ -266,6 +268,60 @@ TSharedRef<Insights::FTableColumn> FTimersViewColumnFactory::CreateInstanceCount
 	return ColumnRef;
 }
 
+////////////////////////////////////////////////////////////////////////////////////////////////////
+
+TSharedRef<Insights::FTableColumn> FTimersViewColumnFactory::CreateChildInstanceCountColumn()
+{
+	using namespace Insights;
+
+	TSharedRef<FTimersTableColumn> ColumnRef = MakeShared<FTimersTableColumn>(FTimersViewColumns::ChildInstanceCountColumnID);
+	FTimersTableColumn& Column = *ColumnRef;
+
+	Column.SetShortName(LOCTEXT("ChildInstanceCount_ColumnName", "Child Count"));
+	Column.SetTitleName(LOCTEXT("ChildInstanceCount_ColumnTitle", "Child Instance Count"));
+
+	FText Description = LOCTEXT("ChildInstanceCount_ColumnDesc", "Total number of timing event instances of the child timers (callers or callees)");
+	Column.SetDescription(ETraceFrameType::TraceFrameType_Count, Description);
+	Column.SetDescription(ETraceFrameType::TraceFrameType_Game, Description);
+	Column.SetDescription(ETraceFrameType::TraceFrameType_Rendering, Description);
+
+	Column.SetFlags(ETableColumnFlags::CanBeHidden | ETableColumnFlags::CanBeFiltered);
+
+	Column.SetHorizontalAlignment(HAlign_Right);
+	Column.SetInitialWidth(60.0f);
+
+	Column.SetDataType(ETableCellDataType::Int64);
+
+	class FInstanceCountValueGetter : public FTableCellValueGetter
+	{
+	public:
+		virtual const TOptional<FTableCellValue> GetValue(const FTableColumn& Column, const FBaseTreeNode& Node) const
+		{
+			ensure(Column.GetId() == FTimersViewColumns::ChildInstanceCountColumnID);
+			const FTimerNode& TimerNode = static_cast<const FTimerNode&>(Node);
+			int64 TotalChildInstanceCount = 0;
+			for (const auto& ChildNode : TimerNode.GetChildren())
+			{
+				TotalChildInstanceCount += ChildNode->As<FTimerNode>().GetAggregatedStats().InstanceCount;
+			}
+			return TOptional<FTableCellValue>(FTableCellValue(TotalChildInstanceCount));
+		}
+	};
+
+	TSharedRef<ITableCellValueGetter> Getter = MakeShared<FInstanceCountValueGetter>();
+	Column.SetValueGetter(Getter);
+
+	TSharedRef<ITableCellValueFormatter> Formatter = MakeShared<FInt64ValueFormatterAsNumber>();
+	Column.SetValueFormatter(Formatter);
+
+	//TSharedRef<ITableCellValueSorter> Sorter = MakeShared<FSorterByInt64Value>(ColumnRef);
+	TSharedRef<ITableCellValueSorter> Sorter = MakeShared<FTimerNodeSortingByInstanceCount>(ColumnRef);
+	Column.SetValueSorter(Sorter);
+	Column.SetInitialSortMode(EColumnSortMode::Descending);
+
+	return ColumnRef;
+}
+
 ////////////////////////////////////////////////////////////////////////////////////////////////////
 // Inclusive Time Columns
 ////////////////////////////////////////////////////////////////////////////////////////////////////
diff --git a/Engine/Source/Developer/TraceInsights/Private/Insights/ViewModels/TimersViewColumnFactory.h b/Engine/Source/Developer/TraceInsights/Private/Insights/ViewModels/TimersViewColumnFactory.h
--- a/Engine/Source/Developer/TraceInsights/Private/Insights/ViewModels/TimersViewColumnFactory.h
+++ b/Engine/Source/Developer/TraceInsights/Private/Insights/ViewModels/TimersViewColumnFactory.h
@@ -17,6 +17,7 @@ struct FTimersViewColumns
 	static const FName MetaGroupNameColumnID;
 	static const FName TypeColumnID;
 	static const FName InstanceCountColumnID;
+	static const FName ChildInstanceCountColumnID;
 
 	// Inclusive Time columns
 	static const FName TotalInclusiveTimeColumnID;
@@ -97,6 +98,7 @@ public:
 	static TSharedRef<Insights::FTableColumn> CreateMetaGroupNameColumn();
 	static TSharedRef<Insights::FTableColumn> CreateTypeColumn();
 	static TSharedRef<Insights::FTableColumn> CreateInstanceCountColumn();
+	static TSharedRef<Insights::FTableColumn> CreateChildInstanceCountColumn();
 
 	static TSharedRef<Insights::FTableColumn> CreateTotalInclusiveTimeColumn();
 	static TSharedRef<Insights::FTableColumn> CreateMaxInclusiveTimeColumn();
diff --git a/Engine/Source/Developer/TraceInsights/Private/Insights/Widgets/STimerTableCell.h b/Engine/Source/Developer/TraceInsights/Private/Insights/Widgets/STimerTableCell.h
--- a/Engine/Source/Developer/TraceInsights/Private/Insights/Widgets/STimerTableCell.h
+++ b/Engine/Source/Developer/TraceInsights/Private/Insights/Widgets/STimerTableCell.h
@@ -124,7 +124,7 @@ protected:
 
 	EVisibility HasExtraDisplayName() const
 	{
-		return TimerNodePtr->GetType() == ETimerNodeType::Group && TimerNodePtr->HasExtraDisplayName() ? EVisibility::Visible : EVisibility::Collapsed;
+		return TimerNodePtr->HasExtraDisplayName() ? EVisibility::Visible : EVisibility::Collapsed;
 	}
 
 	FText GetValueAsText() const;
diff --git a/Engine/Source/Developer/TraceInsights/Private/Insights/Widgets/STimerTreeView.cpp b/Engine/Source/Developer/TraceInsights/Private/Insights/Widgets/STimerTreeView.cpp
--- a/Engine/Source/Developer/TraceInsights/Private/Insights/Widgets/STimerTreeView.cpp
+++ b/Engine/Source/Developer/TraceInsights/Private/Insights/Widgets/STimerTreeView.cpp
@@ -1025,11 +1025,11 @@ FTimerNodePtr STimerTreeView::CreateTimerNodeRec(const TraceServices::FTimingPro
 	if (Node.Timer == nullptr)
 	{
 		return nullptr;
-		//return MakeShared<FTimerNode>(0, TEXT("!!!!!"), ETimerNodeType::InvalidOrMax);
+		//return MakeShared<FTimerNode>(0, TEXT("!!!!!"), ETimerNodeType::InvalidOrMax, true);
 	}
 
 	const ETimerNodeType Type = Node.Timer->IsGpuTimer ? ETimerNodeType::GpuScope : ETimerNodeType::CpuScope;
-	FTimerNodePtr TimerNodePtr = MakeShared<FTimerNode>(Node.Timer->Id, Node.Timer->Name, Type);
+	FTimerNodePtr TimerNodePtr = MakeShared<FTimerNode>(Node.Timer->Id, Node.Timer->Name, Type, true);
 
 	TraceServices::FTimingProfilerAggregatedStats AggregatedStats;
 	AggregatedStats.InstanceCount = Node.Count;
diff --git a/Engine/Source/Developer/TraceInsights/Private/Insights/Widgets/STimersView.cpp b/Engine/Source/Developer/TraceInsights/Private/Insights/Widgets/STimersView.cpp
--- a/Engine/Source/Developer/TraceInsights/Private/Insights/Widgets/STimersView.cpp
+++ b/Engine/Source/Developer/TraceInsights/Private/Insights/Widgets/STimersView.cpp
@@ -872,7 +872,7 @@ void STimersView::TreeView_BuildPlotTimerMenu(FMenuBuilder& MenuBuilder)
 	}
 
 	MenuBuilder.EndSection();
-	
+
 	MenuBuilder.BeginSection("Game Frame", LOCTEXT("Plot_Series_GameFrame_Section", "Game Frame"));
 
 	// Add/remove game frame stats series to/from graph track
@@ -1921,7 +1921,7 @@ FText STimersView::Mode_GetTooltipText(ETraceFrameType InFrameType) const
 
 	return FText();
 }
- 
+
 ////////////////////////////////////////////////////////////////////////////////////////////////////
 // Sorting
 ////////////////////////////////////////////////////////////////////////////////////////////////////
@@ -2417,7 +2417,7 @@ void STimersView::RebuildTree(bool bResync)
 				const TraceServices::FTimingProfilerTimer& Timer = *(TimerReader->GetTimer(TimerIndex));
 				ensure(Timer.Id == TimerIndex);
 				const ETimerNodeType Type = Timer.IsGpuTimer ? ETimerNodeType::GpuScope : ETimerNodeType::CpuScope;
-				FTimerNodePtr TimerNodePtr = MakeShared<FTimerNode>(Timer.Id, Timer.Name, Type);
+				FTimerNodePtr TimerNodePtr = MakeShared<FTimerNode>(Timer.Id, Timer.Name, Type, false);
 				TimerNodePtr->SetDefaultSortOrder(TimerIndex + 1);
 
 				if (GraphTrack.IsValid())
