Move Insights tests into System.Insights.

diff --git a/Engine/Source/Developer/TraceInsights/Private/Insights/Tests/FunctionalTests/ExportCommandsTests.cpp b/Engine/Source/Developer/TraceInsights/Private/Insights/Tests/FunctionalTests/ExportCommandsTests.cpp
--- a/Engine/Source/Developer/TraceInsights/Private/Insights/Tests/FunctionalTests/ExportCommandsTests.cpp
+++ b/Engine/Source/Developer/TraceInsights/Private/Insights/Tests/FunctionalTests/ExportCommandsTests.cpp
@@ -45,7 +45,7 @@ void VerifyExportedLines(const FString& ExportReportPath, const FString& CmdLogP
 
 #if !WITH_EDITOR
 
-IMPLEMENT_SIMPLE_AUTOMATION_TEST(FСommandsExportWindowsTest, "Insights.Analysis.ExecCmd.CommandsExport(Windows)", EAutomationTestFlags::ApplicationContextMask | EAutomationTestFlags::ProductFilter)
+IMPLEMENT_SIMPLE_AUTOMATION_TEST(FСommandsExportWindowsTest, "System.Insights.Trace.Analysis.ExecCmd.CommandsExport(Windows)", EAutomationTestFlags::ApplicationContextMask | EAutomationTestFlags::ProductFilter)
 bool FСommandsExportWindowsTest::RunTest(const FString& Parameters)
 {
 	TSharedPtr<FInsightsManager> InsightsManager = FInsightsManager::Get();
diff --git a/Engine/Source/Developer/TraceInsights/Private/Insights/Tests/FunctionalTests/MemoryInsightsTests.cpp b/Engine/Source/Developer/TraceInsights/Private/Insights/Tests/FunctionalTests/MemoryInsightsTests.cpp
--- a/Engine/Source/Developer/TraceInsights/Private/Insights/Tests/FunctionalTests/MemoryInsightsTests.cpp
+++ b/Engine/Source/Developer/TraceInsights/Private/Insights/Tests/FunctionalTests/MemoryInsightsTests.cpp
@@ -27,7 +27,7 @@ DECLARE_LOG_CATEGORY_EXTERN(MemoryInsightsTests, Log, All);
 
 #if !WITH_EDITOR
 
-BEGIN_DEFINE_SPEC(FAutomationDriverUnrealInsightsHubMemoryInsightsTest, "Insights.Hub.MemoryInsights", EAutomationTestFlags::ApplicationContextMask | EAutomationTestFlags::ProductFilter)
+BEGIN_DEFINE_SPEC(FAutomationDriverUnrealInsightsHubMemoryInsightsTest, "System.Insights.Hub.MemoryInsights", EAutomationTestFlags::ApplicationContextMask | EAutomationTestFlags::ProductFilter)
 FAutomationDriverPtr Driver;
 TSharedPtr<SWindow> AutomationWindow;
 END_DEFINE_SPEC(FAutomationDriverUnrealInsightsHubMemoryInsightsTest)
@@ -60,7 +60,7 @@ void FAutomationDriverUnrealInsightsHubMemoryInsightsTest::Define()
 					TSharedPtr<FInsightsManager> InsightsManager = FInsightsManager::Get();
 					if (!InsightsManager.IsValid())
 					{
-						AddError("Insigts manager should not be null");
+						AddError("Insights manager should not be null");
 						return;
 					}
 
@@ -140,7 +140,7 @@ void FAutomationDriverUnrealInsightsHubMemoryInsightsTest::Define()
 		});
 }
 
-IMPLEMENT_SIMPLE_AUTOMATION_TEST(FMemoryInsightsUploadLLMXMLReportsTraceTest, "Insights.Analysis.MemoryInsights.UploadMemoryInsightsLLMXMLReportsTrace", EAutomationTestFlags::ApplicationContextMask | EAutomationTestFlags::ProductFilter)
+IMPLEMENT_SIMPLE_AUTOMATION_TEST(FMemoryInsightsUploadLLMXMLReportsTraceTest, "System.Insights.Trace.Analysis.MemoryInsights.UploadMemoryInsightsLLMXMLReportsTrace", EAutomationTestFlags::ApplicationContextMask | EAutomationTestFlags::ProductFilter)
 bool FMemoryInsightsUploadLLMXMLReportsTraceTest::RunTest(const FString& Parameters)
 {
 	const FString ReportGraphsXMLPath = FPaths::RootDir() / TEXT("EngineTest/SourceAssets/Utrace/ReportGraphs.xml");
@@ -310,7 +310,7 @@ bool MemoryInsightsAllocationsQueryTableTest(const FString& Parameters, const TM
 	return true;
 }
 
-IMPLEMENT_COMPLEX_AUTOMATION_TEST(FMemoryInsightsAllocationsQueryTableEditorPackageTest, "Insights.Analysis.MemoryInsights.AllocationsQueryTable.Editor/Package", EAutomationTestFlags::ApplicationContextMask | EAutomationTestFlags::ProductFilter)
+IMPLEMENT_COMPLEX_AUTOMATION_TEST(FMemoryInsightsAllocationsQueryTableEditorPackageTest, "System.Insights.Trace.Analysis.MemoryInsights.AllocationsQueryTable.Editor/Package", EAutomationTestFlags::ApplicationContextMask | EAutomationTestFlags::ProductFilter)
 bool FMemoryInsightsAllocationsQueryTableEditorPackageTest::RunTest(const FString& Parameters)
 {
 	bool bSuccess = MemoryInsightsAllocationsQueryTableTest(Parameters, AllocsTimeMarkerEditorPackageGetterMap, this);
@@ -340,7 +340,7 @@ void FMemoryInsightsAllocationsQueryTableEditorPackageTest::GetTests(TArray<FStr
 	}
 }
 
-IMPLEMENT_COMPLEX_AUTOMATION_TEST(FMemoryInsightsAllocationsQueryTableStandaloneTest, "Insights.Analysis.MemoryInsights.AllocationsQueryTable.Standalone", EAutomationTestFlags::ApplicationContextMask | EAutomationTestFlags::ProductFilter)
+IMPLEMENT_COMPLEX_AUTOMATION_TEST(FMemoryInsightsAllocationsQueryTableStandaloneTest, "System.Insights.Trace.Analysis.MemoryInsights.AllocationsQueryTable.Standalone", EAutomationTestFlags::ApplicationContextMask | EAutomationTestFlags::ProductFilter)
 bool FMemoryInsightsAllocationsQueryTableStandaloneTest::RunTest(const FString& Parameters)
 {
 	bool bSuccess = MemoryInsightsAllocationsQueryTableTest(Parameters, AllocsTimeMarkerStandaloneGameGetterMap, this);
diff --git a/Engine/Source/Developer/TraceInsights/Private/Insights/Tests/FunctionalTests/SessionBrowserTests.cpp b/Engine/Source/Developer/TraceInsights/Private/Insights/Tests/FunctionalTests/SessionBrowserTests.cpp
--- a/Engine/Source/Developer/TraceInsights/Private/Insights/Tests/FunctionalTests/SessionBrowserTests.cpp
+++ b/Engine/Source/Developer/TraceInsights/Private/Insights/Tests/FunctionalTests/SessionBrowserTests.cpp
@@ -17,7 +17,7 @@ DECLARE_LOG_CATEGORY_EXTERN(UnrealInsightsHubTests, Log, All);
 
 #if !WITH_EDITOR
 
-BEGIN_DEFINE_SPEC(FAutomationDriverUnrealInsightsSessionBrowserTest, "Insights.Hub.SessionBrowser", EAutomationTestFlags::ApplicationContextMask | EAutomationTestFlags::ProductFilter)
+BEGIN_DEFINE_SPEC(FAutomationDriverUnrealInsightsSessionBrowserTest, "System.Insights.Hub.SessionBrowser", EAutomationTestFlags::ApplicationContextMask | EAutomationTestFlags::ProductFilter)
 FAutomationDriverPtr Driver;
 TSharedPtr<SWindow> AutomationWindow;
 END_DEFINE_SPEC(FAutomationDriverUnrealInsightsSessionBrowserTest)
diff --git a/Engine/Source/Developer/TraceInsights/Private/Insights/Tests/FunctionalTests/TimingProfilerTests.cpp b/Engine/Source/Developer/TraceInsights/Private/Insights/Tests/FunctionalTests/TimingProfilerTests.cpp
--- a/Engine/Source/Developer/TraceInsights/Private/Insights/Tests/FunctionalTests/TimingProfilerTests.cpp
+++ b/Engine/Source/Developer/TraceInsights/Private/Insights/Tests/FunctionalTests/TimingProfilerTests.cpp
@@ -18,7 +18,7 @@ DEFINE_LOG_CATEGORY(TimingProfilerTests);
 
 #if !WITH_EDITOR && WITH_DEV_AUTOMATION_TESTS
 
-IMPLEMENT_SIMPLE_AUTOMATION_TEST(FEnumerateTest, "Insights.Analysis.TimingInsights.Enumerate", EAutomationTestFlags::ApplicationContextMask | EAutomationTestFlags::ProductFilter)
+IMPLEMENT_SIMPLE_AUTOMATION_TEST(FEnumerateTest, "System.Insights.Trace.Analysis.TimingInsights.Enumerate", EAutomationTestFlags::ApplicationContextMask | EAutomationTestFlags::ProductFilter)
 bool FEnumerateTest::RunTest(const FString& Parameters)
 {
 	FInsightsTestUtils Utils(this);
@@ -48,7 +48,7 @@ bool FEnumerateTest::RunTest(const FString& Parameters)
 	return !HasAnyErrors();
 }
 
-IMPLEMENT_SIMPLE_AUTOMATION_TEST(FEnumeratePerformanceTest, "Insights.Analysis.TimingInsights.EnumeratePerformance", EAutomationTestFlags::ApplicationContextMask | EAutomationTestFlags::PerfFilter)
+IMPLEMENT_SIMPLE_AUTOMATION_TEST(FEnumeratePerformanceTest, "System.Insights.Trace.Analysis.TimingInsights.EnumeratePerformance", EAutomationTestFlags::ApplicationContextMask | EAutomationTestFlags::PerfFilter)
 bool FEnumeratePerformanceTest::RunTest(const FString& Parameters)
 {
 	FInsightsTestUtils Utils(this);
@@ -88,7 +88,7 @@ bool FEnumeratePerformanceTest::RunTest(const FString& Parameters)
 	return !HasAnyErrors();
 }
 
-IMPLEMENT_SIMPLE_AUTOMATION_TEST(FEnumerateFastTest, "Insights.Analysis.TimingInsights.EnumerateFast", EAutomationTestFlags::ApplicationContextMask | EAutomationTestFlags::ProductFilter)
+IMPLEMENT_SIMPLE_AUTOMATION_TEST(FEnumerateFastTest, "System.Insights.Analysis.TimingInsights.EnumerateFast", EAutomationTestFlags::ApplicationContextMask | EAutomationTestFlags::ProductFilter)
 bool FEnumerateFastTest::RunTest(const FString& Parameters)
 {
 	FInsightsTestUtils Utils(this);
@@ -118,7 +118,7 @@ bool FEnumerateFastTest::RunTest(const FString& Parameters)
 	return !HasAnyErrors();
 }
 
-IMPLEMENT_SIMPLE_AUTOMATION_TEST(EnumerateEventsToFile, "Insights.Analysis.TimingInsights.EnumerateEventsToFile", EAutomationTestFlags::ApplicationContextMask | EAutomationTestFlags::ProductFilter)
+IMPLEMENT_SIMPLE_AUTOMATION_TEST(EnumerateEventsToFile, "System.Insights.Analysis.TimingInsights.EnumerateEventsToFile", EAutomationTestFlags::ApplicationContextMask | EAutomationTestFlags::ProductFilter)
 bool EnumerateEventsToFile::RunTest(const FString& Parameters)
 {
 	double SessionTime = 0.0;
@@ -159,7 +159,7 @@ bool EnumerateEventsToFile::RunTest(const FString& Parameters)
 	return !HasAnyErrors();
 }
 
-IMPLEMENT_SIMPLE_AUTOMATION_TEST(EnumerateScopesToFile, "Insights.Analysis.TimingInsights.EnumerateScopesToFile", EAutomationTestFlags::ApplicationContextMask | EAutomationTestFlags::ProductFilter)
+IMPLEMENT_SIMPLE_AUTOMATION_TEST(EnumerateScopesToFile, "System.Insights.Trace.Analysis.TimingInsights.EnumerateScopesToFile", EAutomationTestFlags::ApplicationContextMask | EAutomationTestFlags::ProductFilter)
 bool EnumerateScopesToFile::RunTest(const FString& Parameters)
 {
 	double SessionTime = 0.0;
@@ -200,7 +200,7 @@ bool EnumerateScopesToFile::RunTest(const FString& Parameters)
 	return !HasAnyErrors();
 }
 
-IMPLEMENT_SIMPLE_AUTOMATION_TEST(InsightsEnumerate10K, "Insights.Analysis.TimingInsights.Enumerate10K", EAutomationTestFlags::ApplicationContextMask | EAutomationTestFlags::ProductFilter)
+IMPLEMENT_SIMPLE_AUTOMATION_TEST(InsightsEnumerate10K, "System.Insights.Trace.Analysis.TimingInsights.Enumerate10K", EAutomationTestFlags::ApplicationContextMask | EAutomationTestFlags::ProductFilter)
 bool InsightsEnumerate10K::RunTest(const FString& Parameters)
 {
 	FTimingProfilerTests::FEnumerateTestParams Params;
@@ -213,7 +213,7 @@ bool InsightsEnumerate10K::RunTest(const FString& Parameters)
 	return !HasAnyErrors();
 }
 
-IMPLEMENT_SIMPLE_AUTOMATION_TEST(InsightsEnumerate100K, "Insights.Analysis.TimingInsights.Enumerate100K", EAutomationTestFlags::ApplicationContextMask | EAutomationTestFlags::ProductFilter)
+IMPLEMENT_SIMPLE_AUTOMATION_TEST(InsightsEnumerate100K, "System.Insights.Trace.Analysis.TimingInsights.Enumerate100K", EAutomationTestFlags::ApplicationContextMask | EAutomationTestFlags::ProductFilter)
 bool InsightsEnumerate100K::RunTest(const FString& Parameters)
 {
 	FTimingProfilerTests::FEnumerateTestParams Params;
@@ -226,7 +226,7 @@ bool InsightsEnumerate100K::RunTest(const FString& Parameters)
 	return !HasAnyErrors();
 }
 
-IMPLEMENT_SIMPLE_AUTOMATION_TEST(EnumerateByEndTimeAsyncAllTracks, "Insights.Analysis.TimingInsights.EnumerateByEndTimeAsyncAllTracks", EAutomationTestFlags::ApplicationContextMask | EAutomationTestFlags::ProductFilter)
+IMPLEMENT_SIMPLE_AUTOMATION_TEST(EnumerateByEndTimeAsyncAllTracks, "System.Insights.Trace.Analysis.TimingInsights.EnumerateByEndTimeAsyncAllTracks", EAutomationTestFlags::ApplicationContextMask | EAutomationTestFlags::ProductFilter)
 bool EnumerateByEndTimeAsyncAllTracks::RunTest(const FString& Parameters)
 {
 	FTimingProfilerTests::FEnumerateTestParams Params;
@@ -239,7 +239,7 @@ bool EnumerateByEndTimeAsyncAllTracks::RunTest(const FString& Parameters)
 	return !HasAnyErrors();
 }
 
-IMPLEMENT_SIMPLE_AUTOMATION_TEST(EnumerateByEndTimeAsyncGameThreadTrack, "Insights.Analysis.TimingInsights.EnumerateByEndTimeAsyncGameThreadTrack", EAutomationTestFlags::ApplicationContextMask | EAutomationTestFlags::ProductFilter)
+IMPLEMENT_SIMPLE_AUTOMATION_TEST(EnumerateByEndTimeAsyncGameThreadTrack, "System.Insights.Trace.Analysis.TimingInsights.EnumerateByEndTimeAsyncGameThreadTrack", EAutomationTestFlags::ApplicationContextMask | EAutomationTestFlags::ProductFilter)
 bool EnumerateByEndTimeAsyncGameThreadTrack::RunTest(const FString& Parameters)
 {
 	FTimingProfilerTests::FEnumerateTestParams Params;
@@ -252,7 +252,7 @@ bool EnumerateByEndTimeAsyncGameThreadTrack::RunTest(const FString& Parameters)
 	return !HasAnyErrors();
 }
 
-IMPLEMENT_SIMPLE_AUTOMATION_TEST(EnumerateByEndTimeAllTracks10sIntervals, "Insights.Analysis.TimingInsights.EnumerateByEndTimeAllTracks10sIntervals", EAutomationTestFlags::ApplicationContextMask | EAutomationTestFlags::ProductFilter)
+IMPLEMENT_SIMPLE_AUTOMATION_TEST(EnumerateByEndTimeAllTracks10sIntervals, "System.Insights.Trace.Analysis.TimingInsights.EnumerateByEndTimeAllTracks10sIntervals", EAutomationTestFlags::ApplicationContextMask | EAutomationTestFlags::ProductFilter)
 bool EnumerateByEndTimeAllTracks10sIntervals::RunTest(const FString& Parameters)
 {
 	FTimingProfilerTests::FEnumerateTestParams Params;
@@ -265,7 +265,7 @@ bool EnumerateByEndTimeAllTracks10sIntervals::RunTest(const FString& Parameters)
 	return !HasAnyErrors();
 }
 
-IMPLEMENT_SIMPLE_AUTOMATION_TEST(EnumerateByEndTimeAllTracks5sIntervals, "Insights.Analysis.TimingInsights.EnumerateByEndTimeAllTracks5sIntervals", EAutomationTestFlags::ApplicationContextMask | EAutomationTestFlags::ProductFilter)
+IMPLEMENT_SIMPLE_AUTOMATION_TEST(EnumerateByEndTimeAllTracks5sIntervals, "System.Insights.Trace.Analysis.TimingInsights.EnumerateByEndTimeAllTracks5sIntervals", EAutomationTestFlags::ApplicationContextMask | EAutomationTestFlags::ProductFilter)
 bool EnumerateByEndTimeAllTracks5sIntervals::RunTest(const FString& Parameters)
 {
 	FTimingProfilerTests::FEnumerateTestParams Params;
@@ -282,7 +282,7 @@ bool EnumerateByEndTimeAllTracks5sIntervals::RunTest(const FString& Parameters)
 ////////////////////////////Enumerate Async Ordered by Start Time Tests/////////////////////////////
 ////////////////////////////////////////////////////////////////////////////////////////////////////
 
-IMPLEMENT_SIMPLE_AUTOMATION_TEST(EnumerateByStartTimeAsyncAllTracks, "Insights.Analysis.TimingInsights.EnumerateByStartTimeAsyncAllTracks", EAutomationTestFlags::ApplicationContextMask | EAutomationTestFlags::ProductFilter)
+IMPLEMENT_SIMPLE_AUTOMATION_TEST(EnumerateByStartTimeAsyncAllTracks, "System.Insights.Trace.Analysis.TimingInsights.EnumerateByStartTimeAsyncAllTracks", EAutomationTestFlags::ApplicationContextMask | EAutomationTestFlags::ProductFilter)
 bool EnumerateByStartTimeAsyncAllTracks::RunTest(const FString& Parameters)
 {
 	FTimingProfilerTests::FEnumerateTestParams Params;
@@ -295,7 +295,7 @@ bool EnumerateByStartTimeAsyncAllTracks::RunTest(const FString& Parameters)
 	return !HasAnyErrors();
 }
 
-IMPLEMENT_SIMPLE_AUTOMATION_TEST(EnumerateByStartTimeAsyncGameThreadTrack, "Insights.Analysis.TimingInsights.EnumerateByStartTimeAsyncGameThreadTrack", EAutomationTestFlags::ApplicationContextMask | EAutomationTestFlags::ProductFilter)
+IMPLEMENT_SIMPLE_AUTOMATION_TEST(EnumerateByStartTimeAsyncGameThreadTrack, "System.Insights.Trace.Analysis.TimingInsights.EnumerateByStartTimeAsyncGameThreadTrack", EAutomationTestFlags::ApplicationContextMask | EAutomationTestFlags::ProductFilter)
 bool EnumerateByStartTimeAsyncGameThreadTrack::RunTest(const FString& Parameters)
 {
 	FTimingProfilerTests::FEnumerateTestParams Params;
@@ -308,7 +308,7 @@ bool EnumerateByStartTimeAsyncGameThreadTrack::RunTest(const FString& Parameters
 	return !HasAnyErrors();
 }
 
-IMPLEMENT_SIMPLE_AUTOMATION_TEST(EnumerateByStartTimeAllTracks10sIntervals, "Insights.Analysis.TimingInsights.EnumerateByStartTimeAllTracks10sIntervals", EAutomationTestFlags::ApplicationContextMask | EAutomationTestFlags::ProductFilter)
+IMPLEMENT_SIMPLE_AUTOMATION_TEST(EnumerateByStartTimeAllTracks10sIntervals, "System.Insights.Trace.Analysis.TimingInsights.EnumerateByStartTimeAllTracks10sIntervals", EAutomationTestFlags::ApplicationContextMask | EAutomationTestFlags::ProductFilter)
 bool EnumerateByStartTimeAllTracks10sIntervals::RunTest(const FString& Parameters)
 {
 	FTimingProfilerTests::FEnumerateTestParams Params;
@@ -321,7 +321,7 @@ bool EnumerateByStartTimeAllTracks10sIntervals::RunTest(const FString& Parameter
 	return !HasAnyErrors();
 }
 
-IMPLEMENT_SIMPLE_AUTOMATION_TEST(EnumerateByStartTimeAllTracks5sIntervals, "Insights.Analysis.TimingInsights.EnumerateByStartTimeAllTracks5sIntervals", EAutomationTestFlags::ApplicationContextMask | EAutomationTestFlags::ProductFilter)
+IMPLEMENT_SIMPLE_AUTOMATION_TEST(EnumerateByStartTimeAllTracks5sIntervals, "System.Insights.Trace.Analysis.TimingInsights.EnumerateByStartTimeAllTracks5sIntervals", EAutomationTestFlags::ApplicationContextMask | EAutomationTestFlags::ProductFilter)
 bool EnumerateByStartTimeAllTracks5sIntervals::RunTest(const FString& Parameters)
 {
 	FTimingProfilerTests::FEnumerateTestParams Params;
diff --git a/Engine/Source/Developer/TraceInsights/Private/Insights/Tests/FunctionalTests/UITests.cpp b/Engine/Source/Developer/TraceInsights/Private/Insights/Tests/FunctionalTests/UITests.cpp
--- a/Engine/Source/Developer/TraceInsights/Private/Insights/Tests/FunctionalTests/UITests.cpp
+++ b/Engine/Source/Developer/TraceInsights/Private/Insights/Tests/FunctionalTests/UITests.cpp
@@ -14,7 +14,7 @@
 
 DECLARE_LOG_CATEGORY_EXTERN(UITests, Log, All);
 
-IMPLEMENT_SIMPLE_AUTOMATION_TEST(FHideAndShowAllTimingViewTabs, "Insights.Analysis.TimingInsights.HideAndShowAllTimingViewTabs", EAutomationTestFlags::ApplicationContextMask | EAutomationTestFlags::ProductFilter)
+IMPLEMENT_SIMPLE_AUTOMATION_TEST(FHideAndShowAllTimingViewTabs, "System.Insights.Trace.Analysis.TimingInsights.HideAndShowAllTimingViewTabs", EAutomationTestFlags::ApplicationContextMask | EAutomationTestFlags::ProductFilter)
 bool FHideAndShowAllTimingViewTabs::RunTest(const FString& Parameters)
 {
 	TSharedPtr<FTimingProfilerManager> TimingProfilerManager = FTimingProfilerManager::Get();
@@ -36,7 +36,7 @@ bool FHideAndShowAllTimingViewTabs::RunTest(const FString& Parameters)
 	return true;
 }
 
-IMPLEMENT_SIMPLE_AUTOMATION_TEST(FMemoryFilterValueConverterTest, "Insights.Analysis.MemoryFilterValueConverter", EAutomationTestFlags::ApplicationContextMask | EAutomationTestFlags::ProductFilter)
+IMPLEMENT_SIMPLE_AUTOMATION_TEST(FMemoryFilterValueConverterTest, "System.Insights.Trace.Analysis.MemoryFilterValueConverter", EAutomationTestFlags::ApplicationContextMask | EAutomationTestFlags::ProductFilter)
 bool FMemoryFilterValueConverterTest::RunTest(const FString& Parameters)
 {
 	Insights::FMemoryFilterValueConverter Converter;
@@ -72,7 +72,7 @@ bool FMemoryFilterValueConverterTest::RunTest(const FString& Parameters)
 	return !HasAnyErrors();
 }
 
-IMPLEMENT_SIMPLE_AUTOMATION_TEST(FTimeFilterValueConverterTest, "Insights.Analysis.TimeFilterValueConverter", EAutomationTestFlags::ApplicationContextMask | EAutomationTestFlags::ProductFilter)
+IMPLEMENT_SIMPLE_AUTOMATION_TEST(FTimeFilterValueConverterTest, "System.Insights.Trace.Analysis.TimeFilterValueConverter", EAutomationTestFlags::ApplicationContextMask | EAutomationTestFlags::ProductFilter)
 bool FTimeFilterValueConverterTest::RunTest(const FString& Parameters)
 {
 	Insights::FTimeFilterValueConverter Converter;
diff --git a/Engine/Source/Developer/TraceInsights/Private/Insights/Tests/FunctionalTests/UnrealTraceServer.cpp b/Engine/Source/Developer/TraceInsights/Private/Insights/Tests/FunctionalTests/UnrealTraceServer.cpp
--- a/Engine/Source/Developer/TraceInsights/Private/Insights/Tests/FunctionalTests/UnrealTraceServer.cpp
+++ b/Engine/Source/Developer/TraceInsights/Private/Insights/Tests/FunctionalTests/UnrealTraceServer.cpp
@@ -10,7 +10,7 @@
 
 #if !WITH_EDITOR
 
-IMPLEMENT_SIMPLE_AUTOMATION_TEST(FUnrealTraceServerStartingTest, "Insights.Hub.UnrealTraceServer.Starting", EAutomationTestFlags::ApplicationContextMask | EAutomationTestFlags::ProductFilter)
+IMPLEMENT_SIMPLE_AUTOMATION_TEST(FUnrealTraceServerStartingTest, "System.Insights.Hub.UnrealTraceServer.Starting", EAutomationTestFlags::ApplicationContextMask | EAutomationTestFlags::ProductFilter)
 bool FUnrealTraceServerStartingTest::RunTest(const FString& Parameters)
 {
 	TSharedPtr<FInsightsManager> InsightsManager = FInsightsManager::Get();
diff --git a/Engine/Source/Developer/TraceInsights/Private/Insights/Tests/UnitTests.cpp b/Engine/Source/Developer/TraceInsights/Private/Insights/Tests/UnitTests.cpp
--- a/Engine/Source/Developer/TraceInsights/Private/Insights/Tests/UnitTests.cpp
+++ b/Engine/Source/Developer/TraceInsights/Private/Insights/Tests/UnitTests.cpp
@@ -2,7 +2,7 @@
 #include "CoreMinimal.h"
 #include "Misc/AutomationTest.h"
 
-IMPLEMENT_SIMPLE_AUTOMATION_TEST(FTraceInsightsUnitTest, "Insights.Analysis.TraceInsightsUnitTest", EAutomationTestFlags::ApplicationContextMask | EAutomationTestFlags::SmokeFilter)
+IMPLEMENT_SIMPLE_AUTOMATION_TEST(FTraceInsightsUnitTest, "System.Insights.Trace.Analysis.UnitTest", EAutomationTestFlags::ApplicationContextMask | EAutomationTestFlags::SmokeFilter)
 
 bool FTraceInsightsUnitTest::RunTest(const FString& Parameters)
 {
diff --git a/Engine/Source/Developer/TraceServices/Private/Tests/UnitTests.cpp b/Engine/Source/Developer/TraceServices/Private/Tests/UnitTests.cpp
--- a/Engine/Source/Developer/TraceServices/Private/Tests/UnitTests.cpp
+++ b/Engine/Source/Developer/TraceServices/Private/Tests/UnitTests.cpp
@@ -4,7 +4,7 @@
 #include "Common/PagedArray.h"
 #include "Misc/AutomationTest.h"
 
-IMPLEMENT_SIMPLE_AUTOMATION_TEST(FPagedArrayFilteringTest, "Insights.Analysis.PagedArrayFiltering", EAutomationTestFlags::ApplicationContextMask | EAutomationTestFlags::SmokeFilter)
+IMPLEMENT_SIMPLE_AUTOMATION_TEST(FPagedArrayFilteringTest, "System.Insights.Trace.Analysis.PagedArrayFiltering", EAutomationTestFlags::ApplicationContextMask | EAutomationTestFlags::SmokeFilter)
 
 bool operator!=(const FInt32Interval& lhs, const FInt32Interval& rhs)
 {
