[Insights] Fixed few typos.

diff --git a/Engine/Source/Developer/TraceInsights/Private/Insights/Tests/FunctionalTests/TimingProfilerTests.cpp b/Engine/Source/Developer/TraceInsights/Private/Insights/Tests/FunctionalTests/TimingProfilerTests.cpp
--- a/Engine/Source/Developer/TraceInsights/Private/Insights/Tests/FunctionalTests/TimingProfilerTests.cpp
+++ b/Engine/Source/Developer/TraceInsights/Private/Insights/Tests/FunctionalTests/TimingProfilerTests.cpp
@@ -230,7 +230,7 @@ IMPLEMENT_SIMPLE_AUTOMATION_TEST(EnumerateByEndTimeAsyncAllTracks, "System.Insig
 bool EnumerateByEndTimeAsyncAllTracks::RunTest(const FString& Parameters)
 {
 	FTimingProfilerTests::FEnumerateTestParams Params;
-	Params.Interval = 24.0 * 3600; //A day - Should be big enought to contain any valid session in [0, 0 + interval]
+	Params.Interval = 24.0 * 3600; //A day - Should be big enough to contain any valid session in [0, 0 + interval]
 	Params.NumEnumerations = 1;
 	Params.SortOrder = TraceServices::EEventSortOrder::ByEndTime;
 
@@ -243,7 +243,7 @@ IMPLEMENT_SIMPLE_AUTOMATION_TEST(EnumerateByEndTimeAsyncGameThreadTrack, "System
 bool EnumerateByEndTimeAsyncGameThreadTrack::RunTest(const FString& Parameters)
 {
 	FTimingProfilerTests::FEnumerateTestParams Params;
-	Params.Interval = 24.0 * 3600; //A day - Should be big enought to contain any valid session in [0, 0 + interval]
+	Params.Interval = 24.0 * 3600; //A day - Should be big enough to contain any valid session in [0, 0 + interval]
 	Params.NumEnumerations = 1;
 	Params.SortOrder = TraceServices::EEventSortOrder::ByEndTime;
 
@@ -286,7 +286,7 @@ IMPLEMENT_SIMPLE_AUTOMATION_TEST(EnumerateByStartTimeAsyncAllTracks, "System.Ins
 bool EnumerateByStartTimeAsyncAllTracks::RunTest(const FString& Parameters)
 {
 	FTimingProfilerTests::FEnumerateTestParams Params;
-	Params.Interval = 24.0 * 3600; //A day - Should be big enought to contain any valid session in [0, 0 + interval]
+	Params.Interval = 24.0 * 3600; //A day - Should be big enough to contain any valid session in [0, 0 + interval]
 	Params.NumEnumerations = 1;
 	Params.SortOrder = TraceServices::EEventSortOrder::ByStartTime;
 
@@ -299,7 +299,7 @@ IMPLEMENT_SIMPLE_AUTOMATION_TEST(EnumerateByStartTimeAsyncGameThreadTrack, "Syst
 bool EnumerateByStartTimeAsyncGameThreadTrack::RunTest(const FString& Parameters)
 {
 	FTimingProfilerTests::FEnumerateTestParams Params;
-	Params.Interval = 24.0 * 3600; //A day - Should be big enought to contain any valid session in [0, 0 + interval]
+	Params.Interval = 24.0 * 3600; //A day - Should be big enough to contain any valid session in [0, 0 + interval]
 	Params.NumEnumerations = 1;
 	Params.SortOrder = TraceServices::EEventSortOrder::ByStartTime;
 
diff --git a/Engine/Source/Runtime/Core/Private/ProfilingDebugging/Microsoft/EventTracingForWindows.cpp b/Engine/Source/Runtime/Core/Private/ProfilingDebugging/Microsoft/EventTracingForWindows.cpp
--- a/Engine/Source/Runtime/Core/Private/ProfilingDebugging/Microsoft/EventTracingForWindows.cpp
+++ b/Engine/Source/Runtime/Core/Private/ProfilingDebugging/Microsoft/EventTracingForWindows.cpp
@@ -261,7 +261,7 @@ void FPlatformEvents::Enable(FPlatformEventsTrace::EEventType Event)
 		{
 			Properties->EnableFlags |= EVENT_TRACE_FLAG_PROFILE;
 		}
-		
+
 		ULONG Status = ::ControlTraceW(NULL, KERNEL_LOGGER_NAMEW, Properties, EVENT_TRACE_CONTROL_UPDATE);
 		if (Status != ERROR_SUCCESS)
 		{
@@ -293,7 +293,7 @@ void FPlatformEvents::Disable(FPlatformEventsTrace::EEventType Event)
 		// if event is already disabled, do nothing
 		return;
 	}
-	
+
 	EnumRemoveFlags(EnabledEvents, Event);
 
 	if (EnabledEvents == FPlatformEventsTrace::EEventType::None)
@@ -335,7 +335,7 @@ uint32 FPlatformEvents::Run()
 
 		// ProcessTrace is blocking function that enters processing loop
 		// and calls TraceEventCallback callback for each event
-		// it exits processing loop when trace is closed with ::CloseTrace 
+		// it exits processing loop when trace is closed with ::CloseTrace
 		::ProcessTrace(&TraceHandle, 1, NULL, NULL);
 	}
 	else
@@ -354,7 +354,7 @@ uint32 FPlatformEvents::Run()
 void FPlatformEvents::Stop()
 {
 	// because ETW is started in background thread, then first make sure
-	// ETW startup code has completed (sucessfully or not) before actually
+	// ETW startup code has completed (successfully or not) before actually
 	// stopping it
 	ReadyEvent->Wait(MAX_uint32);
 
@@ -408,7 +408,7 @@ bool FPlatformEvents::StartETW()
 				if (!::AdjustTokenPrivileges(Token, 0, &TokenPrivileges, sizeof(TokenPrivileges), NULL, NULL))
 				{
 					uint32 Error = ::GetLastError();
-					UE_LOG(LogPlatformEvents, Warning, TEXT("Cannot enable profile privlege for process: 0x%08x"), Error);
+					UE_LOG(LogPlatformEvents, Warning, TEXT("Cannot enable profile privilege for process: 0x%08x"), Error);
 				}
 			}
 			::CloseHandle(Token);
@@ -493,7 +493,7 @@ bool FPlatformEvents::StartETW()
 		}
 	}
 
-	// open trace for receving callbacks
+	// open trace for receiving callbacks
 	{
 		EVENT_TRACE_LOGFILEW LogFile;
 		ZeroMemory(&LogFile, sizeof(LogFile));
