[Insights] Increased version for Unreal Insights to 1.05.

diff --git a/Engine/Source/Developer/TraceInsights/Public/Insights/Version.h b/Engine/Source/Developer/TraceInsights/Public/Insights/Version.h
--- a/Engine/Source/Developer/TraceInsights/Public/Insights/Version.h
+++ b/Engine/Source/Developer/TraceInsights/Public/Insights/Version.h
@@ -5,10 +5,10 @@
 // See https://semver.org/
 
 #define UNREAL_INSIGHTS_VERSION_MAJOR 1
-#define UNREAL_INSIGHTS_VERSION_MINOR 4
+#define UNREAL_INSIGHTS_VERSION_MINOR 5
 #define UNREAL_INSIGHTS_VERSION_PATCH 0
 
-#define UNREAL_INSIGHTS_VERSION_STRING "1.04"
+#define UNREAL_INSIGHTS_VERSION_STRING "1.05"
 
 /** Extra identifier string (like: "alpha", "beta.1", "beta.2", "test", "featureX", etc.). */
 //#define UNREAL_INSIGHTS_VERSION_ID_STRING "beta.1"
