[Insights] Make task relations more descriptive when they point to hidden tracks. Also add a warning in this case.

diff --git a/Engine/Source/Developer/TraceInsights/Private/Insights/TaskGraphProfiler/TaskGraphProfilerManager.cpp b/Engine/Source/Developer/TraceInsights/Private/Insights/TaskGraphProfiler/TaskGraphProfilerManager.cpp
--- a/Engine/Source/Developer/TraceInsights/Private/Insights/TaskGraphProfiler/TaskGraphProfilerManager.cpp
+++ b/Engine/Source/Developer/TraceInsights/Private/Insights/TaskGraphProfiler/TaskGraphProfilerManager.cpp
@@ -5,6 +5,7 @@
 #include "Async/TaskGraphInterfaces.h"
 #include "Features/IModularFeatures.h"
 #include "Framework/Docking/TabManager.h"
+#include "Logging/MessageLog.h"
 #include "Modules/ModuleManager.h"
 #include "TraceServices/Model/TasksProfiler.h"
 #include "Widgets/Docking/SDockTab.h"
@@ -315,8 +316,12 @@ void FTaskGraphProfilerManager::ShowTaskRelations(const TraceServices::FTaskInfo
 	{
 		GetSingleTaskRelationsForAll(Task->NestedTasks);
 	}
+
+	OutputWarnings();
 }
 
+////////////////////////////////////////////////////////////////////////////////////////////////////
+
 void FTaskGraphProfilerManager::GetSingleTaskTransitions(const TraceServices::FTaskInfo* Task, const TraceServices::ITasksProvider* TasksProvider, const FThreadTrackEvent* InSelectedEvent)
 {
 	if (Task->CreatedTimestamp != Task->LaunchedTimestamp || Task->CreatedThreadId != Task->LaunchedThreadId)
@@ -607,6 +612,16 @@ void FTaskGraphProfilerManager::AddRelation(const FThreadTrackEvent* InSelectedE
 		}
 	}
 
+	if (TaskRelationPtr->GetSourceTrack().IsValid() && !TaskRelationPtr->GetSourceTrack()->IsVisible())
+	{
+		HiddenTrackNames.Add(TaskRelationPtr->GetSourceTrack()->GetName());
+	}
+
+	if (TaskRelationPtr->GetTargetTrack().IsValid() && !TaskRelationPtr->GetTargetTrack()->IsVisible())
+	{
+		HiddenTrackNames.Add(TaskRelationPtr->GetTargetTrack()->GetName());
+	}
+
 	if (TaskRelationPtr->GetSourceTrack().IsValid() && TaskRelationPtr->GetTargetTrack().IsValid())
 	{
 		TimingView->AddRelation(Relation);
@@ -649,6 +664,8 @@ void FTaskGraphProfilerManager::ClearTaskRelations()
 		{
 			return Relations->Is<FTaskGraphRelation>();
 		});
+
+	HiddenTrackNames.Empty();
 }
 
 ////////////////////////////////////////////////////////////////////////////////////////////////////
@@ -837,6 +854,55 @@ void FTaskGraphProfilerManager::RegisterOnWindowClosedEventHandle()
 	}
 }
 
+////////////////////////////////////////////////////////////////////////////////////////////////////
+
+void FTaskGraphProfilerManager::OutputWarnings()
+{
+	if (HiddenTrackNames.Num() == 0)
+	{
+		return;
+	}
+
+	TStringBuilder<256> TrackList;
+
+	constexpr int32 MaxListedTracks = 3;
+	int32 NumTracks = 0;
+
+	for (FString TrackName : HiddenTrackNames)
+	{
+		++NumTracks;
+		TrackList.Append(TrackName);
+		TrackList.Append(TEXT(","));
+
+		if (NumTracks >= MaxListedTracks)
+		{
+			break;
+		}
+	}
+
+	int32 RemainingTracksNum = HiddenTrackNames.Num() - NumTracks;
+	FText WarningMessage;
+	if (RemainingTracksNum > 0)
+	{
+		WarningMessage = FText::Format(LOCTEXT("RelationsOnMoreHiddenTracksWarningFmt", "Some task relations point to hidden tracks: {0} and {1} more."),
+									   FText::FromString(TrackList.ToString()),
+									   RemainingTracksNum);
+	}
+	else
+	{
+		TrackList.RemoveSuffix(1);
+		WarningMessage = FText::Format(LOCTEXT("RelationsOnHiddenTracksWarningFmt", "Some task relations point to hidden tracks: {0}."),
+									   FText::FromString(TrackList.ToString()));
+	}
+
+	FName LogListingName = FTimingProfilerManager::Get()->GetLogListingName();
+	FMessageLog ReportMessageLog(LogListingName);
+	ReportMessageLog.Warning(WarningMessage);
+	ReportMessageLog.Notify();
+}
+
+////////////////////////////////////////////////////////////////////////////////////////////////////
+
 } // namespace Insights
 
 #undef LOCTEXT_NAMESPACE
diff --git a/Engine/Source/Developer/TraceInsights/Private/Insights/TaskGraphProfiler/TaskGraphProfilerManager.h b/Engine/Source/Developer/TraceInsights/Private/Insights/TaskGraphProfiler/TaskGraphProfilerManager.h
--- a/Engine/Source/Developer/TraceInsights/Private/Insights/TaskGraphProfiler/TaskGraphProfilerManager.h
+++ b/Engine/Source/Developer/TraceInsights/Private/Insights/TaskGraphProfiler/TaskGraphProfilerManager.h
@@ -154,6 +154,8 @@ private:
 	void InitializeColorCode();
 	int32 GetRelationDisplayDepth(TSharedPtr<const FThreadTimingTrack> Track, double Time, int32 KnownDepth);
 
+	void OutputWarnings();
+
 private:
 	bool bIsInitialized;
 	bool bIsAvailable;
@@ -184,6 +186,8 @@ private:
 	bool bShowNestedTasks = false;
 	
 	FDelegateHandle OnWindowClosedEventHandle;
+
+	TSet<FString> HiddenTrackNames;
 };
 
 } // namespace Insights
diff --git a/Engine/Source/Developer/TraceInsights/Private/Insights/TaskGraphProfiler/ViewModels/TaskGraphRelation.cpp b/Engine/Source/Developer/TraceInsights/Private/Insights/TaskGraphProfiler/ViewModels/TaskGraphRelation.cpp
--- a/Engine/Source/Developer/TraceInsights/Private/Insights/TaskGraphProfiler/ViewModels/TaskGraphRelation.cpp
+++ b/Engine/Source/Developer/TraceInsights/Private/Insights/TaskGraphProfiler/ViewModels/TaskGraphRelation.cpp
@@ -70,20 +70,41 @@ void FTaskGraphRelation::Draw(const FDrawContext& DrawContext, const FTimingTrac
 		return;
 	}
 
-	int32 ActualSourceDepth = FMath::Min(SourceDepth, (int32)FTimingProfilerManager::Get()->GetEventDepthLimit() - 1);
-	float Y1 = SourceTrackShared->GetPosY();
-	Y1 += Viewport.GetLayout().GetLaneY(ActualSourceDepth) + Viewport.GetLayout().EventH / 2.0f;
-	if (SourceTrackShared->GetChildTrack() && SourceTrackShared->GetChildTrack()->GetHeight() > 0.0f)
+	if (!SourceTrackShared->IsVisible() && !TargetTrackShared->IsVisible())
 	{
-		Y1 += SourceTrackShared->GetChildTrack()->GetHeight() + Viewport.GetLayout().ChildTimelineDY;
+		return;
+	}
+
+	float Y1 = 0.0f;
+	float Y2 = 0.0f;
+	if (SourceTrackShared->IsVisible())
+	{
+		int32 ActualSourceDepth = FMath::Min(SourceDepth, (int32)FTimingProfilerManager::Get()->GetEventDepthLimit() - 1);
+		Y1 = SourceTrackShared->GetPosY();
+		Y1 += Viewport.GetLayout().GetLaneY(ActualSourceDepth) + Viewport.GetLayout().EventH / 2.0f;
+		if (SourceTrackShared->GetChildTrack() && SourceTrackShared->GetChildTrack()->GetHeight() > 0.0f)
+		{
+			Y1 += SourceTrackShared->GetChildTrack()->GetHeight() + Viewport.GetLayout().ChildTimelineDY;
+		}
+	}
+	else
+	{
+		Y1 = TargetTrackShared->GetPosY();
 	}
 
-	int32 ActualTargetDepth = FMath::Min(TargetDepth, (int32)FTimingProfilerManager::Get()->GetEventDepthLimit() - 1);
-	float Y2 = TargetTrackShared->GetPosY();
-	Y2 += Viewport.GetLayout().GetLaneY(ActualTargetDepth) + Viewport.GetLayout().EventH / 2.0f;
-	if (TargetTrackShared->GetChildTrack() && TargetTrackShared->GetChildTrack()->GetHeight() > 0.0f)
+	if (TargetTrackShared->IsVisible())
+	{
+		int32 ActualTargetDepth = FMath::Min(TargetDepth, (int32)FTimingProfilerManager::Get()->GetEventDepthLimit() - 1);
+		Y2 = TargetTrackShared->GetPosY();
+		Y2 += Viewport.GetLayout().GetLaneY(ActualTargetDepth) + Viewport.GetLayout().EventH / 2.0f;
+		if (TargetTrackShared->GetChildTrack() && TargetTrackShared->GetChildTrack()->GetHeight() > 0.0f)
+		{
+			Y2 += TargetTrackShared->GetChildTrack()->GetHeight() + Viewport.GetLayout().ChildTimelineDY;
+		}
+	}
+	else
 	{
-		Y2 += TargetTrackShared->GetChildTrack()->GetHeight() + Viewport.GetLayout().ChildTimelineDY;
+		Y2 = SourceTrackShared->GetPosY();
 	}
 
 	const FVector2D StartPoint = FVector2D(X1, Y1);
