[Insights] Fixed "Open Source File" functionality (in Log View, Timers panel, Tasks panel and in the "Mem Allocs" table/tree view) to handle correctly the case when there is no source code accessor.

diff --git a/Engine/Source/Developer/TraceInsights/Private/Insights/MemoryProfiler/Widgets/SMemAllocTableTreeView.cpp b/Engine/Source/Developer/TraceInsights/Private/Insights/MemoryProfiler/Widgets/SMemAllocTableTreeView.cpp
--- a/Engine/Source/Developer/TraceInsights/Private/Insights/MemoryProfiler/Widgets/SMemAllocTableTreeView.cpp
+++ b/Engine/Source/Developer/TraceInsights/Private/Insights/MemoryProfiler/Widgets/SMemAllocTableTreeView.cpp
@@ -595,19 +595,6 @@ TSharedPtr<SWidget> SMemAllocTableTreeView::ConstructToolbar()
 			]
 		];
 
-	//for (const TSharedRef<ITableTreeViewPreset>& ViewPreset : AvailableViewPresets)
-	//{
-	//	Box->AddSlot()
-	//		.AutoWidth()
-	//		.Padding(4.0f, 0.0f, 0.0f, 0.0f)
-	//		[
-	//			SNew(SButton)
-	//			.Text(ViewPreset->GetName())
-	//			.ToolTipText(ViewPreset->GetToolTip())
-	//			.OnClicked(this, &SMemAllocTableTreeView::OnApplyViewPreset, (const ITableTreeViewPreset*)&ViewPreset.Get())
-	//		];
-	//}
-
 	Box->AddSlot()
 		.AutoWidth()
 		.Padding(4.0f, 0.0f, 0.0f, 0.0f)
@@ -654,10 +641,10 @@ void SMemAllocTableTreeView::InitAvailableViewPresets()
 		}
 		virtual void GetColumnConfigSet(TArray<FTableColumnConfig>& InOutConfigSet) const override
 		{
-			InOutConfigSet.Add({ FTable::GetHierarchyColumnId(),               true, 200.0f });
-			InOutConfigSet.Add({ FMemAllocTableColumns::CountColumnId,         true, 100.0f });
-			InOutConfigSet.Add({ FMemAllocTableColumns::SizeColumnId,          true, 100.0f });
-			InOutConfigSet.Add({ FMemAllocTableColumns::TagColumnId,           true, 120.0f });
+			InOutConfigSet.Add({ FTable::GetHierarchyColumnId(),                    true, 200.0f });
+			InOutConfigSet.Add({ FMemAllocTableColumns::CountColumnId,              true, 100.0f });
+			InOutConfigSet.Add({ FMemAllocTableColumns::SizeColumnId,               true, 100.0f });
+			InOutConfigSet.Add({ FMemAllocTableColumns::TagColumnId,                true, 120.0f });
 			InOutConfigSet.Add({ FMemAllocTableColumns::AllocFunctionColumnId,      true, 550.0f });
 			InOutConfigSet.Add({ FMemAllocTableColumns::AllocCallstackSizeColumnId, true, 100.0f });
 		}
@@ -695,21 +682,21 @@ void SMemAllocTableTreeView::InitAvailableViewPresets()
 		}
 		virtual void GetColumnConfigSet(TArray<FTableColumnConfig>& InOutConfigSet) const override
 		{
-			InOutConfigSet.Add({ FTable::GetHierarchyColumnId(),                 true, 200.0f });
-			InOutConfigSet.Add({ FMemAllocTableColumns::StartEventIndexColumnId, true, 100.0f });
-			InOutConfigSet.Add({ FMemAllocTableColumns::EndEventIndexColumnId,   true, 100.0f });
-			InOutConfigSet.Add({ FMemAllocTableColumns::EventDistanceColumnId,   true, 100.0f });
-			InOutConfigSet.Add({ FMemAllocTableColumns::StartTimeColumnId,       true, 100.0f });
-			InOutConfigSet.Add({ FMemAllocTableColumns::EndTimeColumnId,         true, 100.0f });
-			InOutConfigSet.Add({ FMemAllocTableColumns::DurationColumnId,        true, 100.0f });
-			InOutConfigSet.Add({ FMemAllocTableColumns::AddressColumnId,         true, 120.0f });
-			InOutConfigSet.Add({ FMemAllocTableColumns::MemoryPageColumnId,      true, 120.0f });
-			InOutConfigSet.Add({ FMemAllocTableColumns::CountColumnId,           true, 100.0f });
-			InOutConfigSet.Add({ FMemAllocTableColumns::SizeColumnId,            true, 100.0f });
-			InOutConfigSet.Add({ FMemAllocTableColumns::TagColumnId,             true, 120.0f });
-			InOutConfigSet.Add({ FMemAllocTableColumns::AllocFunctionColumnId,        true, 550.0f });
-			InOutConfigSet.Add({ FMemAllocTableColumns::AllocSourceFileColumnId,      true, 550.0f });
-			InOutConfigSet.Add({ FMemAllocTableColumns::AllocCallstackSizeColumnId,   true, 100.0f });
+			InOutConfigSet.Add({ FTable::GetHierarchyColumnId(),                    true, 200.0f });
+			InOutConfigSet.Add({ FMemAllocTableColumns::StartEventIndexColumnId,    true, 100.0f });
+			InOutConfigSet.Add({ FMemAllocTableColumns::EndEventIndexColumnId,      true, 100.0f });
+			InOutConfigSet.Add({ FMemAllocTableColumns::EventDistanceColumnId,      true, 100.0f });
+			InOutConfigSet.Add({ FMemAllocTableColumns::StartTimeColumnId,          true, 100.0f });
+			InOutConfigSet.Add({ FMemAllocTableColumns::EndTimeColumnId,            true, 100.0f });
+			InOutConfigSet.Add({ FMemAllocTableColumns::DurationColumnId,           true, 100.0f });
+			InOutConfigSet.Add({ FMemAllocTableColumns::AddressColumnId,            true, 120.0f });
+			InOutConfigSet.Add({ FMemAllocTableColumns::MemoryPageColumnId,         true, 120.0f });
+			InOutConfigSet.Add({ FMemAllocTableColumns::CountColumnId,              true, 100.0f });
+			InOutConfigSet.Add({ FMemAllocTableColumns::SizeColumnId,               true, 100.0f });
+			InOutConfigSet.Add({ FMemAllocTableColumns::TagColumnId,                true, 120.0f });
+			InOutConfigSet.Add({ FMemAllocTableColumns::AllocFunctionColumnId,      true, 550.0f });
+			InOutConfigSet.Add({ FMemAllocTableColumns::AllocSourceFileColumnId,    true, 550.0f });
+			InOutConfigSet.Add({ FMemAllocTableColumns::AllocCallstackSizeColumnId, true, 100.0f });
 		}
 	};
 	AvailableViewPresets.Add(MakeShared<FDetailedViewPreset>());
@@ -755,10 +742,10 @@ void SMemAllocTableTreeView::InitAvailableViewPresets()
 		}
 		virtual void GetColumnConfigSet(TArray<FTableColumnConfig>& InOutConfigSet) const override
 		{
-			InOutConfigSet.Add({ FTable::GetHierarchyColumnId(),          true, 400.0f });
-			InOutConfigSet.Add({ FMemAllocTableColumns::CountColumnId,    true, 100.0f });
-			InOutConfigSet.Add({ FMemAllocTableColumns::SizeColumnId,     true, 100.0f });
-			InOutConfigSet.Add({ FMemAllocTableColumns::TagColumnId,      true, 200.0f });
+			InOutConfigSet.Add({ FTable::GetHierarchyColumnId(),               true, 400.0f });
+			InOutConfigSet.Add({ FMemAllocTableColumns::CountColumnId,         true, 100.0f });
+			InOutConfigSet.Add({ FMemAllocTableColumns::SizeColumnId,          true, 100.0f });
+			InOutConfigSet.Add({ FMemAllocTableColumns::TagColumnId,           true, 200.0f });
 			InOutConfigSet.Add({ FMemAllocTableColumns::AllocFunctionColumnId, true, 200.0f });
 		}
 	};
@@ -805,11 +792,11 @@ void SMemAllocTableTreeView::InitAvailableViewPresets()
 		}
 		virtual void GetColumnConfigSet(TArray<FTableColumnConfig>& InOutConfigSet) const override
 		{
-			InOutConfigSet.Add({ FTable::GetHierarchyColumnId(),          true, 200.0f });
-			InOutConfigSet.Add({ FMemAllocTableColumns::AddressColumnId,  true, 120.0f });
-			InOutConfigSet.Add({ FMemAllocTableColumns::CountColumnId,    true, 100.0f });
-			InOutConfigSet.Add({ FMemAllocTableColumns::SizeColumnId,     true, 100.0f });
-			InOutConfigSet.Add({ FMemAllocTableColumns::TagColumnId,      true, 120.0f });
+			InOutConfigSet.Add({ FTable::GetHierarchyColumnId(),               true, 200.0f });
+			InOutConfigSet.Add({ FMemAllocTableColumns::AddressColumnId,       true, 120.0f });
+			InOutConfigSet.Add({ FMemAllocTableColumns::CountColumnId,         true, 100.0f });
+			InOutConfigSet.Add({ FMemAllocTableColumns::SizeColumnId,          true, 100.0f });
+			InOutConfigSet.Add({ FMemAllocTableColumns::TagColumnId,           true, 120.0f });
 			InOutConfigSet.Add({ FMemAllocTableColumns::AllocFunctionColumnId, true, 400.0f });
 		}
 	};
@@ -856,9 +843,9 @@ void SMemAllocTableTreeView::InitAvailableViewPresets()
 		}
 		virtual void GetColumnConfigSet(TArray<FTableColumnConfig>& InOutConfigSet) const override
 		{
-			InOutConfigSet.Add({ FTable::GetHierarchyColumnId(),          true, 200.0f });
-			InOutConfigSet.Add({ FMemAllocTableColumns::CountColumnId,    true, 100.0f });
-			InOutConfigSet.Add({ FMemAllocTableColumns::SizeColumnId,     true, 100.0f });
+			InOutConfigSet.Add({ FTable::GetHierarchyColumnId(),               true, 200.0f });
+			InOutConfigSet.Add({ FMemAllocTableColumns::CountColumnId,         true, 100.0f });
+			InOutConfigSet.Add({ FMemAllocTableColumns::SizeColumnId,          true, 100.0f });
 			InOutConfigSet.Add({ FMemAllocTableColumns::AllocFunctionColumnId, true, 400.0f });
 		}
 	};
@@ -1025,7 +1012,7 @@ void SMemAllocTableTreeView::InitAvailableViewPresets()
 		}
 		virtual FText GetToolTip() const override
 		{
-			return 
+			return
 				bIsAllocCallstack
 				? (bIsInvertedCallstack ?
 				LOCTEXT("InvertedCallstack_Alloc_PresetToolTip", "Inverted Alloc Callstack Breakdown View\nConfigure the tree view to show a breakdown of allocations by inverted callstack.") :
@@ -1065,10 +1052,10 @@ void SMemAllocTableTreeView::InitAvailableViewPresets()
 		}
 		virtual void GetColumnConfigSet(TArray<FTableColumnConfig>& InOutConfigSet) const override
 		{
-			InOutConfigSet.Add({ FTable::GetHierarchyColumnId(),          true, 400.0f });
-			InOutConfigSet.Add({ FMemAllocTableColumns::CountColumnId,    true, 100.0f });
-			InOutConfigSet.Add({ FMemAllocTableColumns::SizeColumnId,     true, 100.0f });
-			InOutConfigSet.Add({ FMemAllocTableColumns::TagColumnId,      true, 200.0f });
+			InOutConfigSet.Add({ FTable::GetHierarchyColumnId(),       true, 400.0f });
+			InOutConfigSet.Add({ FMemAllocTableColumns::CountColumnId, true, 100.0f });
+			InOutConfigSet.Add({ FMemAllocTableColumns::SizeColumnId,  true, 100.0f });
+			InOutConfigSet.Add({ FMemAllocTableColumns::TagColumnId,   true, 200.0f });
 			if (bIsAllocCallstack)
 			{
 				InOutConfigSet.Add({ FMemAllocTableColumns::AllocFunctionColumnId, true, 200.0f });
@@ -1130,11 +1117,11 @@ void SMemAllocTableTreeView::InitAvailableViewPresets()
 		}
 		virtual void GetColumnConfigSet(TArray<FTableColumnConfig>& InOutConfigSet) const override
 		{
-			InOutConfigSet.Add({ FTable::GetHierarchyColumnId(),          true, 200.0f });
-			InOutConfigSet.Add({ FMemAllocTableColumns::AddressColumnId,  true, 120.0f });
-			InOutConfigSet.Add({ FMemAllocTableColumns::CountColumnId,    true, 100.0f });
-			InOutConfigSet.Add({ FMemAllocTableColumns::SizeColumnId,     true, 100.0f });
-			InOutConfigSet.Add({ FMemAllocTableColumns::TagColumnId,      true, 120.0f });
+			InOutConfigSet.Add({ FTable::GetHierarchyColumnId(),               true, 200.0f });
+			InOutConfigSet.Add({ FMemAllocTableColumns::AddressColumnId,       true, 120.0f });
+			InOutConfigSet.Add({ FMemAllocTableColumns::CountColumnId,         true, 100.0f });
+			InOutConfigSet.Add({ FMemAllocTableColumns::SizeColumnId,          true, 100.0f });
+			InOutConfigSet.Add({ FMemAllocTableColumns::TagColumnId,           true, 120.0f });
 			InOutConfigSet.Add({ FMemAllocTableColumns::AllocFunctionColumnId, true, 400.0f });
 		}
 	};
@@ -1601,10 +1588,14 @@ TSharedPtr<FMemAllocNode> SMemAllocTableTreeView::GetSingleSelectedMemAllocNode(
 {
 	if (TreeView->GetNumItemsSelected() == 1)
 	{
-		FMemAllocNodePtr SelectedTreeNode = StaticCastSharedPtr<FMemAllocNode>(TreeView->GetSelectedItems()[0]);
-		if (SelectedTreeNode.IsValid() && !SelectedTreeNode->IsGroup())
+		FTableTreeNodePtr TreeNode = TreeView->GetSelectedItems()[0];
+		if (TreeNode.IsValid() && TreeNode->Is<FMemAllocNode>())
 		{
-			return SelectedTreeNode;
+			TSharedPtr<FMemAllocNode> SelectedTreeNode = StaticCastSharedPtr<FMemAllocNode>(TreeNode);
+			if (SelectedTreeNode.IsValid() && !SelectedTreeNode->IsGroup())
+			{
+				return SelectedTreeNode;
+			}
 		}
 	}
 	return nullptr;
@@ -1612,88 +1603,168 @@ TSharedPtr<FMemAllocNode> SMemAllocTableTreeView::GetSingleSelectedMemAllocNode(
 
 ////////////////////////////////////////////////////////////////////////////////////////////////////
 
-void SMemAllocTableTreeView::ExtendMenu(FMenuBuilder& MenuBuilder)
+TSharedPtr<FCallstackFrameGroupNode> SMemAllocTableTreeView::GetSingleSelectedCallstackFrameGroupNode() const
 {
+	if (TreeView->GetNumItemsSelected() == 1)
+	{
+		FTableTreeNodePtr TreeNode = TreeView->GetSelectedItems()[0];
+		if (TreeNode.IsValid() && TreeNode->Is<FCallstackFrameGroupNode>())
+		{
+			TSharedPtr<FCallstackFrameGroupNode> SelectedTreeNode = StaticCastSharedPtr<FCallstackFrameGroupNode>(TreeNode);
+			if (SelectedTreeNode.IsValid() && SelectedTreeNode->IsGroup())
+			{
+				return SelectedTreeNode;
+			}
+		}
+	}
+	return nullptr;
+}
+
+////////////////////////////////////////////////////////////////////////////////////////////////////
+
+void SMemAllocTableTreeView::ExtendMenu(TSharedRef<FExtender> Extender)
+{
+	Extender->AddMenuExtension("Misc", EExtensionHook::Before, nullptr, FMenuExtensionDelegate::CreateSP(this, &SMemAllocTableTreeView::ExtendMenuAllocation));
+	Extender->AddMenuExtension("Misc", EExtensionHook::Before, nullptr, FMenuExtensionDelegate::CreateSP(this, &SMemAllocTableTreeView::ExtendMenuCallstackFrame));
+	Extender->AddMenuExtension("Misc", EExtensionHook::After, nullptr, FMenuExtensionDelegate::CreateSP(this, &SMemAllocTableTreeView::ExtendMenuExportSnapshot));
+}
+
+////////////////////////////////////////////////////////////////////////////////////////////////////
+
+void SMemAllocTableTreeView::ExtendMenuAllocation(FMenuBuilder& MenuBuilder)
+{
+	FMemAllocNodePtr SingleSelectedMemAllocNode = GetSingleSelectedMemAllocNode();
+	if (!SingleSelectedMemAllocNode.IsValid())
+	{
+		return;
+	}
+
 	ISourceCodeAccessModule& SourceCodeAccessModule = FModuleManager::LoadModuleChecked<ISourceCodeAccessModule>("SourceCodeAccess");
 	ISourceCodeAccessor& SourceCodeAccessor = SourceCodeAccessModule.GetAccessor();
 
-	FMemAllocNodePtr SingleSelectedMemAllocNode = GetSingleSelectedMemAllocNode();
-	if (SingleSelectedMemAllocNode.IsValid() && CountSourceFiles(*SingleSelectedMemAllocNode) > 0)
+	MenuBuilder.BeginSection("Allocation", LOCTEXT("ContextMenu_Section_Allocation", "Allocation"));
 	{
-		MenuBuilder.BeginSection("Allocation", LOCTEXT("ContextMenu_Section_OpenSource", "Allocation"));
+		FText ItemLabel;
+		FText ItemToolTip;
+
+		if (SourceCodeAccessor.CanAccessSourceCode())
+		{
+			ItemLabel = FText::Format(LOCTEXT("ContextMenu_OpenSourceAllocCallstack_SubMenu", "Open in {0} | Alloc Callstack"),
+				SourceCodeAccessor.GetNameText());
+			ItemToolTip = FText::Format(LOCTEXT("ContextMenu_OpenSourceAllocCallstack_SubMenu_Desc", "Opens the source file of the selected allocation callstack frame in {0}."),
+				SourceCodeAccessor.GetNameText());
+		}
+		else
 		{
-			FText ItemLabel = FText::Format(LOCTEXT("ContextMenu_Open_SubMenu_Alloc", "Open in {0} | Alloc Callstack"), SourceCodeAccessor.GetNameText());
-			FText ItemToolTip = FText::Format(LOCTEXT("ContextMenu_Open_Desc_SubMenu_Alloc", "Open source file of selected allocation callstack frame in {0}."), SourceCodeAccessor.GetNameText());
+			ItemLabel = LOCTEXT("ContextMenu_AllocCallstack_SubMenu", "Alloc Callstack");
+			ItemToolTip = LOCTEXT("ContextMenu_SourceCodeAccessorNA", "Source Code Accessor is not available.");
+		}
 
-			// Alloc callstack
-			MenuBuilder.AddSubMenu
-			(
-				ItemLabel,
-				ItemToolTip,
-				FNewMenuDelegate::CreateSP(this, &SMemAllocTableTreeView::BuildOpenSourceSubMenu, true),
-				false,
-				FSlateIcon(SourceCodeAccessor.GetStyleSet(), SourceCodeAccessor.GetOpenIconName())
-			);
+		// Alloc Callstack
+		MenuBuilder.AddSubMenu
+		(
+			ItemLabel,
+			ItemToolTip,
+			FNewMenuDelegate::CreateSP(this, &SMemAllocTableTreeView::BuildOpenSourceSubMenu, true),
+			false,
+			FSlateIcon(SourceCodeAccessor.GetStyleSet(), SourceCodeAccessor.GetOpenIconName())
+		);
 
-			ItemLabel = FText::Format(LOCTEXT("ContextMenu_Open_SubMenu_Free", "Open in {0} | Free Callstack"), SourceCodeAccessor.GetNameText());
-			ItemToolTip = FText::Format(LOCTEXT("ContextMenu_Open_Desc_SubMenu_Free", "Open source file of selected free callstack frame in {0}."), SourceCodeAccessor.GetNameText());
+		if (SourceCodeAccessor.CanAccessSourceCode())
+		{
+			ItemLabel = FText::Format(LOCTEXT("ContextMenu_OpenSourceFreeCallstack_SubMenu", "Open in {0} | Free Callstack"),
+				SourceCodeAccessor.GetNameText());
+			ItemToolTip = FText::Format(LOCTEXT("ContextMenu_OpenSourceFreeCallstack_SubMenu_Desc", "Opens the source file of the selected free callstack frame in {0}."),
+				SourceCodeAccessor.GetNameText());
+		}
+		else
+		{
+			ItemLabel = LOCTEXT("ContextMenu_FreeCallstack_SubMenu", "Free Callstack");
+			ItemToolTip = LOCTEXT("ContextMenu_SourceCodeAccessorNA", "Source Code Accessor is not available.");
+		}
+
+		// Free Callstack
+		MenuBuilder.AddSubMenu
+		(
+			ItemLabel,
+			ItemToolTip,
+			FNewMenuDelegate::CreateSP(this, &SMemAllocTableTreeView::BuildOpenSourceSubMenu, false),
+			false,
+			FSlateIcon(SourceCodeAccessor.GetStyleSet(), SourceCodeAccessor.GetOpenIconName())
+		);
+	}
+	MenuBuilder.EndSection();
+}
+
+////////////////////////////////////////////////////////////////////////////////////////////////////
 
-			// Free Callstack
-			MenuBuilder.AddSubMenu
-			(
+void SMemAllocTableTreeView::ExtendMenuCallstackFrame(FMenuBuilder & MenuBuilder)
+{
+	TSharedPtr<FCallstackFrameGroupNode> SingleSelectedCallstackFrameGroupNode = GetSingleSelectedCallstackFrameGroupNode();
+	if (!SingleSelectedCallstackFrameGroupNode.IsValid())
+	{
+		return;
+	}
+
+	ISourceCodeAccessModule& SourceCodeAccessModule = FModuleManager::LoadModuleChecked<ISourceCodeAccessModule>("SourceCodeAccess");
+	ISourceCodeAccessor& SourceCodeAccessor = SourceCodeAccessModule.GetAccessor();
+
+	MenuBuilder.BeginSection("CallstackFrame", LOCTEXT("ContextMenu_Section_CallstackFrame", "Callstack Frame"));
+	{
+		if (SourceCodeAccessor.CanAccessSourceCode())
+		{
+			FText ItemLabel = FText::Format(LOCTEXT("ContextMenu_OpenSourceFile", "Open Source File in {0}"),
+				SourceCodeAccessor.GetNameText());
+
+			FText FileName = GetSelectedCallstackFrameFileName();
+			FText ItemToolTip = FText::Format(LOCTEXT("ContextMenu_OpenSourceFile_Desc", "Opens the source file of the selected callstack frame in {0}.\n{1}"),
+				SourceCodeAccessor.GetNameText(),
+				FileName);
+
+			MenuBuilder.AddMenuEntry(
 				ItemLabel,
 				ItemToolTip,
-				FNewMenuDelegate::CreateSP(this, &SMemAllocTableTreeView::BuildOpenSourceSubMenu, false),
-				false,
-				FSlateIcon(SourceCodeAccessor.GetStyleSet(), SourceCodeAccessor.GetOpenIconName())
-			);
+				FSlateIcon(SourceCodeAccessor.GetStyleSet(), SourceCodeAccessor.GetOpenIconName()),
+				FUIAction(
+					FExecuteAction::CreateSP(this, &SMemAllocTableTreeView::OpenCallstackFrameSourceFileInIDE),
+					FCanExecuteAction::CreateSP(this, &SMemAllocTableTreeView::CanOpenCallstackFrameSourceFileInIDE)));
 		}
-		MenuBuilder.EndSection();
-	}
-	else
-	{
-		MenuBuilder.BeginSection("CallstackFrame", LOCTEXT("ContextMenu_Section_CallstackFrame", "Callstack Frame"));
+		else
 		{
-			FText ItemLabel = FText::Format(LOCTEXT("ContextMenu_Open", "Open in {0}"), SourceCodeAccessor.GetNameText());
+			FText ItemLabel = LOCTEXT("ContextMenu_OpenSourceFile_NoAccessor", "Open Source File");
+
 			FText FileName = GetSelectedCallstackFrameFileName();
-			FText ItemToolTip = FText::Format(LOCTEXT("ContextMenu_Open_Desc", "Open source file of selected callstack frame in {0}.\n{1}"), SourceCodeAccessor.GetNameText(), FileName);
-
-			FUIAction Action_OpenIDE
-			(
-				FExecuteAction::CreateSP(this, &SMemAllocTableTreeView::OpenCallstackFrameSourceFileInIDE),
-				FCanExecuteAction::CreateSP(this, &SMemAllocTableTreeView::CanOpenCallstackFrameSourceFileInIDE)
-			);
-			MenuBuilder.AddMenuEntry
-			(
+			FText ItemToolTip = FText::Format(LOCTEXT("ContextMenu_OpenSourceFile_NoAccessor_Desc_Fmt", "{0}\nSource Code Accessor is not available."),
+				FileName);
+
+			MenuBuilder.AddMenuEntry(
 				ItemLabel,
 				ItemToolTip,
 				FSlateIcon(SourceCodeAccessor.GetStyleSet(), SourceCodeAccessor.GetOpenIconName()),
-				Action_OpenIDE,
+				FUIAction(FExecuteAction(), FCanExecuteAction::CreateLambda([]() { return false; })),
 				NAME_None,
-				EUserInterfaceActionType::Button
-			);
+				EUserInterfaceActionType::None);
 		}
-		MenuBuilder.EndSection();
 	}
+	MenuBuilder.EndSection();
+}
 
-	{
-		const FText ItemLabel = LOCTEXT("ContextMenu_Export_SubMenu", "Export Snapshot...");
-		const FText ItemToolTip = LOCTEXT("ContextMenu_Export_Desc_SubMenu", "Export memory snapshot to construct diff later.");
+////////////////////////////////////////////////////////////////////////////////////////////////////
 
-		FUIAction Action_ExportSnapshot
-		(
+void SMemAllocTableTreeView::ExtendMenuExportSnapshot(FMenuBuilder& MenuBuilder)
+{
+	const FText ItemLabel = LOCTEXT("ContextMenu_Export_SubMenu", "Export Snapshot...");
+	const FText ItemToolTip = LOCTEXT("ContextMenu_Export_Desc_SubMenu", "Export memory snapshot to construct diff later.");
+
+	MenuBuilder.AddMenuEntry(
+		ItemLabel,
+		ItemToolTip,
+		FSlateIcon(FAppStyle::GetAppStyleSetName(), "Icons.Save"),
+		FUIAction(
 			FExecuteAction::CreateSP(this, &SMemAllocTableTreeView::ExportMemorySnapshot),
-			FCanExecuteAction::CreateSP(this, &SMemAllocTableTreeView::IsExportMemorySnapshotAvailable)
-		);
-		MenuBuilder.AddMenuEntry(
-			ItemLabel,
-			ItemToolTip,
-			FSlateIcon(FAppStyle::GetAppStyleSetName(), "Icons.Save"),
-			Action_ExportSnapshot,
-			NAME_None,
-			EUserInterfaceActionType::Button
-		);
-	}
+			FCanExecuteAction::CreateSP(this, &SMemAllocTableTreeView::IsExportMemorySnapshotAvailable)),
+		NAME_None,
+		EUserInterfaceActionType::Button);
 }
 
 ////////////////////////////////////////////////////////////////////////////////////////////////////
@@ -1772,14 +1843,24 @@ bool SMemAllocTableTreeView::BuildOpenSourceSubMenuItems(FMenuBuilder& MenuBuild
 
 				ItemLabel = FText::Format(LOCTEXT("ContextMenu_OpenSource_Fmt1", "{0} ({1}) \u2192 {2}"),
 					FileName,
-					FText::AsNumber(Frame->Symbol->Line),
+					FText::AsNumber(Frame->Symbol->Line, &FNumberFormattingOptions::DefaultNoGrouping()),
 					SymbolName);
 
-				ItemToolTip = FText::Format(LOCTEXT("ContextMenu_OpenSource_Desc_Fmt1", "Open source file of selected callstack frame in {0}.\n{1} (line {2})\n\u2192 {3}"),
-					SourceCodeAccessor.GetNameText(),
-					FText::FromString(Frame->Symbol->File),
-					FText::AsNumber(Frame->Symbol->Line),
-					FText::FromString(Frame->Symbol->Name));
+				if (SourceCodeAccessor.CanAccessSourceCode())
+				{
+					ItemToolTip = FText::Format(LOCTEXT("ContextMenu_OpenSource_Desc_Fmt1", "Opens the source file of the selected callstack frame in {0}.\n{1} (line {2})\n\u2192 {3}"),
+						SourceCodeAccessor.GetNameText(),
+						FText::FromString(Frame->Symbol->File),
+						FText::AsNumber(Frame->Symbol->Line, &FNumberFormattingOptions::DefaultNoGrouping()),
+						FText::FromString(Frame->Symbol->Name));
+				}
+				else
+				{
+					ItemToolTip = FText::Format(LOCTEXT("ContextMenu_OpenSource_NoAccessor_Desc_Fmt1", "{0} (line {1})\n\u2192 {2}\nSource Code Accessor is not available."),
+						FText::FromString(Frame->Symbol->File),
+						FText::AsNumber(Frame->Symbol->Line, &FNumberFormattingOptions::DefaultNoGrouping()),
+						FText::FromString(Frame->Symbol->Name));
+				}
 			}
 			else
 			{
@@ -1788,26 +1869,28 @@ bool SMemAllocTableTreeView::BuildOpenSourceSubMenuItems(FMenuBuilder& MenuBuild
 					FText::FromString(FString::Printf(TEXT("0x%llX"), Frame->Addr)),
 					FText::FromString(TraceServices::QueryResultToString(Frame->Symbol->GetResult())));
 
-				ItemToolTip = FText::Format(LOCTEXT("ContextMenu_OpenSource_Desc_Fmt2", "Open source file of selected callstack frame in {0}."),
-					SourceCodeAccessor.GetNameText());
+				if (SourceCodeAccessor.CanAccessSourceCode())
+				{
+					ItemToolTip = FText::Format(LOCTEXT("ContextMenu_OpenSource_Desc_Fmt2", "Opens the source file of the selected callstack frame in {0}."),
+						SourceCodeAccessor.GetNameText());
+				}
+				else
+				{
+					ItemToolTip = LOCTEXT("ContextMenu_SourceCodeAccessorNA", "Source Code Accessor is not available.");
+				}
 			}
 
-			const bool bFileExists = FPaths::FileExists(Frame->Symbol->File);
+			const bool bCanOpenSource = SourceCodeAccessor.CanAccessSourceCode() && FPaths::FileExists(Frame->Symbol->File);
 
-			FUIAction Action_OpenIDE
-			(
-				FExecuteAction::CreateSP(this, &SMemAllocTableTreeView::OpenSourceFileInIDE, Frame->Symbol->File, uint32(Frame->Symbol->Line)),
-				FCanExecuteAction::CreateLambda([bFileExists]() { return bFileExists; })
-			);
-			MenuBuilder.AddMenuEntry
-			(
+			MenuBuilder.AddMenuEntry(
 				ItemLabel,
 				ItemToolTip,
 				FSlateIcon(),
-				Action_OpenIDE,
+				FUIAction(
+					FExecuteAction::CreateSP(this, &SMemAllocTableTreeView::OpenSourceFileInIDE, Frame->Symbol->File, uint32(Frame->Symbol->Line)),
+					FCanExecuteAction::CreateLambda([bCanOpenSource]() { return bCanOpenSource; })),
 				NAME_None,
-				EUserInterfaceActionType::Button
-			);
+				EUserInterfaceActionType::Button);
 
 			++NumSourceFiles;
 		}
@@ -1839,26 +1922,13 @@ void SMemAllocTableTreeView::BuildOpenSourceSubMenu(FMenuBuilder& MenuBuilder, b
 
 		if (!bHasAnySourceFilesToOpen)
 		{
-			struct FLocal
-			{
-				static bool ReturnFalse()
-				{
-					return false;
-				}
-			};
-
-			FUIAction DummyUIAction;
-			DummyUIAction.CanExecuteAction = FCanExecuteAction::CreateStatic(&FLocal::ReturnFalse);
-
-			MenuBuilder.AddMenuEntry
-			(
-				LOCTEXT("ContextMenu_OpenSourceNA", "No Source File Available"),
+			MenuBuilder.AddMenuEntry(
+				LOCTEXT("ContextMenu_OpenSourceNA", "Not Available"),
 				TAttribute<FText>(),
 				FSlateIcon(),
-				DummyUIAction,
+				FUIAction(FExecuteAction(), FCanExecuteAction::CreateLambda([]() { return false; })),
 				NAME_None,
-				EUserInterfaceActionType::None
-			);
+				EUserInterfaceActionType::None);
 		}
 	}
 	MenuBuilder.EndSection();
@@ -1868,10 +1938,9 @@ void SMemAllocTableTreeView::BuildOpenSourceSubMenu(FMenuBuilder& MenuBuilder, b
 
 void SMemAllocTableTreeView::OpenSourceFileInIDE(const TCHAR* InFile, uint32 Line) const
 {
-	ISourceCodeAccessModule& SourceCodeAccessModule = FModuleManager::LoadModuleChecked<ISourceCodeAccessModule>("SourceCodeAccess");
-
 	const FString File = InFile;
 
+	ISourceCodeAccessModule& SourceCodeAccessModule = FModuleManager::LoadModuleChecked<ISourceCodeAccessModule>("SourceCodeAccess");
 	if (FPaths::FileExists(File))
 	{
 		ISourceCodeAccessor& SourceCodeAccessor = SourceCodeAccessModule.GetAccessor();
diff --git a/Engine/Source/Developer/TraceInsights/Private/Insights/MemoryProfiler/Widgets/SMemAllocTableTreeView.h b/Engine/Source/Developer/TraceInsights/Private/Insights/MemoryProfiler/Widgets/SMemAllocTableTreeView.h
--- a/Engine/Source/Developer/TraceInsights/Private/Insights/MemoryProfiler/Widgets/SMemAllocTableTreeView.h
+++ b/Engine/Source/Developer/TraceInsights/Private/Insights/MemoryProfiler/Widgets/SMemAllocTableTreeView.h
@@ -13,6 +13,7 @@
 namespace Insights
 {
 
+class FCallstackFrameGroupNode;
 class FMemAllocNode;
 class FMemoryRuleSpec;
 
@@ -95,9 +96,10 @@ protected:
 	virtual void InternalCreateGroupings() override;
 
 	TSharedPtr<FMemAllocNode> GetSingleSelectedMemAllocNode() const;
+	TSharedPtr<FCallstackFrameGroupNode> GetSingleSelectedCallstackFrameGroupNode() const;
 	uint32 CountSourceFiles(FMemAllocNode& MemAllocNode);
 
-	virtual void ExtendMenu(FMenuBuilder& MenuBuilder) override;
+	virtual void ExtendMenu(TSharedRef<FExtender> Extender) override;
 	bool BuildOpenSourceSubMenuItems(FMenuBuilder& MenuBuilder, const TraceServices::FCallstack& Callstack);
 	void BuildOpenSourceSubMenu(FMenuBuilder& MenuBuilder, bool bIsAllocCallstack);
 	bool CanOpenCallstackFrameSourceFileInIDE() const;
@@ -108,6 +110,10 @@ protected:
 	bool IsExportMemorySnapshotAvailable() const;
 
 private:
+	void ExtendMenuAllocation(FMenuBuilder& MenuBuilder);
+	void ExtendMenuCallstackFrame(FMenuBuilder& MenuBuilder);
+	void ExtendMenuExportSnapshot(FMenuBuilder& MenuBuilder);
+
 	void OnQueryInvalidated();
 	void StartQuery();
 	void UpdateQuery(TraceServices::IAllocationsProvider::EQueryStatus& OutStatus);
diff --git a/Engine/Source/Developer/TraceInsights/Private/Insights/Table/Widgets/STableTreeView.cpp b/Engine/Source/Developer/TraceInsights/Private/Insights/Table/Widgets/STableTreeView.cpp
--- a/Engine/Source/Developer/TraceInsights/Private/Insights/Table/Widgets/STableTreeView.cpp
+++ b/Engine/Source/Developer/TraceInsights/Private/Insights/Table/Widgets/STableTreeView.cpp
@@ -494,6 +494,13 @@ TSharedPtr<SWidget> STableTreeView::TreeView_GetMenuContent()
 	const bool bShouldCloseWindowAfterMenuSelection = true;
 	FMenuBuilder MenuBuilder(bShouldCloseWindowAfterMenuSelection, CommandList.ToSharedRef());
 
+	MenuBuilder.SetExtendersEnabled(true);
+
+	TSharedRef<FExtender> Extender = MakeShared<FExtender>();
+	MenuBuilder.PushExtender(Extender);
+
+	ExtendMenu(Extender);
+
 	// Selection menu
 	MenuBuilder.BeginSection("Selection", LOCTEXT("ContextMenu_Section_Selection", "Selection"));
 	{
@@ -551,8 +558,6 @@ TSharedPtr<SWidget> STableTreeView::TreeView_GetMenuContent()
 
 	MenuBuilder.EndSection();
 
-	ExtendMenu(MenuBuilder);
-
 	MenuBuilder.BeginSection("Misc", LOCTEXT("ContextMenu_Section_Misc", "Miscellaneous"));
 	{
 		MenuBuilder.AddMenuEntry
@@ -593,6 +598,10 @@ TSharedPtr<SWidget> STableTreeView::TreeView_GetMenuContent()
 	}
 	MenuBuilder.EndSection();
 
+	ExtendMenu(MenuBuilder);
+
+	MenuBuilder.PopExtender();
+
 	return MenuBuilder.MakeWidget();
 }
 
diff --git a/Engine/Source/Developer/TraceInsights/Private/Insights/Table/Widgets/STableTreeView.h b/Engine/Source/Developer/TraceInsights/Private/Insights/Table/Widgets/STableTreeView.h
--- a/Engine/Source/Developer/TraceInsights/Private/Insights/Table/Widgets/STableTreeView.h
+++ b/Engine/Source/Developer/TraceInsights/Private/Insights/Table/Widgets/STableTreeView.h
@@ -450,6 +450,7 @@ protected:
 	void CountNumNodesPerDepthRec(FBaseTreeNode* InRoot, TArray<int32>& InOutNumNodesPerDepth, int32 InDepth, int32 InMaxDepth, int InMaxNodes) const;
 	void SetExpandValueForChildGroupsRec(FBaseTreeNode* InRoot, int32 InDepth, int32 InMaxDepth, bool InValue);
 
+	virtual void ExtendMenu(TSharedRef<FExtender> Extender) {}
 	virtual void ExtendMenu(FMenuBuilder& Menu) {}
 
 	typedef TFunctionRef<void(TArray<FBaseTreeNodePtr>& InNodes)> WriteToFileCallback;
diff --git a/Engine/Source/Developer/TraceInsights/Private/Insights/TaskGraphProfiler/Widgets/STaskTableTreeView.cpp b/Engine/Source/Developer/TraceInsights/Private/Insights/TaskGraphProfiler/Widgets/STaskTableTreeView.cpp
--- a/Engine/Source/Developer/TraceInsights/Private/Insights/TaskGraphProfiler/Widgets/STaskTableTreeView.cpp
+++ b/Engine/Source/Developer/TraceInsights/Private/Insights/TaskGraphProfiler/Widgets/STaskTableTreeView.cpp
@@ -59,8 +59,17 @@ public:
 	UE_DISABLE_OPTIMIZATION_SHIP
 	virtual void RegisterCommands() override
 	{
-		UI_COMMAND(Command_GoToTask, "Go To Task", "Pan and zoom to the task in Timing View.", EUserInterfaceActionType::Button, FInputChord());
-		UI_COMMAND(Command_OpenInIDE, "Open in IDE", "Open the source location where the selected task was launched in IDE.", EUserInterfaceActionType::Button, FInputChord());
+		UI_COMMAND(Command_GoToTask,
+			"Go To Task",
+			"Pans and zooms to the task in the Timing View.",
+			EUserInterfaceActionType::Button,
+			FInputChord());
+
+		UI_COMMAND(Command_OpenInIDE,
+			"Open in IDE",
+			"Opens the source location where the selected task was launched, in IDE.",
+			EUserInterfaceActionType::Button,
+			FInputChord());
 	}
 	UE_ENABLE_OPTIMIZATION_SHIP
 
@@ -97,35 +106,72 @@ void STaskTableTreeView::Construct(const FArguments& InArgs, TSharedPtr<Insights
 
 void STaskTableTreeView::ExtendMenu(FMenuBuilder& MenuBuilder)
 {
-	ISourceCodeAccessModule& SourceCodeAccessModule = FModuleManager::LoadModuleChecked<ISourceCodeAccessModule>("SourceCodeAccess");
-	ISourceCodeAccessor& SourceCodeAccessor = SourceCodeAccessModule.GetAccessor();
-
-	MenuBuilder.BeginSection("Node", LOCTEXT("ContextMenu_Section_Task", "Task"));
+	MenuBuilder.BeginSection("Task", LOCTEXT("ContextMenu_Section_Task", "Task"));
 
 	{
-		MenuBuilder.AddMenuEntry
-		(
+		MenuBuilder.AddMenuEntry(
 			FTaskTableTreeViewCommands::Get().Command_GoToTask,
 			NAME_None,
 			TAttribute<FText>(),
 			TAttribute<FText>(),
-			FSlateIcon(FInsightsStyle::GetStyleSetName(), "Icons.GoToTask")
-		);
+			FSlateIcon(FInsightsStyle::GetStyleSetName(), "Icons.GoToTask"));
 	}
+
 	{
-		MenuBuilder.AddMenuEntry
-		(
+		FString File;
+		uint32 Line = 0;
+		bool bIsValidSource = GetSourceFileAndLineForSelectedTask(File, Line);
+
+		ISourceCodeAccessModule& SourceCodeAccessModule = FModuleManager::LoadModuleChecked<ISourceCodeAccessModule>("SourceCodeAccess");
+		ISourceCodeAccessor& SourceCodeAccessor = SourceCodeAccessModule.GetAccessor();
+
+		FText ItemLabel;
+		FText ItemToolTip;
+
+		if (SourceCodeAccessor.CanAccessSourceCode())
+		{
+			ItemLabel = FText::Format(LOCTEXT("ContextMenu_OpenSource", "Open Source in {0}"), SourceCodeAccessor.GetNameText());
+			if (bIsValidSource)
+			{
+				ItemToolTip = FText::Format(LOCTEXT("ContextMenu_OpenSource_Desc1", "Opens the source location where the selected task was launched, in {0}.\n{1} ({2})"),
+					SourceCodeAccessor.GetNameText(),
+					FText::FromString(File),
+					FText::AsNumber(Line, &FNumberFormattingOptions::DefaultNoGrouping()));
+			}
+			else
+			{
+				ItemToolTip = FText::Format(LOCTEXT("ContextMenu_OpenSource_Desc2", "Opens the source location where the selected task was launched, in {0}."),
+					SourceCodeAccessor.GetNameText());
+			}
+		}
+		else
+		{
+			ItemLabel = LOCTEXT("ContextMenu_OpenSourceNA", "Open Source");
+			if (bIsValidSource)
+			{
+				ItemToolTip = FText::Format(LOCTEXT("ContextMenu_OpenSourceNA_Desc1", "{1} ({2})\nSource Code Accessor is not available."),
+					FText::FromString(File),
+					FText::AsNumber(Line, &FNumberFormattingOptions::DefaultNoGrouping()));
+			}
+			else
+			{
+				ItemToolTip = LOCTEXT("ContextMenu_OpenSourceNA_Desc2", "Source Code Accessor is not available.");
+			}
+		}
+
+		MenuBuilder.AddMenuEntry(
 			FTaskTableTreeViewCommands::Get().Command_OpenInIDE,
 			NAME_None,
-			TAttribute<FText>(),
-			TAttribute<FText>(),
-			FSlateIcon(SourceCodeAccessor.GetStyleSet(), SourceCodeAccessor.GetOpenIconName())
-		);
+			ItemLabel,
+			ItemToolTip,
+			FSlateIcon(SourceCodeAccessor.GetStyleSet(), SourceCodeAccessor.GetOpenIconName()));
 	}
 
 	MenuBuilder.EndSection();
 }
 
+////////////////////////////////////////////////////////////////////////////////////////////////////
+
 void STaskTableTreeView::AddCommmands()
 {
 	FTaskTableTreeViewCommands::Register();
@@ -733,62 +779,88 @@ void STaskTableTreeView::ContextMenu_GoToTask_Execute()
 
 bool STaskTableTreeView::ContextMenu_OpenInIDE_CanExecute() const
 {
-	TArray<FTableTreeNodePtr> SelectedItems;
-	TreeView->GetSelectedItems(SelectedItems);
-
-	if (SelectedItems.Num() != 1)
-	{
-		return false;
-	}
-
-	FTaskNodePtr SelectedTask = StaticCastSharedPtr<FTaskNode>(SelectedItems[0]);
-	const FTaskEntry* TaskEntry = SelectedTask.IsValid() ? SelectedTask->GetTask() : nullptr;
+	ISourceCodeAccessModule& SourceCodeAccessModule = FModuleManager::LoadModuleChecked<ISourceCodeAccessModule>("SourceCodeAccess");
+	ISourceCodeAccessor& SourceCodeAccessor = SourceCodeAccessModule.GetAccessor();
 
-	if (!SelectedTask.IsValid() || SelectedTask->IsGroup() || TaskEntry == nullptr || TaskEntry->DebugName == nullptr || TaskEntry->DebugName[FCString::Strlen(TaskEntry->DebugName) - 1] != TEXT(')'))
+	if (!SourceCodeAccessor.CanAccessSourceCode())
 	{
 		return false;
 	}
 
-	return true;
+	FString File;
+	uint32 Line = 0;
+	return GetSourceFileAndLineForSelectedTask(File, Line);
 }
 
 ////////////////////////////////////////////////////////////////////////////////////////////////////
 
-void STaskTableTreeView::ContextMenu_OpenInIDE_Execute()
+bool STaskTableTreeView::GetSourceFileAndLineForSelectedTask(FString& OutFile, uint32& OutLine) const
 {
 	TArray<FTableTreeNodePtr> SelectedItems;
 	TreeView->GetSelectedItems(SelectedItems);
 
 	if (SelectedItems.Num() != 1)
 	{
-		return;
+		return false;
 	}
 
-	FTaskNodePtr SelectedTask = StaticCastSharedPtr<FTaskNode>(SelectedItems[0]);
-	const FTaskEntry* TaskEntry = SelectedTask.IsValid() ? SelectedTask->GetTask() : nullptr;
+	FTableTreeNodePtr SelectedTreeNode = SelectedItems[0];
+	if (!SelectedTreeNode.IsValid() || !SelectedTreeNode->Is<FTaskNode>())
+	{
+		return false;
+	}
 
+	FTaskNodePtr SelectedTask = StaticCastSharedPtr<FTaskNode>(SelectedTreeNode);
+
+	const FTaskEntry* TaskEntry = SelectedTask->GetTask();
 	if (TaskEntry == nullptr || TaskEntry->DebugName == nullptr)
 	{
-		return;
+		return false;
 	}
 
 	const FString DebugName = TaskEntry->DebugName;
+	if (DebugName.Len() < 4 ||
+		DebugName[DebugName.Len() - 1] != TEXT(')'))
+	{
+		return false;
+	}
+
 	int32 OpenBracketPos;
-	if (DebugName[DebugName.Len() - 1] != TEXT(')') || !DebugName.FindChar(TEXT('('), OpenBracketPos) || OpenBracketPos > DebugName.Len() - 3)
+	if (!DebugName.FindChar(TEXT('('), OpenBracketPos) ||
+		OpenBracketPos > DebugName.Len() - 3)
 	{
-		return;
+		return false;
 	}
-	FString Filename = DebugName.Left(OpenBracketPos);
-	if (!FPaths::FileExists(Filename))
+
+	OutFile = DebugName.Left(OpenBracketPos);
+
+	FString LineStr = DebugName.Mid(OpenBracketPos + 1, DebugName.Len() - OpenBracketPos - 2);
+	OutLine = FCString::Atoi(*LineStr);
+
+	return true;
+}
+
+////////////////////////////////////////////////////////////////////////////////////////////////////
+
+void STaskTableTreeView::ContextMenu_OpenInIDE_Execute()
+{
+	FString File;
+	uint32 Line = 0;
+	if (!GetSourceFileAndLineForSelectedTask(File, Line))
 	{
 		return;
 	}
-	FString LineStr = DebugName.Mid(OpenBracketPos + 1, DebugName.Len() - OpenBracketPos - 2);
-	int32 Line = FCString::Atoi(*LineStr);
 
 	ISourceCodeAccessModule& SourceCodeAccessModule = FModuleManager::LoadModuleChecked<ISourceCodeAccessModule>("SourceCodeAccess");
-	ISourceCodeAccessor& SourceCodeAccessor = SourceCodeAccessModule.GetAccessor();
-	SourceCodeAccessor.OpenFileAtLine(Filename, Line);
+	if (FPaths::FileExists(File))
+	{
+		ISourceCodeAccessor& SourceCodeAccessor = SourceCodeAccessModule.GetAccessor();
+		SourceCodeAccessor.OpenFileAtLine(File, Line);
+	}
+	else
+	{
+		SourceCodeAccessModule.OnOpenFileFailed().Broadcast(File);
+	}
 }
 
 ////////////////////////////////////////////////////////////////////////////////////////////////////
@@ -817,7 +889,7 @@ void STaskTableTreeView::SearchForItem(TSharedPtr<FTableTaskCancellationToken> C
 {
 	TSharedPtr<FTaskTable> TaskTable = GetTaskTable();
 	TArray<FTaskEntry>& Tasks = TaskTable->GetTaskEntries();
-	
+
 	uint32 NumEntries = (uint32)Tasks.Num();
 	for (uint32 Index = 0; Index < NumEntries; ++Index)
 	{
diff --git a/Engine/Source/Developer/TraceInsights/Private/Insights/TaskGraphProfiler/Widgets/STaskTableTreeView.h b/Engine/Source/Developer/TraceInsights/Private/Insights/TaskGraphProfiler/Widgets/STaskTableTreeView.h
--- a/Engine/Source/Developer/TraceInsights/Private/Insights/TaskGraphProfiler/Widgets/STaskTableTreeView.h
+++ b/Engine/Source/Developer/TraceInsights/Private/Insights/TaskGraphProfiler/Widgets/STaskTableTreeView.h
@@ -120,6 +120,8 @@ private:
 	bool ContextMenu_OpenInIDE_CanExecute() const;
 	void ContextMenu_OpenInIDE_Execute();
 
+	bool GetSourceFileAndLineForSelectedTask(FString& OutFile, uint32& OutLine) const;
+
 private:
 	double QueryStartTime = 0.0f;
 	double QueryEndTime = 0.0f;
diff --git a/Engine/Source/Developer/TraceInsights/Private/Insights/Widgets/SLogView.cpp b/Engine/Source/Developer/TraceInsights/Private/Insights/Widgets/SLogView.cpp
--- a/Engine/Source/Developer/TraceInsights/Private/Insights/Widgets/SLogView.cpp
+++ b/Engine/Source/Developer/TraceInsights/Private/Insights/Widgets/SLogView.cpp
@@ -1225,28 +1225,56 @@ TSharedPtr<SWidget> SLogView::ListView_GetContextMenu()
 
 			ISourceCodeAccessModule& SourceCodeAccessModule = FModuleManager::LoadModuleChecked<ISourceCodeAccessModule>("SourceCodeAccess");
 			ISourceCodeAccessor& SourceCodeAccessor = SourceCodeAccessModule.GetAccessor();
+			if (SourceCodeAccessor.CanAccessSourceCode())
+			{
+				const FText ItemLabel = FText::Format(LOCTEXT("ContextMenu_OpenSource", "Open Source in {0}"),
+					SourceCodeAccessor.GetNameText());
 
-			const FText ItemLabel = FText::Format(LOCTEXT("ContextMenu_OpenSource", "Open Source in {0}"), SourceCodeAccessor.GetNameText());
+				FText ItemToolTip;
+				if (!File.IsEmpty())
+				{
+					ItemToolTip = FText::Format(LOCTEXT("ContextMenu_OpenSource_Desc1", "Opens the source file of the selected message in {0}.\n{1} ({2})"),
+						SourceCodeAccessor.GetNameText(),
+						FText::FromString(File),
+						FText::AsNumber(Line, &FNumberFormattingOptions::DefaultNoGrouping()));
+				}
+				else
+				{
+					ItemToolTip = FText::Format(LOCTEXT("ContextMenu_OpenSource_Desc2", "Opens the source file of the selected message in {0}."),
+						SourceCodeAccessor.GetNameText());
+				}
 
-			FText ItemToolTip;
-			if (!File.IsEmpty())
-			{
-				ItemToolTip = FText::Format(LOCTEXT("ContextMenu_OpenSource_Desc1", "Opens the source file of the selected message in {0}.\n{1} ({2})"),
-					SourceCodeAccessor.GetNameText(), FText::FromString(File), FText::AsNumber(Line, &FNumberFormattingOptions::DefaultNoGrouping()));
+				MenuBuilder.AddMenuEntry(
+					FLogViewCommands::Get().Command_OpenSource,
+					NAME_None,
+					ItemLabel,
+					ItemToolTip,
+					FSlateIcon(SourceCodeAccessor.GetStyleSet(), SourceCodeAccessor.GetOpenIconName()));
 			}
 			else
 			{
-				ItemToolTip = FText::Format(LOCTEXT("ContextMenu_OpenSource_Desc2", "Opens the source file of the selected message in {0}."),
-					SourceCodeAccessor.GetNameText());
-			}
+				const FText ItemLabel = LOCTEXT("ContextMenu_OpenSource_NoAccessor", "Open Source");
 
-			MenuBuilder.AddMenuEntry(
-				FLogViewCommands::Get().Command_OpenSource,
-				NAME_None,
-				ItemLabel,
-				ItemToolTip,
-				FSlateIcon(SourceCodeAccessor.GetStyleSet(), SourceCodeAccessor.GetOpenIconName())
-			);
+				FText ItemToolTip;
+				if (!File.IsEmpty())
+				{
+					ItemToolTip = FText::Format(LOCTEXT("ContextMenu_OpenSource_NoAccessor_Desc1", "{0} ({1})\nSource Code Accessor is not available."),
+						FText::FromString(File),
+						FText::AsNumber(Line, &FNumberFormattingOptions::DefaultNoGrouping()));
+				}
+				else
+				{
+					ItemToolTip = LOCTEXT("ContextMenu_OpenSource_NoAccessor_Desc2", "Source Code Accessor is not available.");
+				}
+
+				MenuBuilder.AddMenuEntry(
+					ItemLabel,
+					ItemToolTip,
+					FSlateIcon(SourceCodeAccessor.GetStyleSet(), SourceCodeAccessor.GetOpenIconName()),
+					FUIAction(FExecuteAction(), FCanExecuteAction::CreateLambda([]() { return false; })),
+					NAME_None,
+					EUserInterfaceActionType::None);
+			}
 		}
 	}
 	MenuBuilder.EndSection();
diff --git a/Engine/Source/Developer/TraceInsights/Private/Insights/Widgets/STimersView.cpp b/Engine/Source/Developer/TraceInsights/Private/Insights/Widgets/STimersView.cpp
--- a/Engine/Source/Developer/TraceInsights/Private/Insights/Widgets/STimersView.cpp
+++ b/Engine/Source/Developer/TraceInsights/Private/Insights/Widgets/STimersView.cpp
@@ -610,18 +610,35 @@ TSharedPtr<SWidget> STimersView::TreeView_GetMenuContent()
 				bIsValidSource = SelectedNode->GetSourceFileAndLine(File, Line);
 			}
 
-			FText ItemLabel = FText::Format(LOCTEXT("ContextMenu_OpenSource", "Open Source in {0}"), SourceCodeAccessor.GetNameText());
-
+			FText ItemLabel;
 			FText ItemToolTip;
-			if (bIsValidSource)
+
+			if (SourceCodeAccessor.CanAccessSourceCode())
 			{
-				ItemToolTip = FText::Format(LOCTEXT("ContextMenu_OpenSource_Desc1", "Opens the source file of the selected timer in {0}.\n{1} ({2})"),
-					SourceCodeAccessor.GetNameText(), FText::FromString(File), FText::AsNumber(Line, &FNumberFormattingOptions::DefaultNoGrouping()));
+				ItemLabel = FText::Format(LOCTEXT("ContextMenu_OpenSource", "Open Source in {0}"), SourceCodeAccessor.GetNameText());
+				if (bIsValidSource)
+				{
+					ItemToolTip = FText::Format(LOCTEXT("ContextMenu_OpenSource_Desc1", "Opens the source file of the selected timer in {0}.\n{1} ({2})"),
+						SourceCodeAccessor.GetNameText(), FText::FromString(File), FText::AsNumber(Line, &FNumberFormattingOptions::DefaultNoGrouping()));
+				}
+				else
+				{
+					ItemToolTip = FText::Format(LOCTEXT("ContextMenu_OpenSource_Desc2", "Opens the source file of the selected timer in {0}."),
+						SourceCodeAccessor.GetNameText());
+				}
 			}
 			else
 			{
-				ItemToolTip = FText::Format(LOCTEXT("ContextMenu_OpenSource_Desc2", "Opens the source file of the selected timer in {0}."),
-					SourceCodeAccessor.GetNameText());
+				ItemLabel = LOCTEXT("ContextMenu_OpenSourceNA", "Open Source");
+				if (bIsValidSource)
+				{
+					ItemToolTip = FText::Format(LOCTEXT("ContextMenu_OpenSourceNA_Desc1", "{1} ({2})\nSource Code Accessor is not available."),
+						FText::FromString(File), FText::AsNumber(Line, &FNumberFormattingOptions::DefaultNoGrouping()));
+				}
+				else
+				{
+					ItemToolTip = LOCTEXT("ContextMenu_OpenSourceNA_Desc2", "Source Code Accessor is not available.");
+				}
 			}
 
 			MenuBuilder.AddMenuEntry(
@@ -629,8 +646,7 @@ TSharedPtr<SWidget> STimersView::TreeView_GetMenuContent()
 				NAME_None,
 				ItemLabel,
 				ItemToolTip,
-				FSlateIcon(SourceCodeAccessor.GetStyleSet(), SourceCodeAccessor.GetOpenIconName())
-			);
+				FSlateIcon(SourceCodeAccessor.GetStyleSet(), SourceCodeAccessor.GetOpenIconName()));
 		}
 	}
 	MenuBuilder.EndSection();
@@ -1484,16 +1500,33 @@ void STimersView::TreeView_OnSelectionChanged(FTimerNodePtr SelectedItem, ESelec
 {
 	if (SelectInfo != ESelectInfo::Direct)
 	{
-		TArray<FTimerNodePtr> SelectedItems = TreeView->GetSelectedItems();
-		if (SelectedItems.Num() == 1 && SelectedItems[0]->GetType() != ETimerNodeType::Group)
+		FTimerNodePtr SelectedNode = GetSingleSelectedTimerNode();
+		if (SelectedNode.IsValid() && SelectedNode->GetType() != ETimerNodeType::Group)
 		{
-			FTimingProfilerManager::Get()->SetSelectedTimer(SelectedItems[0]->GetTimerId());
+			FTimingProfilerManager::Get()->SetSelectedTimer(SelectedNode->GetTimerId());
 		}
 	}
 }
 
 ////////////////////////////////////////////////////////////////////////////////////////////////////
 
+FTimerNodePtr STimersView::GetSingleSelectedTimerNode() const
+{
+	if (TreeView->GetNumItemsSelected() != 1)
+	{
+		return nullptr;
+	}
+	const TArray<FTimerNodePtr> SelectedNodes = TreeView->GetSelectedItems();
+	const int32 NumSelectedNodes = SelectedNodes.Num();
+	if (NumSelectedNodes == 1)
+	{
+		return SelectedNodes[0];
+	}
+	return nullptr;
+}
+
+////////////////////////////////////////////////////////////////////////////////////////////////////
+
 void STimersView::TreeView_OnGetChildren(FTimerNodePtr InParent, TArray<FTimerNodePtr>& OutChildren)
 {
 	const TArray<Insights::FBaseTreeNodePtr>& Children = InParent->GetFilteredChildren();
@@ -2811,15 +2844,16 @@ void STimersView::ContextMenu_CopyToClipboard_Execute()
 		return;
 	}
 
-	TArray<Insights::FBaseTreeNodePtr> SelectedNodes;
-	for (FTimerNodePtr TimerPtr : TreeView->GetSelectedItems())
+	const TArray<FTimerNodePtr> SelectedTimerNodes = TreeView->GetSelectedItems();
+	if (SelectedTimerNodes.Num() == 0)
 	{
-		SelectedNodes.Add(TimerPtr);
+		return;
 	}
 
-	if (SelectedNodes.Num() == 0)
+	TArray<Insights::FBaseTreeNodePtr> SelectedNodes;
+	for (FTimerNodePtr TimerPtr : SelectedTimerNodes)
 	{
-		return;
+		SelectedNodes.Add(TimerPtr);
 	}
 
 	FString ClipboardText;
@@ -2854,15 +2888,16 @@ void STimersView::ContextMenu_Export_Execute()
 		return;
 	}
 
-	TArray<Insights::FBaseTreeNodePtr> SelectedNodes;
-	for (FTimerNodePtr TimerPtr : TreeView->GetSelectedItems())
+	const TArray<FTimerNodePtr> SelectedTimerNodes = TreeView->GetSelectedItems();
+	if (SelectedTimerNodes.Num() == 0)
 	{
-		SelectedNodes.Add(TimerPtr);
+		return;
 	}
 
-	if (SelectedNodes.Num() == 0)
+	TArray<Insights::FBaseTreeNodePtr> SelectedNodes;
+	for (FTimerNodePtr TimerPtr : SelectedTimerNodes)
 	{
-		return;
+		SelectedNodes.Add(TimerPtr);
 	}
 
 	const FString DialogTitle = LOCTEXT("Export_Title", "Export Aggregated Timer Stats").ToString();
@@ -3260,30 +3295,33 @@ IFileHandle* STimersView::OpenExportFile(const TCHAR* InFilename) const
 
 bool STimersView::ContextMenu_OpenSource_CanExecute() const
 {
-	const TArray<FTimerNodePtr> SelectedNodes = TreeView->GetSelectedItems();
-	const int32 NumSelectedNodes = SelectedNodes.Num();
-	if (NumSelectedNodes == 1)
+	ISourceCodeAccessModule& SourceCodeAccessModule = FModuleManager::LoadModuleChecked<ISourceCodeAccessModule>("SourceCodeAccess");
+	ISourceCodeAccessor& SourceCodeAccessor = SourceCodeAccessModule.GetAccessor();
+
+	if (!SourceCodeAccessor.CanAccessSourceCode())
 	{
-		FTimerNodePtr SelectedNode = SelectedNodes[0];
-		if (SelectedNode.IsValid() && SelectedNode->GetType() != ETimerNodeType::Group)
-		{
-			FString File;
-			uint32 Line = 0;
-			return SelectedNode->GetSourceFileAndLine(File, Line);
-		}
+		return false;
 	}
-	return false;
+
+	FTimerNodePtr SelectedNode = GetSingleSelectedTimerNode();
+	if (!SelectedNode.IsValid())
+	{
+		return false;
+	}
+
+	FString File;
+	uint32 Line = 0;
+	return SelectedNode->GetSourceFileAndLine(File, Line);
 }
 
 ////////////////////////////////////////////////////////////////////////////////////////////////////
 
 void STimersView::ContextMenu_OpenSource_Execute() const
 {
-	const TArray<FTimerNodePtr> SelectedNodes = TreeView->GetSelectedItems();
-	const int32 NumSelectedNodes = SelectedNodes.Num();
-	if (NumSelectedNodes == 1)
+	FTimerNodePtr SelectedNode = GetSingleSelectedTimerNode();
+	if (SelectedNode.IsValid())
 	{
-		OpenSourceFileInIDE(SelectedNodes[0]);
+		OpenSourceFileInIDE(SelectedNode);
 	}
 }
 
@@ -3291,25 +3329,22 @@ void STimersView::ContextMenu_OpenSource_Execute() const
 
 bool STimersView::ContextMenu_FindInstance_CanExecute() const
 {
-	const TArray<FTimerNodePtr> SelectedNodes = TreeView->GetSelectedItems();
-	return SelectedNodes.Num() == 1 && !(SelectedNodes[0]->GetType() == ETimerNodeType::Group);
+	FTimerNodePtr SelectedNode = GetSingleSelectedTimerNode();
+	return SelectedNode.IsValid() && SelectedNode->GetType() != ETimerNodeType::Group;
 }
 
 ////////////////////////////////////////////////////////////////////////////////////////////////////
 
 void STimersView::ContextMenu_FindInstance_Execute(bool bFindMax) const
 {
-	const TArray<FTimerNodePtr> SelectedNodes = TreeView->GetSelectedItems();
-	if (SelectedNodes.Num() != 1)
+	FTimerNodePtr SelectedNode = GetSingleSelectedTimerNode();
+	if (!SelectedNode.IsValid())
 	{
 		return;
 	}
 
-	FTimerNodePtr SelectedNode = SelectedNodes[0];
-
 	TSharedPtr<STimingProfilerWindow> Wnd = FTimingProfilerManager::Get()->GetProfilerWindow();
 	TSharedPtr<STimingView> TimingView = Wnd.IsValid() ? Wnd->GetTimingView() : nullptr;
-
 	if (!TimingView.IsValid())
 	{
 		return;
@@ -3323,16 +3358,14 @@ void STimersView::ContextMenu_FindInstance_Execute(bool bFindMax) const
 
 bool STimersView::ContextMenu_FindInstanceInSelection_CanExecute() const
 {
-	const TArray<FTimerNodePtr> SelectedNodes = TreeView->GetSelectedItems();
-
-	if (SelectedNodes.Num() != 1 || SelectedNodes[0]->GetType() == ETimerNodeType::Group)
+	FTimerNodePtr SelectedNode = GetSingleSelectedTimerNode();
+	if (!SelectedNode.IsValid() || SelectedNode->GetType() == ETimerNodeType::Group)
 	{
 		return false;
 	}
 
 	TSharedPtr<STimingProfilerWindow> Wnd = FTimingProfilerManager::Get()->GetProfilerWindow();
 	TSharedPtr<STimingView> TimingView = Wnd.IsValid() ? Wnd->GetTimingView() : nullptr;
-
 	if (TimingView.IsValid())
 	{
 		return TimingView->GetSelectionEndTime() > TimingView->GetSelectionStartTime();
@@ -3345,17 +3378,14 @@ bool STimersView::ContextMenu_FindInstanceInSelection_CanExecute() const
 
 void STimersView::ContextMenu_FindInstanceInSelection_Execute(bool bFindMax) const
 {
-	const TArray<FTimerNodePtr> SelectedNodes = TreeView->GetSelectedItems();
-	if (SelectedNodes.Num() != 1)
+	FTimerNodePtr SelectedNode = GetSingleSelectedTimerNode();
+	if (!SelectedNode.IsValid() || SelectedNode->GetType() == ETimerNodeType::Group)
 	{
 		return;
 	}
 
-	FTimerNodePtr SelectedNode = SelectedNodes[0];
-
 	TSharedPtr<STimingProfilerWindow> Wnd = FTimingProfilerManager::Get()->GetProfilerWindow();
 	TSharedPtr<STimingView> TimingView = Wnd.IsValid() ? Wnd->GetTimingView() : nullptr;
-
 	if (!TimingView.IsValid())
 	{
 		return;
@@ -3369,24 +3399,27 @@ void STimersView::ContextMenu_FindInstanceInSelection_Execute(bool bFindMax) con
 
 void STimersView::OpenSourceFileInIDE(FTimerNodePtr InNode) const
 {
-	if (InNode.IsValid() && InNode->GetType() != ETimerNodeType::Group)
+	if (!InNode.IsValid() || InNode->GetType() == ETimerNodeType::Group)
 	{
-		FString File;
-		uint32 Line = 0;
-		bool bIsValidSource = InNode->GetSourceFileAndLine(File, Line);
-		if (bIsValidSource)
-		{
-			ISourceCodeAccessModule& SourceCodeAccessModule = FModuleManager::LoadModuleChecked<ISourceCodeAccessModule>("SourceCodeAccess");
-			if (FPaths::FileExists(File))
-			{
-				ISourceCodeAccessor& SourceCodeAccessor = SourceCodeAccessModule.GetAccessor();
-				SourceCodeAccessor.OpenFileAtLine(File, Line);
-			}
-			else
-			{
-				SourceCodeAccessModule.OnOpenFileFailed().Broadcast(File);
-			}
-		}
+		return;
+	}
+
+	FString File;
+	uint32 Line = 0;
+	if (!InNode->GetSourceFileAndLine(File, Line))
+	{
+		return;
+	}
+
+	ISourceCodeAccessModule& SourceCodeAccessModule = FModuleManager::LoadModuleChecked<ISourceCodeAccessModule>("SourceCodeAccess");
+	if (FPaths::FileExists(File))
+	{
+		ISourceCodeAccessor& SourceCodeAccessor = SourceCodeAccessModule.GetAccessor();
+		SourceCodeAccessor.OpenFileAtLine(File, Line);
+	}
+	else
+	{
+		SourceCodeAccessModule.OnOpenFileFailed().Broadcast(File);
 	}
 }
 
diff --git a/Engine/Source/Developer/TraceInsights/Private/Insights/Widgets/STimersView.h b/Engine/Source/Developer/TraceInsights/Private/Insights/Widgets/STimersView.h
--- a/Engine/Source/Developer/TraceInsights/Private/Insights/Widgets/STimersView.h
+++ b/Engine/Source/Developer/TraceInsights/Private/Insights/Widgets/STimersView.h
@@ -182,6 +182,8 @@ private:
 	/** Called by STreeView when selection has changed. */
 	void TreeView_OnSelectionChanged(FTimerNodePtr SelectedItem, ESelectInfo::Type SelectInfo);
 
+	FTimerNodePtr GetSingleSelectedTimerNode() const;
+
 	/** Called by STreeView when a tree item is double clicked. */
 	void TreeView_OnMouseButtonDoubleClick(FTimerNodePtr TreeNode);
 
