Fix dialog text

Fix dialog text when comfirming removal of watch directory.

diff --git a/Engine/Source/Developer/TraceInsights/Private/Insights/Widgets/SStartPageWindow.cpp b/Engine/Source/Developer/TraceInsights/Private/Insights/Widgets/SStartPageWindow.cpp
--- a/Engine/Source/Developer/TraceInsights/Private/Insights/Widgets/SStartPageWindow.cpp
+++ b/Engine/Source/Developer/TraceInsights/Private/Insights/Widgets/SStartPageWindow.cpp
@@ -1592,8 +1592,8 @@ FReply STraceDirectoryItem::OnDelete()
 	if (Model)
 	{
 		const EAppReturnType::Type bConfirmed = FMessageDialog::Open(EAppMsgType::OkCancel, FText::Format(
-			LOCTEXT("WatchDirRemoveConfirmBody", "This will remove {0} from watch directories?"),
-			LOCTEXT("WatchDirRemoveConfirmTitle", "Confirm removing watch directory")));
+			LOCTEXT("WatchDirRemoveConfirmBody", "This will remove '{0}' from monitored directories?"), FText::FromString(Model->Path)),
+			LOCTEXT("WatchDirRemoveConfirmTitle", "Confirm removing monitored directory"));
 		if (bConfirmed == EAppReturnType::Ok)
 		{
 			const TCHAR* Path = *Model->Path;
