[Insights] Fixed crash when "TimingInsights.ExportTimerStatistics" execute command is specified with an invalid filename. An error message is logged instead.

diff --git a/Engine/Source/Developer/TraceInsights/Private/Insights/TimingProfilerManager.cpp b/Engine/Source/Developer/TraceInsights/Private/Insights/TimingProfilerManager.cpp
--- a/Engine/Source/Developer/TraceInsights/Private/Insights/TimingProfilerManager.cpp
+++ b/Engine/Source/Developer/TraceInsights/Private/Insights/TimingProfilerManager.cpp
@@ -661,7 +661,12 @@ bool FTimingProfilerManager::Exec(const TCHAR* Cmd, FOutputDevice& Ar)
 		check(FInsightsManager::Get().IsValid() && FInsightsManager::Get()->GetSession().IsValid());
 		Insights::FTimingExporter Exporter(*FInsightsManager::Get()->GetSession().Get());
 		Insights::FTimingExporter::FExportThreadsParams Params; // default
-		Exporter.ExportThreadsAsText(Cmd, Params);
+
+		const bool bUseEscape = true;
+		FString Filename = FParse::Token(Cmd, bUseEscape);
+		Ar.Logf(TEXT("  Filename: %s"), *Filename);
+
+		Exporter.ExportThreadsAsText(Filename, Params);
 		return true;
 	}
 
@@ -671,7 +676,12 @@ bool FTimingProfilerManager::Exec(const TCHAR* Cmd, FOutputDevice& Ar)
 		check(FInsightsManager::Get().IsValid() && FInsightsManager::Get()->GetSession().IsValid());
 		Insights::FTimingExporter Exporter(*FInsightsManager::Get()->GetSession().Get());
 		Insights::FTimingExporter::FExportTimersParams Params; // default
-		Exporter.ExportTimersAsText(Cmd, Params);
+
+		const bool bUseEscape = true;
+		FString Filename = FParse::Token(Cmd, bUseEscape);
+		Ar.Logf(TEXT("  Filename: %s"), *Filename);
+
+		Exporter.ExportTimersAsText(Filename, Params);
 		return true;
 	}
 
diff --git a/Engine/Source/Developer/TraceInsights/Private/Insights/ViewModels/TimingExporter.cpp b/Engine/Source/Developer/TraceInsights/Private/Insights/ViewModels/TimingExporter.cpp
--- a/Engine/Source/Developer/TraceInsights/Private/Insights/ViewModels/TimingExporter.cpp
+++ b/Engine/Source/Developer/TraceInsights/Private/Insights/ViewModels/TimingExporter.cpp
@@ -733,12 +733,20 @@ int32 FTimingExporter::ExportTimerStatisticsAsText(const FString& Filename, FExp
 	FStopwatch Stopwatch;
 	Stopwatch.Start();
 
-	TraceServices::Table2Csv(*StatsTable, *Filename);
+	bool bSuccess = TraceServices::Table2Csv(*StatsTable, *Filename);
 
 	Stopwatch.Stop();
 	const double TotalTime = Stopwatch.GetAccumulatedTime();
-	UE_LOG(TraceInsights, Log, TEXT("Exported timing statistics to file in %.3fs (\"%s\")."), TotalTime, *Filename);
-	return 1;
+
+	if (bSuccess)
+	{
+		UE_LOG(TraceInsights, Log, TEXT("Exported timing statistics to file in %.3fs (\"%s\")."), TotalTime, *Filename);
+	}
+	else
+	{
+		UE_LOG(TraceInsights, Error, TEXT("Failed to write the CSV file (\"%s\")!"), *Filename);
+	}
+	return bSuccess ? 1 : -2;
 }
 
 ////////////////////////////////////////////////////////////////////////////////////////////////////
diff --git a/Engine/Source/Developer/TraceServices/Private/Common/Utils.cpp b/Engine/Source/Developer/TraceServices/Private/Common/Utils.cpp
--- a/Engine/Source/Developer/TraceServices/Private/Common/Utils.cpp
+++ b/Engine/Source/Developer/TraceServices/Private/Common/Utils.cpp
@@ -12,10 +12,14 @@ DEFINE_LOG_CATEGORY(LogTraceServices);
 namespace TraceServices
 {
 
-void Table2Csv(const IUntypedTable& Table, const TCHAR* Filename)
+bool Table2Csv(const IUntypedTable& Table, const TCHAR* Filename)
 {
-	TSharedPtr<FArchive> OutputFile = MakeShareable(IFileManager::Get().CreateFileWriter(Filename));
-	check(OutputFile);
+	TUniquePtr<FArchive> OutputFile(IFileManager::Get().CreateFileWriter(Filename));
+	if (!OutputFile.IsValid())
+	{
+		return false;
+	}
+
 	FString Header;
 	const ITableLayout& Layout = Table.GetLayout();
 	const uint64 ColumnCount = Layout.GetColumnCount();
@@ -72,6 +76,8 @@ void Table2Csv(const IUntypedTable& Table, const TCHAR* Filename)
 		OutputFile->Serialize((void*)AnsiLine.Get(), AnsiLine.Length());
 	}
 	OutputFile->Close();
+
+	return true;
 }
 
 void StringFormat(TCHAR* Out, uint64 MaxOut, TCHAR* Temp, uint64 MaxTemp, const TCHAR* FormatString, const uint8* FormatArgs)
diff --git a/Engine/Source/Developer/TraceServices/Public/TraceServices/Containers/Tables.h b/Engine/Source/Developer/TraceServices/Public/TraceServices/Containers/Tables.h
--- a/Engine/Source/Developer/TraceServices/Public/TraceServices/Containers/Tables.h
+++ b/Engine/Source/Developer/TraceServices/Public/TraceServices/Containers/Tables.h
@@ -74,6 +74,6 @@ public:
 	virtual ITableReader<RowType>* CreateReader() const = 0;
 };
 
-TRACESERVICES_API void Table2Csv(const IUntypedTable& Table, const TCHAR* Filename);
+TRACESERVICES_API bool Table2Csv(const IUntypedTable& Table, const TCHAR* Filename);
 
 } // namespace TraceServices
