[Insights] The arrow down/up and autocomplete functionality is described in controls that support it using a hint text and tooltip.

#changelist validated
#virtualized

diff --git a/Engine/Source/Developer/TraceInsights/Private/Insights/ViewModels/Filters.cpp b/Engine/Source/Developer/TraceInsights/Private/Insights/ViewModels/Filters.cpp
--- a/Engine/Source/Developer/TraceInsights/Private/Insights/ViewModels/Filters.cpp
+++ b/Engine/Source/Developer/TraceInsights/Private/Insights/ViewModels/Filters.cpp
@@ -278,6 +278,20 @@ TSharedRef<FFilterState> FFilterState::DeepCopy() const
 
 ////////////////////////////////////////////////////////////////////////////////////////////////////
 
+FText FFilterWithSuggestionsValueConverter::GetTooltipText() const
+{
+	return LOCTEXT("FilterWithSuggestionsValueConverterTooltip", "Enter the value to search for. This field has auto-complete. Start typing or press the arrow down or arrow up key to see options.");
+}
+
+////////////////////////////////////////////////////////////////////////////////////////////////////
+
+FText FFilterWithSuggestionsValueConverter::GetHintText() const
+{
+	return LOCTEXT("FilterWithSuggestionsValueConverterHint", "Start typing or press arrow down or up to see options.");
+}
+
+////////////////////////////////////////////////////////////////////////////////////////////////////
+
 } // namespace Insights
 
 #undef LOCTEXT_NAMESPACE
\ No newline at end of file
diff --git a/Engine/Source/Developer/TraceInsights/Private/Insights/ViewModels/Filters.h b/Engine/Source/Developer/TraceInsights/Private/Insights/ViewModels/Filters.h
--- a/Engine/Source/Developer/TraceInsights/Private/Insights/ViewModels/Filters.h
+++ b/Engine/Source/Developer/TraceInsights/Private/Insights/ViewModels/Filters.h
@@ -266,6 +266,16 @@ protected:
 
 ////////////////////////////////////////////////////////////////////////////////////////////////////
 
+class FFilterWithSuggestionsValueConverter : public IFilterValueConverter
+{
+public:
+	virtual bool Convert(const FString& Input, double& Output, FText& OutError) const override { return true; }
+	virtual FText GetTooltipText() const override;
+	virtual FText GetHintText() const override;
+};
+
+////////////////////////////////////////////////////////////////////////////////////////////////////
+
 class FFilterWithSuggestions : public FFilter
 {
 	INSIGHTS_DECLARE_RTTI(FFilterWithSuggestions, FFilter)
@@ -277,6 +287,11 @@ public:
 	FFilterWithSuggestions(int32 InKey, FText InName, FText InDesc, EFilterDataType InDataType, TSharedPtr<IFilterValueConverter> InConverter, SupportedOperatorsArrayPtr InSupportedOperators)
 		: FFilter(InKey, InName, InDesc, InDataType, InConverter, InSupportedOperators)
 	{
+		if (!InConverter.IsValid())
+		{
+			// Add a default converter to add a hint text describing that this filter supports auto-complete.
+			Converter = MakeShared<FFilterWithSuggestionsValueConverter>();
+		}
 	}
 
 	virtual ~FFilterWithSuggestions()
