[Insights] Fix dialog that asks for confirmation before deleting a watch directory.

diff --git a/Engine/Source/Developer/TraceInsights/Private/Insights/Common/MessageDialogUtils.cpp b/Engine/Source/Developer/TraceInsights/Private/Insights/Common/MessageDialogUtils.cpp
new file mode 100644
--- /dev/null
+++ b/Engine/Source/Developer/TraceInsights/Private/Insights/Common/MessageDialogUtils.cpp
@@ -0,0 +1,44 @@
+// Copyright Epic Games, Inc. All Rights Reserved.
+
+#include "MessageDialogUtils.h"
+
+#include "Dialog/SCustomDialog.h"
+#include "Widgets/Text/STextBlock.h"
+
+#define LOCTEXT_NAMESPACE "MessageDialogUtils"
+
+namespace Insights
+{
+
+//////////////////////////////////////////////////////////////////////////
+
+EDialogResponse FMessageDialogUtils::ShowChoiceDialog(const FText& Title, const FText& Content)
+{
+	TSharedRef<SCustomDialog> Dialog = SNew(SCustomDialog)
+		.Title(Title)
+		.Content()
+		[
+			SNew(STextBlock).Text(Content)
+		]
+		.Buttons({
+		SCustomDialog::FButton(LOCTEXT("OK", "OK")),
+		SCustomDialog::FButton(LOCTEXT("Cancel", "Cancel"))
+		});
+
+	// returns 0 when OK is pressed, 1 when Cancel is pressed, -1 if the window is closed
+	const int ButtonPressed = Dialog->ShowModal();
+	if (ButtonPressed == 0)
+	{
+		return EDialogResponse::OK;
+	}
+
+	return EDialogResponse::Cancel;
+}
+
+//////////////////////////////////////////////////////////////////////////
+
+} // namespace Insights
+
+//////////////////////////////////////////////////////////////////////////
+
+#undef LOCTEXT_NAMESPACE
diff --git a/Engine/Source/Developer/TraceInsights/Private/Insights/Common/MessageDialogUtils.h b/Engine/Source/Developer/TraceInsights/Private/Insights/Common/MessageDialogUtils.h
new file mode 100644
--- /dev/null
+++ b/Engine/Source/Developer/TraceInsights/Private/Insights/Common/MessageDialogUtils.h
@@ -0,0 +1,28 @@
+// Copyright Epic Games, Inc. All Rights Reserved.
+
+#pragma once
+
+#include "Internationalization/Text.h"
+
+namespace Insights
+{
+
+////////////////////////////////////////////////////////////////////////////////////////////////////
+
+enum class EDialogResponse
+{
+	OK = 0,
+	Cancel = 1
+};
+
+////////////////////////////////////////////////////////////////////////////////////////////////////
+
+class FMessageDialogUtils
+{
+public:
+	static EDialogResponse ShowChoiceDialog(const FText& Title, const FText& Content);
+};
+
+////////////////////////////////////////////////////////////////////////////////////////////////////
+
+} // namespace Insights
\ No newline at end of file
diff --git a/Engine/Source/Developer/TraceInsights/Private/Insights/Widgets/SStartPageWindow.cpp b/Engine/Source/Developer/TraceInsights/Private/Insights/Widgets/SStartPageWindow.cpp
--- a/Engine/Source/Developer/TraceInsights/Private/Insights/Widgets/SStartPageWindow.cpp
+++ b/Engine/Source/Developer/TraceInsights/Private/Insights/Widgets/SStartPageWindow.cpp
@@ -46,6 +46,7 @@
 #endif // WITH_EDITOR
 
 // Insights
+#include "Insights/Common/MessageDialogUtils.h"
 #include "Insights/Common/Stopwatch.h"
 #include "Insights/ImportTool/TableImportTool.h"
 #include "Insights/InsightsManager.h"
@@ -1591,15 +1592,17 @@ FReply STraceDirectoryItem::OnDelete()
 
 	if (Model)
 	{
-		const EAppReturnType::Type bConfirmed = FMessageDialog::Open(EAppMsgType::OkCancel, FText::Format(
-			LOCTEXT("WatchDirRemoveConfirmBody", "This will remove '{0}' from monitored directories?"), FText::FromString(Model->Path)),
-			LOCTEXT("WatchDirRemoveConfirmTitle", "Confirm removing monitored directory"));
-		if (bConfirmed == EAppReturnType::Ok)
+		Insights::EDialogResponse Response = Insights::FMessageDialogUtils::ShowChoiceDialog(LOCTEXT("MonitoredDirRemoveConfirmTitle", "Confirm removing monitored directory"), FText::Format(
+			LOCTEXT("MonitoredDirRemoveConfirmFmt", "This will remove \"{0}\" from monitored directories.\n\nConfirm removing monitored directory?"),
+			FText::FromString(Model->Path)));
+
+		if (Response == Insights::EDialogResponse::OK)
 		{
 			const TCHAR* Path = *Model->Path;
 			FInsightsManager::Get()->GetStoreClient()->SetStoreDirectories(nullptr, {}, { Path });
 		}
 	}
+
 	return FReply::Handled();
 }
 
diff --git a/Engine/Source/Developer/TraceInsights/TraceInsights.Build.cs b/Engine/Source/Developer/TraceInsights/TraceInsights.Build.cs
--- a/Engine/Source/Developer/TraceInsights/TraceInsights.Build.cs
+++ b/Engine/Source/Developer/TraceInsights/TraceInsights.Build.cs
@@ -57,6 +57,7 @@ public class TraceInsights : ModuleRules
 			new string[] {
 				"ImageCore",
 				"SlateCore",
+				"ToolWidgets"
 			}
 		);
 
