[Insights] Prevent the doubling of a warning when task relations are displayed from the Task Table.

diff --git a/Engine/Source/Developer/TraceInsights/Private/Insights/TaskGraphProfiler/Widgets/STaskTableTreeView.cpp b/Engine/Source/Developer/TraceInsights/Private/Insights/TaskGraphProfiler/Widgets/STaskTableTreeView.cpp
--- a/Engine/Source/Developer/TraceInsights/Private/Insights/TaskGraphProfiler/Widgets/STaskTableTreeView.cpp
+++ b/Engine/Source/Developer/TraceInsights/Private/Insights/TaskGraphProfiler/Widgets/STaskTableTreeView.cpp
@@ -732,8 +732,6 @@ void STaskTableTreeView::ContextMenu_GoToTask_Execute()
 		return;
 	}
 
-	FTaskGraphProfilerManager::Get()->ShowTaskRelations(TaskEntry->GetId());
-
 	double Duration = (TaskEntry->GetFinishedTimestamp() - TaskEntry->GetCreatedTimestamp()) * 1.5;
 	TimingView->ZoomOnTimeInterval(TaskEntry->GetCreatedTimestamp() - Duration * 0.15, Duration);
 
@@ -770,9 +768,15 @@ void STaskTableTreeView::ContextMenu_GoToTask_Execute()
 
 	FTimingEventSearchParameters SearchParams(TaskEntry->StartedTimestamp, TaskEntry->FinishedTimestamp, ETimingEventSearchFlags::StopAtFirstMatch, SearchFilter);
 
-
 	const TSharedPtr<const ITimingEvent> FoundEvent = Track->SearchEvent(SearchParams);
-	TimingView->SelectTimingEvent(FoundEvent, true);
+	if (FoundEvent.IsValid() && Track->IsVisible())
+	{
+		TimingView->SelectTimingEvent(FoundEvent, true);
+	}
+	else
+	{
+		FTaskGraphProfilerManager::Get()->ShowTaskRelations(TaskEntry->GetId());
+	}
 }
 
 ////////////////////////////////////////////////////////////////////////////////////////////////////
