[Insights]
  - Fixed formatted text in tooltips re open source file when accessor is not available.
  - Timing View: Added line number in addition to file name for selected event (info displayed in the bottom-right corner of the view).

diff --git a/Engine/Source/Developer/TraceInsights/Private/Insights/TaskGraphProfiler/Widgets/STaskTableTreeView.cpp b/Engine/Source/Developer/TraceInsights/Private/Insights/TaskGraphProfiler/Widgets/STaskTableTreeView.cpp
--- a/Engine/Source/Developer/TraceInsights/Private/Insights/TaskGraphProfiler/Widgets/STaskTableTreeView.cpp
+++ b/Engine/Source/Developer/TraceInsights/Private/Insights/TaskGraphProfiler/Widgets/STaskTableTreeView.cpp
@@ -149,7 +149,7 @@ void STaskTableTreeView::ExtendMenu(FMenuBuilder& MenuBuilder)
 			ItemLabel = LOCTEXT("ContextMenu_OpenSourceNA", "Open Source");
 			if (bIsValidSource)
 			{
-				ItemToolTip = FText::Format(LOCTEXT("ContextMenu_OpenSourceNA_Desc1", "{1} ({2})\nSource Code Accessor is not available."),
+				ItemToolTip = FText::Format(LOCTEXT("ContextMenu_OpenSourceNA_Desc1", "{0} ({1})\nSource Code Accessor is not available."),
 					FText::FromString(File),
 					FText::AsNumber(Line, &FNumberFormattingOptions::DefaultNoGrouping()));
 			}
diff --git a/Engine/Source/Developer/TraceInsights/Private/Insights/ViewModels/ThreadTimingTrack.cpp b/Engine/Source/Developer/TraceInsights/Private/Insights/ViewModels/ThreadTimingTrack.cpp
--- a/Engine/Source/Developer/TraceInsights/Private/Insights/ViewModels/ThreadTimingTrack.cpp
+++ b/Engine/Source/Developer/TraceInsights/Private/Insights/ViewModels/ThreadTimingTrack.cpp
@@ -1223,8 +1223,9 @@ void FThreadTimingTrack::PostDraw(const ITimingTrackDrawContext& Context) const
 
 			if (Timer->File)
 			{
-				FString SourceFile(Timer->File);
-				DrawSelectedEventInfoEx(StatsText, TimerName, FPaths::GetCleanFilename(SourceFile), Context.GetViewport(), Context.GetDrawContext(), Helper.GetWhiteBrush(), Helper.GetEventFont());
+				FString SourceFile = FPaths::GetCleanFilename(FString(Timer->File));
+				FString SourceFileAndLine = FString::Printf(TEXT("%s (%d)"), *SourceFile, Timer->Line);
+				DrawSelectedEventInfoEx(StatsText, TimerName, SourceFileAndLine, Context.GetViewport(), Context.GetDrawContext(), Helper.GetWhiteBrush(), Helper.GetEventFont());
 			}
 			else
 			{
diff --git a/Engine/Source/Developer/TraceInsights/Private/Insights/Widgets/STimersView.cpp b/Engine/Source/Developer/TraceInsights/Private/Insights/Widgets/STimersView.cpp
--- a/Engine/Source/Developer/TraceInsights/Private/Insights/Widgets/STimersView.cpp
+++ b/Engine/Source/Developer/TraceInsights/Private/Insights/Widgets/STimersView.cpp
@@ -619,7 +619,9 @@ TSharedPtr<SWidget> STimersView::TreeView_GetMenuContent()
 				if (bIsValidSource)
 				{
 					ItemToolTip = FText::Format(LOCTEXT("ContextMenu_OpenSource_Desc1", "Opens the source file of the selected timer in {0}.\n{1} ({2})"),
-						SourceCodeAccessor.GetNameText(), FText::FromString(File), FText::AsNumber(Line, &FNumberFormattingOptions::DefaultNoGrouping()));
+						SourceCodeAccessor.GetNameText(),
+						FText::FromString(File),
+						FText::AsNumber(Line, &FNumberFormattingOptions::DefaultNoGrouping()));
 				}
 				else
 				{
@@ -632,8 +634,9 @@ TSharedPtr<SWidget> STimersView::TreeView_GetMenuContent()
 				ItemLabel = LOCTEXT("ContextMenu_OpenSourceNA", "Open Source");
 				if (bIsValidSource)
 				{
-					ItemToolTip = FText::Format(LOCTEXT("ContextMenu_OpenSourceNA_Desc1", "{1} ({2})\nSource Code Accessor is not available."),
-						FText::FromString(File), FText::AsNumber(Line, &FNumberFormattingOptions::DefaultNoGrouping()));
+					ItemToolTip = FText::Format(LOCTEXT("ContextMenu_OpenSourceNA_Desc1", "{0} ({1})\nSource Code Accessor is not available."),
+						FText::FromString(File),
+						FText::AsNumber(Line, &FNumberFormattingOptions::DefaultNoGrouping()));
 				}
 				else
 				{
