[Insights]
  - Fixed trace server control when connection fails. Also lowered the number of attempts to connect to UTS.
  - Updated code using TUniquePtr to use Reset().

diff --git a/Engine/Source/Developer/TraceInsights/Private/Insights/InsightsManager.cpp b/Engine/Source/Developer/TraceInsights/Private/Insights/InsightsManager.cpp
--- a/Engine/Source/Developer/TraceInsights/Private/Insights/InsightsManager.cpp
+++ b/Engine/Source/Developer/TraceInsights/Private/Insights/InsightsManager.cpp
@@ -460,8 +460,7 @@ FString FInsightsManager::GetStoreDir()
 bool FInsightsManager::ConnectToStore(const TCHAR* Host, uint32 Port)
 {
 	using namespace UE::Trace;
-	FStoreClient* Client = FStoreClient::Connect(Host, Port);
-	StoreClient = TUniquePtr<FStoreClient>(Client);
+	StoreClient.Reset(FStoreClient::Connect(Host, Port));
 	if (!StoreClient.IsValid())
 	{
 		return false;
diff --git a/Engine/Source/Developer/TraceInsights/Private/Insights/TraceInsightsModule.cpp b/Engine/Source/Developer/TraceInsights/Private/Insights/TraceInsightsModule.cpp
--- a/Engine/Source/Developer/TraceInsights/Private/Insights/TraceInsightsModule.cpp
+++ b/Engine/Source/Developer/TraceInsights/Private/Insights/TraceInsightsModule.cpp
@@ -185,7 +185,7 @@ void FTraceInsightsModule::CreateDefaultStore()
 	StoreServiceDesc.StoreDir = *StoreDir;
 	StoreServiceDesc.RecorderPort = 1981;
 	StoreServiceDesc.ThreadCount = 2;
-	StoreService = TUniquePtr<UE::Trace::FStoreService>(UE::Trace::FStoreService::Create(StoreServiceDesc));
+	StoreService.Reset(UE::Trace::FStoreService::Create(StoreServiceDesc));
 
 	if (StoreService.IsValid())
 	{
@@ -202,8 +202,7 @@ FString FTraceInsightsModule::GetDefaultStoreDir()
 {
 	using UE::Trace::FStoreClient;
 
-	FStoreClient* StoreClientPtr = FStoreClient::Connect(TEXT("localhost"));
-	TUniquePtr<FStoreClient> StoreClient = TUniquePtr<FStoreClient>(StoreClientPtr);
+	TUniquePtr<FStoreClient> StoreClient(FStoreClient::Connect(TEXT("localhost")));
 
 	if (!StoreClient)
 	{
diff --git a/Engine/Source/Developer/TraceInsights/Private/Insights/Widgets/STraceServerControl.cpp b/Engine/Source/Developer/TraceInsights/Private/Insights/Widgets/STraceServerControl.cpp
--- a/Engine/Source/Developer/TraceInsights/Private/Insights/Widgets/STraceServerControl.cpp
+++ b/Engine/Source/Developer/TraceInsights/Private/Insights/Widgets/STraceServerControl.cpp
@@ -15,9 +15,9 @@
 // Number of times to attempt state change before failing
 constexpr uint32 GStateChangeRetries = 6;
 // Number of times to attempt reconnection while starting server
-constexpr uint32 GStartConnectAttempts = 25;
+constexpr uint32 GStartConnectAttempts = 5;
 // Number of seconds between each reconnection attempt
-constexpr float GStartConnectFrequencySeconds = 0.25;
+constexpr float GStartConnectFrequencySeconds = 0.5;
 
 /////////////////////////////////////////////////////////////////////////////////////////
 
@@ -36,9 +36,9 @@ BEGIN_SLATE_FUNCTION_BUILD_OPTIMIZATION
 void STraceServerControl::MakeMenu(FMenuBuilder& Builder)
 {
 	// Create client
-	if (!Client.IsValid())
+	if (!Client)
 	{
-		Client = TUniquePtr<UE::Trace::FStoreClient>(UE::Trace::FStoreClient::Connect(*Host, Port));
+		Client.Reset(UE::Trace::FStoreClient::Connect(*Host, Port));
 		if (Client)
 		{
 			ChangeState(EState::NotConnected, EState::Connected);
@@ -235,7 +235,7 @@ void STraceServerControl::OnStart_Clicked()
 		uint32 Attempts = 0;
 		while (!Client && Attempts++ < GStartConnectAttempts)
 		{
-			Client = TUniquePtr<UE::Trace::FStoreClient>(UE::Trace::FStoreClient::Connect(TEXT("127.0.0.1")));
+			Client.Reset(UE::Trace::FStoreClient::Connect(TEXT("127.0.0.1")));
 			FPlatformProcess::Sleep(GStartConnectFrequencySeconds);
 		}
 		if (Client)
@@ -244,6 +244,7 @@ void STraceServerControl::OnStart_Clicked()
 		}
 		else
 		{
+			ChangeState(EState::Connecting, EState::NotConnected);
 			UE_LOG(LogTemp, Warning, TEXT("Failed to connect to store"));
 		}
 	}, TStatId(), {CommandTask});
