Fix Insights tests handling UTS setup
Use sponsor argument to pass editor pid before starting UTS

#rnx

diff --git a/Engine/Source/Developer/TraceInsights/Private/Insights/Tests/FunctionalTests/MemoryInsightsTests.cpp b/Engine/Source/Developer/TraceInsights/Private/Insights/Tests/FunctionalTests/MemoryInsightsTests.cpp
--- a/Engine/Source/Developer/TraceInsights/Private/Insights/Tests/FunctionalTests/MemoryInsightsTests.cpp
+++ b/Engine/Source/Developer/TraceInsights/Private/Insights/Tests/FunctionalTests/MemoryInsightsTests.cpp
@@ -115,7 +115,7 @@ void FAutomationDriverUnrealInsightsHubMemoryInsightsTest::Define()
 					const FString SuccessTestResult = TEXT("Test Completed. Result={Success}");
 
 					// Test live trace
-					FString TraceParameters = FString::Printf(TEXT("-InsightsTest -ABSLOG=\"%s\" -AutoQuit -ExecOnAnalysisCompleteCmd=\"Automation RunTests Insights.UploadMemoryInsightsLLMXMLReportsTrace\" -OpenTraceFile=\"%s\""), *LogPath, *StoreTracePath);
+					FString TraceParameters = FString::Printf(TEXT("-InsightsTest -ABSLOG=\"%s\" -AutoQuit -ExecOnAnalysisCompleteCmd=\"Automation RunTests System.Insights.Trace.Analysis.MemoryInsights.UploadMemoryInsightsLLMXMLReportsTrace\" -OpenTraceFile=\"%s\""), *LogPath, *StoreTracePath);
 					InsightsManager->OpenUnrealInsights(*TraceParameters);
 					bool bLineFound = Utils.FileContainsString(LogPath, SuccessTestResult, 60.0f);
 					TestTrue("Test for live trace should pass", bLineFound);
diff --git a/Engine/Source/Developer/TraceInsights/Private/Insights/Tests/InsightsTestUtils.cpp b/Engine/Source/Developer/TraceInsights/Private/Insights/Tests/InsightsTestUtils.cpp
--- a/Engine/Source/Developer/TraceInsights/Private/Insights/Tests/InsightsTestUtils.cpp
+++ b/Engine/Source/Developer/TraceInsights/Private/Insights/Tests/InsightsTestUtils.cpp
@@ -170,7 +170,19 @@ bool FInsightsTestUtils::SetupUTS(double Timeout, bool bUseFork) const
 {
 	const FString UnrealTraceServerName = TEXT("UnrealTraceServer");
 
+	if (FPlatformProcess::IsApplicationRunning(*UnrealTraceServerName))
+	{
+		Test->AddInfo(TEXT("UTS is already running"));
+		return true;
+	}
+
 	FString UTSPath = FPlatformProcess::GenerateApplicationPath("UnrealTraceServer", EBuildConfiguration::Development);
+	if (!FPaths::FileExists(UTSPath))
+	{
+		Test->AddError(FString::Printf(TEXT("UTS executable can't be found at '%s'"), *UTSPath));
+		return false;
+	}
+
 	FString UTSParameters;
 	if (bUseFork)
 	{
@@ -180,6 +192,7 @@ bool FInsightsTestUtils::SetupUTS(double Timeout, bool bUseFork) const
 	{
 		UTSParameters = TEXT("daemon");
 	}
+	UTSParameters += FString::Printf(TEXT(" --sponsor %d"), FPlatformProcess::GetCurrentProcessId());
 	constexpr bool bLaunchDetached = true;
 	constexpr bool bLaunchHidden = false;
 	constexpr bool bLaunchReallyHidden = false;
@@ -187,8 +200,9 @@ bool FInsightsTestUtils::SetupUTS(double Timeout, bool bUseFork) const
 	const int32 PriorityModifier = 0;
 	const TCHAR* OptionalWorkingDirectory = nullptr;
 	void* PipeWriteChild = nullptr;
-	void* PipeReadChild = nullptr;
-	FProcHandle UTSHandle = FPlatformProcess::CreateProc(*UTSPath, *UTSParameters, bLaunchDetached, bLaunchHidden, bLaunchReallyHidden, &ProcessID, PriorityModifier, OptionalWorkingDirectory, PipeWriteChild, PipeReadChild);
+	void* PipeOutput = nullptr;
+	verify(FPlatformProcess::CreatePipe(PipeOutput, PipeWriteChild));
+	FProcHandle UTSHandle = FPlatformProcess::CreateProc(*UTSPath, *UTSParameters, bLaunchDetached, bLaunchHidden, bLaunchReallyHidden, &ProcessID, PriorityModifier, OptionalWorkingDirectory, PipeWriteChild, nullptr);
 	if (!UTSHandle.IsValid())
 	{
 		Test->AddError(TEXT("The UTSHandle should be valid"));
@@ -196,7 +210,7 @@ bool FInsightsTestUtils::SetupUTS(double Timeout, bool bUseFork) const
 	}
 
 	double StartTime = FPlatformTime::Seconds();
-	while (FPlatformTime::Seconds() - StartTime < Timeout)
+	while (FPlatformTime::Seconds() - StartTime < Timeout && FPlatformProcess::IsProcRunning(UTSHandle))
 	{
 		FPlatformProcess::Sleep(0.1f);
 		if (!FPlatformProcess::IsApplicationRunning(*UnrealTraceServerName))
@@ -214,6 +228,13 @@ bool FInsightsTestUtils::SetupUTS(double Timeout, bool bUseFork) const
 	}
 
 	Test->AddError(TEXT("UTS failed to start"));
+	if (!FPlatformProcess::IsProcRunning(UTSHandle))
+	{
+		FString StringOutput = FPlatformProcess::ReadPipe(PipeOutput);
+		int32 ExitCode = 0;
+		FPlatformProcess::GetProcReturnCode(UTSHandle, &ExitCode);
+		Test->AddError(FString::Printf(TEXT("UTS exitcode=%d stdout:\n%s"), ExitCode, *StringOutput));
+	}
 	return false;
 }
 
