[Insights] Fixed UI refering to "monitored directories" (instead of "watch directories"). Code (variable names) still uses "WatchDir" naming.

diff --git a/Engine/Source/Developer/TraceInsights/Private/Insights/Widgets/SStartPageWindow.cpp b/Engine/Source/Developer/TraceInsights/Private/Insights/Widgets/SStartPageWindow.cpp
--- a/Engine/Source/Developer/TraceInsights/Private/Insights/Widgets/SStartPageWindow.cpp
+++ b/Engine/Source/Developer/TraceInsights/Private/Insights/Widgets/SStartPageWindow.cpp
@@ -1458,7 +1458,7 @@ TSharedRef<SWidget> STraceDirectoryItem::ConstructOperations()
 {
 	TSharedRef<SHorizontalBox> Box = SNew(SHorizontalBox);
 
-	if (Model && EnumHasAllFlags(Model->Operations, ETraceDirOperations::ModifyStore) )
+	if (Model && EnumHasAllFlags(Model->Operations, ETraceDirOperations::ModifyStore))
 	{
 		Box->AddSlot()
 			.AutoWidth()
@@ -1475,14 +1475,14 @@ TSharedRef<SWidget> STraceDirectoryItem::ConstructOperations()
 				]
 			];
 	}
-	if (Model && EnumHasAllFlags(Model->Operations, ETraceDirOperations::Delete) )
+	if (Model && EnumHasAllFlags(Model->Operations, ETraceDirOperations::Delete))
 	{
 		Box->AddSlot()
 			.AutoWidth()
 			[
 				SNew(SButton)
 				.ButtonStyle(&FAppStyle::Get().GetWidgetStyle<FButtonStyle>("SimpleButton"))
-				.ToolTipText(LOCTEXT("WatchDirsRemoveTooltip", "Remove watch directory (files will not be deleted)"))
+				.ToolTipText(LOCTEXT("WatchDirsRemoveTooltip", "Removes the monitored directory. Files will not be deleted."))
 				.OnClicked_Raw(this, &STraceDirectoryItem::OnDelete)
 				.IsEnabled_Static(&STraceStoreWindow::CanChangeStoreSettings)
 				[
@@ -1494,12 +1494,17 @@ TSharedRef<SWidget> STraceDirectoryItem::ConstructOperations()
 	}
 	if (Model && EnumHasAllFlags(Model->Operations, ETraceDirOperations::Explore))
 	{
+		// If it has a Delete button then it is a "monitored directory".
+		bool bIsWatchDir = (Model && EnumHasAllFlags(Model->Operations, ETraceDirOperations::Delete));
+
 		Box->AddSlot()
 			.AutoWidth()
 			[
 				SNew(SButton)
 				.ButtonStyle(&FAppStyle::Get().GetWidgetStyle<FButtonStyle>("SimpleButton"))
-				.ToolTipText(LOCTEXT("ExploreTraceStoreDirButtonToolTip", "Explore the Trace Store Directory"))
+				.ToolTipText(bIsWatchDir ? 
+								LOCTEXT("ExploreWatchDirButtonToolTip", "Explores the monitored directory.") :
+								LOCTEXT("ExploreTraceStoreDirButtonToolTip", "Explores the Trace Store Directory."))
 				.OnClicked_Raw(this, &STraceDirectoryItem::OnExplore)
 				.IsEnabled_Static(&STraceStoreWindow::CanChangeStoreSettings)
 				[
@@ -1529,7 +1534,7 @@ FReply STraceDirectoryItem::OnModifyStore()
 
 	if (IDesktopPlatform* DesktopPlatform = FDesktopPlatformModule::Get())
 	{
-		const FString Title = LOCTEXT("SetTraceStoreDirectorySelectPopupTitle", "Set Trace Store Directory").ToString();
+		const FString Title = LOCTEXT("SetTraceStoreDirectory_DialogTitle", "Set Trace Store Directory").ToString();
 
 		FString CurrentStoreDirectory = Window->GetStoreDirectory();
 		FString SelectedDirectory;
@@ -1539,7 +1544,7 @@ FReply STraceDirectoryItem::OnModifyStore()
 			CurrentStoreDirectory,
 			SelectedDirectory);
 
-		const bool bIsWatchDir = Window->WatchDirectoriesModel.FindByPredicate([&](const auto& Directory){ return FPathViews::Equals(SelectedDirectory, Directory->Path);}) != nullptr;
+		const bool bIsWatchDir = Window->WatchDirectoriesModel.FindByPredicate([&](const auto& Directory) { return FPathViews::Equals(SelectedDirectory, Directory->Path); }) != nullptr;
 		const bool bIsCurrentStoreDir = FPathViews::Equals(SelectedDirectory, CurrentStoreDirectory);
 
 		if (bHasSelected && !bIsCurrentStoreDir)
@@ -1550,8 +1555,8 @@ FReply STraceDirectoryItem::OnModifyStore()
 			TArray<FString> RemoveWatchDirs;
 			if (bIsWatchDir)
 			{
-				// If we are selecting a watch dir as new store dir, make
-				// sure we remove it as watch directory.
+				// If we are selecting a monitored dir as new store dir, make
+				// sure we remove it as monitored directory.
 				RemoveWatchDirs.Emplace(SelectedDirectory);
 			}
 			if (!FInsightsManager::Get()->GetStoreClient()->SetStoreDirectories(*SelectedDirectory, AddWatchDirs, RemoveWatchDirs))
@@ -1618,8 +1623,9 @@ bool STraceDirectoryItem::CanModifyStore() const
 
 FText STraceDirectoryItem::ModifyStoreTooltip() const
 {
-	return CanModifyStore() ? LOCTEXT("SetTraceStoreDirButtonToolTip", "Set the Trace Store Directory.") :
-		LOCTEXT("SetTraceStoreDirButtonTooltipInactive", "Set the Trace Store Directory.\nNot available while live trace sessions are running.");
+	return CanModifyStore() ?
+		LOCTEXT("SetTraceStoreDirButtonToolTip", "Sets the Trace Store Directory.") :
+		LOCTEXT("SetTraceStoreDirButtonTooltipInactive", "Sets the Trace Store Directory.\nNot available while live trace sessions are running.");
 }
 
 ////////////////////////////////////////////////////////////////////////////////////////////////////
@@ -1762,7 +1768,7 @@ TSharedRef<SWidget> STraceStoreWindow::ConstructTraceStoreDirectoryPanel()
 			[
 				SNew(SButton)
 				.ButtonStyle(&FAppStyle::Get().GetWidgetStyle<FButtonStyle>("SimpleButton"))
-				.ToolTipText(LOCTEXT("ExploreTraceStoreDirButtonToolTip", "Explore the Trace Store Directory"))
+				.ToolTipText(LOCTEXT("ExploreTraceStoreDirButtonToolTip", "Explores the Trace Store Directory."))
 				.OnClicked(this, &STraceStoreWindow::ExploreTraceStoreDirectory_OnClicked)
 				.AddMetaData(FDriverMetaData::Id("ExploreTraceStoreDirButton"))
 				.IsEnabled_Static(&STraceStoreWindow::CanChangeStoreSettings)
@@ -1795,7 +1801,7 @@ TSharedRef<SWidget> STraceStoreWindow::ConstructTraceStoreDirectoryPanel()
 				.VAlign(VAlign_Center)
 				[
 					SNew(STextBlock)
-					.Text(LOCTEXT("StoreDirLabel", "Trace store directory. New traces will be stored here:"))
+					.Text(LOCTEXT("StoreDirLabel", "Trace Store Directory (new traces will be stored here):"))
 				]
 
 				+ SVerticalBox::Slot()
@@ -1834,7 +1840,7 @@ TSharedRef<SWidget> STraceStoreWindow::ConstructTraceStoreDirectoryPanel()
 				[
 					SNew(SButton)
 					.ButtonStyle(&FAppStyle::Get().GetWidgetStyle<FButtonStyle>("Button"))
-					.ToolTipText(LOCTEXT("WatchDirsAddTooltip", "Add additional watch directory..."))
+					.ToolTipText(LOCTEXT("WatchDirsAddTooltip", "Adds an additional directory to monitor for traces."))
 					.IsEnabled_Static(&STraceStoreWindow::CanChangeStoreSettings)
 					.OnClicked(this, &STraceStoreWindow::AddWatchDir_Clicked)
 					[
@@ -2501,7 +2507,7 @@ void STraceStoreWindow::RefreshTraceList()
 				StoreDirListView->RequestListRefresh();
 			}
 
-			// Update additional watch directory models
+			// Update additional monitored directories model
 			static const FName DirColor[] =
 			{
 				FName("Colors.AccentBlue"),
@@ -3625,7 +3631,7 @@ FReply STraceStoreWindow::AddWatchDir_Clicked()
 	FSlateApplication::Get().CloseToolTip();
 	if (IDesktopPlatform* DesktopPlatform = FDesktopPlatformModule::Get())
 	{
-		const FString Title = LOCTEXT("AddWatchDirectorySelectPopupTitle", "Add Watch Directory").ToString();
+		const FString Title = LOCTEXT("AddWatchDirectory_DialogTitle", "Add Monitored Directory").ToString();
 
 		const FString& CurrentStoreDirectory = StoreDirectoryModel.IsEmpty() ? FString() : StoreDirectoryModel.Last()->Path;
 		FString SelectedDirectory;
diff --git a/Engine/Source/Developer/TraceInsights/Private/Insights/Widgets/SStartPageWindow.h b/Engine/Source/Developer/TraceInsights/Private/Insights/Widgets/SStartPageWindow.h
--- a/Engine/Source/Developer/TraceInsights/Private/Insights/Widgets/SStartPageWindow.h
+++ b/Engine/Source/Developer/TraceInsights/Private/Insights/Widgets/SStartPageWindow.h
@@ -462,7 +462,7 @@ private:
 	//////////////////////////////////////////////////
 	// Store settings
 
-	/** View model for a source of traces. This could be the trace store default directory or a watch directory.*/
+	/** View model for a source of traces. This could be the trace store default directory or an additional monitored directory. */
 	struct FTraceDirectoryModel
 	{
 		FTraceDirectoryModel(FString&& InPath, const FName& InColor, ETraceDirOperations InOperations)
