[Insights]
  - Partially fixed STraceServerControl when destructed before async tasks are completed.
  - Added cancelation for async tasks and ensured the async tasks are not running in the same time.
  - Updated tooltips re Trace Server Control menu actions.

diff --git a/Engine/Source/Developer/TraceInsights/Private/Insights/Widgets/STraceServerControl.cpp b/Engine/Source/Developer/TraceInsights/Private/Insights/Widgets/STraceServerControl.cpp
--- a/Engine/Source/Developer/TraceInsights/Private/Insights/Widgets/STraceServerControl.cpp
+++ b/Engine/Source/Developer/TraceInsights/Private/Insights/Widgets/STraceServerControl.cpp
@@ -31,6 +31,14 @@ STraceServerControl::STraceServerControl(const TCHAR* InHost, uint32 InPort, FNa
 
 /////////////////////////////////////////////////////////////////////////////////////////
 
+STraceServerControl::~STraceServerControl()
+{
+	bIsCancelRequested = true;
+	FScopeLock _(&AsyncTaskLock); // wait for async tasks to complete
+}
+
+/////////////////////////////////////////////////////////////////////////////////////////
+
 BEGIN_SLATE_FUNCTION_BUILD_OPTIMIZATION
 
 void STraceServerControl::MakeMenu(FMenuBuilder& Builder)
@@ -46,8 +54,11 @@ void STraceServerControl::MakeMenu(FMenuBuilder& Builder)
 	}
 
 	// If connected kick off status and version check
-	TriggerStatusUpdate();
-	
+	if (State.load(std::memory_order_relaxed) == EState::Connected)
+	{
+		TriggerStatusUpdate();
+	}
+
 	if (bIsLocalHost)
 	{
 		Builder.BeginSection("LocalTraceServer", LOCTEXT("Section_LocalServer", "Local Trace Server"));
@@ -67,40 +78,40 @@ void STraceServerControl::MakeMenu(FMenuBuilder& Builder)
 	if (bIsLocalHost)
 	{
 		Builder.AddMenuEntry(
-			 LOCTEXT("ServerControlSponsoredLabel", "Sponsored Mode"),
-			 LOCTEXT("ServerControlSponsoredTooltip", "In sponsored mode the server only runs as long as local processes that uses it are alive."),
-			 FSlateIcon(), //?
-			 FUIAction(
-				  FExecuteAction::CreateRaw(this, &STraceServerControl::OnSponsored_Changed),
-				  FCanExecuteAction::CreateRaw(this, &STraceServerControl::AreControlsEnabled),
-				  FGetActionCheckState::CreateLambda([this](){ return IsSponsored() ? ECheckBoxState::Checked : ECheckBoxState::Unchecked; })
-			 ),
-			 NAME_None,
-			 EUserInterfaceActionType::ToggleButton
+			LOCTEXT("ServerControlSponsoredLabel", "Sponsored Mode"),
+			LOCTEXT("ServerControlSponsoredTooltip", "In sponsored mode the server only runs as long as local processes that uses it are alive."),
+			FSlateIcon(), //?
+			FUIAction(
+				FExecuteAction::CreateRaw(this, &STraceServerControl::OnSponsored_Changed),
+				FCanExecuteAction::CreateRaw(this, &STraceServerControl::AreControlsEnabled),
+				FGetActionCheckState::CreateLambda([this](){ return IsSponsored() ? ECheckBoxState::Checked : ECheckBoxState::Unchecked; })
+			),
+			NAME_None,
+			EUserInterfaceActionType::ToggleButton
 		);
-		
+
 		Builder.AddMenuEntry(
-			 LOCTEXT("ServerControlStartLabel", "Start"),
-			 LOCTEXT("ServerControlStartTooltop", "Start Server"),
-			 FSlateIcon(StyleSet, "Icons.TraceServerStart"),
-			 FUIAction(
-				  FExecuteAction::CreateRaw(this, &STraceServerControl::OnStart_Clicked),
-				  FCanExecuteAction::CreateRaw(this, &STraceServerControl::CanServerBeStarted)
-			 ),
-			 NAME_None,
-			 EUserInterfaceActionType::Button
+			LOCTEXT("ServerControlStartLabel", "Start"),
+			LOCTEXT("ServerControlStartTooltop", "Starts the Trace Server"),
+			FSlateIcon(StyleSet, "Icons.TraceServerStart"),
+			FUIAction(
+				FExecuteAction::CreateRaw(this, &STraceServerControl::OnStart_Clicked),
+				FCanExecuteAction::CreateRaw(this, &STraceServerControl::CanServerBeStarted)
+			),
+			NAME_None,
+			EUserInterfaceActionType::Button
 		);
 
 		Builder.AddMenuEntry(
-			 LOCTEXT("ServerControlStopLabel", "Stop"),
-			 LOCTEXT("ServerControlStopTooltop", "Stop Server. Any running traces will be cancelled."),
-			 FSlateIcon(StyleSet, "Icons.TraceServerStop"),
-			 FUIAction(
-				  FExecuteAction::CreateRaw(this, &STraceServerControl::OnStop_Clicked),
-				  FCanExecuteAction::CreateRaw(this, &STraceServerControl::CanServerBeStopped)
-			 ),
-			 NAME_None,
-			 EUserInterfaceActionType::Button
+			LOCTEXT("ServerControlStopLabel", "Stop"),
+			LOCTEXT("ServerControlStopTooltop", "Stops the Trace Server. Any running traces will be cancelled."),
+			FSlateIcon(StyleSet, "Icons.TraceServerStop"),
+			FUIAction(
+				FExecuteAction::CreateRaw(this, &STraceServerControl::OnStop_Clicked),
+				FCanExecuteAction::CreateRaw(this, &STraceServerControl::CanServerBeStopped)
+			),
+			NAME_None,
+			EUserInterfaceActionType::Button
 		);
 	}
 	Builder.EndSection();
@@ -145,14 +156,20 @@ bool STraceServerControl::ChangeState(EState Expected, EState ChangeTo, uint32 A
 
 void STraceServerControl::TriggerStatusUpdate()
 {
-	if (State.load(std::memory_order_relaxed) == EState::Connected && ChangeState(EState::Connected, EState::CheckStatus))
+	FGraphEventRef CheckStatusTask = FFunctionGraphTask::CreateAndDispatchWhenReady([this]
 	{
-		FFunctionGraphTask::CreateAndDispatchWhenReady([this]
+		if (bIsCancelRequested)
+		{
+			return;
+		}
+		FScopeLock _(&AsyncTaskLock);
+
+		if (ChangeState(EState::Connected, EState::CheckStatus, GStateChangeRetries))
 		{
 			UpdateStatus();
 			ChangeState(EState::CheckStatus, EState::Connected);
-		});
-	}
+		}
+	});
 }
 
 /////////////////////////////////////////////////////////////////////////////////////////
@@ -173,13 +190,13 @@ void STraceServerControl::UpdateStatus()
 		ChangeState(EState::CheckStatus, EState::NotConnected);
 		return;
 	}
-	
+
 	TStringBuilder<64> PortsStringBuilder;
 	if (Status)
 	{
 		bSponsored.store(Status->GetSponsored());
 		PortsStringBuilder << TEXT("Recorder Port: ") << Status->GetRecorderPort()
-		<< TEXT(", Store Port: ") << Status->GetStorePort();
+							<< TEXT(", Store Port: ") << Status->GetStorePort();
 	}
 
 	// If not previously checked, also query version information
@@ -218,44 +235,70 @@ void STraceServerControl::OnStart_Clicked()
 {
 	FGraphEventRef CommandTask = FFunctionGraphTask::CreateAndDispatchWhenReady([this]
 	{
+		if (bIsCancelRequested)
+		{
+			return;
+		}
+		FScopeLock _(&AsyncTaskLock);
+
 		if (ChangeState(EState::NotConnected, EState::Command, GStateChangeRetries))
 		{
 #if UE_TRACE_SERVER_CONTROLS_ENABLED
 			FTraceServerControls::Start();
 #endif
 			ChangeState(EState::Command, EState::Connecting);
+
+			if (!Client)
+			{
+				for (uint32 Attempts = 0; Attempts < GStartConnectAttempts; ++Attempts)
+				{
+					UE::Trace::FStoreClient* NewClient = UE::Trace::FStoreClient::Connect(TEXT("127.0.0.1"));
+					if (bIsCancelRequested)
+					{
+						if (NewClient)
+						{
+							delete NewClient;
+						}
+						break;
+					}
+					if (NewClient)
+					{
+						Client.Reset(NewClient);
+						break;
+					}
+					FPlatformProcess::Sleep(GStartConnectFrequencySeconds);
+				}
+			}
+
+			if (Client)
+			{
+				ChangeState(EState::Connecting, EState::Connected);
+			}
+			else
+			{
+				ChangeState(EState::Connecting, EState::NotConnected);
+				UE_LOG(TraceInsights, Warning, TEXT("Failed to connect to store."));
+			}
 		}
 		else
 		{
 			UE_LOG(TraceInsights, Warning, TEXT("Failed to start server."));
 		}
 	});
-	FFunctionGraphTask::CreateAndDispatchWhenReady([this]
-	{
-		uint32 Attempts = 0;
-		while (!Client && Attempts++ < GStartConnectAttempts)
-		{
-			Client.Reset(UE::Trace::FStoreClient::Connect(TEXT("127.0.0.1")));
-			FPlatformProcess::Sleep(GStartConnectFrequencySeconds);
-		}
-		if (Client)
-		{
-			ChangeState(EState::Connecting, EState::Connected);
-		}
-		else
-		{
-			ChangeState(EState::Connecting, EState::NotConnected);
-			UE_LOG(LogTemp, Warning, TEXT("Failed to connect to store"));
-		}
-	}, TStatId(), {CommandTask});
 }
 
 /////////////////////////////////////////////////////////////////////////////////////////
 
 void STraceServerControl::OnStop_Clicked()
 {
-	FFunctionGraphTask::CreateAndDispatchWhenReady([this]
+	FGraphEventRef CommandTask = FFunctionGraphTask::CreateAndDispatchWhenReady([this]
 	{
+		if (bIsCancelRequested)
+		{
+			return;
+		}
+		FScopeLock _(&AsyncTaskLock);
+
 		if (ChangeState(EState::Connected, EState::Command, GStateChangeRetries))
 		{
 #if UE_TRACE_SERVER_CONTROLS_ENABLED
@@ -276,8 +319,14 @@ void STraceServerControl::OnStop_Clicked()
 
 void STraceServerControl::OnSponsored_Changed()
 {
-	FFunctionGraphTask::CreateAndDispatchWhenReady([this]
+	FGraphEventRef CommandTask = FFunctionGraphTask::CreateAndDispatchWhenReady([this]
 	{
+		if (bIsCancelRequested)
+		{
+			return;
+		}
+		FScopeLock _(&AsyncTaskLock);
+
 		if (ChangeState(EState::Connected, EState::Command, GStateChangeRetries))
 		{
 			const bool Success = Client->SetSponsored(!IsSponsored());
diff --git a/Engine/Source/Developer/TraceInsights/Public/Insights/Widgets/STraceServerControl.h b/Engine/Source/Developer/TraceInsights/Public/Insights/Widgets/STraceServerControl.h
--- a/Engine/Source/Developer/TraceInsights/Public/Insights/Widgets/STraceServerControl.h
+++ b/Engine/Source/Developer/TraceInsights/Public/Insights/Widgets/STraceServerControl.h
@@ -16,7 +16,7 @@ class TRACEINSIGHTS_API STraceServerControl
 {
 public:
 	STraceServerControl(const TCHAR* Host, uint32 Port = 0, FName StyleSet = NAME_None);
-	~STraceServerControl() = default;
+	~STraceServerControl();
 
 	void MakeMenu(FMenuBuilder& Builder);
 
@@ -36,9 +36,9 @@ private:
 	void UpdateStatus();
 	void ResetStatus();
 
-	bool CanServerBeStarted() const { return bIsLocalHost && State.load(std::memory_order_relaxed) == EState::NotConnected; }
-	bool CanServerBeStopped() const { return bIsLocalHost && State.load(std::memory_order_relaxed) == EState::Connected; }
-	bool AreControlsEnabled() const { return bIsLocalHost && State.load(std::memory_order_relaxed) == EState::Connected; }
+	bool CanServerBeStarted() const { return !bIsCancelRequested && bIsLocalHost && State.load(std::memory_order_relaxed) == EState::NotConnected; }
+	bool CanServerBeStopped() const { return !bIsCancelRequested && bIsLocalHost && State.load(std::memory_order_relaxed) == EState::Connected; }
+	bool AreControlsEnabled() const { return !bIsCancelRequested && bIsLocalHost && State.load(std::memory_order_relaxed) == EState::Connected; }
 	bool IsSponsored() const { return bSponsored.load(std::memory_order_relaxed); }
 
 	void OnStart_Clicked();
@@ -50,14 +50,17 @@ private:
 	std::atomic<bool> bCanServerBeStarted = false;
 	std::atomic<bool> bCanServerBeStopped = false;
 	std::atomic<bool> bSponsored = false;
+	std::atomic<bool> bIsCancelRequested = false;
+
+	FCriticalSection AsyncTaskLock;
 	FCriticalSection StringsLock;
 	FString StatusString;
 
 	FString Host;
-	uint32 Port;
+	uint32 Port = 0;
 	FName StyleSet;
-	bool bIsLocalHost;
+	bool bIsLocalHost = false;
 	TUniquePtr<UE::Trace::FStoreClient> Client;
-	
+
 	friend const TCHAR* LexState(EState);
 };
