[Insights] Further improved drawing of graph track series. Disabled antialias for lines and maintained 1px line thickness no matter of screen scale.
Also increased the minimum width of the context menu (to lower the resize effect for longer series names).

diff --git a/Engine/Source/Developer/TraceInsights/Private/Insights/ViewModels/GraphTrack.cpp b/Engine/Source/Developer/TraceInsights/Private/Insights/ViewModels/GraphTrack.cpp
--- a/Engine/Source/Developer/TraceInsights/Private/Insights/ViewModels/GraphTrack.cpp
+++ b/Engine/Source/Developer/TraceInsights/Private/Insights/ViewModels/GraphTrack.cpp
@@ -203,14 +203,16 @@ void FGraphTrack::Draw(const ITimingTrackDrawContext& Context) const
 
 void FGraphTrack::DrawSeries(const FGraphSeries& Series, FDrawContext& DrawContext, const FTimingTrackViewport& Viewport) const
 {
+	const float Scale = DrawContext.Geometry.GetAccumulatedLayoutTransform().GetScale();
+	const float PixelUnit = 1.0f / Scale;
+
 	// Set clipping area. This area tries to take into account the optional border area that allows
 	// graph tracks to optionally act like 'event' tracks if required (with respect to layout, anyway).
-	// The GetBorderY() - 1.0f calculation is desgiend to avoid clipping the line rasterization, as the custom verts
+	// The GetBorderY() - 1.0f calculation is designed to avoid clipping the line rasterization, as the custom verts
 	// of the fill and the outer lines appear to get rasterized differently, the latter missing one pixel
 	// on its upper side.
 	{
 		const FVector2f AbsPos = FVector2f(DrawContext.Geometry.GetAbsolutePosition());
-		const float Scale = DrawContext.Geometry.GetAccumulatedLayoutTransform().GetScale();
 		const float L = AbsPos.X;
 		const float R = AbsPos.X + (Viewport.GetWidth() * Scale);
 		const float T = AbsPos.Y + ((GetPosY() + (GetBorderY() - 1.0f)) * Scale);
@@ -239,10 +241,11 @@ void FGraphTrack::DrawSeries(const FGraphSeries& Series, FDrawContext& DrawConte
 		DrawContext.LayerId++;
 	}
 
+	const float LocalPosX = 0.0f;
 	const float LocalPosY = FMath::RoundToFloat(GetPosY());
 
 	FPaintGeometry Geo = DrawContext.Geometry.ToPaintGeometry();
-	Geo.AppendTransform(FSlateLayoutTransform(FVector2D(0.0f, LocalPosY * DrawContext.Geometry.GetAccumulatedLayoutTransform().GetScale())));
+	Geo.AppendTransform(FSlateLayoutTransform(FVector2D(LocalPosX * Scale, LocalPosY * Scale)));
 
 	if (IsAnyOptionEnabled(EGraphOptions::ShowPolygon))
 	{
@@ -352,17 +355,19 @@ void FGraphTrack::DrawSeries(const FGraphSeries& Series, FDrawContext& DrawConte
 
 	if (IsAnyOptionEnabled(EGraphOptions::ShowLines))
 	{
-		// Find scale to get the size of a pixel
-		const float Scale = DrawContext.Geometry.GetAccumulatedLayoutTransform().GetScale();
-
 		FPaintGeometry LineGeo = Geo;
-		LineGeo.AppendTransform(FSlateLayoutTransform(FVector2D(0.5f / Scale, 0.5f / Scale)));
+		LineGeo.AppendTransform(FSlateLayoutTransform(FVector2D(0.5f, 0.5f)));
 
 		// Disable pixel snapping here so lines line up with boxes/polys correctly.
 		const ESlateDrawEffect LineDrawEffects = DrawContext.DrawEffects | ESlateDrawEffect::NoPixelSnapping;
 
+#if 0
 		constexpr bool bAntialias = true;
 		const float Thickness = FMath::Max(1.0f, Scale);
+#else
+		constexpr bool bAntialias = false;
+		const float Thickness = 1.0f;
+#endif
 
 		for (int32 BatchIndex = 0; BatchIndex < Series.LinePoints.Num(); ++BatchIndex)
 		{
@@ -385,14 +390,17 @@ void FGraphTrack::DrawSeries(const FGraphSeries& Series, FDrawContext& DrawConte
 #define INSIGHTS_GRAPH_TRACK_DRAW_POINTS_AS_RECTANGLES 0
 #if !INSIGHTS_GRAPH_TRACK_DRAW_POINTS_AS_RECTANGLES
 
+		const float OffsetX = PixelUnit / 2.0f;
+		const float OffsetY = PixelUnit / 2.0f;
+
 		if (IsAnyOptionEnabled(EGraphOptions::ShowPointsWithBorder))
 		{
 			// Draw points (border).
 			for (int32 Index = 0; Index < NumPoints; ++Index)
 			{
 				const FVector2D& Pt = Series.Points[Index];
-				const float PtX = static_cast<float>(Pt.X) - PointVisualSize / 2.0f - 0.5f;
-				const float PtY = LocalPosY + static_cast<float>(Pt.Y) - PointVisualSize / 2.0f - 0.5f;
+				const float PtX = LocalPosX + static_cast<float>(Pt.X) - PointVisualSize / 2.0f - 1.0f + OffsetX;
+				const float PtY = LocalPosY + static_cast<float>(Pt.Y) - PointVisualSize / 2.0f - 1.0f + OffsetY;
 				DrawContext.DrawBox(PtX, PtY, PointVisualSize + 2.0f, PointVisualSize + 2.0f, PointBrush, Series.BorderColor);
 			}
 			DrawContext.LayerId++;
@@ -402,8 +410,8 @@ void FGraphTrack::DrawSeries(const FGraphSeries& Series, FDrawContext& DrawConte
 		for (int32 Index = 0; Index < NumPoints; ++Index)
 		{
 			const FVector2D& Pt = Series.Points[Index];
-			const float PtX = static_cast<float>(Pt.X) - PointVisualSize / 2.0f + 0.5f;
-			const float PtY = LocalPosY + static_cast<float>(Pt.Y) - PointVisualSize / 2.0f + 0.5f;
+			const float PtX = LocalPosX + static_cast<float>(Pt.X) - PointVisualSize / 2.0f + OffsetX;
+			const float PtY = LocalPosY + static_cast<float>(Pt.Y) - PointVisualSize / 2.0f + OffsetY;
 			DrawContext.DrawBox(PtX, PtY, PointVisualSize, PointVisualSize, PointBrush, Series.Color);
 		}
 		DrawContext.LayerId++;
@@ -420,7 +428,7 @@ void FGraphTrack::DrawSeries(const FGraphSeries& Series, FDrawContext& DrawConte
 			for (int32 Index = 0; Index < NumPoints; ++Index)
 			{
 				const FVector2D& Pt = Series.Points[Index];
-				const float PtX = static_cast<float>(Pt.X) - BorderPtSize / 2.0f + 0.5f;
+				const float PtX = LocalPosX + static_cast<float>(Pt.X) - BorderPtSize / 2.0f + 0.5f;
 				const float PtY = LocalPosY + static_cast<float>(Pt.Y) - BorderPtSize / 2.0f + 0.5f;
 				DrawContext.DrawBox(PtX, PtY, BorderPtSize, BorderPtSize, BorderBrush, Series.BorderColor);
 				//DrawContext.DrawRotatedBox(PtX, PtY, BorderPtSize, BorderPtSize, BorderBrush, Series.BorderColor, Angle, BorderRotationPoint);
@@ -434,7 +442,7 @@ void FGraphTrack::DrawSeries(const FGraphSeries& Series, FDrawContext& DrawConte
 		for (int32 Index = 0; Index < NumPoints; ++Index)
 		{
 			const FVector2D& Pt = Series.Points[Index];
-			const float PtX = static_cast<float>(Pt.X) - PtSize / 2.0f + 0.5f;
+			const float PtX = LocalPosX + static_cast<float>(Pt.X) - PtSize / 2.0f + 0.5f;
 			const float PtY = LocalPosY + static_cast<float>(Pt.Y) - PtSize / 2.0f + 0.5f;
 			DrawContext.DrawBox(PtX, PtY, PtSize, PtSize, WhiteBrush, Series.Color);
 			//DrawContext.DrawRotatedBox(PtX, PtY, PtSize, PtSize, WhiteBrush, Series.Color, Angle, RotationPoint);
@@ -446,36 +454,39 @@ void FGraphTrack::DrawSeries(const FGraphSeries& Series, FDrawContext& DrawConte
 
 	if (IsAnyOptionEnabled(EGraphOptions::ShowDebugInfo)) // for debugging only
 	{
-		// Find scale to get the size of a pixel
-		const float PixelUnit = 1.0f / DrawContext.Geometry.GetAccumulatedLayoutTransform().GetScale();
 		FPaintGeometry LineGeo = Geo;
-		LineGeo.AppendTransform(FSlateLayoutTransform(FVector2D(0.5f * PixelUnit, 0.5f * PixelUnit)));
+		LineGeo.AppendTransform(FSlateLayoutTransform(FVector2D(0.5f, 0.5f)));
 
 		// Disable pixel snapping here so lines line up with boxes/polys correctly.
 		const ESlateDrawEffect LineDrawEffects = DrawContext.DrawEffects | ESlateDrawEffect::NoPixelSnapping;
 
-		// Draw white corner at (0, 0) using MakeLines.
+		if (false)
 		{
-			TArray<FVector2f> HLine;
-			HLine.Add(FVector2f(0.0f, 0.0f));
-			HLine.Add(FVector2f(10.0f, 0.0f));
-			FSlateDrawElement::MakeLines(DrawContext.ElementList, DrawContext.LayerId, LineGeo, HLine, LineDrawEffects, FLinearColor::White, false, 1.0f);
+			// Draw white corner at (0, 0) using MakeLines.
+			{
+				TArray<FVector2f> HLine;
+				HLine.Add(FVector2f(0.0f, 0.0f));
+				HLine.Add(FVector2f(10.0f, 0.0f));
+				FSlateDrawElement::MakeLines(DrawContext.ElementList, DrawContext.LayerId, LineGeo, HLine, LineDrawEffects, FLinearColor::White, false, 1.0f);
+
+				TArray<FVector2f> VLine;
+				VLine.Add(FVector2f(0.0f, 0.0f));
+				VLine.Add(FVector2f(0.0f, 10.0f));
+				FSlateDrawElement::MakeLines(DrawContext.ElementList, DrawContext.LayerId, LineGeo, VLine, LineDrawEffects, FLinearColor::White, false, 1.0f);
+			}
+			DrawContext.LayerId++;
 
-			TArray<FVector2f> VLine;
-			VLine.Add(FVector2f(0.0f, 0.0f));
-			VLine.Add(FVector2f(0.0f, 10.0f));
-			FSlateDrawElement::MakeLines(DrawContext.ElementList, DrawContext.LayerId, LineGeo, VLine, LineDrawEffects, FLinearColor::White, false, 1.0f);
+			// Draw corner at (0, 0) using DrawBox.
+			DrawContext.DrawBox(LocalPosX + 2.0f, LocalPosY, 6.0f, PixelUnit, WhiteBrush, Series.Color);
+			DrawContext.DrawBox(LocalPosX, LocalPosY + 2.0f, PixelUnit, 6.0f, WhiteBrush, Series.Color);
+			DrawContext.LayerId++;
 		}
-		DrawContext.LayerId++;
-
-		// Draw corner at (0, 0) using DrawBox.
-		DrawContext.DrawBox(2.0f, LocalPosY, 6.0f, PixelUnit, WhiteBrush, Series.Color);
-		DrawContext.DrawBox(0.0f, LocalPosY + 2.0f, PixelUnit, 6.0f, WhiteBrush, Series.Color);
-		DrawContext.LayerId++;
 
 		const int32 NumDbgPoints = Series.Points.Num();
 
-		// Draw white cross lines (15x15) at each point.
+		float Thickness = 1.0f;
+
+		// Draw white cross lines (17x17 screen pixels) at each point.
 		for (int32 Index = 0; Index < NumDbgPoints; ++Index)
 		{
 			const FVector2D& Pt = Series.Points[Index];
@@ -483,25 +494,35 @@ void FGraphTrack::DrawSeries(const FGraphSeries& Series, FDrawContext& DrawConte
 			const float PtY = static_cast<float>(Pt.Y);
 
 			TArray<FVector2f> HLine;
-			HLine.Add(FVector2f(PtX - 7.0f, PtY));
-			HLine.Add(FVector2f(PtX + 8.0f, PtY));
-			FSlateDrawElement::MakeLines(DrawContext.ElementList, DrawContext.LayerId, LineGeo, HLine, LineDrawEffects, FLinearColor::White, false, 1.0f);
+			HLine.Add(FVector2f(PtX - 8.0f * PixelUnit, PtY));
+			HLine.Add(FVector2f(PtX + 9.0f * PixelUnit, PtY));
+			FSlateDrawElement::MakeLines(DrawContext.ElementList, DrawContext.LayerId, LineGeo, HLine, LineDrawEffects, FLinearColor::White, false, Thickness);
 
 			TArray<FVector2f> VLine;
-			VLine.Add(FVector2f(PtX, PtY - 7.0f));
-			VLine.Add(FVector2f(PtX, PtY + 8.0f));
-			FSlateDrawElement::MakeLines(DrawContext.ElementList, DrawContext.LayerId, LineGeo, VLine, LineDrawEffects, FLinearColor::White, false, 1.0f);
+			VLine.Add(FVector2f(PtX, PtY - 8.0f * PixelUnit));
+			VLine.Add(FVector2f(PtX, PtY + 9.0f * PixelUnit));
+			FSlateDrawElement::MakeLines(DrawContext.ElementList, DrawContext.LayerId, LineGeo, VLine, LineDrawEffects, FLinearColor::White, false, Thickness);
 		}
 		DrawContext.LayerId++;
 
-		// Draw cross lines (11x11) at each point.
+		// Draw black cross lines (11x11 screen pixels) at each point.
 		for (int32 Index = 0; Index < NumDbgPoints; ++Index)
 		{
 			const FVector2D& Pt = Series.Points[Index];
-			const float PtX = static_cast<float>(Pt.X);
+			const float PtX = LocalPosX + static_cast<float>(Pt.X);
+			const float PtY = LocalPosY + static_cast<float>(Pt.Y);
+			DrawContext.DrawBox(PtX - 5.0f * PixelUnit, PtY, 11.0f * PixelUnit, PixelUnit, WhiteBrush, FLinearColor::Black);
+			DrawContext.DrawBox(PtX, PtY - 5.0f * PixelUnit, PixelUnit, 11.0f * PixelUnit, WhiteBrush, FLinearColor::Black);
+		}
+		DrawContext.LayerId++;
+
+		// Draw a red screen pixel at each point.
+		for (int32 Index = 0; Index < NumDbgPoints; ++Index)
+		{
+			const FVector2D& Pt = Series.Points[Index];
+			const float PtX = LocalPosX + static_cast<float>(Pt.X);
 			const float PtY = LocalPosY + static_cast<float>(Pt.Y);
-			DrawContext.DrawBox(PtX - 5.0f, PtY, 11.0f, PixelUnit, WhiteBrush, Series.Color);
-			DrawContext.DrawBox(PtX, PtY - 5.0f, PixelUnit, 11.0f, WhiteBrush, Series.Color);
+			DrawContext.DrawBox(PtX, PtY, PixelUnit, PixelUnit, WhiteBrush, FLinearColor::Red);
 		}
 		DrawContext.LayerId++;
 	}
@@ -550,6 +571,9 @@ void FGraphTrack::DrawEvent(const ITimingTrackDrawContext& Context, const ITimin
 		DrawContext.LayerId++;
 	}
 
+	const float PX = EventX1;
+	const float PY = EventY1 - 0.5f;
+
 	// Draw highlighted line.
 	if ((EventX2 > EventX1) &&
 		(AreAllOptionsEnabled(EGraphOptions::UseEventDuration | EGraphOptions::ShowLines) ||
@@ -565,16 +589,16 @@ void FGraphTrack::DrawEvent(const ITimingTrackDrawContext& Context, const ITimin
 			W = 1.0f;
 		}
 
-		DrawContext.DrawBox(EventX1 - 1.0f, EventY1 - 1.0f, W + 2.0f, 3.0f, WhiteBrush, HighlightColor);
+		DrawContext.DrawBox(PX - 1.0f, PY - 1.0f, W + 2.0f, 3.0f, WhiteBrush, HighlightColor);
 		DrawContext.LayerId++;
-		DrawContext.DrawBox(EventX1, EventY1, W, 1.0f, WhiteBrush, Series->Color);
+		DrawContext.DrawBox(PX, PY, W, 1.0f, WhiteBrush, Series->Color);
 		DrawContext.LayerId++;
 	}
 
 	// Draw highlighted point.
-	DrawContext.DrawBox(EventX1 - PointVisualSize / 2.0f - 1.5f, EventY1 - PointVisualSize / 2.0f - 1.5f, PointVisualSize + 4.0f, PointVisualSize + 4.0f, PointBrush, HighlightColor);
+	DrawContext.DrawBox(PX - PointVisualSize / 2.0f - 1.5f, PY - PointVisualSize / 2.0f - 1.5f, PointVisualSize + 4.0f, PointVisualSize + 4.0f, PointBrush, HighlightColor);
 	DrawContext.LayerId++;
-	DrawContext.DrawBox(EventX1 - PointVisualSize / 2.0f + 0.5f, EventY1 - PointVisualSize / 2.0f + 0.5f, PointVisualSize, PointVisualSize, PointBrush, Series->Color);
+	DrawContext.DrawBox(PX - PointVisualSize / 2.0f + 0.5f, PY - PointVisualSize / 2.0f + 0.5f, PointVisualSize, PointVisualSize, PointBrush, Series->Color);
 	DrawContext.LayerId++;
 }
 
@@ -790,7 +814,8 @@ void FGraphTrack::BuildContextMenu(FMenuBuilder& MenuBuilder)
 	{
 		MenuBuilder.AddWidget(
 			SNew(SBox)
-			.MaxDesiredHeight(150.0f)
+			.MinDesiredWidth(250.0f)
+			.MaxDesiredHeight(135.0f)
 			[
 				SNew(SGraphSeriesList, SharedThis(this))
 			],
diff --git a/Engine/Source/Developer/TraceInsights/Public/Insights/ViewModels/GraphTrack.h b/Engine/Source/Developer/TraceInsights/Public/Insights/ViewModels/GraphTrack.h
--- a/Engine/Source/Developer/TraceInsights/Public/Insights/ViewModels/GraphTrack.h
+++ b/Engine/Source/Developer/TraceInsights/Public/Insights/ViewModels/GraphTrack.h
@@ -51,7 +51,7 @@ class TRACEINSIGHTS_API FGraphTrack : public FBaseTimingTrack
 
 private:
 	// Visual size of points (in pixels).
-	static constexpr float PointVisualSize = 5.0f;
+	static constexpr float PointVisualSize = 5.5f;
 
 	// Size of points (in pixels) used in reduction algorithm.
 	static constexpr double PointSizeX = 3.0f;
