[TRACE] Export memory snapshot.
- Added Package, AllocThreadName and FreeThreadName columns to the exported table
- Use filtered table result instead of all nodes as a datasource
#tests yes

diff --git a/Engine/Source/Developer/TraceInsights/Private/Insights/MemoryProfiler/Widgets/SMemAllocTableTreeView.cpp b/Engine/Source/Developer/TraceInsights/Private/Insights/MemoryProfiler/Widgets/SMemAllocTableTreeView.cpp
--- a/Engine/Source/Developer/TraceInsights/Private/Insights/MemoryProfiler/Widgets/SMemAllocTableTreeView.cpp
+++ b/Engine/Source/Developer/TraceInsights/Private/Insights/MemoryProfiler/Widgets/SMemAllocTableTreeView.cpp
@@ -2161,6 +2161,9 @@ void SMemAllocTableTreeView::ExportMemorySnapshot() const
 	const FTableColumn& TagColumn = *GetTable()->FindColumn(FMemAllocTableColumns::TagColumnId);
 	const FTableColumn& AssetColumn = *GetTable()->FindColumn(FMemAllocTableColumns::AssetColumnId);
 	const FTableColumn& ClassNameColumn = *GetTable()->FindColumn(FMemAllocTableColumns::ClassNameColumnId);
+	const FTableColumn& PackageColumn = *GetTable()->FindColumn(FMemAllocTableColumns::PackageColumnId);
+	const FTableColumn& AllocThreadColumn = *GetTable()->FindColumn(FMemAllocTableColumns::AllocThreadColumnId);
+	const FTableColumn& FreeThreadColumn = *GetTable()->FindColumn(FMemAllocTableColumns::FreeThreadColumnId);
 	const FTableColumn& AllocFunctionColumn = *GetTable()->FindColumn(FMemAllocTableColumns::AllocFunctionColumnId);
 	const FTableColumn& FreeFunctionColumn = *GetTable()->FindColumn(FMemAllocTableColumns::FreeFunctionColumnId);
 	const FTableColumn& AllocSourceFileColumn = *GetTable()->FindColumn(FMemAllocTableColumns::AllocSourceFileColumnId);
@@ -2168,7 +2171,7 @@ void SMemAllocTableTreeView::ExportMemorySnapshot() const
 
 	// 2. Iterate over TreeNodes
 	TStringBuilder<2048> Buffer;
-	for (const TSharedPtr<FTableTreeNode>& Node : TableRowNodes)
+	for (const TSharedPtr<FTableTreeNode>& Node : *FilteredNodesPtr)
 	{
 		// Export only leaves
 		if (Node->IsGroup()) continue;
@@ -2189,6 +2192,9 @@ void SMemAllocTableTreeView::ExportMemorySnapshot() const
 			WriteColumnHeader(Buffer, TagColumn); Buffer += Separator;
 			WriteColumnHeader(Buffer, AssetColumn); Buffer += Separator;
 			WriteColumnHeader(Buffer, ClassNameColumn); Buffer += Separator;
+			WriteColumnHeader(Buffer, PackageColumn); Buffer += Separator;
+			WriteColumnHeader(Buffer, AllocThreadColumn); Buffer += Separator;
+			WriteColumnHeader(Buffer, FreeThreadColumn); Buffer += Separator;
 			WriteColumnHeader(Buffer, AllocFunctionColumn); Buffer += Separator;
 			WriteColumnHeader(Buffer, FreeFunctionColumn); Buffer += Separator;
 			WriteColumnHeader(Buffer, AllocSourceFileColumn); Buffer += Separator;
@@ -2212,6 +2218,9 @@ void SMemAllocTableTreeView::ExportMemorySnapshot() const
 		WriteColumn(Buffer, TagColumn, *Node); Buffer += Separator;
 		WriteColumn(Buffer, AssetColumn, *Node); Buffer += Separator;
 		WriteColumn(Buffer, ClassNameColumn, *Node); Buffer += Separator;
+		WriteColumn(Buffer, PackageColumn, *Node); Buffer += Separator;
+		WriteColumn(Buffer, AllocThreadColumn, *Node); Buffer += Separator;
+		WriteColumn(Buffer, FreeThreadColumn, *Node); Buffer += Separator;
 		WriteColumn(Buffer, AllocFunctionColumn, *Node); Buffer += Separator;
 		WriteColumn(Buffer, FreeFunctionColumn, *Node); Buffer += Separator;
 		WriteColumn(Buffer, AllocSourceFileColumn, *Node); Buffer += Separator;
@@ -2231,7 +2240,7 @@ void SMemAllocTableTreeView::ExportMemorySnapshot() const
 
 bool SMemAllocTableTreeView::IsExportMemorySnapshotAvailable() const
 {
-	return !TableRowNodes.IsEmpty();
+	return FilteredNodesPtr && !FilteredNodesPtr->IsEmpty();
 }
 
 ////////////////////////////////////////////////////////////////////////////////////////////////////
