Properly flag Insights tests for Program context

#rnx

diff --git a/Engine/Source/Developer/TraceInsights/Private/Insights/Tests/FunctionalTests/ExportCommandsTests.cpp b/Engine/Source/Developer/TraceInsights/Private/Insights/Tests/FunctionalTests/ExportCommandsTests.cpp
--- a/Engine/Source/Developer/TraceInsights/Private/Insights/Tests/FunctionalTests/ExportCommandsTests.cpp
+++ b/Engine/Source/Developer/TraceInsights/Private/Insights/Tests/FunctionalTests/ExportCommandsTests.cpp
@@ -43,9 +43,9 @@ void VerifyExportedLines(const FString& ExportReportPath, const FString& CmdLogP
 	Test->AddError(FString::Printf(TEXT("VerifyExportedLines timed out while trying to find line '%s' from '%s'"), *ExpectedResult, *ExportReportPath));
 }
 
-#if !WITH_EDITOR
+#if WITH_AUTOMATION_TESTS
 
-IMPLEMENT_SIMPLE_AUTOMATION_TEST(FСommandsExportWindowsTest, "System.Insights.Trace.Analysis.ExecCmd.CommandsExport(Windows)", EAutomationTestFlags::ApplicationContextMask | EAutomationTestFlags::ProductFilter)
+IMPLEMENT_SIMPLE_AUTOMATION_TEST(FСommandsExportWindowsTest, "System.Insights.Trace.Analysis.ExecCmd.CommandsExport(Windows)", EAutomationTestFlags::ProgramContext | EAutomationTestFlags::EngineFilter)
 bool FСommandsExportWindowsTest::RunTest(const FString& Parameters)
 {
 	TSharedPtr<FInsightsManager> InsightsManager = FInsightsManager::Get();
@@ -217,4 +217,4 @@ bool FСommandsExportWindowsTest::RunTest(const FString& Parameters)
 	return true;
 }
 
-#endif //!WITH_EDITOR
+#endif //WITH_AUTOMATION_TESTS
diff --git a/Engine/Source/Developer/TraceInsights/Private/Insights/Tests/FunctionalTests/MemoryInsightsTests.cpp b/Engine/Source/Developer/TraceInsights/Private/Insights/Tests/FunctionalTests/MemoryInsightsTests.cpp
--- a/Engine/Source/Developer/TraceInsights/Private/Insights/Tests/FunctionalTests/MemoryInsightsTests.cpp
+++ b/Engine/Source/Developer/TraceInsights/Private/Insights/Tests/FunctionalTests/MemoryInsightsTests.cpp
@@ -25,9 +25,9 @@
 
 DECLARE_LOG_CATEGORY_EXTERN(MemoryInsightsTests, Log, All);
 
-#if !WITH_EDITOR
+#if WITH_AUTOMATION_TESTS
 
-BEGIN_DEFINE_SPEC(FAutomationDriverUnrealInsightsHubMemoryInsightsTest, "System.Insights.Hub.MemoryInsights", EAutomationTestFlags::ApplicationContextMask | EAutomationTestFlags::ProductFilter)
+BEGIN_DEFINE_SPEC(FAutomationDriverUnrealInsightsHubMemoryInsightsTest, "System.Insights.Hub.MemoryInsights", EAutomationTestFlags::ProgramContext | EAutomationTestFlags::EngineFilter)
 FAutomationDriverPtr Driver;
 TSharedPtr<SWindow> AutomationWindow;
 END_DEFINE_SPEC(FAutomationDriverUnrealInsightsHubMemoryInsightsTest)
@@ -111,13 +111,13 @@ void FAutomationDriverUnrealInsightsHubMemoryInsightsTest::Define()
 					const FString StoreTracePath = StoreDir / FString::Printf(TEXT("%s.utrace"), *TraceName);
 					const FString StoreCachePath = StoreDir / FString::Printf(TEXT("%s.ucache"), *TraceName);
 					const FString LogDirPath = ProjectDir / TEXT("TestResults");
-					const FString LogPath = ProjectDir / TEXT("TestResults/Log.txt");
+					const FString TestLogPath = ProjectDir / TEXT("TestResults/Log.txt");
 					const FString SuccessTestResult = TEXT("Test Completed. Result={Success}");
 
 					// Test live trace
-					FString TraceParameters = FString::Printf(TEXT("-InsightsTest -ABSLOG=\"%s\" -AutoQuit -ExecOnAnalysisCompleteCmd=\"Automation RunTests System.Insights.Trace.Analysis.MemoryInsights.UploadMemoryInsightsLLMXMLReportsTrace\" -OpenTraceFile=\"%s\""), *LogPath, *StoreTracePath);
+					FString TraceParameters = FString::Printf(TEXT("-InsightsTest -ABSLOG=\"%s\" -AutoQuit -ExecOnAnalysisCompleteCmd=\"Automation RunTests System.Insights.Trace.Analysis.MemoryInsights.UploadMemoryInsightsLLMXMLReportsTrace\" -OpenTraceFile=\"%s\""), *TestLogPath, *StoreTracePath);
 					InsightsManager->OpenUnrealInsights(*TraceParameters);
-					bool bLineFound = Utils.FileContainsString(LogPath, SuccessTestResult, 60.0f);
+					bool bLineFound = Utils.FileContainsString(TestLogPath, SuccessTestResult, 60.0f);
 					TestTrue("Test for live trace should pass", bLineFound);
 
 					IFileManager::Get().DeleteDirectory(*LogDirPath, false, true);
@@ -125,7 +125,7 @@ void FAutomationDriverUnrealInsightsHubMemoryInsightsTest::Define()
 
 					// Test stopped trace 
 					InsightsManager->OpenUnrealInsights(*TraceParameters);
-					bLineFound = Utils.FileContainsString(LogPath, SuccessTestResult, 60.0f);
+					bLineFound = Utils.FileContainsString(TestLogPath, SuccessTestResult, 60.0f);
 					TestTrue("Test for stopped trace should pass", bLineFound);
 
 					IFileManager::Get().DeleteDirectory(*LogDirPath, false, true);
@@ -140,7 +140,7 @@ void FAutomationDriverUnrealInsightsHubMemoryInsightsTest::Define()
 		});
 }
 
-IMPLEMENT_SIMPLE_AUTOMATION_TEST(FMemoryInsightsUploadLLMXMLReportsTraceTest, "System.Insights.Trace.Analysis.MemoryInsights.UploadMemoryInsightsLLMXMLReportsTrace", EAutomationTestFlags::ApplicationContextMask | EAutomationTestFlags::ProductFilter)
+IMPLEMENT_SIMPLE_AUTOMATION_TEST(FMemoryInsightsUploadLLMXMLReportsTraceTest, "System.Insights.Trace.Analysis.MemoryInsights.UploadMemoryInsightsLLMXMLReportsTrace", EAutomationTestFlags::ProgramContext | EAutomationTestFlags::EngineFilter)
 bool FMemoryInsightsUploadLLMXMLReportsTraceTest::RunTest(const FString& Parameters)
 {
 	const FString ReportGraphsXMLPath = FPaths::RootDir() / TEXT("EngineTest/SourceAssets/Utrace/ReportGraphs.xml");
@@ -149,7 +149,7 @@ bool FMemoryInsightsUploadLLMXMLReportsTraceTest::RunTest(const FString& Paramet
 	FMemorySharedState* SharedState = FMemoryProfilerManager::Get()->GetSharedState();
 	if (!SharedState)
 	{
-		AddError("ProfilerWindow should be valid. Please, run this test throught Insights Session automation tab");
+		AddError("ProfilerWindow should be valid. Please, run this test through Insights Session automation tab");
 		return false;
 	}
 
@@ -310,7 +310,7 @@ bool MemoryInsightsAllocationsQueryTableTest(const FString& Parameters, const TM
 	return true;
 }
 
-IMPLEMENT_COMPLEX_AUTOMATION_TEST(FMemoryInsightsAllocationsQueryTableEditorPackageTest, "System.Insights.Trace.Analysis.MemoryInsights.AllocationsQueryTable.Editor/Package", EAutomationTestFlags::ApplicationContextMask | EAutomationTestFlags::ProductFilter)
+IMPLEMENT_COMPLEX_AUTOMATION_TEST(FMemoryInsightsAllocationsQueryTableEditorPackageTest, "System.Insights.Trace.Analysis.MemoryInsights.AllocationsQueryTable.Editor/Package", EAutomationTestFlags::ProgramContext | EAutomationTestFlags::EngineFilter)
 bool FMemoryInsightsAllocationsQueryTableEditorPackageTest::RunTest(const FString& Parameters)
 {
 	bool bSuccess = MemoryInsightsAllocationsQueryTableTest(Parameters, AllocsTimeMarkerEditorPackageGetterMap, this);
@@ -340,7 +340,7 @@ void FMemoryInsightsAllocationsQueryTableEditorPackageTest::GetTests(TArray<FStr
 	}
 }
 
-IMPLEMENT_COMPLEX_AUTOMATION_TEST(FMemoryInsightsAllocationsQueryTableStandaloneTest, "System.Insights.Trace.Analysis.MemoryInsights.AllocationsQueryTable.Standalone", EAutomationTestFlags::ApplicationContextMask | EAutomationTestFlags::ProductFilter)
+IMPLEMENT_COMPLEX_AUTOMATION_TEST(FMemoryInsightsAllocationsQueryTableStandaloneTest, "System.Insights.Trace.Analysis.MemoryInsights.AllocationsQueryTable.Standalone", EAutomationTestFlags::ProgramContext | EAutomationTestFlags::EngineFilter)
 bool FMemoryInsightsAllocationsQueryTableStandaloneTest::RunTest(const FString& Parameters)
 {
 	bool bSuccess = MemoryInsightsAllocationsQueryTableTest(Parameters, AllocsTimeMarkerStandaloneGameGetterMap, this);
@@ -371,4 +371,4 @@ void FMemoryInsightsAllocationsQueryTableStandaloneTest::GetTests(TArray<FString
 	}
 }
 
-#endif //!WITH_EDITOR
\ No newline at end of file
+#endif //WITH_AUTOMATION_TESTS
\ No newline at end of file
diff --git a/Engine/Source/Developer/TraceInsights/Private/Insights/Tests/FunctionalTests/SessionBrowserTests.cpp b/Engine/Source/Developer/TraceInsights/Private/Insights/Tests/FunctionalTests/SessionBrowserTests.cpp
--- a/Engine/Source/Developer/TraceInsights/Private/Insights/Tests/FunctionalTests/SessionBrowserTests.cpp
+++ b/Engine/Source/Developer/TraceInsights/Private/Insights/Tests/FunctionalTests/SessionBrowserTests.cpp
@@ -15,9 +15,9 @@
 
 DECLARE_LOG_CATEGORY_EXTERN(UnrealInsightsHubTests, Log, All);
 
-#if !WITH_EDITOR
+#if WITH_AUTOMATION_TESTS
 
-BEGIN_DEFINE_SPEC(FAutomationDriverUnrealInsightsSessionBrowserTest, "System.Insights.Hub.SessionBrowser", EAutomationTestFlags::ApplicationContextMask | EAutomationTestFlags::ProductFilter)
+BEGIN_DEFINE_SPEC(FAutomationDriverUnrealInsightsSessionBrowserTest, "System.Insights.Hub.SessionBrowser", EAutomationTestFlags::ProgramContext | EAutomationTestFlags::EngineFilter)
 FAutomationDriverPtr Driver;
 TSharedPtr<SWindow> AutomationWindow;
 END_DEFINE_SPEC(FAutomationDriverUnrealInsightsSessionBrowserTest)
@@ -145,4 +145,4 @@ void FAutomationDriverUnrealInsightsSessionBrowserTest::Define()
 		});
 }
 
-#endif //!WITH_EDITOR
\ No newline at end of file
+#endif //WITH_AUTOMATION_TESTS
\ No newline at end of file
diff --git a/Engine/Source/Developer/TraceInsights/Private/Insights/Tests/FunctionalTests/TimingProfilerTests.cpp b/Engine/Source/Developer/TraceInsights/Private/Insights/Tests/FunctionalTests/TimingProfilerTests.cpp
--- a/Engine/Source/Developer/TraceInsights/Private/Insights/Tests/FunctionalTests/TimingProfilerTests.cpp
+++ b/Engine/Source/Developer/TraceInsights/Private/Insights/Tests/FunctionalTests/TimingProfilerTests.cpp
@@ -16,9 +16,9 @@ DEFINE_LOG_CATEGORY(TimingProfilerTests);
 
 ////////////////////////////////////////////////////////////////////////////////////////////////////
 
-#if !WITH_EDITOR && WITH_DEV_AUTOMATION_TESTS
+#if WITH_AUTOMATION_TESTS
 
-IMPLEMENT_SIMPLE_AUTOMATION_TEST(FEnumerateTest, "System.Insights.Trace.Analysis.TimingInsights.Enumerate", EAutomationTestFlags::ApplicationContextMask | EAutomationTestFlags::ProductFilter)
+IMPLEMENT_SIMPLE_AUTOMATION_TEST(FEnumerateTest, "System.Insights.Analysis.TimingInsights.Enumerate", EAutomationTestFlags::ApplicationContextMask | EAutomationTestFlags::EngineFilter)
 bool FEnumerateTest::RunTest(const FString& Parameters)
 {
 	FInsightsTestUtils Utils(this);
@@ -48,7 +48,7 @@ bool FEnumerateTest::RunTest(const FString& Parameters)
 	return !HasAnyErrors();
 }
 
-IMPLEMENT_SIMPLE_AUTOMATION_TEST(FEnumeratePerformanceTest, "System.Insights.Trace.Analysis.TimingInsights.EnumeratePerformance", EAutomationTestFlags::ApplicationContextMask | EAutomationTestFlags::PerfFilter)
+IMPLEMENT_SIMPLE_AUTOMATION_TEST(FEnumeratePerformanceTest, "System.Insights.Analysis.TimingInsights.EnumeratePerformance", EAutomationTestFlags::ApplicationContextMask | EAutomationTestFlags::PerfFilter)
 bool FEnumeratePerformanceTest::RunTest(const FString& Parameters)
 {
 	FInsightsTestUtils Utils(this);
@@ -88,7 +88,7 @@ bool FEnumeratePerformanceTest::RunTest(const FString& Parameters)
 	return !HasAnyErrors();
 }
 
-IMPLEMENT_SIMPLE_AUTOMATION_TEST(FEnumerateFastTest, "System.Insights.Analysis.TimingInsights.EnumerateFast", EAutomationTestFlags::ApplicationContextMask | EAutomationTestFlags::ProductFilter)
+IMPLEMENT_SIMPLE_AUTOMATION_TEST(FEnumerateFastTest, "System.Insights.Analysis.TimingInsights.EnumerateFast", EAutomationTestFlags::ApplicationContextMask | EAutomationTestFlags::EngineFilter)
 bool FEnumerateFastTest::RunTest(const FString& Parameters)
 {
 	FInsightsTestUtils Utils(this);
@@ -118,7 +118,7 @@ bool FEnumerateFastTest::RunTest(const FString& Parameters)
 	return !HasAnyErrors();
 }
 
-IMPLEMENT_SIMPLE_AUTOMATION_TEST(EnumerateEventsToFile, "System.Insights.Analysis.TimingInsights.EnumerateEventsToFile", EAutomationTestFlags::ApplicationContextMask | EAutomationTestFlags::ProductFilter)
+IMPLEMENT_SIMPLE_AUTOMATION_TEST(EnumerateEventsToFile, "System.Insights.Trace.Analysis.TimingInsights.EnumerateEventsToFile", EAutomationTestFlags::ProgramContext | EAutomationTestFlags::EngineFilter)
 bool EnumerateEventsToFile::RunTest(const FString& Parameters)
 {
 	double SessionTime = 0.0;
@@ -159,7 +159,7 @@ bool EnumerateEventsToFile::RunTest(const FString& Parameters)
 	return !HasAnyErrors();
 }
 
-IMPLEMENT_SIMPLE_AUTOMATION_TEST(EnumerateScopesToFile, "System.Insights.Trace.Analysis.TimingInsights.EnumerateScopesToFile", EAutomationTestFlags::ApplicationContextMask | EAutomationTestFlags::ProductFilter)
+IMPLEMENT_SIMPLE_AUTOMATION_TEST(EnumerateScopesToFile, "System.Insights.Trace.Analysis.TimingInsights.EnumerateScopesToFile", EAutomationTestFlags::ProgramContext | EAutomationTestFlags::EngineFilter)
 bool EnumerateScopesToFile::RunTest(const FString& Parameters)
 {
 	double SessionTime = 0.0;
@@ -200,7 +200,7 @@ bool EnumerateScopesToFile::RunTest(const FString& Parameters)
 	return !HasAnyErrors();
 }
 
-IMPLEMENT_SIMPLE_AUTOMATION_TEST(InsightsEnumerate10K, "System.Insights.Trace.Analysis.TimingInsights.Enumerate10K", EAutomationTestFlags::ApplicationContextMask | EAutomationTestFlags::ProductFilter)
+IMPLEMENT_SIMPLE_AUTOMATION_TEST(InsightsEnumerate10K, "System.Insights.Trace.Analysis.TimingInsights.Enumerate10K", EAutomationTestFlags::ProgramContext | EAutomationTestFlags::EngineFilter)
 bool InsightsEnumerate10K::RunTest(const FString& Parameters)
 {
 	FTimingProfilerTests::FEnumerateTestParams Params;
@@ -213,7 +213,7 @@ bool InsightsEnumerate10K::RunTest(const FString& Parameters)
 	return !HasAnyErrors();
 }
 
-IMPLEMENT_SIMPLE_AUTOMATION_TEST(InsightsEnumerate100K, "System.Insights.Trace.Analysis.TimingInsights.Enumerate100K", EAutomationTestFlags::ApplicationContextMask | EAutomationTestFlags::ProductFilter)
+IMPLEMENT_SIMPLE_AUTOMATION_TEST(InsightsEnumerate100K, "System.Insights.Trace.Analysis.TimingInsights.Enumerate100K", EAutomationTestFlags::ProgramContext | EAutomationTestFlags::EngineFilter)
 bool InsightsEnumerate100K::RunTest(const FString& Parameters)
 {
 	FTimingProfilerTests::FEnumerateTestParams Params;
@@ -226,7 +226,7 @@ bool InsightsEnumerate100K::RunTest(const FString& Parameters)
 	return !HasAnyErrors();
 }
 
-IMPLEMENT_SIMPLE_AUTOMATION_TEST(EnumerateByEndTimeAsyncAllTracks, "System.Insights.Trace.Analysis.TimingInsights.EnumerateByEndTimeAsyncAllTracks", EAutomationTestFlags::ApplicationContextMask | EAutomationTestFlags::ProductFilter)
+IMPLEMENT_SIMPLE_AUTOMATION_TEST(EnumerateByEndTimeAsyncAllTracks, "System.Insights.Trace.Analysis.TimingInsights.EnumerateByEndTimeAsyncAllTracks", EAutomationTestFlags::ProgramContext | EAutomationTestFlags::EngineFilter)
 bool EnumerateByEndTimeAsyncAllTracks::RunTest(const FString& Parameters)
 {
 	FTimingProfilerTests::FEnumerateTestParams Params;
@@ -239,7 +239,7 @@ bool EnumerateByEndTimeAsyncAllTracks::RunTest(const FString& Parameters)
 	return !HasAnyErrors();
 }
 
-IMPLEMENT_SIMPLE_AUTOMATION_TEST(EnumerateByEndTimeAsyncGameThreadTrack, "System.Insights.Trace.Analysis.TimingInsights.EnumerateByEndTimeAsyncGameThreadTrack", EAutomationTestFlags::ApplicationContextMask | EAutomationTestFlags::ProductFilter)
+IMPLEMENT_SIMPLE_AUTOMATION_TEST(EnumerateByEndTimeAsyncGameThreadTrack, "System.Insights.Trace.Analysis.TimingInsights.EnumerateByEndTimeAsyncGameThreadTrack", EAutomationTestFlags::ProgramContext | EAutomationTestFlags::EngineFilter)
 bool EnumerateByEndTimeAsyncGameThreadTrack::RunTest(const FString& Parameters)
 {
 	FTimingProfilerTests::FEnumerateTestParams Params;
@@ -252,7 +252,7 @@ bool EnumerateByEndTimeAsyncGameThreadTrack::RunTest(const FString& Parameters)
 	return !HasAnyErrors();
 }
 
-IMPLEMENT_SIMPLE_AUTOMATION_TEST(EnumerateByEndTimeAllTracks10sIntervals, "System.Insights.Trace.Analysis.TimingInsights.EnumerateByEndTimeAllTracks10sIntervals", EAutomationTestFlags::ApplicationContextMask | EAutomationTestFlags::ProductFilter)
+IMPLEMENT_SIMPLE_AUTOMATION_TEST(EnumerateByEndTimeAllTracks10sIntervals, "System.Insights.Trace.Analysis.TimingInsights.EnumerateByEndTimeAllTracks10sIntervals", EAutomationTestFlags::ProgramContext | EAutomationTestFlags::EngineFilter)
 bool EnumerateByEndTimeAllTracks10sIntervals::RunTest(const FString& Parameters)
 {
 	FTimingProfilerTests::FEnumerateTestParams Params;
@@ -265,7 +265,7 @@ bool EnumerateByEndTimeAllTracks10sIntervals::RunTest(const FString& Parameters)
 	return !HasAnyErrors();
 }
 
-IMPLEMENT_SIMPLE_AUTOMATION_TEST(EnumerateByEndTimeAllTracks5sIntervals, "System.Insights.Trace.Analysis.TimingInsights.EnumerateByEndTimeAllTracks5sIntervals", EAutomationTestFlags::ApplicationContextMask | EAutomationTestFlags::ProductFilter)
+IMPLEMENT_SIMPLE_AUTOMATION_TEST(EnumerateByEndTimeAllTracks5sIntervals, "System.Insights.Trace.Analysis.TimingInsights.EnumerateByEndTimeAllTracks5sIntervals", EAutomationTestFlags::ProgramContext | EAutomationTestFlags::EngineFilter)
 bool EnumerateByEndTimeAllTracks5sIntervals::RunTest(const FString& Parameters)
 {
 	FTimingProfilerTests::FEnumerateTestParams Params;
@@ -282,7 +282,7 @@ bool EnumerateByEndTimeAllTracks5sIntervals::RunTest(const FString& Parameters)
 ////////////////////////////Enumerate Async Ordered by Start Time Tests/////////////////////////////
 ////////////////////////////////////////////////////////////////////////////////////////////////////
 
-IMPLEMENT_SIMPLE_AUTOMATION_TEST(EnumerateByStartTimeAsyncAllTracks, "System.Insights.Trace.Analysis.TimingInsights.EnumerateByStartTimeAsyncAllTracks", EAutomationTestFlags::ApplicationContextMask | EAutomationTestFlags::ProductFilter)
+IMPLEMENT_SIMPLE_AUTOMATION_TEST(EnumerateByStartTimeAsyncAllTracks, "System.Insights.Trace.Analysis.TimingInsights.EnumerateByStartTimeAsyncAllTracks", EAutomationTestFlags::ProgramContext | EAutomationTestFlags::EngineFilter)
 bool EnumerateByStartTimeAsyncAllTracks::RunTest(const FString& Parameters)
 {
 	FTimingProfilerTests::FEnumerateTestParams Params;
@@ -295,7 +295,7 @@ bool EnumerateByStartTimeAsyncAllTracks::RunTest(const FString& Parameters)
 	return !HasAnyErrors();
 }
 
-IMPLEMENT_SIMPLE_AUTOMATION_TEST(EnumerateByStartTimeAsyncGameThreadTrack, "System.Insights.Trace.Analysis.TimingInsights.EnumerateByStartTimeAsyncGameThreadTrack", EAutomationTestFlags::ApplicationContextMask | EAutomationTestFlags::ProductFilter)
+IMPLEMENT_SIMPLE_AUTOMATION_TEST(EnumerateByStartTimeAsyncGameThreadTrack, "System.Insights.Trace.Analysis.TimingInsights.EnumerateByStartTimeAsyncGameThreadTrack", EAutomationTestFlags::ProgramContext | EAutomationTestFlags::EngineFilter)
 bool EnumerateByStartTimeAsyncGameThreadTrack::RunTest(const FString& Parameters)
 {
 	FTimingProfilerTests::FEnumerateTestParams Params;
@@ -308,7 +308,7 @@ bool EnumerateByStartTimeAsyncGameThreadTrack::RunTest(const FString& Parameters
 	return !HasAnyErrors();
 }
 
-IMPLEMENT_SIMPLE_AUTOMATION_TEST(EnumerateByStartTimeAllTracks10sIntervals, "System.Insights.Trace.Analysis.TimingInsights.EnumerateByStartTimeAllTracks10sIntervals", EAutomationTestFlags::ApplicationContextMask | EAutomationTestFlags::ProductFilter)
+IMPLEMENT_SIMPLE_AUTOMATION_TEST(EnumerateByStartTimeAllTracks10sIntervals, "System.Insights.Trace.Analysis.TimingInsights.EnumerateByStartTimeAllTracks10sIntervals", EAutomationTestFlags::ProgramContext | EAutomationTestFlags::EngineFilter)
 bool EnumerateByStartTimeAllTracks10sIntervals::RunTest(const FString& Parameters)
 {
 	FTimingProfilerTests::FEnumerateTestParams Params;
@@ -321,7 +321,7 @@ bool EnumerateByStartTimeAllTracks10sIntervals::RunTest(const FString& Parameter
 	return !HasAnyErrors();
 }
 
-IMPLEMENT_SIMPLE_AUTOMATION_TEST(EnumerateByStartTimeAllTracks5sIntervals, "System.Insights.Trace.Analysis.TimingInsights.EnumerateByStartTimeAllTracks5sIntervals", EAutomationTestFlags::ApplicationContextMask | EAutomationTestFlags::ProductFilter)
+IMPLEMENT_SIMPLE_AUTOMATION_TEST(EnumerateByStartTimeAllTracks5sIntervals, "System.Insights.Trace.Analysis.TimingInsights.EnumerateByStartTimeAllTracks5sIntervals", EAutomationTestFlags::ProgramContext | EAutomationTestFlags::EngineFilter)
 bool EnumerateByStartTimeAllTracks5sIntervals::RunTest(const FString& Parameters)
 {
 	FTimingProfilerTests::FEnumerateTestParams Params;
@@ -333,7 +333,7 @@ bool EnumerateByStartTimeAllTracks5sIntervals::RunTest(const FString& Parameters
 
 	return !HasAnyErrors();
 }
-#endif //!WITH_EDITOR && WITH_DEV_AUTOMATION_TESTS
+#endif //WITH_AUTOMATION_TESTS
 
 void FTimingProfilerTests::RunEnumerateBenchmark(const FEnumerateTestParams& InParams, FCheckValues& OutCheckValues)
 {
diff --git a/Engine/Source/Developer/TraceInsights/Private/Insights/Tests/FunctionalTests/UITests.cpp b/Engine/Source/Developer/TraceInsights/Private/Insights/Tests/FunctionalTests/UITests.cpp
--- a/Engine/Source/Developer/TraceInsights/Private/Insights/Tests/FunctionalTests/UITests.cpp
+++ b/Engine/Source/Developer/TraceInsights/Private/Insights/Tests/FunctionalTests/UITests.cpp
@@ -14,7 +14,7 @@
 
 DECLARE_LOG_CATEGORY_EXTERN(UITests, Log, All);
 
-IMPLEMENT_SIMPLE_AUTOMATION_TEST(FHideAndShowAllTimingViewTabs, "System.Insights.Trace.Analysis.TimingInsights.HideAndShowAllTimingViewTabs", EAutomationTestFlags::ApplicationContextMask | EAutomationTestFlags::ProductFilter)
+IMPLEMENT_SIMPLE_AUTOMATION_TEST(FHideAndShowAllTimingViewTabs, "System.Insights.Analysis.TimingInsights.HideAndShowAllTimingViewTabs", EAutomationTestFlags::ApplicationContextMask | EAutomationTestFlags::EngineFilter)
 bool FHideAndShowAllTimingViewTabs::RunTest(const FString& Parameters)
 {
 	TSharedPtr<FTimingProfilerManager> TimingProfilerManager = FTimingProfilerManager::Get();
@@ -36,7 +36,7 @@ bool FHideAndShowAllTimingViewTabs::RunTest(const FString& Parameters)
 	return true;
 }
 
-IMPLEMENT_SIMPLE_AUTOMATION_TEST(FMemoryFilterValueConverterTest, "System.Insights.Trace.Analysis.MemoryFilterValueConverter", EAutomationTestFlags::ApplicationContextMask | EAutomationTestFlags::ProductFilter)
+IMPLEMENT_SIMPLE_AUTOMATION_TEST(FMemoryFilterValueConverterTest, "System.Insights.Analysis.MemoryFilterValueConverter", EAutomationTestFlags::ApplicationContextMask | EAutomationTestFlags::EngineFilter)
 bool FMemoryFilterValueConverterTest::RunTest(const FString& Parameters)
 {
 	Insights::FMemoryFilterValueConverter Converter;
@@ -72,7 +72,7 @@ bool FMemoryFilterValueConverterTest::RunTest(const FString& Parameters)
 	return !HasAnyErrors();
 }
 
-IMPLEMENT_SIMPLE_AUTOMATION_TEST(FTimeFilterValueConverterTest, "System.Insights.Trace.Analysis.TimeFilterValueConverter", EAutomationTestFlags::ApplicationContextMask | EAutomationTestFlags::ProductFilter)
+IMPLEMENT_SIMPLE_AUTOMATION_TEST(FTimeFilterValueConverterTest, "System.Insights.Analysis.TimeFilterValueConverter", EAutomationTestFlags::ApplicationContextMask | EAutomationTestFlags::EngineFilter)
 bool FTimeFilterValueConverterTest::RunTest(const FString& Parameters)
 {
 	Insights::FTimeFilterValueConverter Converter;
diff --git a/Engine/Source/Developer/TraceInsights/Private/Insights/Tests/FunctionalTests/UnrealTraceServer.cpp b/Engine/Source/Developer/TraceInsights/Private/Insights/Tests/FunctionalTests/UnrealTraceServer.cpp
--- a/Engine/Source/Developer/TraceInsights/Private/Insights/Tests/FunctionalTests/UnrealTraceServer.cpp
+++ b/Engine/Source/Developer/TraceInsights/Private/Insights/Tests/FunctionalTests/UnrealTraceServer.cpp
@@ -8,9 +8,9 @@
 #include "Logging/LogMacros.h"
 #include "Misc/AutomationTest.h"
 
-#if !WITH_EDITOR
+#if WITH_AUTOMATION_TESTS
 
-IMPLEMENT_SIMPLE_AUTOMATION_TEST(FUnrealTraceServerStartingTest, "System.Insights.Hub.UnrealTraceServer.Starting", EAutomationTestFlags::ApplicationContextMask | EAutomationTestFlags::ProductFilter)
+IMPLEMENT_SIMPLE_AUTOMATION_TEST(FUnrealTraceServerStartingTest, "System.Insights.Hub.UnrealTraceServer.Starting", EAutomationTestFlags::ProgramContext | EAutomationTestFlags::EngineFilter)
 bool FUnrealTraceServerStartingTest::RunTest(const FString& Parameters)
 {
 	TSharedPtr<FInsightsManager> InsightsManager = FInsightsManager::Get();
@@ -113,4 +113,4 @@ bool FUnrealTraceServerStartingTest::RunTest(const FString& Parameters)
 	return true;
 }
 
-#endif //!WITH_EDITOR
\ No newline at end of file
+#endif //WITH_AUTOMATION_TESTS
\ No newline at end of file
diff --git a/Engine/Source/Developer/TraceServices/Private/Tests/UnitTests.cpp b/Engine/Source/Developer/TraceServices/Private/Tests/UnitTests.cpp
--- a/Engine/Source/Developer/TraceServices/Private/Tests/UnitTests.cpp
+++ b/Engine/Source/Developer/TraceServices/Private/Tests/UnitTests.cpp
@@ -4,7 +4,7 @@
 #include "Common/PagedArray.h"
 #include "Misc/AutomationTest.h"
 
-IMPLEMENT_SIMPLE_AUTOMATION_TEST(FPagedArrayFilteringTest, "System.Insights.Trace.Analysis.PagedArrayFiltering", EAutomationTestFlags::ApplicationContextMask | EAutomationTestFlags::SmokeFilter)
+IMPLEMENT_SIMPLE_AUTOMATION_TEST(FPagedArrayFilteringTest, "System.Insights.Analysis.PagedArrayFiltering", EAutomationTestFlags::ApplicationContextMask | EAutomationTestFlags::SmokeFilter)
 
 bool operator!=(const FInt32Interval& lhs, const FInt32Interval& rhs)
 {
