Trace autostart for Mac/Linux

Implements trace auto-start/auto-connect for Linux and Mac platforms. Applications can now automatically start tracing whenever Unreal Insights application is running. The behaviour can be controlled using the UI in Insights session browser or by starting the application with `-traceautostart=x`.

diff --git a/Engine/Source/Developer/TraceInsights/Private/Insights/Widgets/SStartPageWindow.cpp b/Engine/Source/Developer/TraceInsights/Private/Insights/Widgets/SStartPageWindow.cpp
--- a/Engine/Source/Developer/TraceInsights/Private/Insights/Widgets/SStartPageWindow.cpp
+++ b/Engine/Source/Developer/TraceInsights/Private/Insights/Widgets/SStartPageWindow.cpp
@@ -923,6 +923,7 @@ STraceStoreWindow::~STraceStoreWindow()
 		FEngineAnalytics::GetProvider().RecordEvent(TEXT("Insights.Usage.SessionBrowser"), FAnalyticsEventAttribute(TEXT("Duration"), DurationActive));
 	}
 #endif // WITH_EDITOR
+	DisableAutoConnect();
 }
 
 ////////////////////////////////////////////////////////////////////////////////////////////////////
@@ -1926,7 +1927,6 @@ TSharedRef<SWidget> STraceStoreWindow::ConstructAutoStartPanel()
 		]
 	];
 
-#if PLATFORM_WINDOWS
 	Box->AddSlot()
 		.AutoWidth()
 		.Padding(6.0f, 0.0f, 0.0f, 0.0f)
@@ -1952,7 +1952,6 @@ TSharedRef<SWidget> STraceStoreWindow::ConstructAutoStartPanel()
 				.Text(LOCTEXT("AutoConnect_Text", "Auto-connect"))
 			]
 		];
-#endif // PLATFORM_WINDOWS
 
 	return Box;
 }
@@ -2856,6 +2855,14 @@ void STraceStoreWindow::EnableAutoConnect()
 	{
 		UE_LOG(TraceInsights, Warning, TEXT("[TraceStore] Failed to create AutoConnect event."));
 	}
+#elif PLATFORM_MAC || PLATFORM_LINUX
+    ensure(AutoConnectEvent == SEM_FAILED);
+	sem_unlink("/UnrealInsightsAutoConnect");
+    AutoConnectEvent = sem_open("/UnrealInsightsAutoConnect", O_CREAT | O_WRONLY | O_EXCL, 0644, 1);
+    if (AutoConnectEvent == SEM_FAILED)
+    {
+        UE_LOG(TraceInsights, Warning, TEXT("[TraceStore] Failed to create AutoConnect semaphore: %d"), errno);
+    }
 #endif
 }
 
@@ -2869,6 +2876,16 @@ void STraceStoreWindow::DisableAutoConnect()
 		CloseHandle(AutoConnectEvent);
 		AutoConnectEvent = nullptr;
 	}
+#elif PLATFORM_MAC || PLATFORM_LINUX
+    if (AutoConnectEvent != SEM_FAILED)
+    {
+    	sem_close(AutoConnectEvent);
+        AutoConnectEvent = SEM_FAILED;
+        if (sem_unlink("/UnrealInsightsAutoConnect"))
+        {
+	        UE_LOG(TraceInsights, Warning, TEXT("[TraceStore] Failed to remove AutoConnect semaphore: %d"), errno);
+        }
+    }
 #endif
 }
 
diff --git a/Engine/Source/Developer/TraceInsights/Private/Insights/Widgets/SStartPageWindow.h b/Engine/Source/Developer/TraceInsights/Private/Insights/Widgets/SStartPageWindow.h
--- a/Engine/Source/Developer/TraceInsights/Private/Insights/Widgets/SStartPageWindow.h
+++ b/Engine/Source/Developer/TraceInsights/Private/Insights/Widgets/SStartPageWindow.h
@@ -13,6 +13,10 @@
 #include "Widgets/SOverlay.h"
 #include "Widgets/Views/SListView.h"
 
+#if PLATFORM_MAC || PLATFORM_LINUX
+#include <semaphore.h>
+#endif
+
 ////////////////////////////////////////////////////////////////////////////////////////////////////
 
 class FActiveTimerHandle;
@@ -685,6 +689,8 @@ private:
 
 #if PLATFORM_WINDOWS
 	void* AutoConnectEvent = nullptr;
+#elif PLATFORM_MAC || PLATFORM_LINUX
+    sem_t* AutoConnectEvent = SEM_FAILED;
 #endif
 };
 
diff --git a/Engine/Source/Runtime/Core/Private/ProfilingDebugging/TraceAuxiliary.cpp b/Engine/Source/Runtime/Core/Private/ProfilingDebugging/TraceAuxiliary.cpp
--- a/Engine/Source/Runtime/Core/Private/ProfilingDebugging/TraceAuxiliary.cpp
+++ b/Engine/Source/Runtime/Core/Private/ProfilingDebugging/TraceAuxiliary.cpp
@@ -13,6 +13,7 @@
 #endif
 #if PLATFORM_LINUX || PLATFORM_MAC
 # 	include <sys/wait.h>
+#   include <semaphore.h>
 #endif
 
 #if defined(WITH_UNREAL_TRACE_LAUNCH)
@@ -1741,7 +1742,18 @@ void FTraceAuxiliary::TryAutoConnect()
 			::CloseHandle(KnownEvent);
 		}
 	}
-#endif // PLATFORM_WINDOWS
+#elif PLATFORM_MAC || PLATFORM_LINUX
+    if (GTraceAutoStart && !IsConnected())
+    {
+        sem_t* AutoConnectSemaphore = sem_open("/UnrealInsightsAutoConnect", O_RDONLY);
+        if (AutoConnectSemaphore != SEM_FAILED)
+        {
+            UE_LOG(LogTrace, Display, TEXT("Unreal Insights instance detected, auto-connecting to local trace server..."));
+            Start(EConnectionType::Network, TEXT("127.0.0.1"), GTraceAuxiliary.HasCommandlineChannels() ? nullptr : TEXT("default"), nullptr);
+            sem_close(AutoConnectSemaphore);
+        }
+    }
+#endif // PLATFORMS
 #endif // UE_TRACE_ENABLED
 }
 
