Rolled back to "static" implementation, added comment explaining what this test aims to check and removed region export related operations.

diff --git a/Engine/Source/Developer/TraceInsights/Private/Insights/Tests/FunctionalTests/ExportCommandsTests.cpp b/Engine/Source/Developer/TraceInsights/Private/Insights/Tests/FunctionalTests/ExportCommandsTests.cpp
--- a/Engine/Source/Developer/TraceInsights/Private/Insights/Tests/FunctionalTests/ExportCommandsTests.cpp
+++ b/Engine/Source/Developer/TraceInsights/Private/Insights/Tests/FunctionalTests/ExportCommandsTests.cpp
@@ -45,6 +45,8 @@ void VerifyExportedLines(const FString& ExportReportPath, const FString& CmdLogP
 
 #if WITH_AUTOMATION_TESTS
 
+// The objective of these tests is to verify if the export command generates the predicted files and if these files encompass the anticipated data.
+// This confirmation is realized through the utilization of a pre-recorded trace file, which ensures data reliability.
 IMPLEMENT_SIMPLE_AUTOMATION_TEST(FСommandsExportWindowsTest, "System.Insights.Trace.Analysis.ExecCmd.CommandsExport(Windows)", EAutomationTestFlags::ProgramContext | EAutomationTestFlags::EngineFilter)
 bool FСommandsExportWindowsTest::RunTest(const FString& Parameters)
 {
@@ -64,7 +66,6 @@ bool FСommandsExportWindowsTest::RunTest(const FString& Parameters)
 	const FString CmdTimersLogPath = TEXT("/TestResults/Logs/cmd_timers.log");
 	const FString CmdTimingEventsLogPath = TEXT("/TestResults/Logs/cmd_timing_events.log");
 	const FString CmdTimingEventsNonDefaultLogPath = TEXT("/TestResults/Logs/cmd_timing_non_default.log");
-	const FString CmdRegionsLogPath = TEXT("/TestResults/Logs/cmd_regions.log");
 	const FString CmdExportLogPath = TEXT("/TestResults/Logs/cmd_export.log");
 	const FString ExportThreadsTask = TEXT("TimingInsights.ExportThreads /TestResults/SingleCommand/Threads.csv");
 	const FString ExportTimersTask = TEXT("TimingInsights.ExportTimers /TestResults/SingleCommand/Timers.csv");
@@ -88,7 +89,7 @@ bool FСommandsExportWindowsTest::RunTest(const FString& Parameters)
 
 	if (PlatformFile.DirectoryExists(*TestResultsDirPath))
 	{
-		AddWarning(TEXT("The TestResults directory already exists. Deleting to avoid undefined behavior"));
+		AddInfo(TEXT("The TestResults directory already exists. Deleting to avoid undefined behavior"));
 		IFileManager::Get().DeleteDirectory(*TestResultsDirPath, false, true);
 		double StartTime = FPlatformTime::Seconds();
 		while (FPlatformTime::Seconds() - StartTime < Timeout)
@@ -148,18 +149,6 @@ bool FСommandsExportWindowsTest::RunTest(const FString& Parameters)
 
 	VerifyExportedLines(ExportTimingEventsNonDefaultReportPath, CmdTimingEventsNonDefaultLogPath, TEXT("timing events"), Utils, this, Timeout);
 
-	// Export regions
-	InsightsParameters = FString::Printf(TEXT("-OpenTraceFile=\"%s\" -ABSLOG=\"%s\" -AutoQuit -NoUI  -ExecOnAnalysisCompleteCmd=\"%s\" -log"), *StoreTracePath, *CmdRegionsLogPath, *ExportTimerStatisticsTask);
-	InsightsManager->OpenUnrealInsights(*InsightsParameters);
-
-	bLineFound = Utils.FileContainsString(CmdRegionsLogPath, TEXT("Exported timing statistics for 10 regions"), Timeout);
-	TestTrue(FString::Printf(TEXT("Line '%s' should exists in file: '%s'"), TEXT("Exported timing statistics for 10 regions"), *CmdRegionsLogPath), bLineFound);
-
-	TArray<FString> RegionFiles;
-	FString FilePattern = TEXT("*.csv"); // Change the pattern to filter .csv files
-	IFileManager::Get().FindFiles(RegionFiles, *FPaths::Combine(*SingleCommandDirPath, *FilePattern), true, false);
-	TestEqual(FString::Printf(TEXT("Should be 14 csv files but found '%d'"), RegionFiles.Num()), RegionFiles.Num(), 14);
-
 	// Export.rsp
 	PlatformFile.CopyFile(*LogResultExportPath, *SourceExportPath);
 	TestTrue("Rsp in log directory should exists after copy", PlatformFile.FileExists(*LogResultExportPath));
