Add support for exclusive time tracking for Insights Cooking Profiler metrics.

diff --git a/Engine/Source/Developer/TraceInsights/Private/Insights/CookProfiler/ViewModels/PackageEntry.cpp b/Engine/Source/Developer/TraceInsights/Private/Insights/CookProfiler/ViewModels/PackageEntry.cpp
--- a/Engine/Source/Developer/TraceInsights/Private/Insights/CookProfiler/ViewModels/PackageEntry.cpp
+++ b/Engine/Source/Developer/TraceInsights/Private/Insights/CookProfiler/ViewModels/PackageEntry.cpp
@@ -10,10 +10,14 @@ namespace Insights
 FPackageEntry::FPackageEntry(const TraceServices::FPackageData& PackageData)
 	: Id(PackageData.Id)
 	, Name(PackageData.Name)
-	, LoadTime(PackageData.LoadTime)
-	, SaveTime(PackageData.SaveTime)
-	, BeginCacheForCookedPlatformData(PackageData.BeginCacheForCookedPlatformData)
-	, IsCachedCookedPlatformDataLoaded(PackageData.IsCachedCookedPlatformDataLoaded)
+	, LoadTimeIncl(PackageData.LoadTimeIncl)
+	, LoadTimeExcl(PackageData.LoadTimeExcl)
+	, SaveTimeIncl(PackageData.SaveTimeIncl)
+	, SaveTimeExcl(PackageData.SaveTimeExcl)
+	, BeginCacheForCookedPlatformDataIncl(PackageData.BeginCacheForCookedPlatformDataIncl)
+	, BeginCacheForCookedPlatformDataExcl(PackageData.BeginCacheForCookedPlatformDataExcl)
+	, IsCachedCookedPlatformDataLoadedIncl(PackageData.IsCachedCookedPlatformDataLoadedIncl)
+	, IsCachedCookedPlatformDataLoadedExcl(PackageData.IsCachedCookedPlatformDataLoadedExcl)
 	, AssetClass(PackageData.AssetClass)
 {
 }
diff --git a/Engine/Source/Developer/TraceInsights/Private/Insights/CookProfiler/ViewModels/PackageEntry.h b/Engine/Source/Developer/TraceInsights/Private/Insights/CookProfiler/ViewModels/PackageEntry.h
--- a/Engine/Source/Developer/TraceInsights/Private/Insights/CookProfiler/ViewModels/PackageEntry.h
+++ b/Engine/Source/Developer/TraceInsights/Private/Insights/CookProfiler/ViewModels/PackageEntry.h
@@ -21,19 +21,36 @@ public:
 
 	uint64 GetId() const { return Id; }
 	const TCHAR* GetName() const { return Name; }
-	const double GetLoadTime() const { return LoadTime; }
-	const double GetSaveTime() const { return SaveTime; }
-	const double GetBeginCacheForCookedPlatformData() const { return BeginCacheForCookedPlatformData; }
-	const double GetIsCachedCookedPlatformDataLoaded() const { return IsCachedCookedPlatformDataLoaded; }
+
+	const double GetLoadTimeIncl() const { return LoadTimeIncl; }
+	const double GetLoadTimeExcl() const { return LoadTimeExcl; }
+
+	const double GetSaveTimeIncl() const { return SaveTimeIncl; }
+	const double GetSaveTimeExcl() const { return SaveTimeExcl; }
+
+	const double GetBeginCacheForCookedPlatformDataIncl() const { return BeginCacheForCookedPlatformDataIncl; }
+	const double GetBeginCacheForCookedPlatformDataExcl() const { return BeginCacheForCookedPlatformDataExcl; }
+
+	const double GetIsCachedCookedPlatformDataLoadedIncl() const { return IsCachedCookedPlatformDataLoadedIncl; }
+	const double GetIsCachedCookedPlatformDataLoadedExcl() const { return IsCachedCookedPlatformDataLoadedExcl; }
+
 	const TCHAR* GetAssetClass() const { return AssetClass; }
 
 private:
 	uint64 Id;
 	const TCHAR* Name; 
-	double LoadTime; 
-	double SaveTime;
-	double BeginCacheForCookedPlatformData;
-	double IsCachedCookedPlatformDataLoaded;
+	double LoadTimeIncl;
+	double LoadTimeExcl;
+
+	double SaveTimeIncl;
+	double SaveTimeExcl;
+
+	double BeginCacheForCookedPlatformDataIncl;
+	double BeginCacheForCookedPlatformDataExcl;
+
+	double IsCachedCookedPlatformDataLoadedIncl;
+	double IsCachedCookedPlatformDataLoadedExcl;
+
 	const TCHAR* AssetClass;
 };
 
diff --git a/Engine/Source/Developer/TraceInsights/Private/Insights/CookProfiler/ViewModels/PackageTable.cpp b/Engine/Source/Developer/TraceInsights/Private/Insights/CookProfiler/ViewModels/PackageTable.cpp
--- a/Engine/Source/Developer/TraceInsights/Private/Insights/CookProfiler/ViewModels/PackageTable.cpp
+++ b/Engine/Source/Developer/TraceInsights/Private/Insights/CookProfiler/ViewModels/PackageTable.cpp
@@ -20,10 +20,19 @@ namespace Insights
 
 const FName FPackageTableColumns::IdColumnId(TEXT("Id"));
 const FName FPackageTableColumns::NameColumnId(TEXT("Name"));
-const FName FPackageTableColumns::LoadTimeColumnId(TEXT("LoadTime"));
-const FName FPackageTableColumns::SaveTimeColumnId(TEXT("SaveTime"));
-const FName FPackageTableColumns::BeginCacheForCookedPlatformDataTimeColumnId(TEXT("BeginCacheForCookedPlatformDataTimeColumnId"));
-const FName FPackageTableColumns::GetIsCachedCookedPlatformDataLoadedColumnId(TEXT("GetIsCachedCookedPlatformDataLoaded"));
+
+const FName FPackageTableColumns::LoadTimeInclColumnId(TEXT("LoadTimeIncl"));
+const FName FPackageTableColumns::LoadTimeExclColumnId(TEXT("LoadTimeExcl"));
+
+const FName FPackageTableColumns::SaveTimeInclColumnId(TEXT("SaveTimeIncl"));
+const FName FPackageTableColumns::SaveTimeExclColumnId(TEXT("SaveTimeExcl")); 
+
+const FName FPackageTableColumns::BeginCacheForCookedPlatformDataTimeInclColumnId(TEXT("BeginCacheForCookedPlatformDataTimeIncl"));
+const FName FPackageTableColumns::BeginCacheForCookedPlatformDataTimeExclColumnId(TEXT("BeginCacheForCookedPlatformDataTimeExcl"));
+
+const FName FPackageTableColumns::GetIsCachedCookedPlatformDataLoadedInclColumnId(TEXT("GetIsCachedCookedPlatformDataLoadedIncl"));
+const FName FPackageTableColumns::GetIsCachedCookedPlatformDataLoadedExclColumnId(TEXT("GetIsCachedCookedPlatformDataLoadedExcl"));
+
 const FName FPackageTableColumns::PackageAssetClassColumnId(TEXT("AssetClass"));
 
 ////////////////////////////////////////////////////////////////////////////////////////////////////
@@ -63,10 +72,19 @@ struct DefaultPackageFieldGetterFuncts
 {
 	static FTableCellValue GetId(const FTableColumn& Column, const FPackageEntry& Package) { return FTableCellValue((int64)Package.GetId());	}
 	static FTableCellValue GetName(const FTableColumn& Column, const FPackageEntry& Package) { return FTableCellValue((const TCHAR*)Package.GetName());	}
-	static FTableCellValue GetLoadTime(const FTableColumn& Column, const FPackageEntry& Package) { return FTableCellValue((double)Package.GetLoadTime());	}
-	static FTableCellValue GetSaveTime(const FTableColumn& Column, const FPackageEntry& Package) { return FTableCellValue((double)Package.GetSaveTime());	}
-	static FTableCellValue GetBeginCacheForCookedPlatformData(const FTableColumn& Column, const FPackageEntry& Package) { return FTableCellValue((double)Package.GetBeginCacheForCookedPlatformData());	}
-	static FTableCellValue GetIsCachedCookedPlatformDataLoaded(const FTableColumn& Column, const FPackageEntry& Package) { return FTableCellValue((double)Package.GetIsCachedCookedPlatformDataLoaded());	}
+
+	static FTableCellValue GetLoadTimeIncl(const FTableColumn& Column, const FPackageEntry& Package) { return FTableCellValue((double)Package.GetLoadTimeIncl());	}
+	static FTableCellValue GetLoadTimeExcl(const FTableColumn& Column, const FPackageEntry& Package) { return FTableCellValue((double)Package.GetLoadTimeExcl());	}
+
+	static FTableCellValue GetSaveTimeIncl(const FTableColumn& Column, const FPackageEntry& Package) { return FTableCellValue((double)Package.GetSaveTimeIncl());	}
+	static FTableCellValue GetSaveTimeExcl(const FTableColumn& Column, const FPackageEntry& Package) { return FTableCellValue((double)Package.GetSaveTimeExcl());	}
+
+	static FTableCellValue GetBeginCacheForCookedPlatformDataIncl(const FTableColumn& Column, const FPackageEntry& Package) { return FTableCellValue((double)Package.GetBeginCacheForCookedPlatformDataIncl());	}
+	static FTableCellValue GetBeginCacheForCookedPlatformDataExcl(const FTableColumn& Column, const FPackageEntry& Package) { return FTableCellValue((double)Package.GetBeginCacheForCookedPlatformDataExcl());	}
+
+	static FTableCellValue GetIsCachedCookedPlatformDataLoadedIncl(const FTableColumn& Column, const FPackageEntry& Package) { return FTableCellValue((double)Package.GetIsCachedCookedPlatformDataLoadedIncl());	}
+	static FTableCellValue GetIsCachedCookedPlatformDataLoadedExcl(const FTableColumn& Column, const FPackageEntry& Package) { return FTableCellValue((double)Package.GetIsCachedCookedPlatformDataLoadedExcl());	}
+
 	static FTableCellValue GetAssetClass(const FTableColumn& Column, const FPackageEntry& Package) { return FTableCellValue((const TCHAR*)Package.GetAssetClass());	}
 };
 
@@ -146,16 +164,144 @@ void FPackageTable::AddDefaultColumns()
 		AddColumn(ColumnRef);
 	}
 	//////////////////////////////////////////////////
-	// Load Time Column
+	// Inclusive Load Time Column
+	{
+		TSharedRef<FTableColumn> ColumnRef = MakeShared<FTableColumn>(FPackageTableColumns::LoadTimeInclColumnId);
+		FTableColumn& Column = *ColumnRef;
+
+		Column.SetIndex(ColumnIndex++);
+
+		Column.SetShortName(LOCTEXT("LoadTimeInclColumnName", "I. Load Time"));
+		Column.SetTitleName(LOCTEXT("LoadTimeInclColumnTitle", "Inclusive Load Time"));
+		Column.SetDescription(LOCTEXT("LoadTimeInclColumnDesc", "The inclusive time it took to load the package."));
+
+		Column.SetFlags(ETableColumnFlags::ShouldBeVisible | ETableColumnFlags::CanBeHidden | ETableColumnFlags::CanBeFiltered);
+
+		Column.SetHorizontalAlignment(HAlign_Left);
+		Column.SetInitialWidth(100.0f);
+
+		Column.SetDataType(ETableCellDataType::Double);
+
+		TSharedRef<ITableCellValueGetter> Getter = MakeShared<FPackageColumnValueGetter<DefaultPackageFieldGetterFuncts::GetLoadTimeIncl>>();
+		Column.SetValueGetter(Getter);
+
+		TSharedRef<ITableCellValueFormatter> Formatter = MakeShared<FDoubleValueFormatterAsTimeAuto>();
+		Column.SetValueFormatter(Formatter);
+
+		TSharedRef<ITableCellValueSorter> Sorter = MakeShared<FSorterByDoubleValue>(ColumnRef);
+		Column.SetValueSorter(Sorter);
+
+		Column.SetAggregation(ETableColumnAggregation::Sum);
+
+		AddColumn(ColumnRef);
+	}
+	//////////////////////////////////////////////////
+	// Exclusive Load Time Column
+	{
+		TSharedRef<FTableColumn> ColumnRef = MakeShared<FTableColumn>(FPackageTableColumns::LoadTimeExclColumnId);
+		FTableColumn& Column = *ColumnRef;
+
+		Column.SetIndex(ColumnIndex++);
+
+		Column.SetShortName(LOCTEXT("LoadTimeExclColumnName", "E. Load Time"));
+		Column.SetTitleName(LOCTEXT("LoadTimeExclColumnTitle", "Exclusive Load Time"));
+		Column.SetDescription(LOCTEXT("LoadTimeExclColumnDesc", "The exclusive time it took to load the package."));
+
+		Column.SetFlags(ETableColumnFlags::ShouldBeVisible | ETableColumnFlags::CanBeHidden | ETableColumnFlags::CanBeFiltered);
+
+		Column.SetHorizontalAlignment(HAlign_Left);
+		Column.SetInitialWidth(100.0f);
+
+		Column.SetDataType(ETableCellDataType::Double);
+
+		TSharedRef<ITableCellValueGetter> Getter = MakeShared<FPackageColumnValueGetter<DefaultPackageFieldGetterFuncts::GetLoadTimeExcl>>();
+		Column.SetValueGetter(Getter);
+
+		TSharedRef<ITableCellValueFormatter> Formatter = MakeShared<FDoubleValueFormatterAsTimeAuto>();
+		Column.SetValueFormatter(Formatter);
+
+		TSharedRef<ITableCellValueSorter> Sorter = MakeShared<FSorterByDoubleValue>(ColumnRef);
+		Column.SetValueSorter(Sorter);
+
+		Column.SetAggregation(ETableColumnAggregation::Sum);
+
+		AddColumn(ColumnRef);
+	}
+	//////////////////////////////////////////////////
+	// Inclusive Save Time Column
+	{
+		TSharedRef<FTableColumn> ColumnRef = MakeShared<FTableColumn>(FPackageTableColumns::SaveTimeInclColumnId);
+		FTableColumn& Column = *ColumnRef;
+
+		Column.SetIndex(ColumnIndex++);
+
+		Column.SetShortName(LOCTEXT("SaveTimeInclColumnName", "I. Save Time"));
+		Column.SetTitleName(LOCTEXT("SaveTimeInclColumnTitle", "Inclusive Save Time"));
+		Column.SetDescription(LOCTEXT("SaveTimeInclColumnDesc", "The inclusive time it took to save the package."));
+
+		Column.SetFlags(ETableColumnFlags::ShouldBeVisible | ETableColumnFlags::CanBeHidden | ETableColumnFlags::CanBeFiltered);
+
+		Column.SetHorizontalAlignment(HAlign_Left);
+		Column.SetInitialWidth(100.0f);
+
+		Column.SetDataType(ETableCellDataType::Double);
+
+		TSharedRef<ITableCellValueGetter> Getter = MakeShared<FPackageColumnValueGetter<DefaultPackageFieldGetterFuncts::GetSaveTimeIncl>>();
+		Column.SetValueGetter(Getter);
+
+		TSharedRef<ITableCellValueFormatter> Formatter = MakeShared<FDoubleValueFormatterAsTimeAuto>();
+		Column.SetValueFormatter(Formatter);
+
+		TSharedRef<ITableCellValueSorter> Sorter = MakeShared<FSorterByDoubleValue>(ColumnRef);
+		Column.SetValueSorter(Sorter);
+
+		Column.SetAggregation(ETableColumnAggregation::Sum);
+
+		AddColumn(ColumnRef);
+	}
+	//////////////////////////////////////////////////
+	// Exclusive Save Time Column
+	{
+		TSharedRef<FTableColumn> ColumnRef = MakeShared<FTableColumn>(FPackageTableColumns::SaveTimeExclColumnId);
+		FTableColumn& Column = *ColumnRef;
+
+		Column.SetIndex(ColumnIndex++);
+
+		Column.SetShortName(LOCTEXT("SaveTimeExclColumnName", "E. Save Time"));
+		Column.SetTitleName(LOCTEXT("SaveTimeExclColumnTitle", "Exclusive Save Time"));
+		Column.SetDescription(LOCTEXT("SaveTimeEnclColumnDesc", "The exclusive time it took to save the package."));
+
+		Column.SetFlags(ETableColumnFlags::ShouldBeVisible | ETableColumnFlags::CanBeHidden | ETableColumnFlags::CanBeFiltered);
+
+		Column.SetHorizontalAlignment(HAlign_Left);
+		Column.SetInitialWidth(100.0f);
+
+		Column.SetDataType(ETableCellDataType::Double);
+
+		TSharedRef<ITableCellValueGetter> Getter = MakeShared<FPackageColumnValueGetter<DefaultPackageFieldGetterFuncts::GetSaveTimeExcl>>();
+		Column.SetValueGetter(Getter);
+
+		TSharedRef<ITableCellValueFormatter> Formatter = MakeShared<FDoubleValueFormatterAsTimeAuto>();
+		Column.SetValueFormatter(Formatter);
+
+		TSharedRef<ITableCellValueSorter> Sorter = MakeShared<FSorterByDoubleValue>(ColumnRef);
+		Column.SetValueSorter(Sorter);
+
+		Column.SetAggregation(ETableColumnAggregation::Sum);
+
+		AddColumn(ColumnRef);
+	}
+	//////////////////////////////////////////////////
+	// Inclusive BeginCacheForCookedPlatformData Time Column
 	{
-		TSharedRef<FTableColumn> ColumnRef = MakeShared<FTableColumn>(FPackageTableColumns::LoadTimeColumnId);
+		TSharedRef<FTableColumn> ColumnRef = MakeShared<FTableColumn>(FPackageTableColumns::BeginCacheForCookedPlatformDataTimeInclColumnId);
 		FTableColumn& Column = *ColumnRef;
 
 		Column.SetIndex(ColumnIndex++);
 
-		Column.SetShortName(LOCTEXT("LoadTimeColumnName", "LoadTime"));
-		Column.SetTitleName(LOCTEXT("LoadTimeColumnTitle", "LoadTime"));
-		Column.SetDescription(LOCTEXT("LoadTimeColumnDesc", "The time it took to load the package."));
+		Column.SetShortName(LOCTEXT("BeginCacheForCookedPlatformDataInclColumnName", "I. Begin Cache Time"));
+		Column.SetTitleName(LOCTEXT("BeginCacheForCookedPlatformDataInclColumnTitle", "Inclusive BeginCacheForCookedPlatformData"));
+		Column.SetDescription(LOCTEXT("BeginCacheForCookedPlatformDataInclColumnDesc", "The total inclusive time spent in the BeginCacheForCookedPlatformData function for the package."));
 
 		Column.SetFlags(ETableColumnFlags::ShouldBeVisible | ETableColumnFlags::CanBeHidden | ETableColumnFlags::CanBeFiltered);
 
@@ -164,7 +310,7 @@ void FPackageTable::AddDefaultColumns()
 
 		Column.SetDataType(ETableCellDataType::Double);
 
-		TSharedRef<ITableCellValueGetter> Getter = MakeShared<FPackageColumnValueGetter<DefaultPackageFieldGetterFuncts::GetLoadTime>>();
+		TSharedRef<ITableCellValueGetter> Getter = MakeShared<FPackageColumnValueGetter<DefaultPackageFieldGetterFuncts::GetBeginCacheForCookedPlatformDataIncl>>();
 		Column.SetValueGetter(Getter);
 
 		TSharedRef<ITableCellValueFormatter> Formatter = MakeShared<FDoubleValueFormatterAsTimeAuto>();
@@ -178,16 +324,16 @@ void FPackageTable::AddDefaultColumns()
 		AddColumn(ColumnRef);
 	}
 	//////////////////////////////////////////////////
-	// Save Time Column
+	// Exclusive BeginCacheForCookedPlatformData Time Column
 	{
-		TSharedRef<FTableColumn> ColumnRef = MakeShared<FTableColumn>(FPackageTableColumns::SaveTimeColumnId);
+		TSharedRef<FTableColumn> ColumnRef = MakeShared<FTableColumn>(FPackageTableColumns::BeginCacheForCookedPlatformDataTimeExclColumnId);
 		FTableColumn& Column = *ColumnRef;
 
 		Column.SetIndex(ColumnIndex++);
 
-		Column.SetShortName(LOCTEXT("SaveTimeColumnName", "SaveTime"));
-		Column.SetTitleName(LOCTEXT("SaveTimeColumnTitle", "SaveTime"));
-		Column.SetDescription(LOCTEXT("SaveTimeColumnDesc", "The time it took to save the package."));
+		Column.SetShortName(LOCTEXT("BeginCacheForCookedPlatformDataExclColumnName", "E. Begin Cache Time"));
+		Column.SetTitleName(LOCTEXT("BeginCacheForCookedPlatformDataExclColumnTitle", "Exclusive BeginCacheForCookedPlatformData"));
+		Column.SetDescription(LOCTEXT("BeginCacheForCookedPlatformDataExclColumnDesc", "The total exclusive time spent in the BeginCacheForCookedPlatformData function for the package."));
 
 		Column.SetFlags(ETableColumnFlags::ShouldBeVisible | ETableColumnFlags::CanBeHidden | ETableColumnFlags::CanBeFiltered);
 
@@ -196,7 +342,7 @@ void FPackageTable::AddDefaultColumns()
 
 		Column.SetDataType(ETableCellDataType::Double);
 
-		TSharedRef<ITableCellValueGetter> Getter = MakeShared<FPackageColumnValueGetter<DefaultPackageFieldGetterFuncts::GetSaveTime>>();
+		TSharedRef<ITableCellValueGetter> Getter = MakeShared<FPackageColumnValueGetter<DefaultPackageFieldGetterFuncts::GetBeginCacheForCookedPlatformDataExcl>>();
 		Column.SetValueGetter(Getter);
 
 		TSharedRef<ITableCellValueFormatter> Formatter = MakeShared<FDoubleValueFormatterAsTimeAuto>();
@@ -210,16 +356,16 @@ void FPackageTable::AddDefaultColumns()
 		AddColumn(ColumnRef);
 	}
 	//////////////////////////////////////////////////
-	// BeginCacheForCookedPlatformData Time Column
+	// Inclusive IsCachedCookedPlatformDataLoaded Time Column
 	{
-		TSharedRef<FTableColumn> ColumnRef = MakeShared<FTableColumn>(FPackageTableColumns::BeginCacheForCookedPlatformDataTimeColumnId);
+		TSharedRef<FTableColumn> ColumnRef = MakeShared<FTableColumn>(FPackageTableColumns::GetIsCachedCookedPlatformDataLoadedInclColumnId);
 		FTableColumn& Column = *ColumnRef;
 
 		Column.SetIndex(ColumnIndex++);
 
-		Column.SetShortName(LOCTEXT("BeginCacheForCookedPlatformDataColumnName", "BeginCache"));
-		Column.SetTitleName(LOCTEXT("BeginCacheForCookedPlatformDataColumnTitle", "BeginCacheForCookedPlatformData"));
-		Column.SetDescription(LOCTEXT("BeginCacheForCookedPlatformDataColumnDesc", "The total time spent in the BeginCacheForCookedPlatformData function for the package."));
+		Column.SetShortName(LOCTEXT("GetIsCachedCookedPlatformDataLoadedInclColumnName", "I. IsCachedCooked"));
+		Column.SetTitleName(LOCTEXT("GetIsCachedCookedPlatformDataLoadedInclColumnTitle", "Inclusive IsCachedCookedPlatformDataLoaded"));
+		Column.SetDescription(LOCTEXT("GetIsCachedCookedPlatformDataLoadedInclColumnDesc", "The total inclusive time spent in the IsCachedCookedPlatformDataLoaded function for the package."));
 
 		Column.SetFlags(ETableColumnFlags::ShouldBeVisible | ETableColumnFlags::CanBeHidden | ETableColumnFlags::CanBeFiltered);
 
@@ -228,7 +374,7 @@ void FPackageTable::AddDefaultColumns()
 
 		Column.SetDataType(ETableCellDataType::Double);
 
-		TSharedRef<ITableCellValueGetter> Getter = MakeShared<FPackageColumnValueGetter<DefaultPackageFieldGetterFuncts::GetBeginCacheForCookedPlatformData>>();
+		TSharedRef<ITableCellValueGetter> Getter = MakeShared<FPackageColumnValueGetter<DefaultPackageFieldGetterFuncts::GetIsCachedCookedPlatformDataLoadedIncl>>();
 		Column.SetValueGetter(Getter);
 
 		TSharedRef<ITableCellValueFormatter> Formatter = MakeShared<FDoubleValueFormatterAsTimeAuto>();
@@ -242,16 +388,16 @@ void FPackageTable::AddDefaultColumns()
 		AddColumn(ColumnRef);
 	}
 	//////////////////////////////////////////////////
-	// BeginCacheForCookedPlatformData Time Column
+	// Exclusive IsCachedCookedPlatformDataLoaded Time Column
 	{
-		TSharedRef<FTableColumn> ColumnRef = MakeShared<FTableColumn>(FPackageTableColumns::GetIsCachedCookedPlatformDataLoadedColumnId);
+		TSharedRef<FTableColumn> ColumnRef = MakeShared<FTableColumn>(FPackageTableColumns::GetIsCachedCookedPlatformDataLoadedExclColumnId);
 		FTableColumn& Column = *ColumnRef;
 
 		Column.SetIndex(ColumnIndex++);
 
-		Column.SetShortName(LOCTEXT("GetIsCachedCookedPlatformDataLoadedColumnName", "IsCachedCooked"));
-		Column.SetTitleName(LOCTEXT("GetIsCachedCookedPlatformDataLoadedColumnTitle", "IsCachedCookedPlatformDataLoaded"));
-		Column.SetDescription(LOCTEXT("GetIsCachedCookedPlatformDataLoadedColumnDesc", "The total time spent in the IsCachedCookedPlatformDataLoaded function for the package."));
+		Column.SetShortName(LOCTEXT("GetIsCachedCookedPlatformDataLoadedExclColumnName", "E. IsCachedCooked"));
+		Column.SetTitleName(LOCTEXT("GetIsCachedCookedPlatformDataLoadedExclColumnTitle", "Exclusive IsCachedCookedPlatformDataLoaded"));
+		Column.SetDescription(LOCTEXT("GetIsCachedCookedPlatformDataLoadedExclColumnDesc", "The total exclusive time spent in the IsCachedCookedPlatformDataLoaded function for the package."));
 
 		Column.SetFlags(ETableColumnFlags::ShouldBeVisible | ETableColumnFlags::CanBeHidden | ETableColumnFlags::CanBeFiltered);
 
@@ -260,7 +406,7 @@ void FPackageTable::AddDefaultColumns()
 
 		Column.SetDataType(ETableCellDataType::Double);
 
-		TSharedRef<ITableCellValueGetter> Getter = MakeShared<FPackageColumnValueGetter<DefaultPackageFieldGetterFuncts::GetIsCachedCookedPlatformDataLoaded>>();
+		TSharedRef<ITableCellValueGetter> Getter = MakeShared<FPackageColumnValueGetter<DefaultPackageFieldGetterFuncts::GetIsCachedCookedPlatformDataLoadedExcl>>();
 		Column.SetValueGetter(Getter);
 
 		TSharedRef<ITableCellValueFormatter> Formatter = MakeShared<FDoubleValueFormatterAsTimeAuto>();
diff --git a/Engine/Source/Developer/TraceInsights/Private/Insights/CookProfiler/ViewModels/PackageTable.h b/Engine/Source/Developer/TraceInsights/Private/Insights/CookProfiler/ViewModels/PackageTable.h
--- a/Engine/Source/Developer/TraceInsights/Private/Insights/CookProfiler/ViewModels/PackageTable.h
+++ b/Engine/Source/Developer/TraceInsights/Private/Insights/CookProfiler/ViewModels/PackageTable.h
@@ -20,10 +20,19 @@ struct FPackageTableColumns
 {
 	static const FName IdColumnId;
 	static const FName NameColumnId;
-	static const FName LoadTimeColumnId;
-	static const FName SaveTimeColumnId;
-	static const FName BeginCacheForCookedPlatformDataTimeColumnId;
-	static const FName GetIsCachedCookedPlatformDataLoadedColumnId;
+
+	static const FName LoadTimeInclColumnId;
+	static const FName LoadTimeExclColumnId;
+
+	static const FName SaveTimeInclColumnId;
+	static const FName SaveTimeExclColumnId;
+
+	static const FName BeginCacheForCookedPlatformDataTimeInclColumnId;
+	static const FName BeginCacheForCookedPlatformDataTimeExclColumnId;
+
+	static const FName GetIsCachedCookedPlatformDataLoadedInclColumnId;
+	static const FName GetIsCachedCookedPlatformDataLoadedExclColumnId;
+
 	static const FName PackageAssetClassColumnId;
 };
 
diff --git a/Engine/Source/Developer/TraceInsights/Private/Insights/CookProfiler/Widgets/SPackageTableTreeView.cpp b/Engine/Source/Developer/TraceInsights/Private/Insights/CookProfiler/Widgets/SPackageTableTreeView.cpp
--- a/Engine/Source/Developer/TraceInsights/Private/Insights/CookProfiler/Widgets/SPackageTableTreeView.cpp
+++ b/Engine/Source/Developer/TraceInsights/Private/Insights/CookProfiler/Widgets/SPackageTableTreeView.cpp
@@ -134,20 +134,21 @@ void SPackageTableTreeView::RebuildTree(bool bResync)
 	{
 		TraceServices::FProviderReadScopeLock ProviderReadScope(*CookProvider);
 
+		TArray64<FPackageData> PackageAggreagation;
+		CookProvider->CreateAggregation(PackageAggreagation);
 		const uint32 NumPackages = CookProvider->GetNumPackages();
 		Packages.Reserve(NumPackages);
 		TableRowNodes.Reserve(NumPackages);
 
 		TArray<FTableTreeNodePtr>* Nodes = &TableRowNodes;
-		CookProvider->EnumeratePackages(0, 0, [&Packages, &PackageTable, Nodes](const TraceServices::FPackageData& Package)
-			{
-				Packages.Emplace(Package);
-				uint32 Index = Packages.Num() - 1;
-				FName NodeName(Package.Name);
-				FPackageNodePtr NodePtr = MakeShared<FPackageNode>(NodeName, PackageTable, Index);
-				Nodes->Add(NodePtr);
-				return true;
-			});
+		for(const TraceServices::FPackageData& Package : PackageAggreagation)
+		{
+			Packages.Emplace(Package);
+			uint32 Index = static_cast<uint32>(Packages.Num() - 1);
+			FName NodeName(Package.Name);
+			FPackageNodePtr NodePtr = MakeShared<FPackageNode>(NodeName, PackageTable, Index);
+			Nodes->Add(NodePtr);
+		};
 	}
 
 	bDataLoaded = true;
@@ -236,12 +237,16 @@ void SPackageTableTreeView::InternalCreateGroupings()
 			if (Grouping->Is<FTreeNodeGroupingByUniqueValue>())
 			{
 				const FName ColumnId = Grouping->As<FTreeNodeGroupingByUniqueValue>().GetColumnId();
-				if (ColumnId == FPackageTableColumns::BeginCacheForCookedPlatformDataTimeColumnId ||
-					ColumnId == FPackageTableColumns::GetIsCachedCookedPlatformDataLoadedColumnId ||
-					ColumnId == FPackageTableColumns::SaveTimeColumnId ||
+				if (ColumnId == FPackageTableColumns::BeginCacheForCookedPlatformDataTimeInclColumnId ||
+					ColumnId == FPackageTableColumns::BeginCacheForCookedPlatformDataTimeExclColumnId ||
+					ColumnId == FPackageTableColumns::GetIsCachedCookedPlatformDataLoadedInclColumnId ||
+					ColumnId == FPackageTableColumns::GetIsCachedCookedPlatformDataLoadedExclColumnId ||
+					ColumnId == FPackageTableColumns::SaveTimeInclColumnId ||
+					ColumnId == FPackageTableColumns::SaveTimeExclColumnId ||
 					ColumnId == FPackageTableColumns::IdColumnId ||
 					ColumnId == FPackageTableColumns::NameColumnId ||
-					ColumnId == FPackageTableColumns::LoadTimeColumnId)
+					ColumnId == FPackageTableColumns::LoadTimeInclColumnId ||
+					ColumnId == FPackageTableColumns::LoadTimeExclColumnId)
 				{
 					return true;
 				}
@@ -302,12 +307,16 @@ void SPackageTableTreeView::InitAvailableViewPresets()
 		{
 			InOutConfigSet.Add({ FTable::GetHierarchyColumnId(), true, 500.0f });
 			InOutConfigSet.Add({ FPackageTableColumns::IdColumnId, true, 80.0f });
-			InOutConfigSet.Add({ FPackageTableColumns::LoadTimeColumnId, true, 100.0f });
-			InOutConfigSet.Add({ FPackageTableColumns::SaveTimeColumnId, true, 100.0f });
-			InOutConfigSet.Add({ FPackageTableColumns::BeginCacheForCookedPlatformDataTimeColumnId, true, 100.0f });
-			InOutConfigSet.Add({ FPackageTableColumns::GetIsCachedCookedPlatformDataLoadedColumnId, true, 100.0f });
+			InOutConfigSet.Add({ FPackageTableColumns::LoadTimeExclColumnId, true, 100.0f });
+			InOutConfigSet.Add({ FPackageTableColumns::SaveTimeExclColumnId, true, 100.0f });
+			InOutConfigSet.Add({ FPackageTableColumns::BeginCacheForCookedPlatformDataTimeExclColumnId, true, 100.0f });
+			InOutConfigSet.Add({ FPackageTableColumns::GetIsCachedCookedPlatformDataLoadedExclColumnId, true, 100.0f });
 			InOutConfigSet.Add({ FPackageTableColumns::PackageAssetClassColumnId, true, 200.0f });
 			InOutConfigSet.Add({ FPackageTableColumns::NameColumnId, false, 400.0f });
+			InOutConfigSet.Add({ FPackageTableColumns::LoadTimeInclColumnId, true, 100.0f });
+			InOutConfigSet.Add({ FPackageTableColumns::SaveTimeInclColumnId, true, 100.0f });
+			InOutConfigSet.Add({ FPackageTableColumns::BeginCacheForCookedPlatformDataTimeInclColumnId, true, 100.0f });
+			InOutConfigSet.Add({ FPackageTableColumns::GetIsCachedCookedPlatformDataLoadedInclColumnId, true, 100.0f });
 		}
 	};
 	AvailableViewPresets.Add(MakeShared<FDefaultViewPreset>());
@@ -357,12 +366,16 @@ void SPackageTableTreeView::InitAvailableViewPresets()
 		{
 			InOutConfigSet.Add({ FTable::GetHierarchyColumnId(), true, 500.0f });
 			InOutConfigSet.Add({ FPackageTableColumns::IdColumnId, true, 80.0f });
-			InOutConfigSet.Add({ FPackageTableColumns::LoadTimeColumnId, true, 100.0f });
-			InOutConfigSet.Add({ FPackageTableColumns::SaveTimeColumnId, true, 100.0f });
-			InOutConfigSet.Add({ FPackageTableColumns::BeginCacheForCookedPlatformDataTimeColumnId, true, 100.0f });
-			InOutConfigSet.Add({ FPackageTableColumns::GetIsCachedCookedPlatformDataLoadedColumnId, true, 100.0f });
+			InOutConfigSet.Add({ FPackageTableColumns::LoadTimeExclColumnId, true, 100.0f });
+			InOutConfigSet.Add({ FPackageTableColumns::SaveTimeExclColumnId, true, 100.0f });
+			InOutConfigSet.Add({ FPackageTableColumns::BeginCacheForCookedPlatformDataTimeExclColumnId, true, 100.0f });
+			InOutConfigSet.Add({ FPackageTableColumns::GetIsCachedCookedPlatformDataLoadedExclColumnId, true, 100.0f });
 			InOutConfigSet.Add({ FPackageTableColumns::PackageAssetClassColumnId, true, 200.0f });
 			InOutConfigSet.Add({ FPackageTableColumns::NameColumnId, false, 400.0f });
+			InOutConfigSet.Add({ FPackageTableColumns::LoadTimeInclColumnId, true, 100.0f });
+			InOutConfigSet.Add({ FPackageTableColumns::SaveTimeInclColumnId, true, 100.0f });
+			InOutConfigSet.Add({ FPackageTableColumns::BeginCacheForCookedPlatformDataTimeInclColumnId, true, 100.0f });
+			InOutConfigSet.Add({ FPackageTableColumns::GetIsCachedCookedPlatformDataLoadedInclColumnId, true, 100.0f });
 		}
 	};
 	AvailableViewPresets.Add(MakeShared<FPackagePathViewPreset>());
@@ -412,12 +425,17 @@ void SPackageTableTreeView::InitAvailableViewPresets()
 		{
 			InOutConfigSet.Add({ FTable::GetHierarchyColumnId(), true, 300.0f });
 			InOutConfigSet.Add({ FPackageTableColumns::IdColumnId, true, 80.0f });
-			InOutConfigSet.Add({ FPackageTableColumns::LoadTimeColumnId, true, 100.0f });
-			InOutConfigSet.Add({ FPackageTableColumns::SaveTimeColumnId, true, 100.0f });
-			InOutConfigSet.Add({ FPackageTableColumns::BeginCacheForCookedPlatformDataTimeColumnId, true, 100.0f });
-			InOutConfigSet.Add({ FPackageTableColumns::GetIsCachedCookedPlatformDataLoadedColumnId, true, 100.0f });
+			InOutConfigSet.Add({ FPackageTableColumns::LoadTimeExclColumnId, true, 100.0f });
+			InOutConfigSet.Add({ FPackageTableColumns::SaveTimeExclColumnId, true, 100.0f });
+			InOutConfigSet.Add({ FPackageTableColumns::BeginCacheForCookedPlatformDataTimeExclColumnId, true, 100.0f });
+			InOutConfigSet.Add({ FPackageTableColumns::GetIsCachedCookedPlatformDataLoadedExclColumnId, true, 100.0f });
 			InOutConfigSet.Add({ FPackageTableColumns::NameColumnId, true, 400.0f });
 			InOutConfigSet.Add({ FPackageTableColumns::PackageAssetClassColumnId, false, 200.0f });
+			InOutConfigSet.Add({ FPackageTableColumns::LoadTimeInclColumnId, true, 100.0f });
+			InOutConfigSet.Add({ FPackageTableColumns::SaveTimeInclColumnId, true, 100.0f });
+			InOutConfigSet.Add({ FPackageTableColumns::BeginCacheForCookedPlatformDataTimeInclColumnId, true, 100.0f });
+			InOutConfigSet.Add({ FPackageTableColumns::GetIsCachedCookedPlatformDataLoadedInclColumnId, true, 100.0f });
+
 		}
 	};
 	AvailableViewPresets.Add(MakeShared<FAssetClassViewPreset>());
diff --git a/Engine/Source/Developer/TraceServices/Private/Analyzers/CookAnalysis.cpp b/Engine/Source/Developer/TraceServices/Private/Analyzers/CookAnalysis.cpp
--- a/Engine/Source/Developer/TraceServices/Private/Analyzers/CookAnalysis.cpp
+++ b/Engine/Source/Developer/TraceServices/Private/Analyzers/CookAnalysis.cpp
@@ -27,8 +27,14 @@ void FCookAnalyzer::OnAnalysisBegin(const FOnAnalysisContext& Context)
 	auto& Builder = Context.InterfaceBuilder;
 
 	Builder.RouteEvent(RouteId_Package, "CookTrace", "Package");
-	Builder.RouteEvent(RouteId_PackageStat, "CookTrace", "PackageStat");
 	Builder.RouteEvent(RouteId_PackageAssetClass, "CookTrace", "PackageAssetClass");
+
+	//V1
+	Builder.RouteEvent(RouteId_PackageStat, "CookTrace", "PackageStat");
+
+	//V2
+	Builder.RouteEvent(RouteId_PackageStatBeginScope, "CookTrace", "PackageStatBeginScope");
+	Builder.RouteEvent(RouteId_PackageStatEndScope, "CookTrace", "PackageStatEndScope");
 }
 
 bool FCookAnalyzer::OnEvent(uint16 RouteId, EStyle Style, const FOnEventContext& Context)
@@ -62,24 +68,24 @@ bool FCookAnalyzer::OnEvent(uint16 RouteId, EStyle Style, const FOnEventContext&
 			case EPackageEventStatType::LoadPackage:
 			{
 				// We measure the Loadtime in multiple scopes so we receive many LoadPackage events.
-				Package->LoadTime += Value;
+				Package->LoadTimeIncl += Value;
 				break;
 			}
 			case EPackageEventStatType::SavePackage:
 			{
-				Package->SaveTime = Value;
+				Package->SaveTimeIncl = Value;
 				break;
 			}
 			case EPackageEventStatType::BeginCacheForCookedPlatformData:
 			{
 				// A BeginCacheForCookedPlatformData event is received for each asset in the package. We add the values to get the total time for the package.
-				Package->BeginCacheForCookedPlatformData += Value;
+				Package->BeginCacheForCookedPlatformDataIncl += Value;
 				break;
 			}
 			case EPackageEventStatType::IsCachedCookedPlatformDataLoaded:
 			{
 				// A IsCachedCookedPlatformDataLoaded event is received for each asset in the package. We add the values to get the total time for the package.
-				Package->IsCachedCookedPlatformDataLoaded += Value;
+				Package->IsCachedCookedPlatformDataLoadedIncl += Value;
 				break;
 			}
 			default:
@@ -89,6 +95,34 @@ bool FCookAnalyzer::OnEvent(uint16 RouteId, EStyle Style, const FOnEventContext&
 		}
 		break;
 	}
+	case RouteId_PackageStatBeginScope:
+	{
+		uint64 PackageId = EventData.GetValue<uint64>("Id");
+		double Timestamp = Context.EventTime.AsSeconds(EventData.GetValue<uint64>("Time"));
+		EPackageEventStatType StatType = (EPackageEventStatType)EventData.GetValue<uint8>("StatType");
+
+		FProviderEditScopeLock ProviderEditScope(CookProfilerProvider);
+		CookProfilerProvider.AddScopeEntry(Context.ThreadInfo.GetId(), PackageId, Timestamp, StatType, true);
+
+		FAnalysisSessionEditScope _EditScope(Session);
+		Session.UpdateDurationSeconds(Timestamp);
+
+		break;
+	}
+	case RouteId_PackageStatEndScope:
+	{
+		uint64 PackageId = EventData.GetValue<uint64>("Id");
+		double Timestamp = Context.EventTime.AsSeconds(EventData.GetValue<uint64>("Time"));
+		EPackageEventStatType StatType = (EPackageEventStatType)EventData.GetValue<uint8>("StatType");
+
+		FProviderEditScopeLock ProviderEditScope(CookProfilerProvider);
+		CookProfilerProvider.AddScopeEntry(Context.ThreadInfo.GetId(), PackageId, Timestamp, StatType, false);
+
+		FAnalysisSessionEditScope _EditScope(Session);
+		Session.UpdateDurationSeconds(Timestamp);
+
+		break; 
+	}
 	case RouteId_PackageAssetClass:
 	{
 		uint64 Id = EventData.GetValue<uint64>("Id");
diff --git a/Engine/Source/Developer/TraceServices/Private/Analyzers/CookAnalysis.h b/Engine/Source/Developer/TraceServices/Private/Analyzers/CookAnalysis.h
--- a/Engine/Source/Developer/TraceServices/Private/Analyzers/CookAnalysis.h
+++ b/Engine/Source/Developer/TraceServices/Private/Analyzers/CookAnalysis.h
@@ -21,9 +21,16 @@ public:
 private:
 	enum : uint16
 	{
+		// Common
 		RouteId_Package,
-		RouteId_PackageStat,
 		RouteId_PackageAssetClass,
+
+		// Version 1
+		RouteId_PackageStat,
+
+		//Version 2 (UE 5.5+)
+		RouteId_PackageStatBeginScope,
+		RouteId_PackageStatEndScope,
 	};
 
 	IAnalysisSession& Session;
diff --git a/Engine/Source/Developer/TraceServices/Private/Model/CookProfilerProvider.cpp b/Engine/Source/Developer/TraceServices/Private/Model/CookProfilerProvider.cpp
--- a/Engine/Source/Developer/TraceServices/Private/Model/CookProfilerProvider.cpp
+++ b/Engine/Source/Developer/TraceServices/Private/Model/CookProfilerProvider.cpp
@@ -56,6 +56,8 @@ FPackageData* FCookProfilerProvider::EditPackage(uint64 Id)
 
 uint32 FCookProfilerProvider::FindOrAddPackage(uint64 Id)
 {
+	EditAccessCheck();
+
 	uint32* Index = PackageIdToIndexMap.Find(Id);
 	if (Index != nullptr)
 	{
@@ -69,4 +71,118 @@ uint32 FCookProfilerProvider::FindOrAddPackage(uint64 Id)
 	}
 }
 
+TPagedArray<FPackageScope>& FCookProfilerProvider::FindOrAddScopeEntries(uint32 ThreadId)
+{
+	EditAccessCheck();
+
+	TPagedArray<FPackageScope>** ThreadScopeEntriesPtr = ScopeEntries.Find(ThreadId);
+	if (ThreadScopeEntriesPtr == nullptr)
+	{
+		return *(ScopeEntries.Add(ThreadId, new TPagedArray<FPackageScope>(Session.GetLinearAllocator(), 4096)));
+	}
+
+	return **ThreadScopeEntriesPtr;
+}
+
+void FCookProfilerProvider::AddScopeEntry(uint32 ThreadId, uint64 InPackageId, double Timestamp, EPackageEventStatType InType, bool InIsEnterScope)
+{
+	EditAccessCheck();
+
+	TPagedArray<FPackageScope>& Entries = FindOrAddScopeEntries(ThreadId);
+	Entries.EmplaceBack(InPackageId, Timestamp, InType, InIsEnterScope);
+}
+
+void FCookProfilerProvider::CreateAggregation(TArray64<FPackageData>& OutPackages) const
+{
+	struct PackageStackEntry
+	{
+		double StartTime;
+		EPackageEventStatType Type;
+		uint64 InPackageId;
+		double ExclTime = 0.0f;
+	};
+
+	constexpr int MaxStackSize = 128;
+	PackageStackEntry Stack[128];
+	int32 Depth = 0;
+	double LastTime = 0;
+
+	OutPackages = Packages;
+
+	for (auto& Entry : ScopeEntries)
+	{
+		TPagedArray<FPackageScope>& Entries = *Entry.Value;
+		for (auto Iterator = Entries.begin(); Iterator != Entries.end(); ++Iterator)
+		{
+			const FPackageScope& CurrentScope = *Iterator.GetCurrentItem();
+			if (CurrentScope.bIsEnterScope)
+			{
+				// Reset to default values.
+				Stack[Depth] = PackageStackEntry();
+				Stack[Depth].InPackageId = CurrentScope.PackageId;
+				Stack[Depth].Type = CurrentScope.Type;
+				Stack[Depth].StartTime = CurrentScope.Timestamp;
+
+				if (Depth > 0)
+				{
+					Stack[Depth - 1].ExclTime += CurrentScope.Timestamp - LastTime;
+				}
+
+				check(++Depth < MaxStackSize);
+			}
+			else
+			{
+				check(--Depth >= 0);
+				check(Stack[Depth].InPackageId == CurrentScope.PackageId);
+				check(Stack[Depth].Type == CurrentScope.Type);
+
+				Stack[Depth].ExclTime += CurrentScope.Timestamp - LastTime;
+				double InclTime = CurrentScope.Timestamp - Stack[Depth].StartTime;
+
+				uint32 PackageIndex;
+				const uint32* Index = PackageIdToIndexMap.Find(Stack[Depth].InPackageId);
+				if (Index != nullptr)
+				{
+					PackageIndex = *Index;
+				}
+				else
+				{
+					continue;
+				}
+				FPackageData& Package = OutPackages[PackageIndex];
+				
+				switch (Stack[Depth].Type)
+				{
+				case EPackageEventStatType::LoadPackage:
+				{
+					Package.LoadTimeIncl += InclTime;
+					Package.LoadTimeExcl += Stack[Depth].ExclTime;
+					break;
+				}
+				case EPackageEventStatType::SavePackage:
+				{
+					Package.SaveTimeIncl += InclTime;
+					Package.SaveTimeExcl += Stack[Depth].ExclTime;
+					break;
+				}
+				case EPackageEventStatType::BeginCacheForCookedPlatformData:
+				{
+					Package.BeginCacheForCookedPlatformDataIncl += InclTime;
+					Package.BeginCacheForCookedPlatformDataExcl += Stack[Depth].ExclTime;
+					break;
+				}
+				case EPackageEventStatType::IsCachedCookedPlatformDataLoaded:
+				{
+					Package.IsCachedCookedPlatformDataLoadedIncl += InclTime;
+					Package.IsCachedCookedPlatformDataLoadedExcl += Stack[Depth].ExclTime;
+					break;
+				}
+				}
+			}
+
+			LastTime = CurrentScope.Timestamp;
+		}
+	}
+}
+
 } // namespace TraceServices
diff --git a/Engine/Source/Developer/TraceServices/Private/Model/CookProfilerProviderPrivate.h b/Engine/Source/Developer/TraceServices/Private/Model/CookProfilerProviderPrivate.h
--- a/Engine/Source/Developer/TraceServices/Private/Model/CookProfilerProviderPrivate.h
+++ b/Engine/Source/Developer/TraceServices/Private/Model/CookProfilerProviderPrivate.h
@@ -2,8 +2,10 @@
 
 #pragma once
 
+#include "Common/PagedArray.h"
 #include "Common/ProviderLock.h"
 #include "Containers/Map.h"
+
 #include "TraceServices/Model/CookProfilerProvider.h"
 
 namespace TraceServices
@@ -11,6 +13,24 @@ namespace TraceServices
 
 extern thread_local FProviderLock::FThreadLocalState GCookProviderLockState;
 
+struct FPackageScope
+{
+	FPackageScope()
+	{}
+
+	FPackageScope(uint64 InPackageId, double InTimestamp, EPackageEventStatType InType, bool InIsEnterScope)
+		: PackageId(InPackageId)
+		, Timestamp(InTimestamp)
+		, Type(InType)
+		, bIsEnterScope(InIsEnterScope)
+	{}
+
+	uint64 PackageId;
+	double Timestamp;
+	EPackageEventStatType Type;
+	bool bIsEnterScope;
+};
+
 class FCookProfilerProvider
 	: public ICookProfilerProvider
 	, public IEditableCookProfilerProvider
@@ -28,6 +48,7 @@ public:
 
 	virtual uint32 GetNumPackages() const;
 	virtual void EnumeratePackages(double StartTime, double EndTime, EnumeratePackagesCallback Callback) const override;
+	virtual void CreateAggregation(TArray64<FPackageData>& OutPackages) const override;
 
 	//////////////////////////////////////////////////
 	// Edit operations
@@ -38,10 +59,13 @@ public:
 
 	virtual FPackageData* EditPackage(uint64 Id) override;
 
+	virtual void AddScopeEntry(uint32 ThreadId, uint64 InPackageId, double Timestamp, EPackageEventStatType InType, bool InIsEnterScope);
+
 	//////////////////////////////////////////////////
 
 private:
 	uint32 FindOrAddPackage(uint64 Id);
+	TPagedArray<FPackageScope>& FindOrAddScopeEntries(uint32 ThreadId);
 
 private:
 	mutable FProviderLock Lock;
@@ -50,6 +74,8 @@ private:
 
 	TMap<uint64, uint32> PackageIdToIndexMap;
 	TArray64<FPackageData> Packages;
+
+	TMap<uint32, TPagedArray<FPackageScope>*> ScopeEntries; // The Key is the ThreadId of the scope entries.
 };
 
 } // namespace TraceServices
diff --git a/Engine/Source/Developer/TraceServices/Public/TraceServices/Model/CookProfilerProvider.h b/Engine/Source/Developer/TraceServices/Public/TraceServices/Model/CookProfilerProvider.h
--- a/Engine/Source/Developer/TraceServices/Public/TraceServices/Model/CookProfilerProvider.h
+++ b/Engine/Source/Developer/TraceServices/Public/TraceServices/Model/CookProfilerProvider.h
@@ -3,6 +3,7 @@
 #pragma once
 
 #include "HAL/Platform.h"
+#include "ProfilingDebugging/CookStats.h"
 #include "Templates/Function.h"
 #include "TraceServices/Model/AnalysisSession.h"
 #include "UObject/NameTypes.h"
@@ -15,10 +16,19 @@ struct FPackageData
 	FPackageData(uint64 InId);
 	uint64 Id = -1;
 	const TCHAR* Name;
-	double LoadTime = 0;
-	double SaveTime = 0;
-	double BeginCacheForCookedPlatformData = 0; // BeginCacheForCookedPlatformData is the name of the function from the cooker.
-	double IsCachedCookedPlatformDataLoaded = 0; // IsCachedCookedPlatformDataLoaded is the name of the function from the cooker.
+
+	double LoadTimeIncl = 0;
+	double LoadTimeExcl = 0;
+
+	double SaveTimeIncl = 0;
+	double SaveTimeExcl = 0;
+
+	double BeginCacheForCookedPlatformDataIncl = 0; // BeginCacheForCookedPlatformData is the name of the function from the cooker.
+	double BeginCacheForCookedPlatformDataExcl = 0;
+
+	double IsCachedCookedPlatformDataLoadedIncl = 0; // IsCachedCookedPlatformDataLoaded is the name of the function from the cooker.
+	double IsCachedCookedPlatformDataLoadedExcl = 0;
+
 	const TCHAR* AssetClass;
 };
 
@@ -35,6 +45,8 @@ public:
 
 	virtual uint32 GetNumPackages() const = 0;
 	virtual void EnumeratePackages(double StartTime, double EndTime, EnumeratePackagesCallback Callback) const = 0;
+	
+	virtual void CreateAggregation(TArray64<FPackageData>& OutPackages) const = 0;
 };
 
 class IEditableCookProfilerProvider : public IEditableProvider
@@ -47,6 +59,8 @@ public:
 	virtual void EditAccessCheck() const = 0;
 
 	virtual FPackageData* EditPackage(uint64 Id) = 0;
+
+	virtual void AddScopeEntry(uint32 ThreadId, uint64 InPackageId, double Timestamp, EPackageEventStatType InType, bool InIsEnterScope) = 0;
 };
 
 TRACESERVICES_API FName GetCookProfilerProviderName();
diff --git a/Engine/Source/Runtime/Core/Private/ProfilingDebugging/CookStats.cpp b/Engine/Source/Runtime/Core/Private/ProfilingDebugging/CookStats.cpp
--- a/Engine/Source/Runtime/Core/Private/ProfilingDebugging/CookStats.cpp
+++ b/Engine/Source/Runtime/Core/Private/ProfilingDebugging/CookStats.cpp
@@ -6,6 +6,13 @@
 
 #if ENABLE_COOK_STATS
 
+CORE_API FCookStatsManager::FGatherCookStatsDelegate FCookStatsManager::CookStatsCallbacks;
+
+CORE_API void FCookStatsManager::LogCookStats(AddStatFuncRef AddStat)
+{
+	CookStatsCallbacks.Broadcast(AddStat);
+}
+
 UE_TRACE_CHANNEL_DEFINE(CookChannel)
 
 UE_TRACE_EVENT_BEGIN(CookTrace, Package)
@@ -14,32 +21,70 @@ UE_TRACE_EVENT_BEGIN(CookTrace, Package)
 	UE_TRACE_EVENT_FIELD(uint64, Cycle)
 UE_TRACE_EVENT_END()
 
-UE_TRACE_EVENT_BEGIN(CookTrace, PackageStat)
+UE_TRACE_EVENT_BEGIN(CookTrace, PackageAssetClass)
+	UE_TRACE_EVENT_FIELD(uint64, Id)
+	UE_TRACE_EVENT_FIELD(UE::Trace::WideString, ClassName)
+UE_TRACE_EVENT_END()
+
+void TracePackage(uint64 InId, const FStringView InName)
+{
+	UE_TRACE_LOG(CookTrace, Package, CookChannel)
+		<< Package.Id(InId)
+		<< Package.Name(InName.GetData(), uint16(InName.Len()))
+		<< Package.Cycle(FPlatformTime::Cycles64());
+}
+
+void TracePackageAssetClass(uint64 InId, const FStringView InName)
+{
+	UE_TRACE_LOG(CookTrace, PackageAssetClass, CookChannel)
+		<< PackageAssetClass.Id(InId)
+		<< PackageAssetClass.ClassName(InName.GetData(), uint16(InName.Len()));
+}
+
+#if UE_INSIGHTS_COOK_PROFILER_VERSION == 2
+
+UE_TRACE_EVENT_BEGIN(CookTrace, PackageStatBeginScope)
 	UE_TRACE_EVENT_FIELD(uint64, Id)
-	UE_TRACE_EVENT_FIELD(uint64, Duration)
+	UE_TRACE_EVENT_FIELD(uint64, Time)
 	UE_TRACE_EVENT_FIELD(uint8, StatType)
 UE_TRACE_EVENT_END()
 
-UE_TRACE_EVENT_BEGIN(CookTrace, PackageAssetClass)
+UE_TRACE_EVENT_BEGIN(CookTrace, PackageStatEndScope)
 	UE_TRACE_EVENT_FIELD(uint64, Id)
-	UE_TRACE_EVENT_FIELD(UE::Trace::WideString, ClassName)
+	UE_TRACE_EVENT_FIELD(uint64, Time)
+	UE_TRACE_EVENT_FIELD(uint8, StatType)
 UE_TRACE_EVENT_END()
 
-CORE_API FCookStatsManager::FGatherCookStatsDelegate FCookStatsManager::CookStatsCallbacks;
+void TracePackageStatBeginScope(uint64 InId, uint64 Time, EPackageEventStatType StatType)
+{
+	UE_TRACE_LOG(CookTrace, PackageStatBeginScope, CookChannel)
+		<< PackageStatBeginScope.Id(InId)
+		<< PackageStatBeginScope.Time(Time)
+		<< PackageStatBeginScope.StatType((uint8)StatType);
+}
 
-CORE_API void FCookStatsManager::LogCookStats(AddStatFuncRef AddStat)
+void TracePackageStatEndScope(uint64 InId, uint64 Time, EPackageEventStatType StatType)
 {
-	CookStatsCallbacks.Broadcast(AddStat);
+	UE_TRACE_LOG(CookTrace, PackageStatEndScope, CookChannel)
+		<< PackageStatEndScope.Id(InId)
+		<< PackageStatEndScope.Time(Time)
+		<< PackageStatEndScope.StatType((uint8)StatType);
 }
 
-void TracePackage(uint64 InId, const FStringView InName)
+void TracePackageStat(uint64 InId, uint64 Duration, EPackageEventStatType StatType)
 {
-	UE_TRACE_LOG(CookTrace, Package, CookChannel)
-		<< Package.Id(InId)
-		<< Package.Name(InName.GetData(), uint16(InName.Len()))
-		<< Package.Cycle(FPlatformTime::Cycles64());
 }
 
+#endif
+
+#if UE_INSIGHTS_COOK_PROFILER_VERSION == 1
+
+UE_TRACE_EVENT_BEGIN(CookTrace, PackageStat)
+	UE_TRACE_EVENT_FIELD(uint64, Id)
+	UE_TRACE_EVENT_FIELD(int64, Duration)
+	UE_TRACE_EVENT_FIELD(uint8, StatType)
+UE_TRACE_EVENT_END()
+
 void TracePackageStat(uint64 InId, uint64 Duration, EPackageEventStatType StatType)
 {
 	UE_TRACE_LOG(CookTrace, PackageStat, CookChannel)
@@ -48,12 +93,7 @@ void TracePackageStat(uint64 InId, uint64 Duration, EPackageEventStatType StatTy
 		<< PackageStat.StatType((uint8)StatType);
 }
 
-void TracePackageAssetClass(uint64 InId, const FStringView InName)
-{
-	UE_TRACE_LOG(CookTrace, PackageAssetClass, CookChannel)
-		<< PackageAssetClass.Id(InId)
-		<< PackageAssetClass.ClassName(InName.GetData(), uint16(InName.Len()));
-}
+#endif
 
 bool ShouldTracePackageInfo()
 {
diff --git a/Engine/Source/Runtime/Core/Public/ProfilingDebugging/CookStats.h b/Engine/Source/Runtime/Core/Public/ProfilingDebugging/CookStats.h
--- a/Engine/Source/Runtime/Core/Public/ProfilingDebugging/CookStats.h
+++ b/Engine/Source/Runtime/Core/Public/ProfilingDebugging/CookStats.h
@@ -439,11 +439,49 @@ struct FDerivedDataUsageStats
 
 CORE_API UE_TRACE_CHANNEL_EXTERN(CookChannel);
 
+#define UE_INSIGHTS_COOK_PROFILER_VERSION 2
+
 CORE_API void TracePackage(uint64 InId, const FStringView InName);
-CORE_API void TracePackageStat(uint64 InId, uint64 Duration, EPackageEventStatType StatType);
 CORE_API void TracePackageAssetClass(uint64 InId, const FStringView InName);
 CORE_API bool ShouldTracePackageInfo();
 
+#if UE_INSIGHTS_COOK_PROFILER_VERSION == 2
+
+UE_DEPRECATED(5.5, "TracePackageStat is deprecated, please use TracePackageStatBeginScope and TracePackageStatEndScope instead.")
+CORE_API void TracePackageStat(uint64 InId, uint64 Duration, EPackageEventStatType StatType);
+
+CORE_API void TracePackageStatBeginScope(uint64 InId, uint64 Time, EPackageEventStatType StatType);
+CORE_API void TracePackageStatEndScope(uint64 InId, uint64 Time, EPackageEventStatType StatType);
+
+struct FScopedCookStat
+{
+	FScopedCookStat(const FScopedCookStat&) = delete;
+	FScopedCookStat(FScopedCookStat&&) = delete;
+
+	FScopedCookStat(uint64 Id, EPackageEventStatType InStatType)
+		: Id(Id)
+		, StatType(InStatType)
+	{
+		TracePackageStatBeginScope(Id, FPlatformTime::Cycles64(), StatType);
+	}
+
+	~FScopedCookStat()
+	{
+		TracePackageStatEndScope(Id, FPlatformTime::Cycles64(), StatType);
+	}
+
+private:
+	uint64 Id;
+	EPackageEventStatType StatType;
+};
+
+#endif
+
+#if UE_INSIGHTS_COOK_PROFILER_VERSION == 1
+
+UE_DEPRECATED(5.5, "TracePackageStat is deprecated, please use TracePackageStatBeginScope and TracePackageStatEndScope instead.")
+CORE_API void TracePackageStat(uint64 InId, uint64 Duration, EPackageEventStatType StatType);
+
 struct FScopedCookStat
 {
 	FScopedCookStat(const FScopedCookStat&) = delete;
@@ -458,7 +496,8 @@ struct FScopedCookStat
 
 	~FScopedCookStat()
 	{
-		TracePackageStat(Id, FPlatformTime::Cycles64() - StartTimestamp, StatType);
+		uint64 Duration = FPlatformTime::Cycles64() - StartTimestamp;
+		TracePackageStat(Id, Duration, StatType);
 	}
 
 private:
@@ -467,6 +506,8 @@ private:
 	EPackageEventStatType StatType;
 };
 
+#endif
+
 #define UE_SCOPED_COOK_STAT(Name, StatType) \
 	FScopedCookStat  PREPROCESSOR_JOIN(__CookTimerScope, __LINE__)(Name.ToUnstableInt(), StatType);
 
@@ -485,6 +526,7 @@ private:
 	__CookStatScope.Reset(); \
 	__PrevPackageName = NAME_None;
 #else
+
 #define COOK_STAT(...)
 #define UE_SCOPED_COOK_STAT(...)
 #define UE_MULTI_SCOPED_COOK_STAT_INIT(...)
@@ -493,6 +535,8 @@ private:
 
 #define TracePackage(...)
 #define TracePackageStat(...)
+#define TracePackageStatBeginScope(...)
+#define TracePackageStatEndScope(...)
 #define TracePackageAssetClass(...)
 #define ShouldTracePackageInfo(...) false
 
