[Insights] Improved export to csv/txt performance by reducing I/O write operations (ex: exporting timing events is now 10x faster).

diff --git a/Engine/Source/Developer/TraceInsights/Private/Insights/ViewModels/TimingExporter.cpp b/Engine/Source/Developer/TraceInsights/Private/Insights/ViewModels/TimingExporter.cpp
--- a/Engine/Source/Developer/TraceInsights/Private/Insights/ViewModels/TimingExporter.cpp
+++ b/Engine/Source/Developer/TraceInsights/Private/Insights/ViewModels/TimingExporter.cpp
@@ -73,13 +73,13 @@ void FTimingExporter::Error(const FText& InMessage) const
 
 ////////////////////////////////////////////////////////////////////////////////////////////////////
 
-void FTimingExporter::AppendString(FUtf8StringBuilder& StringBuilder, const TCHAR* Str, UTF8CHAR Separator)
+void FTimingExporter::FUtf8Writer::AppendString(const TCHAR* Str)
 {
 	if (Str == nullptr || Str[0] == TCHAR('\0'))
 	{
 		// nothing to append
 	}
-	else if (FCString::Strchr(Str, Separator) != nullptr)
+	else if (FCString::Strchr(Str, TCHAR(Separator)) != nullptr)
 	{
 		if (FCString::Strchr(Str, TEXT('\"')) != nullptr)
 		{
@@ -102,6 +102,17 @@ void FTimingExporter::AppendString(FUtf8StringBuilder& StringBuilder, const TCHA
 
 ////////////////////////////////////////////////////////////////////////////////////////////////////
 
+void FTimingExporter::FUtf8Writer::WriteStringBuilder(int32 CacheLen)
+{
+	if (StringBuilder.Len() > CacheLen)
+	{
+		FileHandle->Write((const uint8*)StringBuilder.ToString(), StringBuilder.Len() * sizeof(UTF8CHAR));
+		StringBuilder.Reset();
+	}
+}
+
+////////////////////////////////////////////////////////////////////////////////////////////////////
+
 int32 FTimingExporter::ExportThreadsAsText(const FString& Filename, FExportThreadsParams& Params) const
 {
 	checkf(Params.Columns == nullptr, TEXT("Custom list of columns is not yet supported!"));
@@ -114,72 +125,61 @@ int32 FTimingExporter::ExportThreadsAsText(const FString& Filename, FExportThrea
 	{
 		return -1;
 	}
-
-	UTF8CHAR Separator = UTF8CHAR('\t');
-	if (Filename.EndsWith(TEXT(".csv")))
-	{
-		Separator = UTF8CHAR(',');
-	}
-	const UTF8CHAR LineEnd = UTF8CHAR('\n');
-
-	FUtf8StringBuilder StringBuilder;
+	bool bIsCSV = Filename.EndsWith(TEXT(".csv"));
+	FUtf8Writer Writer(ExportFileHandle, bIsCSV);
+	FUtf8StringBuilder& StringBuilder = Writer.GetStringBuilder();
 
 	// Write header.
 	{
 		StringBuilder.Append(UTF8TEXT("Id"));
-		StringBuilder.AppendChar(Separator);
+		Writer.AppendSeparator();
 		StringBuilder.Append(UTF8TEXT("Name"));
-		StringBuilder.AppendChar(Separator);
+		Writer.AppendSeparator();
 		StringBuilder.Append(UTF8TEXT("Group"));
-		StringBuilder.AppendChar(LineEnd);
-
-		ExportFileHandle->Write((const uint8*)StringBuilder.ToString(), StringBuilder.Len() * sizeof(UTF8CHAR));
+		Writer.AppendLineEnd();
 	}
 
 	int32 ThreadCount = 0;
 
 	// Write values.
 	{
-		StringBuilder.Reset();
 		StringBuilder.Appendf(UTF8TEXT("%u"), FGpuTimingTrack::Gpu1ThreadId);
-		StringBuilder.AppendChar(Separator);
+		Writer.AppendSeparator();
 		StringBuilder.Append(UTF8TEXT("GPU1"));
-		StringBuilder.AppendChar(Separator);
+		Writer.AppendSeparator();
 		StringBuilder.Append(UTF8TEXT("GPU"));
-		StringBuilder.AppendChar(LineEnd);
-		ExportFileHandle->Write((const uint8*)StringBuilder.ToString(), StringBuilder.Len() * sizeof(UTF8CHAR));
+		Writer.AppendLineEnd();
 		++ThreadCount;
 
-		StringBuilder.Reset();
 		StringBuilder.Appendf(UTF8TEXT("%u"), FGpuTimingTrack::Gpu2ThreadId);
-		StringBuilder.AppendChar(Separator);
+		Writer.AppendSeparator();
 		StringBuilder.Append(UTF8TEXT("GPU2"));
-		StringBuilder.AppendChar(Separator);
+		Writer.AppendSeparator();
 		StringBuilder.Append(UTF8TEXT("GPU"));
-		StringBuilder.AppendChar(LineEnd);
-		ExportFileHandle->Write((const uint8*)StringBuilder.ToString(), StringBuilder.Len() * sizeof(UTF8CHAR));
+		Writer.AppendLineEnd();
 		++ThreadCount;
 
 		// Iterate the CPU threads.
 		{
 			TraceServices::FAnalysisSessionReadScope SessionReadScope(Session);
+
 			const TraceServices::IThreadProvider& ThreadProvider = TraceServices::ReadThreadProvider(Session);
+
 			ThreadProvider.EnumerateThreads(
 				[&](const TraceServices::FThreadInfo& ThreadInfo)
 				{
-					StringBuilder.Reset();
 					StringBuilder.Appendf(UTF8TEXT("%u"), ThreadInfo.Id);
-					StringBuilder.AppendChar(Separator);
-					AppendString(StringBuilder, ThreadInfo.Name, Separator);
-					StringBuilder.AppendChar(Separator);
-					AppendString(StringBuilder, ThreadInfo.GroupName, Separator);
-					StringBuilder.AppendChar(LineEnd);
-					ExportFileHandle->Write((const uint8*)StringBuilder.ToString(), StringBuilder.Len() * sizeof(UTF8CHAR));
+					Writer.AppendSeparator();
+					Writer.AppendString(ThreadInfo.Name);
+					Writer.AppendSeparator();
+					Writer.AppendString(ThreadInfo.GroupName);
+					Writer.AppendLineEnd();
 					++ThreadCount;
 				});
 		}
 	}
 
+	Writer.Flush();
 	ExportFileHandle->Flush();
 	delete ExportFileHandle;
 	ExportFileHandle = nullptr;
@@ -205,30 +205,22 @@ int32 FTimingExporter::ExportTimersAsText(const FString& Filename, FExportTimers
 	{
 		return -1;
 	}
-
-	UTF8CHAR Separator = UTF8CHAR('\t');
-	if (Filename.EndsWith(TEXT(".csv")))
-	{
-		Separator = UTF8CHAR(',');
-	}
-	const UTF8CHAR LineEnd = UTF8CHAR('\n');
-
-	FUtf8StringBuilder StringBuilder;
+	bool bIsCSV = Filename.EndsWith(TEXT(".csv"));
+	FUtf8Writer Writer(ExportFileHandle, bIsCSV);
+	FUtf8StringBuilder& StringBuilder = Writer.GetStringBuilder();
 
 	// Write header.
 	{
 		StringBuilder.Append(UTF8TEXT("Id"));
-		StringBuilder.AppendChar(Separator);
+		Writer.AppendSeparator();
 		StringBuilder.Append(UTF8TEXT("Type"));
-		StringBuilder.AppendChar(Separator);
+		Writer.AppendSeparator();
 		StringBuilder.Append(UTF8TEXT("Name"));
-		StringBuilder.AppendChar(Separator);
+		Writer.AppendSeparator();
 		StringBuilder.Append(UTF8TEXT("File"));
-		StringBuilder.AppendChar(Separator);
+		Writer.AppendSeparator();
 		StringBuilder.Append(UTF8TEXT("Line"));
-		StringBuilder.AppendChar(LineEnd);
-
-		ExportFileHandle->Write((const uint8*)StringBuilder.ToString(), StringBuilder.Len() * sizeof(UTF8CHAR));
+		Writer.AppendLineEnd();
 	}
 
 	uint32 TimerCount = 0;
@@ -247,24 +239,23 @@ int32 FTimingExporter::ExportTimersAsText(const FString& Filename, FExportTimers
 		for (uint32 TimerIndex = 0; TimerIndex < TimerCount; ++TimerIndex)
 		{
 			const TraceServices::FTimingProfilerTimer& Timer = *(TimerReader->GetTimer(TimerIndex));
-			StringBuilder.Reset();
 			StringBuilder.Appendf(UTF8TEXT("%u"), Timer.Id);
-			StringBuilder.AppendChar(Separator);
+			Writer.AppendSeparator();
 			StringBuilder.Append(Timer.IsGpuTimer ? UTF8TEXT("GPU") : UTF8TEXT("CPU"));
-			StringBuilder.AppendChar(Separator);
-			AppendString(StringBuilder, Timer.Name, Separator);
-			StringBuilder.AppendChar(Separator);
+			Writer.AppendSeparator();
+			Writer.AppendString(Timer.Name);
+			Writer.AppendSeparator();
 			if (Timer.File)
 			{
 				StringBuilder.Append((const UTF8CHAR*)TCHAR_TO_UTF8(Timer.File));
 			}
-			StringBuilder.AppendChar(Separator);
+			Writer.AppendSeparator();
 			StringBuilder.Appendf(UTF8TEXT("%u"), Timer.Line);
-			StringBuilder.AppendChar(LineEnd);
-			ExportFileHandle->Write((const uint8*)StringBuilder.ToString(), StringBuilder.Len() * sizeof(UTF8CHAR));
+			Writer.AppendLineEnd();
 		}
 	}
 
+	Writer.Flush();
 	ExportFileHandle->Flush();
 	delete ExportFileHandle;
 	ExportFileHandle = nullptr;
@@ -354,14 +345,12 @@ void FTimingExporter::ExportTimingEvents_WriteHeader(FExportTimingEventsInternal
 			}
 			else
 			{
-				Params.StringBuilder.AppendChar(Params.Separator);
+				Params.Writer.AppendSeparator();
 			}
-			Params.StringBuilder.Append(Column.GetPlainNameString());
+			Params.Writer.Append(Column.GetPlainNameString());
 		}
 	}
-	Params.StringBuilder.AppendChar(Params.LineEnd);
-
-	Params.ExportFileHandle->Write((const uint8*)Params.StringBuilder.ToString(), Params.StringBuilder.Len() * sizeof(UTF8CHAR));
+	Params.Writer.AppendLineEnd();
 }
 
 ////////////////////////////////////////////////////////////////////////////////////////////////////
@@ -407,8 +396,7 @@ int32 FTimingExporter::ExportTimingEvents_WriteEvents(FExportTimingEventsInterna
 				{
 					if (!Params.UserParams.TimingEventFilter || Params.UserParams.TimingEventFilter(EventStartTime, EventEndTime, EventDepth, Event))
 					{
-						Params.StringBuilder.Reset();
-
+						FUtf8StringBuilder& StringBuilder = Params.Writer.GetStringBuilder();
 						bool bFirst = true;
 						for (const FName& Column : Params.Columns)
 						{
@@ -420,47 +408,45 @@ int32 FTimingExporter::ExportTimingEvents_WriteEvents(FExportTimingEventsInterna
 								}
 								else
 								{
-									Params.StringBuilder.AppendChar(Params.Separator);
+									Params.Writer.AppendSeparator();
 								}
 
 								if (Column == ExportTimingEvents_ThreadIdColumn)
 								{
-									Params.StringBuilder.Appendf(UTF8TEXT("%u"), Params.ThreadId);
+									StringBuilder.Appendf(UTF8TEXT("%u"), Params.ThreadId);
 								}
 								else if (Column == ExportTimingEvents_ThreadNameColumn)
 								{
-									AppendString(Params.StringBuilder, Params.ThreadName, Params.Separator);
+									Params.Writer.AppendString(Params.ThreadName);
 								}
 								else if (Column == ExportTimingEvents_TimerIdColumn)
 								{
-									Params.StringBuilder.Appendf(UTF8TEXT("%u"), Event.TimerIndex);
+									StringBuilder.Appendf(UTF8TEXT("%u"), Event.TimerIndex);
 								}
 								else if (Column == ExportTimingEvents_TimerNameColumn)
 								{
 									const TCHAR* TimerName = Timers.FindRef(Event.TimerIndex);
-									AppendString(Params.StringBuilder, TimerName, Params.Separator);
+									Params.Writer.AppendString(TimerName);
 								}
 								else if (Column == ExportTimingEvents_StartTimeColumn)
 								{
-									Params.StringBuilder.Appendf(UTF8TEXT("%.9g"), EventStartTime);
+									StringBuilder.Appendf(UTF8TEXT("%.9g"), EventStartTime);
 								}
 								else if (Column == ExportTimingEvents_EndTimeColumn)
 								{
-									Params.StringBuilder.Appendf(UTF8TEXT("%.9g"), EventEndTime);
+									StringBuilder.Appendf(UTF8TEXT("%.9g"), EventEndTime);
 								}
 								else if (Column == ExportTimingEvents_DurationColumn)
 								{
-									Params.StringBuilder.Appendf(UTF8TEXT("%.9f"), EventEndTime - EventStartTime);
+									StringBuilder.Appendf(UTF8TEXT("%.9f"), EventEndTime - EventStartTime);
 								}
 								else if (Column == ExportTimingEvents_DepthColumn)
 								{
-									Params.StringBuilder.Appendf(UTF8TEXT("%u"), EventDepth);
+									StringBuilder.Appendf(UTF8TEXT("%u"), EventDepth);
 								}
 							}
 						}
-						Params.StringBuilder.AppendChar(Params.LineEnd);
-
-						Params.ExportFileHandle->Write((const uint8*)Params.StringBuilder.ToString(), Params.StringBuilder.Len() * sizeof(UTF8CHAR));
+						Params.Writer.AppendLineEnd();
 						++TimingEventCount;
 					}
 
@@ -523,17 +509,11 @@ int32 FTimingExporter::ExportTimingEventsAsText(const FString& Filename, FExport
 	{
 		return -1;
 	}
+	bool bIsCSV = Filename.EndsWith(TEXT(".csv"));
+	FUtf8Writer Writer(ExportFileHandle, bIsCSV);
+	FUtf8StringBuilder& StringBuilder = Writer.GetStringBuilder();
 
-	UTF8CHAR Separator = UTF8CHAR('\t');
-	if (Filename.EndsWith(TEXT(".csv")))
-	{
-		Separator = UTF8CHAR(',');
-	}
-	const UTF8CHAR LineEnd = UTF8CHAR('\n');
-
-	FUtf8StringBuilder StringBuilder;
-
-	FExportTimingEventsInternalParams InternalParams = { *this, Params, Columns, ExportFileHandle, Separator, LineEnd, StringBuilder, 0 };
+	FExportTimingEventsInternalParams InternalParams = { *this, Params, Columns, Writer, 0 };
 
 	// Write header.
 	ExportTimingEvents_WriteHeader(InternalParams);
@@ -541,6 +521,7 @@ int32 FTimingExporter::ExportTimingEventsAsText(const FString& Filename, FExport
 	// Write values.
 	int32 TimingEventCount = ExportTimingEvents_WriteEvents(InternalParams);
 
+	Writer.Flush();
 	ExportFileHandle->Flush();
 	delete ExportFileHandle;
 	ExportFileHandle = nullptr;
@@ -770,7 +751,9 @@ FTimingExporter::FThreadFilterFunc FTimingExporter::MakeThreadFilterInclusive(co
 	// Iterate the CPU threads.
 	{
 		TraceServices::FAnalysisSessionReadScope SessionReadScope(Session);
+
 		const TraceServices::IThreadProvider& ThreadProvider = TraceServices::ReadThreadProvider(Session);
+
 		ThreadProvider.EnumerateThreads(
 			[&Threads](const TraceServices::FThreadInfo& ThreadInfo)
 			{
@@ -975,26 +958,18 @@ int32 FTimingExporter::ExportCountersAsText(const FString& Filename, FExportCoun
 	{
 		return -1;
 	}
-
-	UTF8CHAR Separator = UTF8CHAR('\t');
-	if (Filename.EndsWith(TEXT(".csv")))
-	{
-		Separator = UTF8CHAR(',');
-	}
-	const UTF8CHAR LineEnd = UTF8CHAR('\n');
-
-	FUtf8StringBuilder StringBuilder;
+	bool bIsCSV = Filename.EndsWith(TEXT(".csv"));
+	FUtf8Writer Writer(ExportFileHandle, bIsCSV);
+	FUtf8StringBuilder& StringBuilder = Writer.GetStringBuilder();
 
 	// Write header.
 	{
 		StringBuilder.Append(UTF8TEXT("Id"));
-		StringBuilder.AppendChar(Separator);
+		Writer.AppendSeparator();
 		StringBuilder.Append(UTF8TEXT("Type"));
-		StringBuilder.AppendChar(Separator);
+		Writer.AppendSeparator();
 		StringBuilder.Append(UTF8TEXT("Name"));
-		StringBuilder.AppendChar(LineEnd);
-
-		ExportFileHandle->Write((const uint8*)StringBuilder.ToString(), StringBuilder.Len() * sizeof(UTF8CHAR));
+		Writer.AppendLineEnd();
 	}
 
 	int32 CounterCount = 0;
@@ -1003,33 +978,34 @@ int32 FTimingExporter::ExportCountersAsText(const FString& Filename, FExportCoun
 	if (true) // TraceServices::ReadCounterProvider(Session)
 	{
 		TraceServices::FAnalysisSessionReadScope SessionReadScope(Session);
+
 		const TraceServices::ICounterProvider& CounterProvider = TraceServices::ReadCounterProvider(Session);
 
-		CounterProvider.EnumerateCounters([&](uint32 CounterId, const TraceServices::ICounter& Counter)
-		{
-			StringBuilder.Reset();
-			StringBuilder.Appendf(UTF8TEXT("%u"), CounterId);
-			StringBuilder.AppendChar(Separator);
-			if (Counter.IsFloatingPoint())
-			{
-				StringBuilder.Append(UTF8TEXT("Double"));
-			}
-			else
+		CounterProvider.EnumerateCounters(
+			[&](uint32 CounterId, const TraceServices::ICounter& Counter)
 			{
-				StringBuilder.Append(UTF8TEXT("Int64"));
-			}
-			if (Counter.IsResetEveryFrame())
-			{
-				StringBuilder.Append(UTF8TEXT("|ResetEveryFrame"));
-			}
-			StringBuilder.AppendChar(Separator);
-			AppendString(StringBuilder, Counter.GetName(), Separator);
-			StringBuilder.AppendChar(LineEnd);
-			ExportFileHandle->Write((const uint8*)StringBuilder.ToString(), StringBuilder.Len() * sizeof(UTF8CHAR));
-			++CounterCount;
-		});
+				StringBuilder.Appendf(UTF8TEXT("%u"), CounterId);
+				Writer.AppendSeparator();
+				if (Counter.IsFloatingPoint())
+				{
+					StringBuilder.Append(UTF8TEXT("Double"));
+				}
+				else
+				{
+					StringBuilder.Append(UTF8TEXT("Int64"));
+				}
+				if (Counter.IsResetEveryFrame())
+				{
+					StringBuilder.Append(UTF8TEXT("|ResetEveryFrame"));
+				}
+				Writer.AppendSeparator();
+				Writer.AppendString(Counter.GetName());
+				Writer.AppendLineEnd();
+				++CounterCount;
+			});
 	}
 
+	Writer.Flush();
 	ExportFileHandle->Flush();
 	delete ExportFileHandle;
 	ExportFileHandle = nullptr;
@@ -1055,34 +1031,27 @@ int32 FTimingExporter::ExportCounterAsText(const FString& Filename, uint32 Count
 	{
 		return -1;
 	}
-
-	UTF8CHAR Separator = UTF8CHAR('\t');
-	if (Filename.EndsWith(TEXT(".csv")))
-	{
-		Separator = UTF8CHAR(',');
-	}
-	const UTF8CHAR LineEnd = UTF8CHAR('\n');
-
-	FUtf8StringBuilder StringBuilder;
+	bool bIsCSV = Filename.EndsWith(TEXT(".csv"));
+	FUtf8Writer Writer(ExportFileHandle, bIsCSV);
+	FUtf8StringBuilder& StringBuilder = Writer.GetStringBuilder();
 
 	// Write header.
 	if (Params.bExportOps)
 	{
 		StringBuilder.Append(UTF8TEXT("Time"));
-		StringBuilder.AppendChar(Separator);
+		Writer.AppendSeparator();
 		StringBuilder.Append(UTF8TEXT("Op"));
-		StringBuilder.AppendChar(Separator);
+		Writer.AppendSeparator();
 		StringBuilder.Append(UTF8TEXT("Value"));
-		StringBuilder.AppendChar(LineEnd);
+		Writer.AppendLineEnd();
 	}
 	else
 	{
 		StringBuilder.Append(UTF8TEXT("Time"));
-		StringBuilder.AppendChar(Separator);
+		Writer.AppendSeparator();
 		StringBuilder.Append(UTF8TEXT("Value"));
-		StringBuilder.AppendChar(LineEnd);
+		Writer.AppendLineEnd();
 	}
-	ExportFileHandle->Write((const uint8*)StringBuilder.ToString(), StringBuilder.Len() * sizeof(UTF8CHAR));
 
 	FString CounterName;
 	int32 ValueCount = 0;
@@ -1091,98 +1060,97 @@ int32 FTimingExporter::ExportCounterAsText(const FString& Filename, uint32 Count
 	if (true) // TraceServices::ReadCounterProvider(Session)
 	{
 		TraceServices::FAnalysisSessionReadScope SessionReadScope(Session);
-		const TraceServices::ICounterProvider& CounterProvider = TraceServices::ReadCounterProvider(Session);
 
-		CounterProvider.ReadCounter(CounterId, [&](const TraceServices::ICounter& Counter)
-		{
-			CounterName = Counter.GetName();
+		const TraceServices::ICounterProvider& CounterProvider = TraceServices::ReadCounterProvider(Session);
 
-			// Iterate the counter values.
-			if (Params.bExportOps)
+		CounterProvider.ReadCounter(CounterId,
+			[&](const TraceServices::ICounter& Counter)
 			{
-				if (Counter.IsFloatingPoint())
-				{
-					Counter.EnumerateFloatOps(Params.IntervalStartTime, Params.IntervalEndTime, false, [&](double Time, TraceServices::ECounterOpType Op, double Value)
-						{
-							StringBuilder.Reset();
-							StringBuilder.Appendf(UTF8TEXT("%.9f"), Time);
-							StringBuilder.AppendChar(Separator);
-							switch (Op)
-							{
-							case TraceServices::ECounterOpType::Set:
-								StringBuilder.Append(UTF8TEXT("Set"));
-								break;
-							case TraceServices::ECounterOpType::Add:
-								StringBuilder.Append(UTF8TEXT("Add"));
-								break;
-							default:
-								StringBuilder.Appendf(UTF8TEXT("%d"), int32(Op));
-							}
-							StringBuilder.AppendChar(Separator);
-							StringBuilder.Appendf(UTF8TEXT("%.9f"), Value);
-							StringBuilder.AppendChar(LineEnd);
-							ExportFileHandle->Write((const uint8*)StringBuilder.ToString(), StringBuilder.Len() * sizeof(UTF8CHAR));
-							++ValueCount;
-						});
-				}
-				else
+				CounterName = Counter.GetName();
+
+				// Iterate the counter values.
+				if (Params.bExportOps)
 				{
-					Counter.EnumerateOps(Params.IntervalStartTime, Params.IntervalEndTime, false, [&](double Time, TraceServices::ECounterOpType Op, int64 IntValue)
-						{
-							StringBuilder.Reset();
-							StringBuilder.Appendf(UTF8TEXT("%.9f"), Time);
-							StringBuilder.AppendChar(Separator);
-							switch (Op)
+					if (Counter.IsFloatingPoint())
+					{
+						Counter.EnumerateFloatOps(Params.IntervalStartTime, Params.IntervalEndTime, false,
+							[&](double Time, TraceServices::ECounterOpType Op, double Value)
 							{
-							case TraceServices::ECounterOpType::Set:
-								StringBuilder.Append(UTF8TEXT("Set"));
-								break;
-							case TraceServices::ECounterOpType::Add:
-								StringBuilder.Append(UTF8TEXT("Add"));
-								break;
-							default:
-								StringBuilder.Appendf(UTF8TEXT("%d"), int32(Op));
-							}
-							StringBuilder.AppendChar(Separator);
-							StringBuilder.Appendf(UTF8TEXT("%lli"), IntValue);
-							StringBuilder.AppendChar(LineEnd);
-							ExportFileHandle->Write((const uint8*)StringBuilder.ToString(), StringBuilder.Len() * sizeof(UTF8CHAR));
-							++ValueCount;
-						});
-				}
-			}
-			else
-			{
-				if (Counter.IsFloatingPoint())
-				{
-					Counter.EnumerateFloatValues(Params.IntervalStartTime, Params.IntervalEndTime, false, [&](double Time, double Value)
+								StringBuilder.Appendf(UTF8TEXT("%.9f"), Time);
+								Writer.AppendSeparator();
+								switch (Op)
+								{
+									case TraceServices::ECounterOpType::Set:
+										StringBuilder.Append(UTF8TEXT("Set"));
+										break;
+									case TraceServices::ECounterOpType::Add:
+										StringBuilder.Append(UTF8TEXT("Add"));
+										break;
+									default:
+										StringBuilder.Appendf(UTF8TEXT("%d"), int32(Op));
+								}
+								Writer.AppendSeparator();
+								StringBuilder.Appendf(UTF8TEXT("%.9f"), Value);
+								Writer.AppendLineEnd();
+								++ValueCount;
+							});
+					}
+					else
 					{
-						StringBuilder.Reset();
-						StringBuilder.Appendf(UTF8TEXT("%.9f"), Time);
-						StringBuilder.AppendChar(Separator);
-						StringBuilder.Appendf(UTF8TEXT("%.9f"), Value);
-						StringBuilder.AppendChar(LineEnd);
-						ExportFileHandle->Write((const uint8*)StringBuilder.ToString(), StringBuilder.Len() * sizeof(UTF8CHAR));
-						++ValueCount;
-					});
+						Counter.EnumerateOps(Params.IntervalStartTime, Params.IntervalEndTime, false,
+							[&](double Time, TraceServices::ECounterOpType Op, int64 IntValue)
+							{
+								StringBuilder.Appendf(UTF8TEXT("%.9f"), Time);
+								Writer.AppendSeparator();
+								switch (Op)
+								{
+									case TraceServices::ECounterOpType::Set:
+										StringBuilder.Append(UTF8TEXT("Set"));
+										break;
+									case TraceServices::ECounterOpType::Add:
+										StringBuilder.Append(UTF8TEXT("Add"));
+										break;
+									default:
+										StringBuilder.Appendf(UTF8TEXT("%d"), int32(Op));
+								}
+								Writer.AppendSeparator();
+								StringBuilder.Appendf(UTF8TEXT("%lli"), IntValue);
+								Writer.AppendLineEnd();
+								++ValueCount;
+							});
+					}
 				}
 				else
 				{
-					Counter.EnumerateValues(Params.IntervalStartTime, Params.IntervalEndTime, false, [&](double Time, int64 IntValue)
+					if (Counter.IsFloatingPoint())
 					{
-						StringBuilder.Reset();
-						StringBuilder.Appendf(UTF8TEXT("%.9f"), Time);
-						StringBuilder.AppendChar(Separator);
-						StringBuilder.Appendf(UTF8TEXT("%lli"), IntValue);
-						StringBuilder.AppendChar(LineEnd);
-						ExportFileHandle->Write((const uint8*)StringBuilder.ToString(), StringBuilder.Len() * sizeof(UTF8CHAR));
-						++ValueCount;
-					});
+						Counter.EnumerateFloatValues(Params.IntervalStartTime, Params.IntervalEndTime, false,
+							[&](double Time, double Value)
+							{
+								StringBuilder.Appendf(UTF8TEXT("%.9f"), Time);
+								Writer.AppendSeparator();
+								StringBuilder.Appendf(UTF8TEXT("%.9f"), Value);
+								Writer.AppendLineEnd();
+								++ValueCount;
+							});
+					}
+					else
+					{
+						Counter.EnumerateValues(Params.IntervalStartTime, Params.IntervalEndTime, false,
+							[&](double Time, int64 IntValue)
+							{
+								StringBuilder.Appendf(UTF8TEXT("%.9f"), Time);
+								Writer.AppendSeparator();
+								StringBuilder.Appendf(UTF8TEXT("%lli"), IntValue);
+								Writer.AppendLineEnd();
+								++ValueCount;
+							});
+					}
 				}
-			}
-		});
+			});
 	}
 
+	Writer.Flush();
 	ExportFileHandle->Flush();
 	delete ExportFileHandle;
 	ExportFileHandle = nullptr;
diff --git a/Engine/Source/Developer/TraceInsights/Private/Insights/ViewModels/TimingExporter.h b/Engine/Source/Developer/TraceInsights/Private/Insights/ViewModels/TimingExporter.h
--- a/Engine/Source/Developer/TraceInsights/Private/Insights/ViewModels/TimingExporter.h
+++ b/Engine/Source/Developer/TraceInsights/Private/Insights/ViewModels/TimingExporter.h
@@ -112,17 +112,61 @@ public:
 	};
 
 private:
-	typedef TUtf8StringBuilder<1024> FUtf8StringBuilder;
+	static constexpr int32 StringBuilderBufferSize = 32 * 1024;
+	typedef TUtf8StringBuilder<StringBuilderBufferSize> FUtf8StringBuilder;
+
+	class FUtf8Writer
+	{
+	public:
+		FUtf8Writer(IFileHandle* InFileHandle, bool bIsCSV)
+			: FileHandle(InFileHandle)
+			, Separator(bIsCSV ? UTF8CHAR(',') : UTF8CHAR('\t'))
+		{
+			check(FileHandle);
+		}
+
+		~FUtf8Writer()
+		{
+		}
+
+		FUtf8StringBuilder& GetStringBuilder() { return StringBuilder; }
+		UTF8CHAR GetSeparator() const { return Separator; }
+		UTF8CHAR GetLineEnd() const { return UTF8CHAR('\n'); }
+
+		void AppendSeparator()
+		{
+			StringBuilder.AppendChar(Separator);
+		}
+
+		void AppendLineEnd()
+		{
+			StringBuilder.AppendChar(UTF8CHAR('\n'));
+			WriteStringBuilder(StringBuilderBufferSize - 1024);
+		}
+
+		void AppendString(const TCHAR* InString);
+		void Append(const FString& InString) { StringBuilder.Append(InString); }
+
+		void Flush()
+		{
+			WriteStringBuilder(0);
+		}
+
+	private:
+		void WriteStringBuilder(int32 CacheLen);
+
+	private:
+		IFileHandle* FileHandle;
+		FUtf8StringBuilder StringBuilder;
+		UTF8CHAR Separator;
+	};
 
 	struct FExportTimingEventsInternalParams
 	{
 		const FTimingExporter& Exporter;
 		const FExportTimingEventsParams& UserParams;
 		const TArray<FName>& Columns;
-		IFileHandle* ExportFileHandle;
-		const UTF8CHAR Separator;
-		const UTF8CHAR LineEnd;
-		TUtf8StringBuilder<1024>& StringBuilder;
+		FUtf8Writer& Writer;
 		uint32 ThreadId;
 		const TCHAR* ThreadName;
 	};
@@ -187,7 +231,7 @@ public:
 	 * Note: The set is referenced in the returned function.
 	 * @param IncludedTimers The set of timer ids to be accepted by filter.
 	 * @return The TimingEventFilter function.
-	 */ 
+	 */
 	static FTimingEventFilterFunc MakeTimingEventFilterByTimersInclusive(const TSet<uint32>& IncludedTimers);
 
 	/**
@@ -195,13 +239,12 @@ public:
 	 * Note: The set is referenced in the returned function.
 	 * @param ExcludedTimers The set of timer ids to be rejected by filter.
 	 * @return The TimingEventFilter function. Can be nullptr (i.e. no filter).
-	 */ 
+	 */
 	static FTimingEventFilterFunc MakeTimingEventFilterByTimersExclusive(const TSet<uint32>& ExcludedTimers);
 
 private:
 	IFileHandle* OpenExportFile(const TCHAR* InFilename) const;
 	void Error(const FText& InMessage) const;
-	static void AppendString(FUtf8StringBuilder& StringBuilder, const TCHAR* String, UTF8CHAR Separator);
 	void ExportTimingEvents_InitColumns() const;
 	void ExportTimingEvents_WriteHeader(FExportTimingEventsInternalParams& Params) const;
 	int32 ExportTimingEvents_WriteEvents(FExportTimingEventsInternalParams& Params) const;
