[Insights] Fixed utrace disappearing after renaming it with same name as a previous trace.

diff --git a/Engine/Source/Developer/TraceInsights/Private/Insights/Widgets/SStartPageWindow.cpp b/Engine/Source/Developer/TraceInsights/Private/Insights/Widgets/SStartPageWindow.cpp
--- a/Engine/Source/Developer/TraceInsights/Private/Insights/Widgets/SStartPageWindow.cpp
+++ b/Engine/Source/Developer/TraceInsights/Private/Insights/Widgets/SStartPageWindow.cpp
@@ -330,6 +330,8 @@ public:
 		UE_LOG(TraceInsights, Verbose, TEXT("[TraceStore] Renamed utrace file (\"%s\")."), *NewTraceFile);
 		Trace.Name = FText::FromString(NewTraceName);
 		Trace.Uri = FText::FromString(NewTraceFile);
+
+		TraceStoreWindow->TraceViewModelMap.Remove(Trace.TraceId);
 		Trace.TraceId = FTraceViewModel::InvalidTraceId; // cannot be open until its TraceId is updated
 		Trace.ChangeSerial = 0; // to force update
 
@@ -2558,8 +2560,8 @@ void STraceStoreWindow::RefreshTraceList()
 					{
 						// This trace was removed.
 						RemovedTraces++;
-						TraceViewModelMap.Remove(Trace.TraceId);
 						TraceViewModels.RemoveAtSwap(TraceIndex);
+						TraceViewModelMap.Remove(Trace.TraceId);
 						TraceIndex--;
 						TraceViewModelCount--;
 					}
@@ -2586,9 +2588,10 @@ void STraceStoreWindow::RefreshTraceList()
 					// This trace was added.
 					AddedTraces++;
 					TSharedPtr<FTraceViewModel> TracePtr = MakeShared<FTraceViewModel>();
+					TracePtr->TraceId = SourceTrace.TraceId;
 					UpdateTrace(*TracePtr, SourceTrace);
 					TraceViewModels.Add(TracePtr);
-					TraceViewModelMap.Add(SourceTrace.TraceId, TracePtr);
+					TraceViewModelMap.Add(TracePtr->TraceId, TracePtr);
 				}
 			}
 		}
@@ -2629,7 +2632,7 @@ bool STraceStoreWindow::IsConnected() const
 
 void STraceStoreWindow::UpdateTrace(FTraceViewModel& InOutTrace, const Insights::FStoreBrowserTraceInfo& InSourceTrace)
 {
-	InOutTrace.TraceId = InSourceTrace.TraceId;
+	check(InOutTrace.TraceId == InSourceTrace.TraceId);
 
 	InOutTrace.ChangeSerial = InSourceTrace.ChangeSerial;
 
