[Insights]
  - Added support for exporting counters by timing regions (ExportCounterAsText API).
  - Added support for exporting timing events by timing region (ExportTimingEventsAsText API; TimingInsights.ExportTimingEvents filenamePattern -region= ...)
  - Added new export commands for Counters:
    * TimingInsights.ExportCounters filename
    * TimingInsights.ExportCounterValues filenamePattern -counter= -startTime= -endTime= -region=
  - Added support for {region} and {counter} in execute commands that supports regions and counters.

diff --git a/Engine/Source/Developer/TraceInsights/Private/Insights/TimingProfilerManager.cpp b/Engine/Source/Developer/TraceInsights/Private/Insights/TimingProfilerManager.cpp
--- a/Engine/Source/Developer/TraceInsights/Private/Insights/TimingProfilerManager.cpp
+++ b/Engine/Source/Developer/TraceInsights/Private/Insights/TimingProfilerManager.cpp
@@ -4,6 +4,7 @@
 
 #include "MessageLogModule.h"
 #include "Modules/ModuleManager.h"
+#include "TraceServices/Model/Counters.h"
 #include "Widgets/Docking/SDockTab.h"
 #include "WorkspaceMenuStructure.h"
 #include "WorkspaceMenuStructureModule.h"
@@ -664,7 +665,8 @@ bool FTimingProfilerManager::Exec(const TCHAR* Cmd, FOutputDevice& Ar)
 
 		const bool bUseEscape = true;
 		FString Filename = FParse::Token(Cmd, bUseEscape);
-		Ar.Logf(TEXT("  Filename: %s"), *Filename);
+		Filename.TrimQuotesInline();
+		Ar.Logf(TEXT("  Filename: \"%s\""), *Filename);
 
 		Exporter.ExportThreadsAsText(Filename, Params);
 		return true;
@@ -679,7 +681,8 @@ bool FTimingProfilerManager::Exec(const TCHAR* Cmd, FOutputDevice& Ar)
 
 		const bool bUseEscape = true;
 		FString Filename = FParse::Token(Cmd, bUseEscape);
-		Ar.Logf(TEXT("  Filename: %s"), *Filename);
+		Filename.TrimQuotesInline();
+		Ar.Logf(TEXT("  Filename: \"%s\""), *Filename);
 
 		Exporter.ExportTimersAsText(Filename, Params);
 		return true;
@@ -702,7 +705,9 @@ bool FTimingProfilerManager::Exec(const TCHAR* Cmd, FOutputDevice& Ar)
 
 		const bool bUseEscape = true;
 		FString Filename = FParse::Token(Cmd, bUseEscape);
-		Ar.Logf(TEXT("  Filename: %s"), *Filename);
+		Filename.TrimQuotesInline();
+		Ar.Logf(TEXT("  Filename: \"%s\""), *Filename);
+		// '{region}' in Filename (if any) will be replaced with the resolved name of the region.
 
 		while (Cmd && Cmd[0] != TEXT('\0'))
 		{
@@ -716,6 +721,7 @@ bool FTimingProfilerManager::Exec(const TCHAR* Cmd, FOutputDevice& Ar)
 				static constexpr TCHAR TimersToken[] = TEXT("-timers=");
 				static constexpr TCHAR StartTimeToken[] = TEXT("-startTime=");
 				static constexpr TCHAR EndTimeToken[] = TEXT("-endTime=");
+				static constexpr TCHAR RegionToken[] = TEXT("-region=");
 
 				if (Token.StartsWith(ColumnsToken))
 				{
@@ -759,6 +765,17 @@ bool FTimingProfilerManager::Exec(const TCHAR* Cmd, FOutputDevice& Ar)
 					Token.RightChopInline(UE_ARRAY_COUNT(EndTimeToken) - 1);
 					Params.IntervalEndTime = atof(TCHAR_TO_ANSI(*Token));
 				}
+				else if (Token.StartsWith(RegionToken))
+				{
+					// Comma-delimited list of region names. Supports *?-type wildcard.
+					// Each region is exported to a separate file.
+					// '{region}' in Filename (if any) will be replaced with the resolved name of the region.
+					// Default: -Region=
+					// Example: -Region="RegionName,Game_*,OtherRegion"
+					Token.RightChopInline(UE_ARRAY_COUNT(RegionToken) - 1);
+					Token.TrimQuotesInline();
+					Params.Region = TCHAR_TO_ANSI(*Token);
+				}
 				else
 				{
 					Ar.Logf(ELogVerbosity::Warning, TEXT("Unknown Cmd Param: %s"), *Token);
@@ -789,7 +806,9 @@ bool FTimingProfilerManager::Exec(const TCHAR* Cmd, FOutputDevice& Ar)
 
 		const bool bUseEscape = true;
 		FString Filename = FParse::Token(Cmd, bUseEscape);
-		Ar.Logf(TEXT("  Filename: %s"), *Filename);
+		Filename.TrimQuotesInline();
+		Ar.Logf(TEXT("  Filename: \"%s\""), *Filename);
+		// '{region}' in Filename (if any) will be replaced with the resolved name of the region.
 
 		while (Cmd && Cmd[0] != TEXT('\0'))
 		{
@@ -801,9 +820,9 @@ bool FTimingProfilerManager::Exec(const TCHAR* Cmd, FOutputDevice& Ar)
 				static constexpr TCHAR ColumnsToken[] = TEXT("-columns=");
 				static constexpr TCHAR ThreadsToken[] = TEXT("-threads=");
 				static constexpr TCHAR TimersToken[] = TEXT("-timers=");
-				static constexpr TCHAR RegionToken[] = TEXT("-region=");
 				static constexpr TCHAR StartTimeToken[] = TEXT("-startTime=");
 				static constexpr TCHAR EndTimeToken[] = TEXT("-endTime=");
+				static constexpr TCHAR RegionToken[] = TEXT("-region=");
 
 				if (Token.StartsWith(ColumnsToken))
 				{
@@ -850,8 +869,8 @@ bool FTimingProfilerManager::Exec(const TCHAR* Cmd, FOutputDevice& Ar)
 				else if (Token.StartsWith(RegionToken))
 				{
 					// Comma-delimited list of region names. Supports *?-type wildcard.
-					// Each region is exported to a separate file. The '*' char in Filename (if any)
-					// will be replaced with the resolved name of the region.
+					// Each region is exported to a separate file.
+					// '{region}' in Filename (if any) will be replaced with the resolved name of the region.
 					// Default: -Region=
 					// Example: -Region="RegionName,Game_*,OtherRegion"
 					Token.RightChopInline(UE_ARRAY_COUNT(RegionToken) - 1);
@@ -872,8 +891,141 @@ bool FTimingProfilerManager::Exec(const TCHAR* Cmd, FOutputDevice& Ar)
 			Params.ThreadFilter = [](unsigned int){ return true; };
 		}
 
-		int32 Result = Exporter.ExportTimerStatisticsAsText(Filename, Params);
-		return Result > 0;
+		Exporter.ExportTimerStatisticsAsText(Filename, Params);
+		return true;
+	}
+
+	if (FParse::Command(&Cmd, TEXT("TimingInsights.ExportCounters")))
+	{
+		Ar.Logf(TEXT("TimingInsights.ExportCounters %s"), Cmd);
+		check(FInsightsManager::Get().IsValid() && FInsightsManager::Get()->GetSession().IsValid());
+		Insights::FTimingExporter Exporter(*FInsightsManager::Get()->GetSession().Get());
+		Insights::FTimingExporter::FExportCountersParams Params; // default
+		Exporter.ExportCountersAsText(Cmd, Params);
+		return true;
+	}
+
+	if (FParse::Command(&Cmd, TEXT("TimingInsights.ExportCounterValues")))
+	{
+		Ar.Logf(TEXT("TimingInsights.ExportCounterValues %s"), Cmd);
+		check(FInsightsManager::Get().IsValid() && FInsightsManager::Get()->GetSession().IsValid());
+		Insights::FTimingExporter Exporter(*FInsightsManager::Get()->GetSession().Get());
+		Insights::FTimingExporter::FExportCounterParams Params; // default
+
+		TArray<FString> Counters; // list of counters to export
+
+		// These variables needs to be in the same scope with the call to Exporter.ExportCounterAsText().
+		TArray<FName> Columns; // referenced by Params.Columns
+
+		//////////////////////////////////////////////////
+
+		const bool bUseEscape = true;
+		FString Filename = FParse::Token(Cmd, bUseEscape);
+		Filename.TrimQuotesInline();
+		Ar.Logf(TEXT("  Filename: \"%s\""), *Filename);
+		// '{counter}' in Filename (if any) will be replaced with the name of the counter.
+		// '{region}' in Filename (if any) will be replaced with the resolved name of the region.
+
+		while (Cmd && Cmd[0] != TEXT('\0'))
+		{
+			FString Token;
+			if (FParse::Token(Cmd, Token, bUseEscape))
+			{
+				static constexpr TCHAR CounterToken[] = TEXT("-counter=");
+				static constexpr TCHAR ColumnsToken[] = TEXT("-columns=");
+				static constexpr TCHAR StartTimeToken[] = TEXT("-startTime=");
+				static constexpr TCHAR EndTimeToken[] = TEXT("-endTime=");
+				static constexpr TCHAR RegionToken[] = TEXT("-region=");
+
+				if (Token.StartsWith(CounterToken))
+				{
+					// Comma-delimited list of counter names. Supports *?-type wildcard.
+					// Default: -counter=
+					// Example: -counter="PC / *"
+					Token.RightChopInline(UE_ARRAY_COUNT(CounterToken) - 1);
+					Token.TrimQuotesInline();
+					Token.ParseIntoArray(Counters, TEXT(","), true);
+				}
+				else if (Token.StartsWith(ColumnsToken))
+				{
+					// Comma-delimited list of column names. Supports *?-type wildcard.
+					// Default: -columns="Time,Value"
+					// Example: -columns="*" -columns="Value"
+					Token.RightChopInline(UE_ARRAY_COUNT(ColumnsToken) - 1);
+					Token.TrimQuotesInline();
+					Exporter.MakeExportTimingEventsColumnList(Token, Columns);
+					Params.Columns = &Columns;
+				}
+				else if (Token.StartsWith(StartTimeToken))
+				{
+					// Default: -startTime=-infinite
+					// Example: -startTime=10.0
+					Token.RightChopInline(UE_ARRAY_COUNT(StartTimeToken) - 1);
+					Params.IntervalStartTime = atof(TCHAR_TO_ANSI(*Token));
+				}
+				else if (Token.StartsWith(EndTimeToken))
+				{
+					// Default: -endTime=+infinite
+					// Example: -endTime=20.0
+					Token.RightChopInline(UE_ARRAY_COUNT(EndTimeToken) - 1);
+					Params.IntervalEndTime = atof(TCHAR_TO_ANSI(*Token));
+				}
+				else if (Token.StartsWith(RegionToken))
+				{
+					// Comma-delimited list of region names. Supports *?-type wildcard.
+					// Each region is exported to a separate file.
+					// '{region}' in Filename (if any) will be replaced with the resolved name of the region.
+					// Default: -Region=
+					// Example: -Region="RegionName,Game_*,OtherRegion"
+					Token.RightChopInline(UE_ARRAY_COUNT(RegionToken) - 1);
+					Token.TrimQuotesInline();
+					Params.Region = TCHAR_TO_ANSI(*Token);
+				}
+				else
+				{
+					Ar.Logf(ELogVerbosity::Warning, TEXT("Unknown Cmd Param: %s"), *Token);
+				}
+			}
+		}
+
+		struct FExportCounterInfo
+		{
+			uint32 Id;
+			FString Name;
+		};
+		TArray<FExportCounterInfo> CountersToExport;
+
+		if (Counters.Num() > 0)
+		{
+			for (const FString& CounterWildcard : Counters)
+			{
+				Ar.Logf(TEXT("  Searching counters with name: \"%s\""), *CounterWildcard);
+			}
+			const TraceServices::IAnalysisSession& Session = *FInsightsManager::Get()->GetSession().Get();
+			TraceServices::FAnalysisSessionReadScope SessionReadScope(Session);
+			const TraceServices::ICounterProvider& CounterProvider = TraceServices::ReadCounterProvider(Session);
+			CounterProvider.EnumerateCounters([&Ar, &Counters, &CountersToExport](uint32 CounterId, const TraceServices::ICounter& Counter)
+			{
+				FString CounterName(Counter.GetName());
+				for (const FString& CounterWildcard : Counters)
+				{
+					if (CounterName.MatchesWildcard(CounterWildcard))
+					{
+						CountersToExport.Add({ CounterId, CounterName });
+						break;
+					}
+				}
+			});
+		}
+
+		Ar.Logf(TEXT("  Exporting values for %d counters..."), CountersToExport.Num());
+		for (const FExportCounterInfo& Counter : CountersToExport)
+		{
+			Ar.Logf(TEXT("  Exporting counter: \"%s\" (id=%d)"), *Counter.Name, Counter.Id);
+			Exporter.ExportCounterAsText(Filename, Counter.Id, Params);
+		}
+
+		return true;
 	}
 
 	return false;
diff --git a/Engine/Source/Developer/TraceInsights/Private/Insights/ViewModels/TimingExporter.cpp b/Engine/Source/Developer/TraceInsights/Private/Insights/ViewModels/TimingExporter.cpp
--- a/Engine/Source/Developer/TraceInsights/Private/Insights/ViewModels/TimingExporter.cpp
+++ b/Engine/Source/Developer/TraceInsights/Private/Insights/ViewModels/TimingExporter.cpp
@@ -11,6 +11,7 @@
 #include "TraceServices/Model/Counters.h"
 #include "TraceServices/Model/Regions.h"
 #include "TraceServices/Model/Threads.h"
+#include "TraceServices/Model/TimingProfiler.h"
 
 // TraceInsights
 #include "Insights/Common/Stopwatch.h"
@@ -115,7 +116,11 @@ void FTimingExporter::FUtf8Writer::WriteStringBuilder(int32 CacheLen)
 
 int32 FTimingExporter::ExportThreadsAsText(const FString& Filename, FExportThreadsParams& Params) const
 {
-	checkf(Params.Columns == nullptr, TEXT("Custom list of columns is not yet supported!"));
+	if (Params.Columns != nullptr)
+	{
+		UE_LOG(TraceInsights, Error, TEXT("Custom list of columns is not yet supported!"));
+		return -1;
+	}
 
 	FStopwatch Stopwatch;
 	Stopwatch.Start();
@@ -195,7 +200,11 @@ int32 FTimingExporter::ExportThreadsAsText(const FString& Filename, FExportThrea
 
 int32 FTimingExporter::ExportTimersAsText(const FString& Filename, FExportTimersParams& Params) const
 {
-	checkf(Params.Columns == nullptr, TEXT("Custom list of columns is not yet supported!"));
+	if (Params.Columns != nullptr)
+	{
+		UE_LOG(TraceInsights, Error, TEXT("Custom list of columns is not yet supported!"));
+		return -1;
+	}
 
 	FStopwatch Stopwatch;
 	Stopwatch.Start();
@@ -496,8 +505,47 @@ int32 FTimingExporter::ExportTimingEvents_WriteEvents(FExportTimingEventsInterna
 
 ////////////////////////////////////////////////////////////////////////////////////////////////////
 
+int32 FTimingExporter::ExportTimingEventsAsTextByRegions(const FString& FilenamePattern, FExportTimingEventsParams& Params) const
+{
+	TMap<FString, FTimeRegionGroup> RegionGroups;
+	GetRegions(Params.Region, RegionGroups);
+
+	if (RegionGroups.Num() == 0)
+	{
+		UE_LOG(TraceInsights, Error, TEXT("Unable to find any region with name pattern '%s'."), *Params.Region);
+		return -1;
+	}
+
+	FStopwatch Stopwatch;
+	Stopwatch.Start();
+
+	// Export timing statistics for each region.
+	FExportTimingEventsParams RegionParams = Params;
+	RegionParams.Region.Reset();
+	int32 ExportedRegionCount = EnumerateRegions(RegionGroups, FilenamePattern,
+		[this, &RegionParams](const FString& Filename, const FString& RegionName, double IntervalStartTime, double IntervalEndTime)
+		{
+			RegionParams.IntervalStartTime = IntervalStartTime;
+			RegionParams.IntervalEndTime = IntervalEndTime;
+			UE_LOG(TraceInsights, Display, TEXT("Exporting timing statistics for region '%s' [%f .. %f] to '%s'"), *RegionName, RegionParams.IntervalStartTime, RegionParams.IntervalEndTime, *Filename);
+			ExportTimingEventsAsText(Filename, RegionParams);
+		});
+
+	Stopwatch.Stop();
+	const double TotalTime = Stopwatch.GetAccumulatedTime();
+	UE_LOG(TraceInsights, Log, TEXT("Exported timing statistics for %d regions in %.3fs."), ExportedRegionCount, TotalTime);
+	return ExportedRegionCount;
+}
+
+////////////////////////////////////////////////////////////////////////////////////////////////////
+
 int32 FTimingExporter::ExportTimingEventsAsText(const FString& Filename, FExportTimingEventsParams& Params) const
 {
+	if (!Params.Region.IsEmpty())
+	{
+		return ExportTimingEventsAsTextByRegions(Filename, Params);
+	}
+
 	ExportTimingEvents_InitColumns();
 	const TArray<FName>& Columns = Params.Columns ? *Params.Columns : ExportTimingEventsDefaultColumns;
 
@@ -535,12 +583,12 @@ int32 FTimingExporter::ExportTimingEventsAsText(const FString& Filename, FExport
 
 ////////////////////////////////////////////////////////////////////////////////////////////////////
 
-int32 FTimingExporter::ExportTimerStatisticsAsTextByRegions(const FString& Filename, FExportTimerStatisticsParams& Params) const
+void FTimingExporter::GetRegions(const FString& InRegionNamePattern, TMap<FString, FTimeRegionGroup>& OutRegionGroups) const
 {
 	class FRegionNameSpec
 	{
 	public:
-		FRegionNameSpec(FString& InNamePatternList)
+		FRegionNameSpec(const FString& InNamePatternList)
 		{
 			InNamePatternList.ParseIntoArray(NamePatterns, TEXT(","), true);
 		}
@@ -560,18 +608,7 @@ int32 FTimingExporter::ExportTimerStatisticsAsTextByRegions(const FString& Filen
 	private:
 		TArray<FString> NamePatterns;
 	};
-	FRegionNameSpec RegionNameSpec(Params.Region);
-
-	struct FTimeRegionInterval
-	{
-		double StartTime;
-		double EndTime;
-	};
-	struct FTimeRegionGroup
-	{
-		TArray<FTimeRegionInterval> Intervals;
-	};
-	TMap<FString, FTimeRegionGroup> RegionGroups;
+	FRegionNameSpec RegionNameSpec(InRegionNamePattern);
 
 	// Detect regions
 	int32 RegionCount = 0;
@@ -581,10 +618,10 @@ int32 FTimingExporter::ExportTimerStatisticsAsTextByRegions(const FString& Filen
 		const TraceServices::IRegionProvider& RegionProvider = TraceServices::ReadRegionProvider(Session);
 		TraceServices::FProviderReadScopeLock RegionProviderScopedLock(RegionProvider);
 
-		UE_LOG(TraceInsights, Log, TEXT("Looking for regions: '%s'"), *Params.Region);
+		UE_LOG(TraceInsights, Log, TEXT("Looking for regions: '%s'"), *InRegionNamePattern);
 
 		RegionProvider.EnumerateRegions(0.0, std::numeric_limits<double>::max(),
-			[&RegionCount, &RegionNameSpec, &RegionGroups](const TraceServices::FTimeRegion& InRegion) -> bool
+			[&RegionCount, &RegionNameSpec, &OutRegionGroups](const TraceServices::FTimeRegion& InRegion) -> bool
 			{
 				if (!RegionNameSpec.Match(InRegion.Text))
 				{
@@ -594,10 +631,10 @@ int32 FTimingExporter::ExportTimerStatisticsAsTextByRegions(const FString& Filen
 				// Handle duplicate region names, individual regions may appear multiple times
 				// we append numbers to allow for unique export filenames.
 				FString RegionName = InRegion.Text;
-				FTimeRegionGroup* ExistingRegionGroup = RegionGroups.Find(RegionName);
+				FTimeRegionGroup* ExistingRegionGroup = OutRegionGroups.Find(RegionName);
 				if (!ExistingRegionGroup)
 				{
-					ExistingRegionGroup = &RegionGroups.Add(RegionName, FTimeRegionGroup{});
+					ExistingRegionGroup = &OutRegionGroups.Add(RegionName, FTimeRegionGroup{});
 				}
 				ExistingRegionGroup->Intervals.Add(FTimeRegionInterval{ InRegion.BeginTime, InRegion.EndTime });
 				++RegionCount;
@@ -607,23 +644,20 @@ int32 FTimingExporter::ExportTimerStatisticsAsTextByRegions(const FString& Filen
 	DetectRegionsStopwatch.Stop();
 	const double DetectRegionsTime = DetectRegionsStopwatch.GetAccumulatedTime();
 	UE_LOG(TraceInsights, Display, TEXT("Detected %d regions in %.3fs."), RegionCount, DetectRegionsTime);
+}
 
-	if (RegionGroups.Num() == 0)
-	{
-		UE_LOG(TraceInsights, Error, TEXT("Unable to find any region with name pattern '%s'."), *Params.Region);
-		return -1;
-	}
-
-	FStopwatch Stopwatch;
-	Stopwatch.Start();
+////////////////////////////////////////////////////////////////////////////////////////////////////
 
-	// Export timing statistics for each region.
+int32 FTimingExporter::EnumerateRegions(const TMap<FString, FTimeRegionGroup>& InRegionGroups, const FString& InFilenamePattern,
+	TFunction<void(const FString& /*Filename*/, const FString& /*RegionName*/, double /*IntervalStartTime*/, double /*IntervalEndTime*/)> InCallback) const
+{
 	constexpr int32 MaxIntervalsPerRegion = 100;
 	constexpr int32 MaxExportedRegions = 10000;
 	int32 ExportedRegionCount = 0;
-	for (auto& KV : RegionGroups)
+
+	for (auto& KV : InRegionGroups)
 	{
-		FString RegionName = KV.Key;
+		FString RegionName(KV.Key);
 		const FString InvalidFileSystemChars = FPaths::GetInvalidFileSystemChars();
 		for (int32 CharIndex = 0; CharIndex < InvalidFileSystemChars.Len(); CharIndex++)
 		{
@@ -635,36 +669,33 @@ int32 FTimingExporter::ExportTimerStatisticsAsTextByRegions(const FString& Filen
 		int32 IntervalIndex = 0;
 		for (const FTimeRegionInterval& Interval : KV.Value.Intervals)
 		{
-			FString RegionFilename;
+			FString Filename(InFilenamePattern);
 			if (IntervalIndex == 0)
 			{
-				RegionFilename = Filename.Replace(TEXT("*"), *RegionName);
+				Filename.ReplaceInline(TEXT("*"), *RegionName); // for backward compatibility
+				Filename.ReplaceInline(TEXT("{region}"), *RegionName);
 			}
 			else
 			{
 				FString UniqueRegionName = FString::Printf(TEXT("%s_%d"), *RegionName, IntervalIndex);
-				RegionFilename = Filename.Replace(TEXT("*"), *UniqueRegionName);
+				Filename.ReplaceInline(TEXT("*"), *UniqueRegionName); // for backward compatibility
+				Filename.ReplaceInline(TEXT("{region}"), *UniqueRegionName);
 			}
 			++IntervalIndex;
 
-			FExportTimerStatisticsParams RegionParams = Params;
-			RegionParams.Region.Reset();
-			RegionParams.IntervalStartTime = Interval.StartTime;
-			RegionParams.IntervalEndTime = Interval.EndTime;
-			UE_LOG(TraceInsights, Display, TEXT("Exporting timing statistics for region '%s' [%f .. %f] to '%s'"), *KV.Key, RegionParams.IntervalStartTime, RegionParams.IntervalEndTime, *RegionFilename);
-			ExportTimerStatisticsAsText(RegionFilename, RegionParams);
+			InCallback(Filename, KV.Key, Interval.StartTime, Interval.EndTime);
 
 			++ExportedRegionCount;
 
 			// Avoid writing too many files...
 			if (IntervalIndex >= MaxIntervalsPerRegion)
 			{
-				UE_LOG(TraceInsights, Error, TEXT("Too many intervals for region '%s'! Exporting timing statistics to separate file per interval for this region is not allowed to continue."), *KV.Key);
+				UE_LOG(TraceInsights, Error, TEXT("Too many intervals for region '%s'! Exporting to separate file per interval for this region is not allowed to continue."), *KV.Key);
 				break;
 			}
 			if (ExportedRegionCount >= MaxExportedRegions)
 			{
-				UE_LOG(TraceInsights, Error, TEXT("Too many regions! Exporting timing statistics to separate file per region is not allowed to continue."));
+				UE_LOG(TraceInsights, Error, TEXT("Too many regions! Exporting to separate file per region is not allowed to continue."));
 				break;
 			}
 		}
@@ -675,6 +706,37 @@ int32 FTimingExporter::ExportTimerStatisticsAsTextByRegions(const FString& Filen
 		}
 	}
 
+	return ExportedRegionCount;
+}
+
+////////////////////////////////////////////////////////////////////////////////////////////////////
+
+int32 FTimingExporter::ExportTimerStatisticsAsTextByRegions(const FString& FilenamePattern, FExportTimerStatisticsParams& Params) const
+{
+	TMap<FString, FTimeRegionGroup> RegionGroups;
+	GetRegions(Params.Region, RegionGroups);
+
+	if (RegionGroups.Num() == 0)
+	{
+		UE_LOG(TraceInsights, Error, TEXT("Unable to find any region with name pattern '%s'."), *Params.Region);
+		return -1;
+	}
+
+	FStopwatch Stopwatch;
+	Stopwatch.Start();
+
+	// Export timing statistics for each region.
+	FExportTimerStatisticsParams RegionParams = Params;
+	RegionParams.Region.Reset();
+	int32 ExportedRegionCount = EnumerateRegions(RegionGroups, FilenamePattern,
+		[this, &RegionParams](const FString& Filename, const FString& RegionName, double IntervalStartTime, double IntervalEndTime)
+		{
+			RegionParams.IntervalStartTime = IntervalStartTime;
+			RegionParams.IntervalEndTime = IntervalEndTime;
+			UE_LOG(TraceInsights, Display, TEXT("Exporting timing statistics for region '%s' [%f .. %f] to '%s'"), *RegionName, RegionParams.IntervalStartTime, RegionParams.IntervalEndTime, *Filename);
+			ExportTimerStatisticsAsText(Filename, RegionParams);
+		});
+
 	Stopwatch.Stop();
 	const double TotalTime = Stopwatch.GetAccumulatedTime();
 	UE_LOG(TraceInsights, Log, TEXT("Exported timing statistics for %d regions in %.3fs."), ExportedRegionCount, TotalTime);
@@ -948,7 +1010,11 @@ FTimingExporter::FTimingEventFilterFunc FTimingExporter::MakeTimingEventFilterBy
 
 int32 FTimingExporter::ExportCountersAsText(const FString& Filename, FExportCountersParams& Params) const
 {
-	checkf(Params.Columns == nullptr, TEXT("Custom list of columns is not yet supported!"));
+	if (Params.Columns != nullptr)
+	{
+		UE_LOG(TraceInsights, Error, TEXT("Custom list of columns is not yet supported!"));
+		return -1;
+	}
 
 	FStopwatch Stopwatch;
 	Stopwatch.Start();
@@ -1019,13 +1085,89 @@ int32 FTimingExporter::ExportCountersAsText(const FString& Filename, FExportCoun
 
 ////////////////////////////////////////////////////////////////////////////////////////////////////
 
-int32 FTimingExporter::ExportCounterAsText(const FString& Filename, uint32 CounterId, FExportCounterParams& Params) const
+int32 FTimingExporter::ExportCounterAsTextByRegions(const FString& FilenamePattern, uint32 CounterId, FExportCounterParams& Params) const
 {
-	checkf(Params.Columns == nullptr, TEXT("Custom list of columns is not yet supported!"));
+	TMap<FString, FTimeRegionGroup> RegionGroups;
+	GetRegions(Params.Region, RegionGroups);
+
+	if (RegionGroups.Num() == 0)
+	{
+		UE_LOG(TraceInsights, Error, TEXT("Unable to find any region with name pattern '%s'."), *Params.Region);
+		return -1;
+	}
 
 	FStopwatch Stopwatch;
 	Stopwatch.Start();
 
+	// Export counter for each region.
+	FExportCounterParams RegionParams = Params;
+	RegionParams.Region.Reset();
+	int32 ExportedRegionCount = EnumerateRegions(RegionGroups, FilenamePattern,
+		[this, CounterId, &RegionParams](const FString& Filename, const FString& RegionName, double IntervalStartTime, double IntervalEndTime)
+		{
+			RegionParams.IntervalStartTime = IntervalStartTime;
+			RegionParams.IntervalEndTime = IntervalEndTime;
+			UE_LOG(TraceInsights, Display, TEXT("Exporting counter %u for region '%s' [%f .. %f] to '%s'"), CounterId , *RegionName, RegionParams.IntervalStartTime, RegionParams.IntervalEndTime, *Filename);
+			ExportCounterAsText(Filename, CounterId, RegionParams);
+		});
+
+	Stopwatch.Stop();
+	const double TotalTime = Stopwatch.GetAccumulatedTime();
+	UE_LOG(TraceInsights, Log, TEXT("Exported counter values for %d regions in %.3fs."), ExportedRegionCount, TotalTime);
+	return ExportedRegionCount;
+}
+
+////////////////////////////////////////////////////////////////////////////////////////////////////
+
+int32 FTimingExporter::ExportCounterAsText(const FString& FilenamePattern, uint32 CounterId, FExportCounterParams& Params) const
+{
+	if (Params.Columns != nullptr)
+	{
+		UE_LOG(TraceInsights, Error, TEXT("Custom list of columns is not yet supported!"));
+		return -1;
+	}
+
+	if (!Params.Region.IsEmpty())
+	{
+		return ExportCounterAsTextByRegions(FilenamePattern, CounterId, Params);
+	}
+
+	FStopwatch Stopwatch;
+	Stopwatch.Start();
+
+	FString CounterName;
+
+	if (true) // TraceServices::ReadCounterProvider(Session)
+	{
+		TraceServices::FAnalysisSessionReadScope SessionReadScope(Session);
+		const TraceServices::ICounterProvider& CounterProvider = TraceServices::ReadCounterProvider(Session);
+		CounterProvider.ReadCounter(CounterId,
+			[&](const TraceServices::ICounter& Counter)
+			{
+				CounterName = Counter.GetName();
+			});
+	}
+
+	if (CounterName.IsEmpty())
+	{
+		UE_LOG(TraceInsights, Error, TEXT("Invalid counter!"));
+		return -1;
+	}
+
+	FString Filename(FilenamePattern);
+	if (Filename.Contains(TEXT("{counter}")))
+	{
+		FString CounterFilename(CounterName);
+		const FString InvalidFileSystemChars = FPaths::GetInvalidFileSystemChars();
+		for (int32 CharIndex = 0; CharIndex < InvalidFileSystemChars.Len(); CharIndex++)
+		{
+			FString Char = FString().AppendChar(InvalidFileSystemChars[CharIndex]);
+			CounterFilename.ReplaceInline(*Char, TEXT("_"));
+		}
+		CounterFilename.TrimStartAndEndInline();
+		Filename.ReplaceInline(TEXT("{counter}"), *CounterFilename);
+	}
+
 	IFileHandle* ExportFileHandle = OpenExportFile(*Filename);
 	if (!ExportFileHandle)
 	{
@@ -1053,7 +1195,6 @@ int32 FTimingExporter::ExportCounterAsText(const FString& Filename, uint32 Count
 		Writer.AppendLineEnd();
 	}
 
-	FString CounterName;
 	int32 ValueCount = 0;
 
 	// Write values.
@@ -1066,8 +1207,6 @@ int32 FTimingExporter::ExportCounterAsText(const FString& Filename, uint32 Count
 		CounterProvider.ReadCounter(CounterId,
 			[&](const TraceServices::ICounter& Counter)
 			{
-				CounterName = Counter.GetName();
-
 				// Iterate the counter values.
 				if (Params.bExportOps)
 				{
diff --git a/Engine/Source/Developer/TraceInsights/Private/Insights/ViewModels/TimingExporter.h b/Engine/Source/Developer/TraceInsights/Private/Insights/ViewModels/TimingExporter.h
--- a/Engine/Source/Developer/TraceInsights/Private/Insights/ViewModels/TimingExporter.h
+++ b/Engine/Source/Developer/TraceInsights/Private/Insights/ViewModels/TimingExporter.h
@@ -2,15 +2,27 @@
 
 #pragma once
 
-#include "CoreMinimal.h"
+#include "CoreTypes.h"
+
+#include "Containers/Array.h"
+#include "Containers/Map.h"
+#include "Containers/Set.h"
+#include "HAL/Platform.h"
+#include "Templates/Function.h"
+#include "UObject/NameTypes.h"
 
-#include "TraceServices/Model/TimingProfiler.h"
 #include <limits>
 
 ////////////////////////////////////////////////////////////////////////////////////////////////////
 
 class IFileHandle;
 
+namespace TraceServices
+{
+	class IAnalysisSession;
+	struct FTimingProfilerEvent;
+}
+
 ////////////////////////////////////////////////////////////////////////////////////////////////////
 
 namespace Insights
@@ -69,18 +81,18 @@ public:
 		 */
 		double IntervalStartTime = -std::numeric_limits<double>::infinity();
 		double IntervalEndTime = +std::numeric_limits<double>::infinity();
-	};
 
-	struct FExportTimerStatisticsParams : public FExportTimingEventsParams
-	{
 		/**
-		 * The time region to be exported. This is defined by having corresponding "RegionStart:Name" and
-		 * "RegionEnd:Name" bookmarks in the bookmarks channel.
-		 * If empty falls back to IntervalStartTime and IntervalEndTime.
+		 * The time region to be exported.
+		 * If empty, falls back to IntervalStartTime and IntervalEndTime.
 		 */
 		FString Region;
 	};
 
+	struct FExportTimerStatisticsParams : public FExportTimingEventsParams
+	{
+	};
+
 	struct FExportCountersParams
 	{
 		/**
@@ -106,11 +118,28 @@ public:
 		double IntervalEndTime = +std::numeric_limits<double>::infinity();
 
 		/**
-		 * If true, will export the operations and the operation values, instead of final values.
+		 * The time region to be exported.
+		 * If empty, falls back to IntervalStartTime and IntervalEndTime.
+		 */
+		FString Region;
+
+		/**
+		 * If true, will export values with the corresponding operation type, instead of the final values.
 		 */
 		bool bExportOps = false;
 	};
 
+	struct FTimeRegionInterval
+	{
+		double StartTime;
+		double EndTime;
+	};
+
+	struct FTimeRegionGroup
+	{
+		TArray<FTimeRegionInterval> Intervals;
+	};
+
 private:
 	static constexpr int32 StringBuilderBufferSize = 32 * 1024;
 	typedef TUtf8StringBuilder<StringBuilderBufferSize> FUtf8StringBuilder;
@@ -200,6 +229,10 @@ public:
 
 	void MakeExportTimingEventsColumnList(const FString& InColumnsString, TArray<FName>& OutColumnList) const;
 
+	void GetRegions(const FString& InRegionNamePattern, TMap<FString, FTimeRegionGroup>& OutRegionGroups) const;
+	int32 EnumerateRegions(const TMap<FString, FTimeRegionGroup>& InRegionGroups, const FString& InFilenamePattern,
+		TFunction<void(const FString& /*Filename*/, const FString& /*RegionName*/, double /*IntervalStartTime*/, double /*IntervalEndTime*/)> InCallback) const;
+
 	//////////////////////////////////////////////////////////////////////
 	// Utilities to make FThreadFilterFunc filters.
 
@@ -248,7 +281,9 @@ private:
 	void ExportTimingEvents_InitColumns() const;
 	void ExportTimingEvents_WriteHeader(FExportTimingEventsInternalParams& Params) const;
 	int32 ExportTimingEvents_WriteEvents(FExportTimingEventsInternalParams& Params) const;
+	int32 ExportTimingEventsAsTextByRegions(const FString& FilenamePattern, FExportTimingEventsParams& Params) const;
 	int32 ExportTimerStatisticsAsTextByRegions(const FString& Filename, FExportTimerStatisticsParams& Params) const;
+	int32 ExportCounterAsTextByRegions(const FString& Filename, uint32 CounterId, FExportCounterParams& Params) const;
 
 private:
 	const TraceServices::IAnalysisSession& Session;
