[Insights] Frame Stats graphs now exclude nested events.

diff --git a/Engine/Source/Developer/TraceInsights/Private/Insights/Tests/FunctionalTests/TimingProfilerTests.cpp b/Engine/Source/Developer/TraceInsights/Private/Insights/Tests/FunctionalTests/TimingProfilerTests.cpp
--- a/Engine/Source/Developer/TraceInsights/Private/Insights/Tests/FunctionalTests/TimingProfilerTests.cpp
+++ b/Engine/Source/Developer/TraceInsights/Private/Insights/Tests/FunctionalTests/TimingProfilerTests.cpp
@@ -430,7 +430,7 @@ void FTimingProfilerTests::RunEnumerateAsyncBenchmark(const FEnumerateTestParams
 					{
 						TaskCheckValues.AddDefaulted(NumTasks);
 					};
-					Params.Callback = [&TaskCheckValues, SessionTime](double EventStartTime, double EventEndTime, uint32 EventDepth, const TraceServices::FTimingProfilerEvent& Event, uint32 TaskIndex)
+					Params.EventRangeCallback = [&TaskCheckValues, SessionTime](double EventStartTime, double EventEndTime, uint32 EventDepth, const TraceServices::FTimingProfilerEvent& Event, uint32 TaskIndex)
 					{
 						EventEndTime = FMath::Min(EventEndTime, SessionTime);
 						TaskCheckValues[TaskIndex].TotalEventDuration += EventEndTime - EventStartTime;
@@ -567,7 +567,7 @@ void FTimingProfilerTests::RunEnumerateAsyncAllTracksBenchmark(const FEnumerateT
 							{
 								TaskCheckValues.AddDefaulted(NumTasks);
 							};
-							Params.Callback = [&OutCheckValues, &TaskCheckValues](double EventStartTime, double EventEndTime, uint32 EventDepth, const TraceServices::FTimingProfilerEvent& Event, uint32 TaskIndex)
+							Params.EventRangeCallback = [&OutCheckValues, &TaskCheckValues](double EventStartTime, double EventEndTime, uint32 EventDepth, const TraceServices::FTimingProfilerEvent& Event, uint32 TaskIndex)
 							{
 								EventEndTime = FMath::Min(EventEndTime, OutCheckValues.SessionDuration);
 								TaskCheckValues[TaskIndex].TotalEventDuration += EventEndTime - EventStartTime;
diff --git a/Engine/Source/Developer/TraceInsights/Private/Insights/ViewModels/FrameStatsHelper.cpp b/Engine/Source/Developer/TraceInsights/Private/Insights/ViewModels/FrameStatsHelper.cpp
--- a/Engine/Source/Developer/TraceInsights/Private/Insights/ViewModels/FrameStatsHelper.cpp
+++ b/Engine/Source/Developer/TraceInsights/Private/Insights/ViewModels/FrameStatsHelper.cpp
@@ -54,54 +54,82 @@ void FFrameStatsHelper::ProcessTimeline(TArray<FFrameStatsCachedEvent>& FrameSta
 		TimingProfilerProvider.ReadTimeline(TimelineIndex,
 			[SessionDuration, &FrameStatsEvents, TimerReader, TimerId](const TraceServices::ITimingProfilerProvider::Timeline& Timeline)
 			{
+				struct TaskData
+				{
+					double StartTime;
+					int32 NestedDepth = 0;
+				};
+
+				TArray<TaskData> DataArray;
+
 				TraceServices::ITimeline<TraceServices::FTimingProfilerEvent>::EnumerateAsyncParams Params;
 				Params.IntervalStart = 0;
 				Params.IntervalEnd = SessionDuration;
 				Params.Resolution = 0.0;
-				Params.SetupCallback = [](uint32 NumTasks) {};
-				Params.Callback = [TimerReader, &FrameStatsEvents, SessionDuration, TimerId](double StartTime, double EndTime, uint32 Depth, const TraceServices::FTimingProfilerEvent& Event, uint32 TaskIndex)
+				Params.SetupCallback = [&DataArray](uint32 NumTasks) { DataArray.AddDefaulted(NumTasks); };
+				Params.EventCallback = [TimerReader, &FrameStatsEvents, TimerId, &DataArray](bool bIsEnter, double Time, const TraceServices::FTimingProfilerEvent& Event, uint32 TaskIndex)
 				{
 					const TraceServices::FTimingProfilerTimer* Timer = TimerReader->GetTimer(Event.TimerIndex);
 					if (ensure(Timer != nullptr))
 					{
 						if (Timer->Id == TimerId)
 						{
-							int32 Index = Algo::UpperBoundBy(FrameStatsEvents, StartTime, &FFrameStatsCachedEvent::FrameStartTime);
-							if (Index > 0)
+							TaskData& CurrentTaskData = DataArray[TaskIndex];
+							if (bIsEnter)
 							{
-								--Index;
-							}
+								if (CurrentTaskData.NestedDepth == 0)
+								{
+									CurrentTaskData.StartTime = Time;
+								}
 
-							// This can can happen when the event is between frames.
-							if (StartTime > FrameStatsEvents[Index].FrameEndTime)
+								++CurrentTaskData.NestedDepth;
+							}
+							else
 							{
-								Index++;
-								if (Index >= FrameStatsEvents.Num())
+								check(CurrentTaskData.NestedDepth > 0);
+								if (--CurrentTaskData.NestedDepth > 0)
 								{
 									return TraceServices::EEventEnumerate::Continue;
 								}
-							}
 
-							do
-							{
-								FFrameStatsCachedEvent& Entry = FrameStatsEvents[Index];
-
-								if (EndTime < Entry.FrameStartTime)
+								int32 Index = Algo::UpperBoundBy(FrameStatsEvents, CurrentTaskData.StartTime, &FFrameStatsCachedEvent::FrameStartTime);
+								if (Index > 0)
 								{
-									return TraceServices::EEventEnumerate::Continue;
+									--Index;
 								}
 
-								if (StartTime < Entry.FrameStartTime)
+								// This can can happen when the event is between frames.
+								if (CurrentTaskData.StartTime > FrameStatsEvents[Index].FrameEndTime)
 								{
-									StartTime = Entry.FrameStartTime;
+									Index++;
+									if (Index >= FrameStatsEvents.Num())
+									{
+										return TraceServices::EEventEnumerate::Continue;
+									}
 								}
 
-								const double Duration = FMath::Min(EndTime, Entry.FrameEndTime) - StartTime;
-								ensure(Duration >= 0.0f);
-								for (double Value = Entry.Duration.load(); !Entry.Duration.compare_exchange_strong(Value, Value + Duration););
+								double EndTime = Time;
+								do
+								{
+									FFrameStatsCachedEvent& Entry = FrameStatsEvents[Index];
+
+									if (EndTime < Entry.FrameStartTime)
+									{
+										return TraceServices::EEventEnumerate::Continue;
+									}
+
+									if (CurrentTaskData.StartTime < Entry.FrameStartTime)
+									{
+										CurrentTaskData.StartTime = Entry.FrameStartTime;
+									}
 
-								Index++;
-							} while (Index < FrameStatsEvents.Num());
+									const double Duration = FMath::Min(EndTime, Entry.FrameEndTime) - CurrentTaskData.StartTime;
+									ensure(Duration >= 0.0f);
+									for (double Value = Entry.Duration.load(); !Entry.Duration.compare_exchange_strong(Value, Value + Duration););
+
+									Index++;
+								} while (Index < FrameStatsEvents.Num());
+							}
 						}
 					}
 					return TraceServices::EEventEnumerate::Continue;
diff --git a/Engine/Source/Developer/TraceInsights/Private/Insights/ViewModels/ThreadTimingTrack.cpp b/Engine/Source/Developer/TraceInsights/Private/Insights/ViewModels/ThreadTimingTrack.cpp
--- a/Engine/Source/Developer/TraceInsights/Private/Insights/ViewModels/ThreadTimingTrack.cpp
+++ b/Engine/Source/Developer/TraceInsights/Private/Insights/ViewModels/ThreadTimingTrack.cpp
@@ -1005,7 +1005,7 @@ void FThreadTimingTrack::BuildFilteredDrawState(ITimingEventsTrackDrawStateBuild
 							{
 								FilteredEvents.AddDefaulted(NumTasks);
 							};
-							Params.Callback = [this, &Builder, TimerReader, FilterEventType, &FilteredEvents](double StartTime, double EndTime, uint32 Depth, const TraceServices::FTimingProfilerEvent& Event, uint32 TaskIndex)
+							Params.EventRangeCallback = [this, &Builder, TimerReader, FilterEventType, &FilteredEvents](double StartTime, double EndTime, uint32 Depth, const TraceServices::FTimingProfilerEvent& Event, uint32 TaskIndex)
 							{
 								const TraceServices::FTimingProfilerTimer* Timer = TimerReader->GetTimer(Event.TimerIndex);
 								if (ensure(Timer != nullptr))
@@ -1138,7 +1138,7 @@ void FThreadTimingTrack::BuildFilteredDrawState(ITimingEventsTrackDrawStateBuild
 							Context.AddFilterData<int64>(static_cast<int32>(EFilterField::Metadata), 0);
 						}
 					};
-					Params.Callback = [this, &Builder, TimerReader, &FilteredEvents, &FilterContexts](double StartTime, double EndTime, uint32 Depth, const TraceServices::FTimingProfilerEvent& Event, uint32 TaskIndex)
+					Params.EventRangeCallback = [this, &Builder, TimerReader, &FilteredEvents, &FilterContexts](double StartTime, double EndTime, uint32 Depth, const TraceServices::FTimingProfilerEvent& Event, uint32 TaskIndex)
 					{
 						const TraceServices::FTimingProfilerTimer* Timer = TimerReader->GetTimer(Event.TimerIndex);
 						if (ensure(Timer != nullptr))
@@ -1980,7 +1980,7 @@ TSharedPtr<const ITimingEvent> FThreadTimingTrack::FindMaxEventInstance(uint32 I
 			{
 				CandidateEvents.AddDefaulted(NumTasks);
 			};
-			Params.Callback = [TimerReader, &CandidateEvents, InTimerId](double StartTime, double EndTime, uint32 Depth, const TraceServices::FTimingProfilerEvent& Event, uint32 TaskIndex)
+			Params.EventRangeCallback = [TimerReader, &CandidateEvents, InTimerId](double StartTime, double EndTime, uint32 Depth, const TraceServices::FTimingProfilerEvent& Event, uint32 TaskIndex)
 			{
 				const TraceServices::FTimingProfilerTimer* Timer = TimerReader->GetTimer(Event.TimerIndex);
 				if (ensure(Timer != nullptr))
@@ -2063,7 +2063,7 @@ TSharedPtr<const ITimingEvent> FThreadTimingTrack::FindMinEventInstance(uint32 I
 			{
 				CandidateEvents.AddDefaulted(NumTasks);
 			};
-			Params.Callback = [TimerReader, &CandidateEvents, InTimerId](double StartTime, double EndTime, uint32 Depth, const TraceServices::FTimingProfilerEvent& Event, uint32 TaskIndex)
+			Params.EventRangeCallback = [TimerReader, &CandidateEvents, InTimerId](double StartTime, double EndTime, uint32 Depth, const TraceServices::FTimingProfilerEvent& Event, uint32 TaskIndex)
 			{
 				const TraceServices::FTimingProfilerTimer* Timer = TimerReader->GetTimer(Event.TimerIndex);
 				if (ensure(Timer != nullptr))
diff --git a/Engine/Source/Developer/TraceInsights/Private/Insights/ViewModels/TimingGraphTrack.cpp b/Engine/Source/Developer/TraceInsights/Private/Insights/ViewModels/TimingGraphTrack.cpp
--- a/Engine/Source/Developer/TraceInsights/Private/Insights/ViewModels/TimingGraphTrack.cpp
+++ b/Engine/Source/Developer/TraceInsights/Private/Insights/ViewModels/TimingGraphTrack.cpp
@@ -476,7 +476,7 @@ void FTimingGraphTrack::UpdateTimerSeries(FTimingGraphSeries& Series, const FTim
 						{
 							Events.AddDefaulted(NumTasks);
 						};
-						Params.Callback = [&Events, TimerReader, &Series](double StartTime, double EndTime, uint32 Depth, const TraceServices::FTimingProfilerEvent& Event, uint32 TaskIndex)
+						Params.EventRangeCallback = [&Events, TimerReader, &Series](double StartTime, double EndTime, uint32 Depth, const TraceServices::FTimingProfilerEvent& Event, uint32 TaskIndex)
 						{
 							const TraceServices::FTimingProfilerTimer* Timer = TimerReader->GetTimer(Event.TimerIndex);
 							if (ensure(Timer != nullptr))
diff --git a/Engine/Source/Developer/TraceServices/Public/Model/AsyncEnumerateTask.h b/Engine/Source/Developer/TraceServices/Public/Model/AsyncEnumerateTask.h
--- a/Engine/Source/Developer/TraceServices/Public/Model/AsyncEnumerateTask.h
+++ b/Engine/Source/Developer/TraceServices/Public/Model/AsyncEnumerateTask.h
@@ -21,7 +21,9 @@ struct FAsyncEnumerateTaskData
 	double EndTime;
 	EEventSortOrder SortOrder;
 	const FDetailLevel<InEventType, SettingsType>* DetailLevel;
-	typename ITimeline<InEventType>::AsyncEventRangeCallback Callback;
+
+	typename ITimeline<InEventType>::AsyncEventRangeCallback EventRangeCallback;
+	typename ITimeline<InEventType>::AsyncEventCallback EventCallback;
 };
 
 template<typename InEventType, typename SettingsType>
@@ -253,7 +255,7 @@ public:
 						{
 							if (OutputEvent.EndTime >= Data.StartTime && OutputEvent.StartTime <= Data.EndTime)
 							{
-								if (this->Data.Callback(OutputEvent.StartTime, OutputEvent.EndTime, OutputEvent.Depth, OutputEvent.Event, Data.TaskIndex) == EEventEnumerate::Stop)
+								if (this->Data.EventRangeCallback(OutputEvent.StartTime, OutputEvent.EndTime, OutputEvent.Depth, OutputEvent.Event, Data.TaskIndex) == EEventEnumerate::Stop)
 								{
 									return EEventEnumerate::Stop;
 								}
@@ -298,7 +300,7 @@ public:
 
 					if (Time >= Data.StartTime && StackEntry.StartTime <= Data.EndTime)
 					{
-						if (this->Data.Callback(StackEntry.StartTime, Time, CurrentDepth, Event, Data.TaskIndex) == EEventEnumerate::Stop)
+						if (this->Data.EventRangeCallback(StackEntry.StartTime, Time, CurrentDepth, Event, Data.TaskIndex) == EEventEnumerate::Stop)
 						{
 							return EEventEnumerate::Stop;
 						}
@@ -311,13 +313,27 @@ public:
 
 	void DoWork()
 	{
-		if (Data.SortOrder == EEventSortOrder::ByEndTime)
+		if (Data.EventRangeCallback != nullptr)
 		{
-			EnumerateOrderedByEndTime();
+			if (Data.SortOrder == EEventSortOrder::ByEndTime)
+			{
+				EnumerateOrderedByEndTime();
+			}
+			else if (Data.SortOrder == EEventSortOrder::ByStartTime)
+			{
+				EnumerateOrderedByStartTime();
+			}
 		}
-		else if (Data.SortOrder == EEventSortOrder::ByStartTime)
+		else
 		{
-			EnumerateOrderedByStartTime();
+			EnumerateAsyncAlgoritm<InEventType, SettingsType>::EnumerateEventsDownSampled(Data,
+				[](int32 InInitialDepth)
+				{
+				},
+				[this](bool IsEnter, double Time, const EventType& Event)
+				{
+					return this->Data.EventCallback(IsEnter, Time, Event, Data.TaskIndex);
+				});
 		}
 	}
 
diff --git a/Engine/Source/Developer/TraceServices/Public/Model/MonotonicTimeline.h b/Engine/Source/Developer/TraceServices/Public/Model/MonotonicTimeline.h
--- a/Engine/Source/Developer/TraceServices/Public/Model/MonotonicTimeline.h
+++ b/Engine/Source/Developer/TraceServices/Public/Model/MonotonicTimeline.h
@@ -623,6 +623,9 @@ public:
 
 	virtual void EnumerateEventsDownSampledAsync(const typename ITimeline<EventType>::EnumerateAsyncParams& EnumerateAsyncParams) const override
 	{
+		// Exactly one of these 2 callbacks must be set.
+		check((EnumerateAsyncParams.EventRangeCallback == nullptr) ^ (EnumerateAsyncParams.EventCallback == nullptr));
+
 		if (EnumerateAsyncParams.IntervalEnd < EnumerateAsyncParams.IntervalStart)
 		{
 			return;
@@ -689,7 +692,8 @@ public:
 			TaskData.EndTime = EnumerateAsyncParams.IntervalEnd;
 			TaskData.SortOrder = EnumerateAsyncParams.SortOrder;
 			TaskData.DetailLevel = &DetailLevel;
-			TaskData.Callback = EnumerateAsyncParams.Callback;
+			TaskData.EventCallback = EnumerateAsyncParams.EventCallback;
+			TaskData.EventRangeCallback = EnumerateAsyncParams.EventRangeCallback;
 
 			TSharedRef<FAsyncTask<FEnumarateAsyncTask>> AsyncTask = MakeShared<FAsyncTask<FEnumarateAsyncTask>>(TaskData);
 			WorkerTasks.Add(AsyncTask);
diff --git a/Engine/Source/Developer/TraceServices/Public/TraceServices/Containers/Timelines.h b/Engine/Source/Developer/TraceServices/Public/TraceServices/Containers/Timelines.h
--- a/Engine/Source/Developer/TraceServices/Public/TraceServices/Containers/Timelines.h
+++ b/Engine/Source/Developer/TraceServices/Public/TraceServices/Containers/Timelines.h
@@ -26,7 +26,8 @@ public:
 	typedef InEventType EventType;
 	typedef TFunctionRef<EEventEnumerate(bool /*bStart*/, double /*Time*/, const EventType& /*Event*/)> EventCallback;
 	typedef TFunctionRef<EEventEnumerate(double /*StartTime*/, double /*EndTime*/, uint32 /*Depth*/, const EventType&/*Event*/)> EventRangeCallback;
-	typedef TFunction<EEventEnumerate(double /*StartTime*/, double /*EndTime*/, uint32 /*Depth*/, const EventType&/*Event*/, uint32 TaskIndex)> AsyncEventRangeCallback;
+	typedef TFunction<EEventEnumerate(double /*StartTime*/, double /*EndTime*/, uint32 /*Depth*/, const EventType&/*Event*/, uint32 /*TaskIndex*/)> AsyncEventRangeCallback;
+	typedef TFunction<EEventEnumerate(bool /*bStart*/, double /*Time*/, const EventType& /*Event*/, uint32 /*TaskIndex*/)> AsyncEventCallback;
 	typedef TFunction<void(uint32 /*NumTasks*/)> TaskInfoCallback;
 
 	struct FTimelineEventInfo
@@ -45,7 +46,10 @@ public:
 		double MaxOccupancy = 0.875; // Leave 1/8 threads free to avoid starvation.
 		EEventSortOrder SortOrder = EEventSortOrder::ByEndTime;
 		typename ITimeline<EventType>::TaskInfoCallback SetupCallback;
-		typename ITimeline<EventType>::AsyncEventRangeCallback Callback;
+
+		// Only one of these 2 callbacks must be set.
+		typename ITimeline<EventType>::AsyncEventRangeCallback EventRangeCallback;
+		typename ITimeline<EventType>::AsyncEventCallback EventCallback;
 	};
 
 	virtual ~ITimeline() = default;
