[Insights] Fix task table display bugs. Missing values will now be shown with N/A instead of 0.

diff --git a/Engine/Source/Developer/TraceInsights/Private/Insights/TaskGraphProfiler/ViewModels/TaskTable.cpp b/Engine/Source/Developer/TraceInsights/Private/Insights/TaskGraphProfiler/ViewModels/TaskTable.cpp
--- a/Engine/Source/Developer/TraceInsights/Private/Insights/TaskGraphProfiler/ViewModels/TaskTable.cpp
+++ b/Engine/Source/Developer/TraceInsights/Private/Insights/TaskGraphProfiler/ViewModels/TaskTable.cpp
@@ -3,6 +3,7 @@
 #include "TaskTable.h"
 
 // Insights
+#include "Insights/Common/TimeUtils.h"
 #include "Insights/TaskGraphProfiler/ViewModels/TaskNode.h"
 #include "Insights/TaskGraphProfiler/ViewModels/TaskTable.h"
 #include "Insights/Table/ViewModels/TableCellValueFormatter.h"
@@ -72,6 +73,18 @@ public:
 
 ////////////////////////////////////////////////////////////////////////////////////////////////////
 
+double GetRelativeValue(double ValueA, double ValueB)
+{
+	if (ValueA == TraceServices::FTaskInfo::InvalidTimestamp || ValueB == TraceServices::FTaskInfo::InvalidTimestamp)
+	{
+		return TraceServices::FTaskInfo::InvalidTimestamp;
+	}
+
+	return ValueA - ValueB;
+}
+
+////////////////////////////////////////////////////////////////////////////////////////////////////
+
 struct DefaultTaskFieldGetterFuncts
 {
 	static FTableCellValue GetDebugName(const FTableColumn& Column, const FTaskEntry& Task) { return FTableCellValue(Task.GetDebugName()); }
@@ -108,23 +121,12 @@ struct DefaultTaskFieldGetterFuncts
 struct RelativeToPreviousTaskFieldGetterFuncts
 {
 	static FTableCellValue GetCreatedTimestamp(const FTableColumn& Column, const FTaskEntry& Task) { return FTableCellValue(Task.GetCreatedTimestamp()); }
-	static FTableCellValue GetLaunchedTimestamp(const FTableColumn& Column, const FTaskEntry& Task) { return FTableCellValue(Task.GetLaunchedTimestamp() - Task.GetCreatedTimestamp()); }
-	static FTableCellValue GetScheduledTimestamp(const FTableColumn& Column, const FTaskEntry& Task) { return FTableCellValue(Task.GetScheduledTimestamp() - Task.GetLaunchedTimestamp()); }
-	static FTableCellValue GetStartedTimestamp(const FTableColumn& Column, const FTaskEntry& Task) { return FTableCellValue(Task.GetStartedTimestamp() - Task.GetScheduledTimestamp()); }
-	static FTableCellValue GetFinishedTimestamp(const FTableColumn& Column, const FTaskEntry& Task) { return FTableCellValue(Task.GetFinishedTimestamp() - Task.GetStartedTimestamp()); }
-	static FTableCellValue GetCompletedTimestamp(const FTableColumn& Column, const FTaskEntry& Task) 
-	{
-		// CompletedTimestamp can be 0, so don't show negative numbers in this case.
-		if (Task.GetCompletedTimestamp() != 0)
-		{
-			return FTableCellValue(Task.GetCompletedTimestamp() - Task.GetFinishedTimestamp());
-		}
-		else
-		{
-			return FTableCellValue(Task.GetCompletedTimestamp());
-		}
-	}
-	static FTableCellValue GetDestroyedTimestamp(const FTableColumn& Column, const FTaskEntry& Task) { return FTableCellValue(Task.GetDestroyedTimestamp() - Task.GetCompletedTimestamp()); }
+	static FTableCellValue GetLaunchedTimestamp(const FTableColumn& Column, const FTaskEntry& Task) { return FTableCellValue(GetRelativeValue(Task.GetLaunchedTimestamp(), Task.GetCreatedTimestamp())); }
+	static FTableCellValue GetScheduledTimestamp(const FTableColumn& Column, const FTaskEntry& Task) { return FTableCellValue(GetRelativeValue(Task.GetScheduledTimestamp(), Task.GetLaunchedTimestamp())); }
+	static FTableCellValue GetStartedTimestamp(const FTableColumn& Column, const FTaskEntry& Task) { return FTableCellValue(GetRelativeValue(Task.GetStartedTimestamp(), Task.GetScheduledTimestamp())); }
+	static FTableCellValue GetFinishedTimestamp(const FTableColumn& Column, const FTaskEntry& Task) { return FTableCellValue(GetRelativeValue(Task.GetFinishedTimestamp(), Task.GetStartedTimestamp())); }
+	static FTableCellValue GetCompletedTimestamp(const FTableColumn& Column, const FTaskEntry& Task) { return FTableCellValue(GetRelativeValue(Task.GetCompletedTimestamp(), Task.GetFinishedTimestamp())); }
+	static FTableCellValue GetDestroyedTimestamp(const FTableColumn& Column, const FTaskEntry& Task) { return FTableCellValue(GetRelativeValue(Task.GetDestroyedTimestamp(), Task.GetCompletedTimestamp())); }
 };
 
 ////////////////////////////////////////////////////////////////////////////////////////////////////
@@ -132,23 +134,50 @@ struct RelativeToPreviousTaskFieldGetterFuncts
 struct RelativeToCreatedTaskFieldGetterFuncts
 {
 	static FTableCellValue GetCreatedTimestamp(const FTableColumn& Column, const FTaskEntry& Task) { return FTableCellValue(Task.GetCreatedTimestamp()); }
-	static FTableCellValue GetLaunchedTimestamp(const FTableColumn& Column, const FTaskEntry& Task) { return FTableCellValue(Task.GetLaunchedTimestamp() - Task.GetCreatedTimestamp()); }
-	static FTableCellValue GetScheduledTimestamp(const FTableColumn& Column, const FTaskEntry& Task) { return FTableCellValue(Task.GetScheduledTimestamp() - Task.GetCreatedTimestamp()); }
-	static FTableCellValue GetStartedTimestamp(const FTableColumn& Column, const FTaskEntry& Task) { return FTableCellValue(Task.GetStartedTimestamp() - Task.GetCreatedTimestamp()); }
-	static FTableCellValue GetFinishedTimestamp(const FTableColumn& Column, const FTaskEntry& Task) { return FTableCellValue(Task.GetFinishedTimestamp() - Task.GetCreatedTimestamp()); }
-	static FTableCellValue GetCompletedTimestamp(const FTableColumn& Column, const FTaskEntry& Task)
+	static FTableCellValue GetLaunchedTimestamp(const FTableColumn& Column, const FTaskEntry& Task) { return FTableCellValue(GetRelativeValue(Task.GetLaunchedTimestamp(), Task.GetCreatedTimestamp())); }
+	static FTableCellValue GetScheduledTimestamp(const FTableColumn& Column, const FTaskEntry& Task) { return FTableCellValue(GetRelativeValue(Task.GetScheduledTimestamp(), Task.GetCreatedTimestamp())); }
+	static FTableCellValue GetStartedTimestamp(const FTableColumn& Column, const FTaskEntry& Task) { return FTableCellValue(GetRelativeValue(Task.GetStartedTimestamp(), Task.GetCreatedTimestamp())); }
+	static FTableCellValue GetFinishedTimestamp(const FTableColumn& Column, const FTaskEntry& Task) { return FTableCellValue(GetRelativeValue(Task.GetFinishedTimestamp(), Task.GetCreatedTimestamp())); }
+	static FTableCellValue GetCompletedTimestamp(const FTableColumn& Column, const FTaskEntry& Task) { return FTableCellValue(GetRelativeValue(Task.GetCompletedTimestamp(), Task.GetCreatedTimestamp())); }
+	static FTableCellValue GetDestroyedTimestamp(const FTableColumn& Column, const FTaskEntry& Task) { return FTableCellValue(GetRelativeValue(Task.GetDestroyedTimestamp(), Task.GetCreatedTimestamp())); }
+};
+
+////////////////////////////////////////////////////////////////////////////////////////////////////
+
+class FTaskDoubleValueFormatterAsTimeAuto : public FTableCellValueFormatter
+{
+public:
+	virtual FText FormatValue(const TOptional<FTableCellValue>& InValue) const override
 	{
-		// CompletedTimestamp can be 0, so don't show negative numbers in this case.
-		if (Task.GetCompletedTimestamp() != 0)
+		if (InValue.IsSet())
 		{
-			return FTableCellValue(Task.GetCompletedTimestamp() - Task.GetCreatedTimestamp());
+			const double Value = InValue.GetValue().Double;
+			if (Value != TraceServices::FTaskInfo::InvalidTimestamp)
+			{
+				return FText::FromString(TimeUtils::FormatTimeAuto(Value));
+			}
 		}
-		else
+
+		return FText::FromString(TEXT("N/A"));
+	}
+
+	virtual FText FormatValueForTooltip(const TOptional<FTableCellValue>& InValue) const override
+	{
+		if (InValue.IsSet())
 		{
-			return FTableCellValue(Task.GetCompletedTimestamp());
+			const double Value = InValue.GetValue().Double;
+			if (Value == 0.0)
+			{
+				return FText::FromString(TEXT("0"));
+			}
+			else if (Value != TraceServices::FTaskInfo::InvalidTimestamp)
+			{
+				return FText::FromString(FString::Printf(TEXT("%f (%s)"), Value, *TimeUtils::FormatTimeAuto(Value)));
+			}
 		}
+
+		return FText::FromString(TEXT("N/A"));
 	}
-	static FTableCellValue GetDestroyedTimestamp(const FTableColumn& Column, const FTaskEntry& Task) { return FTableCellValue(Task.GetDestroyedTimestamp() - Task.GetCreatedTimestamp()); }
 };
 
 ////////////////////////////////////////////////////////////////////////////////////////////////////
@@ -249,7 +278,7 @@ void FTaskTable::AddDefaultColumns()
 		TSharedRef<ITableCellValueGetter> Getter = MakeShared<FTaskColumnValueGetter<DefaultTaskFieldGetterFuncts::GetCreatedTimestamp>>();
 		Column.SetValueGetter(Getter);
 
-		TSharedRef<ITableCellValueFormatter> Formatter = MakeShared<FDoubleValueFormatterAsTimeAuto>();
+		TSharedRef<ITableCellValueFormatter> Formatter = MakeShared<FTaskDoubleValueFormatterAsTimeAuto>();
 		Column.SetValueFormatter(Formatter);
 
 		TSharedRef<ITableCellValueSorter> Sorter = MakeShared<FSorterByDoubleValue>(ColumnRef);
@@ -312,7 +341,7 @@ void FTaskTable::AddDefaultColumns()
 		TSharedRef<ITableCellValueGetter> Getter = MakeShared<FTaskColumnValueGetter<DefaultTaskFieldGetterFuncts::GetLaunchedTimestamp>>();
 		Column.SetValueGetter(Getter);
 
-		TSharedRef<ITableCellValueFormatter> Formatter = MakeShared<FDoubleValueFormatterAsTimeAuto>();
+		TSharedRef<ITableCellValueFormatter> Formatter = MakeShared<FTaskDoubleValueFormatterAsTimeAuto>();
 		Column.SetValueFormatter(Formatter);
 
 		TSharedRef<ITableCellValueSorter> Sorter = MakeShared<FSorterByDoubleValue>(ColumnRef);
@@ -375,7 +404,7 @@ void FTaskTable::AddDefaultColumns()
 		TSharedRef<ITableCellValueGetter> Getter = MakeShared<FTaskColumnValueGetter<DefaultTaskFieldGetterFuncts::GetScheduledTimestamp>>();
 		Column.SetValueGetter(Getter);
 
-		TSharedRef<ITableCellValueFormatter> Formatter = MakeShared<FDoubleValueFormatterAsTimeAuto>();
+		TSharedRef<ITableCellValueFormatter> Formatter = MakeShared<FTaskDoubleValueFormatterAsTimeAuto>();
 		Column.SetValueFormatter(Formatter);
 
 		TSharedRef<ITableCellValueSorter> Sorter = MakeShared<FSorterByDoubleValue>(ColumnRef);
@@ -437,7 +466,7 @@ void FTaskTable::AddDefaultColumns()
 		TSharedRef<ITableCellValueGetter> Getter = MakeShared<FTaskColumnValueGetter<DefaultTaskFieldGetterFuncts::GetStartedTimestamp>>();
 		Column.SetValueGetter(Getter);
 
-		TSharedRef<ITableCellValueFormatter> Formatter = MakeShared<FDoubleValueFormatterAsTimeAuto>();
+		TSharedRef<ITableCellValueFormatter> Formatter = MakeShared<FTaskDoubleValueFormatterAsTimeAuto>();
 		Column.SetValueFormatter(Formatter);
 
 		TSharedRef<ITableCellValueSorter> Sorter = MakeShared<FSorterByDoubleValue>(ColumnRef);
@@ -499,7 +528,7 @@ void FTaskTable::AddDefaultColumns()
 		TSharedRef<ITableCellValueGetter> Getter = MakeShared<FTaskColumnValueGetter<DefaultTaskFieldGetterFuncts::GetFinishedTimestamp>>();
 		Column.SetValueGetter(Getter);
 
-		TSharedRef<ITableCellValueFormatter> Formatter = MakeShared<FDoubleValueFormatterAsTimeAuto>();
+		TSharedRef<ITableCellValueFormatter> Formatter = MakeShared<FTaskDoubleValueFormatterAsTimeAuto>();
 		Column.SetValueFormatter(Formatter);
 
 		TSharedRef<ITableCellValueSorter> Sorter = MakeShared<FSorterByDoubleValue>(ColumnRef);
@@ -531,7 +560,7 @@ void FTaskTable::AddDefaultColumns()
 		TSharedRef<ITableCellValueGetter> Getter = MakeShared<FTaskColumnValueGetter<DefaultTaskFieldGetterFuncts::GetCompletedTimestamp>>();
 		Column.SetValueGetter(Getter);
 
-		TSharedRef<ITableCellValueFormatter> Formatter = MakeShared<FDoubleValueFormatterAsTimeAuto>();
+		TSharedRef<ITableCellValueFormatter> Formatter = MakeShared<FTaskDoubleValueFormatterAsTimeAuto>();
 		Column.SetValueFormatter(Formatter);
 
 		TSharedRef<ITableCellValueSorter> Sorter = MakeShared<FSorterByDoubleValue>(ColumnRef);
@@ -593,7 +622,7 @@ void FTaskTable::AddDefaultColumns()
 		TSharedRef<ITableCellValueGetter> Getter = MakeShared<FTaskColumnValueGetter<DefaultTaskFieldGetterFuncts::GetDestroyedTimestamp>>();
 		Column.SetValueGetter(Getter);
 
-		TSharedRef<ITableCellValueFormatter> Formatter = MakeShared<FDoubleValueFormatterAsTimeAuto>();
+		TSharedRef<ITableCellValueFormatter> Formatter = MakeShared<FTaskDoubleValueFormatterAsTimeAuto>();
 		Column.SetValueFormatter(Formatter);
 
 		TSharedRef<ITableCellValueSorter> Sorter = MakeShared<FSorterByDoubleValue>(ColumnRef);
@@ -794,6 +823,7 @@ void FTaskTable::SwitchToAbsoluteTimestamps()
 	FindColumnChecked(FTaskTableColumns::StartedTimestampColumnId)->SetValueGetter(MakeShared<FTaskColumnValueGetter<DefaultTaskFieldGetterFuncts::GetStartedTimestamp>>());
 	FindColumnChecked(FTaskTableColumns::FinishedTimestampColumnId)->SetValueGetter(MakeShared<FTaskColumnValueGetter<DefaultTaskFieldGetterFuncts::GetFinishedTimestamp>>());
 	FindColumnChecked(FTaskTableColumns::CompletedTimestampColumnId)->SetValueGetter(MakeShared<FTaskColumnValueGetter<DefaultTaskFieldGetterFuncts::GetCompletedTimestamp>>());
+	FindColumnChecked(FTaskTableColumns::DestroyedTimestampColumnId)->SetValueGetter(MakeShared<FTaskColumnValueGetter<DefaultTaskFieldGetterFuncts::GetDestroyedTimestamp>>());
 }
 
 ////////////////////////////////////////////////////////////////////////////////////////////////////
diff --git a/Engine/Source/Developer/TraceServices/Public/TraceServices/Model/TasksProfiler.h b/Engine/Source/Developer/TraceServices/Public/TraceServices/Model/TasksProfiler.h
--- a/Engine/Source/Developer/TraceServices/Public/TraceServices/Model/TasksProfiler.h
+++ b/Engine/Source/Developer/TraceServices/Public/TraceServices/Model/TasksProfiler.h
@@ -27,7 +27,7 @@ struct FTaskInfo
 	// thread - ENamedThreads::AnyThread), task priority, thread priority (in case of a worker thread), and queue index (in case of a named
 	// thread - main queue or local queue). See ENamedThreads to split the info into separate fields
 
-	static constexpr double InvalidTimestamp = 0.0;
+	static constexpr double InvalidTimestamp = -std::numeric_limits<double>::max();
 
 	// some tasks can be created with delayed launch, otherwise "created" and "launched" timestamps and thread IDs are equal
 	double CreatedTimestamp = InvalidTimestamp;
