A try to increase timeout for tests that are failing on MacOS and Linux.

diff --git a/Engine/Source/Developer/TraceInsights/Private/Insights/Tests/InsightsTestUtils.cpp b/Engine/Source/Developer/TraceInsights/Private/Insights/Tests/InsightsTestUtils.cpp
--- a/Engine/Source/Developer/TraceInsights/Private/Insights/Tests/InsightsTestUtils.cpp
+++ b/Engine/Source/Developer/TraceInsights/Private/Insights/Tests/InsightsTestUtils.cpp
@@ -85,6 +85,7 @@ bool FInsightsTestUtils::FileContainsString(const FString& PathToFile, const FSt
 	{
 		if (!FPaths::FileExists(PathToFile))
 		{
+			Test->AddInfo("Unable to find EngineTest.log at " + PathToFile);
 			FPlatformProcess::Sleep(0.1f);
 		}
 		else
