[Insights] Fixed selection of a frame in the Frames panel to disable auto-scroll in Timing View. Same for Log View, when a log message is clicked and Timing View needs to scroll, auto-scroll will be disabled.

diff --git a/Engine/Source/Developer/TraceInsights/Private/Insights/Widgets/SFrameTrack.cpp b/Engine/Source/Developer/TraceInsights/Private/Insights/Widgets/SFrameTrack.cpp
--- a/Engine/Source/Developer/TraceInsights/Private/Insights/Widgets/SFrameTrack.cpp
+++ b/Engine/Source/Developer/TraceInsights/Private/Insights/Widgets/SFrameTrack.cpp
@@ -556,6 +556,8 @@ void SFrameTrack::SelectFrameAtMousePosition(double X, double Y, bool JoinCurren
 					Duration = EndTime - StartTime;
 				}
 
+				TimingView->SetAutoScroll(false);
+
 				if (bZoomTimingViewOnFrameSelection)
 				{
 					const double EndTime = FMath::Min(StartTime + Duration, TimingView->GetViewport().GetMaxValidTime());
diff --git a/Engine/Source/Developer/TraceInsights/Private/Insights/Widgets/SLogView.cpp b/Engine/Source/Developer/TraceInsights/Private/Insights/Widgets/SLogView.cpp
--- a/Engine/Source/Developer/TraceInsights/Private/Insights/Widgets/SLogView.cpp
+++ b/Engine/Source/Developer/TraceInsights/Private/Insights/Widgets/SLogView.cpp
@@ -1012,6 +1012,7 @@ void SLogView::SelectLogMessage(TSharedPtr<FLogMessage> LogMessage)
 				}
 				else
 				{
+					TimingView->SetAutoScroll(false);
 					TimingView->SetAndCenterOnTimeMarker(Time);
 				}
 			}
diff --git a/Engine/Source/Developer/TraceInsights/Private/Insights/Widgets/STimingView.h b/Engine/Source/Developer/TraceInsights/Private/Insights/Widgets/STimingView.h
--- a/Engine/Source/Developer/TraceInsights/Private/Insights/Widgets/STimingView.h
+++ b/Engine/Source/Developer/TraceInsights/Private/Insights/Widgets/STimingView.h
@@ -346,6 +346,9 @@ public:
 	bool IsTimeSelected(double Time) const { return Time >= SelectionStartTime && Time < SelectionEndTime; }
 	bool IsTimeSelectedInclusive(double Time) const { return Time >= SelectionStartTime && Time <= SelectionEndTime; }
 
+	bool IsAutoScrollEnabled() const { return bAutoScroll; }
+	void SetAutoScroll(bool bOnOff);
+
 	void ScrollAtPosY(float ScrollPosY);
 	void BringIntoViewY(float InTopY, float InBottomY);
 	void BringScrollableTrackIntoView(const FBaseTimingTrack& Track);
@@ -442,8 +445,6 @@ protected:
 	////////////////////////////////////////////////////////////////////////////////////////////////////
 	// Auto-Scroll
 
-	void SetAutoScroll(bool bOnOff);
-
 	void AutoScroll_OnCheckStateChanged(ECheckBoxState NewRadioState);
 	ECheckBoxState AutoScroll_IsChecked() const;
 
