[UnrealInsights] Added Branch and Version columns to session browser
#tests local

diff --git a/Engine/Source/Developer/TraceInsights/Private/Insights/Widgets/SStartPageWindow.cpp b/Engine/Source/Developer/TraceInsights/Private/Insights/Widgets/SStartPageWindow.cpp
--- a/Engine/Source/Developer/TraceInsights/Private/Insights/Widgets/SStartPageWindow.cpp
+++ b/Engine/Source/Developer/TraceInsights/Private/Insights/Widgets/SStartPageWindow.cpp
@@ -210,6 +210,26 @@ public:
 					.ToolTip(STraceListRow::GetTraceTooltip())
 				];
 		}
+		else if (ColumnName == TraceStoreColumns::BuildBranch)
+		{
+			return SNew(SBox)
+				.Padding(FMargin(4.0f, 0.0f))
+				[
+					SNew(STextBlock)
+					.Text(this, &STraceListRow::GetTraceBranch)
+					.ToolTip(STraceListRow::GetTraceTooltip())
+				];
+		}
+		else if (ColumnName == TraceStoreColumns::BuildVersion)
+		{
+			return SNew(SBox)
+				.Padding(FMargin(4.0f, 0.0f))
+				[
+					SNew(STextBlock)
+					.Text(this, &STraceListRow::GetTraceBuildVersion)
+					.ToolTip(STraceListRow::GetTraceTooltip())
+				];
+		}
 		else if (ColumnName == TraceStoreColumns::Size)
 		{
 			return SNew(SBox)
@@ -804,11 +824,11 @@ public:
 			AddGridPanelRow(GridPanel, Row++, LOCTEXT("TraceTooltip_Platform", "Platform:"), &STraceListRow::GetTracePlatform);
 			AddGridPanelRow(GridPanel, Row++, LOCTEXT("TraceTooltip_AppName", "App Name:"), &STraceListRow::GetTraceAppName);
 			AddGridPanelRow(GridPanel, Row++, LOCTEXT("TraceTooltip_CommandLine", "Command Line:"), &STraceListRow::GetTraceCommandLine, &STraceListRow::GetTraceCommandLineHighlightText);
-			AddGridPanelRow(GridPanel, Row++, LOCTEXT("TraceTooltip_Branch", "Branch:"), &STraceListRow::GetTraceBranch);
-			AddGridPanelRow(GridPanel, Row++, LOCTEXT("TraceTooltip_BuildVersion", "Build Version:"), &STraceListRow::GetTraceBuildVersion);
-			AddGridPanelRow(GridPanel, Row++, LOCTEXT("TraceTooltip_Changelist", "Changelist:"), &STraceListRow::GetTraceChangelist, nullptr, &STraceListRow::TraceChangelistVisibility);
 			AddGridPanelRow(GridPanel, Row++, LOCTEXT("TraceTooltip_BuildConfig", "Build Config:"), &STraceListRow::GetTraceBuildConfiguration);
 			AddGridPanelRow(GridPanel, Row++, LOCTEXT("TraceTooltip_BuildTarget", "Build Target:"), &STraceListRow::GetTraceBuildTarget);
+			AddGridPanelRow(GridPanel, Row++, LOCTEXT("TraceTooltip_Branch", "Build Branch:"), &STraceListRow::GetTraceBranch);
+			AddGridPanelRow(GridPanel, Row++, LOCTEXT("TraceTooltip_BuildVersion", "Build Version:"), &STraceListRow::GetTraceBuildVersion);
+			AddGridPanelRow(GridPanel, Row++, LOCTEXT("TraceTooltip_Changelist", "Changelist:"), &STraceListRow::GetTraceChangelist, nullptr, &STraceListRow::TraceChangelistVisibility);
 			AddGridPanelRow(GridPanel, Row++, LOCTEXT("TraceTooltip_Timestamp", "Timestamp:"), &STraceListRow::GetTraceTimestampForTooltip);
 			AddGridPanelRow(GridPanel, Row++, LOCTEXT("TraceTooltip_Size", "File Size:"), &STraceListRow::GetTraceSizeForTooltip);
 			AddGridPanelRow(GridPanel, Row++, LOCTEXT("TraceTooltip_Status", "Status:"), &STraceListRow::GetTraceStatusForTooltip);
@@ -1171,6 +1191,15 @@ TSharedRef<SWidget> STraceStoreWindow::ConstructFiltersToolbar()
 			FSlateIcon(FInsightsStyle::GetStyleSetName(), "Icons.Filter.ToolBar"),
 			false);
 
+		// Filter by Version
+		ToolbarBuilder.AddComboButton(
+			FUIAction(),
+			FOnGetContent::CreateSP(this, &STraceStoreWindow::MakeVersionFilterMenu),
+			LOCTEXT("FilterByVersionText", "Version"),
+			LOCTEXT("FilterByVersionToolTip", "Filters the list of trace sessions by Version."),
+			FSlateIcon(FInsightsStyle::GetStyleSetName(), "Icons.Filter.ToolBar"),
+			false);
+
 		// Filter by Size
 		ToolbarBuilder.AddComboButton(
 			FUIAction(),
@@ -1328,6 +1357,38 @@ TSharedRef<SWidget> STraceStoreWindow::ConstructSessionsPanel()
 				]
 			]
 
+			+ SHeaderRow::Column(TraceStoreColumns::BuildBranch)
+			.FillWidth(0.2f)
+			.InitialSortMode(EColumnSortMode::Ascending)
+			.SortMode(this, &STraceStoreWindow::GetSortModeForColumn, TraceStoreColumns::BuildBranch)
+			.OnSort(this, &STraceStoreWindow::OnSortModeChanged)
+			.OnGetMenuContent(this, &STraceStoreWindow::MakeBranchColumnHeaderMenu)
+			[
+				SNew(SBox)
+				.VAlign(VAlign_Center)
+				[
+					SNew(STextBlock)
+					.Text(LOCTEXT("BranchColumn", "Build Branch"))
+					.ColorAndOpacity_Lambda([this] { return FilterByBranch->IsEmpty() ? FLinearColor(0.5f, 0.5f, 0.5f, 1.0f) : FLinearColor(0.3f, 0.75f, 1.0f, 1.0f); })
+				]
+			]
+
+			+ SHeaderRow::Column(TraceStoreColumns::BuildVersion)
+			.FillWidth(0.25f)
+			.InitialSortMode(EColumnSortMode::Ascending)
+			.SortMode(this, &STraceStoreWindow::GetSortModeForColumn, TraceStoreColumns::BuildVersion)
+			.OnSort(this, &STraceStoreWindow::OnSortModeChanged)
+			.OnGetMenuContent(this, &STraceStoreWindow::MakeVersionColumnHeaderMenu)
+			[
+				SNew(SBox)
+				.VAlign(VAlign_Center)
+				[
+					SNew(STextBlock)
+					.Text(LOCTEXT("BuildVersionColumn", "Build Version"))
+					.ColorAndOpacity_Lambda([this] { return FilterByVersion->IsEmpty() ? FLinearColor(0.5f, 0.5f, 0.5f, 1.0f) : FLinearColor(0.3f, 0.75f, 1.0f, 1.0f); })
+				]
+			]
+
 			+ SHeaderRow::Column(TraceStoreColumns::Size)
 			.FixedWidth(100.0f)
 			.HAlignHeader(HAlign_Right)
@@ -3506,6 +3567,43 @@ void STraceStoreWindow::BuildBranchFilterSubMenu(FMenuBuilder& InMenuBuilder)
 
 ////////////////////////////////////////////////////////////////////////////////////////////////////
 
+TSharedRef<SWidget> STraceStoreWindow::MakeVersionColumnHeaderMenu()
+{
+	FSlateApplication::Get().CloseToolTip();
+
+	FMenuBuilder MenuBuilder(/*bInShouldCloseWindowAfterMenuSelection=*/true, nullptr);
+
+	MenuBuilder.BeginSection("Filter", LOCTEXT("ColumnHeaderMenuSection_VersionFilter", "Version Filter"));
+	BuildVersionFilterSubMenu(MenuBuilder);
+	MenuBuilder.EndSection();
+
+	return MenuBuilder.MakeWidget();
+}
+
+////////////////////////////////////////////////////////////////////////////////////////////////////
+
+TSharedRef<SWidget> STraceStoreWindow::MakeVersionFilterMenu()
+{
+	FSlateApplication::Get().CloseToolTip();
+
+	FMenuBuilder MenuBuilder(/*bInShouldCloseWindowAfterMenuSelection=*/true, nullptr);
+
+	MenuBuilder.BeginSection("Filter", LOCTEXT("MenuSection_VersionFilter", "Version Filter"));
+	BuildVersionFilterSubMenu(MenuBuilder);
+	MenuBuilder.EndSection();
+
+	return MenuBuilder.MakeWidget();
+}
+
+////////////////////////////////////////////////////////////////////////////////////////////////////
+
+void STraceStoreWindow::BuildVersionFilterSubMenu(FMenuBuilder& InMenuBuilder)
+{
+	FilterByVersion->BuildMenu(InMenuBuilder, *this);
+}
+
+////////////////////////////////////////////////////////////////////////////////////////////////////
+
 TSharedRef<SWidget> STraceStoreWindow::MakeSizeColumnHeaderMenu()
 {
 	FSlateApplication::Get().CloseToolTip();
@@ -3928,6 +4026,13 @@ FText FTraceFilterByStatus::ValueToText(const bool InValue) const
 
 ////////////////////////////////////////////////////////////////////////////////////////////////////
 
+FTraceFilterByVersion::FTraceFilterByVersion()
+{
+	ToggleAllActionTooltip = LOCTEXT("FilterByVersion_ToggleAll_Tooltip", "Shows or hides all versions.");
+}
+
+////////////////////////////////////////////////////////////////////////////////////////////////////
+
 void STraceStoreWindow::FilterByNameSearchBox_OnTextChanged(const FText& InFilterText)
 {
 	FilterByName->SetRawFilterText(InFilterText);
@@ -3975,6 +4080,9 @@ void STraceStoreWindow::CreateFilters()
 	FilterByBranch = MakeShared<FTraceFilterByBranch>();
 	Filters->Add(FilterByBranch);
 
+	FilterByVersion = MakeShared<FTraceFilterByVersion>();
+	Filters->Add(FilterByVersion);
+
 	FilterBySize = MakeShared<FTraceFilterBySize>();
 	Filters->Add(FilterBySize);
 
@@ -4008,6 +4116,7 @@ void STraceStoreWindow::UpdateFiltering()
 		FilterByBuildConfig->IsEmpty() &&
 		FilterByBuildTarget->IsEmpty() &&
 		FilterByBranch->IsEmpty() &&
+		FilterByVersion->IsEmpty() &&
 		FilterBySize->IsEmpty() &&
 		FilterByStatus->IsEmpty())
 	{
@@ -4209,6 +4318,32 @@ void STraceStoreWindow::UpdateSorting()
 				});
 		}
 	}
+	else if (SortColumn == TraceStoreColumns::BuildBranch)
+	{
+		if (SortMode == EColumnSortMode::Ascending)
+		{
+			FilteredTraceViewModels.Sort([](const TSharedPtr<FTraceViewModel>& A, const TSharedPtr<FTraceViewModel>& B)
+				{ return A->Branch.CompareTo(B->Branch) < 0; });
+		}
+		else
+		{
+			FilteredTraceViewModels.Sort([](const TSharedPtr<FTraceViewModel>& A, const TSharedPtr<FTraceViewModel>& B)
+				{ return B->Branch.CompareTo(A->Branch) < 0; });
+		}
+	}
+	else if (SortColumn == TraceStoreColumns::BuildVersion)
+	{
+		if (SortMode == EColumnSortMode::Ascending)
+		{
+			FilteredTraceViewModels.Sort([](const TSharedPtr<FTraceViewModel>& A, const TSharedPtr<FTraceViewModel>& B)
+				{ return A->BuildVersion.CompareTo(B->BuildVersion) < 0; });
+		}
+		else
+		{
+			FilteredTraceViewModels.Sort([](const TSharedPtr<FTraceViewModel>& A, const TSharedPtr<FTraceViewModel>& B)
+				{ return B->BuildVersion.CompareTo(A->BuildVersion) < 0; });
+		}
+	}
 	else if (SortColumn == TraceStoreColumns::Size)
 	{
 		if (SortMode == EColumnSortMode::Ascending)
diff --git a/Engine/Source/Developer/TraceInsights/Private/Insights/Widgets/SStartPageWindow.h b/Engine/Source/Developer/TraceInsights/Private/Insights/Widgets/SStartPageWindow.h
--- a/Engine/Source/Developer/TraceInsights/Private/Insights/Widgets/SStartPageWindow.h
+++ b/Engine/Source/Developer/TraceInsights/Private/Insights/Widgets/SStartPageWindow.h
@@ -43,6 +43,8 @@ namespace TraceStoreColumns
 	static const FName AppName(TEXT("AppName"));
 	static const FName BuildConfig(TEXT("BuildConfig"));
 	static const FName BuildTarget(TEXT("BuildTarget"));
+	static const FName BuildBranch(TEXT("BuildBranch"));
+	static const FName BuildVersion(TEXT("BuildVersion"));
 	static const FName Size(TEXT("Size"));
 	static const FName Status(TEXT("Status"));
 }
@@ -306,6 +308,18 @@ protected:
 	virtual FText ValueToText(const bool InValue) const override;
 };
 
+class FTraceFilterByVersion : public FTraceFilterByStringSet
+{
+public:
+	FTraceFilterByVersion();
+
+protected:
+	virtual FString GetFilterValueForTrace(const FTraceViewModel& InTrace) const override
+	{
+		return InTrace.BuildVersion.ToString();
+	}
+};
+
 ////////////////////////////////////////////////////////////////////////////////////////////////////
 
 /** Implements the Trace Store window. */
@@ -425,6 +439,10 @@ private:
 	TSharedRef<SWidget> MakeBranchFilterMenu();
 	void BuildBranchFilterSubMenu(FMenuBuilder& InMenuBuilder);
 
+	TSharedRef<SWidget> MakeVersionColumnHeaderMenu();
+	TSharedRef<SWidget> MakeVersionFilterMenu();
+	void BuildVersionFilterSubMenu(FMenuBuilder& InMenuBuilder);
+
 	TSharedRef<SWidget> MakeSizeColumnHeaderMenu();
 	TSharedRef<SWidget> MakeSizeFilterMenu();
 	void BuildSizeFilterSubMenu(FMenuBuilder& InMenuBuilder);
@@ -649,6 +667,7 @@ private:
 	TSharedPtr<FTraceFilterByBuildConfig> FilterByBuildConfig;
 	TSharedPtr<FTraceFilterByBuildTarget> FilterByBuildTarget;
 	TSharedPtr<FTraceFilterByBranch> FilterByBranch;
+	TSharedPtr<FTraceFilterByVersion> FilterByVersion;
 	TSharedPtr<FTraceFilterBySize> FilterBySize;
 	TSharedPtr<FTraceFilterByStatus> FilterByStatus;
 
