[Insights] Added TRACEINSIGHTS_API to SMajorTabWindow class.

diff --git a/Engine/Source/Developer/TraceInsights/Public/Insights/Widgets/SMajorTabWindow.h b/Engine/Source/Developer/TraceInsights/Public/Insights/Widgets/SMajorTabWindow.h
--- a/Engine/Source/Developer/TraceInsights/Public/Insights/Widgets/SMajorTabWindow.h
+++ b/Engine/Source/Developer/TraceInsights/Public/Insights/Widgets/SMajorTabWindow.h
@@ -25,7 +25,7 @@ namespace Insights
 {
 
 /** Implements the base class for a major tab window. */
-class SMajorTabWindow : public SCompoundWidget
+class TRACEINSIGHTS_API SMajorTabWindow : public SCompoundWidget
 {
 public:
 	/** Default constructor. */
