Fixed mismatched printf format arguments.

diff --git a/Engine/Plugins/Animation/ControlRig/Source/ControlRig/Private/ControlRigBlueprintGeneratedClass.cpp b/Engine/Plugins/Animation/ControlRig/Source/ControlRig/Private/ControlRigBlueprintGeneratedClass.cpp
--- a/Engine/Plugins/Animation/ControlRig/Source/ControlRig/Private/ControlRigBlueprintGeneratedClass.cpp
+++ b/Engine/Plugins/Animation/ControlRig/Source/ControlRig/Private/ControlRigBlueprintGeneratedClass.cpp
@@ -30,7 +30,7 @@ void UControlRigBlueprintGeneratedClass::Serialize(FArchive& Ar)
 
 	// for debugging purposes we'll give this VM a name that's useful.
 	static TAtomic<uint32> NumVMs{ 0 };
-	static constexpr TCHAR Format[] = TEXT("%s_VM_%zu");
+	static constexpr TCHAR Format[] = TEXT("%s_VM_%u");
 	const FString VMDebugName = FString::Printf(Format, *GetName(), uint32(++NumVMs));
 	URigVM* VM = NewObject<URigVM>(GetTransientPackage(), *VMDebugName);
 
diff --git a/Engine/Plugins/Animation/ControlRig/Source/ControlRig/Private/ControlRigComponent.cpp b/Engine/Plugins/Animation/ControlRig/Source/ControlRig/Private/ControlRigComponent.cpp
--- a/Engine/Plugins/Animation/ControlRig/Source/ControlRig/Private/ControlRigComponent.cpp
+++ b/Engine/Plugins/Animation/ControlRig/Source/ControlRig/Private/ControlRigComponent.cpp
@@ -1924,7 +1924,7 @@ bool UControlRigComponent::EnsureCalledOutsideOfBracket(const TCHAR* InCallingFu
 			}
 			else
 			{
-				ReportError(FString::Printf(TEXT("Cannot be called during the PreConstructionEvent - use ConstructionScript instead."), InCallingFunctionName));
+				ReportError(TEXT("Cannot be called during the PreConstructionEvent - use ConstructionScript instead."));
 				return false;
 			}
 		}
@@ -1938,7 +1938,7 @@ bool UControlRigComponent::EnsureCalledOutsideOfBracket(const TCHAR* InCallingFu
 			}
 			else
 			{
-				ReportError(FString::Printf(TEXT("Cannot be called during the PostConstructionEvent - use ConstructionScript instead."), InCallingFunctionName));
+				ReportError(TEXT("Cannot be called during the PostConstructionEvent - use ConstructionScript instead."));
 				return false;
 			}
 		}
@@ -1952,7 +1952,7 @@ bool UControlRigComponent::EnsureCalledOutsideOfBracket(const TCHAR* InCallingFu
 			}
 			else
 			{
-				ReportError(FString::Printf(TEXT("Cannot be called during the InitEvent - use ConstructionScript instead."), InCallingFunctionName));
+				ReportError(TEXT("Cannot be called during the InitEvent - use ConstructionScript instead."));
 				return false;
 			}
 		}
@@ -1966,7 +1966,7 @@ bool UControlRigComponent::EnsureCalledOutsideOfBracket(const TCHAR* InCallingFu
 			}
 			else
 			{
-				ReportError(FString::Printf(TEXT("Cannot be called during the ForwardsSolveEvent - use ConstructionScript instead."), InCallingFunctionName));
+				ReportError(TEXT("Cannot be called during the ForwardsSolveEvent - use ConstructionScript instead."));
 				return false;
 			}
 		}
diff --git a/Engine/Plugins/Animation/ControlRig/Source/ControlRig/Private/Rigs/RigHierarchy.cpp b/Engine/Plugins/Animation/ControlRig/Source/ControlRig/Private/Rigs/RigHierarchy.cpp
--- a/Engine/Plugins/Animation/ControlRig/Source/ControlRig/Private/Rigs/RigHierarchy.cpp
+++ b/Engine/Plugins/Animation/ControlRig/Source/ControlRig/Private/Rigs/RigHierarchy.cpp
@@ -5362,7 +5362,7 @@ void URigHierarchy::CleanupInvalidCaches()
 	{
 		// Give the validation hierarchy a unique name for debug purposes and to avoid possible memory reuse
 		static TAtomic<uint32> HierarchyNameIndex = 0;
-		static constexpr TCHAR Format[] = TEXT("CacheValidationHierarchy_%zu");
+		static constexpr TCHAR Format[] = TEXT("CacheValidationHierarchy_%u");
 		const FString ValidationHierarchyName = FString::Printf(Format, uint32(++HierarchyNameIndex));
 		
 		HierarchyForCacheValidation = NewObject<URigHierarchy>(this, *ValidationHierarchyName, RF_Transient);
diff --git a/Engine/Plugins/Animation/ControlRig/Source/ControlRig/Private/Units/Execution/RigUnit_DynamicHierarchy.cpp b/Engine/Plugins/Animation/ControlRig/Source/ControlRig/Private/Units/Execution/RigUnit_DynamicHierarchy.cpp
--- a/Engine/Plugins/Animation/ControlRig/Source/ControlRig/Private/Units/Execution/RigUnit_DynamicHierarchy.cpp
+++ b/Engine/Plugins/Animation/ControlRig/Source/ControlRig/Private/Units/Execution/RigUnit_DynamicHierarchy.cpp
@@ -38,7 +38,7 @@ bool FRigUnit_DynamicHierarchyBase::IsValidToRunInContext(
 		if(OutErrorMessage)
 		{
 			static constexpr TCHAR ErrorMessageFormat[] = TEXT("Node has hit the Procedural Element Limit. Check the Class Settings under Hierarchy.");
-			*OutErrorMessage = FString::Printf(ErrorMessageFormat, *FRigUnit_PrepareForExecution::EventName.ToString());
+			*OutErrorMessage = ErrorMessageFormat;
 		}
 		return false;
 	}
diff --git a/Engine/Plugins/Animation/ControlRig/Source/ControlRigEditor/Private/ControlRigDrawingDetails.cpp b/Engine/Plugins/Animation/ControlRig/Source/ControlRigEditor/Private/ControlRigDrawingDetails.cpp
--- a/Engine/Plugins/Animation/ControlRig/Source/ControlRigEditor/Private/ControlRigDrawingDetails.cpp
+++ b/Engine/Plugins/Animation/ControlRig/Source/ControlRigEditor/Private/ControlRigDrawingDetails.cpp
@@ -281,7 +281,7 @@ void FControlRigDrawContainerDetails::ImportCurvesFromFBX(const FString& InFileP
 				}
 				else
 				{
-					Instruction.Name = *FString::Printf(TEXT("%s_Line%d"), FbxCurveNode->GetName(), LineIndex);
+					Instruction.Name = *FString::Printf(TEXT("%hs_Line%d"), FbxCurveNode->GetName(), LineIndex);
 				}
 				Instruction.Positions.Reset();
 
diff --git a/Engine/Plugins/Bridge/Source/Bridge/Private/UI/BrowserBinding.cpp b/Engine/Plugins/Bridge/Source/Bridge/Private/UI/BrowserBinding.cpp
--- a/Engine/Plugins/Bridge/Source/Bridge/Private/UI/BrowserBinding.cpp
+++ b/Engine/Plugins/Bridge/Source/Bridge/Private/UI/BrowserBinding.cpp
@@ -83,11 +83,11 @@ TSharedPtr<SWidget> FAssetDragDropCustomOp::GetDefaultDecorator() const
 
     if (Count > 1)
     {
-       PopupWebBrowser->LoadString(FString::Printf(TEXT("<!DOCTYPE html><html lang=\"en\"> <head> <meta charset=\"UTF-8\"/> <meta name=\"viewport\" content=\"width=device-width, initial-scale=1.0\"/> <style>*{padding: 0px;}body{padding: 0px; margin: 0px;}#container{display: flex; position: relative; width: 100%; height: 100%; min-width: 120px; min-height: 120px; background: #202020; justify-content: center; align-items: center;}#full-image{max-width: 110px; max-height: 110px; display: block; font-size: 0;}#number-circle{position: absolute; border-radius: 50%; width: 18px; height: 18px; padding: 4px; background: #fff; color: #666; text-align: center; font: 12px Arial, sans-serif; box-shadow: 1px 1px 1px #888888; opacity: 0.5;}</style> </head> <body> <div id=\"container\"> <img id=\"full-image\" src=\"%s\"/> <div id=\"number-circle\">+%d</div></div></body></html>"), *ImageUrl, Count-1), TEXT(""));
+       PopupWebBrowser->LoadString(FString::Printf(TEXT("<!DOCTYPE html><html lang=\"en\"> <head> <meta charset=\"UTF-8\"/> <meta name=\"viewport\" content=\"width=device-width, initial-scale=1.0\"/> <style>*{padding: 0px;}body{padding: 0px; margin: 0px;}#container{display: flex; position: relative; width: 100%%; height: 100%%; min-width: 120px; min-height: 120px; background: #202020; justify-content: center; align-items: center;}#full-image{max-width: 110px; max-height: 110px; display: block; font-size: 0;}#number-circle{position: absolute; border-radius: 50%%; width: 18px; height: 18px; padding: 4px; background: #fff; color: #666; text-align: center; font: 12px Arial, sans-serif; box-shadow: 1px 1px 1px #888888; opacity: 0.5;}</style> </head> <body> <div id=\"container\"> <img id=\"full-image\" src=\"%s\"/> <div id=\"number-circle\">+%d</div></div></body></html>"), *ImageUrl, Count-1), TEXT(""));
     }
     else
     {
-       PopupWebBrowser->LoadString(FString::Printf(TEXT("<!DOCTYPE html><html lang=\"en\"> <head> <meta charset=\"UTF-8\"/> <meta name=\"viewport\" content=\"width=device-width, initial-scale=1.0\"/> <style>*{padding: 0px;}body{padding: 0px; margin: 0px;}#container{display: flex; position: relative; width: 100%; height: 100%; min-width: 120px; min-height: 120px; background: #202020; justify-content: center; align-items: center;}#full-image{max-width: 110px; max-height: 110px; display: block; font-size: 0;}#number-circle{position: absolute; border-radius: 50%; width: 18px; height: 18px; padding: 4px; background: #fff; color: #666; text-align: center; font: 16px Arial, sans-serif; box-shadow: 1px 1px 1px #888888; opacity: 0.5;}</style> </head> <body> <div id=\"container\"> <img id=\"full-image\" src=\"%s\"/></div></body></html>"), *ImageUrl), TEXT(""));
+       PopupWebBrowser->LoadString(FString::Printf(TEXT("<!DOCTYPE html><html lang=\"en\"> <head> <meta charset=\"UTF-8\"/> <meta name=\"viewport\" content=\"width=device-width, initial-scale=1.0\"/> <style>*{padding: 0px;}body{padding: 0px; margin: 0px;}#container{display: flex; position: relative; width: 100%%; height: 100%%; min-width: 120px; min-height: 120px; background: #202020; justify-content: center; align-items: center;}#full-image{max-width: 110px; max-height: 110px; display: block; font-size: 0;}#number-circle{position: absolute; border-radius: 50%%; width: 18px; height: 18px; padding: 4px; background: #fff; color: #666; text-align: center; font: 16px Arial, sans-serif; box-shadow: 1px 1px 1px #888888; opacity: 0.5;}</style> </head> <body> <div id=\"container\"> <img id=\"full-image\" src=\"%s\"/></div></body></html>"), *ImageUrl), TEXT(""));
     }
 
 	return SNew(SBox)
diff --git a/Engine/Plugins/Experimental/AVCodecs/NVCodecs/Source/NVENC/Private/Video/Encoders/VideoEncoderNVENC.cpp b/Engine/Plugins/Experimental/AVCodecs/NVCodecs/Source/NVENC/Private/Video/Encoders/VideoEncoderNVENC.cpp
--- a/Engine/Plugins/Experimental/AVCodecs/NVCodecs/Source/NVENC/Private/Video/Encoders/VideoEncoderNVENC.cpp
+++ b/Engine/Plugins/Experimental/AVCodecs/NVCodecs/Source/NVENC/Private/Video/Encoders/VideoEncoderNVENC.cpp
@@ -386,11 +386,11 @@ FAVResult FEncoderNVENC::CreateD3D11Device(TSharedRef<FAVDevice> const& InDevice
 	HRESULT Result = InDevice->GetContext<FVideoContextD3D11>()->Device->QueryInterface(__uuidof(IDXGIDevice), (void**)DXGIDevice.GetInitReference());
 	if (Result != S_OK)
 	{
-		return FAVResult(EAVResult::Fatal, FString::Printf(TEXT("ID3D11Device::QueryInterface() failed 0x%X - %s."), Result, *AVGetComErrorDescription(Result)), TEXT("D3D11"));
+		return FAVResult(EAVResult::Fatal, FString::Printf(TEXT("ID3D11Device::QueryInterface() failed 0x%X - %s."), (uint32)Result, *AVGetComErrorDescription(Result)), TEXT("D3D11"));
 	}
 	else if ((Result = DXGIDevice->GetAdapter(Adapter.GetInitReference())) != S_OK)
 	{
-		return FAVResult(EAVResult::Fatal, FString::Printf(TEXT("DXGIDevice::GetAdapter() failed 0x%X - %s."), Result, *AVGetComErrorDescription(Result)), TEXT("D3D11"));
+		return FAVResult(EAVResult::Fatal, FString::Printf(TEXT("DXGIDevice::GetAdapter() failed 0x%X - %s."), (uint32)Result, *AVGetComErrorDescription(Result)), TEXT("D3D11"));
 	}
 
 	uint32 DeviceFlags = 0;
@@ -410,7 +410,7 @@ FAVResult FEncoderNVENC::CreateD3D11Device(TSharedRef<FAVDevice> const& InDevice
 			 OutEncoderDeviceContext.GetInitReference()))
 		!= S_OK)
 	{
-		return FAVResult(EAVResult::Fatal, FString::Printf(TEXT("D3D11CreateDevice() failed 0x%X - %s."), Result, *AVGetComErrorDescription(Result)), TEXT("D3D11"));
+		return FAVResult(EAVResult::Fatal, FString::Printf(TEXT("D3D11CreateDevice() failed 0x%X - %s."), (uint32)Result, *AVGetComErrorDescription(Result)), TEXT("D3D11"));
 	}
 
 	return FAVResult(EAVResult::Success, TEXT("Created D3D11 device for NVENC."), TEXT("D3D11"));
diff --git a/Engine/Plugins/Experimental/AutomationUtils/Source/AutomationUtils/Private/AutomationUtilsBlueprintLibrary.cpp b/Engine/Plugins/Experimental/AutomationUtils/Source/AutomationUtils/Private/AutomationUtilsBlueprintLibrary.cpp
--- a/Engine/Plugins/Experimental/AutomationUtils/Source/AutomationUtils/Private/AutomationUtilsBlueprintLibrary.cpp
+++ b/Engine/Plugins/Experimental/AutomationUtils/Source/AutomationUtils/Private/AutomationUtilsBlueprintLibrary.cpp
@@ -257,7 +257,7 @@ void FAutomationUtilsGameplayAutomationScreenshotInstance::HandleScreenshotData(
 	TArray64<uint8> CompressedBitmap;
 	FImageUtils::PNGCompressImageArray(InSizeX, InSizeY, TArrayView64<const FColor>(InImageData.GetData(), InImageData.Num()), CompressedBitmap);
 	FFileHelper::SaveArrayToFile(CompressedBitmap, *DeterminedPath);
-	GLog->Log(FString::Printf(TEXT("Saved %d bytes of screenshot image to %s"), CompressedBitmap.Num(), *DeterminedPath));
+	GLog->Log(FString::Printf(TEXT("Saved %" INT64_FMT " bytes of screenshot image to %s"), CompressedBitmap.Num(), *DeterminedPath));
 
 	//Save Metadata Json
 	FString MetadataPath = FPaths::ChangeExtension(DeterminedPath, TEXT("json"));
diff --git a/Engine/Plugins/Experimental/ChaosVehiclesPlugin/Source/ChaosVehicles/Private/ChaosVehicleManagerAsyncCallback.cpp b/Engine/Plugins/Experimental/ChaosVehiclesPlugin/Source/ChaosVehicles/Private/ChaosVehicleManagerAsyncCallback.cpp
--- a/Engine/Plugins/Experimental/ChaosVehiclesPlugin/Source/ChaosVehicles/Private/ChaosVehicleManagerAsyncCallback.cpp
+++ b/Engine/Plugins/Experimental/ChaosVehiclesPlugin/Source/ChaosVehicles/Private/ChaosVehicleManagerAsyncCallback.cpp
@@ -269,7 +269,7 @@ void FNetworkVehicleInputs::DecayData(float DecayAmount)
 
 const FString FNetworkVehicleInputs::DebugData()
 {
-	return FString::Printf(TEXT("FNetworkVehicleInputs | Throttle = %f Brake = %f Roll = %f Pitch = %f Yaw = %f Steering = %f Handbrake = %f Gear = %d | VehicleInputs size = %d | ControlInputs size = %d | NetworkInputs = %d"),
+	return FString::Printf(TEXT("FNetworkVehicleInputs | Throttle = %f Brake = %f Roll = %f Pitch = %f Yaw = %f Steering = %f Handbrake = %f Gear = %d | VehicleInputs size = %zd | ControlInputs size = %zd | NetworkInputs = %zd"),
 			VehicleInputs.ThrottleInput, VehicleInputs.BrakeInput, VehicleInputs.RollInput, VehicleInputs.PitchInput,
 			VehicleInputs.YawInput, VehicleInputs.SteeringInput, VehicleInputs.HandbrakeInput, TransmissionTargetGear, sizeof(FVehicleInputs) * 8, sizeof(FControlInputs) * 8, sizeof(FNetworkVehicleInputs) * 8);
 }
diff --git a/Engine/Plugins/Experimental/PythonScriptPlugin/Source/PythonScriptPlugin/Private/PyConversion.cpp b/Engine/Plugins/Experimental/PythonScriptPlugin/Source/PythonScriptPlugin/Private/PyConversion.cpp
--- a/Engine/Plugins/Experimental/PythonScriptPlugin/Source/PythonScriptPlugin/Private/PyConversion.cpp
+++ b/Engine/Plugins/Experimental/PythonScriptPlugin/Source/PythonScriptPlugin/Private/PyConversion.cpp
@@ -647,7 +647,7 @@ FPyConversionResult PythonizeEnumEntry(const int64 Val, const UEnum* EnumType, P
 		}
 	}
 
-	PYCONVERSION_RETURN(FPyConversionResult::Failure(), TEXT("PythonizeEnumEntry"), *FString::Printf(TEXT("Cannot pythonize '%d' (int64) as '%s'"), Val, *PyUtil::GetFriendlyTypename(PyEnumType)));
+	PYCONVERSION_RETURN(FPyConversionResult::Failure(), TEXT("PythonizeEnumEntry"), *FString::Printf(TEXT("Cannot pythonize '%" INT64_FMT "' (int64) as '%s'"), Val, *PyUtil::GetFriendlyTypename(PyEnumType)));
 }
 
 PyObject* PythonizeEnumEntry(const int64 Val, const UEnum* EnumType, const ESetErrorState SetErrorState)
diff --git a/Engine/Plugins/Experimental/PythonScriptPlugin/Source/PythonScriptPlugin/Private/PyEngine.cpp b/Engine/Plugins/Experimental/PythonScriptPlugin/Source/PythonScriptPlugin/Private/PyEngine.cpp
--- a/Engine/Plugins/Experimental/PythonScriptPlugin/Source/PythonScriptPlugin/Private/PyEngine.cpp
+++ b/Engine/Plugins/Experimental/PythonScriptPlugin/Source/PythonScriptPlugin/Private/PyEngine.cpp
@@ -42,7 +42,7 @@ PyTypeObject InitializePyActorIteratorType(const char* InTypeName, const char* I
 			}
 			if (!IterWorld)
 			{
-				PyUtil::SetPythonError(PyExc_TypeError, InSelf, *FString::Printf(TEXT("'world' cannot be 'None'"), *PyUtil::GetFriendlyTypename(PyWorldObj)));
+				PyUtil::SetPythonError(PyExc_TypeError, InSelf, TEXT("'world' cannot be 'None'"));
 				return -1;
 			}
 
@@ -128,7 +128,7 @@ PyObject* GetBlueprintGeneratedTypes(PyObject* InSelf, PyObject* InArgs)
 							FString& AssetToGenerate = AssetsToGenerate.AddDefaulted_GetRef();
 							if (!PyConversion::Nativize(ValueItem, AssetToGenerate))
 							{
-								PyUtil::SetPythonError(PyExc_TypeError, TEXT("get_blueprint_generated_types"), *FString::Printf(TEXT("Cannot convert argument %d (%s) at index %d to 'string'"), ArgIndex, *PyUtil::GetFriendlyTypename(PyArg), SequenceIndex));
+								PyUtil::SetPythonError(PyExc_TypeError, TEXT("get_blueprint_generated_types"), *FString::Printf(TEXT("Cannot convert argument %zd (%s) at index %zd to 'string'"), ArgIndex, *PyUtil::GetFriendlyTypename(PyArg), SequenceIndex));
 								return nullptr;
 							}
 						}
@@ -139,7 +139,7 @@ PyObject* GetBlueprintGeneratedTypes(PyObject* InSelf, PyObject* InArgs)
 					FString& AssetToGenerate = AssetsToGenerate.AddDefaulted_GetRef();
 					if (!PyConversion::Nativize(PyArg, AssetToGenerate))
 					{
-						PyUtil::SetPythonError(PyExc_TypeError, TEXT("get_blueprint_generated_types"), *FString::Printf(TEXT("Cannot convert argument %d (%s) to 'string'"), ArgIndex, *PyUtil::GetFriendlyTypename(PyArg)));
+						PyUtil::SetPythonError(PyExc_TypeError, TEXT("get_blueprint_generated_types"), *FString::Printf(TEXT("Cannot convert argument %zd (%s) to 'string'"), ArgIndex, *PyUtil::GetFriendlyTypename(PyArg)));
 						return nullptr;
 					}
 				}
diff --git a/Engine/Plugins/Experimental/PythonScriptPlugin/Source/PythonScriptPlugin/Private/PyGenUtil.cpp b/Engine/Plugins/Experimental/PythonScriptPlugin/Source/PythonScriptPlugin/Private/PyGenUtil.cpp
--- a/Engine/Plugins/Experimental/PythonScriptPlugin/Source/PythonScriptPlugin/Private/PyGenUtil.cpp
+++ b/Engine/Plugins/Experimental/PythonScriptPlugin/Source/PythonScriptPlugin/Private/PyGenUtil.cpp
@@ -2268,8 +2268,8 @@ void PythonizeValueImpl(const FProperty* InProp, const void* InPropValue, const
 			else // Use the path string as default value.
 			{
 				OutPythonDefaultValue += bUseStrictTyping
-					? FString::Printf(TEXT("%sFieldPath(\"%s\")"), UnrealNamespace, *GetPropertyPythonName(FieldPathProp), *FieldPathStrValue)
-					: FString::Printf(TEXT("FieldPath(\"%s\")"), *GetPropertyPythonName(FieldPathProp), *FieldPathStrValue);
+					? FString::Printf(TEXT("%sFieldPath(\"%s\")"), UnrealNamespace, *FieldPathStrValue)
+					: FString::Printf(TEXT("FieldPath(\"%s\")"), *FieldPathStrValue);
 			}
 		}
 		else if (const FObjectPropertyBase* ObjProp = CastField<const FObjectPropertyBase>(InProp))
diff --git a/Engine/Plugins/Experimental/PythonScriptPlugin/Source/PythonScriptPlugin/Private/PyUtil.cpp b/Engine/Plugins/Experimental/PythonScriptPlugin/Source/PythonScriptPlugin/Private/PyUtil.cpp
--- a/Engine/Plugins/Experimental/PythonScriptPlugin/Source/PythonScriptPlugin/Private/PyUtil.cpp
+++ b/Engine/Plugins/Experimental/PythonScriptPlugin/Source/PythonScriptPlugin/Private/PyUtil.cpp
@@ -664,7 +664,7 @@ int ValidateContainerLenParam(PyObject* InPyObj, int32 &OutLen, const char* InPy
 
 	if (OutLen < 0)
 	{
-		SetPythonError(PyExc_Exception, InErrorCtxt, *FString::Printf(TEXT("'len' must be positive"), UTF8_TO_TCHAR(InPythonArgName)));
+		SetPythonError(PyExc_Exception, InErrorCtxt, *FString::Printf(TEXT("'%s' must be positive"), UTF8_TO_TCHAR(InPythonArgName)));
 		return -1;
 	}
 
@@ -675,7 +675,7 @@ int ValidateContainerIndexParam(const Py_ssize_t InIndex, const Py_ssize_t InLen
 {
 	if (InIndex < 0 || InIndex >= InLen)
 	{
-		SetPythonError(PyExc_IndexError, InErrorCtxt, *FString::Printf(TEXT("Index %d is out-of-bounds (len: %d) for property '%s' (%s)"), InIndex, InLen, *InProp->GetName(), *InProp->GetClass()->GetName()));
+		SetPythonError(PyExc_IndexError, InErrorCtxt, *FString::Printf(TEXT("Index %zd is out-of-bounds (len: %zd) for property '%s' (%s)"), InIndex, InLen, *InProp->GetName(), *InProp->GetClass()->GetName()));
 		return -1;
 	}
 
diff --git a/Engine/Plugins/Experimental/PythonScriptPlugin/Source/PythonScriptPlugin/Private/PyWrapperArray.cpp b/Engine/Plugins/Experimental/PythonScriptPlugin/Source/PythonScriptPlugin/Private/PyWrapperArray.cpp
--- a/Engine/Plugins/Experimental/PythonScriptPlugin/Source/PythonScriptPlugin/Private/PyWrapperArray.cpp
+++ b/Engine/Plugins/Experimental/PythonScriptPlugin/Source/PythonScriptPlugin/Private/PyWrapperArray.cpp
@@ -425,7 +425,7 @@ PyObject* FPyWrapperArray::GetItem(FPyWrapperArray* InSelf, Py_ssize_t InIndex)
 	PyObject* PyItemObj = nullptr;
 	if (!PyConversion::PythonizeProperty(InSelf->ArrayProp->Inner, SelfScriptArrayHelper.GetRawPtr(ResolvedIndex), PyItemObj))
 	{
-		PyUtil::SetPythonError(PyExc_TypeError, InSelf, *FString::Printf(TEXT("Failed to convert element property '%s' (%s) at index %d"), *InSelf->ArrayProp->Inner->GetName(), *InSelf->ArrayProp->Inner->GetClass()->GetName(), ResolvedIndex));
+		PyUtil::SetPythonError(PyExc_TypeError, InSelf, *FString::Printf(TEXT("Failed to convert element property '%s' (%s) at index %zd"), *InSelf->ArrayProp->Inner->GetName(), *InSelf->ArrayProp->Inner->GetClass()->GetName(), ResolvedIndex));
 		return nullptr;
 	}
 	return PyItemObj;
@@ -450,7 +450,7 @@ int FPyWrapperArray::SetItem(FPyWrapperArray* InSelf, Py_ssize_t InIndex, PyObje
 
 	if (!PyConversion::NativizeProperty(InValue, InSelf->ArrayProp->Inner, SelfScriptArrayHelper.GetRawPtr(ResolvedIndex)))
 	{
-		PyUtil::SetPythonError(PyExc_TypeError, InSelf, *FString::Printf(TEXT("Failed to convert element property '%s' (%s) at index %d"), *InSelf->ArrayProp->Inner->GetName(), *InSelf->ArrayProp->Inner->GetClass()->GetName(), ResolvedIndex));
+		PyUtil::SetPythonError(PyExc_TypeError, InSelf, *FString::Printf(TEXT("Failed to convert element property '%s' (%s) at index %zd"), *InSelf->ArrayProp->Inner->GetName(), *InSelf->ArrayProp->Inner->GetClass()->GetName(), ResolvedIndex));
 		return -1;
 	}
 
@@ -718,7 +718,7 @@ PyObject* FPyWrapperArray::Pop(FPyWrapperArray* InSelf, Py_ssize_t InIndex)
 	PyObject* PyReturnValue = nullptr;
 	if (!PyConversion::PythonizeProperty(InSelf->ArrayProp->Inner, SelfScriptArrayHelper.GetRawPtr(ResolvedIndex), PyReturnValue))
 	{
-		PyUtil::SetPythonError(PyExc_TypeError, InSelf, *FString::Printf(TEXT("Failed to convert element property '%s' (%s) at index %d"), *InSelf->ArrayProp->Inner->GetName(), *InSelf->ArrayProp->Inner->GetClass()->GetName(), ResolvedIndex));
+		PyUtil::SetPythonError(PyExc_TypeError, InSelf, *FString::Printf(TEXT("Failed to convert element property '%s' (%s) at index %zd"), *InSelf->ArrayProp->Inner->GetName(), *InSelf->ArrayProp->Inner->GetClass()->GetName(), ResolvedIndex));
 		return nullptr;
 	}
 
diff --git a/Engine/Plugins/Experimental/PythonScriptPlugin/Source/PythonScriptPlugin/Private/PyWrapperFixedArray.cpp b/Engine/Plugins/Experimental/PythonScriptPlugin/Source/PythonScriptPlugin/Private/PyWrapperFixedArray.cpp
--- a/Engine/Plugins/Experimental/PythonScriptPlugin/Source/PythonScriptPlugin/Private/PyWrapperFixedArray.cpp
+++ b/Engine/Plugins/Experimental/PythonScriptPlugin/Source/PythonScriptPlugin/Private/PyWrapperFixedArray.cpp
@@ -395,7 +395,7 @@ PyObject* FPyWrapperFixedArray::GetItem(FPyWrapperFixedArray* InSelf, Py_ssize_t
 	PyObject* PyItemObj = nullptr;
 	if (!PyConversion::PythonizeProperty_Direct(InSelf->ArrayProp, GetItemPtr(InSelf, ResolvedIndex), PyItemObj))
 	{
-		PyUtil::SetPythonError(PyExc_TypeError, InSelf, *FString::Printf(TEXT("Failed to convert property '%s' (%s) at index %d"), *InSelf->ArrayProp->GetName(), *InSelf->ArrayProp->GetClass()->GetName(), ResolvedIndex));
+		PyUtil::SetPythonError(PyExc_TypeError, InSelf, *FString::Printf(TEXT("Failed to convert property '%s' (%s) at index %zd"), *InSelf->ArrayProp->GetName(), *InSelf->ArrayProp->GetClass()->GetName(), ResolvedIndex));
 		return nullptr;
 	}
 	return PyItemObj;
@@ -417,7 +417,7 @@ int FPyWrapperFixedArray::SetItem(FPyWrapperFixedArray* InSelf, Py_ssize_t InInd
 
 	if (!PyConversion::NativizeProperty_Direct(InValue, InSelf->ArrayProp, GetItemPtr(InSelf, ResolvedIndex)))
 	{
-		PyUtil::SetPythonError(PyExc_TypeError, InSelf, *FString::Printf(TEXT("Failed to convert property '%s' (%s) at index %d"), *InSelf->ArrayProp->GetName(), *InSelf->ArrayProp->GetClass()->GetName(), ResolvedIndex));
+		PyUtil::SetPythonError(PyExc_TypeError, InSelf, *FString::Printf(TEXT("Failed to convert property '%s' (%s) at index %zd"), *InSelf->ArrayProp->GetName(), *InSelf->ArrayProp->GetClass()->GetName(), ResolvedIndex));
 		return -1;
 	}
 
diff --git a/Engine/Plugins/Experimental/PythonScriptPlugin/Source/PythonScriptPlugin/Private/PyWrapperMap.cpp b/Engine/Plugins/Experimental/PythonScriptPlugin/Source/PythonScriptPlugin/Private/PyWrapperMap.cpp
--- a/Engine/Plugins/Experimental/PythonScriptPlugin/Source/PythonScriptPlugin/Private/PyWrapperMap.cpp
+++ b/Engine/Plugins/Experimental/PythonScriptPlugin/Source/PythonScriptPlugin/Private/PyWrapperMap.cpp
@@ -761,14 +761,14 @@ FPyWrapperMap* FPyWrapperMap::CastPyObject(PyObject* InPyObject, PyTypeObject* I
 						FPyObjectPtr PairSequence = FPyObjectPtr::StealReference(PySequence_Fast(PairItem, ""));
 						if (!PairSequence)
 						{
-							PyUtil::SetPythonError(PyExc_TypeError, NewMap.Get(), *FString::Printf(TEXT("Failed to convert item at index %d to a sequence"), SequenceIndex));
+							PyUtil::SetPythonError(PyExc_TypeError, NewMap.Get(), *FString::Printf(TEXT("Failed to convert item at index %zd to a sequence"), SequenceIndex));
 							return nullptr;
 						}
 
 						const Py_ssize_t PairLen = PySequence_Fast_GET_SIZE(PairSequence.Get());
 						if (PairLen != 2)
 						{
-							PyUtil::SetPythonError(PyExc_TypeError, NewMap.Get(), *FString::Printf(TEXT("Failed to convert item at index %d as it was not a pair (len != 2)"), SequenceIndex));
+							PyUtil::SetPythonError(PyExc_TypeError, NewMap.Get(), *FString::Printf(TEXT("Failed to convert item at index %zd as it was not a pair (len != 2)"), SequenceIndex));
 							return nullptr;
 						}
 
diff --git a/Engine/Plugins/Experimental/PythonScriptPlugin/Source/PythonScriptPlugin/Private/PyWrapperText.cpp b/Engine/Plugins/Experimental/PythonScriptPlugin/Source/PythonScriptPlugin/Private/PyWrapperText.cpp
--- a/Engine/Plugins/Experimental/PythonScriptPlugin/Source/PythonScriptPlugin/Private/PyWrapperText.cpp
+++ b/Engine/Plugins/Experimental/PythonScriptPlugin/Source/PythonScriptPlugin/Private/PyWrapperText.cpp
@@ -88,12 +88,12 @@ bool ExtractFormatArguments(FPyWrapperText* InSelf, PyObject* InObj, const int32
 					FFormatArgumentData& FormatArg = InOutFormatArgs.AddDefaulted_GetRef();
 					if (!ExtractFormatArgumentKey(InSelf, KeyItem, FormatArg))
 					{
-						PyUtil::SetPythonError(PyExc_TypeError, InSelf, *FString::Printf(TEXT("Cannot convert format argument %d (%s) at index %d"), InArgIndex, *PyUtil::GetFriendlyTypename(InObj), SequenceIndex));
+						PyUtil::SetPythonError(PyExc_TypeError, InSelf, *FString::Printf(TEXT("Cannot convert format argument %d (%s) at index %zd"), InArgIndex, *PyUtil::GetFriendlyTypename(InObj), SequenceIndex));
 						return false;
 					}
 					if (!ExtractFormatArgumentValue(InSelf, ValueItem, FormatArg))
 					{
-						PyUtil::SetPythonError(PyExc_TypeError, InSelf, *FString::Printf(TEXT("Cannot convert format argument %d (%s) with key '%s' at index %d"), InArgIndex, *PyUtil::GetFriendlyTypename(InObj), *FormatArg.ArgumentName, SequenceIndex));
+						PyUtil::SetPythonError(PyExc_TypeError, InSelf, *FString::Printf(TEXT("Cannot convert format argument %d (%s) with key '%s' at index %zd"), InArgIndex, *PyUtil::GetFriendlyTypename(InObj), *FormatArg.ArgumentName, SequenceIndex));
 						return false;
 					}
 				}
@@ -113,7 +113,7 @@ bool ExtractFormatArguments(FPyWrapperText* InSelf, PyObject* InObj, const int32
 					FormatArg.ArgumentName = FString::FromInt(InArgIndex);
 					if (!ExtractFormatArgumentValue(InSelf, ValueItem, FormatArg))
 					{
-						PyUtil::SetPythonError(PyExc_TypeError, InSelf, *FString::Printf(TEXT("Cannot convert format argument %d (%s) at index %d"), InArgIndex, *PyUtil::GetFriendlyTypename(InObj), SequenceIndex));
+						PyUtil::SetPythonError(PyExc_TypeError, InSelf, *FString::Printf(TEXT("Cannot convert format argument %d (%s) at index %zd"), InArgIndex, *PyUtil::GetFriendlyTypename(InObj), SequenceIndex));
 						return false;
 					}
 				}
diff --git a/Engine/Plugins/Experimental/RemoteSession/Source/RemoteSession/Private/RemoteSessionRole.cpp b/Engine/Plugins/Experimental/RemoteSession/Source/RemoteSession/Private/RemoteSessionRole.cpp
--- a/Engine/Plugins/Experimental/RemoteSession/Source/RemoteSession/Private/RemoteSessionRole.cpp
+++ b/Engine/Plugins/Experimental/RemoteSession/Source/RemoteSession/Private/RemoteSessionRole.cpp
@@ -123,7 +123,7 @@ void FRemoteSessionRole::Tick(float DeltaTime)
 		}
 		else
 		{
-			CloseWithError(FString::Printf(TEXT("State change failed! Closing connection"), CurrentState, OldState));
+			CloseWithError(TEXT("State change failed! Closing connection"));
 		}
 	}
 
diff --git a/Engine/Plugins/FX/Niagara/Source/Niagara/Private/NiagaraEmitter.cpp b/Engine/Plugins/FX/Niagara/Source/Niagara/Private/NiagaraEmitter.cpp
--- a/Engine/Plugins/FX/Niagara/Source/Niagara/Private/NiagaraEmitter.cpp
+++ b/Engine/Plugins/FX/Niagara/Source/Niagara/Private/NiagaraEmitter.cpp
@@ -1763,7 +1763,7 @@ void FVersionedNiagaraEmitterData::CacheFromCompiledData(const FNiagaraDataSetCo
 				if (ScriptCBufferSize > RHIMaxCBufferSize)
 				{
 #if !(UE_BUILD_SHIPPING || UE_BUILD_TEST)
-					GEngine->AddOnScreenDebugMessage(uint64(this), 1.f, FColor::Red, *FString::Printf(TEXT("GPU Simulation(%s) is disabled due to using too much constant buffer space (%d/%d)."), GetDebugSimName(), ScriptCBufferSize, RHIMaxCBufferSize));
+					GEngine->AddOnScreenDebugMessage(uint64(this), 1.f, FColor::Red, *FString::Printf(TEXT("GPU Simulation(%s) is disabled due to using too much constant buffer space (%" UINT64_FMT "/%" UINT64_FMT ")."), GetDebugSimName(), ScriptCBufferSize, RHIMaxCBufferSize));
 #endif
 					bIsAllowedToExecute = false;
 				}
diff --git a/Engine/Plugins/FX/Niagara/Source/Niagara/Private/NiagaraPerfBaseline.cpp b/Engine/Plugins/FX/Niagara/Source/Niagara/Private/NiagaraPerfBaseline.cpp
--- a/Engine/Plugins/FX/Niagara/Source/Niagara/Private/NiagaraPerfBaseline.cpp
+++ b/Engine/Plugins/FX/Niagara/Source/Niagara/Private/NiagaraPerfBaseline.cpp
@@ -1214,7 +1214,7 @@ void FNiagaraPerfBaselineHandler::Tick(UWorld* World, float DeltaSeconds)
 					{
 						BaselineGenerationState = EBaselineGenState::Generating;
 
-						FString Message = FString::Printf(TEXT("Generating Niagara Perf Baselines..."), WorldTimeToGenerate - World->GetTimeSeconds());
+						FString Message = TEXT("Generating Niagara Perf Baselines...");
 						GEngine->AddOnScreenDebugMessage(GenBaselineMessageID, 4.0f, FColor::White, Message);
 					}
 					else
diff --git a/Engine/Plugins/FX/Niagara/Source/NiagaraEditorWidgets/Private/NiagaraEditorWidgetsUtilities.cpp b/Engine/Plugins/FX/Niagara/Source/NiagaraEditorWidgets/Private/NiagaraEditorWidgetsUtilities.cpp
--- a/Engine/Plugins/FX/Niagara/Source/NiagaraEditorWidgets/Private/NiagaraEditorWidgetsUtilities.cpp
+++ b/Engine/Plugins/FX/Niagara/Source/NiagaraEditorWidgets/Private/NiagaraEditorWidgetsUtilities.cpp
@@ -564,7 +564,7 @@ bool FNiagaraStackEditorWidgetsUtilities::HandleDropForStackEntry(const FDragDro
 
 FString FNiagaraStackEditorWidgetsUtilities::StackEntryToStringForListDebug(UNiagaraStackEntry* StackEntry)
 {
-	return FString::Printf(TEXT("0x%08x - %s - %s"), StackEntry, *StackEntry->GetClass()->GetName(), *StackEntry->GetDisplayName().ToString());
+	return FString::Printf(TEXT("0x%08" UPTRINT_x_FMT " - %s - %s"), (UPTRINT)StackEntry, *StackEntry->GetClass()->GetName(), *StackEntry->GetDisplayName().ToString());
 }
 
 #undef LOCTEXT_NAMESPACE
diff --git a/Engine/Plugins/Importers/USDImporter/Source/USDStageEditorViewModels/Private/USDObjectFieldViewModel.cpp b/Engine/Plugins/Importers/USDImporter/Source/USDStageEditorViewModels/Private/USDObjectFieldViewModel.cpp
--- a/Engine/Plugins/Importers/USDImporter/Source/USDStageEditorViewModels/Private/USDObjectFieldViewModel.cpp
+++ b/Engine/Plugins/Importers/USDImporter/Source/USDStageEditorViewModels/Private/USDObjectFieldViewModel.cpp
@@ -419,7 +419,7 @@ void FUsdObjectFieldsViewModel::Refresh(const UE::FUsdStageWeak& InUsdStage, con
 				// Just show arrays as readonly strings for now
 				if (VtValue.IsArrayValued())
 				{
-					FString Stringified = FString::Printf(TEXT("%d elements: "), VtValue.GetArraySize());
+					FString Stringified = FString::Printf(TEXT("%zu elements: "), VtValue.GetArraySize());
 
 					// This array it's too large to even stringify fast enough, so for now just show the element count
 					if (VtValue.GetArraySize() > 5000)
diff --git a/Engine/Plugins/Interchange/Runtime/Source/Parsers/GLTFCore/Private/GLTF/ExtensionsHandler.cpp b/Engine/Plugins/Interchange/Runtime/Source/Parsers/GLTFCore/Private/GLTF/ExtensionsHandler.cpp
--- a/Engine/Plugins/Interchange/Runtime/Source/Parsers/GLTFCore/Private/GLTF/ExtensionsHandler.cpp
+++ b/Engine/Plugins/Interchange/Runtime/Source/Parsers/GLTFCore/Private/GLTF/ExtensionsHandler.cpp
@@ -400,7 +400,7 @@ namespace GLTF
 				default:
 					if (!ensure(false))
 					{
-						Messages.Emplace(RuntimeWarningSeverity(), FString::Printf(TEXT("Material.Extension not supported: %s"), *ToString(Extension)));
+						Messages.Emplace(RuntimeWarningSeverity(), FString::Printf(TEXT("Material.Extension not supported: %s"), ToString(Extension)));
 					}
 					break;
 			}
@@ -561,7 +561,7 @@ namespace GLTF
 				default:
 					if (!ensure(false))
 					{
-						Messages.Emplace(RuntimeWarningSeverity(), FString::Printf(TEXT("Primitive.Extension not supported: %s"), *ToString(Extension)));
+						Messages.Emplace(RuntimeWarningSeverity(), FString::Printf(TEXT("Primitive.Extension not supported: %s"), ToCStr(ToString(Extension))));
 					}
 				break;
 			}
diff --git a/Engine/Plugins/Media/ElectraPlayer/Source/ElectraPlayerRuntime/Private/Runtime/Player/DASH/ManifestDASH.cpp b/Engine/Plugins/Media/ElectraPlayer/Source/ElectraPlayerRuntime/Private/Runtime/Player/DASH/ManifestDASH.cpp
--- a/Engine/Plugins/Media/ElectraPlayer/Source/ElectraPlayerRuntime/Private/Runtime/Player/DASH/ManifestDASH.cpp
+++ b/Engine/Plugins/Media/ElectraPlayer/Source/ElectraPlayerRuntime/Private/Runtime/Player/DASH/ManifestDASH.cpp
@@ -3337,7 +3337,7 @@ FManifestDASHInternal::FRepresentation::ESearchResult FManifestDASHInternal::FRe
 				if (!bWarnedAboutInconsistentNumbering)
 				{
 					bWarnedAboutInconsistentNumbering = true;
-					LogMessage(InPlayerSessionServices, IInfoLog::ELevel::Warning, FString::Printf(TEXT("Representation \"%s\" <SegmentTimeline> 'n' value %lld is not the expected %d. This may cause playback issues"), *MPDRepresentation->GetID(), (long long int)CurrentN, (long long int)PreviousN+1));
+					LogMessage(InPlayerSessionServices, IInfoLog::ELevel::Warning, FString::Printf(TEXT("Representation \"%s\" <SegmentTimeline> 'n' value %lld is not the expected %lld. This may cause playback issues"), *MPDRepresentation->GetID(), (long long int)CurrentN, (long long int)PreviousN+1));
 				}
 			}
 
diff --git a/Engine/Plugins/Media/ElectraPlayer/Source/ElectraPlayerRuntime/Private/Runtime/Player/DASH/StreamReaderDASH.cpp b/Engine/Plugins/Media/ElectraPlayer/Source/ElectraPlayerRuntime/Private/Runtime/Player/DASH/StreamReaderDASH.cpp
--- a/Engine/Plugins/Media/ElectraPlayer/Source/ElectraPlayerRuntime/Private/Runtime/Player/DASH/StreamReaderDASH.cpp
+++ b/Engine/Plugins/Media/ElectraPlayer/Source/ElectraPlayerRuntime/Private/Runtime/Player/DASH/StreamReaderDASH.cpp
@@ -752,7 +752,7 @@ FErrorDetail FStreamReaderDASH::FStreamHandler::GetInitSegment(TSharedPtrTS<cons
 					Request->ConnectionInfo = *ci;
 				}
 				StreamSelector->ReportDownloadEnd(ds);
-				return CreateError(FString::Printf(TEXT("Track preparation of init segment \"%s\" failed. %d"), *LoadReq->URL, *parseError.GetMessage()), INTERNAL_ERROR_INIT_SEGMENT_PARSE_ERROR);
+				return CreateError(FString::Printf(TEXT("Track preparation of init segment \"%s\" failed. %s"), *LoadReq->URL, *parseError.GetMessage()), INTERNAL_ERROR_INIT_SEGMENT_PARSE_ERROR);
 			}
 		}
 		else
diff --git a/Engine/Plugins/Media/MediaCompositing/Source/MediaCompositingEditor/Private/Sequencer/MovieSceneMediaPlayerSectionRecorder.cpp b/Engine/Plugins/Media/MediaCompositing/Source/MediaCompositingEditor/Private/Sequencer/MovieSceneMediaPlayerSectionRecorder.cpp
--- a/Engine/Plugins/Media/MediaCompositing/Source/MediaCompositingEditor/Private/Sequencer/MovieSceneMediaPlayerSectionRecorder.cpp
+++ b/Engine/Plugins/Media/MediaCompositing/Source/MediaCompositingEditor/Private/Sequencer/MovieSceneMediaPlayerSectionRecorder.cpp
@@ -142,7 +142,7 @@ void FMovieSceneMediaPlayerSectionRecorder::StartPlayerRecording(float CurrentTi
 	FString RecordeDataPath;
 	if (RecordingSettings.bRecordMediaFrame)
 	{
-		RecordeDataPath = FPaths::Combine(FPackageName::LongPackageNameToFilename(MediaSourceBasePackageName) + FString::Printf(TEXT("%_%016u"), FDateTime::Now().GetTicks()));
+		RecordeDataPath = FPaths::Combine(FPackageName::LongPackageNameToFilename(MediaSourceBasePackageName) + FString::Printf(TEXT("%%_%016" INT64_FMT), FDateTime::Now().GetTicks()));
 		FString RecordedBaseName = FPaths::Combine(RecordeDataPath, RecordingSettings.BaseFilename);
 		FMediaRecorder::FMediaRecorderData RecorderData = FMediaRecorder::FMediaRecorderData(ObjectToRecord->GetPlayerFacade(), RecordedBaseName);
 		RecorderData.CompressionQuality = RecordingSettings.CompressionQuality;
diff --git a/Engine/Plugins/Media/WmfMedia/Source/WmfMedia/Private/Player/WmfMediaTracks.cpp b/Engine/Plugins/Media/WmfMedia/Source/WmfMedia/Private/Player/WmfMediaTracks.cpp
--- a/Engine/Plugins/Media/WmfMedia/Source/WmfMedia/Private/Player/WmfMediaTracks.cpp
+++ b/Engine/Plugins/Media/WmfMedia/Source/WmfMedia/Private/Player/WmfMediaTracks.cpp
@@ -1512,7 +1512,7 @@ bool FWmfMediaTracks::AddStreamToTracks(uint32 StreamIndex, bool IsVideoDevice,
 
 	for (DWORD TypeIndex = 0; TypeIndex < NumMediaTypes; ++TypeIndex)
 	{
-		OutInfo += FString::Printf(TEXT("\tFormat %i\n"), TypeIndex);
+		OutInfo += FString::Printf(TEXT("\tFormat %u\n"), (uint32)TypeIndex);
 
 		// get media type
 		TComPtr<IMFMediaType> MediaType;
diff --git a/Engine/Plugins/Media/WmfMedia/Source/WmfMedia/Private/Wmf/WmfMediaUtils.cpp b/Engine/Plugins/Media/WmfMedia/Source/WmfMedia/Private/Wmf/WmfMediaUtils.cpp
--- a/Engine/Plugins/Media/WmfMedia/Source/WmfMedia/Private/Wmf/WmfMediaUtils.cpp
+++ b/Engine/Plugins/Media/WmfMedia/Source/WmfMedia/Private/Wmf/WmfMediaUtils.cpp
@@ -664,7 +664,7 @@ namespace WmfMedia
 			}
 			else if (Guid == MF_MT_INTERLACE_MODE)
 			{
-				Dump += FString::Printf(TEXT("\t%s: %d (%s)\n"), *GuidName, Item.ulVal, *InterlaceModeToString((MFVideoInterlaceMode)Item.ulVal));
+				Dump += FString::Printf(TEXT("\t%s: %lu (%s)\n"), *GuidName, Item.ulVal, *InterlaceModeToString((MFVideoInterlaceMode)Item.ulVal));
 			}
 			else if (Guid == MF_MT_PIXEL_ASPECT_RATIO)
 			{
@@ -684,7 +684,7 @@ namespace WmfMedia
 				{
 					MFVideoArea* Area = (MFVideoArea*)Item.caub.pElems;
 
-					Dump += FString::Printf(TEXT("\t%s: (%f,%f) (%d,%d)\n"),
+					Dump += FString::Printf(TEXT("\t%s: (%f,%f) (%ld,%ld)\n"),
 						*GuidName,
 						Area->OffsetX.value + (static_cast<float>(Area->OffsetX.fract) / 65536.0f),
 						Area->OffsetY.value + (static_cast<float>(Area->OffsetY.fract) / 65536.0f),
@@ -697,11 +697,11 @@ namespace WmfMedia
 				switch (Item.vt)
 				{
 				case VT_UI4:
-					Dump += FString::Printf(TEXT("\t%s: %d\n"), *GuidName, Item.ulVal);
+					Dump += FString::Printf(TEXT("\t%s: %lu\n"), *GuidName, Item.ulVal);
 					break;
 
 				case VT_UI8:
-					Dump += FString::Printf(TEXT("\t%s: %ll\n"), *GuidName, Item.uhVal.QuadPart);
+					Dump += FString::Printf(TEXT("\t%s: %llu\n"), *GuidName, Item.uhVal.QuadPart);
 					break;
 
 				case VT_R8:
@@ -1004,17 +1004,17 @@ namespace WmfMedia
 	FString GuidToString(const GUID& Guid)
 	{
 		return FString::Printf(TEXT("%08x-%04x-%04x-%02x%02x%02x%02x%02x%02x%02x%02x"),
-			Guid.Data1,
-			Guid.Data2,
-			Guid.Data3,
-			Guid.Data4[0],
-			Guid.Data4[1],
-			Guid.Data4[2],
-			Guid.Data4[3],
-			Guid.Data4[4],
-			Guid.Data4[5],
-			Guid.Data4[6],
-			Guid.Data4[7]
+			(uint32)Guid.Data1,
+			(uint32)Guid.Data2,
+			(uint32)Guid.Data3,
+			(uint32)Guid.Data4[0],
+			(uint32)Guid.Data4[1],
+			(uint32)Guid.Data4[2],
+			(uint32)Guid.Data4[3],
+			(uint32)Guid.Data4[4],
+			(uint32)Guid.Data4[5],
+			(uint32)Guid.Data4[6],
+			(uint32)Guid.Data4[7]
 		);
 	}
 
@@ -1434,7 +1434,7 @@ namespace WmfMedia
 
 		if (BufferLength == 0)
 		{
-			return FString::Printf(TEXT("0x%08x"), Result);
+			return FString::Printf(TEXT("0x%08x"), (uint32)Result);
 		}
 
 		// remove line break
diff --git a/Engine/Plugins/Online/OnlineSubsystemSteam/Source/Private/OnlineIdentityInterfaceSteam.cpp b/Engine/Plugins/Online/OnlineSubsystemSteam/Source/Private/OnlineIdentityInterfaceSteam.cpp
--- a/Engine/Plugins/Online/OnlineSubsystemSteam/Source/Private/OnlineIdentityInterfaceSteam.cpp
+++ b/Engine/Plugins/Online/OnlineSubsystemSteam/Source/Private/OnlineIdentityInterfaceSteam.cpp
@@ -293,7 +293,7 @@ void FOnlineIdentitySteam::GetLinkedAccountAuthToken(int32 LocalUserNum, const F
 
 			if (RemoteServiceIdentity.IsEmpty())
 			{
-				UE_LOG_ONLINE_IDENTITY(Warning, TEXT("FOnlineIdentitySteam::GetLinkedAccountAuthToken DefaultRemoteServiceIdentity not set"), *TokenType);
+				UE_LOG_ONLINE_IDENTITY(Warning, TEXT("FOnlineIdentitySteam::GetLinkedAccountAuthToken DefaultRemoteServiceIdentity not set"));
 				Delegate.ExecuteIfBound(LocalUserNum, false, FExternalAuthToken());
 				return;
 			}
diff --git a/Engine/Plugins/Online/OnlineSubsystemSteam/Source/Private/OnlineSessionInterfaceSteam.cpp b/Engine/Plugins/Online/OnlineSubsystemSteam/Source/Private/OnlineSessionInterfaceSteam.cpp
--- a/Engine/Plugins/Online/OnlineSubsystemSteam/Source/Private/OnlineSessionInterfaceSteam.cpp
+++ b/Engine/Plugins/Online/OnlineSubsystemSteam/Source/Private/OnlineSessionInterfaceSteam.cpp
@@ -896,7 +896,7 @@ bool FOnlineSessionSteam::JoinSession(int32 PlayerNum, FName SessionName, const
 			}
 			else
 			{
-				UE_LOG_ONLINE_SESSION(Warning, TEXT("Invalid session info on search result"), *SessionName.ToString());
+				UE_LOG_ONLINE_SESSION(Warning, TEXT("Invalid session info on search result"));
 			}
 		}
 		else
diff --git a/Engine/Plugins/Online/OnlineSubsystemSteam/Source/Private/OnlineUserCloudInterfaceSteam.cpp b/Engine/Plugins/Online/OnlineSubsystemSteam/Source/Private/OnlineUserCloudInterfaceSteam.cpp
--- a/Engine/Plugins/Online/OnlineSubsystemSteam/Source/Private/OnlineUserCloudInterfaceSteam.cpp
+++ b/Engine/Plugins/Online/OnlineSubsystemSteam/Source/Private/OnlineUserCloudInterfaceSteam.cpp
@@ -356,7 +356,7 @@ void FOnlineUserCloudSteam::DumpCloudState(const FUniqueNetId& UserId)
 		TotalBytes = TotalAvailable = 0;
 	}
 
-	UE_LOG_ONLINE_CLOUD(Verbose, TEXT("Steam Disk Quota: %d / %d"), TotalAvailable, TotalBytes);
+	UE_LOG_ONLINE_CLOUD(Verbose, TEXT("Steam Disk Quota: %" UINT64_FMT " / %" UINT64_FMT), TotalAvailable, TotalBytes);
 	UE_LOG_ONLINE_CLOUD(Verbose, TEXT("Game does %shave cloud storage enabled."), SteamRemoteStorage()->IsCloudEnabledForApp() ? TEXT("") : TEXT("NOT "));
 	UE_LOG_ONLINE_CLOUD(Verbose, TEXT("User does %shave cloud storage enabled."), SteamRemoteStorage()->IsCloudEnabledForAccount() ? TEXT("") : TEXT("NOT "));
 }
diff --git a/Engine/Plugins/Online/OnlineSubsystemSteam/Source/Private/SocketSubsystemSteam.cpp b/Engine/Plugins/Online/OnlineSubsystemSteam/Source/Private/SocketSubsystemSteam.cpp
--- a/Engine/Plugins/Online/OnlineSubsystemSteam/Source/Private/SocketSubsystemSteam.cpp
+++ b/Engine/Plugins/Online/OnlineSubsystemSteam/Source/Private/SocketSubsystemSteam.cpp
@@ -833,7 +833,7 @@ void FSocketSubsystemSteam::DumpAllOpenSteamSessions()
 	for (TUniqueNetIdMap<FSteamP2PConnectionInfo>::TConstIterator It(AcceptedConnections); It; ++It)
 	{
 		UE_LOG_ONLINE(Verbose, TEXT("- Connection %s"), *It->Key->ToDebugString());
-		UE_LOG_ONLINE(Verbose, TEXT("--  Last Update Time: %d"), It->Value.LastReceivedTime);
+		UE_LOG_ONLINE(Verbose, TEXT("--  Last Update Time: %f"), It->Value.LastReceivedTime);
 		FString ConnectedChannels(TEXT(""));
 		for (int32 i = 0; i < It->Value.ConnectedChannels.Num(); ++i)
 		{
diff --git a/Engine/Plugins/Online/OnlineSubsystemUtils/Source/OnlineSubsystemUtils/Private/IpNetDriver.cpp b/Engine/Plugins/Online/OnlineSubsystemUtils/Source/OnlineSubsystemUtils/Private/IpNetDriver.cpp
--- a/Engine/Plugins/Online/OnlineSubsystemUtils/Source/OnlineSubsystemUtils/Private/IpNetDriver.cpp
+++ b/Engine/Plugins/Online/OnlineSubsystemUtils/Source/OnlineSubsystemUtils/Private/IpNetDriver.cpp
@@ -935,7 +935,7 @@ bool UIpNetDriver::InitBase( bool bInitAsClient, FNetworkNotify* InNotify, const
 	if (CVarNetIpNetDriverUseReceiveThread.GetValueOnAnyThread() != 0 && SocketSubsystem->IsSocketWaitSupported())
 	{
 		SocketReceiveThreadRunnable = MakeUnique<FReceiveThreadRunnable>(this);
-		SocketReceiveThread.Reset(FRunnableThread::Create(SocketReceiveThreadRunnable.Get(), *FString::Printf(TEXT("IpNetDriver Receive Thread"), *NetDriverName.ToString())));
+		SocketReceiveThread.Reset(FRunnableThread::Create(SocketReceiveThreadRunnable.Get(), *FString::Printf(TEXT("IpNetDriver Receive Thread: %s"), *NetDriverName.ToString())));
 	}
 
 	SetSocketAndLocalAddress(Resolver->GetFirstSocket());
diff --git a/Engine/Plugins/Online/OnlineSubsystemUtils/Source/OnlineSubsystemUtils/Private/Tests/TestSharingInterface.cpp b/Engine/Plugins/Online/OnlineSubsystemUtils/Source/OnlineSubsystemUtils/Private/Tests/TestSharingInterface.cpp
--- a/Engine/Plugins/Online/OnlineSubsystemUtils/Source/OnlineSubsystemUtils/Private/Tests/TestSharingInterface.cpp
+++ b/Engine/Plugins/Online/OnlineSubsystemUtils/Source/OnlineSubsystemUtils/Private/Tests/TestSharingInterface.cpp
@@ -191,7 +191,7 @@ void FTestSharingInterface::OnNewsFeedRead(int32 LocalPlayer, bool bWasSuccessfu
 	SharingInterface->ClearOnReadNewsFeedCompleteDelegate_Handle(LocalPlayer, DelegateHandle);
 	if( ++ResponsesReceived == MAX_LOCAL_PLAYERS )
 	{
-		UE_LOG_ONLINE_SHARING(Display, TEXT("FTestSharingInterface TESTS COMPLETED"), LocalPlayer, bWasSuccessful);
+		UE_LOG_ONLINE_SHARING(Display, TEXT("FTestSharingInterface TESTS COMPLETED"));
 		delete this;
 	}
 }
diff --git a/Engine/Plugins/Online/OnlineSubsystemUtils/Source/OnlineSubsystemUtils/Private/Tests/TestStatsInterface.cpp b/Engine/Plugins/Online/OnlineSubsystemUtils/Source/OnlineSubsystemUtils/Private/Tests/TestStatsInterface.cpp
--- a/Engine/Plugins/Online/OnlineSubsystemUtils/Source/OnlineSubsystemUtils/Private/Tests/TestStatsInterface.cpp
+++ b/Engine/Plugins/Online/OnlineSubsystemUtils/Source/OnlineSubsystemUtils/Private/Tests/TestStatsInterface.cpp
@@ -62,7 +62,7 @@ void FTestStatsInterface::Test(UWorld* InWorld, const TCHAR* Cmd)
 
 	if (StatsToRead.Num() == 0)
 	{
-		UE_LOG_ONLINE_STATS(Warning, TEXT("[FTestStatsInterface::Test] No stat names were specified. Please add the stat names you want to use for the test."), *Subsystem);
+		UE_LOG_ONLINE_STATS(Warning, TEXT("[FTestStatsInterface::Test] No stat names were specified. Please add the stat names you want to use for the test."));
 
 		bOverallSuccess = false;
 		return;
diff --git a/Engine/Plugins/PCG/Source/PCG/Private/Graph/PCGPinDependencyExpression.cpp b/Engine/Plugins/PCG/Source/PCG/Private/Graph/PCGPinDependencyExpression.cpp
--- a/Engine/Plugins/PCG/Source/PCG/Private/Graph/PCGPinDependencyExpression.cpp
+++ b/Engine/Plugins/PCG/Source/PCG/Private/Graph/PCGPinDependencyExpression.cpp
@@ -135,7 +135,7 @@ FString FPCGPinDependencyExpression::ToString() const
 				// Decompose pin ID into task ID and pin index which is friendly to read.
 				const FPCGTaskId NodeId = PCGPinIdHelpers::GetNodeIdFromPinId(Expression[I]);
 				const uint64 PinIndex = PCGPinIdHelpers::GetPinIndexFromPinId(Expression[I]);
-				ExpressionString += FString::Printf(TEXT("%u_%u"), NodeId, PinIndex);
+				ExpressionString += FString::Printf(TEXT("%" UINT64_FMT "_%" UINT64_FMT), NodeId, PinIndex);
 			}
 
 			InDisjunction = true;
diff --git a/Engine/Plugins/PCG/Source/PCG/Private/Utils/PCGGraphExecutionLogging.cpp b/Engine/Plugins/PCG/Source/PCG/Private/Utils/PCGGraphExecutionLogging.cpp
--- a/Engine/Plugins/PCG/Source/PCG/Private/Utils/PCGGraphExecutionLogging.cpp
+++ b/Engine/Plugins/PCG/Source/PCG/Private/Utils/PCGGraphExecutionLogging.cpp
@@ -77,7 +77,7 @@ namespace PCGGraphExecutionLogging
 				}
 				bFirstInput = false;
 
-				InputString += FString::Printf(TEXT("%u->'%s'"), Input.TaskId, Input.OutPin ? *Input.OutPin->Properties.Label.ToString() : TEXT(""));
+				InputString += FString::Printf(TEXT("%" UINT64_FMT "->'%s'"), Input.TaskId, Input.OutPin ? *Input.OutPin->Properties.Label.ToString() : TEXT(""));
 			}
 
 			return InputString;
@@ -89,7 +89,7 @@ namespace PCGGraphExecutionLogging
 			bool bFirstSuccessor = true;
 			for (const FPCGTaskId& SuccessorId : *SuccessorIds)
 			{
-				SuccessorsString += bFirstSuccessor ? FString::Printf(TEXT("%u"), SuccessorId) : FString::Printf(TEXT(",%u"), SuccessorId);
+				SuccessorsString += bFirstSuccessor ? FString::Printf(TEXT("%" UINT64_FMT), SuccessorId) : FString::Printf(TEXT(",%" UINT64_FMT), SuccessorId);
 				bFirstSuccessor = false;
 			}
 		}
@@ -281,7 +281,7 @@ namespace PCGGraphExecutionLogging
 		{
 			const FPCGTaskId NodeId = PCGPinIdHelpers::GetNodeIdFromPinId(PinId);
 			const uint64 PinIndex = PCGPinIdHelpers::GetPinIndexFromPinId(PinId);
-			PinIdsToDeactivateString += bFirst ? FString::Printf(TEXT("%u_%u"), NodeId, PinIndex) : FString::Printf(TEXT(",%u_%u"), NodeId, PinIndex);
+			PinIdsToDeactivateString += bFirst ? FString::Printf(TEXT("%" UINT64_FMT "_%" UINT64_FMT), NodeId, PinIndex) : FString::Printf(TEXT(",%" UINT64_FMT "_%" UINT64_FMT), NodeId, PinIndex);
 			bFirst = false;
 		}
 
@@ -300,7 +300,7 @@ namespace PCGGraphExecutionLogging
 		UE_LOG(LogPCG, Log, TEXT("         [%s/%s] %s\t\tEXECUTE"),
 			*Task.SourceComponent->GetOwner()->GetName(),
 			Task.SourceComponent->GetGraph() ? *Task.SourceComponent->GetGraph()->GetName() : TEXT("MISSINGGRAPH"),
-			*FString::Printf(TEXT("%u'%s'"), Task.NodeId, Task.Node ? *Task.Node->GetNodeTitle(EPCGNodeTitleType::ListView).ToString() : TEXT("")));
+			*FString::Printf(TEXT("%" UINT64_FMT "'%s'"), Task.NodeId, Task.Node ? *Task.Node->GetNodeTitle(EPCGNodeTitleType::ListView).ToString() : TEXT("")));
 #endif // !(UE_BUILD_SHIPPING || UE_BUILD_TEST) || USE_LOGGING_IN_SHIPPING
 	}
 
@@ -315,7 +315,7 @@ namespace PCGGraphExecutionLogging
 		UE_LOG(LogPCG, Warning, TEXT("[%s/%s] %s\t\tCACHING DISABLED"),
 			*Task.SourceComponent->GetOwner()->GetName(),
 			Task.SourceComponent->GetGraph() ? *Task.SourceComponent->GetGraph()->GetName() : TEXT("MISSINGGRAPH"),
-			*FString::Printf(TEXT("%u'%s'"), Task.NodeId, Task.Node ? *Task.Node->GetNodeTitle(EPCGNodeTitleType::ListView).ToString() : TEXT("")));
+			*FString::Printf(TEXT("%" UINT64_FMT "'%s'"), Task.NodeId, Task.Node ? *Task.Node->GetNodeTitle(EPCGNodeTitleType::ListView).ToString() : TEXT("")));
 #endif // !(UE_BUILD_SHIPPING || UE_BUILD_TEST) || USE_LOGGING_IN_SHIPPING
 	}
 
diff --git a/Engine/Plugins/Runtime/AnimationBudgetAllocator/Source/AnimationBudgetAllocator/Private/AnimationBudgetAllocator.cpp b/Engine/Plugins/Runtime/AnimationBudgetAllocator/Source/AnimationBudgetAllocator/Private/AnimationBudgetAllocator.cpp
--- a/Engine/Plugins/Runtime/AnimationBudgetAllocator/Source/AnimationBudgetAllocator/Private/AnimationBudgetAllocator.cpp
+++ b/Engine/Plugins/Runtime/AnimationBudgetAllocator/Source/AnimationBudgetAllocator/Private/AnimationBudgetAllocator.cpp
@@ -741,7 +741,7 @@ void FAnimationBudgetAllocator::Update(float DeltaSeconds)
 					{
 						if(GAnimationBudgetDebugShowAddresses != 0)
 						{
-							DebugString += FString::Printf(TEXT("0x%llx %d %s %s\n"), &ComponentData, ComponentData.TickRate, ComponentData.bInterpolate ? TEXT("I") : TEXT(" "), ComponentData.bReducedWork ? TEXT("Lo") : TEXT("Hi"));
+							DebugString += FString::Printf(TEXT("0x%" UPTRINT_x_FMT " %d %s %s\n"), (UPTRINT)&ComponentData, ComponentData.TickRate, ComponentData.bInterpolate ? TEXT("I") : TEXT(" "), ComponentData.bReducedWork ? TEXT("Lo") : TEXT("Hi"));
 						}
 						else
 						{
@@ -850,7 +850,7 @@ void FAnimationBudgetAllocator::OnHUDPostRender(AHUD* HUD, UCanvas* Canvas)
 						{
 							if(GAnimationBudgetDebugShowAddresses != 0)
 							{
-								DebugString += FString::Printf(TEXT("0x%llx %d %s %s\n"), &ComponentData, ComponentData.TickRate, ComponentData.bInterpolate ? TEXT("I") : TEXT(" "), ComponentData.bReducedWork ? TEXT("Lo") : TEXT("Hi"));
+								DebugString += FString::Printf(TEXT("0x%" UPTRINT_x_FMT " %d %s %s\n"), (UPTRINT)&ComponentData, ComponentData.TickRate, ComponentData.bInterpolate ? TEXT("I") : TEXT(" "), ComponentData.bReducedWork ? TEXT("Lo") : TEXT("Hi"));
 							}
 							else
 							{
diff --git a/Engine/Plugins/Runtime/GameplayAbilities/Source/GameplayAbilities/Private/Tests/GameplayTagQueryTests.cpp b/Engine/Plugins/Runtime/GameplayAbilities/Source/GameplayAbilities/Private/Tests/GameplayTagQueryTests.cpp
--- a/Engine/Plugins/Runtime/GameplayAbilities/Source/GameplayAbilities/Private/Tests/GameplayTagQueryTests.cpp
+++ b/Engine/Plugins/Runtime/GameplayAbilities/Source/GameplayAbilities/Private/Tests/GameplayTagQueryTests.cpp
@@ -125,8 +125,8 @@ bool FGameplayTagQueryMatchesTagRequirementsTest::RunTest(const FString& Paramet
 	} // for each scenario
 
 	AddInfo(FString::Printf(TEXT(" Ran %d tests against FGameplayTagRequirements and FGameplayTagQuery in %.4lf ms"), NumTestScenarios * NumTestsPerScenario, FPlatformTime::ToMilliseconds64(MetRequirementsCycles + QueryMatchesCycles)));
-	AddInfo(FString::Printf(TEXT("	FGameplayTagRequirements took %ull cycles (%.4lf ms) and gave %d matches"), MetRequirementsCycles, FPlatformTime::ToMilliseconds64(MetRequirementsCycles), NumMetRequirements));
-	AddInfo(FString::Printf(TEXT("	FGameplayTagQuery took %ull cycles (%.4lf ms) and gave %d matches"), QueryMatchesCycles, FPlatformTime::ToMilliseconds64(QueryMatchesCycles), NumQueryMatches));
+	AddInfo(FString::Printf(TEXT("	FGameplayTagRequirements took %" UINT64_FMT " cycles (%.4lf ms) and gave %u matches"), MetRequirementsCycles, FPlatformTime::ToMilliseconds64(MetRequirementsCycles), NumMetRequirements));
+	AddInfo(FString::Printf(TEXT("	FGameplayTagQuery took %" UINT64_FMT " cycles (%.4lf ms) and gave %u matches"), QueryMatchesCycles, FPlatformTime::ToMilliseconds64(QueryMatchesCycles), NumQueryMatches));
 
 	TestEqual(TEXT("FGameplayTagRequirements has same matching rules as FGameplayTagQuery"), NumQueryMatches, NumMetRequirements);
 	return NumMetRequirements == NumQueryMatches;
diff --git a/Engine/Plugins/Runtime/ModelViewViewModel/Source/ModelViewViewModel/Private/View/MVVMViewClass.cpp b/Engine/Plugins/Runtime/ModelViewViewModel/Source/ModelViewViewModel/Private/View/MVVMViewClass.cpp
--- a/Engine/Plugins/Runtime/ModelViewViewModel/Source/ModelViewViewModel/Private/View/MVVMViewClass.cpp
+++ b/Engine/Plugins/Runtime/ModelViewViewModel/Source/ModelViewViewModel/Private/View/MVVMViewClass.cpp
@@ -540,7 +540,7 @@ FString FMVVMViewClass_Binding::ToString(const UMVVMViewClass* ViewClass, FToStr
 	{
 		AddNewLine();
 		StringBuilder << TEXT("Sources: 0x");
-		StringBuilder << FString::Printf(TEXT("%x"), SourceBitField);
+		StringBuilder << FString::Printf(TEXT("%" UINT64_x_FMT), SourceBitField);
 	}
 
 	if (Args.bAddFlags)
diff --git a/Engine/Plugins/Runtime/NetworkPrediction/Source/NetworkPrediction/Public/NetworkPredictionCues.h b/Engine/Plugins/Runtime/NetworkPrediction/Source/NetworkPrediction/Public/NetworkPredictionCues.h
--- a/Engine/Plugins/Runtime/NetworkPrediction/Source/NetworkPrediction/Public/NetworkPredictionCues.h
+++ b/Engine/Plugins/Runtime/NetworkPrediction/Source/NetworkPrediction/Public/NetworkPredictionCues.h
@@ -339,7 +339,7 @@ struct FSavedCue
 
 	FString GetDebugName() const
 	{
-		return FString::Printf(TEXT("[%s 0x%X] @ (Frame %d/%dms)"), *FGlobalCueTypeTable::Get().GetTypeName(ID), (int64)this, Frame, Time);
+		return FString::Printf(TEXT("[%s 0x%" INT64_FMT "] @ (Frame %d/%dms)"), *FGlobalCueTypeTable::Get().GetTypeName(ID), (int64)this, Frame, Time);
 	}
 
 	FNetSimCueTypeId ID = 0;
diff --git a/Engine/Plugins/Runtime/RigVM/Source/RigVM/Public/RigVMTypeUtils.h b/Engine/Plugins/Runtime/RigVM/Source/RigVM/Public/RigVMTypeUtils.h
--- a/Engine/Plugins/Runtime/RigVM/Source/RigVM/Public/RigVMTypeUtils.h
+++ b/Engine/Plugins/Runtime/RigVM/Source/RigVM/Public/RigVMTypeUtils.h
@@ -345,7 +345,7 @@ namespace RigVMTypeUtils
 		{
 			if (InClassArgType == EClassArgType::AsClass)
 			{
-				return FString::Printf(RigVMTypeUtils::TSubclassOfPrefix, Class->GetPrefixCPP(), *Class->GetName());
+				return FString::Printf(RigVMTypeUtils::TSubclassOfTemplate, Class->GetPrefixCPP(), *Class->GetName());
 			}
 			else if (Class->IsChildOf(UInterface::StaticClass()))
 			{
diff --git a/Engine/Plugins/Runtime/RigVM/Source/RigVMDeveloper/Private/RigVMCompiler/RigVMCodeGenerator.cpp b/Engine/Plugins/Runtime/RigVM/Source/RigVMDeveloper/Private/RigVMCompiler/RigVMCodeGenerator.cpp
--- a/Engine/Plugins/Runtime/RigVM/Source/RigVMDeveloper/Private/RigVMCompiler/RigVMCodeGenerator.cpp
+++ b/Engine/Plugins/Runtime/RigVM/Source/RigVMDeveloper/Private/RigVMCompiler/RigVMCodeGenerator.cpp
@@ -927,7 +927,7 @@ FString FRigVMCodeGenerator::DumpHeader(const FRigVMExtendedExecuteContext& Cont
 	Lines.Emplace();
 	
 	Lines.Add(Format(RigVM_UClassDefinitionFormat, *ModuleName.ToUpper(), *ClassName));
-	Lines.Add(Format(RigVM_GetVMHashFormat, FString::Printf(TEXT("%lu"), VM->GetVMHash())));
+	Lines.Add(Format(RigVM_GetVMHashFormat, FString::Printf(TEXT("%u"), VM->GetVMHash())));
 	Lines.Add(Format(RigVM_GetEntryNamesFormat, *FString::Join(FormattedEntries, RigVM_CommaSeparator)));
 	Lines.Emplace();
 	Lines.Add(FString(RigVM_DeclareInitializeFormat));
diff --git a/Engine/Plugins/Runtime/RigVM/Source/RigVMDeveloper/Private/RigVMCompiler/RigVMCompiler.cpp b/Engine/Plugins/Runtime/RigVM/Source/RigVMDeveloper/Private/RigVMCompiler/RigVMCompiler.cpp
--- a/Engine/Plugins/Runtime/RigVM/Source/RigVMDeveloper/Private/RigVMCompiler/RigVMCompiler.cpp
+++ b/Engine/Plugins/Runtime/RigVM/Source/RigVMDeveloper/Private/RigVMCompiler/RigVMCompiler.cpp
@@ -1602,7 +1602,7 @@ bool URigVMCompiler::TraverseExpression(const FRigVMExprAST* InExpr, FRigVMCompi
 			{
 				if(!BlockInfo->ExecuteStateOperand.IsValid())
 				{
-					static constexpr TCHAR Format[] = TEXT("BlockExecuteState_%zu");
+					static constexpr TCHAR Format[] = TEXT("BlockExecuteState_%u");
 					const FString BlockStateName = FString::Printf(Format, NextBlockHash.GetValue());
 					BlockInfo->ExecuteStateOperand = WorkData.AddProperty(
 						ERigVMMemoryType::Work,
diff --git a/Engine/Plugins/Runtime/RigVM/Source/RigVMDeveloper/Private/RigVMModel/RigVMController.cpp b/Engine/Plugins/Runtime/RigVM/Source/RigVMDeveloper/Private/RigVMModel/RigVMController.cpp
--- a/Engine/Plugins/Runtime/RigVM/Source/RigVMDeveloper/Private/RigVMModel/RigVMController.cpp
+++ b/Engine/Plugins/Runtime/RigVM/Source/RigVMDeveloper/Private/RigVMModel/RigVMController.cpp
@@ -3124,7 +3124,7 @@ URigVMRerouteNode* URigVMController::AddRerouteNodeOnPin(const FString& InPinPat
 		// AddRerouteNodeOnPin(const FString& InPinPath, bool bAsInput, const FVector2D& InPosition, const FString& InNodeName, bool bSetupUndoRedo, bool bPrintPythonCommand)
 		
 		RigVMPythonUtils::Print(GetSchema()->GetGraphOuterName(GetGraph()), 
-							FString::Printf(TEXT("blueprint.get_controller_by_name('%s').add_reroute_node_on_pin('%s', %s, %s, %s '%s')"),
+							FString::Printf(TEXT("blueprint.get_controller_by_name('%s').add_reroute_node_on_pin('%s', %s, %s '%s')"),
 											*GraphName,
 											*GetSchema()->GetSanitizedPinPath(InPinPath),
 											(bAsInput) ? TEXT("True") : TEXT("False"),
diff --git a/Engine/Plugins/Tests/EditorTests/Source/EditorTests/Private/UnrealEd/CollisionAutomationTests.cpp b/Engine/Plugins/Tests/EditorTests/Source/EditorTests/Private/UnrealEd/CollisionAutomationTests.cpp
--- a/Engine/Plugins/Tests/EditorTests/Source/EditorTests/Private/UnrealEd/CollisionAutomationTests.cpp
+++ b/Engine/Plugins/Tests/EditorTests/Source/EditorTests/Private/UnrealEd/CollisionAutomationTests.cpp
@@ -115,7 +115,7 @@ namespace CollisionAutomationTests
 		if (Diff > Tolerance)
 		{
 			//UE_LOG(CollisionAutomationTestLog, Log, TEXT("%d:HitResult=(%s)"), iTest+1, *OutHits[iHits].ToString());
-			TestBase->AddError(FString::Printf(TEXT("Test %d:%s %s mismatch. Should be %df but is actually %f."), TestIndex, *TestName, *ParameterName, ExpectedFloat, ResultFloat));
+			TestBase->AddError(FString::Printf(TEXT("Test %d:%s %s mismatch. Should be %f but is actually %f."), TestIndex, *TestName, *ParameterName, ExpectedFloat, ResultFloat));
 		}
 	}
 
diff --git a/Engine/Plugins/Tests/EditorTests/Source/EditorTests/Private/UnrealEd/FbxAutomationTests.cpp b/Engine/Plugins/Tests/EditorTests/Source/EditorTests/Private/UnrealEd/FbxAutomationTests.cpp
--- a/Engine/Plugins/Tests/EditorTests/Source/EditorTests/Private/UnrealEd/FbxAutomationTests.cpp
+++ b/Engine/Plugins/Tests/EditorTests/Source/EditorTests/Private/UnrealEd/FbxAutomationTests.cpp
@@ -476,7 +476,7 @@ bool FFbxImportAssetsAutomationTest::RunTest(const FString& Parameters)
 				FString LodFile = BaseLODFile.Replace(TEXT("_lod00"), *LodIndexString);
 				if (!FPaths::FileExists(LodFile))
 				{
-					ExecutionInfo.AddError(FString::Printf(TEXT("%s: Cannot Add Lod because file %s do not exist on disk!"), *LodFile));
+					ExecutionInfo.AddError(FString::Printf(TEXT("%s: Cannot Add Lod because file %s do not exist on disk!"), *CleanFilename, *LodFile));
 					CurTestSuccessful = false;
 					continue;
 				}
diff --git a/Engine/Plugins/Tests/InterchangeTests/Source/InterchangeTests/Private/ImportTestFunctions/StaticMeshImportTestFunctions.cpp b/Engine/Plugins/Tests/InterchangeTests/Source/InterchangeTests/Private/ImportTestFunctions/StaticMeshImportTestFunctions.cpp
--- a/Engine/Plugins/Tests/InterchangeTests/Source/InterchangeTests/Private/ImportTestFunctions/StaticMeshImportTestFunctions.cpp
+++ b/Engine/Plugins/Tests/InterchangeTests/Source/InterchangeTests/Private/ImportTestFunctions/StaticMeshImportTestFunctions.cpp
@@ -726,7 +726,7 @@ FInterchangeTestFunctionResult UStaticMeshImportTestFunctions::CheckSocketName(U
 	int32 SocketCount = Mesh->Sockets.Num();
 	if (SocketIndex >= SocketCount)
 	{
-		Result.AddError(FString::Printf(TEXT("The imported mesh doesn't contain %s sockets."), SocketIndex));
+		Result.AddError(FString::Printf(TEXT("The imported mesh doesn't contain %d sockets."), SocketIndex));
 	}
 
 	FString SocketName = Mesh->Sockets[SocketIndex]->SocketName.ToString();
@@ -746,7 +746,7 @@ FInterchangeTestFunctionResult UStaticMeshImportTestFunctions::CheckSocketLocati
 	int32 SocketCount = Mesh->Sockets.Num();
 	if (SocketIndex >= SocketCount)
 	{
-		Result.AddError(FString::Printf(TEXT("The imported mesh doesn't contain %s sockets."), SocketIndex));
+		Result.AddError(FString::Printf(TEXT("The imported mesh doesn't contain %d sockets."), SocketIndex));
 	}
 
 	FVector SocketLocation = Mesh->Sockets[SocketIndex]->RelativeLocation;
diff --git a/Engine/Source/Developer/Android/AndroidDeviceDetection/Private/AndroidDeviceDetectionModule.cpp b/Engine/Source/Developer/Android/AndroidDeviceDetection/Private/AndroidDeviceDetectionModule.cpp
--- a/Engine/Source/Developer/Android/AndroidDeviceDetection/Private/AndroidDeviceDetectionModule.cpp
+++ b/Engine/Source/Developer/Android/AndroidDeviceDetection/Private/AndroidDeviceDetectionModule.cpp
@@ -854,7 +854,7 @@ public:
 			DeviceSpecs.AndroidProperties.DeviceBuildNumber = DeviceInfo->BuildNumber;
 			// this is used in the same way as PlatformMemoryBucket..
 			// to establish the nearest GB Android has a different rounding algo (hence 384 used here). See GenericPlatformMemory::GetMemorySizeBucket.
-			DeviceSpecs.AndroidProperties.TotalPhysicalGB = FString::Printf(TEXT("%d"),(((uint64)DeviceInfo->TotalPhysicalKB + 384 * 1024 - 1) / 1024 / 1024));
+			DeviceSpecs.AndroidProperties.TotalPhysicalGB = FString::Printf(TEXT("%" UINT64_FMT),(((uint64)DeviceInfo->TotalPhysicalKB + 384 * 1024 - 1) / 1024 / 1024));
 			
 			DeviceSpecs.AndroidProperties.UsingHoudini = false;
 			DeviceSpecs.AndroidProperties.VulkanAvailable = !(DeviceInfo->VulkanVersion.IsEmpty() || DeviceInfo->VulkanVersion.Contains(TEXT("0.0.0")));
diff --git a/Engine/Source/Developer/DeviceManager/Private/Widgets/Processes/SDeviceProcesses.cpp b/Engine/Source/Developer/DeviceManager/Private/Widgets/Processes/SDeviceProcesses.cpp
--- a/Engine/Source/Developer/DeviceManager/Private/Widgets/Processes/SDeviceProcesses.cpp
+++ b/Engine/Source/Developer/DeviceManager/Private/Widgets/Processes/SDeviceProcesses.cpp
@@ -387,7 +387,7 @@ FReply SDeviceProcesses::HandleTerminateProcessButtonClicked()
 				for (int32 FailedProcessIndex = 0; FailedProcessIndex < FailedProcesses.Num(); ++FailedProcessIndex)
 				{
 					const FTargetDeviceProcessInfo& FailedProcessInfo = FailedProcesses[FailedProcessIndex]->GetProcessInfo();
-					ProcessInfo += FString::Printf(TEXT("%s (PID: %d)\n"), *FailedProcessInfo.Name, FailedProcessInfo.Id);
+					ProcessInfo += FString::Printf(TEXT("%s (PID: %" INT64_FMT ")\n"), *FailedProcessInfo.Name, FailedProcessInfo.Id);
 				}
 
 				const FText ErrorMessage = FText::Format(LOCTEXT("FailedToTerminateProcessesMessage", "The following processes could not be terminated.\nYou may not have the required permissions:\n\n{0}"), FText::FromString(ProcessInfo));
diff --git a/Engine/Source/Developer/FileUtilities/Private/ZipArchiveReader.cpp b/Engine/Source/Developer/FileUtilities/Private/ZipArchiveReader.cpp
--- a/Engine/Source/Developer/FileUtilities/Private/ZipArchiveReader.cpp
+++ b/Engine/Source/Developer/FileUtilities/Private/ZipArchiveReader.cpp
@@ -85,7 +85,7 @@ FZipArchiveReader::FImpl::FImpl(IFileHandle* InFileHandle, FOutputDevice* ErrorH
 		if (ErrorHandler)
 		{
 			ErrorHandler->Log(LogZipArchive.GetCategoryName(), ELogVerbosity::Display,
-				FString::Printf(TEXT("Could not create ZipSourceFunction: %s"), zip_error_strerror(&ZipError)));
+				FString::Printf(TEXT("Could not create ZipSourceFunction: %hs"), zip_error_strerror(&ZipError)));
 		}
 		zip_error_fini(&ZipError);
 		Destruct();
@@ -99,7 +99,7 @@ FZipArchiveReader::FImpl::FImpl(IFileHandle* InFileHandle, FOutputDevice* ErrorH
 		if (ErrorHandler)
 		{
 			ErrorHandler->Log(LogZipArchive.GetCategoryName(), ELogVerbosity::Display,
-				FString::Printf(TEXT("Could not parse zip file: %s"), zip_error_strerror(&ZipError)));
+				FString::Printf(TEXT("Could not parse zip file: %hs"), zip_error_strerror(&ZipError)));
 		}
 		zip_error_fini(&ZipError);
 		Destruct();
@@ -228,7 +228,7 @@ bool FZipArchiveReader::FImpl::TryReadFile(FStringView FileName, TArray<uint8>&
 		if (ErrorHandler)
 		{
 			ErrorHandler->Log(LogZipArchive.GetCategoryName(), ELogVerbosity::Display,
-				FString::Printf(TEXT("Embedded file %s has size " UINT64_FMT " which is too large to store in a TArray."),
+				FString::Printf(TEXT("Embedded file %.*s has size %" UINT64_FMT " which is too large to store in a TArray."),
 					FileName.Len(), FileName.GetData(), ZipFileStat.size));
 		}
 		return false;
diff --git a/Engine/Source/Developer/IoStoreUtilities/Private/IoStoreUtilities.cpp b/Engine/Source/Developer/IoStoreUtilities/Private/IoStoreUtilities.cpp
--- a/Engine/Source/Developer/IoStoreUtilities/Private/IoStoreUtilities.cpp
+++ b/Engine/Source/Developer/IoStoreUtilities/Private/IoStoreUtilities.cpp
@@ -7107,11 +7107,11 @@ namespace DescribeUtils
 				FExportDesc* ExportDesc = ExportByKeyMap.FindRef(Key);
 				if (ExportDesc && bIncludeName)
 				{
-					return FString::Printf(TEXT("0x%llX '%s'"), PackageObjectIndex.Value(), *ExportDesc->FullName.ToString());
+					return FString::Printf(TEXT("0x%" UINT64_X_FMT " '%s'"), PackageObjectIndex.Value(), *ExportDesc->FullName.ToString());
 				}
 				else
 				{
-					return FString::Printf(TEXT("0x%llX"), PackageObjectIndex.Value());
+					return FString::Printf(TEXT("0x%" UINT64_X_FMT), PackageObjectIndex.Value());
 				}
 			}
 			else if (PackageObjectIndex.IsScriptImport())
@@ -7119,20 +7119,20 @@ namespace DescribeUtils
 				const FScriptObjectDesc* ScriptObjectDesc = ScriptObjectByGlobalIdMap.Find(PackageObjectIndex);
 				if (ScriptObjectDesc && bIncludeName)
 				{
-					return FString::Printf(TEXT("0x%llX '%s'"), PackageObjectIndex.Value(), *ScriptObjectDesc->FullName.ToString());
+					return FString::Printf(TEXT("0x%" UINT64_X_FMT " '%s'"), PackageObjectIndex.Value(), *ScriptObjectDesc->FullName.ToString());
 				}
 				else
 				{
-					return FString::Printf(TEXT("0x%llX"), PackageObjectIndex.Value());
+					return FString::Printf(TEXT("0x%" UINT64_X_FMT), PackageObjectIndex.Value());
 				}
 			}
 			else if (PackageObjectIndex.IsExport())
 			{
-				return FString::Printf(TEXT("%d"), PackageObjectIndex.Value());
+				return FString::Printf(TEXT("%" UINT64_X_FMT), PackageObjectIndex.Value());
 			}
 			else
 			{
-				return FString::Printf(TEXT("0x%llX"), PackageObjectIndex.Value());
+				return FString::Printf(TEXT("0x%" UINT64_X_FMT), PackageObjectIndex.Value());
 			}
 		}
 	};
diff --git a/Engine/Source/Developer/PakFileUtilities/Private/PakFileUtilities.cpp b/Engine/Source/Developer/PakFileUtilities/Private/PakFileUtilities.cpp
--- a/Engine/Source/Developer/PakFileUtilities/Private/PakFileUtilities.cpp
+++ b/Engine/Source/Developer/PakFileUtilities/Private/PakFileUtilities.cpp
@@ -811,7 +811,7 @@ FString FCompressedFileBuffer::GetDDCKeyString(const uint8* UncompressedFile, co
 {
 	FString KeyString;
 
-	KeyString += FString::Printf(TEXT("_F:%s_C:%s_B:%d_"), *CompressionFormat.ToString(), *FCompression::GetCompressorDDCSuffix(CompressionFormat), BlockSize);
+	KeyString += FString::Printf(TEXT("_F:%s_C:%s_B:%" INT64_FMT "_"), *CompressionFormat.ToString(), *FCompression::GetCompressorDDCSuffix(CompressionFormat), BlockSize);
 	
 	FSHA1 HashState;
 	HashState.Update(UncompressedFile, UncompressedFileSize);
@@ -3673,16 +3673,16 @@ bool AuditPakFiles( const FString& InputPath, bool bOnlyDeleted, const FString&
 			}
 			else if (Itr.Value.PakPriority == HighestPakPriority)
 			{
-				WriteCSVLine( FString::Printf( TEXT("%s,Fresh,%s,,%d,,%d,%s,%s,%s"), *AssetName, *Itr.Value.PakFilename, Itr.Value.PakPriority, Itr.Value.Size, *AssetPath, *PatchDotChunk, *OpenOrderText ) );
+				WriteCSVLine( FString::Printf( TEXT("%s,Fresh,%s,,%d,,%" INT64_FMT ",%s,%s,%s"), *AssetName, *Itr.Value.PakFilename, Itr.Value.PakPriority, Itr.Value.Size, *AssetPath, *PatchDotChunk, *OpenOrderText ) );
 			}
 			else
 			{
-				WriteCSVLine( FString::Printf( TEXT("%s,Inherited,%s,,%d,,%d,%s,%s,%s"), *AssetName, *Itr.Value.PakFilename, Itr.Value.PakPriority, Itr.Value.Size, *AssetPath, *PatchDotChunk, *OpenOrderText  ) );
+				WriteCSVLine( FString::Printf( TEXT("%s,Inherited,%s,,%d,,%" INT64_FMT ",%s,%s,%s"), *AssetName, *Itr.Value.PakFilename, Itr.Value.PakPriority, Itr.Value.Size, *AssetPath, *PatchDotChunk, *OpenOrderText  ) );
 			}
 		}
 		else if (DeletedRevision->PakPriority == Itr.Value.PakPriority)
 		{
-			WriteCSVLine( FString::Printf( TEXT("%s,Moved,%s,%s,%d,,%d,%s,%s,%s"), *AssetName, *Itr.Value.PakFilename, *DeletedRevision->PakFilename, Itr.Value.PakPriority, Itr.Value.Size, *AssetPath, *PatchDotChunk, *OpenOrderText ) );
+			WriteCSVLine( FString::Printf( TEXT("%s,Moved,%s,%s,%d,,%" INT64_FMT ",%s,%s,%s"), *AssetName, *Itr.Value.PakFilename, *DeletedRevision->PakFilename, Itr.Value.PakPriority, Itr.Value.Size, *AssetPath, *PatchDotChunk, *OpenOrderText ) );
 		}
 		else if (DeletedRevision->PakPriority > Itr.Value.PakPriority)
 		{
@@ -3691,7 +3691,7 @@ bool AuditPakFiles( const FString& InputPath, bool bOnlyDeleted, const FString&
 		}
 		else if (DeletedRevision->PakPriority < Itr.Value.PakPriority)
 		{
-			WriteCSVLine( FString::Printf( TEXT("%s,Restored,%s,%s,%d,%d,%d,%s,%s,%s"), *AssetName, *Itr.Value.PakFilename, *DeletedRevision->PakFilename, Itr.Value.PakPriority, DeletedRevision->PakPriority, Itr.Value.Size, *AssetPath, *PatchDotChunk, *OpenOrderText ) );
+			WriteCSVLine( FString::Printf( TEXT("%s,Restored,%s,%s,%d,%d,%" INT64_FMT ",%s,%s,%s"), *AssetName, *Itr.Value.PakFilename, *DeletedRevision->PakFilename, Itr.Value.PakPriority, DeletedRevision->PakPriority, Itr.Value.Size, *AssetPath, *PatchDotChunk, *OpenOrderText ) );
 		}
 
 		if( bFileExists && bSortByOrdering && bHasOpenOrder )
diff --git a/Engine/Source/Developer/TraceInsights/Private/Insights/NetworkingProfiler/Widgets/SPacketContentView.cpp b/Engine/Source/Developer/TraceInsights/Private/Insights/NetworkingProfiler/Widgets/SPacketContentView.cpp
--- a/Engine/Source/Developer/TraceInsights/Private/Insights/NetworkingProfiler/Widgets/SPacketContentView.cpp
+++ b/Engine/Source/Developer/TraceInsights/Private/Insights/NetworkingProfiler/Widgets/SPacketContentView.cpp
@@ -1349,7 +1349,7 @@ int32 SPacketContentView::OnPaint(const FPaintArgs& Args, const FGeometry& Allot
 		DrawContext.DrawText
 		(
 			DbgX, DbgY,
-			FString::Printf(TEXT("U: %llu ms    D: %llu ms + %llu ms = %llu ms (%d fps)"),
+			FString::Printf(TEXT("U: %" UINT64_FMT " ms    D: %" UINT64_FMT " ms + %" UINT64_FMT " ms = %" UINT64_FMT " ms (%" INT64_FMT " fps)"),
 				AvgUpdateDurationMs, // average duration of UpdateState calls
 				AvgDrawDurationMs, // drawing time
 				AvgOnPaintDurationMs - AvgDrawDurationMs, // other overhead to OnPaint calls
diff --git a/Engine/Source/Developer/TraceInsights/Private/Insights/NetworkingProfiler/Widgets/SPacketView.cpp b/Engine/Source/Developer/TraceInsights/Private/Insights/NetworkingProfiler/Widgets/SPacketView.cpp
--- a/Engine/Source/Developer/TraceInsights/Private/Insights/NetworkingProfiler/Widgets/SPacketView.cpp
+++ b/Engine/Source/Developer/TraceInsights/Private/Insights/NetworkingProfiler/Widgets/SPacketView.cpp
@@ -939,7 +939,7 @@ int32 SPacketView::OnPaint(const FPaintArgs& Args, const FGeometry& AllottedGeom
 		DrawContext.DrawText
 		(
 			DbgX, DbgY,
-			FString::Printf(TEXT("U: %llu ms, D: %llu ms + %llu ms = %llu ms (%d fps)"),
+			FString::Printf(TEXT("U: %" UINT64_FMT " ms, D: %" UINT64_FMT " ms + %" UINT64_FMT " ms = %" UINT64_FMT " ms (%" INT64_FMT " fps)"),
 				AvgUpdateDurationMs, // caching time
 				AvgDrawDurationMs, // drawing time
 				AvgOnPaintDurationMs - AvgDrawDurationMs, // other overhead to OnPaint calls
diff --git a/Engine/Source/Developer/TraceInsights/Private/Insights/TaskGraphProfiler/ViewModels/TaskTimingTrack.cpp b/Engine/Source/Developer/TraceInsights/Private/Insights/TaskGraphProfiler/ViewModels/TaskTimingTrack.cpp
--- a/Engine/Source/Developer/TraceInsights/Private/Insights/TaskGraphProfiler/ViewModels/TaskTimingTrack.cpp
+++ b/Engine/Source/Developer/TraceInsights/Private/Insights/TaskGraphProfiler/ViewModels/TaskTimingTrack.cpp
@@ -753,7 +753,7 @@ void FTaskTimingTrack::BuildDrawState(ITimingEventsTrackDrawStateBuilder& Builde
 			const TraceServices::FTaskInfo* TaskInfo = TasksProvider->TryGetTask(Relation.RelativeId);
 			if (TaskInfo)
 			{
-				Builder.AddEvent(TaskInfo->StartedTimestamp, TaskInfo->FinishedTimestamp, ++Depth, *FString::Printf(TEXT("Prerequisite Task %d Executing"), TaskInfo->Id), 0, FTaskGraphProfilerManager::Get()->GetColorForTaskEventAsPackedARGB(ETaskEventType::PrerequisiteStarted));
+				Builder.AddEvent(TaskInfo->StartedTimestamp, TaskInfo->FinishedTimestamp, ++Depth, *FString::Printf(TEXT("Prerequisite Task %" UINT64_FMT " Executing"), TaskInfo->Id), 0, FTaskGraphProfilerManager::Get()->GetColorForTaskEventAsPackedARGB(ETaskEventType::PrerequisiteStarted));
 			}
 		}
 
@@ -767,7 +767,7 @@ void FTaskTimingTrack::BuildDrawState(ITimingEventsTrackDrawStateBuilder& Builde
 			const TraceServices::FTaskInfo* TaskInfo = TasksProvider->TryGetTask(Relation.RelativeId);
 			if (TaskInfo)
 			{
-				Builder.AddEvent(TaskInfo->StartedTimestamp, TaskInfo->FinishedTimestamp, ++Depth, *FString::Printf(TEXT("Parent Task %d Executing"), TaskInfo->Id), 0, FTaskGraphProfilerManager::Get()->GetColorForTaskEventAsPackedARGB(ETaskEventType::ParentStarted));
+				Builder.AddEvent(TaskInfo->StartedTimestamp, TaskInfo->FinishedTimestamp, ++Depth, *FString::Printf(TEXT("Parent Task %" UINT64_FMT " Executing"), TaskInfo->Id), 0, FTaskGraphProfilerManager::Get()->GetColorForTaskEventAsPackedARGB(ETaskEventType::ParentStarted));
 			}
 		}
 
@@ -781,7 +781,7 @@ void FTaskTimingTrack::BuildDrawState(ITimingEventsTrackDrawStateBuilder& Builde
 			const TraceServices::FTaskInfo* TaskInfo = TasksProvider->TryGetTask(Relation.RelativeId);
 			if (TaskInfo)
 			{
-				Builder.AddEvent(TaskInfo->StartedTimestamp, TaskInfo->FinishedTimestamp, ++Depth, *FString::Printf(TEXT("Nested Task %d Executing"), TaskInfo->Id), 0, FTaskGraphProfilerManager::Get()->GetColorForTaskEventAsPackedARGB(ETaskEventType::NestedStarted));
+				Builder.AddEvent(TaskInfo->StartedTimestamp, TaskInfo->FinishedTimestamp, ++Depth, *FString::Printf(TEXT("Nested Task %" UINT64_FMT " Executing"), TaskInfo->Id), 0, FTaskGraphProfilerManager::Get()->GetColorForTaskEventAsPackedARGB(ETaskEventType::NestedStarted));
 			}
 		}
 
@@ -795,7 +795,7 @@ void FTaskTimingTrack::BuildDrawState(ITimingEventsTrackDrawStateBuilder& Builde
 			const TraceServices::FTaskInfo* TaskInfo = TasksProvider->TryGetTask(Relation.RelativeId);
 			if (TaskInfo)
 			{
-				Builder.AddEvent(TaskInfo->StartedTimestamp, TaskInfo->FinishedTimestamp, ++Depth, *FString::Printf(TEXT("Subsequent Task %d Executing"), TaskInfo->Id), 0, FTaskGraphProfilerManager::Get()->GetColorForTaskEventAsPackedARGB(ETaskEventType::SubsequentStarted));
+				Builder.AddEvent(TaskInfo->StartedTimestamp, TaskInfo->FinishedTimestamp, ++Depth, *FString::Printf(TEXT("Subsequent Task %" UINT64_FMT " Executing"), TaskInfo->Id), 0, FTaskGraphProfilerManager::Get()->GetColorForTaskEventAsPackedARGB(ETaskEventType::SubsequentStarted));
 			}
 		}
 	}
@@ -1042,7 +1042,7 @@ void FTaskTimingTrack::InitTooltip(FTooltipDrawState& InOutTooltip, const ITimin
 		return;
 	}
 
-	InOutTooltip.AddNameValueTextLine(TEXT("Task Id:"), FString::Printf(TEXT("%d"), Task->Id));
+	InOutTooltip.AddNameValueTextLine(TEXT("Task Id:"), FString::Printf(TEXT("%" UINT64_FMT), Task->Id));
 
 	switch (TaskTrackEvent.GetTaskEventType())
 	{
@@ -1057,11 +1057,11 @@ void FTaskTimingTrack::InitTooltip(FTooltipDrawState& InOutTooltip, const ITimin
 		const TCHAR* StartedThreadName = ThreadProvider.GetThreadName(Task->StartedThreadId);
 		if (StartedThreadName)
 		{
-			InOutTooltip.AddNameValueTextLine(TEXT("Executed Thread Id:"), FString::Printf(TEXT("%d (%s)"), Task->StartedThreadId, StartedThreadName));
+			InOutTooltip.AddNameValueTextLine(TEXT("Executed Thread Id:"), FString::Printf(TEXT("%u (%s)"), Task->StartedThreadId, StartedThreadName));
 		}
 		else
 		{
-			InOutTooltip.AddNameValueTextLine(TEXT("Executed Thread Id:"), FString::Printf(TEXT("%d (<unknown>)"), Task->StartedThreadId));
+			InOutTooltip.AddNameValueTextLine(TEXT("Executed Thread Id:"), FString::Printf(TEXT("%u (<unknown>)"), Task->StartedThreadId));
 		}
 		break;
 	}
diff --git a/Engine/Source/Developer/TraceInsights/Private/Insights/TaskGraphProfiler/ViewModels/TaskTrackEvent.cpp b/Engine/Source/Developer/TraceInsights/Private/Insights/TaskGraphProfiler/ViewModels/TaskTrackEvent.cpp
--- a/Engine/Source/Developer/TraceInsights/Private/Insights/TaskGraphProfiler/ViewModels/TaskTrackEvent.cpp
+++ b/Engine/Source/Developer/TraceInsights/Private/Insights/TaskGraphProfiler/ViewModels/TaskTrackEvent.cpp
@@ -95,13 +95,13 @@ FString FTaskTrackEvent::GetEventName() const
 	case ETaskEventType::Completed:
 		return TEXT("Completed");
 	case ETaskEventType::PrerequisiteStarted:
-		return FString::Printf(TEXT("Prerequisite Task %d Executing"), GetTaskId());	
+		return FString::Printf(TEXT("Prerequisite Task %" UINT64_FMT " Executing"), GetTaskId());	
 	case ETaskEventType::ParentStarted:
-		return FString::Printf(TEXT("Parent Task %d Executing"), GetTaskId());
+		return FString::Printf(TEXT("Parent Task %" UINT64_FMT " Executing"), GetTaskId());
 	case ETaskEventType::NestedStarted:
-		return FString::Printf(TEXT("Nested Task %d Executing"), GetTaskId());
+		return FString::Printf(TEXT("Nested Task %" UINT64_FMT " Executing"), GetTaskId());
 	case ETaskEventType::SubsequentStarted:
-		return FString::Printf(TEXT("Subsequent Task %d Executing"), GetTaskId());
+		return FString::Printf(TEXT("Subsequent Task %" UINT64_FMT " Executing"), GetTaskId());
 
 	default:
 		checkf(false, TEXT("Unknown task event type"));
diff --git a/Engine/Source/Developer/TraceInsights/Private/Insights/ViewModels/FrameTimingTrack.cpp b/Engine/Source/Developer/TraceInsights/Private/Insights/ViewModels/FrameTimingTrack.cpp
--- a/Engine/Source/Developer/TraceInsights/Private/Insights/ViewModels/FrameTimingTrack.cpp
+++ b/Engine/Source/Developer/TraceInsights/Private/Insights/ViewModels/FrameTimingTrack.cpp
@@ -545,21 +545,21 @@ bool FFrameTimingTrack::FindFrame(const FTimingEventSearchParameters& InParamete
 
 const FString FFrameTimingTrack::GetShortFrameName(const uint64 InFrameIndex) const
 {
-	return FString::Printf(TEXT("%d"), InFrameIndex);
+	return FString::Printf(TEXT("%" UINT64_FMT), InFrameIndex);
 }
 
 ////////////////////////////////////////////////////////////////////////////////////////////////////
 
 const FString FFrameTimingTrack::GetFrameName(const uint64 InFrameIndex) const
 {
-	return FString::Printf(TEXT("%s Frame %d"), FFrameTrackDrawHelper::FrameTypeToString(FrameType), InFrameIndex);
+	return FString::Printf(TEXT("%s Frame %" UINT64_FMT), FFrameTrackDrawHelper::FrameTypeToString(FrameType), InFrameIndex);
 }
 
 ////////////////////////////////////////////////////////////////////////////////////////////////////
 
 const FString FFrameTimingTrack::GetCompleteFrameName(const uint64 InFrameIndex, const double InFrameDuration) const
 {
-	return FString::Printf(TEXT("%s Frame %d (%s)"), FFrameTrackDrawHelper::FrameTypeToString(FrameType), InFrameIndex, *TimeUtils::FormatTimeAuto(InFrameDuration));
+	return FString::Printf(TEXT("%s Frame %" UINT64_FMT " (%s)"), FFrameTrackDrawHelper::FrameTypeToString(FrameType), InFrameIndex, *TimeUtils::FormatTimeAuto(InFrameDuration));
 }
 
 ////////////////////////////////////////////////////////////////////////////////////////////////////
diff --git a/Engine/Source/Developer/TraceInsights/Private/Insights/Widgets/SFrameTrack.cpp b/Engine/Source/Developer/TraceInsights/Private/Insights/Widgets/SFrameTrack.cpp
--- a/Engine/Source/Developer/TraceInsights/Private/Insights/Widgets/SFrameTrack.cpp
+++ b/Engine/Source/Developer/TraceInsights/Private/Insights/Widgets/SFrameTrack.cpp
@@ -786,7 +786,7 @@ int32 SFrameTrack::OnPaint(const FPaintArgs& Args, const FGeometry& AllottedGeom
 		DrawContext.DrawText
 		(
 			DbgX, DbgY,
-			FString::Printf(TEXT("U: %llu ms, D: %llu ms + %llu ms = %llu ms (%d fps)"),
+			FString::Printf(TEXT("U: %" UINT64_FMT " ms, D: %" UINT64_FMT " ms + %" UINT64_FMT " ms = %" UINT64_FMT " ms (%" INT64_FMT " fps)"),
 				AvgUpdateDurationMs, // caching time
 				AvgDrawDurationMs, // drawing time
 				AvgOnPaintDurationMs - AvgDrawDurationMs, // other overhead to OnPaint calls
diff --git a/Engine/Source/Developer/TraceInsights/Private/Insights/Widgets/STimingView.cpp b/Engine/Source/Developer/TraceInsights/Private/Insights/Widgets/STimingView.cpp
--- a/Engine/Source/Developer/TraceInsights/Private/Insights/Widgets/STimingView.cpp
+++ b/Engine/Source/Developer/TraceInsights/Private/Insights/Widgets/STimingView.cpp
@@ -1586,7 +1586,7 @@ int32 STimingView::OnPaint(const FPaintArgs& Args, const FGeometry& AllottedGeom
 		DrawContext.DrawText
 		(
 			DbgX, DbgY,
-			FString::Printf(TEXT("D: %llu ms + %llu ms + %llu ms + %llu ms = %llu ms | + %llu ms = %llu ms (%d fps)"),
+			FString::Printf(TEXT("D: %" UINT64_FMT " ms + %" UINT64_FMT " ms + %" UINT64_FMT " ms + %" UINT64_FMT " ms = %" UINT64_FMT " ms | + %" UINT64_FMT " ms = %" UINT64_FMT " ms (%" INT64_FMT " fps)"),
 				AvgPreDrawTracksDurationMs, // pre-draw tracks time
 				AvgDrawTracksDurationMs, // draw tracks time
 				AvgPostDrawTracksDurationMs, // post-draw tracks time
diff --git a/Engine/Source/Developer/Virtualization/Private/DumpPackagePayloadInfoCommand.cpp b/Engine/Source/Developer/Virtualization/Private/DumpPackagePayloadInfoCommand.cpp
--- a/Engine/Source/Developer/Virtualization/Private/DumpPackagePayloadInfoCommand.cpp
+++ b/Engine/Source/Developer/Virtualization/Private/DumpPackagePayloadInfoCommand.cpp
@@ -18,7 +18,7 @@ FString BytesToString(int64 SizeInBytes)
 {
 	if (SizeInBytes < (8 *1024))
 	{
-		return FString::Printf(TEXT("%4d bytes"), SizeInBytes);
+		return FString::Printf(TEXT("%4" INT64_FMT " bytes"), SizeInBytes);
 	}
 	else if (SizeInBytes < (1024 * 1024))
 	{
diff --git a/Engine/Source/Developer/Windows/ShaderFormatD3D/Private/D3DShaderCompilerDXC.cpp b/Engine/Source/Developer/Windows/ShaderFormatD3D/Private/D3DShaderCompilerDXC.cpp
--- a/Engine/Source/Developer/Windows/ShaderFormatD3D/Private/D3DShaderCompilerDXC.cpp
+++ b/Engine/Source/Developer/Windows/ShaderFormatD3D/Private/D3DShaderCompilerDXC.cpp
@@ -132,7 +132,7 @@ static void LogFailedHRESULT(const TCHAR* FailedExpressionStr, HRESULT Result)
 {
 	if (Result == E_OUTOFMEMORY)
 	{
-		const FString ErrorReport = FString::Printf(TEXT("%s failed: Result=0x%08x (E_OUTOFMEMORY)"), FailedExpressionStr, Result);
+		const FString ErrorReport = FString::Printf(TEXT("%s failed: Result=0x%08x (E_OUTOFMEMORY)"), FailedExpressionStr, (uint32)Result);
 		FSCWErrorCode::Report(FSCWErrorCode::OutOfMemory, ErrorReport);
 		UE_LOG(LogD3DShaderCompiler, Fatal, TEXT("%s"), *ErrorReport);
 	}
diff --git a/Engine/Source/Editor/BlueprintGraph/Private/EdGraphSchema_K2.cpp b/Engine/Source/Editor/BlueprintGraph/Private/EdGraphSchema_K2.cpp
--- a/Engine/Source/Editor/BlueprintGraph/Private/EdGraphSchema_K2.cpp
+++ b/Engine/Source/Editor/BlueprintGraph/Private/EdGraphSchema_K2.cpp
@@ -4404,7 +4404,7 @@ bool UEdGraphSchema_K2::DefaultValueSimpleValidation(const FEdGraphPinType& PinT
 
 			if (!FPackageName::IsValidObjectPath(NewDefaultValue, &PathReason))
 			{
-				DVSV_RETURN_MSG(FString::Printf(TEXT("Soft Reference '%s' is invalid format for object pin '%s':"), *NewDefaultValue, *PinName.ToString(), *PathReason.ToString()));
+				DVSV_RETURN_MSG(FString::Printf(TEXT("Soft Reference '%s' is invalid format for object pin '%s': %s"), *NewDefaultValue, *PinName.ToString(), *PathReason.ToString()));
 			}
 
 			// Class and IsAsset validation is not foolproof for soft references, skip
diff --git a/Engine/Source/Editor/Cascade/Private/Cascade.cpp b/Engine/Source/Editor/Cascade/Private/Cascade.cpp
--- a/Engine/Source/Editor/Cascade/Private/Cascade.cpp
+++ b/Engine/Source/Editor/Cascade/Private/Cascade.cpp
@@ -4864,7 +4864,7 @@ void ParticleSystem_DumpInfo(UParticleSystem* InParticleSystem)
 						for (int32 ModuleIdx = 0; ModuleIdx < LODLevel->Modules.Num(); ModuleIdx++)
 						{
 							UParticleModule* Module = LODLevel->Modules[ModuleIdx];
-							ModuleDump += FString::Printf(TEXT("0x%08x,"), PTRINT(Module));
+							ModuleDump += FString::Printf(TEXT("0x%08" PTRINT_FMT ","), PTRINT(Module));
 						}
 						UE_LOG(LogCascade, Log, TEXT("%s"), *ModuleDump);
 					}
diff --git a/Engine/Source/Editor/DistCurveEditor/Private/CurveEditorViewportClient.cpp b/Engine/Source/Editor/DistCurveEditor/Private/CurveEditorViewportClient.cpp
--- a/Engine/Source/Editor/DistCurveEditor/Private/CurveEditorViewportClient.cpp
+++ b/Engine/Source/Editor/DistCurveEditor/Private/CurveEditorViewportClient.cpp
@@ -1203,7 +1203,7 @@ void FCurveEditorViewportClient::DrawEntry(FViewport* Viewport, FCanvas* Canvas,
 				FString KeyComment;
 				if (bSnapToFrames)
 				{
-					KeyComment = FString::Printf(TEXT("(%df,%3.2f)"), FMath::RoundToInt(NewKey.X/InSnapAmount), NewKey.Y);
+					KeyComment = FString::Printf(TEXT("(%" INT64_FMT "f,%3.2f)"), FMath::RoundToInt(NewKey.X/InSnapAmount), NewKey.Y);
 				}
 				else
 				{
diff --git a/Engine/Source/Editor/KismetCompiler/Private/KismetCompiler.cpp b/Engine/Source/Editor/KismetCompiler/Private/KismetCompiler.cpp
--- a/Engine/Source/Editor/KismetCompiler/Private/KismetCompiler.cpp
+++ b/Engine/Source/Editor/KismetCompiler/Private/KismetCompiler.cpp
@@ -1823,7 +1823,7 @@ void FKismetCompilerContext::PrintVerboseInfoStruct(UStruct* Struct) const
 				Prop->GetOffset_ForDebug(),
 				Prop->ElementSize,
 				Prop->ArrayDim,
-				FText::FromString(FString::Printf(TEXT("%x"), Prop->PropertyFlags))
+				FText::FromString(FString::Printf(TEXT("%" UINT64_x_FMT), (uint64)Prop->PropertyFlags))
 			).ToString()
 		);
 	}
diff --git a/Engine/Source/Editor/PropertyEditor/Private/PropertyHandleImpl.cpp b/Engine/Source/Editor/PropertyEditor/Private/PropertyHandleImpl.cpp
--- a/Engine/Source/Editor/PropertyEditor/Private/PropertyHandleImpl.cpp
+++ b/Engine/Source/Editor/PropertyEditor/Private/PropertyHandleImpl.cpp
@@ -183,20 +183,20 @@ FString FPropertyValueImpl::GetPropertyValueArray() const
 					if ( FArrayProperty* ArrayProperty = CastField<FArrayProperty>(NodeProperty) )
 					{
 						FScriptArrayHelper ArrayHelper(ArrayProperty, Addr);
-						String = FString::Printf( TEXT("%(%d)"), ArrayHelper.Num() );
+						String = FString::Printf( TEXT("%%(%d)"), ArrayHelper.Num() );
 					}
 					else if ( CastField<FSetProperty>(NodeProperty) != nullptr )	
 					{
-						String = FString::Printf( TEXT("%(%d)"), FScriptSetHelper::Num(Addr) );
+						String = FString::Printf( TEXT("%%(%d)"), FScriptSetHelper::Num(Addr) );
 					}
 					else if (FMapProperty* MapProperty = CastField<FMapProperty>(NodeProperty))
 					{
 						FScriptMapHelper MapHelper(MapProperty, Addr);
-						String = FString::Printf(TEXT("%(%d)"), MapHelper.Num());
+						String = FString::Printf(TEXT("%%(%d)"), MapHelper.Num());
 					}
 					else
 					{
-						String = FString::Printf( TEXT("%[%d]"), NodeProperty->ArrayDim );
+						String = FString::Printf( TEXT("%%[%d]"), NodeProperty->ArrayDim );
 					}
 				}
 			}
diff --git a/Engine/Source/Editor/Sequencer/Private/SSequencerTimePanel.cpp b/Engine/Source/Editor/Sequencer/Private/SSequencerTimePanel.cpp
--- a/Engine/Source/Editor/Sequencer/Private/SSequencerTimePanel.cpp
+++ b/Engine/Source/Editor/Sequencer/Private/SSequencerTimePanel.cpp
@@ -327,8 +327,8 @@ FText SSequencerTimePanel::GetSupportedTimeRange() const
 	int64 Seconds =  (TotalMaxSeconds % 60   );
 
 	FString String = Days > 0
-		? FString::Printf(TEXT("+/- %02dd %02dh %02dm %02ds"), Days, Hours, Minutes, Seconds)
-		: FString::Printf(TEXT("+/- %02dh %02dm %02ds"), Hours, Minutes, Seconds);
+		? FString::Printf(TEXT("+/- %02" INT64_FMT "d %02" INT64_FMT "h %02" INT64_FMT "m %02" INT64_FMT "s"), Days, Hours, Minutes, Seconds)
+		: FString::Printf(TEXT("+/- %02" INT64_FMT "h %02" INT64_FMT "m %02" INT64_FMT "s"), Hours, Minutes, Seconds);
 
 	return FText::FromString(String);
 }
diff --git a/Engine/Source/Editor/SequencerCore/Public/MVVM/Views/TreeViewTraits.h b/Engine/Source/Editor/SequencerCore/Public/MVVM/Views/TreeViewTraits.h
--- a/Engine/Source/Editor/SequencerCore/Public/MVVM/Views/TreeViewTraits.h
+++ b/Engine/Source/Editor/SequencerCore/Public/MVVM/Views/TreeViewTraits.h
@@ -47,7 +47,7 @@ public:
 	static FString DebugDump( PtrType InPtr )
 	{
 		UE::Sequencer::TViewModelPtr<ExtensionType> Pinned = InPtr.Pin();
-		return Pinned ? FString::Printf(TEXT("0x%08x"), Pinned.AsModel().Get()) : FString(TEXT("nullptr"));
+		return Pinned ? FString::Printf(TEXT("0x%08" UPTRINT_x_FMT), (UPTRINT)Pinned.AsModel().Get()) : FString(TEXT("nullptr"));
 	}
 
 	class SerializerType{};
diff --git a/Engine/Source/Editor/TurnkeySupport/Private/TurnkeySupportModule.cpp b/Engine/Source/Editor/TurnkeySupport/Private/TurnkeySupportModule.cpp
--- a/Engine/Source/Editor/TurnkeySupport/Private/TurnkeySupportModule.cpp
+++ b/Engine/Source/Editor/TurnkeySupport/Private/TurnkeySupportModule.cpp
@@ -2380,11 +2380,11 @@ bool GetSdkInfoFromTurnkey(FString Line, FName& PlatformName, FString& DeviceId,
 		FParse::Value(*Info, *FString::Printf(TEXT("Allowed_%s="), *Name), Max);
 		FParse::Value(*Info, *FString::Printf(TEXT("Current_%s="), *Name), Current);
 		// also handle no name at all (for device, etc)
-		FParse::Value(*Info, *FString::Printf(TEXT("MinAllowed="), *Name), Min);
-		FParse::Value(*Info, *FString::Printf(TEXT("MaxAllowed="), *Name), Max);
-		FParse::Value(*Info, *FString::Printf(TEXT("Allowed="), *Name), Min);
-		FParse::Value(*Info, *FString::Printf(TEXT("Allowed="), *Name), Max);
-		FParse::Value(*Info, *FString::Printf(TEXT("Current=")), Current);
+		FParse::Value(*Info, TEXT("MinAllowed="), Min);
+		FParse::Value(*Info, TEXT("MaxAllowed="), Max);
+		FParse::Value(*Info, TEXT("Allowed="), Min);
+		FParse::Value(*Info, TEXT("Allowed="), Max);
+		FParse::Value(*Info, TEXT("Current="), Current);
 		SdkInfo.SDKVersions.Add(Name, { Min, Max, Current });
 	}
 
diff --git a/Engine/Source/Editor/UnrealEd/Private/Commandlets/AssetRegUtilCommandlet.cpp b/Engine/Source/Editor/UnrealEd/Private/Commandlets/AssetRegUtilCommandlet.cpp
--- a/Engine/Source/Editor/UnrealEd/Private/Commandlets/AssetRegUtilCommandlet.cpp
+++ b/Engine/Source/Editor/UnrealEd/Private/Commandlets/AssetRegUtilCommandlet.cpp
@@ -420,7 +420,7 @@ bool UAssetRegUtilCommandlet::GenerateOrderFile(TMap<FString, int64>& OutputOrde
 		for (TMap<FString, int64>::TConstIterator It(OutputOrderMap); It; ++It)
 		{
 			FString OutputLine;
-			OutputLine = FString::Printf(TEXT("\"%s\" %llu\n"), *It->Key, NewOrderIndex++);
+			OutputLine = FString::Printf(TEXT("\"%s\" %d\n"), *It->Key, NewOrderIndex++);
 			OutArc->Serialize(const_cast<ANSICHAR*>(StringCast<ANSICHAR>(*OutputLine).Get()), OutputLine.Len());
 		}
 
diff --git a/Engine/Source/Editor/UnrealEd/Private/Commandlets/AssetSizeQueryCommandlet.cpp b/Engine/Source/Editor/UnrealEd/Private/Commandlets/AssetSizeQueryCommandlet.cpp
--- a/Engine/Source/Editor/UnrealEd/Private/Commandlets/AssetSizeQueryCommandlet.cpp
+++ b/Engine/Source/Editor/UnrealEd/Private/Commandlets/AssetSizeQueryCommandlet.cpp
@@ -244,7 +244,7 @@ int32 UAssetSizeQueryCommandlet::Main(const FString& FullCommandLine)
 			for (TPair<FTopLevelAssetPath, int64>& ClassSizePair : FilteredClassCompressedSizes)
 			{
 				// we add to both maps at the same time to we know the lookup succeeds.
-				Lines.Add(FString::Printf(TEXT("%s,%lld,%lld"), *ClassSizePair.Key.ToString(), FilteredClassMatchedAssets[ClassSizePair.Key].Num(), ClassSizePair.Value));
+				Lines.Add(FString::Printf(TEXT("%s,%d,%" INT64_FMT), *ClassSizePair.Key.ToString(), FilteredClassMatchedAssets[ClassSizePair.Key].Num(), ClassSizePair.Value));
 			}
 		}
 		else if (OutputCSVType == EOutputCSVType::Assets)
diff --git a/Engine/Source/Editor/UnrealEd/Private/Commandlets/ContentCommandlets.cpp b/Engine/Source/Editor/UnrealEd/Private/Commandlets/ContentCommandlets.cpp
--- a/Engine/Source/Editor/UnrealEd/Private/Commandlets/ContentCommandlets.cpp
+++ b/Engine/Source/Editor/UnrealEd/Private/Commandlets/ContentCommandlets.cpp
@@ -3227,7 +3227,7 @@ int32 UWrangleContentCommandlet::Main( const FString& Params )
 
 			// dump out a line to the .csv file
 			// @todo: sort by size to Excel's 65536 limit gets the biggest objects
-			FString CSVLine = FString::Printf(TEXT("%s,%s,%d%s"), *Object.PackageName, *Object.ObjectName, Object.SerialSize, LINE_TERMINATOR);
+			FString CSVLine = FString::Printf(TEXT("%s,%s,%" INT64_FMT "%s"), *Object.PackageName, *Object.ObjectName, Object.SerialSize, LINE_TERMINATOR);
 			CSVFile->Serialize(TCHAR_TO_ANSI(*CSVLine), CSVLine.Len());
 		}
 	}
diff --git a/Engine/Source/Editor/UnrealEd/Private/Commandlets/SummarizeTraceCommandlet.cpp b/Engine/Source/Editor/UnrealEd/Private/Commandlets/SummarizeTraceCommandlet.cpp
--- a/Engine/Source/Editor/UnrealEd/Private/Commandlets/SummarizeTraceCommandlet.cpp
+++ b/Engine/Source/Editor/UnrealEd/Private/Commandlets/SummarizeTraceCommandlet.cpp
@@ -1280,7 +1280,7 @@ bool USummarizeTraceCommandlet::GenerateBookmarksCSV(const FSummarizeBookmarksPr
 		}
 
 		// note newline is at the front of every data line to prevent final extraneous newline, per customary for csv
-		CsvUtils::WriteAsUTF8String(CsvHandle.Get(), FString::Printf(TEXT("\n%s,%d,%f,%f,"),
+		CsvUtils::WriteAsUTF8String(CsvHandle.Get(), FString::Printf(TEXT("\n%s,%" UINT64_FMT ",%f,%f,"),
 			*Bookmark.Value.Name, Bookmark.Value.Count,
 			Bookmark.Value.FirstSeconds, Bookmark.Value.LastSeconds));
 	}
diff --git a/Engine/Source/Editor/UnrealEd/Private/Commandlets/TextAssetCommandlet.cpp b/Engine/Source/Editor/UnrealEd/Private/Commandlets/TextAssetCommandlet.cpp
--- a/Engine/Source/Editor/UnrealEd/Private/Commandlets/TextAssetCommandlet.cpp
+++ b/Engine/Source/Editor/UnrealEd/Private/Commandlets/TextAssetCommandlet.cpp
@@ -917,7 +917,7 @@ bool UTextAssetCommandlet::DoTextAssetProcessing(const FProcessingArgs& InArgs)
 
 		if (CSVWriter != nullptr)
 		{
-			FString CSVLine = FString::Printf(TEXT("%f,%i,%f,%f,%f,%f\n"), TotalTime, NumFiles, AvgFileTime, MinFileTime, MaxFileTime, IterationPackageLoadTime);
+			FString CSVLine = FString::Printf(TEXT("%f,%" INT64_FMT ",%f,%f,%f,%f\n"), TotalTime, NumFiles, AvgFileTime, MinFileTime, MaxFileTime, IterationPackageLoadTime);
 			CSVWriter->Serialize(TCHAR_TO_ANSI(*CSVLine), CSVLine.Len());
 		}
 
diff --git a/Engine/Source/Editor/UnrealEd/Private/Cooker/DiffWriterArchive.cpp b/Engine/Source/Editor/UnrealEd/Private/Cooker/DiffWriterArchive.cpp
--- a/Engine/Source/Editor/UnrealEd/Private/Cooker/DiffWriterArchive.cpp
+++ b/Engine/Source/Editor/UnrealEd/Private/Cooker/DiffWriterArchive.cpp
@@ -598,7 +598,7 @@ void FAccumulator::OnSecondSaveComplete(int64 InHeaderSize)
 	if (HeaderSize != InHeaderSize)
 	{
 		MessageCallback(ELogVerbosity::Error, FString::Printf(
-			TEXT("%s: Indeterministic header size. When saving the package twice into memory, first header size %d != second header size %d. Callstacks for indeterminism in the exports will be incorrect.")
+			TEXT("%s: Indeterministic header size. When saving the package twice into memory, first header size %" INT64_FMT " != second header size %" INT64_FMT ". Callstacks for indeterminism in the exports will be incorrect.")
 			TEXT("\n\tDumping differences from first and second memory saves."),
 			*this->Filename, HeaderSize, InHeaderSize));
 
@@ -965,7 +965,7 @@ void FAccumulator::CompareWithPreviousForSection(const FPackageData& SourcePacka
 		}
 
 		MessageCallback(ELogVerbosity::Display, FString::Printf(
-			TEXT("%s: Logging %d bytes around offset: %lld (%016X) in the OnDisk package:"),
+			TEXT("%s: Logging %d bytes around offset: %" INT64_FMT " (%016" INT64_X_FMT ") in the OnDisk package:"),
 			*SectionFilename,
 			BytesToLog, LocalOffset, LocalOffset
 		));
@@ -978,7 +978,7 @@ void FAccumulator::CompareWithPreviousForSection(const FPackageData& SourcePacka
 		}
 
 		MessageCallback(ELogVerbosity::Display, FString::Printf(
-			TEXT("%s: Logging %d bytes around offset: %lld (%016X) in the InMemory package:"),
+			TEXT("%s: Logging %d bytes around offset: %" INT64_FMT " (%016" INT64_X_FMT ") in the InMemory package:"),
 			*SectionFilename, BytesToLog, LocalOffset, LocalOffset
 		));
 		HexDumpLines = FCompressionUtil::HexDumpLines(DestPackage.Data + DestPackage.StartOffset,
diff --git a/Engine/Source/Editor/UnrealEd/Private/Cooker/DiffWriterZenHeader.cpp b/Engine/Source/Editor/UnrealEd/Private/Cooker/DiffWriterZenHeader.cpp
--- a/Engine/Source/Editor/UnrealEd/Private/Cooker/DiffWriterZenHeader.cpp
+++ b/Engine/Source/Editor/UnrealEd/Private/Cooker/DiffWriterZenHeader.cpp
@@ -426,7 +426,7 @@ FStringView FZenPackageExportsForDiff::GetExportPackageRelativePath(uint64 Publi
 	FString& ExportPath = ExportPaths.FindOrAdd(PublicExportHash);
 	if (ExportPath.IsEmpty())
 	{
-		ExportPath = FString::Printf(TEXT("<UnknownExport.%u>"), PublicExportHash);
+		ExportPath = FString::Printf(TEXT("<UnknownExport.%" UINT64_FMT ">"), PublicExportHash);
 	}
 	return ExportPath;
 }
diff --git a/Engine/Source/Editor/UnrealEd/Private/EditorDomain/EditorDomainUtils.cpp b/Engine/Source/Editor/UnrealEd/Private/EditorDomain/EditorDomainUtils.cpp
--- a/Engine/Source/Editor/UnrealEd/Private/EditorDomain/EditorDomainUtils.cpp
+++ b/Engine/Source/Editor/UnrealEd/Private/EditorDomain/EditorDomainUtils.cpp
@@ -1838,7 +1838,7 @@ bool TrySavePackage(UPackage* Package)
 			});
 
 		FString Message = FString::Printf(TEXT("Could not save package to EditorDomain because BulkData size is too large. ")
-			TEXT("Package=%s, BulkDataSize=%d, EditorDomain.MaxBulkDataSize=%d")
+			TEXT("Package=%s, BulkDataSize=%" UINT64_FMT ", EditorDomain.MaxBulkDataSize=%" INT64_FMT)
 			TEXT("\n\tWe did not detect this until after trying to save the package, which is bad for performance. Resave the package ")
 			TEXT("or debug why the size was not detected by the BulkDataRegistry."),
 			*WriteToString<256>(PackageName), PackageWriter->GetBulkDataSize(), GMaxBulkDataSize);
diff --git a/Engine/Source/Editor/UnrealEd/Private/Fbx/FbxMainImport.cpp b/Engine/Source/Editor/UnrealEd/Private/Fbx/FbxMainImport.cpp
--- a/Engine/Source/Editor/UnrealEd/Private/Fbx/FbxMainImport.cpp
+++ b/Engine/Source/Editor/UnrealEd/Private/Fbx/FbxMainImport.cpp
@@ -3341,7 +3341,7 @@ FString GetFbxPropertyStringValue(const FbxProperty& Property)
 	case eFbxDouble2:
 	{
 		FbxDouble2 Vec = Property.Get<FbxDouble2>();
-		ValueStr = FString::Printf(TEXT("(%f, %f, %f, %f)"), Vec[0], Vec[1]);
+		ValueStr = FString::Printf(TEXT("(%f, %f)"), Vec[0], Vec[1]);
 	}
 	break;
 	case eFbxDouble3:
diff --git a/Engine/Source/Runtime/AIModule/Private/BehaviorTree/Tasks/BTTask_PlayAnimation.cpp b/Engine/Source/Runtime/AIModule/Private/BehaviorTree/Tasks/BTTask_PlayAnimation.cpp
--- a/Engine/Source/Runtime/AIModule/Private/BehaviorTree/Tasks/BTTask_PlayAnimation.cpp
+++ b/Engine/Source/Runtime/AIModule/Private/BehaviorTree/Tasks/BTTask_PlayAnimation.cpp
@@ -93,9 +93,9 @@ EBTNodeResult::Type UBTTask_PlayAnimation::AbortTask(UBehaviorTreeComponent& Own
 
 FString UBTTask_PlayAnimation::GetStaticDescription() const
 {
-	return FString::Printf(TEXT("%s: '%s'"), *Super::GetStaticDescription(), *GetNameSafe(AnimationToPlay)
+	return FString::Printf(TEXT("%s: '%s'%s%s"), *Super::GetStaticDescription(), *GetNameSafe(AnimationToPlay)
 		, bLooping ? TEXT(", looping") : TEXT("")
-		, bNonBlocking ? TEXT(", non-blocking") : TEXT("blocking"));
+		, bNonBlocking ? TEXT(", non-blocking") : TEXT(", blocking"));
 }
 
 void UBTTask_PlayAnimation::OnAnimationTimerDone()
diff --git a/Engine/Source/Runtime/AIModule/Private/EnvironmentQuery/EnvQueryDebugHelpers.cpp b/Engine/Source/Runtime/AIModule/Private/EnvironmentQuery/EnvQueryDebugHelpers.cpp
--- a/Engine/Source/Runtime/AIModule/Private/EnvironmentQuery/EnvQueryDebugHelpers.cpp
+++ b/Engine/Source/Runtime/AIModule/Private/EnvironmentQuery/EnvQueryDebugHelpers.cpp
@@ -214,9 +214,9 @@ void UEnvQueryDebugHelpers::LogQueryInternal(FEnvQueryInstance& Query, const FNa
 		// table header		
 		{
 			FString HeaderString;
-			PRINT_TABLE_ROW(FString::Printf(TEXT("Item "), EQSLocalData.NumValidItems), HeaderString, 27);
+			PRINT_TABLE_ROW(TEXT("Item "), HeaderString, 27);
 
-			PRINT_TABLE_ROW(FString::Printf(TEXT(" | Score"), EQSLocalData.NumValidItems), HeaderString, 8);
+			PRINT_TABLE_ROW(TEXT(" | Score"), HeaderString, 8);
 
 			for (int32 TestIdx = 0; TestIdx < NumTests; TestIdx++)
 			{
diff --git a/Engine/Source/Runtime/AIModule/Private/GameplayDebugger/GameplayDebuggerCategory_AI.cpp b/Engine/Source/Runtime/AIModule/Private/GameplayDebugger/GameplayDebuggerCategory_AI.cpp
--- a/Engine/Source/Runtime/AIModule/Private/GameplayDebugger/GameplayDebuggerCategory_AI.cpp
+++ b/Engine/Source/Runtime/AIModule/Private/GameplayDebugger/GameplayDebuggerCategory_AI.cpp
@@ -114,15 +114,16 @@ void FGameplayDebuggerCategory_AI::FRepDataPath::Serialize(FArchive& Ar)
 static FString DescribeTaskHelper(const UGameplayTask& TaskOb)
 {
 	const UObject* OwnerOb = Cast<const UObject>(TaskOb.GetTaskOwner());
+#if ENABLE_VISUAL_LOG
+	FString StateName = *TaskOb.GetTaskStateName();
+#else
+	const TCHAR* StateName = TEXT("UnknownTaskStateName");
+#endif
 	return FString::Printf(TEXT("\n  {white}%s%s {%s}%s:%d {white}Owner:{yellow}%s {white}Res:{yellow}%s"),
 		*TaskOb.GetName(),
 		TaskOb.GetInstanceName() != NAME_None ? *FString::Printf(TEXT(" {yellow}[%s]"), *TaskOb.GetInstanceName().ToString()) : TEXT(""),
 		TaskOb.IsActive() ? TEXT("green") : TEXT("orange"),
-#if ENABLE_VISUAL_LOG 
-		*TaskOb.GetTaskStateName(),
-#else
-		TEXT("UnknownTaskStateName"),
-#endif
+		ToCStr(StateName),
 		TaskOb.GetPriority(),
 		*GetNameSafe(OwnerOb),
 		TaskOb.GetRequiredResources().IsEmpty() ? TEXT("None") : *TaskOb.GetRequiredResources().GetDebugDescription());
diff --git a/Engine/Source/Runtime/AVEncoder/Private/Decoders/Windows/VideoDecoderH264_Windows.cpp b/Engine/Source/Runtime/AVEncoder/Private/Decoders/Windows/VideoDecoderH264_Windows.cpp
--- a/Engine/Source/Runtime/AVEncoder/Private/Decoders/Windows/VideoDecoderH264_Windows.cpp
+++ b/Engine/Source/Runtime/AVEncoder/Private/Decoders/Windows/VideoDecoderH264_Windows.cpp
@@ -1223,7 +1223,7 @@ bool FVideoDecoderH264_WindowsImpl::InternalDecoderCreate()
 			}
 			else if (FAILED(Result = Decoder->ProcessMessage(MFT_MESSAGE_SET_D3D_MANAGER, reinterpret_cast<ULONG_PTR>(DXDeviceInfo.DxDeviceManager.GetReference()))))
 			{
-				FallbackToSwDecoding(FString::Printf(TEXT("Failed to set MFT_MESSAGE_SET_D3D_MANAGER: 0x%X %s"), Result, *GetComErrorDescription(Result)));
+				FallbackToSwDecoding(FString::Printf(TEXT("Failed to set MFT_MESSAGE_SET_D3D_MANAGER: 0x%X %s"), (uint32)Result, *GetComErrorDescription(Result)));
 			}
 		}
 	#if 0	// PLATFORM_WINDOWS
diff --git a/Engine/Source/Runtime/Analytics/AnalyticsET/Private/AnalyticsPerfTracker.cpp b/Engine/Source/Runtime/Analytics/AnalyticsET/Private/AnalyticsPerfTracker.cpp
--- a/Engine/Source/Runtime/Analytics/AnalyticsET/Private/AnalyticsPerfTracker.cpp
+++ b/Engine/Source/Runtime/Analytics/AnalyticsET/Private/AnalyticsPerfTracker.cpp
@@ -56,7 +56,7 @@ bool FAnalyticsPerfTracker::Tick(float DeltaTime)
 		double Now = FPlatformTime::Seconds();
 		if (WindowExpired(Now))
 		{
-			LogFile.Serialize(*FString::Printf(TEXT("%s,%s,%s,%f,%f,%f,%d,%d,%d,%d,%d"),
+			LogFile.Serialize(*FString::Printf(TEXT("%s,%s,%s,%f,%f,%f,%d,%d,%" UINT64_FMT ",%" UINT64_FMT ",%" UINT64_FMT),
 				*StartDate,
 				*CL,
 				*RunID,
diff --git a/Engine/Source/Runtime/AudioCaptureImplementations/Windows/AudioCaputureWasapi/Private/WasapiCaptureLog.h b/Engine/Source/Runtime/AudioCaptureImplementations/Windows/AudioCaputureWasapi/Private/WasapiCaptureLog.h
--- a/Engine/Source/Runtime/AudioCaptureImplementations/Windows/AudioCaputureWasapi/Private/WasapiCaptureLog.h
+++ b/Engine/Source/Runtime/AudioCaptureImplementations/Windows/AudioCaputureWasapi/Private/WasapiCaptureLog.h
@@ -6,6 +6,6 @@
 
 #define WASAPI_CAPTURE_LOG_RESULT(FunctionName, Result) \
 	{ \
-		FString ErrorString = FString::Printf(TEXT("%s -> 0x%X (line: %d)"), TEXT( FunctionName ), Result, __LINE__); \
+		FString ErrorString = FString::Printf(TEXT("%s -> 0x%X (line: %d)"), TEXT( FunctionName ), (uint32)Result, __LINE__); \
 		UE_LOG(LogAudioCaptureCore, Error, TEXT("WasapiCapture Error: %s"), *ErrorString);				  \
 	}
diff --git a/Engine/Source/Runtime/Core/Private/Windows/WindowsPlatformFileJournal.cpp b/Engine/Source/Runtime/Core/Private/Windows/WindowsPlatformFileJournal.cpp
--- a/Engine/Source/Runtime/Core/Private/Windows/WindowsPlatformFileJournal.cpp
+++ b/Engine/Source/Runtime/Core/Private/Windows/WindowsPlatformFileJournal.cpp
@@ -512,8 +512,8 @@ EFileJournalResult FWindowsPlatformFile::FileJournalReadModified(const TCHAR* Vo
 			if (OutError)
 			{
 				*OutError = FString::Printf(
-					TEXT("FSCTL_READ_UNPRIVILEGED_USN_JOURNAL returned an unexpected value; ReadBufferResultsSize == %u, and the minimum expected is %d."),
-					ReadBufferResultsSize, sizeof(USN));
+					TEXT("FSCTL_READ_UNPRIVILEGED_USN_JOURNAL returned an unexpected value; ReadBufferResultsSize == %u, and the minimum expected is %zu."),
+					(uint32)ReadBufferResultsSize, sizeof(USN));
 			}
 			return EFileJournalResult::JournalInternalError;
 		}
@@ -573,7 +573,7 @@ EFileJournalResult FWindowsPlatformFile::FileJournalReadModified(const TCHAR* Vo
 				*OutError = FString::Printf(
 					TEXT("FSCTL_READ_UNPRIVILEGED_USN_JOURNAL returned an unexpected value; ReadBufferResultsSize == %u, but NextJournalUSN did not increase.")
 					TEXT("Previous value = %" UINT64_FMT ". New value = %" UINT64_FMT "."),
-					ReadBufferResultsSize, static_cast<uint64>(CurrentJournalUSN), static_cast<uint64>(NextJournalUSN));
+					(uint32)ReadBufferResultsSize, static_cast<uint64>(CurrentJournalUSN), static_cast<uint64>(NextJournalUSN));
 			}
 			return EFileJournalResult::JournalInternalError;
 		}
diff --git a/Engine/Source/Runtime/D3D12RHI/Private/D3D12Adapter.cpp b/Engine/Source/Runtime/D3D12RHI/Private/D3D12Adapter.cpp
--- a/Engine/Source/Runtime/D3D12RHI/Private/D3D12Adapter.cpp
+++ b/Engine/Source/Runtime/D3D12RHI/Private/D3D12Adapter.cpp
@@ -1830,7 +1830,7 @@ void FD3D12Adapter::DumpTrackedAllocationData(FOutputDevice& OutputDevice, bool
 			Flags += EnumHasAnyFlags(ResourceDesc.Flags, D3D12_RESOURCE_FLAG_ALLOW_RENDER_TARGET | D3D12_RESOURCE_FLAG_ALLOW_DEPTH_STENCIL) ? "|UAV" : "UAV";
 		}
 
-		OutputData += FString::Printf(TEXT("\tName: %s - Size: %3.3fMB - Width: %d - Height: %d - DepthOrArraySize: %d - MipLevels: %d - Flags: %s - Resident: %s\n"), 
+		OutputData += FString::Printf(TEXT("\tName: %s - Size: %3.3fMB - Width: %" UINT64_FMT " - Height: %d - DepthOrArraySize: %d - MipLevels: %d - Flags: %s - Resident: %s\n"), 
 			*AllocationData.ResourceAllocation->GetResource()->GetName().ToString(), 
 			AllocationData.AllocationSize / (1024.0f * 1024),
 			ResourceDesc.Width, ResourceDesc.Height, ResourceDesc.DepthOrArraySize, ResourceDesc.MipLevels,
@@ -1863,7 +1863,7 @@ void FD3D12Adapter::DumpTrackedAllocationData(FOutputDevice& OutputDevice, bool
 			continue;
 		}
 
-		OutputData += FString::Printf(TEXT("\tName: %s - Size: %3.3fMB - Width: %d - UAV: %s - Resident: %s\n"), 
+		OutputData += FString::Printf(TEXT("\tName: %s - Size: %3.3fMB - Width: %" UINT64_FMT " - UAV: %s - Resident: %s\n"), 
 			*AllocationData.ResourceAllocation->GetResource()->GetName().ToString(), 
 			AllocationData.AllocationSize / (1024.0f * 1024),
 			ResourceDesc.Width,
diff --git a/Engine/Source/Runtime/D3D12RHI/Private/D3D12Util.cpp b/Engine/Source/Runtime/D3D12RHI/Private/D3D12Util.cpp
--- a/Engine/Source/Runtime/D3D12RHI/Private/D3D12Util.cpp
+++ b/Engine/Source/Runtime/D3D12RHI/Private/D3D12Util.cpp
@@ -676,12 +676,12 @@ static FString MakeResourceDescDebugString(const D3D12_RESOURCE_DESC& Desc)
 		ResourceDescString = TEXT("Unknown");
 		break;
 	case D3D12_RESOURCE_DIMENSION_BUFFER:
-		ResourceDescString = FString::Printf(TEXT("Buffer %d bytes"), Desc.Width);
+		ResourceDescString = FString::Printf(TEXT("Buffer %" UINT64_FMT " bytes"), Desc.Width);
 		break;
 	case D3D12_RESOURCE_DIMENSION_TEXTURE1D:
 	case D3D12_RESOURCE_DIMENSION_TEXTURE2D:
 	case D3D12_RESOURCE_DIMENSION_TEXTURE3D:
-		ResourceDescString = FString::Printf(TEXT("Texture %dx%dx%d %s"), Desc.Width, Desc.Height, Desc.DepthOrArraySize, LexToString(Desc.Format));
+		ResourceDescString = FString::Printf(TEXT("Texture %" UINT64_FMT "x%dx%d %s"), Desc.Width, Desc.Height, Desc.DepthOrArraySize, LexToString(Desc.Format));
 	}
 	return ResourceDescString;
 }
diff --git a/Engine/Source/Runtime/Engine/Classes/Engine/EngineBaseTypes.h b/Engine/Source/Runtime/Engine/Classes/Engine/EngineBaseTypes.h
--- a/Engine/Source/Runtime/Engine/Classes/Engine/EngineBaseTypes.h
+++ b/Engine/Source/Runtime/Engine/Classes/Engine/EngineBaseTypes.h
@@ -969,7 +969,7 @@ struct FExposureSettings
 
 	FString ToString() const
 	{
-		return FString::Printf(TEXT("%d,%d"), FixedEV100, bFixed ? 1 : 0);
+		return FString::Printf(TEXT("%f,%d"), FixedEV100, bFixed ? 1 : 0);
 	}
 
 	void SetFromString(const TCHAR *In)
diff --git a/Engine/Source/Runtime/Engine/Private/AudioDerivedData.cpp b/Engine/Source/Runtime/Engine/Private/AudioDerivedData.cpp
--- a/Engine/Source/Runtime/Engine/Private/AudioDerivedData.cpp
+++ b/Engine/Source/Runtime/Engine/Private/AudioDerivedData.cpp
@@ -394,7 +394,7 @@ static uint32 PutDerivedDataInCache(
 
 		if (UE_LOG_ACTIVE(LogAudio,Verbose))
 		{
-			LogString += FString::Printf(TEXT("  Chunk%d %d bytes %s\n"),
+			LogString += FString::Printf(TEXT("  Chunk%d %" INT64_FMT " bytes %s\n"),
 				ChunkIndex,
 				Chunk.BulkData.GetBulkDataSize(),
 				*ChunkDerivedDataKey
diff --git a/Engine/Source/Runtime/Engine/Private/AudioStreamingCache.cpp b/Engine/Source/Runtime/Engine/Private/AudioStreamingCache.cpp
--- a/Engine/Source/Runtime/Engine/Private/AudioStreamingCache.cpp
+++ b/Engine/Source/Runtime/Engine/Private/AudioStreamingCache.cpp
@@ -2316,7 +2316,7 @@ FString FAudioChunkCache::DebugPrint()
 	const double PercentageOfCacheExternalFeatures = NumMegabytesExternalFeatures / MaxCacheSizeMB;
 
 	FString CacheMemoryHeader = *FString::Printf(TEXT("External Features:\t, Force Inline:\t, Retaining:\t, Loaded:\t, Max Potential Usage:\t, \n"));
-	FString CacheMemoryUsage = *FString::Printf(TEXT("%.4f Megabytes (%.3f%% of total capacity)\t %.4f Megabytes (%.3f%% of total capacity)\t %.4f Megabytes (%.3f%% of total capacity)\t,  %.4f Megabytes (%lu bytes)\t, %.4f Megabytes\t, \n"), 
+	FString CacheMemoryUsage = *FString::Printf(TEXT("%.4f Megabytes (%.3f%% of total capacity)\t %.4f Megabytes (%.3f%% of total capacity)\t %.4f Megabytes (%.3f%% of total capacity)\t,  %.4f Megabytes (%" UINT64_FMT " bytes)\t, %.4f Megabytes\t, \n"), 
 		NumMegabytesExternalFeatures, PercentageOfCacheExternalFeatures, NumMegabytesForceInline, PercentageOfCacheForceInlined, NumMegabytesRetained, PercentageOfCacheRetained, NumMegabytesInUse, MemoryCounterBytes.Load(), MaxCacheSizeMB);
 	OutputString += CacheMemoryHeader + CacheMemoryUsage + TEXT("\n");
 
@@ -2551,7 +2551,7 @@ TPair<int, int> FAudioChunkCache::DebugDisplay(UWorld* World, FViewport* Viewpor
 	const double NumMegabytesExternalFeatures = (double)ExternalFeaturesBytes / (1024 * 1024);
 	const float PercentageExternalFeatures = NumBytesCounter > 0 ? (double)ExternalFeaturesBytes / NumBytesCounter : 0;
 
-	FString CacheMemoryUsage = *FString::Printf(TEXT("Using: %.4f Megabytes (%lu bytes). Max Potential Usage: %.4f Megabytes."), 
+	FString CacheMemoryUsage = *FString::Printf(TEXT("Using: %.4f Megabytes (%" UINT64_FMT " bytes). Max Potential Usage: %.4f Megabytes."), 
 		NumMegabytesInUse, GetCurrentMemoryUsageBytes(), MaxCacheSizeMB);
 
 	// We're going to align this horizontally with the number of elements right above it.
diff --git a/Engine/Source/Runtime/Engine/Private/ChartCreation.cpp b/Engine/Source/Runtime/Engine/Private/ChartCreation.cpp
--- a/Engine/Source/Runtime/Engine/Private/ChartCreation.cpp
+++ b/Engine/Source/Runtime/Engine/Private/ChartCreation.cpp
@@ -186,7 +186,7 @@ void FDumpFPSChartToEndpoint::HandleBasicStats()
 	PrintToEndpoint(FString::Printf(TEXT("\tWindow Mode: %s"), *WindowMode));
 	PrintToEndpoint(FString::Printf(TEXT("\tResolution: %dx%d"), GameResolution.X, GameResolution.Y));
 
-	PrintToEndpoint(FString::Printf(TEXT("%i frames collected over %4.2f seconds, disregarding %4.2f seconds (%d frames) for a %4.2f FPS average"),
+	PrintToEndpoint(FString::Printf(TEXT("%" INT64_FMT " frames collected over %4.2f seconds, disregarding %4.2f seconds (%d frames) for a %4.2f FPS average"),
 		Chart.GetNumFrames(),
 		WallClockTimeFromStartOfCharting,
 		Chart.TimeDisregarded,
diff --git a/Engine/Source/Runtime/Engine/Private/DataTableJSON.cpp b/Engine/Source/Runtime/Engine/Private/DataTableJSON.cpp
--- a/Engine/Source/Runtime/Engine/Private/DataTableJSON.cpp
+++ b/Engine/Source/Runtime/Engine/Private/DataTableJSON.cpp
@@ -883,7 +883,7 @@ bool FDataTableImporterJSON::ReadContainerEntry(const TSharedRef<FJsonValue>& In
 			const FString Error = DataTableUtils::AssignStringToPropertyDirect(PropertyValueString, InProperty, (uint8*)InPropertyData);
 			if (Error.Len() > 0)
 			{
-				ImportProblems.Add(FString::Printf(TEXT("Problem assigning string '%s' to entry %d on property '%s' on row '%s' : %s"), InArrayEntryIndex, *PropertyValueString, *InColumnName, *InRowName.ToString(), *Error));
+				ImportProblems.Add(FString::Printf(TEXT("Problem assigning string '%s' to entry %d on property '%s' on row '%s' : %s"), *PropertyValueString, InArrayEntryIndex, *InColumnName, *InRowName.ToString(), *Error));
 				return false;
 			}
 
@@ -902,7 +902,7 @@ bool FDataTableImporterJSON::ReadContainerEntry(const TSharedRef<FJsonValue>& In
 		const FString Error = DataTableUtils::AssignStringToPropertyDirect(PropertyValue, InProperty, (uint8*)InPropertyData);
 		if(Error.Len() > 0)
 		{
-			ImportProblems.Add(FString::Printf(TEXT("Problem assigning string '%s' to entry %d on property '%s' on row '%s' : %s"), InArrayEntryIndex, *PropertyValue, *InColumnName, *InRowName.ToString(), *Error));
+			ImportProblems.Add(FString::Printf(TEXT("Problem assigning string '%s' to entry %d on property '%s' on row '%s' : %s"), *PropertyValue, InArrayEntryIndex, *InColumnName, *InRowName.ToString(), *Error));
 			return false;
 		}
 	}
diff --git a/Engine/Source/Runtime/Engine/Private/DemoNetDriver.cpp b/Engine/Source/Runtime/Engine/Private/DemoNetDriver.cpp
--- a/Engine/Source/Runtime/Engine/Private/DemoNetDriver.cpp
+++ b/Engine/Source/Runtime/Engine/Private/DemoNetDriver.cpp
@@ -109,7 +109,7 @@ static void ConditionallyDisplayBurnInTime(uint32 RecordedCL, float CurrentDemoT
 {
 	if (CVarWithDemoTimeBurnIn.GetValueOnAnyThread() != 0)
 	{
-		GEngine->AddOnScreenDebugMessage(INDEX_NONE, 0.f, FColor::Red, FString::Printf(TEXT("Current CL: %lu | Recorded CL: %lu | Time: %f"), FEngineVersion::Current().GetChangelist(), RecordedCL, CurrentDemoTime), true, FVector2D(3.f, 3.f));
+		GEngine->AddOnScreenDebugMessage(INDEX_NONE, 0.f, FColor::Red, FString::Printf(TEXT("Current CL: %u | Recorded CL: %u | Time: %f"), FEngineVersion::Current().GetChangelist(), RecordedCL, CurrentDemoTime), true, FVector2D(3.f, 3.f));
 	}
 }
 
diff --git a/Engine/Source/Runtime/Engine/Private/EditorFramework/AssetImportData.cpp b/Engine/Source/Runtime/Engine/Private/EditorFramework/AssetImportData.cpp
--- a/Engine/Source/Runtime/Engine/Private/EditorFramework/AssetImportData.cpp
+++ b/Engine/Source/Runtime/Engine/Private/EditorFramework/AssetImportData.cpp
@@ -46,7 +46,7 @@ FString FAssetImportInfo::ToJson() const
 
 	for (int32 Index = 0; Index < SourceFiles.Num(); ++Index)
 	{
-		Json += FString::Printf(TEXT("{ \"RelativeFilename\" : \"%s\", \"Timestamp\" : \"%d\", \"FileMD5\" : \"%s\", \"DisplayLabelName\" : \"%s\" }"),
+		Json += FString::Printf(TEXT("{ \"RelativeFilename\" : \"%s\", \"Timestamp\" : \"%" INT64_FMT "\", \"FileMD5\" : \"%s\", \"DisplayLabelName\" : \"%s\" }"),
 			*SourceFiles[Index].RelativeFilename,
 			SourceFiles[Index].Timestamp.ToUnixTimestamp(),
 			*LexToString(SourceFiles[Index].FileHash),
diff --git a/Engine/Source/Runtime/Engine/Private/GPUSkinCache.cpp b/Engine/Source/Runtime/Engine/Private/GPUSkinCache.cpp
--- a/Engine/Source/Runtime/Engine/Private/GPUSkinCache.cpp
+++ b/Engine/Source/Runtime/Engine/Private/GPUSkinCache.cpp
@@ -2507,7 +2507,7 @@ void FGPUSkinCache::DrawVisualizationInfoText(const FName& GPUSkinCacheVisualiza
 			float UsedMemoryInMB = UsedMemoryInBytes / MBSize;
 			float AvailableMemoryInMB = GSkinCacheSceneMemoryLimitInMB - UsedMemoryInMB;
 
-			FString LowMemoryText = FString::Printf(TEXT("0 - %dMB"), GEngine->GPUSkinCacheVisualizationLowMemoryThresholdInMB);
+			FString LowMemoryText = FString::Printf(TEXT("0 - %fMB"), GEngine->GPUSkinCacheVisualizationLowMemoryThresholdInMB);
 			DrawText(TEXT("Skin Cache Visualization - Memory"), FColor::White);
 			DrawText(FString::Printf(TEXT("Total Limit: %.2fMB"), GSkinCacheSceneMemoryLimitInMB), FColor::White);
 			DrawText(FString::Printf(TEXT("Total Used: %.2fMB"), UsedMemoryInMB), FColor::White);
diff --git a/Engine/Source/Runtime/Engine/Private/GameFramework/RootMotionSource.cpp b/Engine/Source/Runtime/Engine/Private/GameFramework/RootMotionSource.cpp
--- a/Engine/Source/Runtime/Engine/Private/GameFramework/RootMotionSource.cpp
+++ b/Engine/Source/Runtime/Engine/Private/GameFramework/RootMotionSource.cpp
@@ -33,7 +33,7 @@ void RootMotionSourceDebug::PrintOnScreen(const ACharacter& InCharacter, const F
 	// Skip bots, debug player networking.
 	if (InCharacter.IsPlayerControlled())
 	{
-		const FString AdjustedDebugString = FString::Printf(TEXT("[%d] [%s] %s"), (uint64)GFrameCounter, *InCharacter.GetName(), *InString);
+		const FString AdjustedDebugString = FString::Printf(TEXT("[%" UINT64_FMT "] [%s] %s"), (uint64)GFrameCounter, *InCharacter.GetName(), *InString);
 
 		// If on the server, replicate this message to everyone.
 		if (!InCharacter.IsLocallyControlled() && (InCharacter.GetLocalRole() == ROLE_Authority))
diff --git a/Engine/Source/Runtime/Engine/Private/Level.cpp b/Engine/Source/Runtime/Engine/Private/Level.cpp
--- a/Engine/Source/Runtime/Engine/Private/Level.cpp
+++ b/Engine/Source/Runtime/Engine/Private/Level.cpp
@@ -3015,7 +3015,7 @@ void ULevel::FixupActorFolders()
 				const FString NewPath = FString::Printf(TEXT("%s_%s"), *NewParent->GetPath().ToString(), *FGuid::NewGuid().ToString());
 				const FFolder NewFolder = FFolder(OldFolder.GetRootObject(), FName(NewPath));
 				FLevelActorFoldersHelper::RenameFolder(this, OldFolder, NewFolder);
-				const FString NewPath2 = FString::Printf(TEXT("%s/DuplicateFolder_%s"), *NewParent->GetPath().ToString(), *FolderToDelete->GetLabel(), *FGuid::NewGuid().ToString());
+				const FString NewPath2 = FString::Printf(TEXT("%s/DuplicateFolder_%s_%s"), *NewParent->GetPath().ToString(), *FolderToDelete->GetLabel(), *FGuid::NewGuid().ToString());
 				const FFolder NewFolder2 = FFolder(OldFolder.GetRootObject(), FName(NewPath2));
 				FLevelActorFoldersHelper::RenameFolder(this, NewFolder, NewFolder2);
 				// Then delete (mark as deleted) this folder
diff --git a/Engine/Source/Runtime/Engine/Private/Materials/MaterialHLSLEmitter.cpp b/Engine/Source/Runtime/Engine/Private/Materials/MaterialHLSLEmitter.cpp
--- a/Engine/Source/Runtime/Engine/Private/Materials/MaterialHLSLEmitter.cpp
+++ b/Engine/Source/Runtime/Engine/Private/Materials/MaterialHLSLEmitter.cpp
@@ -219,7 +219,7 @@ static FString GenerateMaterialTemplateHLSL(EShaderPlatform ShaderPlatform,
 		{
 			const TCHAR* DerivativeName = (DerivativeIndex == 0) ? TEXT("HWDerivative") : TEXT("AnalyticDerivative");
 
-			EvaluateMaterialAttributesPhase0[DerivativeIndex] =  FString::Printf(TEXT("    FMaterialAttributes MaterialAttributesPhase0;" HLSL_LINE_TERMINATOR), DerivativeName);
+			EvaluateMaterialAttributesPhase0[DerivativeIndex] =  FString::Printf(TEXT("    FMaterialAttributes MaterialAttributesPhase0;" HLSL_LINE_TERMINATOR));
 			EvaluateMaterialAttributesPhase0[DerivativeIndex] += FString::Printf(TEXT("    EvaluatePixelMaterialAttributesPhase0_%s(Parameters, MaterialAttributesPhase0);" HLSL_LINE_TERMINATOR), DerivativeName);
 
 			EvaluateMaterialDeclaration += FString::Printf(TEXT("void EvaluatePixelMaterialAttributesPhase0_%s(in out FMaterialPixelParameters Parameters, out FMaterialAttributes OutResult)" HLSL_LINE_TERMINATOR), DerivativeName);
@@ -234,7 +234,7 @@ static FString GenerateMaterialTemplateHLSL(EShaderPlatform ShaderPlatform,
 				EvaluateMaterialDeclaration += PixelShaderCodePhase1[DerivativeIndex];
 				EvaluateMaterialDeclaration += TEXT("}" HLSL_LINE_TERMINATOR);
 
-				EvaluateMaterialAttributesPhase1[DerivativeIndex] =  FString::Printf(TEXT("    FMaterialAttributes MaterialAttributesPhase1;" HLSL_LINE_TERMINATOR), DerivativeName);
+				EvaluateMaterialAttributesPhase1[DerivativeIndex] =  FString::Printf(TEXT("    FMaterialAttributes MaterialAttributesPhase1;" HLSL_LINE_TERMINATOR));
 				EvaluateMaterialAttributesPhase1[DerivativeIndex] += FString::Printf(TEXT("    EvaluatePixelMaterialAttributesPhase1_%s(Parameters, MaterialAttributesPhase1);" HLSL_LINE_TERMINATOR), DerivativeName);
 				EvaluateMaterialAttributesPhase1[DerivativeIndex] += TEXT("    Parameters.MaterialAttributes = MaterialAttributesPhase1;" HLSL_LINE_TERMINATOR);
 			}
diff --git a/Engine/Source/Runtime/Engine/Private/PerQualityLevelProperties.cpp b/Engine/Source/Runtime/Engine/Private/PerQualityLevelProperties.cpp
--- a/Engine/Source/Runtime/Engine/Private/PerQualityLevelProperties.cpp
+++ b/Engine/Source/Runtime/Engine/Private/PerQualityLevelProperties.cpp
@@ -520,7 +520,7 @@ FString FPerQualityLevelFloat::ToString() const
 
 	for (int32 QL : QualityLevels)
 	{
-		Result = FString::Printf(TEXT("%s, %s=%d"), *Result, *QualityLevelProperty::QualityLevelToFName(QL).ToString(), PerQuality.FindChecked(QL));
+		Result = FString::Printf(TEXT("%s, %s=%f"), *Result, *QualityLevelProperty::QualityLevelToFName(QL).ToString(), PerQuality.FindChecked(QL));
 	}
 #endif
 
diff --git a/Engine/Source/Runtime/Engine/Private/PhysicsEngine/BodySetup.cpp b/Engine/Source/Runtime/Engine/Private/PhysicsEngine/BodySetup.cpp
--- a/Engine/Source/Runtime/Engine/Private/PhysicsEngine/BodySetup.cpp
+++ b/Engine/Source/Runtime/Engine/Private/PhysicsEngine/BodySetup.cpp
@@ -1278,7 +1278,7 @@ bool ShouldSkipDDC(UBodySetup* InSetup, FString& OutReason)
 
 			if (VerticeCount >= SkipDDCThreshold)
 			{
-				OutReason = FString::Printf(TEXT("AggGeom Vertice Count %ld"), VerticeCount);
+				OutReason = FString::Printf(TEXT("AggGeom Vertice Count %d"), VerticeCount);
 				return false;
 			}
 		}
diff --git a/Engine/Source/Runtime/Engine/Private/StaticMesh.cpp b/Engine/Source/Runtime/Engine/Private/StaticMesh.cpp
--- a/Engine/Source/Runtime/Engine/Private/StaticMesh.cpp
+++ b/Engine/Source/Runtime/Engine/Private/StaticMesh.cpp
@@ -4834,7 +4834,7 @@ void UStaticMesh::GetAssetRegistryTags(FAssetRegistryTagsContext Context) const
 	{
 		Bounds = GetRenderData()->Bounds;
 	}
-	const FString ApproxSizeStr = FString::Printf(TEXT("%dx%dx%d"), FMath::RoundToInt(Bounds.BoxExtent.X * 2.0f), FMath::RoundToInt(Bounds.BoxExtent.Y * 2.0f), FMath::RoundToInt(Bounds.BoxExtent.Z * 2.0f));
+	const FString ApproxSizeStr = FString::Printf(TEXT("%" INT64_FMT "x%" INT64_FMT "x%" INT64_FMT), FMath::RoundToInt(Bounds.BoxExtent.X * 2.0f), FMath::RoundToInt(Bounds.BoxExtent.Y * 2.0f), FMath::RoundToInt(Bounds.BoxExtent.Z * 2.0f));
 
 	// Get name of default collision profile
 	FName DefaultCollisionName = NAME_None;
diff --git a/Engine/Source/Runtime/Engine/Private/TextureDerivedData.cpp b/Engine/Source/Runtime/Engine/Private/TextureDerivedData.cpp
--- a/Engine/Source/Runtime/Engine/Private/TextureDerivedData.cpp
+++ b/Engine/Source/Runtime/Engine/Private/TextureDerivedData.cpp
@@ -1590,7 +1590,7 @@ int64 PutDerivedDataInCache(FTexturePlatformData* DerivedData, const FString& De
 				);
 			}
 
-			LogString += FString::Printf(TEXT("  Mip%d %dx%d %d bytes%s %s\n"),
+			LogString += FString::Printf(TEXT("  Mip%d %dx%d %" UINT64_FMT " bytes%s %s\n"),
 				MipIndex,
 				Mip.SizeX,
 				Mip.SizeY,
diff --git a/Engine/Source/Runtime/Engine/Private/UnrealEngine.cpp b/Engine/Source/Runtime/Engine/Private/UnrealEngine.cpp
--- a/Engine/Source/Runtime/Engine/Private/UnrealEngine.cpp
+++ b/Engine/Source/Runtime/Engine/Private/UnrealEngine.cpp
@@ -11777,7 +11777,7 @@ void UEngine::LogPerformanceCapture(UWorld* World, const FString& MapName, const
 		PerfSnapshot.AverageFrameTime = FString::Printf(TEXT("%0.2f"), StatUnitData->FrameTime);
 		PerfSnapshot.AverageGameThreadTime = FString::Printf(TEXT("%0.2f"), StatUnitData->GameThreadTime);
 		PerfSnapshot.AverageRenderThreadTime = FString::Printf(TEXT("%0.2f"), StatUnitData->RenderThreadTime);
-		PerfSnapshot.AverageGPUTime = FString::Printf(TEXT("%0.2f"), StatUnitData->GPUFrameTime);
+		PerfSnapshot.AverageGPUTime = FString::Printf(TEXT("%0.2f"), StatUnitData->GPUFrameTime[0]);
 		// PerfSnapshot.PercentOfFramesAtLeast60FPS = ???;	// @todo
 		// PerfSnapshot.PercentOfFramesAtLeast60FPS = ???;	// @todo
 
@@ -19340,7 +19340,7 @@ int32 UEngine::RenderStatFrameCounter(UWorld* World, FViewport* Viewport, FCanva
 	UFont* Font = FPlatformProperties::SupportsWindowedMode() ? GetSmallFont() : GetMediumFont();
 	const int32 RowHeight = FMath::TruncToInt(Font->GetMaxCharHeight() * 1.1f);
 
-	Canvas->DrawShadowedString(X, Y, *FString::Printf(TEXT("FC: %d"), GFrameCounter), Font, FColor::Green);
+	Canvas->DrawShadowedString(X, Y, *FString::Printf(TEXT("FC: %" UINT64_FMT), GFrameCounter), Font, FColor::Green);
 	Y += RowHeight;
 
 	return Y;
diff --git a/Engine/Source/Runtime/Engine/Private/VisualLogger/VisualLoggerTypes.cpp b/Engine/Source/Runtime/Engine/Private/VisualLogger/VisualLoggerTypes.cpp
--- a/Engine/Source/Runtime/Engine/Private/VisualLogger/VisualLoggerTypes.cpp
+++ b/Engine/Source/Runtime/Engine/Private/VisualLogger/VisualLoggerTypes.cpp
@@ -678,7 +678,7 @@ FString FVisualLoggerHelpers::GenerateTemporaryFilename(const FString& FileExt)
 FString FVisualLoggerHelpers::GenerateFilename(const FString& TempFileName, const FString& Prefix, double StartRecordingTime, double EndTimeStamp)
 {
 	const FString FullFilename = FString::Printf(TEXT("%s_%s"), *Prefix, *TempFileName);
-	const FString TimeFrameString = FString::Printf(TEXT("%d-%d_"), FMath::TruncToInt(StartRecordingTime), FMath::TruncToInt(EndTimeStamp));
+	const FString TimeFrameString = FString::Printf(TEXT("%" INT64_FMT "-%" INT64_FMT "_"), FMath::TruncToInt(StartRecordingTime), FMath::TruncToInt(EndTimeStamp));
 	return FullFilename.Replace(TEXT("VTEMP_"), *TimeFrameString, ESearchCase::CaseSensitive);
 }
 
diff --git a/Engine/Source/Runtime/Engine/Private/WorldPartition/RuntimeHashSet/RuntimePartitionLHGrid.cpp b/Engine/Source/Runtime/Engine/Private/WorldPartition/RuntimeHashSet/RuntimePartitionLHGrid.cpp
--- a/Engine/Source/Runtime/Engine/Private/WorldPartition/RuntimeHashSet/RuntimePartitionLHGrid.cpp
+++ b/Engine/Source/Runtime/Engine/Private/WorldPartition/RuntimeHashSet/RuntimePartitionLHGrid.cpp
@@ -18,7 +18,7 @@ struct FCellCoord
 
 	inline FString ToString() const
 	{
-		return FString::Printf(TEXT("%d_%d_%d_%d"), X, Y, Z, Level);
+		return FString::Printf(TEXT("%" INT64_FMT "_%" INT64_FMT "_%" INT64_FMT "_%d"), X, Y, Z, Level);
 	}
 
 	inline bool operator==(const FCellCoord& Other) const
diff --git a/Engine/Source/Runtime/Engine/Private/WorldPartition/WorldPartitionRuntimeSpatialHash.cpp b/Engine/Source/Runtime/Engine/Private/WorldPartition/WorldPartitionRuntimeSpatialHash.cpp
--- a/Engine/Source/Runtime/Engine/Private/WorldPartition/WorldPartitionRuntimeSpatialHash.cpp
+++ b/Engine/Source/Runtime/Engine/Private/WorldPartition/WorldPartitionRuntimeSpatialHash.cpp
@@ -1188,7 +1188,7 @@ void UWorldPartitionRuntimeSpatialHash::PreSave(FObjectPreSaveContext ObjectSave
 
 FString UWorldPartitionRuntimeSpatialHash::GetCellCoordString(const FGridCellCoord& InCellGlobalCoord)
 {
-	return FString::Printf(TEXT("L%d_X%d_Y%d"), InCellGlobalCoord.Z, InCellGlobalCoord.X, InCellGlobalCoord.Y);
+	return FString::Printf(TEXT("L%" INT64_FMT "_X%" INT64_FMT "_Y%" INT64_FMT), InCellGlobalCoord.Z, InCellGlobalCoord.X, InCellGlobalCoord.Y);
 }
 
 #if WITH_EDITOR
diff --git a/Engine/Source/Runtime/Engine/Public/Net/RepLayout.h b/Engine/Source/Runtime/Engine/Public/Net/RepLayout.h
--- a/Engine/Source/Runtime/Engine/Public/Net/RepLayout.h
+++ b/Engine/Source/Runtime/Engine/Public/Net/RepLayout.h
@@ -195,7 +195,7 @@ public:
 
 	FString ToDebugString() const
 	{
-		return FString::Printf(TEXT("{Cmd: %u, Index: %u, Depth: %u, Ptr: %x}"), CmdIndex, ArrayIndex, ArrayDepth, DataPtr);
+		return FString::Printf(TEXT("{Cmd: %u, Index: %u, Depth: %u, Ptr: 0x%08" UPTRINT_x_FMT "}"), CmdIndex, ArrayIndex, ArrayDepth, (UPTRINT)DataPtr);
 	}
 
 	friend bool operator==(const FRepSharedPropertyKey& A, const FRepSharedPropertyKey& B)
diff --git a/Engine/Source/Runtime/Experimental/GeometryCollectionEngine/Public/GeometryCollection/GeometryCollectionISMPoolDebugDrawComponent.cpp b/Engine/Source/Runtime/Experimental/GeometryCollectionEngine/Public/GeometryCollection/GeometryCollectionISMPoolDebugDrawComponent.cpp
--- a/Engine/Source/Runtime/Experimental/GeometryCollectionEngine/Public/GeometryCollection/GeometryCollectionISMPoolDebugDrawComponent.cpp
+++ b/Engine/Source/Runtime/Experimental/GeometryCollectionEngine/Public/GeometryCollection/GeometryCollectionISMPoolDebugDrawComponent.cpp
@@ -341,7 +341,7 @@ void UGeometryCollectionISMPoolDebugDrawComponent::GetOnScreenMessages(TMultiMap
 		UStaticMesh const* StaticMesh = SelectedComponent->GetStaticMesh();
 		const int32 NumLods = StaticMesh ? StaticMesh->GetNumLODs() : 0;
 
-		FString ISMDescription = FString::Printf(TEXT("Type=%s   Count=%d   Bounds=(%f., %f., %f.)   CullDistance=(%f - %f)   NumLods= %d    NumMaterials=%d   CustomDataSize=%d"), 
+		FString ISMDescription = FString::Printf(TEXT("Type=%s   Count=%d   Bounds=(%f., %f., %f.)   CullDistance=(%d - %d)   NumLods= %d    NumMaterials=%d   CustomDataSize=%d"), 
 			ISMType, NumInstances, BoundsSize.X, BoundsSize.Z, BoundsSize.Z, StartCullDistance, EndCullDistance, NumLods, NumMaterials, NumCustomDataFloats);
 
 		for (int32 LodIndex = 0; LodIndex < NumLods; ++LodIndex)
diff --git a/Engine/Source/Runtime/GeometryCore/Public/Util/RefCountVector.h b/Engine/Source/Runtime/GeometryCore/Public/Util/RefCountVector.h
--- a/Engine/Source/Runtime/GeometryCore/Public/Util/RefCountVector.h
+++ b/Engine/Source/Runtime/GeometryCore/Public/Util/RefCountVector.h
@@ -497,7 +497,7 @@ public:
 
 	FString UsageStats() const
 	{
-		return FString::Printf(TEXT("RefCountSize %llu  FreeSize %llu  FreeMem %llukb"),
+		return FString::Printf(TEXT("RefCountSize %zu  FreeSize %zu  FreeMem %zukb"),
 			RefCounts.GetLength(), FreeIndices.GetLength(), (FreeIndices.GetByteCount() / 1024));
 	}
 
diff --git a/Engine/Source/Runtime/GeometryCore/Public/Util/SmallListSet.h b/Engine/Source/Runtime/GeometryCore/Public/Util/SmallListSet.h
--- a/Engine/Source/Runtime/GeometryCore/Public/Util/SmallListSet.h
+++ b/Engine/Source/Runtime/GeometryCore/Public/Util/SmallListSet.h
@@ -547,7 +547,7 @@ protected:
 public:
 	inline FString MemoryUsage() const
 	{
-		return FString::Printf(TEXT("ListSize %llu  Blocks Count %d  Free %llu  Mem %llukb   Linked Mem %llukb"),
+		return FString::Printf(TEXT("ListSize %zu  Blocks Count %d  Free %zu  Mem %zukb   Linked Mem %zukb"),
 			ListHeads.GetLength(), AllocatedCount, (FreeBlocks.GetLength() * sizeof(int32) / 1024),
 			ListBlocks.GetLength(), (LinkedListElements.GetLength() * sizeof(int32) / 1024));
 	}
diff --git a/Engine/Source/Runtime/MediaUtils/Private/MediaRecorder.cpp b/Engine/Source/Runtime/MediaUtils/Private/MediaRecorder.cpp
--- a/Engine/Source/Runtime/MediaUtils/Private/MediaRecorder.cpp
+++ b/Engine/Source/Runtime/MediaUtils/Private/MediaRecorder.cpp
@@ -293,7 +293,7 @@ void FMediaRecorder::TickRecording()
 			}
 			else
 			{
-				ImageTask->Filename = FString::Printf(TEXT("%s_%.16lu"), *BaseFilename, Sample->GetTime().Time.GetTicks());
+				ImageTask->Filename = FString::Printf(TEXT("%s_%.16" INT64_FMT), *BaseFilename, Sample->GetTime().Time.GetTicks());
 			}
 
 			ImageWriteQueue->Enqueue(MoveTemp(ImageTask), false);
diff --git a/Engine/Source/Runtime/MovieScene/Private/Evaluation/MovieSceneAnimTypeID.cpp b/Engine/Source/Runtime/MovieScene/Private/Evaluation/MovieSceneAnimTypeID.cpp
--- a/Engine/Source/Runtime/MovieScene/Private/Evaluation/MovieSceneAnimTypeID.cpp
+++ b/Engine/Source/Runtime/MovieScene/Private/Evaluation/MovieSceneAnimTypeID.cpp
@@ -4,7 +4,7 @@
 
 FString LexToString(const FMovieSceneAnimTypeID& AnimTypeID)
 {
-	return FString::Printf(TEXT("%#010x"), AnimTypeID.ID);
+	return FString::Printf(TEXT("%#010" UINT64_x_FMT), AnimTypeID.ID);
 }
 
 uint64 FMovieSceneAnimTypeID::Initialize(uint64* StaticPtr, uint32 Seed)
diff --git a/Engine/Source/Runtime/MovieSceneCapture/Private/MovieSceneCapture.cpp b/Engine/Source/Runtime/MovieSceneCapture/Private/MovieSceneCapture.cpp
--- a/Engine/Source/Runtime/MovieSceneCapture/Private/MovieSceneCapture.cpp
+++ b/Engine/Source/Runtime/MovieSceneCapture/Private/MovieSceneCapture.cpp
@@ -419,7 +419,7 @@ void UMovieSceneCapture::Initialize(TSharedPtr<FSceneViewport> InSceneViewport,
 	FormatMappings.Reserve(10);
 	if (FrameRate == FMath::RoundToDouble(FrameRate))
 	{
-		FormatMappings.Add(TEXT("fps"), FString::Printf(TEXT("%d"), FMath::RoundToInt(FrameRate)));
+		FormatMappings.Add(TEXT("fps"), FString::Printf(TEXT("%" INT64_FMT), FMath::RoundToInt(FrameRate)));
 	}
 	else
 	{
diff --git a/Engine/Source/Runtime/NetworkReplayStreaming/NullNetworkReplayStreaming/Private/NullNetworkReplayStreaming.cpp b/Engine/Source/Runtime/NetworkReplayStreaming/NullNetworkReplayStreaming/Private/NullNetworkReplayStreaming.cpp
--- a/Engine/Source/Runtime/NetworkReplayStreaming/NullNetworkReplayStreaming/Private/NullNetworkReplayStreaming.cpp
+++ b/Engine/Source/Runtime/NetworkReplayStreaming/NullNetworkReplayStreaming/Private/NullNetworkReplayStreaming.cpp
@@ -547,7 +547,7 @@ void FNullNetworkReplayStreamer::FlushCheckpoint(const uint32 TimeInMS)
 	{
 		FNullCheckpointListItem CheckpointEvent;
 		CheckpointEvent.Group = TEXT("checkpoint");
-		CheckpointEvent.Metadata = FString::Printf( TEXT("%ld"), FileAr->Tell() );
+		CheckpointEvent.Metadata = FString::Printf( TEXT("%" INT64_FMT), FileAr->Tell() );
 		CheckpointEvent.Time1 = TimeInMS;
 		CheckpointEvent.Time2 = TimeInMS;
 
diff --git a/Engine/Source/Runtime/Online/WebSockets/Private/Lws/LwsWebSocket.cpp b/Engine/Source/Runtime/Online/WebSockets/Private/Lws/LwsWebSocket.cpp
--- a/Engine/Source/Runtime/Online/WebSockets/Private/Lws/LwsWebSocket.cpp
+++ b/Engine/Source/Runtime/Online/WebSockets/Private/Lws/LwsWebSocket.cpp
@@ -309,7 +309,7 @@ int FLwsWebSocket::LwsCallback(lws* Instance, lws_callback_reasons Reason, void*
 
 			if (ReceiveBuffer.Len() + Length > MaxTextMessageBufferSize)
 			{
-				FString ReasonString = FString::Printf(TEXT("Received text message exceeded memory limit of %lu bytes"), MaxTextMessageBufferSize);
+				FString ReasonString = FString::Printf(TEXT("Received text message exceeded memory limit of %" UINT64_FMT " bytes"), MaxTextMessageBufferSize);
 				Close(LWS_CLOSE_STATUS_MESSAGE_TOO_LARGE, ReasonString);
 
 				UE_LOG(LogWebSockets, Verbose, TEXT("Received text message too large - use SetTextMessageMemoryLimit() to increase buffer size. Current Size=%lu"), MaxTextMessageBufferSize);
diff --git a/Engine/Source/Runtime/RHI/Private/PipelineFileCache.cpp b/Engine/Source/Runtime/RHI/Private/PipelineFileCache.cpp
--- a/Engine/Source/Runtime/RHI/Private/PipelineFileCache.cpp
+++ b/Engine/Source/Runtime/RHI/Private/PipelineFileCache.cpp
@@ -928,7 +928,7 @@ FString FPipelineCacheFileFormatPSO::CommonToString() const
 	Mask = UsageMask;
 	Count = BindCount;
 #endif
-	return FString::Printf(TEXT("\"%d,%llu\""), Count, Mask);
+	return FString::Printf(TEXT("\"%" INT64_FMT ",%" UINT64_FMT "\""), Count, Mask);
 }
 
 FString FPipelineCacheFileFormatPSO::ToStringReadable() const
diff --git a/Engine/Source/Runtime/RenderCore/Private/DumpGPU.cpp b/Engine/Source/Runtime/RenderCore/Private/DumpGPU.cpp
--- a/Engine/Source/Runtime/RenderCore/Private/DumpGPU.cpp
+++ b/Engine/Source/Runtime/RenderCore/Private/DumpGPU.cpp
@@ -489,7 +489,7 @@ public:
 		{
 			return TEXT("00000000000000000000");
 		}
-		return FString::Printf(TEXT("%04u%016x"), uint32(GraphBuilderIndex), static_cast<uint64>(reinterpret_cast<size_t>(Ptr)));
+		return FString::Printf(TEXT("%04u%016" UINT64_x_FMT), uint32(GraphBuilderIndex), static_cast<uint64>(reinterpret_cast<size_t>(Ptr)));
 	}
 
 	FString GetUniqueResourceName(const FRDGResource* Resource)
diff --git a/Engine/Source/Runtime/RenderCore/Private/ShaderCore.cpp b/Engine/Source/Runtime/RenderCore/Private/ShaderCore.cpp
--- a/Engine/Source/Runtime/RenderCore/Private/ShaderCore.cpp
+++ b/Engine/Source/Runtime/RenderCore/Private/ShaderCore.cpp
@@ -1421,15 +1421,15 @@ int HandleShaderCompileException(Windows::LPEXCEPTION_POINTERS Info, FString& Ou
 		if (ExceptionCodeToString(Info->ExceptionRecord->ExceptionCode, ExCodeStr))
 		{
 			OutExMsg = FString::Printf(
-				TEXT("Exception: %s, address=0x%016x\n"),
+				TEXT("Exception: %s, address=0x%016" UINT64_x_FMT "\n"),
 				*ExCodeStr,
 				(uint64)Info->ExceptionRecord->ExceptionAddress);
 		}
 		else
 		{
 			OutExMsg = FString::Printf(
-				TEXT("Exception code: 0x%08x, address=0x%016x\n"),
-				Info->ExceptionRecord->ExceptionCode,
+				TEXT("Exception code: 0x%08x, address=0x%016" UINT64_x_FMT "\n"),
+				(uint32)Info->ExceptionRecord->ExceptionCode,
 				(uint64)Info->ExceptionRecord->ExceptionAddress);
 		}
 	}
diff --git a/Engine/Source/Runtime/Renderer/Private/MeshDrawCommandStats.cpp b/Engine/Source/Runtime/Renderer/Private/MeshDrawCommandStats.cpp
--- a/Engine/Source/Runtime/Renderer/Private/MeshDrawCommandStats.cpp
+++ b/Engine/Source/Runtime/Renderer/Private/MeshDrawCommandStats.cpp
@@ -167,7 +167,7 @@ FMeshDrawCommandStatsManager::FMeshDrawCommandStatsManager()
 							FString& PassFriendlyNames = StatCollections[CategoryBudget.Collection].CategoryPassFriendlyNames[CategoryBudget.CategoryName];
 
 							FCoreDelegates::EOnScreenMessageSeverity Severity = CategoryBudget.PrimitiveBudget < *PrimitiveCount ? FCoreDelegates::EOnScreenMessageSeverity::Warning : FCoreDelegates::EOnScreenMessageSeverity::Info;
-							OutMessages.Add(Severity, FText::FromString(FString::Printf(TEXT("%5dK / %5dK - %s (%s)"), 
+							OutMessages.Add(Severity, FText::FromString(FString::Printf(TEXT("%5" UINT64_FMT "K / %5dK - %s (%s)"), 
 									*PrimitiveCount / 1000, 
 									CategoryBudget.PrimitiveBudget / 1000, 
 									*(CategoryBudget.CategoryName.ToString()),
diff --git a/Engine/Source/Runtime/Renderer/Private/RayTracing/RayTracingDebug.cpp b/Engine/Source/Runtime/Renderer/Private/RayTracing/RayTracingDebug.cpp
--- a/Engine/Source/Runtime/Renderer/Private/RayTracing/RayTracingDebug.cpp
+++ b/Engine/Source/Runtime/Renderer/Private/RayTracing/RayTracingDebug.cpp
@@ -1657,9 +1657,9 @@ void FDeferredShadingSceneRenderer::RayTracingDisplayPicking(const FRayTracingPi
 		Writer.EmptyLine();
 
 		FRayTracingAccelerationStructureSize SizeInfo = Geometry->GetSizeInfo();
-		Writer.DrawLine(FText::FromString(FString::Printf(TEXT("Result Size: %u"), SizeInfo.ResultSize)), 10, FColor::Yellow);
-		Writer.DrawLine(FText::FromString(FString::Printf(TEXT("Build Scratch Size: %u"), SizeInfo.BuildScratchSize)), 10, FColor::Yellow);
-		Writer.DrawLine(FText::FromString(FString::Printf(TEXT("Update Scratch Size: %u"), SizeInfo.UpdateScratchSize)), 10, FColor::Yellow);
+		Writer.DrawLine(FText::FromString(FString::Printf(TEXT("Result Size: %" UINT64_FMT), SizeInfo.ResultSize)), 10, FColor::Yellow);
+		Writer.DrawLine(FText::FromString(FString::Printf(TEXT("Build Scratch Size: %" UINT64_FMT), SizeInfo.BuildScratchSize)), 10, FColor::Yellow);
+		Writer.DrawLine(FText::FromString(FString::Printf(TEXT("Update Scratch Size: %" UINT64_FMT), SizeInfo.UpdateScratchSize)), 10, FColor::Yellow);
 	}
 	else
 	{
diff --git a/Engine/Source/Runtime/Slate/Public/Framework/Views/TableViewTypeTraits.h b/Engine/Source/Runtime/Slate/Public/Framework/Views/TableViewTypeTraits.h
--- a/Engine/Source/Runtime/Slate/Public/Framework/Views/TableViewTypeTraits.h
+++ b/Engine/Source/Runtime/Slate/Public/Framework/Views/TableViewTypeTraits.h
@@ -165,7 +165,7 @@ public:
 
 	static FString DebugDump( TSharedPtr<T> InPtr )
 	{
-		return InPtr.IsValid() ? FString::Printf(TEXT("0x%08x"), InPtr.Get()) : FString(TEXT("nullptr"));
+		return InPtr.IsValid() ? FString::Printf(TEXT("0x%08" UPTRINT_x_FMT), (UPTRINT)InPtr.Get()) : FString(TEXT("nullptr"));
 	}
 
 	class SerializerType{};
@@ -211,7 +211,7 @@ public:
 
 	static FString DebugDump( TSharedPtr<T, ESPMode::ThreadSafe> InPtr )
 	{
-		return InPtr.IsValid() ? FString::Printf(TEXT("0x%08x"), InPtr.Get()) : FString(TEXT("nullptr"));
+		return InPtr.IsValid() ? FString::Printf(TEXT("0x%08" UPTRINT_x_FMT), (UPTRINT)InPtr.Get()) : FString(TEXT("nullptr"));
 	}
 
 	class SerializerType{};
@@ -257,7 +257,7 @@ public:
 
 	static FString DebugDump( TSharedRef<T> InPtr )
 	{
-		return FString::Printf(TEXT("0x%08x"), &InPtr.Get());
+		return FString::Printf(TEXT("0x%08" UPTRINT_x_FMT), (UPTRINT)&InPtr.Get());
 	}
 
 	class SerializerType{};
@@ -303,7 +303,7 @@ public:
 
 	static FString DebugDump( TSharedRef<T, ESPMode::ThreadSafe> InPtr )
 	{
-		return FString::Printf(TEXT("0x%08x"), &InPtr.Get());
+		return FString::Printf(TEXT("0x%08" UPTRINT_x_FMT), (UPTRINT)&InPtr.Get());
 	}
 
 	class SerializerType{};
@@ -348,7 +348,7 @@ public:
 
 	static FString DebugDump( TWeakPtr<T> InPtr )
 	{
-		return InPtr.IsValid() ? FString::Printf(TEXT("0x%08x"), InPtr.Get()) : FString(TEXT("nullptr"));
+		return InPtr.IsValid() ? FString::Printf(TEXT("0x%08" UPTRINT_x_FMT), (UPTRINT)InPtr.Get()) : FString(TEXT("nullptr"));
 	}
 
 	class SerializerType{};
@@ -394,7 +394,7 @@ public:
 
 	static FString DebugDump( TWeakPtr<T, ESPMode::ThreadSafe> InPtr )
 	{
-		return InPtr.IsValid() ? FString::Printf(TEXT("0x%08x"), InPtr.Pin().Get()) : FString(TEXT("nullptr"));
+		return InPtr.IsValid() ? FString::Printf(TEXT("0x%08" UPTRINT_x_FMT), (UPTRINT)InPtr.Pin().Get()) : FString(TEXT("nullptr"));
 	}
 
 	class SerializerType{};
@@ -443,7 +443,7 @@ public:
 	static FString DebugDump( TWeakObjectPtr<T> InPtr )
 	{
 		T* ObjPtr = InPtr.Get();
-		return ObjPtr ? FString::Printf(TEXT("0x%08x [%s]"), ObjPtr, *ObjPtr->GetName()) : FString(TEXT("nullptr"));
+		return ObjPtr ? FString::Printf(TEXT("0x%08" UPTRINT_x_FMT " [%s]"), (UPTRINT)ObjPtr, *ObjPtr->GetName()) : FString(TEXT("nullptr"));
 	}
 
 	class SerializerType{};
@@ -504,7 +504,7 @@ public:
 
 	static FString DebugDump(TObjectPtr<T> InPtr)
 	{
-		return InPtr ? FString::Printf(TEXT("0x%08x [%s]"), InPtr.Get(), *InPtr->GetName()) : FString(TEXT("nullptr"));
+		return InPtr ? FString::Printf(TEXT("0x%08" UPTRINT_x_FMT " [%s]"), (UPTRINT)InPtr.Get(), *InPtr->GetName()) : FString(TEXT("nullptr"));
 	}
 
 	typedef FGCObject SerializerType;
@@ -556,7 +556,7 @@ public:
 
 	static FString DebugDump( T* InPtr )
 	{
-		return InPtr ? FString::Printf(TEXT("0x%08x [%s]"), InPtr, *InPtr->GetName()) : FString(TEXT("nullptr"));
+		return InPtr ? FString::Printf(TEXT("0x%08" UPTRINT_x_FMT " [%s]"), (UPTRINT)InPtr, *InPtr->GetName()) : FString(TEXT("nullptr"));
 	}
 
 	typedef FGCObject SerializerType;
@@ -602,7 +602,7 @@ public:
 
 	static FString DebugDump( const T* InPtr )
 	{
-		return InPtr ? FString::Printf(TEXT("0x%08x [%s]"), InPtr, *InPtr->GetName()) : FString(TEXT("nullptr"));
+		return InPtr ? FString::Printf(TEXT("0x%08" UPTRINT_x_FMT " [%s]"), (UPTRINT)InPtr, *InPtr->GetName()) : FString(TEXT("nullptr"));
 	}
 
 	typedef FGCObject SerializerType;
@@ -678,7 +678,7 @@ public:
 	static FString DebugDump(TFieldPath<T> InPtr)
 	{
 		T* ObjPtr = InPtr.Get();
-		return ObjPtr ? FString::Printf(TEXT("0x%08x [%s]"), ObjPtr, *ObjPtr->GetName()) : FString(TEXT("nullptr"));
+		return ObjPtr ? FString::Printf(TEXT("0x%08" UPTRINT_x_FMT " [%s]"), (UPTRINT)ObjPtr, *ObjPtr->GetName()) : FString(TEXT("nullptr"));
 	}
 
 	class SerializerType {};
diff --git a/Engine/Source/Runtime/SlateCore/Private/Debugging/ConsoleSlateDebuggerInvalidate.cpp b/Engine/Source/Runtime/SlateCore/Private/Debugging/ConsoleSlateDebuggerInvalidate.cpp
--- a/Engine/Source/Runtime/SlateCore/Private/Debugging/ConsoleSlateDebuggerInvalidate.cpp
+++ b/Engine/Source/Runtime/SlateCore/Private/Debugging/ConsoleSlateDebuggerInvalidate.cpp
@@ -761,7 +761,7 @@ void FConsoleSlateDebuggerInvalidate::HandlePaintDebugInfo(const FPaintArgs& InA
 			InOutDrawElements
 			, InOutLayerId
 			, InAllottedGeometry.ToPaintGeometry(FVector2f(1.f, 1.f), FSlateLayoutTransform(TextElementLocation))
-			, FString::Printf(TEXT("Slate Performance Threshold Reached: %d"), LastPerformanceThresholdFrameCount)
+			, FString::Printf(TEXT("Slate Performance Threshold Reached: %" UINT64_FMT), LastPerformanceThresholdFrameCount)
 			, NormalFontInfo
 			, ESlateDrawEffect::None
 			, FLinearColor::Red);
diff --git a/Engine/Source/Runtime/UMG/Private/Blueprint/GameViewportSubsystem.cpp b/Engine/Source/Runtime/UMG/Private/Blueprint/GameViewportSubsystem.cpp
--- a/Engine/Source/Runtime/UMG/Private/Blueprint/GameViewportSubsystem.cpp
+++ b/Engine/Source/Runtime/UMG/Private/Blueprint/GameViewportSubsystem.cpp
@@ -144,7 +144,7 @@ bool UGameViewportSubsystem::AddToScreen(UWidget* Widget, ULocalPlayer* Player,
 	UGameViewportClient* ViewportClient = World->GetGameViewport();
 	if (!ViewportClient)
 	{
-		FFrame::KismetExecutionMessage(*FString::Printf(TEXT("No game viewport was found."), *Widget->GetName()), ELogVerbosity::Warning);
+		FFrame::KismetExecutionMessage(TEXT("No game viewport was found."), ELogVerbosity::Warning);
 		return false;
 	}
 
diff --git a/Engine/Source/Runtime/VulkanRHI/Private/VulkanChunkedPipelineCache.cpp b/Engine/Source/Runtime/VulkanRHI/Private/VulkanChunkedPipelineCache.cpp
--- a/Engine/Source/Runtime/VulkanRHI/Private/VulkanChunkedPipelineCache.cpp
+++ b/Engine/Source/Runtime/VulkanRHI/Private/VulkanChunkedPipelineCache.cpp
@@ -850,7 +850,7 @@ public:
 			CacheSize = (uint32)TotalSize;
 
 			{
-				FScopedTimeToLog Timer4(FString::Printf(TEXT("FVulkanChunkedPipelineCacheManager: SavePSOCacheInternal serialize %s %d "), *BinaryCacheFileInfo.Filename, TotalSize));
+				FScopedTimeToLog Timer4(FString::Printf(TEXT("FVulkanChunkedPipelineCacheManager: SavePSOCacheInternal serialize %s %zu "), *BinaryCacheFileInfo.Filename, TotalSize));
 
 				SetCacheOffset((uint32)Archive.Tell());
 
diff --git a/Engine/Source/Runtime/VulkanRHI/Private/VulkanRHI.cpp b/Engine/Source/Runtime/VulkanRHI/Private/VulkanRHI.cpp
--- a/Engine/Source/Runtime/VulkanRHI/Private/VulkanRHI.cpp
+++ b/Engine/Source/Runtime/VulkanRHI/Private/VulkanRHI.cpp
@@ -882,7 +882,7 @@ void FVulkanDynamicRHI::SelectDevice()
 		}
 		else
 		{
-			GRHIAdapterUserDriverVersion = FString::Printf(TEXT("%d.%d.%d"), VK_VERSION_MAJOR(Props.driverVersion), VK_VERSION_MINOR(Props.driverVersion), VK_VERSION_PATCH(Props.driverVersion), Props.driverVersion);
+			GRHIAdapterUserDriverVersion = FString::Printf(TEXT("%d.%d.%d"), VK_VERSION_MAJOR(Props.driverVersion), VK_VERSION_MINOR(Props.driverVersion), VK_VERSION_PATCH(Props.driverVersion));
 		}
 
 		GRHIDeviceId = Props.deviceID;
diff --git a/Engine/Source/Runtime/Windows/AudioMixerXAudio2/Private/AudioMixerPlatformXAudio2.cpp b/Engine/Source/Runtime/Windows/AudioMixerXAudio2/Private/AudioMixerPlatformXAudio2.cpp
--- a/Engine/Source/Runtime/Windows/AudioMixerXAudio2/Private/AudioMixerPlatformXAudio2.cpp
+++ b/Engine/Source/Runtime/Windows/AudioMixerXAudio2/Private/AudioMixerPlatformXAudio2.cpp
@@ -45,7 +45,7 @@ THIRD_PARTY_INCLUDES_END
 
 #define XAUDIO2_LOG_RESULT(FunctionName, Result) \
 	{ \
-		FString ErrorString = FString::Printf(TEXT("%s -> 0x%X: %s (line: %d)"), TEXT( FunctionName ), Result, *Audio::ToErrorFString(Result), __LINE__); \
+		FString ErrorString = FString::Printf(TEXT("%s -> 0x%X: %s (line: %d)"), TEXT( FunctionName ), (uint32)Result, *Audio::ToErrorFString(Result), __LINE__); \
 		UE_LOG(LogAudioMixer, Error, TEXT("XAudio2 Error: %s"), *ErrorString);																		\
 	}
 
@@ -53,7 +53,7 @@ THIRD_PARTY_INCLUDES_END
 #define XAUDIO2_GOTO_CLEANUP_ON_FAIL(Result)																										 \
 	if (FAILED(Result))																														 \
 	{																																		 \
-		FString ErrorString = FString::Printf(TEXT("%s -> 0x%X: %s (line: %d)"), TEXT( #Result ), Result, *Audio::ToErrorFString(Result), __LINE__);\
+		FString ErrorString = FString::Printf(TEXT("%s -> 0x%X: %s (line: %d)"), TEXT( #Result ), (uint32)Result, *Audio::ToErrorFString(Result), __LINE__);\
 		UE_LOG(LogAudioMixer, Error, TEXT("XAudio2 Error: %s"), *ErrorString);																 \
 		goto Cleanup;																														 \
 	}
@@ -62,7 +62,7 @@ THIRD_PARTY_INCLUDES_END
 #define XAUDIO2_RETURN_ON_FAIL(Result)																										 \
 	if (FAILED(Result))																														 \
 	{																																		 \
-		FString ErrorString = FString::Printf(TEXT("%s -> 0x%X: %s (line: %d)"), TEXT( #Result ), Result, *Audio::ToErrorFString(Result), __LINE__);\
+		FString ErrorString = FString::Printf(TEXT("%s -> 0x%X: %s (line: %d)"), TEXT( #Result ), (uint32)Result, *Audio::ToErrorFString(Result), __LINE__);\
 		UE_LOG(LogAudioMixer, Error, TEXT("XAudio2 Error: %s"), *ErrorString);																 \
 		return false;																														 \
 	}
diff --git a/Engine/Source/Runtime/Windows/WindowsPlatformFeatures/Private/WindowsVideoRecordingSystem.cpp b/Engine/Source/Runtime/Windows/WindowsPlatformFeatures/Private/WindowsVideoRecordingSystem.cpp
--- a/Engine/Source/Runtime/Windows/WindowsPlatformFeatures/Private/WindowsVideoRecordingSystem.cpp
+++ b/Engine/Source/Runtime/Windows/WindowsPlatformFeatures/Private/WindowsVideoRecordingSystem.cpp
@@ -220,7 +220,7 @@ void FWindowsVideoRecordingSystem::NextRecording()
 {
 	if (Parameters.bAutoContinue)
 	{
-		CurrentFilename = FString::Printf(TEXT("%s_%d.mp4"), *BaseFilename, ++RecordingIndex);
+		CurrentFilename = FString::Printf(TEXT("%s_%" UINT64_FMT ".mp4"), *BaseFilename, ++RecordingIndex);
 	}
 	else
 	{
