Automate Test Case: Verify that Trace.SnapshotFile creates a small utrace from a Packaged Game session outside of an active profiling session

#rnx

diff --git a/Engine/Source/Developer/TraceInsights/Private/Insights/Tests/InsightsTestUtils.cpp b/Engine/Source/Developer/TraceInsights/Private/Insights/Tests/InsightsTestUtils.cpp
--- a/Engine/Source/Developer/TraceInsights/Private/Insights/Tests/InsightsTestUtils.cpp
+++ b/Engine/Source/Developer/TraceInsights/Private/Insights/Tests/InsightsTestUtils.cpp
@@ -177,7 +177,7 @@ bool FInsightsTestUtils::SetupUTS(double Timeout, bool bUseFork) const
 		return true;
 	}
 
-	FString UTSPath = FPlatformProcess::GenerateApplicationPath("UnrealTraceServer", EBuildConfiguration::Development);
+	const FString UTSPath = GetUTSPath();
 	if (!FPaths::FileExists(UTSPath))
 	{
 		Test->AddError(FString::Printf(TEXT("UTS executable can't be found at '%s'"), *UTSPath));
@@ -245,7 +245,7 @@ bool FInsightsTestUtils::KillUTS(double Timeout) const
 {
 	const FString UnrealTraceServerName = TEXT("UnrealTraceServer");
 
-	FString UTSPath = FPlatformProcess::GenerateApplicationPath("UnrealTraceServer", EBuildConfiguration::Development);
+	const FString UTSPath = GetUTSPath();
 	FString UTSParameters = TEXT("kill");
 	constexpr bool bLaunchDetached = true;
 	constexpr bool bLaunchHidden = false;
@@ -282,7 +282,10 @@ bool FInsightsTestUtils::KillUTS(double Timeout) const
 void FInsightsTestUtils::ResetSession() const
 {
 	TSharedPtr<FInsightsManager> InsightsManager = FInsightsManager::Get();
-	InsightsManager->ResetSession();
+	if (InsightsManager.IsValid())
+	{
+		InsightsManager->ResetSession();
+	}
 }
 
 ////////////////////////////////////////////////////////////////////////////////////////////////////
@@ -369,4 +372,36 @@ FString FInsightsTestUtils::GetLiveTrace(const TCHAR* Host, int32 Port) const
 	Test->AddInfo(TEXT("There isn't any live trace"));
 	delete StoreClient;
 	return TEXT("");
+}
+
+FString FInsightsTestUtils::FInsightsTestUtils::GetUTSPath() const
+{
+	#if PLATFORM_WINDOWS
+		const FString EnginePathToUTS = FString(TEXT("Engine/Binaries/Win64/UnrealTraceServer.exe"));
+	#endif
+
+	#if PLATFORM_MAC
+		const FString EnginePathToUTS = FString(TEXT("Engine/Binaries/Mac/UnrealTraceServer"));
+	#endif
+
+	#if PLATFORM_LINUX
+		const FString EnginePathToUTS = FString(TEXT("Engine/Binaries/Linux/UnrealTraceServer"));
+	#endif
+
+	FString RootDirectory = FPaths::RootDir();
+	FString EntireUTSPath = FPaths::Combine(*RootDirectory, EnginePathToUTS);
+	while (FPaths::DirectoryExists(RootDirectory) && !FPaths::FileExists(EntireUTSPath))
+	{
+		EntireUTSPath = FPaths::Combine(*RootDirectory, EnginePathToUTS);
+		if (!FPaths::FileExists(EntireUTSPath))
+		{
+			RootDirectory = FPaths::GetPath(RootDirectory);
+		}
+	}
+	if (!FPaths::DirectoryExists(RootDirectory)) {
+		Test->AddError(TEXT("Coudln't find UTS file"));
+		EntireUTSPath = FString();
+	}
+
+	return EntireUTSPath;
 }
\ No newline at end of file
diff --git a/Engine/Source/Developer/TraceInsights/Public/Insights/Tests/InsightsTestUtils.h b/Engine/Source/Developer/TraceInsights/Public/Insights/Tests/InsightsTestUtils.h
--- a/Engine/Source/Developer/TraceInsights/Public/Insights/Tests/InsightsTestUtils.h
+++ b/Engine/Source/Developer/TraceInsights/Public/Insights/Tests/InsightsTestUtils.h
@@ -21,6 +21,7 @@ public:
 	bool IsTraceHasLiveStatus(const FString& TraceName, const TCHAR* Host = TEXT("localhost"), int32 Port = 0U) const;
 	void ResetSession() const;
 	FString GetLiveTrace(const TCHAR* Host = TEXT("localhost"), int32 Port = 0U) const;
+	FString GetUTSPath() const;
 
 private:
 	FAutomationTestBase* Test;
