[Insights] Add/remove instance series to graph track works if mode is Game/Rendering Frame

diff --git a/Engine/Source/Developer/TraceInsights/Private/Insights/Widgets/STimersView.cpp b/Engine/Source/Developer/TraceInsights/Private/Insights/Widgets/STimersView.cpp
--- a/Engine/Source/Developer/TraceInsights/Private/Insights/Widgets/STimersView.cpp
+++ b/Engine/Source/Developer/TraceInsights/Private/Insights/Widgets/STimersView.cpp
@@ -859,7 +859,7 @@ void STimersView::TreeView_BuildPlotTimerMenu(FMenuBuilder& MenuBuilder)
 	{
 		FUIAction Action_ToggleTimerInGraphTrack;
 		Action_ToggleTimerInGraphTrack.CanExecuteAction = FCanExecuteAction::CreateLambda(CanExecuteAddToGraphTrack);
-		Action_ToggleTimerInGraphTrack.ExecuteAction = FExecuteAction::CreateSP(this, &STimersView::ToggleTimingViewMainGraphEventSeries, SelectedNode);
+		Action_ToggleTimerInGraphTrack.ExecuteAction = FExecuteAction::CreateSP(this, &STimersView::ToggleTimingViewMainGraphEventInstanceSeries, SelectedNode);
 
 		if (SelectedNode.IsValid() &&
 			SelectedNode->GetType() != ETimerNodeType::Group &&
