[Insights] Memory Insights:
  - Fixed crash when first encountered callstack is an empty callstack with a non-zero id.
  - Added distinction between the default empty callstack (i.e. callstack with id 0; "no recorded callstack") and other empty callstacks (with a valid id > 0, but with zero number of frames). Also changed the display names for empty callstacks to be localizable texts (ex: "No Callstack Recorded", "Empty Callstack").
  - Fixed wrong truncation in case the traced callstack has a larger number of stack frames than supported (max 255 stack frames are supported by analysis).

diff --git a/Engine/Source/Developer/TraceInsights/Private/Insights/MemoryProfiler/ViewModels/CallstackFormatting.cpp b/Engine/Source/Developer/TraceInsights/Private/Insights/MemoryProfiler/ViewModels/CallstackFormatting.cpp
new file mode 100644
--- /dev/null
+++ b/Engine/Source/Developer/TraceInsights/Private/Insights/MemoryProfiler/ViewModels/CallstackFormatting.cpp
@@ -0,0 +1,109 @@
+// Copyright Epic Games, Inc. All Rights Reserved.
+
+#include "CallstackFormatting.h"
+
+#define LOCTEXT_NAMESPACE "Insights::CallstackFormatting"
+
+namespace Insights
+{
+
+////////////////////////////////////////////////////////////////////////////////////////////////////
+
+FText GetCallstackNotAvailableString()
+{
+	return LOCTEXT("UnknownCallstack", "Unknown Callstack");
+}
+
+////////////////////////////////////////////////////////////////////////////////////////////////////
+
+FText GetNoCallstackString()
+{
+	return LOCTEXT("NoCallstackRecorded", "No Callstack Recorded");
+}
+
+////////////////////////////////////////////////////////////////////////////////////////////////////
+
+FText GetEmptyCallstackString()
+{
+	return LOCTEXT("EmptyCallstack", "Empty Callstack");
+}
+
+////////////////////////////////////////////////////////////////////////////////////////////////////
+
+void FormatStackFrame(const TraceServices::FStackFrame& Frame, FStringBuilderBase& OutString, EStackFrameFormatFlags FormatFlags)
+{
+	using namespace TraceServices;
+	const ESymbolQueryResult Result = Frame.Symbol->GetResult();
+	switch (Result)
+	{
+		case ESymbolQueryResult::OK:
+			if (EnumHasAnyFlags(FormatFlags, EStackFrameFormatFlags::Module))
+			{
+				OutString.Append(Frame.Symbol->Module);
+				if (FormatFlags != EStackFrameFormatFlags::Module &&
+					FormatFlags != (EStackFrameFormatFlags::Module | EStackFrameFormatFlags::Multiline))
+				{
+					OutString.AppendChar(EnumHasAnyFlags(FormatFlags, EStackFrameFormatFlags::Multiline) ? TEXT('\n') : TEXT('!'));
+				}
+			}
+			if (EnumHasAnyFlags(FormatFlags, EStackFrameFormatFlags::Symbol))
+			{
+				OutString.Append(Frame.Symbol->Name);
+			}
+			if (EnumHasAnyFlags(FormatFlags, EStackFrameFormatFlags::Symbol) &&
+				EnumHasAnyFlags(FormatFlags, EStackFrameFormatFlags::File | EStackFrameFormatFlags::Line))
+			{
+				OutString.AppendChar(EnumHasAnyFlags(FormatFlags, EStackFrameFormatFlags::Multiline) ? TEXT('\n') : TEXT(' '));
+			}
+			if (EnumHasAnyFlags(FormatFlags, EStackFrameFormatFlags::File))
+			{
+				OutString.Append(Frame.Symbol->File);
+			}
+			if (EnumHasAnyFlags(FormatFlags, EStackFrameFormatFlags::Line))
+			{
+				OutString.Appendf(TEXT("(%d)"), Frame.Symbol->Line);
+			}
+			break;
+
+		case ESymbolQueryResult::Mismatch:
+		case ESymbolQueryResult::NotFound:
+		case ESymbolQueryResult::NotLoaded:
+			if (EnumHasAnyFlags(FormatFlags, EStackFrameFormatFlags::Module))
+			{
+				OutString.Append(Frame.Symbol->Module);
+				if (FormatFlags != EStackFrameFormatFlags::Module &&
+					FormatFlags != (EStackFrameFormatFlags::Module | EStackFrameFormatFlags::Multiline))
+				{
+					OutString.AppendChar(EnumHasAnyFlags(FormatFlags, EStackFrameFormatFlags::Multiline) ? TEXT('\n') : TEXT('!'));
+				}
+			}
+			if (EnumHasAnyFlags(FormatFlags, EStackFrameFormatFlags::Symbol))
+			{
+				if (Frame.Addr == 0)
+				{
+					OutString.Append(TEXT("0x00000000"));
+				}
+				else
+				{
+					OutString.Appendf(TEXT("0x%llX"), Frame.Addr);
+				}
+			}
+			if (EnumHasAnyFlags(FormatFlags, EStackFrameFormatFlags::ModuleAndSymbol))
+			{
+				OutString.AppendChar(EnumHasAnyFlags(FormatFlags, EStackFrameFormatFlags::Multiline) ? TEXT('\n') : TEXT(' '));
+			}
+			OutString.Appendf(TEXT("(%s)"), QueryResultToString((Result)));
+			break;
+
+		case ESymbolQueryResult::Pending:
+		default:
+			OutString.Append(QueryResultToString(Result));
+			break;
+	}
+}
+
+////////////////////////////////////////////////////////////////////////////////////////////////////
+
+} // namespace Insights
+
+#undef LOCTEXT_NAMESPACE
diff --git a/Engine/Source/Developer/TraceInsights/Private/Insights/MemoryProfiler/ViewModels/CallstackFormatting.h b/Engine/Source/Developer/TraceInsights/Private/Insights/MemoryProfiler/ViewModels/CallstackFormatting.h
--- a/Engine/Source/Developer/TraceInsights/Private/Insights/MemoryProfiler/ViewModels/CallstackFormatting.h
+++ b/Engine/Source/Developer/TraceInsights/Private/Insights/MemoryProfiler/ViewModels/CallstackFormatting.h
@@ -2,11 +2,13 @@
 
 #pragma once
 
+#include "Internationalization/Text.h"
 #include "Misc/EnumClassFlags.h"
 #include "TraceServices/Model/Callstack.h"
 #include "TraceServices/Model/Modules.h"
 
-////////////////////////////////////////////////////////////////////////////////////////////////////
+namespace Insights
+{
 
 enum class EStackFrameFormatFlags : uint8
 {
@@ -22,92 +24,10 @@ enum class EStackFrameFormatFlags : uint8
 };
 ENUM_CLASS_FLAGS(EStackFrameFormatFlags);
 
-////////////////////////////////////////////////////////////////////////////////////////////////////
-
-inline const TCHAR* GetCallstackNotAvailableString()
-{
-	return TEXT("Unknown Callstack");
-}
-
-////////////////////////////////////////////////////////////////////////////////////////////////////
-
-inline const TCHAR* GetEmptyCallstackString()
-{
-	return TEXT("Empty Callstack");
-}
-
-////////////////////////////////////////////////////////////////////////////////////////////////////
-
-inline void FormatStackFrame(const TraceServices::FStackFrame& Frame, FStringBuilderBase& OutString, EStackFrameFormatFlags FormatFlags)
-{
-	using namespace TraceServices;
-	const ESymbolQueryResult Result = Frame.Symbol->GetResult();
-	switch (Result)
-	{
-		case ESymbolQueryResult::OK:
-			if (EnumHasAnyFlags(FormatFlags, EStackFrameFormatFlags::Module))
-			{
-				OutString.Append(Frame.Symbol->Module);
-				if (FormatFlags != EStackFrameFormatFlags::Module &&
-					FormatFlags != (EStackFrameFormatFlags::Module | EStackFrameFormatFlags::Multiline))
-				{
-					OutString.AppendChar(EnumHasAnyFlags(FormatFlags, EStackFrameFormatFlags::Multiline) ? TEXT('\n') : TEXT('!'));
-				}
-			}
-			if (EnumHasAnyFlags(FormatFlags, EStackFrameFormatFlags::Symbol))
-			{
-				OutString.Append(Frame.Symbol->Name);
-			}
-			if (EnumHasAnyFlags(FormatFlags, EStackFrameFormatFlags::Symbol) &&
-				EnumHasAnyFlags(FormatFlags, EStackFrameFormatFlags::File | EStackFrameFormatFlags::Line))
-			{
-				OutString.AppendChar(EnumHasAnyFlags(FormatFlags, EStackFrameFormatFlags::Multiline) ? TEXT('\n') : TEXT(' '));
-			}
-			if (EnumHasAnyFlags(FormatFlags, EStackFrameFormatFlags::File))
-			{
-				OutString.Append(Frame.Symbol->File);
-			}
-			if (EnumHasAnyFlags(FormatFlags, EStackFrameFormatFlags::Line))
-			{
-				OutString.Appendf(TEXT("(%d)"), Frame.Symbol->Line);
-			}
-			break;
-
-		case ESymbolQueryResult::Mismatch:
-		case ESymbolQueryResult::NotFound:
-		case ESymbolQueryResult::NotLoaded:
-			if (EnumHasAnyFlags(FormatFlags, EStackFrameFormatFlags::Module))
-			{
-				OutString.Append(Frame.Symbol->Module);
-				if (FormatFlags != EStackFrameFormatFlags::Module &&
-					FormatFlags != (EStackFrameFormatFlags::Module | EStackFrameFormatFlags::Multiline))
-				{
-					OutString.AppendChar(EnumHasAnyFlags(FormatFlags, EStackFrameFormatFlags::Multiline) ? TEXT('\n') : TEXT('!'));
-				}
-			}
-			if (EnumHasAnyFlags(FormatFlags, EStackFrameFormatFlags::Symbol))
-			{
-				if (Frame.Addr == 0)
-				{
-					OutString.Append(TEXT("0x00000000"));
-				}
-				else
-				{
-					OutString.Appendf(TEXT("0x%llX"), Frame.Addr);
-				}
-			}
-			if (EnumHasAnyFlags(FormatFlags, EStackFrameFormatFlags::ModuleAndSymbol))
-			{
-				OutString.AppendChar(EnumHasAnyFlags(FormatFlags, EStackFrameFormatFlags::Multiline) ? TEXT('\n') : TEXT(' '));
-			}
-			OutString.Appendf(TEXT("(%s)"), QueryResultToString((Result)));
-			break;
+FText GetCallstackNotAvailableString();
+FText GetNoCallstackString();
+FText GetEmptyCallstackString();
 
-		case ESymbolQueryResult::Pending:
-		default:
-			OutString.Append(QueryResultToString(Result));
-			break;
-	}
-}
+void FormatStackFrame(const TraceServices::FStackFrame& Frame, FStringBuilderBase& OutString, EStackFrameFormatFlags FormatFlags);
 
-////////////////////////////////////////////////////////////////////////////////////////////////////
+} // namespace Insights
diff --git a/Engine/Source/Developer/TraceInsights/Private/Insights/MemoryProfiler/ViewModels/MemAllocGroupingByCallstack.cpp b/Engine/Source/Developer/TraceInsights/Private/Insights/MemoryProfiler/ViewModels/MemAllocGroupingByCallstack.cpp
--- a/Engine/Source/Developer/TraceInsights/Private/Insights/MemoryProfiler/ViewModels/MemAllocGroupingByCallstack.cpp
+++ b/Engine/Source/Developer/TraceInsights/Private/Insights/MemoryProfiler/ViewModels/MemAllocGroupingByCallstack.cpp
@@ -99,6 +99,7 @@ void FMemAllocGroupingByCallstack::GroupNodes(const TArray<FTableTreeNodePtr>& N
 	CallstackGroups.Add(Root);
 
 	FTableTreeNode* UnsetGroupPtr = nullptr;
+	FTableTreeNode* NoCallstackGroupPtr = nullptr;
 	FTableTreeNode* EmptyCallstackGroupPtr = nullptr;
 	TMap<const TraceServices::FCallstack*, FCallstackGroup*> GroupMapByCallstack; // Callstack* -> FCallstackGroup*
 
@@ -116,7 +117,8 @@ void FMemAllocGroupingByCallstack::GroupNodes(const TArray<FTableTreeNodePtr>& N
 		}
 
 		FCallstackGroup* GroupPtr = Root;
-		int32 NumFrames = -1;
+		bool bIsEmptyCallstack = false;
+		bool bIsDefaultEmptyCallstack = false;
 
 		const FMemAllocNode& MemAllocNode = static_cast<const FMemAllocNode&>(*NodePtr);
 		const FMemoryAlloc* Alloc = MemAllocNode.GetMemAlloc();
@@ -131,63 +133,68 @@ void FMemAllocGroupingByCallstack::GroupNodes(const TArray<FTableTreeNodePtr>& N
 			}
 			else if (Callstack)
 			{
-				NumFrames = static_cast<int32>(Callstack->Num());
-				check(NumFrames <= 256); // see Callstack->Frame(uint8)
-
-				for (int32 FrameDepth = NumFrames - 1; FrameDepth >= 0; --FrameDepth)
+				if (Callstack->Num() == 0)
 				{
-					const TraceServices::FStackFrame* Frame = Callstack->Frame(static_cast<uint8>(bIsInverted ? NumFrames - FrameDepth - 1 : FrameDepth));
-					check(Frame != nullptr);
-
-					if (bLocalShouldSkipFilteredFrames &&
-						Frame->Symbol->FilterStatus.load() == TraceServices::EResolvedSymbolFilterStatus::Filtered)
-					{
-						continue;
-					}
-
-					if (bLocalIsGroupingByFunction)
+					bIsEmptyCallstack = true;
+					bIsDefaultEmptyCallstack = (Callstack->GetEmptyId() == 0);
+				}
+				else
+				{
+					int32 NumFrames = static_cast<int32>(Callstack->Num());
+					check(NumFrames > 0);
+					check(NumFrames <= 256); // see Callstack->Frame(uint8)
+					for (int32 FrameDepth = NumFrames - 1; FrameDepth >= 0; --FrameDepth)
 					{
-						const FName GroupName = GetGroupName(Frame);
+						const TraceServices::FStackFrame* Frame = Callstack->Frame(static_cast<uint8>(bIsInverted ? NumFrames - FrameDepth - 1 : FrameDepth));
+						check(Frame != nullptr);
 
-						// Merge with parent group, if it has the same name (i.e. same function).
-						if (GroupPtr->Parent != nullptr && GroupPtr->Name == GroupName)
+						if (bLocalShouldSkipFilteredFrames &&
+							Frame->Symbol->FilterStatus.load() == TraceServices::EResolvedSymbolFilterStatus::Filtered)
 						{
-							GroupPtr = GroupPtr->Parent;
+							continue;
 						}
 
-						// Merge groups by name.
-						FCallstackGroup** GroupPtrPtr = GroupPtr->GroupMapByName.Find(GroupName);
-						if (!GroupPtrPtr)
-						{
-							GroupPtr = CreateGroup(CallstackGroups, GroupPtr, GroupName, InParentTable, Frame);
-							check(GroupPtr->Parent != nullptr);
-							GroupPtr->Parent->GroupMapByName.Add(GroupName, GroupPtr);
-						}
-						else
-						{
-							GroupPtr = *GroupPtrPtr;
-						}
-					}
-					else
-					{
-						// Merge groups by unique callstack frame.
-						FCallstackGroup** GroupPtrPtr = GroupPtr->GroupMap.Find(Frame->Addr);
-						if (!GroupPtrPtr)
+						if (bLocalIsGroupingByFunction)
 						{
 							const FName GroupName = GetGroupName(Frame);
-							GroupPtr = CreateGroup(CallstackGroups, GroupPtr, GroupName, InParentTable, Frame);
-							check(GroupPtr->Parent != nullptr);
-							GroupPtr->Parent->GroupMap.Add(Frame->Addr, GroupPtr);
+
+							// Merge with parent group, if it has the same name (i.e. same function).
+							if (GroupPtr->Parent != nullptr && GroupPtr->Name == GroupName)
+							{
+								GroupPtr = GroupPtr->Parent;
+							}
+
+							// Merge groups by name.
+							FCallstackGroup** GroupPtrPtr = GroupPtr->GroupMapByName.Find(GroupName);
+							if (!GroupPtrPtr)
+							{
+								GroupPtr = CreateGroup(CallstackGroups, GroupPtr, GroupName, InParentTable, Frame);
+								check(GroupPtr->Parent != nullptr);
+								GroupPtr->Parent->GroupMapByName.Add(GroupName, GroupPtr);
+							}
+							else
+							{
+								GroupPtr = *GroupPtrPtr;
+							}
 						}
 						else
 						{
-							GroupPtr = *GroupPtrPtr;
+							// Merge groups by unique callstack frame.
+							FCallstackGroup** GroupPtrPtr = GroupPtr->GroupMap.Find(Frame->Addr);
+							if (!GroupPtrPtr)
+							{
+								const FName GroupName = GetGroupName(Frame);
+								GroupPtr = CreateGroup(CallstackGroups, GroupPtr, GroupName, InParentTable, Frame);
+								check(GroupPtr->Parent != nullptr);
+								GroupPtr->Parent->GroupMap.Add(Frame->Addr, GroupPtr);
+							}
+							else
+							{
+								GroupPtr = *GroupPtrPtr;
+							}
 						}
 					}
-				}
 
-				if (NumFrames > 0)
-				{
 					GroupMapByCallstack.Add(Callstack, GroupPtr);
 				}
 			}
@@ -197,13 +204,24 @@ void FMemAllocGroupingByCallstack::GroupNodes(const TArray<FTableTreeNodePtr>& N
 		{
 			GroupPtr->Node->AddChildAndSetParent(NodePtr);
 		}
-		else if (NumFrames == 0)
+		else if (bIsEmptyCallstack)
 		{
-			if (!EmptyCallstackGroupPtr)
+			if (bIsDefaultEmptyCallstack)
+			{
+				if (!NoCallstackGroupPtr)
+				{
+					NoCallstackGroupPtr = CreateNoCallstackGroup(InParentTable, ParentGroup);
+				}
+				NoCallstackGroupPtr->AddChildAndSetParent(NodePtr);
+			}
+			else
 			{
-				EmptyCallstackGroupPtr = CreateEmptyCallstackGroup(InParentTable, ParentGroup);
+				if (!EmptyCallstackGroupPtr)
+				{
+					EmptyCallstackGroupPtr = CreateEmptyCallstackGroup(InParentTable, ParentGroup);
+				}
+				EmptyCallstackGroupPtr->AddChildAndSetParent(NodePtr);
 			}
-			EmptyCallstackGroupPtr->AddChildAndSetParent(NodePtr);
 		}
 		else
 		{
@@ -273,8 +291,19 @@ FMemAllocGroupingByCallstack::FCallstackGroup* FMemAllocGroupingByCallstack::Cre
 
 FTableTreeNode* FMemAllocGroupingByCallstack::CreateUnsetGroup(TWeakPtr<FTable> ParentTable, FTableTreeNode& Parent) const
 {
-	static FName NotAvailableName(TEXT("N/A"));
-	FTableTreeNodePtr NodePtr = MakeShared<FTableTreeNode>(NotAvailableName, ParentTable);
+	static FName GroupName(GetCallstackNotAvailableString().ToString());
+	FTableTreeNodePtr NodePtr = MakeShared<FTableTreeNode>(GroupName, ParentTable);
+	NodePtr->SetExpansion(false);
+	Parent.AddChildAndSetParent(NodePtr);
+	return NodePtr.Get();
+}
+
+////////////////////////////////////////////////////////////////////////////////////////////////////
+
+FTableTreeNode* FMemAllocGroupingByCallstack::CreateNoCallstackGroup(TWeakPtr<FTable> ParentTable, FTableTreeNode& Parent) const
+{
+	static FName GroupName(GetNoCallstackString().ToString());
+	FTableTreeNodePtr NodePtr = MakeShared<FTableTreeNode>(GroupName, ParentTable);
 	NodePtr->SetExpansion(false);
 	Parent.AddChildAndSetParent(NodePtr);
 	return NodePtr.Get();
@@ -284,8 +313,8 @@ FTableTreeNode* FMemAllocGroupingByCallstack::CreateUnsetGroup(TWeakPtr<FTable>
 
 FTableTreeNode* FMemAllocGroupingByCallstack::CreateEmptyCallstackGroup(TWeakPtr<FTable> ParentTable, FTableTreeNode& Parent) const
 {
-	static FName NotAvailableName(GetEmptyCallstackString());
-	FTableTreeNodePtr NodePtr = MakeShared<FTableTreeNode>(NotAvailableName, ParentTable);
+	static FName GroupName(GetEmptyCallstackString().ToString());
+	FTableTreeNodePtr NodePtr = MakeShared<FTableTreeNode>(GroupName, ParentTable);
 	NodePtr->SetExpansion(false);
 	Parent.AddChildAndSetParent(NodePtr);
 	return NodePtr.Get();
diff --git a/Engine/Source/Developer/TraceInsights/Private/Insights/MemoryProfiler/ViewModels/MemAllocGroupingByCallstack.h b/Engine/Source/Developer/TraceInsights/Private/Insights/MemoryProfiler/ViewModels/MemAllocGroupingByCallstack.h
--- a/Engine/Source/Developer/TraceInsights/Private/Insights/MemoryProfiler/ViewModels/MemAllocGroupingByCallstack.h
+++ b/Engine/Source/Developer/TraceInsights/Private/Insights/MemoryProfiler/ViewModels/MemAllocGroupingByCallstack.h
@@ -90,6 +90,7 @@ private:
 	FCallstackGroup* CreateGroup(TArray<FCallstackGroup*>& InOutAllCallstackGroup, FCallstackGroup* InParentGroup, const FName InGroupName, TWeakPtr<FTable> InParentTable, const TraceServices::FStackFrame* InFrame) const;
 
 	FTableTreeNode* CreateUnsetGroup(TWeakPtr<FTable> ParentTable, FTableTreeNode& Parent) const;
+	FTableTreeNode* CreateNoCallstackGroup(TWeakPtr<FTable> ParentTable, FTableTreeNode& Parent) const;
 	FTableTreeNode* CreateEmptyCallstackGroup(TWeakPtr<FTable> ParentTable, FTableTreeNode& Parent) const;
 
 private:
diff --git a/Engine/Source/Developer/TraceInsights/Private/Insights/MemoryProfiler/ViewModels/MemAllocNode.cpp b/Engine/Source/Developer/TraceInsights/Private/Insights/MemoryProfiler/ViewModels/MemAllocNode.cpp
--- a/Engine/Source/Developer/TraceInsights/Private/Insights/MemoryProfiler/ViewModels/MemAllocNode.cpp
+++ b/Engine/Source/Developer/TraceInsights/Private/Insights/MemoryProfiler/ViewModels/MemAllocNode.cpp
@@ -78,19 +78,26 @@ FText FMemAllocNode::GetFullCallstackOrSourceFiles(ECallstackType InCallstackTyp
 	}
 
 	const Insights::FMemoryAlloc& Alloc = GetMemAllocChecked();
-	const TraceServices::FCallstack* Callstack = 
+	const TraceServices::FCallstack* Callstack =
 		InCallstackType == ECallstackType::AllocCallstack
-			? Alloc.GetAllocCallstack() 
+			? Alloc.GetAllocCallstack()
 			: Alloc.GetFreeCallstack();
 
 	if (!Callstack)
 	{
-		return FText::FromString(GetCallstackNotAvailableString());
+		return GetCallstackNotAvailableString();
 	}
 
 	if (Callstack->Num() == 0)
 	{
-		return FText::FromString(GetEmptyCallstackString());
+		if (Callstack->GetEmptyId() == 0)
+		{
+			return GetNoCallstackString();
+		}
+		else
+		{
+			return GetEmptyCallstackString();
+		}
 	}
 
 	TStringBuilder<1024> Tooltip;
@@ -119,24 +126,31 @@ FText FMemAllocNode::GetTopFunctionOrSourceFile(ECallstackType InCallstackType,
 	}
 
 	const Insights::FMemoryAlloc& Alloc = GetMemAllocChecked();
-	const TraceServices::FCallstack* Callstack = 
+	const TraceServices::FCallstack* Callstack =
 		InCallstackType == ECallstackType::AllocCallstack
-			? Alloc.GetAllocCallstack() 
+			? Alloc.GetAllocCallstack()
 			: Alloc.GetFreeCallstack();
 
 	if (!Callstack)
 	{
-		return FText::FromString(GetCallstackNotAvailableString());
+		return GetCallstackNotAvailableString();
 	}
 
-	const uint32 NumCallstackFrames = Callstack->Num();
-	if (NumCallstackFrames == 0)
+	if (Callstack->Num() == 0)
 	{
-		return FText::FromString(GetEmptyCallstackString());
+		if (Callstack->GetEmptyId() == 0)
+		{
+			return GetNoCallstackString();
+		}
+		else
+		{
+			return GetEmptyCallstackString();
+		}
 	}
-	check(NumCallstackFrames <= 256); // see Callstack->Frame(uint8)
 
 	const TraceServices::FStackFrame* Frame = nullptr;
+	const uint32 NumCallstackFrames = Callstack->Num();
+	check(NumCallstackFrames <= 256); // see Callstack->Frame(uint8)
 	for (uint32 FrameIndex = 0; FrameIndex < NumCallstackFrames; ++FrameIndex)
 	{
 		Frame = Callstack->Frame(static_cast<uint8>(FrameIndex));
@@ -149,6 +163,7 @@ FText FMemAllocNode::GetTopFunctionOrSourceFile(ECallstackType InCallstackType,
 			break;
 		}
 	}
+	check(Frame != nullptr);
 
 	TStringBuilder<1024> Str;
 	FormatStackFrame(*Frame, Str, (EStackFrameFormatFlags)Flags);
diff --git a/Engine/Source/Developer/TraceInsights/Private/Insights/MemoryProfiler/Widgets/SMemAllocTableTreeView.cpp b/Engine/Source/Developer/TraceInsights/Private/Insights/MemoryProfiler/Widgets/SMemAllocTableTreeView.cpp
--- a/Engine/Source/Developer/TraceInsights/Private/Insights/MemoryProfiler/Widgets/SMemAllocTableTreeView.cpp
+++ b/Engine/Source/Developer/TraceInsights/Private/Insights/MemoryProfiler/Widgets/SMemAllocTableTreeView.cpp
@@ -2117,10 +2117,11 @@ void SMemAllocTableTreeView::ExportMemorySnapshot() const
 			if (Alloc)
 			{
 				const TraceServices::FCallstack* Callstack = bIsAllocCallstack ? Alloc->AllocCallstack : Alloc->FreeCallstack;
+
 				if (!Callstack)
 				{
 					OutData << QuotationMarkBegin;
-					OutData << GetCallstackNotAvailableString();
+					OutData << GetCallstackNotAvailableString().ToString();
 					OutData << QuotationMarkEnd;
 					return;
 				}
@@ -2128,14 +2129,21 @@ void SMemAllocTableTreeView::ExportMemorySnapshot() const
 				if (Callstack->Num() == 0)
 				{
 					OutData << QuotationMarkBegin;
-					OutData << GetEmptyCallstackString();
+					if (Callstack->GetEmptyId() == 0)
+					{
+						OutData << GetNoCallstackString().ToString();
+					}
+					else
+					{
+						OutData << GetEmptyCallstackString().ToString();
+					}
 					OutData << QuotationMarkEnd;
 					return;
 				}
 
-				const uint32 NumCallstackFrames = Callstack->Num();
-				check(NumCallstackFrames <= 256);
 				OutData << QuotationMarkBegin;
+				const uint32 NumCallstackFrames = Callstack->Num();
+				check(NumCallstackFrames <= 256); // see Callstack->Frame(uint8)
 				for (uint32 Index = 0; Index < NumCallstackFrames; ++Index)
 				{
 					if (Index != 0)
diff --git a/Engine/Source/Developer/TraceServices/Private/Analyzers/CallstacksAnalysis.cpp b/Engine/Source/Developer/TraceServices/Private/Analyzers/CallstacksAnalysis.cpp
--- a/Engine/Source/Developer/TraceServices/Private/Analyzers/CallstacksAnalysis.cpp
+++ b/Engine/Source/Developer/TraceServices/Private/Analyzers/CallstacksAnalysis.cpp
@@ -33,14 +33,15 @@ bool FCallstacksAnalyzer::OnEvent(uint16 RouteId, EStyle Style, const FOnEventCo
 	{
 		case RouteId_Callstack:
 			const TArrayReader<uint64>& Frames = Context.EventData.GetArray<uint64>("Frames");
+			uint8 NumFrames = (uint8)FMath::Min(255u, Frames.Num());
 			if (const uint32 Id = Context.EventData.GetValue<uint32>("CallstackId"))
 			{
-				Provider->AddCallstack(Id, Frames.GetData(), uint8(Frames.Num()));
+				Provider->AddCallstack(Id, Frames.GetData(), NumFrames);
 			}
 			// Backward compatibility with legacy memory trace format (5.0-EA).
 			else if (const uint64 Hash = Context.EventData.GetValue<uint64>("Id"))
 			{
-				Provider->AddCallstackWithHash(Hash, Frames.GetData(), uint8(Frames.Num()));
+				Provider->AddCallstackWithHash(Hash, Frames.GetData(), NumFrames);
 			}
 			break;
 	}
diff --git a/Engine/Source/Developer/TraceServices/Private/Model/CallstacksProvider.cpp b/Engine/Source/Developer/TraceServices/Private/Model/CallstacksProvider.cpp
--- a/Engine/Source/Developer/TraceServices/Private/Model/CallstacksProvider.cpp
+++ b/Engine/Source/Developer/TraceServices/Private/Model/CallstacksProvider.cpp
@@ -16,9 +16,7 @@ namespace TraceServices
 
 static const FResolvedSymbol GNeverResolveSymbol(ESymbolQueryResult::NotLoaded, nullptr, nullptr, nullptr, 0, EResolvedSymbolFilterStatus::NotFiltered);
 static const FResolvedSymbol GNotFoundSymbol(ESymbolQueryResult::NotFound, TEXT("Unknown"), nullptr, nullptr, 0, EResolvedSymbolFilterStatus::NotFiltered);
-static const FResolvedSymbol GNoSymbol(ESymbolQueryResult::NotFound, TEXT("No callstack recorded"), nullptr, nullptr, 0, EResolvedSymbolFilterStatus::NotFiltered);
 static constexpr FStackFrame GNotFoundStackFrame = { 0, &GNotFoundSymbol };
-static constexpr FStackFrame GNoStackFrame = { 0, &GNoSymbol };
 static const FCallstack GNotFoundCallstack(&GNotFoundStackFrame, 1);
 
 ////////////////////////////////////////////////////////////////////////////////////////////////////
@@ -40,9 +38,8 @@ FCallstacksProvider::FCallstacksProvider(IAnalysisSession& InSession)
 	, Callstacks(InSession.GetLinearAllocator(), CallstacksPerPage)
 	, Frames(InSession.GetLinearAllocator(), FramesPerPage)
 {
-	// Let the first callstack to be an empty/undefined callstack.
+	// Let the first callstack to be the default empty callstack (i.e. CallstackId == 0, "callstack not recorded").
 	FCallstack& FirstCallstack = Callstacks.PushBack();
-	FirstCallstack.Init(&GNoStackFrame, 1);
 }
 
 ////////////////////////////////////////////////////////////////////////////////////////////////////
@@ -78,24 +75,23 @@ void FCallstacksProvider::AddCallstack(uint32 InCallstackId, const uint64* InFra
 			uint64 EntriesToAdd = PageHeadroom + 1; // Fill page and allocate one on next
 			do { Frames.PushBack(); } while (--EntriesToAdd);
 		}
-	}
-
-	// Append the incoming frames.
-	const uint64 FirstFrame = Frames.Num();
-	for (uint32 FrameIdx = 0; FrameIdx < InFrameCount; ++FrameIdx)
-	{
-		FStackFrame& F = Frames.PushBack();
-		F.Addr = InFrames[FrameIdx];
 
-		if (ModuleProvider)
+		// Append the incoming frames.
+		for (uint32 FrameIdx = 0; FrameIdx < InFrameCount; ++FrameIdx)
 		{
-			// This will return immediately. The result will be empty if the symbol
-			// has not been encountered before, and resolution has been queued up.
-			F.Symbol = ModuleProvider->GetSymbol(InFrames[FrameIdx]);
-		}
-		else
-		{
-			F.Symbol = &GNeverResolveSymbol;
+			FStackFrame& F = Frames.PushBack();
+			F.Addr = InFrames[FrameIdx];
+
+			if (ModuleProvider)
+			{
+				// This will return immediately. The result will be empty if the symbol
+				// has not been encountered before, and resolution has been queued up.
+				F.Symbol = ModuleProvider->GetSymbol(InFrames[FrameIdx]);
+			}
+			else
+			{
+				F.Symbol = &GNeverResolveSymbol;
+			}
 		}
 	}
 
@@ -111,11 +107,16 @@ void FCallstacksProvider::AddCallstack(uint32 InCallstackId, const uint64* InFra
 			while (InCallstackId >= Callstacks.Num())
 			{
 				Callstack = &Callstacks.PushBack();
-				Callstack->Init(&GNoStackFrame, 1);
+				Callstack->InitEmpty(Callstacks.Num() - 1);
 			}
 		}
-		check(Callstack && (Callstack->Num() == 0 || Callstack->Frame(0) == &GNoStackFrame)); // adding same callstack id twice?
-		Callstack->Init(&Frames[FirstFrame], InFrameCount);
+		check(Callstack);
+		check(Callstack->IsEmpty() && Callstack->GetEmptyId() == uint64(InCallstackId)); // adding same callstack id twice?
+		if (InFrameCount > 0)
+		{
+			check(InFrameCount <= Frames.Num());
+			Callstack->Init(&Frames[Frames.Num() - InFrameCount], InFrameCount);
+		}
 	}
 }
 
diff --git a/Engine/Source/Developer/TraceServices/Public/TraceServices/Model/Callstack.h b/Engine/Source/Developer/TraceServices/Public/TraceServices/Model/Callstack.h
--- a/Engine/Source/Developer/TraceServices/Public/TraceServices/Model/Callstack.h
+++ b/Engine/Source/Developer/TraceServices/Public/TraceServices/Model/Callstack.h
@@ -21,18 +21,36 @@ struct TRACESERVICES_API FStackFrame
 /////////////////////////////////////////////////////////////////////
 struct TRACESERVICES_API FCallstack
 {
-	/** Creates an empty callstack. */
+	/** Creates the default empty callstack (id 0). */
 	FCallstack();
+
+	/** Creates an empty callstack but with a non-zero id. */
+	FCallstack(uint64 Id);
+
 	/** Creates an callstack initialized with a certain number of stack frames. */
 	FCallstack(const FStackFrame* FirstEntry, uint8 FrameCount);
+
 	/** Initializes the callstack with a certain number of stack frames. */
 	void Init(const FStackFrame* FirstEntry, uint8 FrameCount);
+
+	/** Initializes an empty callstack but with a non-zero id. */
+	void InitEmpty(uint64 Id);
+
+	/** Gets the callstack id of an empty callstack. */
+	uint64 GetEmptyId() const;
+
+	/** Return true if the number of stack frames is zero. */
+	bool IsEmpty() const;
+
 	/** Gets the number of stack frames in callstack. */
 	uint32 Num() const;
+
 	/** Gets the address at a given stack depth. */
 	uint64 Addr(uint8 Depth) const;
+
 	/** Gets the cached symbol name at a given stack depth. */
 	const TCHAR* Name(uint8 Depth) const;
+
 	/** Gets the entire frame at given depth. */
 	const FStackFrame* Frame(uint8 Depth) const;
 
@@ -88,6 +106,12 @@ inline FCallstack::FCallstack()
 {
 }
 
+/////////////////////////////////////////////////////////////////////
+inline FCallstack::FCallstack(uint64 Id)
+{
+	InitEmpty(Id);
+}
+
 /////////////////////////////////////////////////////////////////////
 inline FCallstack::FCallstack(const FStackFrame* InFirstFrame, uint8 InFrameCount)
 {
@@ -101,6 +125,25 @@ inline void FCallstack::Init(const FStackFrame* InFirstFrame, uint8 InFrameCount
 	CallstackLenIndex = (uint64(InFrameCount) << EntryLenShift) | (~EntryLenMask & uint64(InFirstFrame));
 }
 
+/////////////////////////////////////////////////////////////////////
+inline void FCallstack::InitEmpty(uint64 Id)
+{
+	check((Id & EntryLenMask) == 0);
+	CallstackLenIndex = Id & ~EntryLenMask;
+}
+
+/////////////////////////////////////////////////////////////////////
+inline uint64 FCallstack::GetEmptyId() const
+{
+	return CallstackLenIndex & ~EntryLenMask;
+}
+
+/////////////////////////////////////////////////////////////////////
+inline bool FCallstack::IsEmpty() const
+{
+	return (CallstackLenIndex >> EntryLenShift) == 0;
+}
+
 /////////////////////////////////////////////////////////////////////
 inline uint32 FCallstack::Num() const
 {
