[Insights]
  - LogView: Added visual validation for the issue where log messages are traced out of order. Log messages with non monotonic increasing timestamp will show in red in the Log View.
  - LogView: Minor code optimizations and code refactoring in preparation for further code changes.

diff --git a/Engine/Source/Developer/TraceInsights/Private/Insights/Widgets/SLogView.cpp b/Engine/Source/Developer/TraceInsights/Private/Insights/Widgets/SLogView.cpp
--- a/Engine/Source/Developer/TraceInsights/Private/Insights/Widgets/SLogView.cpp
+++ b/Engine/Source/Developer/TraceInsights/Private/Insights/Widgets/SLogView.cpp
@@ -293,6 +293,16 @@ public:
 			FLogMessageRecord& CacheEntry = ParentWidgetPin->GetCache().Get(LogMessagePin->GetIndex());
 			const double Time = CacheEntry.GetTime();
 
+			// Verify if time is monotonically increasing.
+			if (LogMessagePin->GetIndex() > 0)
+			{
+				FLogMessageRecord& PrevCacheEntry = ParentWidgetPin->GetCache().Get(LogMessagePin->GetIndex() - 1);
+				if (Time < PrevCacheEntry.GetTime())
+				{
+					return FSlateColor(FLinearColor(1.0f, 0.3f, 0.3f, 1.0f));
+				}
+			}
+
 			TSharedPtr<STimingProfilerWindow> Window = FTimingProfilerManager::Get()->GetProfilerWindow();
 			if (Window)
 			{
@@ -303,12 +313,10 @@ public:
 					{
 						if (IsSelected)
 						{
-							//return FSlateColor(FLinearColor(0.0f, 0.1f, 0.5f, 1.0f));
 							return FSlateColor(FLinearColor(0.0f, 0.05f, 0.2f, 1.0f));
 						}
 						else
 						{
-							//return FSlateColor(FLinearColor(0.2f, 0.4f, 0.8f, 1.0f));
 							return FSlateColor(FLinearColor(0.4f, 0.8f, 1.6f, 1.0f));
 						}
 					}
@@ -567,7 +575,7 @@ void SLogView::Reset()
 
 	Cache.Reset();
 
-	Messages.Reset();
+	FilteredMessages.Reset();
 
 	ListView->RebuildList();
 }
@@ -659,7 +667,7 @@ void SLogView::Construct(const FArguments& InArgs)
 					.SelectionMode(ESelectionMode::Single)
 					.OnMouseButtonClick(this, &SLogView::OnMouseButtonClick)
 					.OnSelectionChanged(this, &SLogView::OnSelectionChanged)
-					.ListItemsSource(&Messages)
+					.ListItemsSource(&FilteredMessages)
 					.OnGenerateRow(this, &SLogView::OnGenerateRow)
 					.ConsumeMouseWheel(EConsumeMouseWheel::Always)
 					.OnContextMenuOpening(FOnContextMenuOpening::CreateSP(this, &SLogView::ListView_GetContextMenu))
@@ -769,10 +777,10 @@ void SLogView::Tick(const FGeometry& AllottedGeometry, const double InCurrentTim
 			TMap<FName, int32> DuplicatedCategories;
 			LogProvider.EnumerateCategories([&Categories, &DuplicatedCategories](const TraceServices::FLogCategoryInfo& Category)
 			{
-				FString CategoryStr(Category.Name);
+				FStringView CategoryStr(Category.Name);
 				if (CategoryStr.StartsWith(TEXT("Log")))
 				{
-					CategoryStr.RightChopInline(3, EAllowShrinking::No);
+					CategoryStr.RightChopInline(3);
 				}
 				FName CategoryName(CategoryStr);
 				if (Categories.Contains(CategoryName))
@@ -805,7 +813,7 @@ void SLogView::Tick(const FGeometry& AllottedGeometry, const double InCurrentTim
 			UE_LOG(TraceInsights, Log, TEXT("[LogView] Unique Log Categories: %d"), NumAvailableLogCategories);
 
 			//Cache.Reset();
-			Messages.Reset();
+			FilteredMessages.Reset();
 			TotalNumMessages = 0;
 			//ListView->RebuildList();
 			bIsDirty = true;
@@ -864,7 +872,7 @@ void SLogView::Tick(const FGeometry& AllottedGeometry, const double InCurrentTim
 
 				for (int32 Index = TotalNumMessages; Index < NewMessageCount; Index++)
 				{
-					Messages.Add(MakeShared<FLogMessage>(Index));
+					FilteredMessages.Add(MakeShared<FLogMessage>(Index));
 				}
 
 				const int32 NumAddedMessages = NewMessageCount - TotalNumMessages;
@@ -874,7 +882,7 @@ void SLogView::Tick(const FGeometry& AllottedGeometry, const double InCurrentTim
 				if (SelectedLogMessage.IsValid())
 				{
 					// Restore selection.
-					SelectedLogMessageByLogIndex(SelectedLogMessage->GetIndex());
+					SelectLogMessageByLogIndex(SelectedLogMessage->GetIndex());
 				}
 				bIsDirty = false;
 				DirtyStopwatch.Reset();
@@ -898,7 +906,7 @@ void SLogView::Tick(const FGeometry& AllottedGeometry, const double InCurrentTim
 			UE_LOG(TraceInsights, Log, TEXT("[LogView] RESET"));
 #endif // !WITH_EDITOR
 			Cache.Reset();
-			Messages.Reset();
+			FilteredMessages.Reset();
 			TotalNumMessages = 0;
 			ListView->RebuildList();
 			bIsDirty = (NewMessageCount != 0);
@@ -930,13 +938,13 @@ void SLogView::Tick(const FGeometry& AllottedGeometry, const double InCurrentTim
 			FilteringChangeNumber == Filter.GetChangeNumber())
 		{
 			FLogFilteringAsyncTask& Task = FilteringAsyncTask->GetTask();
-			const TArray<uint32>& FilteredMessages = Task.GetFilteredMessages();
+			const TArray<uint32>& FilteredMessageIndices = Task.GetFilteredMessages();
 
-			// Add filtered messages to current Messages array.
-			const int32 NumFilteredMessages = FilteredMessages.Num();
+			// Add filtered messages to current FilteredMessages array.
+			const int32 NumFilteredMessages = FilteredMessageIndices.Num();
 			for (int32 Index = 0; Index < NumFilteredMessages; Index++)
 			{
-				Messages.Add(MakeShared<FLogMessage>(FilteredMessages[Index]));
+				FilteredMessages.Add(MakeShared<FLogMessage>(FilteredMessageIndices[Index]));
 			}
 
 			TotalNumMessages = Task.GetEndIndex();
@@ -945,7 +953,7 @@ void SLogView::Tick(const FGeometry& AllottedGeometry, const double InCurrentTim
 			if (SelectedLogMessage.IsValid())
 			{
 				// Restore selection.
-				SelectedLogMessageByLogIndex(SelectedLogMessage->GetIndex());
+				SelectLogMessageByLogIndex(SelectedLogMessage->GetIndex());
 			}
 			bIsDirty = false;
 			DirtyStopwatch.Reset();
@@ -980,20 +988,41 @@ TSharedPtr<FLogMessage> SLogView::GetSelectedLogMessage() const
 
 ////////////////////////////////////////////////////////////////////////////////////////////////////
 
-void SLogView::SelectedLogMessageByLogIndex(int32 LogIndex)
+void SLogView::SelectLogMessage(TSharedPtr<FLogMessage> LogMessage)
 {
-	// We are assuming the Messages list is sorted by log index...
-	int32 MessageIndex = Algo::BinarySearchBy(Messages, LogIndex, &FLogMessage::GetIndex);
-	if (MessageIndex != INDEX_NONE)
+	if (LogMessage.IsValid())
 	{
-		ListView->SetItemSelection(Messages[MessageIndex], true);
-		ListView->RequestScrollIntoView(Messages[MessageIndex]);
+		if (!ListView->IsItemSelected(LogMessage))
+		{
+			ListView->SetItemSelection(LogMessage, true, ESelectInfo::Direct);
+		}
+		ListView->RequestScrollIntoView(LogMessage);
 	}
 }
 
 ////////////////////////////////////////////////////////////////////////////////////////////////////
 
-void SLogView::SelectLogMessage(TSharedPtr<FLogMessage> LogMessage)
+void SLogView::SelectLogMessageByLogIndex(int32 LogIndex)
+{
+	if (FilteredMessages.Num() == 0)
+	{
+		return;
+	}
+
+	// We are assuming the FilteredMessages list is sorted by log index.
+	// Find the exact match of the log index in the filtered messages list.
+	int32 MessageIndex = Algo::BinarySearchBy(FilteredMessages, LogIndex, &FLogMessage::GetIndex);
+	if (MessageIndex == INDEX_NONE)
+	{
+		return;
+	}
+
+	SelectLogMessage(FilteredMessages[MessageIndex]);
+}
+
+////////////////////////////////////////////////////////////////////////////////////////////////////
+
+void SLogView::OnSelectedLogMessageChanged(TSharedPtr<FLogMessage> LogMessage)
 {
 	if (LogMessage.IsValid())
 	{
@@ -1023,7 +1052,7 @@ void SLogView::SelectLogMessage(TSharedPtr<FLogMessage> LogMessage)
 
 void SLogView::OnMouseButtonClick(TSharedPtr<FLogMessage> LogMessage)
 {
-	SelectLogMessage(LogMessage);
+	OnSelectedLogMessageChanged(LogMessage);
 }
 
 ////////////////////////////////////////////////////////////////////////////////////////////////////
@@ -1033,7 +1062,7 @@ void SLogView::OnSelectionChanged(TSharedPtr<FLogMessage> LogMessage, ESelectInf
 	if (SelectInfo != ESelectInfo::Direct &&
 		SelectInfo != ESelectInfo::OnMouseClick)
 	{
-		SelectLogMessage(LogMessage);
+		OnSelectedLogMessageChanged(LogMessage);
 	}
 }
 
@@ -1065,7 +1094,7 @@ void SLogView::OnFilterChanged()
 	UE_LOG(TraceInsights, Log, TEXT("[LogView] OnFilterChanged: \"%s\""), *FilterText);
 #endif // !WITH_EDITOR
 	Cache.Reset();
-	Messages.Reset();
+	FilteredMessages.Reset();
 	TotalNumMessages = 0;
 	bIsDirty = true;
 	DirtyStopwatch.Start();
@@ -1076,13 +1105,13 @@ void SLogView::OnFilterChanged()
 
 void SLogView::UpdateStatsText()
 {
-	if (Messages.Num() == TotalNumMessages)
+	if (FilteredMessages.Num() == TotalNumMessages)
 	{
 		StatsText = FText::Format(LOCTEXT("StatsText1", "{0} logs"), FText::AsNumber(TotalNumMessages));
 	}
 	else
 	{
-		StatsText = FText::Format(LOCTEXT("StatsText2", "{0} / {1} logs"), FText::AsNumber(Messages.Num()), FText::AsNumber(TotalNumMessages));
+		StatsText = FText::Format(LOCTEXT("StatsText2", "{0} / {1} logs"), FText::AsNumber(FilteredMessages.Num()), FText::AsNumber(TotalNumMessages));
 	}
 }
 
@@ -1686,7 +1715,7 @@ void SLogView::CopyMessage() const
 
 bool SLogView::CanCopyRange() const
 {
-	if (Messages.Num() == 0)
+	if (FilteredMessages.Num() == 0)
 	{
 		return false;
 	}
@@ -1707,7 +1736,7 @@ bool SLogView::CanCopyRange() const
 
 void SLogView::CopyRange() const
 {
-	if (Messages.Num() == 0)
+	if (FilteredMessages.Num() == 0)
 	{
 		return;
 	}
@@ -1738,7 +1767,7 @@ void SLogView::CopyRange() const
 		TraceServices::FAnalysisSessionReadScope SessionReadScope(*Session.Get());
 		const TraceServices::ILogProvider& LogProvider = TraceServices::ReadLogProvider(*Session.Get());
 
-		for (const TSharedPtr<FLogMessage>& Message : Messages)
+		for (const TSharedPtr<FLogMessage>& Message : FilteredMessages)
 		{
 			LogProvider.ReadMessage(Message->GetIndex(), [this, SelectionStartTime, SelectionEndTime, &StringBuilder, &NumMessagesInSelectedRange](const TraceServices::FLogMessageInfo& MessageInfo)
 			{
@@ -1764,14 +1793,14 @@ void SLogView::CopyRange() const
 
 bool SLogView::CanCopyAll() const
 {
-	return Messages.Num() > 0;
+	return FilteredMessages.Num() > 0;
 }
 
 ////////////////////////////////////////////////////////////////////////////////////////////////////
 
 void SLogView::CopyAll() const
 {
-	if (Messages.Num() == 0)
+	if (FilteredMessages.Num() == 0)
 	{
 		return;
 	}
@@ -1786,7 +1815,7 @@ void SLogView::CopyAll() const
 		TraceServices::FAnalysisSessionReadScope SessionReadScope(*Session.Get());
 		const TraceServices::ILogProvider& LogProvider = TraceServices::ReadLogProvider(*Session.Get());
 
-		for (const TSharedPtr<FLogMessage>& Message : Messages)
+		for (const TSharedPtr<FLogMessage>& Message : FilteredMessages)
 		{
 			LogProvider.ReadMessage(Message->GetIndex(), [this, &StringBuilder](const TraceServices::FLogMessageInfo& MessageInfo)
 			{
@@ -1798,7 +1827,7 @@ void SLogView::CopyAll() const
 	}
 
 	FPlatformApplicationMisc::ClipboardCopy(StringBuilder.ToString());
-	UE_LOG(TraceInsights, Log, TEXT("Copied %d logs to clipboard."), Messages.Num());
+	UE_LOG(TraceInsights, Log, TEXT("Copied %d logs to clipboard."), FilteredMessages.Num());
 }
 
 ////////////////////////////////////////////////////////////////////////////////////////////////////
@@ -1820,7 +1849,7 @@ void SLogView::SaveRange() const
 
 bool SLogView::CanSaveAll() const
 {
-	return Messages.Num() > 0;
+	return FilteredMessages.Num() > 0;
 }
 
 ////////////////////////////////////////////////////////////////////////////////////////////////////
@@ -1835,7 +1864,7 @@ void SLogView::SaveAll() const
 
 void SLogView::SaveLogsToFile(bool bSaveLogsInSelectedRangeOnly) const
 {
-	if (Messages.Num() == 0)
+	if (FilteredMessages.Num() == 0)
 	{
 		return;
 	}
@@ -1920,7 +1949,7 @@ void SLogView::SaveLogsToFile(bool bSaveLogsInSelectedRangeOnly) const
 		TraceServices::FAnalysisSessionReadScope SessionReadScope(*Session.Get());
 		const TraceServices::ILogProvider& LogProvider = TraceServices::ReadLogProvider(*Session.Get());
 
-		for (const TSharedPtr<FLogMessage>& Message : Messages)
+		for (const TSharedPtr<FLogMessage>& Message : FilteredMessages)
 		{
 			LogProvider.ReadMessage(Message->GetIndex(), [this, bSaveLogsInSelectedRangeOnly, SelectionStartTime, SelectionEndTime, Separator, ExportFileHandle, &NumMessagesInSelectedRange](const TraceServices::FLogMessageInfo& MessageInfo)
 			{
diff --git a/Engine/Source/Developer/TraceInsights/Private/Insights/Widgets/SLogView.h b/Engine/Source/Developer/TraceInsights/Private/Insights/Widgets/SLogView.h
--- a/Engine/Source/Developer/TraceInsights/Private/Insights/Widgets/SLogView.h
+++ b/Engine/Source/Developer/TraceInsights/Private/Insights/Widgets/SLogView.h
@@ -71,7 +71,9 @@ public:
 	FLogMessageCache& GetCache() { return Cache; }
 
 	TSharedPtr<FLogMessage> GetSelectedLogMessage() const;
-	void SelectedLogMessageByLogIndex(int32 LogIndex);
+	void SelectLogMessage(TSharedPtr<FLogMessage> LogMessage);
+	void SelectLogMessageByLogIndex(int32 LogIndex);
+	void SelectLogMessageByClosestTime(double Time);
 
 	FText GetFilterText() const { return FilterTextBox->GetText(); }
 
@@ -83,9 +85,9 @@ protected:
 
 	void InitCommandList();
 
-	void SelectLogMessage(TSharedPtr<FLogMessage> LogMessage);
 	void OnMouseButtonClick(TSharedPtr<FLogMessage> LogMessage);
 	void OnSelectionChanged(TSharedPtr<FLogMessage> LogMessage, ESelectInfo::Type SelectInfo);
+	void OnSelectedLogMessageChanged(TSharedPtr<FLogMessage> LogMessage);
 
 	void FilterTextBox_OnTextChanged(const FText& InFilterText);
 	void OnFilterChanged();
@@ -107,7 +109,7 @@ protected:
 
 	bool ShowHideAllCategories_IsChecked() const;
 	void ShowHideAllCategories_Execute();
-	
+
 	bool IsLogCategoryEnabled(FName InName) const;
 	void ToggleCategory(FName InName);
 
@@ -188,8 +190,8 @@ protected:
 	/** Cached log messages. */
 	mutable FLogMessageCache Cache;
 
-	/** List of trace log messages to show in list view (i.e. filtered). */
-	TArray<TSharedPtr<FLogMessage>> Messages; // TODO: this needs virtualisation (an a new SListView)
+	/** List of trace log messages to show in list view. */
+	TArray<TSharedPtr<FLogMessage>> FilteredMessages;
 };
 
 ////////////////////////////////////////////////////////////////////////////////////////////////////
