[Insights]
  - Timing View / Markers track: Made few FString/FStringView optimizations for handling name of log categories.
  - Timing View / Markers track: Optimized FLogProvider::EnumerateMessages to only iterate the log messages in the specified time interval (uses binary search to find start index and stops iteration at first log with time > EndTime). This significantly improves performance of Markers track (Bookmarks / Logs) when there are high number of log messages.

diff --git a/Engine/Source/Developer/TraceInsights/Private/Insights/ViewModels/MarkersTimingTrack.cpp b/Engine/Source/Developer/TraceInsights/Private/Insights/ViewModels/MarkersTimingTrack.cpp
--- a/Engine/Source/Developer/TraceInsights/Private/Insights/ViewModels/MarkersTimingTrack.cpp
+++ b/Engine/Source/Developer/TraceInsights/Private/Insights/ViewModels/MarkersTimingTrack.cpp
@@ -92,14 +92,12 @@ void FMarkersTimingTrack::UpdateTrackNameAndHeight()
 {
 	if (bUseOnlyBookmarks)
 	{
-		const FString NameString = TEXT("Bookmarks");
-		SetName(NameString);
+		SetName(TEXT("Bookmarks"));
 		SetHeight(14.0f);
 	}
 	else
 	{
-		const FString NameString = TEXT("Logs");
-		SetName(NameString);
+		SetName(TEXT("Logs"));
 		SetHeight(28.0f);
 	}
 
@@ -675,10 +673,10 @@ void FTimeMarkerTrackBuilder::AddLogMessage(const TraceServices::FLogMessageInfo
 {
 	Track.NumLogMessages++;
 
-	// Add also the log message imediately on the left of the screen (if any).
+	// Add also the log message immediately on the left of the screen (if any).
 	if (Track.NumLogMessages == 1 && Message.Index > 0)
 	{
-		// Note: Reading message at Index-1 will not work as expected when using filter!
+		// Note: The log message at Index-1 may fail the filter (i.e. not a bookmark or screenshot), so it might not be displayed!
 		//TODO: Search API like: LogProviderPtr->SearchMessage(StartIndex, ESearchDirection::Backward, LambdaPredicate, bResolveFormatString);
 		LogProviderPtr->ReadMessage(
 			Message.Index - 1,
@@ -686,7 +684,9 @@ void FTimeMarkerTrackBuilder::AddLogMessage(const TraceServices::FLogMessageInfo
 	}
 
 	check(Message.Category != nullptr);
-	if (!Track.bUseOnlyBookmarks || Message.Category == Track.BookmarkCategory || Message.Category == Track.ScreenshotCategory)
+	if (!Track.bUseOnlyBookmarks ||
+		Message.Category == Track.BookmarkCategory ||
+		Message.Category == Track.ScreenshotCategory)
 	{
 		float X = Viewport.TimeToSlateUnitsRounded(Message.Time);
 		if (X < 0.0f)
@@ -704,14 +704,14 @@ void FTimeMarkerTrackBuilder::AddLogMessage(const TraceServices::FLogMessageInfo
 FLinearColor FTimeMarkerTrackBuilder::GetColorByCategory(const TCHAR* const Category)
 {
 	// Strip the "Log" prefix.
-	FString CategoryStr(Category);
+	FStringView CategoryStr(Category);
 	if (CategoryStr.StartsWith(TEXT("Log")))
 	{
-		CategoryStr.RightChopInline(3, EAllowShrinking::No);
+		CategoryStr.RightChopInline(3);
 	}
 
 	uint32 Hash = 0;
-	for (const TCHAR* c = *CategoryStr; *c; ++c)
+	for (const TCHAR* c = CategoryStr.GetData(); *c; ++c)
 	{
 		Hash = (Hash + *c) * 0x2c2c57ed;
 	}
@@ -777,10 +777,10 @@ void FTimeMarkerTrackBuilder::Flush(float AvailableTextW)
 		if (AvailableTextW > 6.0f)
 		{
 			// Strip the "Log" prefix.
-			FString CategoryStr(LastCategory);
+			FStringView CategoryStr(LastCategory);
 			if (CategoryStr.StartsWith(TEXT("Log")))
 			{
-				CategoryStr.RightChopInline(3, EAllowShrinking::No);
+				CategoryStr.RightChopInline(3);
 			}
 
 			const int32 HorizontalOffset = FMath::RoundToInt((AvailableTextW - 2.0f) * FontScale);
@@ -797,7 +797,7 @@ void FTimeMarkerTrackBuilder::Flush(float AvailableTextW)
 				TextInfo.LogIndex = LastLogIndex;
 				if (LastWholeCharacterIndexCategory >= 0)
 				{
-					TextInfo.Category.AppendChars(*CategoryStr, LastWholeCharacterIndexCategory + 1);
+					TextInfo.Category.AppendChars(CategoryStr.GetData(), LastWholeCharacterIndexCategory + 1);
 				}
 				if (LastWholeCharacterIndexMessage >= 0)
 				{
diff --git a/Engine/Source/Developer/TraceInsights/Private/Insights/Widgets/STimingView.cpp b/Engine/Source/Developer/TraceInsights/Private/Insights/Widgets/STimingView.cpp
--- a/Engine/Source/Developer/TraceInsights/Private/Insights/Widgets/STimingView.cpp
+++ b/Engine/Source/Developer/TraceInsights/Private/Insights/Widgets/STimingView.cpp
@@ -39,9 +39,9 @@
 #include "Insights/Common/PaintUtils.h"
 #include "Insights/Common/Stopwatch.h"
 #include "Insights/Common/TimeUtils.h"
-#include "Insights/ITimingViewExtender.h"
 #include "Insights/InsightsManager.h"
 #include "Insights/InsightsStyle.h"
+#include "Insights/ITimingViewExtender.h"
 #include "Insights/LoadingProfiler/LoadingProfilerManager.h"
 #include "Insights/LoadingProfiler/Widgets/SLoadingProfilerWindow.h"
 #include "Insights/Log.h"
@@ -60,21 +60,21 @@
 #include "Insights/ViewModels/GraphTrack.h"
 #include "Insights/ViewModels/LoadingTimingTrack.h"
 #include "Insights/ViewModels/MarkersTimingTrack.h"
+#include "Insights/ViewModels/QuickFind.h"
 #include "Insights/ViewModels/RegionsTimingTrack.h"
 #include "Insights/ViewModels/ThreadTimingTrack.h"
+#include "Insights/ViewModels/ThreadTrackEvent.h"
 #include "Insights/ViewModels/TimeFilterValueConverter.h"
 #include "Insights/ViewModels/TimerFilters.h"
 #include "Insights/ViewModels/TimeRulerTrack.h"
 #include "Insights/ViewModels/TimingEventSearch.h"
 #include "Insights/ViewModels/TimingGraphTrack.h"
 #include "Insights/ViewModels/TimingViewDrawHelper.h"
-#include "Insights/ViewModels/QuickFind.h"
+#include "Insights/Widgets/SQuickFind.h"
 #include "Insights/Widgets/SStatsView.h"
 #include "Insights/Widgets/STimersView.h"
 #include "Insights/Widgets/STimingProfilerWindow.h"
 #include "Insights/Widgets/STimingViewTrackList.h"
-#include "Insights/Widgets/SQuickFind.h"
-#include "Insights/ViewModels/ThreadTrackEvent.h"
 
 #include <limits>
 
diff --git a/Engine/Source/Developer/TraceServices/Private/Model/Log.cpp b/Engine/Source/Developer/TraceServices/Private/Model/Log.cpp
--- a/Engine/Source/Developer/TraceServices/Private/Model/Log.cpp
+++ b/Engine/Source/Developer/TraceServices/Private/Model/Log.cpp
@@ -189,47 +189,57 @@ bool FLogProvider::ReadMessage(uint64 Index, TFunctionRef<void(const FLogMessage
 	return true;
 }
 
-void FLogProvider::EnumerateMessages(double StartTime, double EndTime, TFunctionRef<void(const FLogMessageInfo&)> Callback) const
+void FLogProvider::EnumerateMessagesByIndex(uint64 StartIndex, uint64 EndIndex, TFunctionRef<void(const FLogMessageInfo&)> Callback) const
 {
 	Session.ReadAccessCheck();
 
-	if (StartTime > EndTime)
+	uint64 Count = Messages.Num();
+	if (EndIndex > Count)
 	{
-		return;
+		EndIndex = Count;
 	}
-
-	uint64 MessageCount = Messages.Num();
-	if (MessageCount == 0)
+	if (StartIndex >= EndIndex)
 	{
 		return;
 	}
 
-	for (auto It = Messages.GetIteratorFromItem(0); It; ++It)
+	for (auto It = Messages.GetIteratorFromItem(StartIndex); It && It.GetCurrentItemIndex() < EndIndex; ++It)
 	{
-		double Time = It.GetCurrentItem()->Time;
-		if (StartTime <= Time && Time <= EndTime)
-		{
-			ConstructMessage(*It.GetCurrentItem(), It.GetCurrentItemIndex(), Callback);
-		}
+		ConstructMessage(*It.GetCurrentItem(), It.GetCurrentItemIndex(), Callback);
 	}
 }
 
-void FLogProvider::EnumerateMessagesByIndex(uint64 StartIndex, uint64 EndIndex, TFunctionRef<void(const FLogMessageInfo&)> Callback) const
+void FLogProvider::EnumerateMessages(double StartTime, double EndTime, TFunctionRef<void(const FLogMessageInfo&)> Callback) const
 {
 	Session.ReadAccessCheck();
 
-	uint64 Count = Messages.Num();
-	if (EndIndex > Count)
+	if (StartTime > EndTime)
 	{
-		EndIndex = Count;
+		return;
 	}
-	if (StartIndex >= EndIndex)
+
+	uint64 MessageCount = Messages.Num();
+	if (MessageCount == 0)
 	{
 		return;
 	}
 
-	for (auto It = Messages.GetIteratorFromItem(StartIndex); It && It.GetCurrentItemIndex() < EndIndex; ++It)
+	// Find the first log message with Time >= StartTime.
+	uint64 StartIndex = TraceServices::PagedArrayAlgo::LowerBoundBy(Messages, StartTime,
+		[](const FLogMessageInternal& Item) { return Item.Time; });
+	if (StartIndex >= Messages.Num())
 	{
+		return;
+	}
+
+	// Iterate from StartIndex and stop at first log message with Time > EndTime.
+	for (auto It = Messages.GetIteratorFromItem(StartIndex); It; ++It)
+	{
+		double Time = It.GetCurrentItem()->Time;
+		if (Time > EndTime)
+		{
+			break;
+		}
 		ConstructMessage(*It.GetCurrentItem(), It.GetCurrentItemIndex(), Callback);
 	}
 }
