[Insights]
  - Frames View: Added "Scroll Log View" in context menu. Scrolls the Log View at the message with closest timestamp to the start time of selected frame.
  - Timing View / Time Ruler: Added "Scroll Log View" in Time Ruler track's context menu. Scrolls the Log View at the message with closest timestamp to the time of the current mouse position.
  - Timing View / Time Ruler: Added timestamp in context menu for the "Move Time Marker X" actions.
  - Log View: Fixed selection of a single message to reset previous selection if any.

diff --git a/Engine/Source/Developer/TraceInsights/Private/Insights/ViewModels/FrameTrackHelper.cpp b/Engine/Source/Developer/TraceInsights/Private/Insights/ViewModels/FrameTrackHelper.cpp
--- a/Engine/Source/Developer/TraceInsights/Private/Insights/ViewModels/FrameTrackHelper.cpp
+++ b/Engine/Source/Developer/TraceInsights/Private/Insights/ViewModels/FrameTrackHelper.cpp
@@ -16,6 +16,8 @@
 
 #include <limits>
 
+#define LOCTEXT_NAMESPACE "FrameTrack"
+
 ////////////////////////////////////////////////////////////////////////////////////////////////////
 // FFrameTrackSeriesBuilder
 ////////////////////////////////////////////////////////////////////////////////////////////////////
@@ -107,9 +109,21 @@ const TCHAR* FFrameTrackDrawHelper::FrameTypeToString(int32 FrameType)
 {
 	switch (FrameType)
 	{
-	case TraceFrameType_Game:      return TEXT("Game");
-	case TraceFrameType_Rendering: return TEXT("Rendering");
-	default:                       return TEXT("Misc");
+		case TraceFrameType_Game:      return TEXT("Game");
+		case TraceFrameType_Rendering: return TEXT("Rendering");
+		default:                       return TEXT("Unknown");
+	}
+}
+
+////////////////////////////////////////////////////////////////////////////////////////////////////
+
+FText FFrameTrackDrawHelper::FrameTypeToText(int32 FrameType)
+{
+	switch (FrameType)
+	{
+		case TraceFrameType_Game:      return LOCTEXT("GameFrame", "Game Frame");
+		case TraceFrameType_Rendering: return LOCTEXT("RenderingFrame", "Rendering Frame");
+		default:                       return LOCTEXT("UnknownFrame", "Unknown Frame");
 	}
 }
 
@@ -119,9 +133,9 @@ uint32 FFrameTrackDrawHelper::GetColor32ByFrameType(int32 FrameType)
 {
 	switch (FrameType)
 	{
-	case TraceFrameType_Game	:	return 0xFF5555FF;
-	case TraceFrameType_Rendering:	return 0xFFFF5555;
-	default:						return 0xFF666666;
+		case TraceFrameType_Game:      return 0xFF5555FF;
+		case TraceFrameType_Rendering: return 0xFFFF5555;
+		default:                       return 0xFF666666;
 	}
 }
 
@@ -132,9 +146,9 @@ FLinearColor FFrameTrackDrawHelper::GetColorByFrameType(int32 FrameType)
 	constexpr float Alpha = 0.9f;
 	switch (FrameType)
 	{
-	case TraceFrameType_Game:		return FLinearColor(0.75f, 1.0f, 1.0f, Alpha);
-	case TraceFrameType_Rendering:	return FLinearColor(1.0f, 0.75f, 0.75f, Alpha);
-	default:						return FLinearColor(1.0f, 1.0f, 1.0f, Alpha);
+		case TraceFrameType_Game:      return FLinearColor(0.75f, 1.0f, 1.0f, Alpha);
+		case TraceFrameType_Rendering: return FLinearColor(1.0f, 0.75f, 0.75f, Alpha);
+		default:                       return FLinearColor(1.0f, 1.0f, 1.0f, Alpha);
 	}
 }
 
@@ -314,3 +328,5 @@ void FFrameTrackDrawHelper::DrawHighlightedInterval(const FFrameTrackSeries& Ser
 }
 
 ////////////////////////////////////////////////////////////////////////////////////////////////////
+
+#undef LOCTEXT_NAMESPACE
diff --git a/Engine/Source/Developer/TraceInsights/Private/Insights/ViewModels/FrameTrackHelper.h b/Engine/Source/Developer/TraceInsights/Private/Insights/ViewModels/FrameTrackHelper.h
--- a/Engine/Source/Developer/TraceInsights/Private/Insights/ViewModels/FrameTrackHelper.h
+++ b/Engine/Source/Developer/TraceInsights/Private/Insights/ViewModels/FrameTrackHelper.h
@@ -162,6 +162,7 @@ public:
 	void DrawHighlightedInterval(const FFrameTrackSeries& Series, const double StartTime, const double EndTime) const;
 
 	static const TCHAR* FrameTypeToString(int32 FrameType);
+	static FText FrameTypeToText(int32 FrameType);
 	static uint32 GetColor32ByFrameType(int32 FrameType);
 	static FLinearColor GetColorByFrameType(int32 FrameType);
 
diff --git a/Engine/Source/Developer/TraceInsights/Private/Insights/ViewModels/TimeRulerTrack.cpp b/Engine/Source/Developer/TraceInsights/Private/Insights/ViewModels/TimeRulerTrack.cpp
--- a/Engine/Source/Developer/TraceInsights/Private/Insights/ViewModels/TimeRulerTrack.cpp
+++ b/Engine/Source/Developer/TraceInsights/Private/Insights/ViewModels/TimeRulerTrack.cpp
@@ -8,14 +8,18 @@
 #include "Framework/MultiBox/MultiBoxBuilder.h"
 #include "Rendering/DrawElements.h"
 #include "Styling/AppStyle.h"
-#include <limits>
 
 // Insights
 #include "Insights/Common/PaintUtils.h"
 #include "Insights/Common/TimeUtils.h"
 #include "Insights/InsightsStyle.h"
+#include "Insights/TimingProfilerManager.h"
 #include "Insights/ViewModels/DrawHelpers.h"
 #include "Insights/ViewModels/TimingTrackViewport.h"
+#include "Insights/Widgets/SLogView.h"
+#include "Insights/Widgets/STimingProfilerWindow.h"
+
+#include <limits>
 
 #define LOCTEXT_NAMESPACE "TimeRulerTrack"
 
@@ -440,6 +444,8 @@ void FTimeRulerTrack::BuildContextMenu(FMenuBuilder& MenuBuilder)
 		}
 	}
 
+	FText CrtMousePosTimeText = FText::FromString(TimeUtils::FormatTimeAuto(CrtMousePosTime, 2));
+
 	if (VisibleTimeMarkers.Num() > 0)
 	{
 		// Sort TimeMarkers by name.
@@ -458,8 +464,8 @@ void FTimeRulerTrack::BuildContextMenu(FMenuBuilder& MenuBuilder)
 			const FText MarkerNameText = FText::FromString((MarkerNameString.Len() > 0) ? MarkerNameString : TEXT("T"));
 			MenuBuilder.AddMenuEntry
 			(
-				FText::Format(LOCTEXT("ContextMenu_MoveTimeMerker", "Move Time Marker '{0}' Here"), MarkerNameText),
-				FText::Format(LOCTEXT("ContextMenu_MoveTimeMerker_Desc", "Move the time marker '{0}' at the current mouse position."), MarkerNameText),
+				FText::Format(LOCTEXT("ContextMenu_MoveTimeMarker_Fmt", "Move Time Marker '{0}' (\u2192 {1})"), MarkerNameText, CrtMousePosTimeText),
+				FText::Format(LOCTEXT("ContextMenu_MoveTimeMarker_Desc_Fmt", "Move the time marker '{0}' at the current mouse position ({1})."), MarkerNameText, CrtMousePosTimeText),
 				FSlateIcon(),
 				Action_MoveTimeMarker,
 				NAME_None,
@@ -469,6 +475,26 @@ void FTimeRulerTrack::BuildContextMenu(FMenuBuilder& MenuBuilder)
 
 		MenuBuilder.EndSection();
 	}
+
+	MenuBuilder.BeginSection("Misc", LOCTEXT("ContextMenu_Section_Misc", "Misc"));
+
+	FUIAction Action_ScrollLogView
+	(
+		FExecuteAction::CreateSP(this, &FTimeRulerTrack::ContextMenu_ScrollLogView_Execute),
+		FCanExecuteAction::CreateSP(this, &FTimeRulerTrack::ContextMenu_ScrollLogView_CanExecute)
+	);
+	FText Label = FText::Format(LOCTEXT("ContextMenu_ScrollLogView_Fmt", "Scroll Log View (\u2192 {0})"), CrtMousePosTimeText);
+	MenuBuilder.AddMenuEntry
+	(
+		Label,
+		FText::Format(LOCTEXT("ContextMenu_ScrollLogView_Desc_Fmt", "Scrolls the Log View at the message with the closest timestamp to the time of the current mouse position ({0})."), CrtMousePosTimeText),
+		FSlateIcon(),
+		Action_ScrollLogView,
+		NAME_None,
+		EUserInterfaceActionType::Button
+	);
+
+	MenuBuilder.EndSection();
 }
 
 ////////////////////////////////////////////////////////////////////////////////////////////////////
@@ -480,4 +506,38 @@ void FTimeRulerTrack::ContextMenu_MoveTimeMarker_Execute(TSharedRef<Insights::FT
 
 ////////////////////////////////////////////////////////////////////////////////////////////////////
 
+bool FTimeRulerTrack::ContextMenu_ScrollLogView_CanExecute()
+{
+	TSharedPtr<STimingProfilerWindow> TimingWindow = FTimingProfilerManager::Get()->GetProfilerWindow();
+	if (!TimingWindow.IsValid())
+	{
+		return false;
+	}
+	TSharedPtr<SLogView> LogView = TimingWindow->GetLogView();
+	if (!LogView.IsValid())
+	{
+		return false;
+	}
+	return true;
+}
+
+////////////////////////////////////////////////////////////////////////////////////////////////////
+
+void FTimeRulerTrack::ContextMenu_ScrollLogView_Execute()
+{
+	TSharedPtr<STimingProfilerWindow> TimingWindow = FTimingProfilerManager::Get()->GetProfilerWindow();
+	if (!TimingWindow.IsValid())
+	{
+		return;
+	}
+	TSharedPtr<SLogView> LogView = TimingWindow->GetLogView();
+	if (!LogView.IsValid())
+	{
+		return;
+	}
+	LogView->SelectLogMessageByClosestTime(CrtMousePosTime);
+}
+
+////////////////////////////////////////////////////////////////////////////////////////////////////
+
 #undef LOCTEXT_NAMESPACE
diff --git a/Engine/Source/Developer/TraceInsights/Private/Insights/ViewModels/TimeRulerTrack.h b/Engine/Source/Developer/TraceInsights/Private/Insights/ViewModels/TimeRulerTrack.h
--- a/Engine/Source/Developer/TraceInsights/Private/Insights/ViewModels/TimeRulerTrack.h
+++ b/Engine/Source/Developer/TraceInsights/Private/Insights/ViewModels/TimeRulerTrack.h
@@ -101,6 +101,8 @@ public:
 private:
 	void DrawTimeMarker(const ITimingTrackDrawContext& Context, const Insights::FTimeMarker& TimeMarker) const;
 	void ContextMenu_MoveTimeMarker_Execute(TSharedRef<Insights::FTimeMarker> InTimeMarker);
+	bool ContextMenu_ScrollLogView_CanExecute();
+	void ContextMenu_ScrollLogView_Execute();
 
 private:
 	// Slate resources
diff --git a/Engine/Source/Developer/TraceInsights/Private/Insights/Widgets/SFrameTrack.cpp b/Engine/Source/Developer/TraceInsights/Private/Insights/Widgets/SFrameTrack.cpp
--- a/Engine/Source/Developer/TraceInsights/Private/Insights/Widgets/SFrameTrack.cpp
+++ b/Engine/Source/Developer/TraceInsights/Private/Insights/Widgets/SFrameTrack.cpp
@@ -26,8 +26,9 @@
 #include "Insights/ViewModels/FrameStatsHelper.h"
 #include "Insights/ViewModels/FrameTrackHelper.h"
 #include "Insights/ViewModels/ThreadTimingTrack.h"
-#include "Insights/Widgets/STimingProfilerWindow.h"
+#include "Insights/Widgets/SLogView.h"
 #include "Insights/Widgets/STimersView.h"
+#include "Insights/Widgets/STimingProfilerWindow.h"
 #include "Insights/Widgets/STimingView.h"
 
 #include <limits>
@@ -130,6 +131,7 @@ void SFrameTrack::Reset()
 	bDrawVerticalAxisLabelsOnLeftSide = false;
 
 	HoveredSample.Reset();
+	SelectedSample.Reset();
 	TooltipOpacity = 0.0f;
 	TooltipSizeX = 70.0f;
 
@@ -419,7 +421,7 @@ void SFrameTrack::UpdateState()
 			TSet<uint32> Timelines;
 			TSharedPtr<class STimingProfilerWindow> TimingWindow = FTimingProfilerManager::Get()->GetProfilerWindow();
 
-			// Attemp to compute only from visible timelines.
+			// Attempt to compute only from visible timelines.
 			if (TimingWindow.IsValid())
 			{
 				TSharedPtr<STimingView> TimingView = TimingWindow->GetTimingView();
@@ -1099,6 +1101,11 @@ FReply SFrameTrack::OnMouseButtonUp(const FGeometry& MyGeometry, const FPointerE
 			{
 				const bool JoinCurrentSelection = MouseEvent.IsShiftDown();
 
+				if (!JoinCurrentSelection)
+				{
+					SelectedSample = HoveredSample;
+				}
+
 				SelectFrameAtMousePosition(
 					static_cast<float>(MousePositionOnButtonUp.X),
 					static_cast<float>(MousePositionOnButtonUp.Y),
@@ -1122,6 +1129,7 @@ FReply SFrameTrack::OnMouseButtonUp(const FGeometry& MyGeometry, const FPointerE
 			}
 			else if (bIsValidForMouseClick)
 			{
+				SelectedSample = HoveredSample;
 				ShowContextMenu(MouseEvent);
 			}
 
@@ -1414,6 +1422,8 @@ void SFrameTrack::ShowContextMenu(const FPointerEvent& MouseEvent)
 	}
 	MenuBuilder.EndSection();
 
+	CreateSelectedFrameMenu(MenuBuilder);
+
 	TSharedRef<SWidget> MenuWidget = MenuBuilder.MakeWidget();
 
 	FWidgetPath EventPath = MouseEvent.GetEventPath() != nullptr ? *MouseEvent.GetEventPath() : FWidgetPath();
@@ -1423,6 +1433,51 @@ void SFrameTrack::ShowContextMenu(const FPointerEvent& MouseEvent)
 
 ////////////////////////////////////////////////////////////////////////////////////////////////////
 
+void SFrameTrack::CreateSelectedFrameMenu(FMenuBuilder& MenuBuilder)
+{
+	FText SelectedFrameSectionName;
+	if (SelectedSample.IsValid())
+	{
+		SelectedFrameSectionName = FText::Format(LOCTEXT("ContextMenu_Section_SelectedFrame_Fmt", "{0} {1}"),
+			FFrameTrackDrawHelper::FrameTypeToText(SelectedSample.Series->FrameType),
+			FText::AsNumber(SelectedSample.Sample->LargestFrameIndex));
+	}
+	else
+	{
+		SelectedFrameSectionName = LOCTEXT("ContextMenu_Section_NoFrameSelected", "No Frame Selected");
+	}
+	MenuBuilder.BeginSection("SelectedFrame", SelectedFrameSectionName);
+
+	FUIAction Action_ScrollLogView
+	(
+		FExecuteAction::CreateSP(this, &SFrameTrack::ContextMenu_ScrollLogView_Execute),
+		FCanExecuteAction::CreateSP(this, &SFrameTrack::ContextMenu_ScrollLogView_CanExecute)
+	);
+	FText Label;
+	if (SelectedSample.IsValid())
+	{
+		FText StartTimeText = FText::FromString(TimeUtils::FormatTimeAuto(SelectedSample.Sample->LargestFrameStartTime, 2));
+		Label = FText::Format(LOCTEXT("ContextMenu_ScrollLogView_Fmt", "Scroll Log View (\u2192 {0})"), StartTimeText);
+	}
+	else
+	{
+		Label = LOCTEXT("ContextMenu_ScrollLogView", "Scroll Log View");
+	}
+	MenuBuilder.AddMenuEntry
+	(
+		Label,
+		LOCTEXT("ContextMenu_ScrollLogView_Desc", "Scrolls the Log View at the message with the closest timestamp to the start time of the selected frame."),
+		FSlateIcon(),
+		Action_ScrollLogView,
+		NAME_None,
+		EUserInterfaceActionType::Button
+	);
+
+	MenuBuilder.EndSection();
+}
+
+////////////////////////////////////////////////////////////////////////////////////////////////////
+
 void SFrameTrack::ContextMenu_ShowGameFrames_Execute()
 {
 	TSharedPtr<FFrameTrackSeries> Series = FindSeries(ETraceFrameType::TraceFrameType_Game);
@@ -1601,6 +1656,48 @@ bool SFrameTrack::ContextMenu_ZoomTimingViewOnFrameSelection_IsChecked()
 
 ////////////////////////////////////////////////////////////////////////////////////////////////////
 
+void SFrameTrack::ContextMenu_ScrollLogView_Execute()
+{
+	if (!SelectedSample.IsValid())
+	{
+		return;
+	}
+	TSharedPtr<STimingProfilerWindow> TimingWindow = FTimingProfilerManager::Get()->GetProfilerWindow();
+	if (!TimingWindow.IsValid())
+	{
+		return;
+	}
+	TSharedPtr<SLogView> LogView = TimingWindow->GetLogView();
+	if (!LogView.IsValid())
+	{
+		return;
+	}
+	LogView->SelectLogMessageByClosestTime(SelectedSample.Sample->LargestFrameStartTime);
+}
+
+////////////////////////////////////////////////////////////////////////////////////////////////////
+
+bool SFrameTrack::ContextMenu_ScrollLogView_CanExecute()
+{
+	if (!SelectedSample.IsValid())
+	{
+		return false;
+	}
+	TSharedPtr<STimingProfilerWindow> TimingWindow = FTimingProfilerManager::Get()->GetProfilerWindow();
+	if (!TimingWindow.IsValid())
+	{
+		return false;
+	}
+	TSharedPtr<SLogView> LogView = TimingWindow->GetLogView();
+	if (!LogView.IsValid())
+	{
+		return false;
+	}
+	return true;
+}
+
+////////////////////////////////////////////////////////////////////////////////////////////////////
+
 void SFrameTrack::BindCommands()
 {
 }
diff --git a/Engine/Source/Developer/TraceInsights/Private/Insights/Widgets/SFrameTrack.h b/Engine/Source/Developer/TraceInsights/Private/Insights/Widgets/SFrameTrack.h
--- a/Engine/Source/Developer/TraceInsights/Private/Insights/Widgets/SFrameTrack.h
+++ b/Engine/Source/Developer/TraceInsights/Private/Insights/Widgets/SFrameTrack.h
@@ -15,6 +15,7 @@
 #include "Insights/ViewModels/FrameTrackHelper.h"
 #include "Insights/ViewModels/FrameTrackViewport.h"
 
+class FMenuBuilder;
 class SScrollBar;
 class STimingView;
 
@@ -148,6 +149,7 @@ protected:
 	void SelectFrameAtMousePosition(double X, double Y, bool JoinCurrentSelection);
 
 	void ShowContextMenu(const FPointerEvent& MouseEvent);
+	void CreateSelectedFrameMenu(FMenuBuilder& MenuBuilder);
 
 	void ContextMenu_ShowGameFrames_Execute();
 	bool ContextMenu_ShowGameFrames_CanExecute();
@@ -170,6 +172,9 @@ protected:
 	bool ContextMenu_ZoomTimingViewOnFrameSelection_CanExecute();
 	bool ContextMenu_ZoomTimingViewOnFrameSelection_IsChecked();
 
+	void ContextMenu_ScrollLogView_Execute();
+	bool ContextMenu_ScrollLogView_CanExecute();
+
 	/** Binds our UI commands to delegates. */
 	void BindCommands();
 
@@ -185,22 +190,22 @@ protected:
 protected:
 	/** The track's viewport. Encapsulates info about position and scale. */
 	FFrameTrackViewport Viewport;
-	bool bIsViewportDirty;
+	bool bIsViewportDirty = false;
 
 	/** Cached info for all frame series. */
 	TArray<TSharedPtr<FFrameTrackSeries>> AllSeries;
 
-	bool bIsStateDirty;
+	bool bIsStateDirty = false;
 
-	bool bIsAutoZoomEnabled;
+	bool bIsAutoZoomEnabled = false;
 
-	float AutoZoomViewportPos;
-	float AutoZoomViewportScale;
-	float AutoZoomViewportSize;
+	float AutoZoomViewportPos = false;
+	float AutoZoomViewportScale = false;
+	float AutoZoomViewportSize = false;
 
-	bool bZoomTimingViewOnFrameSelection;
+	bool bZoomTimingViewOnFrameSelection = false;
 
-	uint64 AnalysisSyncNextTimestamp;
+	uint64 AnalysisSyncNextTimestamp = 0;
 
 	//////////////////////////////////////////////////
 
@@ -214,26 +219,27 @@ protected:
 
 	/** Mouse position during the call on mouse button down. */
 	FVector2D MousePositionOnButtonDown;
-	float ViewportPosXOnButtonDown;
+	float ViewportPosXOnButtonDown = 0.0f;
 
 	/** Mouse position during the call on mouse button up. */
 	FVector2D MousePositionOnButtonUp;
 
-	bool bIsLMB_Pressed;
-	bool bIsRMB_Pressed;
+	bool bIsLMB_Pressed = false;
+	bool bIsRMB_Pressed = false;
 
 	/** True, if the user is currently interactively scrolling the view (ex.: by holding the left mouse button and dragging). */
-	bool bIsScrolling;
+	bool bIsScrolling = false;
 
-	mutable bool bDrawVerticalAxisLabelsOnLeftSide;
+	mutable bool bDrawVerticalAxisLabelsOnLeftSide = false;
 
 	//////////////////////////////////////////////////
 	// Selection
 
 	FFrameTrackSampleRef HoveredSample;
+	FFrameTrackSampleRef SelectedSample;
 
-	mutable float TooltipOpacity;
-	mutable float TooltipSizeX;
+	mutable float TooltipOpacity = 0.0f;
+	mutable float TooltipSizeX = 0.0f;
 
 	//////////////////////////////////////////////////
 	// Misc
@@ -241,17 +247,17 @@ protected:
 	FGeometry ThisGeometry;
 
 	/** Cursor type. */
-	ECursorType CursorType;
+	ECursorType CursorType = ECursorType::Default;
 
-	STimingView* RegisteredTimingView = nullptr; // For pointer comparison only, do not dereferentiate.
+	STimingView* RegisteredTimingView = nullptr; // for pointer comparison only; do not dereferentiate!
 	FDelegateHandle OnTrackVisibilityChangedHandle;
 	FDelegateHandle OnTrackAddedHandle;
 	FDelegateHandle OnTrackRemovedHandle;
 
 	// Debug stats
-	int32 NumUpdatedFrames;
+	int32 NumUpdatedFrames = 0;
 	TFixedCircularBuffer<uint64, 32> UpdateDurationHistory;
 	mutable TFixedCircularBuffer<uint64, 32> DrawDurationHistory;
 	mutable TFixedCircularBuffer<uint64, 32> OnPaintDurationHistory;
-	mutable uint64 LastOnPaintTime;
+	mutable uint64 LastOnPaintTime = 0;
 };
diff --git a/Engine/Source/Developer/TraceInsights/Private/Insights/Widgets/SLogView.cpp b/Engine/Source/Developer/TraceInsights/Private/Insights/Widgets/SLogView.cpp
--- a/Engine/Source/Developer/TraceInsights/Private/Insights/Widgets/SLogView.cpp
+++ b/Engine/Source/Developer/TraceInsights/Private/Insights/Widgets/SLogView.cpp
@@ -1007,12 +1007,9 @@ TSharedPtr<FLogMessage> SLogView::GetSelectedLogMessage() const
 
 void SLogView::SelectLogMessage(TSharedPtr<FLogMessage> LogMessage)
 {
+	ListView->SetSelection(LogMessage, ESelectInfo::Direct);
 	if (LogMessage.IsValid())
 	{
-		if (!ListView->IsItemSelected(LogMessage))
-		{
-			ListView->SetItemSelection(LogMessage, true, ESelectInfo::Direct);
-		}
 		ListView->RequestScrollIntoView(LogMessage);
 	}
 }
@@ -1213,7 +1210,8 @@ TSharedPtr<SWidget> SLogView::ListView_GetContextMenu()
 			FLogMessageRecord& Record = Cache.Get(SelectedLogMessage->GetIndex());
 			FName CategoryName(Record.GetCategory());
 
-			MenuBuilder.AddMenuEntry(
+			MenuBuilder.AddMenuEntry
+			(
 				FLogViewCommands::Get().Command_HideSelectedCategory,
 				NAME_None,
 				FText::Format(LOCTEXT("HideCategory", "Hide \"{0}\" Category"), Record.GetCategoryAsText()),
@@ -1221,7 +1219,8 @@ TSharedPtr<SWidget> SLogView::ListView_GetContextMenu()
 				FSlateIcon(FAppStyle::Get().GetStyleSetName(), "Icons.Hidden")
 			);
 
-			MenuBuilder.AddMenuEntry(
+			MenuBuilder.AddMenuEntry
+			(
 				FLogViewCommands::Get().Command_ShowOnlySelectedCategory,
 				NAME_None,
 				FText::Format(LOCTEXT("ShowOnlyCategory", "Show Only \"{0}\" Category"), Record.GetCategoryAsText()),
@@ -1230,7 +1229,8 @@ TSharedPtr<SWidget> SLogView::ListView_GetContextMenu()
 			);
 		}
 
-		MenuBuilder.AddMenuEntry(
+		MenuBuilder.AddMenuEntry
+		(
 			FLogViewCommands::Get().Command_ShowAllCategories,
 			NAME_None,
 			TAttribute<FText>(),
@@ -1240,7 +1240,8 @@ TSharedPtr<SWidget> SLogView::ListView_GetContextMenu()
 
 		MenuBuilder.AddSeparator();
 
-		MenuBuilder.AddMenuEntry(
+		MenuBuilder.AddMenuEntry
+		(
 			FLogViewCommands::Get().Command_CopySelected,
 			NAME_None,
 			TAttribute<FText>(),
@@ -1248,7 +1249,8 @@ TSharedPtr<SWidget> SLogView::ListView_GetContextMenu()
 			FSlateIcon(FAppStyle::Get().GetStyleSetName(), "GenericCommands.Copy")
 		);
 
-		MenuBuilder.AddMenuEntry(
+		MenuBuilder.AddMenuEntry
+		(
 			FLogViewCommands::Get().Command_CopyMessage,
 			NAME_None,
 			TAttribute<FText>(),
@@ -1256,7 +1258,8 @@ TSharedPtr<SWidget> SLogView::ListView_GetContextMenu()
 			FSlateIcon(FAppStyle::Get().GetStyleSetName(), "GenericCommands.Copy")
 		);
 
-		MenuBuilder.AddMenuEntry(
+		MenuBuilder.AddMenuEntry
+		(
 			FLogViewCommands::Get().Command_CopyRange,
 			NAME_None,
 			TAttribute<FText>(),
@@ -1264,7 +1267,8 @@ TSharedPtr<SWidget> SLogView::ListView_GetContextMenu()
 			FSlateIcon(FAppStyle::Get().GetStyleSetName(), "GenericCommands.Copy")
 		);
 
-		MenuBuilder.AddMenuEntry(
+		MenuBuilder.AddMenuEntry
+		(
 			FLogViewCommands::Get().Command_CopyAll,
 			NAME_None,
 			TAttribute<FText>(),
@@ -1274,7 +1278,8 @@ TSharedPtr<SWidget> SLogView::ListView_GetContextMenu()
 
 		MenuBuilder.AddSeparator();
 
-		MenuBuilder.AddMenuEntry(
+		MenuBuilder.AddMenuEntry
+		(
 			FLogViewCommands::Get().Command_SaveRange,
 			NAME_None,
 			TAttribute<FText>(),
@@ -1282,7 +1287,8 @@ TSharedPtr<SWidget> SLogView::ListView_GetContextMenu()
 			FSlateIcon(FAppStyle::Get().GetStyleSetName(), "Icons.Save")
 		);
 
-		MenuBuilder.AddMenuEntry(
+		MenuBuilder.AddMenuEntry
+		(
 			FLogViewCommands::Get().Command_SaveAll,
 			NAME_None,
 			TAttribute<FText>(),
@@ -1324,12 +1330,14 @@ TSharedPtr<SWidget> SLogView::ListView_GetContextMenu()
 						SourceCodeAccessor.GetNameText());
 				}
 
-				MenuBuilder.AddMenuEntry(
+				MenuBuilder.AddMenuEntry
+				(
 					FLogViewCommands::Get().Command_OpenSource,
 					NAME_None,
 					ItemLabel,
 					ItemToolTip,
-					FSlateIcon(SourceCodeAccessor.GetStyleSet(), SourceCodeAccessor.GetOpenIconName()));
+					FSlateIcon(SourceCodeAccessor.GetStyleSet(), SourceCodeAccessor.GetOpenIconName())
+				);
 			}
 			else
 			{
@@ -1347,13 +1355,15 @@ TSharedPtr<SWidget> SLogView::ListView_GetContextMenu()
 					ItemToolTip = LOCTEXT("ContextMenu_OpenSource_NoAccessor_Desc2", "Source Code Accessor is not available.");
 				}
 
-				MenuBuilder.AddMenuEntry(
+				MenuBuilder.AddMenuEntry
+				(
 					ItemLabel,
 					ItemToolTip,
 					FSlateIcon(SourceCodeAccessor.GetStyleSet(), SourceCodeAccessor.GetOpenIconName()),
 					FUIAction(FExecuteAction(), FCanExecuteAction::CreateLambda([]() { return false; })),
 					NAME_None,
-					EUserInterfaceActionType::None);
+					EUserInterfaceActionType::None
+				);
 			}
 		}
 	}
@@ -1447,8 +1457,10 @@ void SLogView::CreateVerbosityThresholdMenuSection(FMenuBuilder& MenuBuilder)
 				]
 			];
 
-		MenuBuilder.AddMenuEntry(
-			FUIAction(FExecuteAction::CreateSP(this, &SLogView::VerbosityThreshold_Execute, Threshold.Verbosity),
+		MenuBuilder.AddMenuEntry
+		(
+			FUIAction(
+				FExecuteAction::CreateSP(this, &SLogView::VerbosityThreshold_Execute, Threshold.Verbosity),
 				FCanExecuteAction(),
 				FIsActionChecked::CreateSP(this, &SLogView::VerbosityThreshold_IsChecked, Threshold.Verbosity)),
 			TextBlock,
@@ -1467,11 +1479,13 @@ TSharedRef<SWidget> SLogView::MakeCategoryFilterMenu()
 
 	MenuBuilder.BeginSection("QuickFilter", LOCTEXT("CategoryFilterMenu_Section_QuickFilter", "Quick Filter"));
 	{
-		MenuBuilder.AddMenuEntry(
+		MenuBuilder.AddMenuEntry
+		(
 			LOCTEXT("ShowAllCategories", "Show/Hide All"),
 			LOCTEXT("ShowAllCategories_Tooltip", "Change filtering to show/hide all categories"),
 			FSlateIcon(),
-			FUIAction(FExecuteAction::CreateSP(this, &SLogView::ShowHideAllCategories_Execute),
+			FUIAction(
+				FExecuteAction::CreateSP(this, &SLogView::ShowHideAllCategories_Execute),
 				FCanExecuteAction(),
 				FIsActionChecked::CreateSP(this, &SLogView::ShowHideAllCategories_IsChecked)),
 			NAME_None,
@@ -1503,10 +1517,12 @@ void SLogView::CreateCategoriesFilterMenuSection(FMenuBuilder& MenuBuilder)
 			.ShadowOffset(FVector2D(1.0f, 1.0f))
 			.ColorAndOpacity(FSlateColor(FTimeMarkerTrackBuilder::GetColorByCategory(*CategoryString)));
 
-		MenuBuilder.AddMenuEntry(
-			FUIAction(FExecuteAction::CreateSP(this, &SLogView::ToggleCategory, CategoryName),
-			FCanExecuteAction(),
-			FIsActionChecked::CreateSP(this, &SLogView::IsLogCategoryEnabled, CategoryName)),
+		MenuBuilder.AddMenuEntry
+		(
+			FUIAction(
+				FExecuteAction::CreateSP(this, &SLogView::ToggleCategory, CategoryName),
+				FCanExecuteAction(),
+				FIsActionChecked::CreateSP(this, &SLogView::IsLogCategoryEnabled, CategoryName)),
 			TextBlock,
 			NAME_None,
 			FText::Format(LOCTEXT("Category_Tooltip", "Filter the Log View to show/hide category: {0}"), CategoryText),
diff --git a/Engine/Source/Developer/TraceInsights/Private/Insights/Widgets/SLogView.h b/Engine/Source/Developer/TraceInsights/Private/Insights/Widgets/SLogView.h
--- a/Engine/Source/Developer/TraceInsights/Private/Insights/Widgets/SLogView.h
+++ b/Engine/Source/Developer/TraceInsights/Private/Insights/Widgets/SLogView.h
@@ -164,9 +164,9 @@ protected:
 	FLogFilter Filter;
 	uint64 FilterChangeNumber;
 
-	int32 FilteringStartIndex; // Start index (of the range of log messages to filter) currenly used by the async task
-	int32 FilteringEndIndex; // End index (of the range of log messages to filter) currenly used by the async task
-	uint64 FilteringChangeNumber; // Change number of the filter currenly used by the async task
+	int32 FilteringStartIndex; // Start index (of the range of log messages to filter) currently used by the async task
+	int32 FilteringEndIndex; // End index (of the range of log messages to filter) currently used by the async task
+	uint64 FilteringChangeNumber; // Change number of the filter currently used by the async task
 	TUniquePtr<FAsyncTask<FLogFilteringAsyncTask>> FilteringAsyncTask; // The async task to filter log messages on a worker thread
 	mutable volatile bool bIsFilteringAsyncTaskCancelRequested; // true if we want the async task to finish asap
 
