[Insights]
  - Frames View: Added UI to setup the upper and lower frame thresholds:
      * Upper Threshold --> Frames with duration longer than this threshold will have a red color tint.
      * Lower Threshold --> Frames with duration shorter than this threshold will have a green color tint.
    Thresholds can be specified as a frame duration, in seconds [0.001 .. 1.0] or as a framerate [1 fps ... 1000 fps].
  - Frames View: Added toggles to show horizontal lines for upper and lower thresholds (red & green).

diff --git a/Engine/Source/Developer/TraceInsights/Private/Insights/ViewModels/FrameTrackHelper.cpp b/Engine/Source/Developer/TraceInsights/Private/Insights/ViewModels/FrameTrackHelper.cpp
--- a/Engine/Source/Developer/TraceInsights/Private/Insights/ViewModels/FrameTrackHelper.cpp
+++ b/Engine/Source/Developer/TraceInsights/Private/Insights/ViewModels/FrameTrackHelper.cpp
@@ -198,12 +198,12 @@ void FFrameTrackDrawHelper::DrawCached(const FFrameTrackSeries& Series) const
 		else
 		{
 			ValueY = FMath::RoundToFloat(ViewportY.GetOffsetForValue(Sample.LargestFrameDuration));
-			if (Sample.LargestFrameDuration > 1.0 / 30.0)
+			if (Sample.LargestFrameDuration > UpperThresholdTime)
 			{
 				ColorFill.G *= 0.5f;
 				ColorFill.B *= 0.5f;
 			}
-			else if (Sample.LargestFrameDuration > 1.0 / 60.0)
+			else if (Sample.LargestFrameDuration > LowerThresholdTime)
 			{
 				ColorFill.B *= 0.5f;
 			}
diff --git a/Engine/Source/Developer/TraceInsights/Private/Insights/ViewModels/FrameTrackHelper.h b/Engine/Source/Developer/TraceInsights/Private/Insights/ViewModels/FrameTrackHelper.h
--- a/Engine/Source/Developer/TraceInsights/Private/Insights/ViewModels/FrameTrackHelper.h
+++ b/Engine/Source/Developer/TraceInsights/Private/Insights/ViewModels/FrameTrackHelper.h
@@ -156,6 +156,12 @@ public:
 	FFrameTrackDrawHelper(const FFrameTrackDrawHelper&) = delete;
 	FFrameTrackDrawHelper& operator=(const FFrameTrackDrawHelper&) = delete;
 
+	void SetThresholds(double InUpperThresholdTime, double InLowerThresholdTime)
+	{
+		UpperThresholdTime = InUpperThresholdTime;
+		LowerThresholdTime = InLowerThresholdTime;
+	}
+
 	void DrawBackground() const;
 	void DrawCached(const FFrameTrackSeries& Series) const;
 	void DrawHoveredSample(const FFrameTrackSample& Sample) const;
@@ -176,6 +182,9 @@ private:
 	const FSlateBrush* WhiteBrush;
 	const FSlateBrush* HoveredFrameBorderBrush;
 
+	double UpperThresholdTime = 1.0 / 30.0;
+	double LowerThresholdTime = 1.0 / 60.0;
+
 	// Debug stats.
 	mutable int32 NumFrames;
 	mutable int32 NumDrawSamples;
diff --git a/Engine/Source/Developer/TraceInsights/Private/Insights/Widgets/SFrameTrack.cpp b/Engine/Source/Developer/TraceInsights/Private/Insights/Widgets/SFrameTrack.cpp
--- a/Engine/Source/Developer/TraceInsights/Private/Insights/Widgets/SFrameTrack.cpp
+++ b/Engine/Source/Developer/TraceInsights/Private/Insights/Widgets/SFrameTrack.cpp
@@ -12,6 +12,7 @@
 #include "Styling/AppStyle.h"
 #include "TraceServices/Model/Frames.h"
 #include "TraceServices/Model/TimingProfiler.h"
+#include "Widgets/Input/SEditableTextBox.h"
 #include "Widgets/Layout/SScrollBar.h"
 
 // Insights
@@ -313,7 +314,7 @@ TSharedRef<FFrameTrackSeries> SFrameTrack::FindOrAddSeries(ETraceFrameType Frame
 
 	TSharedRef<FFrameTrackSeries> SeriesRef = MakeShared<FFrameTrackSeries>(FrameType, EFrameTrackSeriesType::Frame);
 	SeriesRef->Color = FFrameTrackDrawHelper::GetColorByFrameType(FrameType);
-	SeriesRef->Name = FText::Format(LOCTEXT("FrameTrackSeriesName_Format", "{0} {1}"), FText::FromString(FFrameTrackDrawHelper::FrameTypeToString(FrameType)), LOCTEXT("Frame", "Frame"));
+	SeriesRef->Name = FFrameTrackDrawHelper::FrameTypeToText(FrameType);
 	AllSeries.Add(SeriesRef);
 	return SeriesRef;
 }
@@ -603,6 +604,8 @@ int32 SFrameTrack::OnPaint(const FPaintArgs& Args, const FGeometry& AllottedGeom
 
 		FFrameTrackDrawHelper Helper(DrawContext, Viewport);
 
+		Helper.SetThresholds(UpperThresholdTime, LowerThresholdTime);
+
 		Helper.DrawBackground();
 
 		// Draw the horizontal axis grid (background layer).
@@ -644,6 +647,25 @@ int32 SFrameTrack::OnPaint(const FPaintArgs& Args, const FGeometry& AllottedGeom
 		// Draw the horizontal axis grid (foreground layer).
 		DrawHorizontalAxisGrid(DrawContext, WhiteBrush, SummaryFont, false);
 
+		if (bShowLowerThresholdLine)
+		{
+			const FLinearColor LineColor(0.5f, 1.0f, 0.5f, 1.0f);
+			const FAxisViewportDouble& ViewportY = Viewport.GetVerticalAxisViewport();
+			const float RoundedViewHeight = FMath::RoundToFloat(ViewportY.GetSize());
+			const float LineY = RoundedViewHeight - FMath::RoundToFloat(ViewportY.GetOffsetForValue(LowerThresholdTime));
+			DrawContext.DrawBox(0.0, LineY, ViewWidth, 1.0, WhiteBrush, LineColor);
+			DrawContext.LayerId++;
+		}
+		if (bShowUpperThresholdLine)
+		{
+			const FLinearColor LineColor(1.0f, 0.5f, 0.5f, 1.0f);
+			const FAxisViewportDouble& ViewportY = Viewport.GetVerticalAxisViewport();
+			const float RoundedViewHeight = FMath::RoundToFloat(ViewportY.GetSize());
+			const float LineY = RoundedViewHeight - FMath::RoundToFloat(ViewportY.GetOffsetForValue(UpperThresholdTime));
+			DrawContext.DrawBox(0.0, LineY, ViewWidth, 1.0, WhiteBrush, LineColor);
+			DrawContext.LayerId++;
+		}
+
 		// Draw the vertical axis grid.
 		DrawVerticalAxisGrid(DrawContext, WhiteBrush, SummaryFont);
 
@@ -938,15 +960,12 @@ void SFrameTrack::DrawVerticalAxisGrid(FDrawContext& DrawContext, const FSlateBr
 	{
 		const FVisibleAxis& Axis = VisibleAxis[Index];
 
-		constexpr double Time60fps = 1.0 / 60.0;
-		constexpr double Time30fps = 1.0 / 30.0;
-
 		FLinearColor TextColor;
-		if (Axis.Value <= Time60fps)
+		if (Axis.Value <= LowerThresholdTime)
 		{
 			TextColor = FLinearColor(0.5f, 1.0f, 0.5f, 1.0f);
 		}
-		else if (Axis.Value <= Time30fps)
+		else if (Axis.Value <= UpperThresholdTime)
 		{
 			TextColor = FLinearColor(1.0f, 1.0f, 0.5f, 1.0f);
 		}
@@ -1354,9 +1373,9 @@ void SFrameTrack::ShowContextMenu(const FPointerEvent& MouseEvent)
 	}
 	MenuBuilder.EndSection();
 
-	MenuBuilder.BeginSection("FrameStats", LOCTEXT("ContextMenu_Section_Stats", "Frame Stats"));
+	MenuBuilder.BeginSection("Timers", LOCTEXT("ContextMenu_Section_Timers", "Timers"));
 
-	FText TooltipTextBase = LOCTEXT("ContextMenu_ShowFrameStatsSeries_Desc", "Shows/hides the {0} series.");
+	FText TooltipTextBase = LOCTEXT("ContextMenu_ShowFrameStatsSeries_Desc", "Shows/hides the {0} timer series.");
 	for (TSharedPtr<FFrameTrackSeries> Series : AllSeries)
 	{
 		if (Series->Type != EFrameTrackSeriesType::TimerFrameStats)
@@ -1386,6 +1405,17 @@ void SFrameTrack::ShowContextMenu(const FPointerEvent& MouseEvent)
 
 	MenuBuilder.EndSection();
 
+	MenuBuilder.AddSeparator();
+
+	MenuBuilder.AddSubMenu
+	(
+		LOCTEXT("ContextMenu_ThresholdsSubMenu", "Setup Thresholds"),
+		LOCTEXT("ContextMenu_ThresholdsSubMenu_Desc", "Setup thresholds."),
+		FNewMenuDelegate::CreateSP(this, &SFrameTrack::CreateThresholdsMenu),
+		false,
+		FSlateIcon()
+	);
+
 	MenuBuilder.BeginSection("Zoom", LOCTEXT("ContextMenu_Section_Zoom", "Zoom"));
 	{
 		FUIAction Action_AutoZoom
@@ -1433,6 +1463,241 @@ void SFrameTrack::ShowContextMenu(const FPointerEvent& MouseEvent)
 
 ////////////////////////////////////////////////////////////////////////////////////////////////////
 
+void SFrameTrack::CreateThresholdsMenu(FMenuBuilder& MenuBuilder)
+{
+	MenuBuilder.BeginSection("Thresholds");// , LOCTEXT("ContextMenu_Section_Thresholds", "Thresholds"));
+
+	MenuBuilder.AddMenuEntry
+	(
+		LOCTEXT("ContextMenu_ShowUpperThresholdLine", "Show Upper Threshold Line"),
+		LOCTEXT("ContextMenu_ShowUpperThresholdLine_Desc", "Shows/hides the red horizontal line for the upper threshold."),
+		FSlateIcon(),
+		FUIAction(
+			FExecuteAction::CreateLambda(
+				[this]()
+				{
+					bShowUpperThresholdLine = !bShowUpperThresholdLine;
+				}),
+			FCanExecuteAction(),
+			FIsActionChecked::CreateLambda(
+				[this]() -> bool
+				{
+					return bShowUpperThresholdLine;
+				})),
+		NAME_None,
+		EUserInterfaceActionType::ToggleButton
+	);
+
+	MenuBuilder.AddMenuEntry
+	(
+		LOCTEXT("ContextMenu_ShowLowerThresholdLine", "Show Lower Threshold Line"),
+		LOCTEXT("ContextMenu_ShowLowerThresholdLine_Desc", "Shows/hides the green horizontal line for the lower threshold."),
+		FSlateIcon(),
+		FUIAction(
+			FExecuteAction::CreateLambda(
+				[this]()
+				{
+					bShowLowerThresholdLine = !bShowLowerThresholdLine;
+				}),
+			FCanExecuteAction(),
+			FIsActionChecked::CreateLambda(
+				[this]() -> bool
+				{
+					return bShowLowerThresholdLine;
+				})),
+		NAME_None,
+		EUserInterfaceActionType::ToggleButton
+	);
+
+	MenuBuilder.AddSeparator();
+
+	MenuBuilder.AddMenuEntry
+	(
+		FUIAction(FExecuteAction(), FCanExecuteAction()),
+		CreateUpperThresholdWidget(),
+		NAME_None,
+		LOCTEXT("UpperThresholdCustomTooltip", "Upper Threshold\nFrames with duration longer than this threshold will have a red color tint.\nCan be specified as a frame duration, in seconds [0.001 .. 1.0] or as a framerate [1 fps ... 1000 fps]."),
+		EUserInterfaceActionType::None
+	);
+
+	MenuBuilder.AddMenuEntry
+	(
+		FUIAction(FExecuteAction(), FCanExecuteAction()),
+		CreateLowerThresholdWidget(),
+		NAME_None,
+		LOCTEXT("LowerThresholdCustomTooltip", "Lower Threshold\nFrames with duration shorter than this threshold will have a green color tint.\nCan be specified as a frame duration, in seconds [0.001 .. 1.0] or as a framerate [1 fps ... 1000 fps]."),
+		EUserInterfaceActionType::None
+	);
+
+	MenuBuilder.EndSection();
+}
+
+////////////////////////////////////////////////////////////////////////////////////////////////////
+
+TSharedRef<SWidget> SFrameTrack::CreateUpperThresholdWidget()
+{
+	return SNew(SHorizontalBox)
+
+	+ SHorizontalBox::Slot()
+	.Padding(FMargin(0.0f, 0.0f, 0.0f, 0.0f))
+	.AutoWidth()
+	.VAlign(VAlign_Center)
+	[
+		SNew(STextBlock)
+		.MinDesiredWidth(110.0f)
+		.Text(LOCTEXT("UpperThresholdText", "Upper Threshold:"))
+	]
+
+	+ SHorizontalBox::Slot()
+	.Padding(FMargin(0.0f, 0.0f, 0.0f, 0.0f))
+	.AutoWidth()
+	.VAlign(VAlign_Center)
+	[
+		SNew(SEditableTextBox)
+		.MinDesiredWidth(60.0f)
+		.HintText(LOCTEXT("UpperThresholdCustomHint", "30 fps"))
+		.Text_Lambda([this]
+		{
+			if (bShowUpperThresholdAsFps)
+			{
+				return FText::FromString(FString::Printf(TEXT("%g fps"), 1.0 / UpperThresholdTime));
+			}
+			else
+			{
+				return FText::FromString(FString::Printf(TEXT("%g"), UpperThresholdTime));
+			}
+		})
+		.OnTextChanged_Lambda([this](const FText& InText)
+		{
+			FString ValueStr = InText.ToString().TrimStartAndEnd();
+			if (ValueStr.IsEmpty())
+			{
+				ValueStr = TEXT("30 fps");
+			}
+			if (ValueStr.EndsWith(TEXT("fps")))
+			{
+				double FPS = atof(TCHAR_TO_ANSI(*ValueStr));
+				UpperThresholdTime = 1.0 / FMath::Clamp(FPS, 1.0 / MaxThresholdTime, 1.0 / MinThresholdTime);
+				bShowUpperThresholdAsFps = true;
+			}
+			else
+			{
+				double Time = atof(TCHAR_TO_ANSI(*ValueStr));
+				UpperThresholdTime = FMath::Clamp(Time, MinThresholdTime, MaxThresholdTime);
+				bShowUpperThresholdAsFps = false;
+			}
+			if (LowerThresholdTime > UpperThresholdTime)
+			{
+				LowerThresholdTime = UpperThresholdTime;
+			}
+		})
+	]
+
+	+ SHorizontalBox::Slot()
+	.Padding(FMargin(8.0f, 0.0f, 12.0f, 0.0f))
+	.AutoWidth()
+	.VAlign(VAlign_Center)
+	[
+		SNew(STextBlock)
+		.Text_Lambda([this]
+		{
+			FString ThresholdTimeStr = TimeUtils::FormatTimeAuto(UpperThresholdTime, 2);
+			if (bShowUpperThresholdAsFps)
+			{
+				return FText::FromString(FString::Printf(TEXT("%s"), *ThresholdTimeStr));
+			}
+			else
+			{
+				return FText::FromString(FString::Printf(TEXT("%s (%.2f fps)"), *ThresholdTimeStr, 1.0 / UpperThresholdTime));
+			}
+		})
+	];
+}
+
+////////////////////////////////////////////////////////////////////////////////////////////////////
+
+TSharedRef<SWidget> SFrameTrack::CreateLowerThresholdWidget()
+{
+	return SNew(SHorizontalBox)
+
+	+ SHorizontalBox::Slot()
+	.Padding(FMargin(0.0f, 0.0f, 0.0f, 0.0f))
+	.AutoWidth()
+	.VAlign(VAlign_Center)
+	[
+		SNew(STextBlock)
+		.MinDesiredWidth(110.0f)
+		.Text(LOCTEXT("LowerThresholdText", "Lower Threshold:"))
+	]
+
+	+ SHorizontalBox::Slot()
+	.Padding(FMargin(0.0f, 0.0f, 0.0f, 0.0f))
+	.AutoWidth()
+	.VAlign(VAlign_Center)
+	[
+		SNew(SEditableTextBox)
+		.MinDesiredWidth(60.0f)
+		.HintText(LOCTEXT("LowerThresholdCustomHint", "60 fps"))
+		.Text_Lambda([this]
+		{
+			if (bShowLowerThresholdAsFps)
+			{
+				return FText::FromString(FString::Printf(TEXT("%g fps"), 1.0 / LowerThresholdTime));
+			}
+			else
+			{
+				return FText::FromString(FString::Printf(TEXT("%g"), LowerThresholdTime));
+			}
+		})
+		.OnTextChanged_Lambda([this](const FText& InText)
+		{
+			FString ValueStr = InText.ToString().TrimStartAndEnd();
+			if (ValueStr.IsEmpty())
+			{
+				ValueStr = TEXT("60 fps");
+			}
+			if (ValueStr.EndsWith(TEXT("fps")))
+			{
+				double FPS = atof(TCHAR_TO_ANSI(*ValueStr));
+				LowerThresholdTime = 1.0 / FMath::Clamp(FPS, 1.0 / MaxThresholdTime, 1.0 / MinThresholdTime);
+				bShowLowerThresholdAsFps = true;
+			}
+			else
+			{
+				double Time = atof(TCHAR_TO_ANSI(*ValueStr));
+				LowerThresholdTime = FMath::Clamp(Time, MinThresholdTime, MaxThresholdTime);
+				bShowLowerThresholdAsFps = false;
+			}
+			if (UpperThresholdTime < LowerThresholdTime)
+			{
+				UpperThresholdTime = LowerThresholdTime;
+			}
+		})
+	]
+
+	+ SHorizontalBox::Slot()
+	.Padding(FMargin(8.0f, 0.0f, 12.0f, 0.0f))
+	.AutoWidth()
+	.VAlign(VAlign_Center)
+	[
+		SNew(STextBlock)
+		.Text_Lambda([this]
+		{
+			FString ThresholdTimeStr = TimeUtils::FormatTimeAuto(LowerThresholdTime, 2);
+			if (bShowLowerThresholdAsFps)
+			{
+				return FText::FromString(FString::Printf(TEXT("%s"), *ThresholdTimeStr));
+			}
+			else
+			{
+				return FText::FromString(FString::Printf(TEXT("%s (%.2f fps)"), *ThresholdTimeStr, 1.0 / LowerThresholdTime));
+			}
+		})
+	];
+}
+
+////////////////////////////////////////////////////////////////////////////////////////////////////
+
 void SFrameTrack::CreateSelectedFrameMenu(FMenuBuilder& MenuBuilder)
 {
 	FText SelectedFrameSectionName;
diff --git a/Engine/Source/Developer/TraceInsights/Private/Insights/Widgets/SFrameTrack.h b/Engine/Source/Developer/TraceInsights/Private/Insights/Widgets/SFrameTrack.h
--- a/Engine/Source/Developer/TraceInsights/Private/Insights/Widgets/SFrameTrack.h
+++ b/Engine/Source/Developer/TraceInsights/Private/Insights/Widgets/SFrameTrack.h
@@ -149,6 +149,9 @@ protected:
 	void SelectFrameAtMousePosition(double X, double Y, bool JoinCurrentSelection);
 
 	void ShowContextMenu(const FPointerEvent& MouseEvent);
+	void CreateThresholdsMenu(FMenuBuilder& MenuBuilder);
+	TSharedRef<SWidget> CreateUpperThresholdWidget();
+	TSharedRef<SWidget> CreateLowerThresholdWidget();
 	void CreateSelectedFrameMenu(FMenuBuilder& MenuBuilder);
 
 	void ContextMenu_ShowGameFrames_Execute();
@@ -207,6 +210,21 @@ protected:
 
 	uint64 AnalysisSyncNextTimestamp = 0;
 
+	//////////////////////////////////////////////////
+	// Thresholds
+
+	static constexpr double MinThresholdTime = 0.001; // == 1ms == 1000 fps
+	static constexpr double MaxThresholdTime = 1.0; // == 1s == 1 fps
+
+	double UpperThresholdTime = 1.0 / 30.0;
+	double LowerThresholdTime = 1.0 / 60.0;
+
+	bool bShowUpperThresholdAsFps = true;
+	bool bShowLowerThresholdAsFps = true;
+
+	bool bShowUpperThresholdLine = false;
+	bool bShowLowerThresholdLine = false;
+
 	//////////////////////////////////////////////////
 
 	TSharedPtr<SScrollBar> HorizontalScrollBar;
diff --git a/Engine/Source/Developer/TraceInsights/Private/Insights/Widgets/STimersView.cpp b/Engine/Source/Developer/TraceInsights/Private/Insights/Widgets/STimersView.cpp
--- a/Engine/Source/Developer/TraceInsights/Private/Insights/Widgets/STimersView.cpp
+++ b/Engine/Source/Developer/TraceInsights/Private/Insights/Widgets/STimersView.cpp
@@ -2745,18 +2745,10 @@ void STimersView::ToggleGraphFrameStatsSeries(TSharedRef<FTimingGraphTrack> Grap
 	{
 		GraphTrack->Show();
 		Series = GraphTrack->AddFrameStatsTimerSeries(TimerId, FrameType, NodePtr->GetColor());
-		FText NameText = FText::FromName(NodePtr->GetName());
-
-		if (FrameType == ETraceFrameType::TraceFrameType_Game)
-		{
-			NameText = FText::Format((LOCTEXT("GameFrameSeriesName", "{0} Game Frame")), NameText);
-		}
-		else if (FrameType == ETraceFrameType::TraceFrameType_Rendering)
-		{
-			NameText = FText::Format((LOCTEXT("RenderingFrameSeriesName", "{0} Rendering Frame")), NameText);
-		}
-
-		Series->SetName(NameText);
+		FText SeriesName = FText::Format(LOCTEXT("FrameStatsTimerSeriesName_Fmt", "{0} ({1})"),
+			FText::FromName(NodePtr->GetName()),
+			FFrameTrackDrawHelper::FrameTypeToText(FrameType));
+		Series->SetName(SeriesName);
 		GraphTrack->SetDirtyFlag();
 		NodePtr->OnAddedToGraph();
 	}
@@ -2822,7 +2814,9 @@ void STimersView::ToggleFrameTrackSeries(FTimerNodePtr TimerNode, ETraceFrameTyp
 	}
 	else
 	{
-		FText SeriesName = FText::Format(LOCTEXT("FrameTrackSeriesName_Format", "{0} {1} {2}"), TimerNode->GetDisplayName(), FText::FromString(FFrameTrackDrawHelper::FrameTypeToString(FrameType)), LOCTEXT("Frame", "Frame"));
+		FText SeriesName = FText::Format(LOCTEXT("FrameStatsTimerSeriesName_Fmt", "{0} ({1})"),
+			TimerNode->GetDisplayName(),
+			FFrameTrackDrawHelper::FrameTypeToText(FrameType));
 		TSharedPtr<FTimerFrameStatsTrackSeries> Series = FrameTrack->AddTimerFrameStatSeries(FrameType, TimerNode->GetTimerId(), TimerNode->GetColor(), SeriesName);
 		TimerNode->OnAddedToGraph();
 	}
