Removed ClientContext mask from Enumerate test cases as they will only work with ProgramContext or EngineContext

diff --git a/Engine/Source/Developer/TraceInsights/Private/Insights/Tests/FunctionalTests/TimingProfilerTests.cpp b/Engine/Source/Developer/TraceInsights/Private/Insights/Tests/FunctionalTests/TimingProfilerTests.cpp
--- a/Engine/Source/Developer/TraceInsights/Private/Insights/Tests/FunctionalTests/TimingProfilerTests.cpp
+++ b/Engine/Source/Developer/TraceInsights/Private/Insights/Tests/FunctionalTests/TimingProfilerTests.cpp
@@ -18,7 +18,7 @@ DEFINE_LOG_CATEGORY(TimingProfilerTests);
 
 #if WITH_AUTOMATION_TESTS
 
-IMPLEMENT_SIMPLE_AUTOMATION_TEST(FEnumerateTest, "System.Insights.Analysis.TimingInsights.Enumerate", EAutomationTestFlags::ApplicationContextMask | EAutomationTestFlags::EngineFilter)
+IMPLEMENT_SIMPLE_AUTOMATION_TEST(FEnumerateTest, "System.Insights.Analysis.TimingInsights.Enumerate", EAutomationTestFlags::EditorContext | EAutomationTestFlags::ProgramContext | EAutomationTestFlags::EngineFilter)
 bool FEnumerateTest::RunTest(const FString& Parameters)
 {
 	FInsightsTestUtils Utils(this);
@@ -48,7 +48,7 @@ bool FEnumerateTest::RunTest(const FString& Parameters)
 	return !HasAnyErrors();
 }
 
-IMPLEMENT_SIMPLE_AUTOMATION_TEST(FEnumeratePerformanceTest, "System.Insights.Analysis.TimingInsights.EnumeratePerformance", EAutomationTestFlags::ApplicationContextMask | EAutomationTestFlags::PerfFilter)
+IMPLEMENT_SIMPLE_AUTOMATION_TEST(FEnumeratePerformanceTest, "System.Insights.Analysis.TimingInsights.EnumeratePerformance", EAutomationTestFlags::EditorContext | EAutomationTestFlags::ProgramContext | EAutomationTestFlags::EngineFilter)
 bool FEnumeratePerformanceTest::RunTest(const FString& Parameters)
 {
 	FInsightsTestUtils Utils(this);
@@ -88,7 +88,7 @@ bool FEnumeratePerformanceTest::RunTest(const FString& Parameters)
 	return !HasAnyErrors();
 }
 
-IMPLEMENT_SIMPLE_AUTOMATION_TEST(FEnumerateFastTest, "System.Insights.Analysis.TimingInsights.EnumerateFast", EAutomationTestFlags::ApplicationContextMask | EAutomationTestFlags::EngineFilter)
+IMPLEMENT_SIMPLE_AUTOMATION_TEST(FEnumerateFastTest, "System.Insights.Analysis.TimingInsights.EnumerateFast", EAutomationTestFlags::EditorContext | EAutomationTestFlags::ProgramContext | EAutomationTestFlags::EngineFilter)
 bool FEnumerateFastTest::RunTest(const FString& Parameters)
 {
 	FInsightsTestUtils Utils(this);
@@ -118,7 +118,7 @@ bool FEnumerateFastTest::RunTest(const FString& Parameters)
 	return !HasAnyErrors();
 }
 
-IMPLEMENT_SIMPLE_AUTOMATION_TEST(EnumerateEventsToFile, "System.Insights.Trace.Analysis.TimingInsights.EnumerateEventsToFile", EAutomationTestFlags::ProgramContext | EAutomationTestFlags::EngineFilter)
+IMPLEMENT_SIMPLE_AUTOMATION_TEST(EnumerateEventsToFile, "System.Insights.Trace.Analysis.TimingInsights.EnumerateEventsToFile", EAutomationTestFlags::EditorContext | EAutomationTestFlags::ProgramContext | EAutomationTestFlags::EngineFilter)
 bool EnumerateEventsToFile::RunTest(const FString& Parameters)
 {
 	double SessionTime = 0.0;
