[Insights]
  - Session Browser: Increased default window size.
  - Session Browser: Renamed main window to "Unreal Insights Frontend".

diff --git a/Engine/Source/Developer/TraceInsights/Private/Insights/TraceInsightsModule.cpp b/Engine/Source/Developer/TraceInsights/Private/Insights/TraceInsightsModule.cpp
--- a/Engine/Source/Developer/TraceInsights/Private/Insights/TraceInsightsModule.cpp
+++ b/Engine/Source/Developer/TraceInsights/Private/Insights/TraceInsightsModule.cpp
@@ -258,10 +258,10 @@ void FTraceInsightsModule::CreateSessionBrowser(const FCreateSessionBrowserParam
 		static_cast<float>(DisplayMetrics.PrimaryDisplayWorkAreaRect.Left),
 		static_cast<float>(DisplayMetrics.PrimaryDisplayWorkAreaRect.Top));
 
-	const FVector2D ClientSize(960.0f * DPIScaleFactor, 640.0f * DPIScaleFactor);
+	const FVector2D ClientSize(1280.0f * DPIScaleFactor, 720.0f * DPIScaleFactor);
 
 	TSharedRef<SWindow> RootWindow = SNew(SWindow)
-		.Title(NSLOCTEXT("TraceInsightsModule", "UnrealInsightsBrowserAppName", "Unreal Insights Session Browser"))
+		.Title(NSLOCTEXT("TraceInsightsModule", "UnrealInsightsBrowserAppName", "Unreal Insights Frontend"))
 		.CreateTitleBar(!bEmbedTitleAreaContent)
 		.SupportsMaximize(true)
 		.SupportsMinimize(true)
@@ -489,7 +489,6 @@ void FTraceInsightsModule::UnregisterMajorTabConfig(const FName& InMajorTabId)
 	TabConfigs.Remove(InMajorTabId);
 }
 
-
 ////////////////////////////////////////////////////////////////////////////////////////////////////
 
 FOnRegisterMajorTabExtensions& FTraceInsightsModule::OnRegisterMajorTabExtension(const FName& InMajorTabId)
