[Insights] Replaced few occurances of "atof(TCHAR_TO_ANSI(" with "FCString::Atof(".

diff --git a/Engine/Source/Developer/TraceInsights/Private/Insights/TimingProfilerManager.cpp b/Engine/Source/Developer/TraceInsights/Private/Insights/TimingProfilerManager.cpp
--- a/Engine/Source/Developer/TraceInsights/Private/Insights/TimingProfilerManager.cpp
+++ b/Engine/Source/Developer/TraceInsights/Private/Insights/TimingProfilerManager.cpp
@@ -756,14 +756,14 @@ bool FTimingProfilerManager::Exec(const TCHAR* Cmd, FOutputDevice& Ar)
 					// Default: -startTime=-infinite
 					// Example: -startTime=10.0
 					Token.RightChopInline(UE_ARRAY_COUNT(StartTimeToken) - 1);
-					Params.IntervalStartTime = atof(TCHAR_TO_ANSI(*Token));
+					Params.IntervalStartTime = FCString::Atof(*Token);
 				}
 				else if (Token.StartsWith(EndTimeToken))
 				{
 					// Default: -endTime=+infinite
 					// Example: -endTime=20.0
 					Token.RightChopInline(UE_ARRAY_COUNT(EndTimeToken) - 1);
-					Params.IntervalEndTime = atof(TCHAR_TO_ANSI(*Token));
+					Params.IntervalEndTime = FCString::Atof(*Token);
 				}
 				else if (Token.StartsWith(RegionToken))
 				{
@@ -857,14 +857,14 @@ bool FTimingProfilerManager::Exec(const TCHAR* Cmd, FOutputDevice& Ar)
 					// Default: -startTime=-infinite
 					// Example: -startTime=10.0
 					Token.RightChopInline(UE_ARRAY_COUNT(StartTimeToken) - 1);
-					Params.IntervalStartTime = atof(TCHAR_TO_ANSI(*Token));
+					Params.IntervalStartTime = FCString::Atof(*Token);
 				}
 				else if (Token.StartsWith(EndTimeToken))
 				{
 					// Default: -endTime=+infinite
 					// Example: -endTime=20.0
 					Token.RightChopInline(UE_ARRAY_COUNT(EndTimeToken) - 1);
-					Params.IntervalEndTime = atof(TCHAR_TO_ANSI(*Token));
+					Params.IntervalEndTime = FCString::Atof(*Token);
 				}
 				else if (Token.StartsWith(RegionToken))
 				{
@@ -961,14 +961,14 @@ bool FTimingProfilerManager::Exec(const TCHAR* Cmd, FOutputDevice& Ar)
 					// Default: -startTime=-infinite
 					// Example: -startTime=10.0
 					Token.RightChopInline(UE_ARRAY_COUNT(StartTimeToken) - 1);
-					Params.IntervalStartTime = atof(TCHAR_TO_ANSI(*Token));
+					Params.IntervalStartTime = FCString::Atof(*Token);
 				}
 				else if (Token.StartsWith(EndTimeToken))
 				{
 					// Default: -endTime=+infinite
 					// Example: -endTime=20.0
 					Token.RightChopInline(UE_ARRAY_COUNT(EndTimeToken) - 1);
-					Params.IntervalEndTime = atof(TCHAR_TO_ANSI(*Token));
+					Params.IntervalEndTime = FCString::Atof(*Token);
 				}
 				else if (Token.StartsWith(RegionToken))
 				{
diff --git a/Engine/Source/Developer/TraceInsights/Private/Insights/Widgets/STimingView.cpp b/Engine/Source/Developer/TraceInsights/Private/Insights/Widgets/STimingView.cpp
--- a/Engine/Source/Developer/TraceInsights/Private/Insights/Widgets/STimingView.cpp
+++ b/Engine/Source/Developer/TraceInsights/Private/Insights/Widgets/STimingView.cpp
@@ -4221,7 +4221,11 @@ TSharedRef<SWidget> STimingView::MakeCompactAutoScrollOptionsMenu()
 					const FString ValueStr = (AutoScrollViewportOffsetPercent == 0.0) ? FString(TEXT("0")) : FString::Printf(TEXT("%g%%"), AutoScrollViewportOffsetPercent * 100.0);
 					return FText::FromString(ValueStr);
 				})
-				.OnTextChanged_Lambda([this](const FText& InText) { SetAutoScrollViewportOffset(atof(TCHAR_TO_ANSI(*InText.ToString())) * 0.01); })
+				.OnTextChanged_Lambda([this](const FText& InText)
+				{
+					double OffsetPercent = FCString::Atof(*InText.ToString());
+					SetAutoScrollViewportOffset(OffsetPercent * 0.01);
+				})
 			]
 			+ SHorizontalBox::Slot()
 			.FillWidth(1.0f)
@@ -4255,7 +4259,11 @@ TSharedRef<SWidget> STimingView::MakeCompactAutoScrollOptionsMenu()
 					const FString ValueStr = (AutoScrollMinDelay == 0.0) ? FString(TEXT("0")) : FString::Printf(TEXT("%gs"), AutoScrollMinDelay);
 					return FText::FromString(ValueStr);
 				})
-				.OnTextChanged_Lambda([this](const FText& InText) { SetAutoScrollDelay(atof(TCHAR_TO_ANSI(*InText.ToString()))); })
+				.OnTextChanged_Lambda([this](const FText& InText)
+				{
+					double Delay = FCString::Atof(*InText.ToString());
+					SetAutoScrollDelay(Delay);
+				})
 			]
 			+ SHorizontalBox::Slot()
 			.FillWidth(1.0f)
@@ -4379,7 +4387,11 @@ TSharedRef<SWidget> STimingView::MakeAutoScrollOptionsMenu()
 					const FString ValueStr = (AutoScrollViewportOffsetPercent == 0.0) ? FString(TEXT("0")) : FString::Printf(TEXT("%g%%"), AutoScrollViewportOffsetPercent * 100.0);
 					return FText::FromString(ValueStr);
 				})
-				.OnTextChanged_Lambda([this](const FText& InText) { SetAutoScrollViewportOffset(atof(TCHAR_TO_ANSI(*InText.ToString())) * 0.01); })
+				.OnTextChanged_Lambda([this](const FText& InText)
+				{
+					double OffsetPercent = FCString::Atof(*InText.ToString());
+					SetAutoScrollViewportOffset(OffsetPercent * 0.01);
+				})
 			],
 			NAME_None,
 			LOCTEXT("AutoScrollViewportOffsetCustom_Tooltip", "Sets a custom value for the viewport offset as percent from viewport's width (when auto-scrolling)."),
@@ -4453,7 +4465,11 @@ TSharedRef<SWidget> STimingView::MakeAutoScrollOptionsMenu()
 					const FString ValueStr = (AutoScrollMinDelay == 0.0) ? FString(TEXT("0")) : FString::Printf(TEXT("%gs"), AutoScrollMinDelay);
 					return FText::FromString(ValueStr);
 				})
-				.OnTextChanged_Lambda([this](const FText& InText) { SetAutoScrollDelay(atof(TCHAR_TO_ANSI(*InText.ToString()))); })
+				.OnTextChanged_Lambda([this](const FText& InText)
+				{
+					double Delay = FCString::Atof(*InText.ToString());
+					SetAutoScrollDelay(Delay);
+				})
 			],
 			NAME_None,
 			LOCTEXT("AutoScrollDelayCustom_Tooltip", "Sets a custom time delay (in seconds) for the auto-scroll update."),
