[Insights]
  - Frames view: Added colored bullet point for each series in the context menu (including frame series and timer series).
  - Frames view: Added threshold presets in the context menu.
  - Frames view: Made the colors of threshold lines to be more saturated (to be more visible).

diff --git a/Engine/Source/Developer/TraceInsights/Private/Insights/Widgets/SFrameTrack.cpp b/Engine/Source/Developer/TraceInsights/Private/Insights/Widgets/SFrameTrack.cpp
--- a/Engine/Source/Developer/TraceInsights/Private/Insights/Widgets/SFrameTrack.cpp
+++ b/Engine/Source/Developer/TraceInsights/Private/Insights/Widgets/SFrameTrack.cpp
@@ -12,6 +12,7 @@
 #include "Styling/AppStyle.h"
 #include "TraceServices/Model/Frames.h"
 #include "TraceServices/Model/TimingProfiler.h"
+#include "Widgets/Input/SButton.h"
 #include "Widgets/Input/SEditableTextBox.h"
 #include "Widgets/Layout/SScrollBar.h"
 
@@ -440,7 +441,7 @@ void SFrameTrack::UpdateState()
 
 			if (!bTimingViewExists)
 			{
-				// Compute the stats for all timelines. 
+				// Compute the stats for all timelines.
 				Insights::FFrameStatsHelper::ComputeFrameStatsForTimer(Frames, TimerSeries->TimerId);
 			}
 
@@ -649,7 +650,7 @@ int32 SFrameTrack::OnPaint(const FPaintArgs& Args, const FGeometry& AllottedGeom
 
 		if (bShowLowerThresholdLine)
 		{
-			const FLinearColor LineColor(0.5f, 1.0f, 0.5f, 1.0f);
+			const FLinearColor LineColor(0.1f, 0.7f, 0.1f, 1.0f);
 			const FAxisViewportDouble& ViewportY = Viewport.GetVerticalAxisViewport();
 			const float RoundedViewHeight = FMath::RoundToFloat(ViewportY.GetSize());
 			const float LineY = RoundedViewHeight - FMath::RoundToFloat(ViewportY.GetOffsetForValue(LowerThresholdTime));
@@ -658,7 +659,7 @@ int32 SFrameTrack::OnPaint(const FPaintArgs& Args, const FGeometry& AllottedGeom
 		}
 		if (bShowUpperThresholdLine)
 		{
-			const FLinearColor LineColor(1.0f, 0.5f, 0.5f, 1.0f);
+			const FLinearColor LineColor(1.0f, 0.1f, 0.1f, 1.0f);
 			const FAxisViewportDouble& ViewportY = Viewport.GetVerticalAxisViewport();
 			const float RoundedViewHeight = FMath::RoundToFloat(ViewportY.GetSize());
 			const float LineY = RoundedViewHeight - FMath::RoundToFloat(ViewportY.GetOffsetForValue(UpperThresholdTime));
@@ -1329,6 +1330,30 @@ void SFrameTrack::ShowContextMenu(const FPointerEvent& MouseEvent)
 	const bool bShouldCloseWindowAfterMenuSelection = true;
 	FMenuBuilder MenuBuilder(bShouldCloseWindowAfterMenuSelection, NULL);
 
+	auto CreateSeriesMenuWidget = [](FLinearColor InIconColor, FText InText) -> TSharedRef<SWidget>
+	{
+		return SNew(SHorizontalBox)
+
+		+ SHorizontalBox::Slot()
+		.FillWidth(1.0f)
+		.VAlign(VAlign_Center)
+		[
+			SNew(STextBlock)
+			.Text(InText)
+		]
+
+		+ SHorizontalBox::Slot()
+		.Padding(FMargin(8.0f, 0.0f, 8.0f, 0.0f))
+		.AutoWidth()
+		.VAlign(VAlign_Center)
+		[
+			SNew(SImage)
+			.Image(FAppStyle::GetBrush("Icons.FilledCircle"))
+			.DesiredSizeOverride(FVector2D(12.0, 12.0))
+			.ColorAndOpacity(InIconColor)
+		];
+	};
+
 	MenuBuilder.BeginSection("Frames", LOCTEXT("ContextMenu_Section_Frames", "Frames"));
 	{
 		struct FLocal
@@ -1347,11 +1372,12 @@ void SFrameTrack::ShowContextMenu(const FPointerEvent& MouseEvent)
 		);
 		MenuBuilder.AddMenuEntry
 		(
-			LOCTEXT("ContextMenu_ShowGameFrames", "Game Frames"),
-			LOCTEXT("ContextMenu_ShowGameFrames_Desc", "Shows/hides the Game frames."),
-			FSlateIcon(),
 			Action_ShowGameFrames,
+			CreateSeriesMenuWidget(
+				FLinearColor(0.3f, 0.3f, 0.7f, 1.0f),
+				LOCTEXT("ContextMenu_ShowGameFrames", "Game Frames")),
 			NAME_None,
+			LOCTEXT("ContextMenu_ShowGameFrames_Desc", "Shows/hides the Game frames."),
 			EUserInterfaceActionType::ToggleButton
 		);
 
@@ -1363,11 +1389,12 @@ void SFrameTrack::ShowContextMenu(const FPointerEvent& MouseEvent)
 		);
 		MenuBuilder.AddMenuEntry
 		(
-			LOCTEXT("ContextMenu_ShowRenderingFrames", "Rendering Frames"),
-			LOCTEXT("ContextMenu_ShowRenderingFrames_Desc", "Shows/hides the Rendering frames."),
-			FSlateIcon(),
 			Action_ShowRenderingFrames,
+			CreateSeriesMenuWidget(
+				FLinearColor(0.7f, 0.3f, 0.3f, 1.0f),
+				LOCTEXT("ContextMenu_ShowRenderingFrames", "Rendering Frames")),
 			NAME_None,
+			LOCTEXT("ContextMenu_ShowRenderingFrames_Desc", "Shows/hides the Rendering frames."),
 			EUserInterfaceActionType::ToggleButton
 		);
 	}
@@ -1375,7 +1402,6 @@ void SFrameTrack::ShowContextMenu(const FPointerEvent& MouseEvent)
 
 	MenuBuilder.BeginSection("Timers", LOCTEXT("ContextMenu_Section_Timers", "Timers"));
 
-	FText TooltipTextBase = LOCTEXT("ContextMenu_ShowFrameStatsSeries_Desc", "Shows/hides the {0} timer series.");
 	for (TSharedPtr<FFrameTrackSeries> Series : AllSeries)
 	{
 		if (Series->Type != EFrameTrackSeriesType::TimerFrameStats)
@@ -1383,22 +1409,20 @@ void SFrameTrack::ShowContextMenu(const FPointerEvent& MouseEvent)
 			continue;
 		}
 
-		TSharedPtr<FTimerFrameStatsTrackSeries> FrameStatSeries = StaticCastSharedPtr<FTimerFrameStatsTrackSeries>(Series);
-
-		ETraceFrameType FrameType = static_cast<ETraceFrameType>(FrameStatSeries->FrameType);
-		FUIAction Action_ShowRenderingFrames
+		TSharedPtr<FTimerFrameStatsTrackSeries> FrameStatsSeries = StaticCastSharedPtr<FTimerFrameStatsTrackSeries>(Series);
+		ETraceFrameType FrameType = static_cast<ETraceFrameType>(FrameStatsSeries->FrameType);
+		FUIAction Action_ShowFrameStatsSeries
 		(
-			FExecuteAction::CreateSP(this, &SFrameTrack::ContextMenu_ShowFrameStats_Execute, FrameType, FrameStatSeries->TimerId),
-			FCanExecuteAction::CreateSP(this, &SFrameTrack::ContextMenu_ShowFrameStats_CanExecute, FrameType, FrameStatSeries->TimerId),
-			FIsActionChecked::CreateSP(this, &SFrameTrack::ContextMenu_ShowFrameStats_IsChecked, FrameType, FrameStatSeries->TimerId)
+			FExecuteAction::CreateSP(this, &SFrameTrack::ContextMenu_ShowFrameStats_Execute, FrameType, FrameStatsSeries->TimerId),
+			FCanExecuteAction::CreateSP(this, &SFrameTrack::ContextMenu_ShowFrameStats_CanExecute, FrameType, FrameStatsSeries->TimerId),
+			FIsActionChecked::CreateSP(this, &SFrameTrack::ContextMenu_ShowFrameStats_IsChecked, FrameType, FrameStatsSeries->TimerId)
 		);
 		MenuBuilder.AddMenuEntry
 		(
-			FText::Format(LOCTEXT("ContextMenu_ShowFrameStatsSeries_Name", "{0}"), FrameStatSeries->Name),
-			FText::Format(TooltipTextBase, FrameStatSeries->Name),
-			FSlateIcon(),
-			Action_ShowRenderingFrames,
+			Action_ShowFrameStatsSeries,
+			CreateSeriesMenuWidget(FrameStatsSeries->Color, FrameStatsSeries->Name),
 			NAME_None,
+			FText::Format(LOCTEXT("ContextMenu_ShowFrameStatsSeries_Desc", "Shows/hides the {0} timer series."), FrameStatsSeries->Name),
 			EUserInterfaceActionType::ToggleButton
 		);
 	}
@@ -1529,6 +1553,17 @@ void SFrameTrack::CreateThresholdsMenu(FMenuBuilder& MenuBuilder)
 		EUserInterfaceActionType::None
 	);
 
+	MenuBuilder.AddSeparator();
+
+	MenuBuilder.AddMenuEntry
+	(
+		FUIAction(FExecuteAction(), FCanExecuteAction()),
+		CreateThresholdPresetsWidget(),
+		NAME_None,
+		LOCTEXT("ThresholdPresetsTooltip", "Threshold Presets"),
+		EUserInterfaceActionType::None
+	);
+
 	MenuBuilder.EndSection();
 }
 
@@ -1554,7 +1589,7 @@ TSharedRef<SWidget> SFrameTrack::CreateUpperThresholdWidget()
 	.VAlign(VAlign_Center)
 	[
 		SNew(SEditableTextBox)
-		.MinDesiredWidth(60.0f)
+		.MinDesiredWidth(50.0f)
 		.HintText(LOCTEXT("UpperThresholdCustomHint", "30 fps"))
 		.Text_Lambda([this]
 		{
@@ -1576,13 +1611,13 @@ TSharedRef<SWidget> SFrameTrack::CreateUpperThresholdWidget()
 			}
 			if (ValueStr.EndsWith(TEXT("fps")))
 			{
-				double FPS = atof(TCHAR_TO_ANSI(*ValueStr));
+				double FPS = FCString::Atof(*ValueStr);
 				UpperThresholdTime = 1.0 / FMath::Clamp(FPS, 1.0 / MaxThresholdTime, 1.0 / MinThresholdTime);
 				bShowUpperThresholdAsFps = true;
 			}
 			else
 			{
-				double Time = atof(TCHAR_TO_ANSI(*ValueStr));
+				double Time = FCString::Atof(*ValueStr);
 				UpperThresholdTime = FMath::Clamp(Time, MinThresholdTime, MaxThresholdTime);
 				bShowUpperThresholdAsFps = false;
 			}
@@ -1636,7 +1671,7 @@ TSharedRef<SWidget> SFrameTrack::CreateLowerThresholdWidget()
 	.VAlign(VAlign_Center)
 	[
 		SNew(SEditableTextBox)
-		.MinDesiredWidth(60.0f)
+		.MinDesiredWidth(50.0f)
 		.HintText(LOCTEXT("LowerThresholdCustomHint", "60 fps"))
 		.Text_Lambda([this]
 		{
@@ -1658,13 +1693,13 @@ TSharedRef<SWidget> SFrameTrack::CreateLowerThresholdWidget()
 			}
 			if (ValueStr.EndsWith(TEXT("fps")))
 			{
-				double FPS = atof(TCHAR_TO_ANSI(*ValueStr));
+				double FPS = FCString::Atof(*ValueStr);
 				LowerThresholdTime = 1.0 / FMath::Clamp(FPS, 1.0 / MaxThresholdTime, 1.0 / MinThresholdTime);
 				bShowLowerThresholdAsFps = true;
 			}
 			else
 			{
-				double Time = atof(TCHAR_TO_ANSI(*ValueStr));
+				double Time = FCString::Atof(*ValueStr);
 				LowerThresholdTime = FMath::Clamp(Time, MinThresholdTime, MaxThresholdTime);
 				bShowLowerThresholdAsFps = false;
 			}
@@ -1698,6 +1733,83 @@ TSharedRef<SWidget> SFrameTrack::CreateLowerThresholdWidget()
 
 ////////////////////////////////////////////////////////////////////////////////////////////////////
 
+void SFrameTrack::SetThresholdsFPS(double InUpperThresholdFPS, double InLowerThresholdFPS)
+{
+	UpperThresholdTime = 1.0 / InUpperThresholdFPS;
+	LowerThresholdTime = 1.0 / InLowerThresholdFPS;
+	bShowUpperThresholdAsFps = true;
+	bShowLowerThresholdAsFps = true;
+}
+
+////////////////////////////////////////////////////////////////////////////////////////////////////
+
+TSharedRef<SWidget> SFrameTrack::CreateThresholdPresetsWidget()
+{
+	return SNew(SHorizontalBox)
+
+	+ SHorizontalBox::Slot()
+	.Padding(FMargin(-30.0f, 0.0f, 0.0f, 0.0f))
+	.AutoWidth()
+	.VAlign(VAlign_Center)
+	[
+		SNew(SButton)
+		.ContentPadding(FMargin(-6.0f, 0.0f, -6.0f, 0.0f))
+		.Text(LOCTEXT("ThresholdPreset20fps", "15/20 fps"))
+		.OnClicked_Lambda([this]() -> FReply
+		{
+			SetThresholdsFPS(15.0, 20.0);
+			return FReply::Handled();
+		})
+	]
+
+	+ SHorizontalBox::Slot()
+	.Padding(FMargin(4.0f, 0.0f, 0.0f, 0.0f))
+	.AutoWidth()
+	.VAlign(VAlign_Center)
+	[
+		SNew(SButton)
+		.ContentPadding(FMargin(-6.0f, 0.0f, -6.0f, 0.0f))
+		.Text(LOCTEXT("ThresholdPreset30fps", "20/30 fps"))
+		.OnClicked_Lambda([this]() -> FReply
+		{
+			SetThresholdsFPS(20.0, 30.0);
+			return FReply::Handled();
+		})
+	]
+
+	+ SHorizontalBox::Slot()
+	.Padding(FMargin(4.0f, 0.0f, 0.0f, 0.0f))
+	.AutoWidth()
+	.VAlign(VAlign_Center)
+	[
+		SNew(SButton)
+		.ContentPadding(FMargin(-6.0f, 0.0f, -6.0f, 0.0f))
+		.Text(LOCTEXT("ThresholdPreset60fps", "30/60 fps"))
+		.OnClicked_Lambda([this]() -> FReply
+		{
+			SetThresholdsFPS(30.0, 60.0);
+			return FReply::Handled();
+		})
+	]
+
+	+ SHorizontalBox::Slot()
+	.Padding(FMargin(4.0f, 0.0f, 9.0f, 0.0f))
+	.AutoWidth()
+	.VAlign(VAlign_Center)
+	[
+		SNew(SButton)
+		.ContentPadding(FMargin(-6.0f, 0.0f, -6.0f, 0.0f))
+		.Text(LOCTEXT("ThresholdPreset120fps", "60/120 fps"))
+		.OnClicked_Lambda([this]() -> FReply
+		{
+			SetThresholdsFPS(60.0, 120.0);
+			return FReply::Handled();
+		})
+	];
+}
+
+////////////////////////////////////////////////////////////////////////////////////////////////////
+
 void SFrameTrack::CreateSelectedFrameMenu(FMenuBuilder& MenuBuilder)
 {
 	FText SelectedFrameSectionName;
diff --git a/Engine/Source/Developer/TraceInsights/Private/Insights/Widgets/SFrameTrack.h b/Engine/Source/Developer/TraceInsights/Private/Insights/Widgets/SFrameTrack.h
--- a/Engine/Source/Developer/TraceInsights/Private/Insights/Widgets/SFrameTrack.h
+++ b/Engine/Source/Developer/TraceInsights/Private/Insights/Widgets/SFrameTrack.h
@@ -152,6 +152,8 @@ protected:
 	void CreateThresholdsMenu(FMenuBuilder& MenuBuilder);
 	TSharedRef<SWidget> CreateUpperThresholdWidget();
 	TSharedRef<SWidget> CreateLowerThresholdWidget();
+	void SetThresholdsFPS(double InUpperThresholdFPS, double InLowerThresholdFPS);
+	TSharedRef<SWidget> CreateThresholdPresetsWidget();
 	void CreateSelectedFrameMenu(FMenuBuilder& MenuBuilder);
 
 	void ContextMenu_ShowGameFrames_Execute();
