[Insights] Refactored the InsightsSettings code in preparation for adding more settings for Timing Profiler / Frames view.

diff --git a/Engine/Programs/UnrealInsights/Config/DefaultUnrealInsightsSettings.ini b/Engine/Programs/UnrealInsights/Config/DefaultUnrealInsightsSettings.ini
--- a/Engine/Programs/UnrealInsights/Config/DefaultUnrealInsightsSettings.ini
+++ b/Engine/Programs/UnrealInsights/Config/DefaultUnrealInsightsSettings.ini
@@ -10,9 +10,10 @@ DefaultZoomLevel=5.0
 ; Auto hide empty tracks (ex.: ones without timing events in the current viewport).
 ; This option can be changed from the "Tracks" drop down menu.
 bAutoHideEmptyTracks=True
-; If enabled, the Timing view will also be zoomed when a frame is selected in the Frames track.
-; This option can be changed from the context menu of the Frames track.
-bAutoZoomOnFrameSelection=False
+; If enabled, the panning is allowed to continue when the mouse cursor reaches the edges of the screen.
+bAllowPanningOnScreenEdges=False
+; The default state for Auto-Scroll, in the Timing view.
+bAutoScroll=False
 ; Frame Alignment. Controls if auto-scroll should align center of the viewport with start of a frame or not.
 ; Valid options: 'none' to disable frame alignment, 'game' / 'rendering' to align with Game / Rendering frames
 AutoScrollFrameAlignment=game
@@ -22,8 +23,9 @@ AutoScrollFrameAlignment=game
 AutoScrollViewportOffsetPercent=0.1
 ; Minimum time between two auto-scroll updates, in [seconds].
 AutoScrollMinDelay=0.3
-; If enabled, the panning is allowed to continue when the mouse cursor reaches the edges of the screen.
-bAllowPanningOnScreenEdges=False
+; If enabled, the Timing view will also be zoomed when a frame is selected in the Frames track.
+; This option can be changed from the context menu of the Frames track.
+bAutoZoomOnFrameSelection=False
 
 [Insights.TimingProfiler.MainGraph]
 ; If enabled, values will be displayed as points in the Main Graph Track in Timing Insights.
diff --git a/Engine/Source/Developer/TraceInsights/Private/Insights/InsightsSettings.cpp b/Engine/Source/Developer/TraceInsights/Private/Insights/InsightsSettings.cpp
--- a/Engine/Source/Developer/TraceInsights/Private/Insights/InsightsSettings.cpp
+++ b/Engine/Source/Developer/TraceInsights/Private/Insights/InsightsSettings.cpp
@@ -2,4 +2,192 @@
 
 #include "InsightsSettings.h"
 
+#include "Containers/StringConv.h"
+#include "HAL/PlatformMath.h"
+#include "Misc/ConfigCacheIni.h"
+#include "Misc/ConfigContext.h"
+#include "Misc/CString.h"
+#include "ProfilingDebugging/MiscTrace.h"
+
+////////////////////////////////////////////////////////////////////////////////////////////////////
+
 FInsightsSettings FInsightsSettings::Defaults(true);
+
+////////////////////////////////////////////////////////////////////////////////////////////////////
+
+FInsightsSettings::FInsightsSettings(bool bInIsDefault)
+	: bIsDefault(bInIsDefault)
+{
+	if (!bIsDefault)
+	{
+		LoadFromConfig();
+	}
+	else
+	{
+		TimersViewInstanceVisibleColumns.Add(TEXT("Count"));
+		TimersViewInstanceVisibleColumns.Add(TEXT("TotalInclTime"));
+		TimersViewInstanceVisibleColumns.Add(TEXT("TotalExclTime"));
+
+		TimersViewGameFrameVisibleColumns.Add(TEXT("MaxInclTime"));
+		TimersViewGameFrameVisibleColumns.Add(TEXT("AverageInclTime"));
+		TimersViewGameFrameVisibleColumns.Add(TEXT("MedianInclTime"));
+		TimersViewGameFrameVisibleColumns.Add(TEXT("MinInclTime"));
+
+		TimersViewRenderingFrameVisibleColumns.Add(TEXT("MaxInclTime"));
+		TimersViewRenderingFrameVisibleColumns.Add(TEXT("AverageInclTime"));
+		TimersViewRenderingFrameVisibleColumns.Add(TEXT("MedianInclTime"));
+		TimersViewRenderingFrameVisibleColumns.Add(TEXT("MinInclTime"));
+	}
+}
+
+////////////////////////////////////////////////////////////////////////////////////////////////////
+
+FInsightsSettings::~FInsightsSettings()
+{
+}
+
+////////////////////////////////////////////////////////////////////////////////////////////////////
+
+void FInsightsSettings::ResetToDefaults()
+{
+	DefaultZoomLevel = Defaults.DefaultZoomLevel;
+	bAutoHideEmptyTracks = Defaults.bAutoHideEmptyTracks;
+	bAllowPanningOnScreenEdges = Defaults.bAllowPanningOnScreenEdges;
+	bAutoZoomOnFrameSelection = Defaults.bAutoZoomOnFrameSelection;
+	AutoScrollFrameAlignment = Defaults.AutoScrollFrameAlignment;
+	AutoScrollViewportOffsetPercent = Defaults.AutoScrollViewportOffsetPercent;
+	AutoScrollMinDelay = Defaults.AutoScrollMinDelay;
+}
+
+////////////////////////////////////////////////////////////////////////////////////////////////////
+
+void FInsightsSettings::LoadFromConfig()
+{
+	if (!FConfigContext::ReadIntoGConfig().Load(TEXT("UnrealInsightsSettings"), SettingsIni))
+	{
+		return;
+	}
+
+	//////////////////////////////////////////////////
+	// [Insights.TimingProfiler]
+
+	GConfig->GetDouble(TEXT("Insights.TimingProfiler"), TEXT("DefaultZoomLevel"), DefaultZoomLevel, SettingsIni);
+	GConfig->GetBool(TEXT("Insights.TimingProfiler"), TEXT("bAutoHideEmptyTracks"), bAutoHideEmptyTracks, SettingsIni);
+	GConfig->GetBool(TEXT("Insights.TimingProfiler"), TEXT("bAllowPanningOnScreenEdges"), bAllowPanningOnScreenEdges, SettingsIni);
+
+	// Auto-scroll options
+	GConfig->GetBool(TEXT("Insights.TimingProfiler"), TEXT("bAutoScroll"), bAutoScroll, SettingsIni);
+	FString FrameAlignment;
+	if (GConfig->GetString(TEXT("Insights.TimingProfiler"), TEXT("AutoScrollFrameAlignment"), FrameAlignment, SettingsIni))
+	{
+		FrameAlignment.TrimStartAndEndInline();
+		if (FrameAlignment.Equals(TEXT("game"), ESearchCase::IgnoreCase))
+		{
+			static_assert((int32)TraceFrameType_Game == 0, "ETraceFrameType");
+			AutoScrollFrameAlignment = (int32)TraceFrameType_Game;
+		}
+		else if (FrameAlignment.Equals(TEXT("rendering"), ESearchCase::IgnoreCase))
+		{
+			static_assert((int32)TraceFrameType_Rendering == 1, "ETraceFrameType");
+			AutoScrollFrameAlignment = (int32)TraceFrameType_Rendering;
+		}
+		else
+		{
+			AutoScrollFrameAlignment = -1;
+		}
+	}
+	GConfig->GetDouble(TEXT("Insights.TimingProfiler"), TEXT("AutoScrollViewportOffsetPercent"), AutoScrollViewportOffsetPercent, SettingsIni);
+	GConfig->GetDouble(TEXT("Insights.TimingProfiler"), TEXT("AutoScrollMinDelay"), AutoScrollMinDelay, SettingsIni);
+
+	GConfig->GetBool(TEXT("Insights.TimingProfiler"), TEXT("bAutoZoomOnFrameSelection"), bAutoZoomOnFrameSelection, SettingsIni);
+
+	//////////////////////////////////////////////////
+	// [Insights.TimingProfiler.MainGraph]
+
+	GConfig->GetBool(TEXT("Insights.TimingProfiler.MainGraph"), TEXT("ShowPoints"), bTimingViewMainGraphShowPoints, SettingsIni);
+	GConfig->GetBool(TEXT("Insights.TimingProfiler.MainGraph"), TEXT("ShowPointsWithBorder"), bTimingViewMainGraphShowPointsWithBorder, SettingsIni);
+	GConfig->GetBool(TEXT("Insights.TimingProfiler.MainGraph"), TEXT("ShowConnectedLines"), bTimingViewMainGraphShowConnectedLines, SettingsIni);
+	GConfig->GetBool(TEXT("Insights.TimingProfiler.MainGraph"), TEXT("ShowPolygons"), bTimingViewMainGraphShowPolygons, SettingsIni);
+	GConfig->GetBool(TEXT("Insights.TimingProfiler.MainGraph"), TEXT("ShowEventDuration"), bTimingViewMainGraphShowEventDuration, SettingsIni);
+	GConfig->GetBool(TEXT("Insights.TimingProfiler.MainGraph"), TEXT("ShowBars"), bTimingViewMainGraphShowBars, SettingsIni);
+	GConfig->GetBool(TEXT("Insights.TimingProfiler.MainGraph"), TEXT("ShowGameFrames"), bTimingViewMainGraphShowGameFrames, SettingsIni);
+	GConfig->GetBool(TEXT("Insights.TimingProfiler.MainGraph"), TEXT("ShowRenderingFrame"), bTimingViewMainGraphShowRenderingFrames, SettingsIni);
+
+	//////////////////////////////////////////////////
+	// [Insights.TimingProfiler.TimersView]
+
+	GConfig->GetArray(TEXT("Insights.TimingProfiler.TimersView"), TEXT("InstanceColumns"), TimersViewInstanceVisibleColumns, SettingsIni);
+	GConfig->GetArray(TEXT("Insights.TimingProfiler.TimersView"), TEXT("GameFrameColumns"), TimersViewGameFrameVisibleColumns, SettingsIni);
+	GConfig->GetArray(TEXT("Insights.TimingProfiler.TimersView"), TEXT("RenderingFrameColumns"), TimersViewRenderingFrameVisibleColumns, SettingsIni);
+
+	GConfig->GetInt(TEXT("Insights.TimingProfiler.TimersView"), TEXT("Mode"), TimersViewMode, SettingsIni);
+
+	GConfig->GetInt(TEXT("Insights.TimingProfiler.TimersView"), TEXT("GroupingMode"), TimersViewGroupingMode, SettingsIni);
+	GConfig->GetBool(TEXT("Insights.TimingProfiler.TimersView"), TEXT("ShowCpuTimers"), bTimersViewShowCpuTimers, SettingsIni);
+	GConfig->GetBool(TEXT("Insights.TimingProfiler.TimersView"), TEXT("ShowGpuTimers"), bTimersViewShowGpuTimers, SettingsIni);
+	GConfig->GetBool(TEXT("Insights.TimingProfiler.TimersView"), TEXT("ShowZeroCountTimers"), bTimersViewShowZeroCountTimers, SettingsIni);
+
+	//////////////////////////////////////////////////
+	// [Insights.MemoryProfiler]
+
+	GConfig->GetArray(TEXT("Insights.MemoryProfiler"), TEXT("SymbolSearchPaths"), SymbolSearchPaths, SettingsIni);
+}
+
+////////////////////////////////////////////////////////////////////////////////////////////////////
+
+void FInsightsSettings::SaveToConfig()
+{
+	//////////////////////////////////////////////////
+	// [Insights.TimingProfiler]
+
+	GConfig->SetDouble(TEXT("Insights.TimingProfiler"), TEXT("DefaultZoomLevel"), DefaultZoomLevel, SettingsIni);
+	GConfig->SetBool(TEXT("Insights.TimingProfiler"), TEXT("bAutoHideEmptyTracks"), bAutoHideEmptyTracks, SettingsIni);
+	GConfig->SetBool(TEXT("Insights.TimingProfiler"), TEXT("bAllowPanningOnScreenEdges"), bAllowPanningOnScreenEdges, SettingsIni);
+
+	// Auto-scroll options
+	GConfig->SetBool(TEXT("Insights.TimingProfiler"), TEXT("bAutoScroll"), bAutoScroll, SettingsIni);
+	const TCHAR* FrameAlignment = (AutoScrollFrameAlignment == 0) ? TEXT("game") : (AutoScrollFrameAlignment == 1) ? TEXT("rendering") : TEXT("none");
+	GConfig->SetString(TEXT("Insights.TimingProfiler"), TEXT("AutoScrollFrameAlignment"), FrameAlignment, SettingsIni);
+	GConfig->SetDouble(TEXT("Insights.TimingProfiler"), TEXT("AutoScrollViewportOffsetPercent"), AutoScrollViewportOffsetPercent, SettingsIni);
+	GConfig->SetDouble(TEXT("Insights.TimingProfiler"), TEXT("AutoScrollMinDelay"), AutoScrollMinDelay, SettingsIni);
+
+	GConfig->SetBool(TEXT("Insights.TimingProfiler"), TEXT("bAutoZoomOnFrameSelection"), bAutoZoomOnFrameSelection, SettingsIni);
+
+	//////////////////////////////////////////////////
+	// [Insights.TimingProfiler.MainGraph]
+
+	GConfig->SetBool(TEXT("Insights.TimingProfiler.MainGraph"), TEXT("ShowPoints"), bTimingViewMainGraphShowPoints, SettingsIni);
+	GConfig->SetBool(TEXT("Insights.TimingProfiler.MainGraph"), TEXT("ShowPointsWithBorder"), bTimingViewMainGraphShowPointsWithBorder, SettingsIni);
+	GConfig->SetBool(TEXT("Insights.TimingProfiler.MainGraph"), TEXT("ShowConnectedLines"), bTimingViewMainGraphShowConnectedLines, SettingsIni);
+	GConfig->SetBool(TEXT("Insights.TimingProfiler.MainGraph"), TEXT("ShowPolygons"), bTimingViewMainGraphShowPolygons, SettingsIni);
+	GConfig->SetBool(TEXT("Insights.TimingProfiler.MainGraph"), TEXT("ShowEventDuration"), bTimingViewMainGraphShowEventDuration, SettingsIni);
+	GConfig->SetBool(TEXT("Insights.TimingProfiler.MainGraph"), TEXT("ShowBars"), bTimingViewMainGraphShowBars, SettingsIni);
+	GConfig->SetBool(TEXT("Insights.TimingProfiler.MainGraph"), TEXT("ShowGameFrames"), bTimingViewMainGraphShowGameFrames, SettingsIni);
+	GConfig->SetBool(TEXT("Insights.TimingProfiler.MainGraph"), TEXT("ShowRenderingFrame"), bTimingViewMainGraphShowRenderingFrames, SettingsIni);
+
+	//////////////////////////////////////////////////
+	// [Insights.TimingProfiler.TimersView]
+
+	GConfig->SetArray(TEXT("Insights.TimingProfiler.TimersView"), TEXT("InstanceColumns"), TimersViewInstanceVisibleColumns, SettingsIni);
+	GConfig->SetArray(TEXT("Insights.TimingProfiler.TimersView"), TEXT("GameFrameColumns"), TimersViewGameFrameVisibleColumns, SettingsIni);
+	GConfig->SetArray(TEXT("Insights.TimingProfiler.TimersView"), TEXT("RenderingFrameColumns"), TimersViewRenderingFrameVisibleColumns, SettingsIni);
+
+	GConfig->SetInt(TEXT("Insights.TimingProfiler.TimersView"), TEXT("Mode"), TimersViewMode, SettingsIni);
+
+	GConfig->SetInt(TEXT("Insights.TimingProfiler.TimersView"), TEXT("GroupingMode"), TimersViewGroupingMode, SettingsIni);
+	GConfig->SetInt(TEXT("Insights.TimingProfiler.TimersView"), TEXT("ShowCpuTimers"), bTimersViewShowCpuTimers, SettingsIni);
+	GConfig->SetInt(TEXT("Insights.TimingProfiler.TimersView"), TEXT("ShowGpuTimers"), bTimersViewShowGpuTimers, SettingsIni);
+	GConfig->SetInt(TEXT("Insights.TimingProfiler.TimersView"), TEXT("ShowZeroCountTimers"), bTimersViewShowZeroCountTimers, SettingsIni);
+
+	//////////////////////////////////////////////////
+	// [Insights.MemoryProfiler]
+
+	GConfig->SetArray(TEXT("Insights.MemoryProfiler"), TEXT("SymbolSearchPaths"), SymbolSearchPaths, SettingsIni);
+
+	//////////////////////////////////////////////////
+
+	GConfig->Flush(false, SettingsIni);
+}
+
+////////////////////////////////////////////////////////////////////////////////////////////////////
+////////////////////////////////////////////////////////////////////////////////////////////////////
diff --git a/Engine/Source/Developer/TraceInsights/Private/Insights/InsightsSettings.h b/Engine/Source/Developer/TraceInsights/Private/Insights/InsightsSettings.h
--- a/Engine/Source/Developer/TraceInsights/Private/Insights/InsightsSettings.h
+++ b/Engine/Source/Developer/TraceInsights/Private/Insights/InsightsSettings.h
@@ -2,10 +2,8 @@
 
 #pragma once
 
-#include "CoreMinimal.h"
-#include "Misc/ConfigCacheIni.h"
-#include "Misc/ConfigContext.h"
-#include "ProfilingDebugging/MiscTrace.h"
+#include "CoreTypes.h"
+#include "Containers/UnrealString.h"
 
 /** Contains all settings for the Unreal Insights, accessible through the main manager. */
 class FInsightsSettings
@@ -13,154 +11,18 @@ class FInsightsSettings
 	friend class SInsightsSettings;
 
 public:
-	FInsightsSettings(bool bInIsDefault = false)
-		: bIsEditing(false)
-		, bIsDefault(bInIsDefault)
-		, DefaultZoomLevel(5.0) // 5 seconds between major tick marks
-		, bAutoHideEmptyTracks(true)
-		, bAllowPanningOnScreenEdges(false)
-		, bAutoZoomOnFrameSelection(false)
-		, AutoScrollFrameAlignment((int32)TraceFrameType_Game) // -1 = none, 0 = game, 1 = rendering
-		, AutoScrollViewportOffsetPercent(0.1) // scrolls forward 10% of viewport's width
-		, AutoScrollMinDelay(0.3) // [seconds]
-		, TimersViewMode((int32)TraceFrameType_Count)
-		, TimersViewGroupingMode(3) // ByType
-		, bTimersViewShowCpuTimers(true)
-		, bTimersViewShowGpuTimers(true)
-		, bTimersViewShowZeroCountTimers(true)
-		, bTimingViewMainGraphShowPoints(false)
-		, bTimingViewMainGraphShowPointsWithBorder(true)
-		, bTimingViewMainGraphShowConnectedLines(true)
-		, bTimingViewMainGraphShowPolygons(true)
-		, bTimingViewMainGraphShowEventDuration(true)
-		, bTimingViewMainGraphShowBars(false)
-		, bTimingViewMainGraphShowGameFrames(true)
-		, bTimingViewMainGraphShowRenderingFrames(true)
-	{
-		if (!bIsDefault)
-		{
-			LoadFromConfig();
-		}
-		else
-		{
-			TimersViewInstanceVisibleColumns.Add(TEXT("Count"));
-			TimersViewInstanceVisibleColumns.Add(TEXT("TotalInclTime"));
-			TimersViewInstanceVisibleColumns.Add(TEXT("TotalExclTime"));
-
-			TimersViewGameFrameVisibleColumns.Add(TEXT("MaxInclTime"));
-			TimersViewGameFrameVisibleColumns.Add(TEXT("AverageInclTime"));
-			TimersViewGameFrameVisibleColumns.Add(TEXT("MedianInclTime"));
-			TimersViewGameFrameVisibleColumns.Add(TEXT("MinInclTime"));
-
-			TimersViewRenderingFrameVisibleColumns.Add(TEXT("MaxInclTime"));
-			TimersViewRenderingFrameVisibleColumns.Add(TEXT("AverageInclTime"));
-			TimersViewRenderingFrameVisibleColumns.Add(TEXT("MedianInclTime"));
-			TimersViewRenderingFrameVisibleColumns.Add(TEXT("MinInclTime"));
-		}
-	}
+	FInsightsSettings(bool bInIsDefault = false);
+	~FInsightsSettings();
 
-	~FInsightsSettings()
-	{
-	}
+	void LoadFromConfig();
+	void SaveToConfig();
 
-	void LoadFromConfig()
+	const FInsightsSettings& GetDefaults() const
 	{
-		if (!FConfigContext::ReadIntoGConfig().Load(TEXT("UnrealInsightsSettings"), SettingsIni))
-		{
-			return;
-		}
-
-		GConfig->GetDouble(TEXT("Insights.TimingProfiler"), TEXT("DefaultZoomLevel"), DefaultZoomLevel, SettingsIni);
-		GConfig->GetBool(TEXT("Insights.TimingProfiler"), TEXT("bAutoHideEmptyTracks"), bAutoHideEmptyTracks, SettingsIni);
-		GConfig->GetBool(TEXT("Insights.TimingProfiler"), TEXT("bAllowPanningOnScreenEdges"), bAllowPanningOnScreenEdges, SettingsIni);
-		GConfig->GetBool(TEXT("Insights.TimingProfiler"), TEXT("bAutoZoomOnFrameSelection"), bAutoZoomOnFrameSelection, SettingsIni);
-
-		// Auto-scroll options
-		GConfig->GetBool(TEXT("Insights.AutoScroll"), TEXT("bAutoScroll"), bAutoScroll, SettingsIni);
-		FString FrameAlignment;
-		if (GConfig->GetString(TEXT("Insights.TimingProfiler"), TEXT("AutoScrollFrameAlignment"), FrameAlignment, SettingsIni))
-		{
-			FrameAlignment.TrimStartAndEndInline();
-			if (FrameAlignment.Equals(TEXT("game"), ESearchCase::IgnoreCase))
-			{
-				static_assert((int32)TraceFrameType_Game == 0, "ETraceFrameType");
-				AutoScrollFrameAlignment = (int32)TraceFrameType_Game;
-			}
-			else if (FrameAlignment.Equals(TEXT("rendering"), ESearchCase::IgnoreCase))
-			{
-				static_assert((int32)TraceFrameType_Rendering == 1, "ETraceFrameType");
-				AutoScrollFrameAlignment = (int32)TraceFrameType_Rendering;
-			}
-			else
-			{
-				AutoScrollFrameAlignment = -1;
-			}
-		}
-		GConfig->GetDouble(TEXT("Insights.TimingProfiler"), TEXT("AutoScrollViewportOffsetPercent"), AutoScrollViewportOffsetPercent, SettingsIni);
-		GConfig->GetDouble(TEXT("Insights.TimingProfiler"), TEXT("AutoScrollMinDelay"), AutoScrollMinDelay, SettingsIni);
-
-		GConfig->GetArray(TEXT("Insights.MemoryProfiler"), TEXT("SymbolSearchPaths"), SymbolSearchPaths, SettingsIni);
-		
-		GConfig->GetArray(TEXT("Insights.TimingProfiler.TimersView"), TEXT("InstanceColumns"), TimersViewInstanceVisibleColumns, SettingsIni);
-		GConfig->GetArray(TEXT("Insights.TimingProfiler.TimersView"), TEXT("GameFrameColumns"), TimersViewGameFrameVisibleColumns, SettingsIni);
-		GConfig->GetArray(TEXT("Insights.TimingProfiler.TimersView"), TEXT("RenderingFrameColumns"), TimersViewRenderingFrameVisibleColumns, SettingsIni);
-
-		GConfig->GetInt(TEXT("Insights.TimingProfiler.TimersView"), TEXT("Mode"), TimersViewMode, SettingsIni);
-
-		GConfig->GetInt(TEXT("Insights.TimingProfiler.TimersView"), TEXT("GroupingMode"), TimersViewGroupingMode, SettingsIni);
-		GConfig->GetBool(TEXT("Insights.TimingProfiler.TimersView"), TEXT("ShowCpuTimers"), bTimersViewShowCpuTimers, SettingsIni);
-		GConfig->GetBool(TEXT("Insights.TimingProfiler.TimersView"), TEXT("ShowGpuTimers"), bTimersViewShowGpuTimers, SettingsIni);
-		GConfig->GetBool(TEXT("Insights.TimingProfiler.TimersView"), TEXT("ShowZeroCountTimers"), bTimersViewShowZeroCountTimers, SettingsIni);
-
-		GConfig->GetBool(TEXT("Insights.TimingProfiler.MainGraph"), TEXT("ShowPoints"), bTimingViewMainGraphShowPoints, SettingsIni);
-		GConfig->GetBool(TEXT("Insights.TimingProfiler.MainGraph"), TEXT("ShowPointsWithBorder"), bTimingViewMainGraphShowPointsWithBorder, SettingsIni);
-		GConfig->GetBool(TEXT("Insights.TimingProfiler.MainGraph"), TEXT("ShowConnectedLines"), bTimingViewMainGraphShowConnectedLines, SettingsIni);
-		GConfig->GetBool(TEXT("Insights.TimingProfiler.MainGraph"), TEXT("ShowPolygons"), bTimingViewMainGraphShowPolygons, SettingsIni);
-		GConfig->GetBool(TEXT("Insights.TimingProfiler.MainGraph"), TEXT("ShowEventDuration"), bTimingViewMainGraphShowEventDuration, SettingsIni);
-		GConfig->GetBool(TEXT("Insights.TimingProfiler.MainGraph"), TEXT("ShowBars"), bTimingViewMainGraphShowBars, SettingsIni);
-		GConfig->GetBool(TEXT("Insights.TimingProfiler.MainGraph"), TEXT("ShowGameFrames"), bTimingViewMainGraphShowGameFrames, SettingsIni);
-		GConfig->GetBool(TEXT("Insights.TimingProfiler.MainGraph"), TEXT("ShowRenderingFrame"), bTimingViewMainGraphShowRenderingFrames, SettingsIni);
-
+		return Defaults;
 	}
 
-	void SaveToConfig()
-	{
-		GConfig->SetDouble(TEXT("Insights.TimingProfiler"), TEXT("DefaultZoomLevel"), DefaultZoomLevel, SettingsIni);
-		GConfig->SetBool(TEXT("Insights.TimingProfiler"), TEXT("bAutoHideEmptyTracks"), bAutoHideEmptyTracks, SettingsIni);
-		GConfig->SetBool(TEXT("Insights.TimingProfiler"), TEXT("bAllowPanningOnScreenEdges"), bAllowPanningOnScreenEdges, SettingsIni);
-		GConfig->SetBool(TEXT("Insights.TimingProfiler"), TEXT("bAutoZoomOnFrameSelection"), bAutoZoomOnFrameSelection, SettingsIni);
-
-		// Auto-scroll options
-		GConfig->SetBool(TEXT("Insights.AutoScroll"), TEXT("bAutoScroll"), bAutoScroll, SettingsIni);
-		const TCHAR* FrameAlignment = (AutoScrollFrameAlignment == 0) ? TEXT("game") : (AutoScrollFrameAlignment == 1) ? TEXT("rendering") : TEXT("none");
-		GConfig->SetString(TEXT("Insights.TimingProfiler"), TEXT("AutoScrollFrameAlignment"), FrameAlignment, SettingsIni);
-		GConfig->SetDouble(TEXT("Insights.TimingProfiler"), TEXT("AutoScrollViewportOffsetPercent"), AutoScrollViewportOffsetPercent, SettingsIni);
-		GConfig->SetDouble(TEXT("Insights.TimingProfiler"), TEXT("AutoScrollMinDelay"), AutoScrollMinDelay, SettingsIni);
-
-		GConfig->SetArray(TEXT("Insights.MemoryProfiler"), TEXT("SymbolSearchPaths"), SymbolSearchPaths, SettingsIni);
-
-		GConfig->SetArray(TEXT("Insights.TimingProfiler.TimersView"), TEXT("InstanceColumns"), TimersViewInstanceVisibleColumns, SettingsIni);
-		GConfig->SetArray(TEXT("Insights.TimingProfiler.TimersView"), TEXT("GameFrameColumns"), TimersViewGameFrameVisibleColumns, SettingsIni);
-		GConfig->SetArray(TEXT("Insights.TimingProfiler.TimersView"), TEXT("RenderingFrameColumns"), TimersViewRenderingFrameVisibleColumns, SettingsIni);
-
-		GConfig->SetInt(TEXT("Insights.TimingProfiler.TimersView"), TEXT("Mode"), TimersViewMode, SettingsIni);
-
-		GConfig->SetInt(TEXT("Insights.TimingProfiler.TimersView"), TEXT("GroupingMode"), TimersViewGroupingMode, SettingsIni);
-		GConfig->SetInt(TEXT("Insights.TimingProfiler.TimersView"), TEXT("ShowCpuTimers"), bTimersViewShowCpuTimers, SettingsIni);
-		GConfig->SetInt(TEXT("Insights.TimingProfiler.TimersView"), TEXT("ShowGpuTimers"), bTimersViewShowGpuTimers, SettingsIni);
-		GConfig->SetInt(TEXT("Insights.TimingProfiler.TimersView"), TEXT("ShowZeroCountTimers"), bTimersViewShowZeroCountTimers, SettingsIni);
-
-		GConfig->SetBool(TEXT("Insights.TimingProfiler.MainGraph"), TEXT("ShowPoints"), bTimingViewMainGraphShowPoints, SettingsIni);
-		GConfig->SetBool(TEXT("Insights.TimingProfiler.MainGraph"), TEXT("ShowPointsWithBorder"), bTimingViewMainGraphShowPointsWithBorder, SettingsIni);
-		GConfig->SetBool(TEXT("Insights.TimingProfiler.MainGraph"), TEXT("ShowConnectedLines"), bTimingViewMainGraphShowConnectedLines, SettingsIni);
-		GConfig->SetBool(TEXT("Insights.TimingProfiler.MainGraph"), TEXT("ShowPolygons"), bTimingViewMainGraphShowPolygons, SettingsIni);
-		GConfig->SetBool(TEXT("Insights.TimingProfiler.MainGraph"), TEXT("ShowEventDuration"), bTimingViewMainGraphShowEventDuration, SettingsIni);
-		GConfig->SetBool(TEXT("Insights.TimingProfiler.MainGraph"), TEXT("ShowBars"), bTimingViewMainGraphShowBars, SettingsIni);
-		GConfig->SetBool(TEXT("Insights.TimingProfiler.MainGraph"), TEXT("ShowGameFrames"), bTimingViewMainGraphShowGameFrames, SettingsIni);
-		GConfig->SetBool(TEXT("Insights.TimingProfiler.MainGraph"), TEXT("ShowRenderingFrame"), bTimingViewMainGraphShowRenderingFrames, SettingsIni);
-
-		GConfig->Flush(false, SettingsIni);
-	}
+	void ResetToDefaults();
 
 	void EnterEditMode()
 	{
@@ -177,24 +39,11 @@ public:
 		return bIsEditing;
 	}
 
-	const FInsightsSettings& GetDefaults() const
-	{
-		return Defaults;
-	}
-
-	void ResetToDefaults()
-	{
-		DefaultZoomLevel = Defaults.DefaultZoomLevel;
-		bAutoHideEmptyTracks = Defaults.bAutoHideEmptyTracks;
-		bAllowPanningOnScreenEdges = Defaults.bAllowPanningOnScreenEdges;
-		bAutoZoomOnFrameSelection = Defaults.bAutoZoomOnFrameSelection;
-		AutoScrollFrameAlignment = Defaults.AutoScrollFrameAlignment;
-		AutoScrollViewportOffsetPercent = Defaults.AutoScrollViewportOffsetPercent;
-		AutoScrollMinDelay = Defaults.AutoScrollMinDelay;
-	}
-
 	#define SET_AND_SAVE(Option, Value) { if (Option != Value) { Option = Value; SaveToConfig(); } }
 
+	//////////////////////////////////////////////////
+	// [Insights.TimingProfiler]
+
 	double GetDefaultZoomLevel() const { return DefaultZoomLevel; }
 	void SetDefaultZoomLevel(double ZoomLevel) { DefaultZoomLevel = ZoomLevel; }
 	void SetAndSaveDefaultZoomLevel(double ZoomLevel) { SET_AND_SAVE(DefaultZoomLevel, ZoomLevel); }
@@ -207,10 +56,6 @@ public:
 	void SetPanningOnScreenEdges(bool bOnOff) { bAllowPanningOnScreenEdges = bOnOff; }
 	void SetAndSavePanningOnScreenEdges(bool bOnOff) { SET_AND_SAVE(bAllowPanningOnScreenEdges, bOnOff); }
 
-	bool IsAutoZoomOnFrameSelectionEnabled() const { return bAutoZoomOnFrameSelection; }
-	void SetAutoZoomOnFrameSelection(bool bOnOff) { bAutoZoomOnFrameSelection = bOnOff; }
-	void SetAndSaveAutoZoomOnFrameSelection(bool bOnOff) { SET_AND_SAVE(bAutoZoomOnFrameSelection, bOnOff); }
-
 	bool IsAutoScrollEnabled() const { return bAutoScroll; }
 	void SetAutoScroll(bool bOnOff) { bAutoScroll = bOnOff; }
 	void SetAndSaveAutoScroll(bool bOnOff) { SET_AND_SAVE(bAutoScroll, bOnOff); }
@@ -227,9 +72,47 @@ public:
 	void SetAutoScrollMinDelay(double Delay) { AutoScrollMinDelay = Delay; }
 	void SetAndSaveAutoScrollMinDelay(double Delay) { SET_AND_SAVE(AutoScrollMinDelay, Delay); }
 
-	const TArray<FString>& GetSymbolSearchPaths() const { return SymbolSearchPaths; }
-	void SetSymbolSearchPaths(const TArray<FString>& SearchPaths) { SymbolSearchPaths = SearchPaths; }
-	void SetAndSaveSymbolSearchPaths(const TArray<FString>& SearchPaths) { SET_AND_SAVE(SymbolSearchPaths, SearchPaths); }
+	bool IsAutoZoomOnFrameSelectionEnabled() const { return bAutoZoomOnFrameSelection; }
+	void SetAutoZoomOnFrameSelection(bool bOnOff) { bAutoZoomOnFrameSelection = bOnOff; }
+	void SetAndSaveAutoZoomOnFrameSelection(bool bOnOff) { SET_AND_SAVE(bAutoZoomOnFrameSelection, bOnOff); }
+
+	//////////////////////////////////////////////////
+	// [Insights.TimingProfiler.MainGraph]
+
+	bool GetTimingViewMainGraphShowPoints() const { return bTimingViewMainGraphShowPoints; }
+	void SetTimingViewMainGraphShowPoints(bool InValue) { bTimingViewMainGraphShowPoints = InValue; }
+	void SetAndSaveTimingViewMainGraphShowPoints(bool InValue) { SET_AND_SAVE(bTimingViewMainGraphShowPoints, InValue); }
+
+	bool GetTimingViewMainGraphShowPointsWithBorder() const { return bTimingViewMainGraphShowPointsWithBorder; }
+	void SetTimingViewMainGraphShowPointsWithBorder(bool InValue) { bTimingViewMainGraphShowPointsWithBorder = InValue; }
+	void SetAndSaveTimingViewMainGraphShowPointsWithBorder(bool InValue) { SET_AND_SAVE(bTimingViewMainGraphShowPointsWithBorder, InValue); }
+
+	bool GetTimingViewMainGraphShowConnectedLines() const { return bTimingViewMainGraphShowConnectedLines; }
+	void SetTimingViewMainGraphShowConnectedLines(bool InValue) { bTimingViewMainGraphShowConnectedLines = InValue; }
+	void SetAndSaveTimingViewMainGraphShowConnectedLines(bool InValue) { SET_AND_SAVE(bTimingViewMainGraphShowConnectedLines, InValue); }
+
+	bool GetTimingViewMainGraphShowPolygons() const { return bTimingViewMainGraphShowPolygons; }
+	void SetTimingViewMainGraphShowPolygons(bool InValue) { bTimingViewMainGraphShowPolygons = InValue; }
+	void SetAndTimingViewMainGraphShowPolygons(bool InValue) { SET_AND_SAVE(bTimingViewMainGraphShowPolygons, InValue); }
+
+	bool GetTimingViewMainGraphShowEventDuration() const { return bTimingViewMainGraphShowEventDuration; }
+	void SetTimingViewMainGraphShowEventDuration(bool InValue) { bTimingViewMainGraphShowEventDuration = InValue; }
+	void SetAndSaveTimingViewMainGraphShowEventDuration(bool InValue) { SET_AND_SAVE(bTimingViewMainGraphShowEventDuration, InValue); }
+
+	bool GetTimingViewMainGraphShowBars() const { return bTimingViewMainGraphShowBars; }
+	void SetTimingViewMainGraphShowBars(bool InValue) { bTimingViewMainGraphShowBars = InValue; }
+	void SetAndSaveTimingViewMainGraphShowBars(bool InValue) { SET_AND_SAVE(bTimingViewMainGraphShowBars, InValue); }
+
+	bool GetTimingViewMainGraphShowGameFrames() const { return bTimingViewMainGraphShowGameFrames; }
+	void SetTimingViewMainGraphShowGameFrames(bool InValue) { bTimingViewMainGraphShowGameFrames = InValue; }
+	void SetAndSaveTimingViewMainGraphShowGameFrames(bool InValue) { SET_AND_SAVE(bTimingViewMainGraphShowGameFrames, InValue); }
+
+	bool GetTimingViewMainGraphShowRenderingFrames() const { return bTimingViewMainGraphShowRenderingFrames; }
+	void SetTimingViewMainGraphShowRenderingFrames(bool InValue) { bTimingViewMainGraphShowRenderingFrames = InValue; }
+	void SetAndSaveTimingViewMainGraphShowRenderingFrames(bool InValue) { SET_AND_SAVE(bTimingViewMainGraphShowRenderingFrames, InValue); }
+
+	//////////////////////////////////////////////////
+	// [Insights.TimingProfiler.TimersView]
 
 	const TArray<FString>& GetTimersViewInstanceVisibleColumns() const { return TimersViewInstanceVisibleColumns; }
 	void SetTimersViewInstanceVisibleColumns(const TArray<FString>& Columns) { TimersViewInstanceVisibleColumns = Columns; }
@@ -263,37 +146,14 @@ public:
 	void SetTimersViewShowZeroCountTimers(bool InValue) { bTimersViewShowZeroCountTimers = InValue; }
 	void SetAndSaveTimersViewShowZeroCountTimers(bool InValue) { SET_AND_SAVE(bTimersViewShowZeroCountTimers, InValue); }
 
-	bool GetTimingViewMainGraphShowPoints() const { return bTimingViewMainGraphShowPoints; }
-	void SetTimingViewMainGraphShowPoints(bool InValue) { bTimingViewMainGraphShowPoints = InValue; }
-	void SetAndSaveTimingViewMainGraphShowPoints(bool InValue) { SET_AND_SAVE(bTimingViewMainGraphShowPoints, InValue); }
-
-	bool GetTimingViewMainGraphShowPointsWithBorder() const { return bTimingViewMainGraphShowPointsWithBorder; }
-	void SetTimingViewMainGraphShowPointsWithBorder(bool InValue) { bTimingViewMainGraphShowPointsWithBorder = InValue; }
-	void SetAndSaveTimingViewMainGraphShowPointsWithBorder(bool InValue) { SET_AND_SAVE(bTimingViewMainGraphShowPointsWithBorder, InValue); }
-
-	bool GetTimingViewMainGraphShowConnectedLines() const { return bTimingViewMainGraphShowConnectedLines; }
-	void SetTimingViewMainGraphShowConnectedLines(bool InValue) { bTimingViewMainGraphShowConnectedLines = InValue; }
-	void SetAndSaveTimingViewMainGraphShowConnectedLines(bool InValue) { SET_AND_SAVE(bTimingViewMainGraphShowConnectedLines, InValue); }
-
-	bool GetTimingViewMainGraphShowPolygons() const { return bTimingViewMainGraphShowPolygons; }
-	void SetTimingViewMainGraphShowPolygons(bool InValue) { bTimingViewMainGraphShowPolygons = InValue; }
-	void SetAndTimingViewMainGraphShowPolygons(bool InValue) { SET_AND_SAVE(bTimingViewMainGraphShowPolygons, InValue); }
-
-	bool GetTimingViewMainGraphShowEventDuration() const { return bTimingViewMainGraphShowEventDuration; }
-	void SetTimingViewMainGraphShowEventDuration(bool InValue) { bTimingViewMainGraphShowEventDuration = InValue; }
-	void SetAndSaveTimingViewMainGraphShowEventDuration(bool InValue) { SET_AND_SAVE(bTimingViewMainGraphShowEventDuration, InValue); }
-
-	bool GetTimingViewMainGraphShowBars() const { return bTimingViewMainGraphShowBars; }
-	void SetTimingViewMainGraphShowBars(bool InValue) { bTimingViewMainGraphShowBars = InValue; }
-	void SetAndSaveTimingViewMainGraphShowBars(bool InValue) { SET_AND_SAVE(bTimingViewMainGraphShowBars, InValue); }
+	//////////////////////////////////////////////////
+	// [Insights.MemoryProfiler]
 
-	bool GetTimingViewMainGraphShowGameFrames() const { return bTimingViewMainGraphShowGameFrames; }
-	void SetTimingViewMainGraphShowGameFrames(bool InValue) { bTimingViewMainGraphShowGameFrames = InValue; }
-	void SetAndSaveTimingViewMainGraphShowGameFrames(bool InValue) { SET_AND_SAVE(bTimingViewMainGraphShowGameFrames, InValue); }
+	const TArray<FString>& GetSymbolSearchPaths() const { return SymbolSearchPaths; }
+	void SetSymbolSearchPaths(const TArray<FString>& SearchPaths) { SymbolSearchPaths = SearchPaths; }
+	void SetAndSaveSymbolSearchPaths(const TArray<FString>& SearchPaths) { SET_AND_SAVE(SymbolSearchPaths, SearchPaths); }
 
-	bool GetTimingViewMainGraphShowRenderingFrames() const { return bTimingViewMainGraphShowRenderingFrames; }
-	void SetTimingViewMainGraphShowRenderingFrames(bool InValue) { bTimingViewMainGraphShowRenderingFrames = InValue; }
-	void SetAndSaveTimingViewMainGraphShowRenderingFrames(bool InValue) { SET_AND_SAVE(bTimingViewMainGraphShowRenderingFrames, InValue); }
+	//////////////////////////////////////////////////
 
 	#undef SET_AND_SAVE
 
@@ -301,94 +161,108 @@ private:
 	/** Contains default settings. */
 	static FInsightsSettings Defaults;
 
-	/** Setting filename ini. */
-	FString SettingsIni;
+	/** Whether this instance contains defaults. */
+	bool bIsDefault = false;
 
 	/** Whether profiler settings is in edit mode. */
-	bool bIsEditing;
+	bool bIsEditing = false;
 
-	/** Whether this instance contains defaults. */
-	bool bIsDefault;
+	/** Settings filename ini. */
+	FString SettingsIni;
 
 	//////////////////////////////////////////////////
-	// Actual settings.
+	// [Insights.TimingProfiler]
 
 	/** The default (initial) zoom level of the Timing view. */
-	double DefaultZoomLevel;
+	double DefaultZoomLevel = 5.0; // 5 seconds between major tick marks
 
 	/** Auto hide empty tracks (ex.: ones without timing events in the current viewport). */
-	bool bAutoHideEmptyTracks;
+	bool bAutoHideEmptyTracks = true;
 
 	/** If enabled, the panning is allowed to continue when mouse cursor reaches the edges of the screen. */
-	bool bAllowPanningOnScreenEdges;
+	bool bAllowPanningOnScreenEdges = false;
 
-	/** If enabled, the Timing View will also be zoomed when a new frame is selected in the Frames track. */
-	bool bAutoZoomOnFrameSelection;
+	/** If enabled, the Timing View will start with auto-scroll enabled. */
+	bool bAutoScroll = false;
 
 	/** -1 to disable frame alignment or the type of frame to align with (0 = Game or 1 = Rendering). */
-	int32 AutoScrollFrameAlignment;
-
-	/** List of search paths to look for symbol files */
-	TArray<FString> SymbolSearchPaths;
-
-	/** If enabled, the Timing View will start with auto-scroll enabled. */
-	bool bAutoScroll;
+	int32 AutoScrollFrameAlignment = 0; // -1 = none, 0 = game, 1 = rendering
 
 	/**
 	 * Viewport offset while auto-scrolling, as percent of viewport width.
 	 * If positive, it offsets the viewport forward, allowing an empty space at the right side of the viewport (i.e. after end of session).
 	 * If negative, it offsets the viewport backward (i.e. end of session will be outside viewport).
 	 */
-	double AutoScrollViewportOffsetPercent;
+	double AutoScrollViewportOffsetPercent = 0.1; // scrolls forward 10% of viewport's width
 
 	/** Minimum time between two auto-scroll updates, in [seconds]. */
-	double AutoScrollMinDelay;
+	double AutoScrollMinDelay = 0.3; // [seconds]
 
-	//** The list of visible columns in the Timers view in the Instance mode. */
-	TArray<FString> TimersViewInstanceVisibleColumns;
+	/** If enabled, the Timing View will also be zoomed when a new frame is selected in the Frames track. */
+	bool bAutoZoomOnFrameSelection = false;
 
-	//** The list of visible columns in the Timers view in the Game Frame mode. */
-	TArray<FString> TimersViewGameFrameVisibleColumns;
+	//////////////////////////////////////////////////
+	// [Insights.TimingProfiler.MainGraph]
 
-	//** The list of visible columns in the Timers view in the Rendering Frame mode. */
-	TArray<FString> TimersViewRenderingFrameVisibleColumns;
+	/** If enabled, values will be displayed as points in the Main Graph Track in Timing Insights. */
+	bool bTimingViewMainGraphShowPoints = false;
 
-	//** The mode for the timers panel. */
-	int32 TimersViewMode;
+	/** If enabled, values will be displayed as points with border in the Main Graph Track in Timing Insights. */
+	bool bTimingViewMainGraphShowPointsWithBorder = true;
 
-	//** The grouping mode for the timers panel. */
-	int32 TimersViewGroupingMode;
+	/** If enabled, values will be displayed as connected lines in the Main Graph Track in Timing Insights. */
+	bool bTimingViewMainGraphShowConnectedLines = true;
 
-	//** If enabled, Cpu timers will be displayed in the Timing View. */
-	bool bTimersViewShowCpuTimers;
+	/** If enabled, values will be displayed as polygons in the Main Graph Track in Timing Insights. */
+	bool bTimingViewMainGraphShowPolygons = true;
 
-	//** If enabled, Gpu timers will be displayed in the Timing View. */
-	bool bTimersViewShowGpuTimers;
+	/** If enabled, uses duration of timing events for connected lines and polygons in the Main Graph Track in Timing Insights. */
+	bool bTimingViewMainGraphShowEventDuration = true;
 
-	//** If enabled, timers with no instances in the selected interval will still be displayed in the Timers View. */
-	bool bTimersViewShowZeroCountTimers;
+	/** If enabled, shows bars corresponding to the duration of the timing events in the Main Graph Track in Timing Insights. */
+	bool bTimingViewMainGraphShowBars = false;
 
-	//** If enabled, values will be displayed as points in the Main Graph Track in Timing Insights. */
-	bool bTimingViewMainGraphShowPoints;
+	/** If enabled, shows game frames in the Main Graph Track in Timing Insights. */
+	bool bTimingViewMainGraphShowGameFrames = true;
 
-	//** If enabled, values will be displayed as points with border in the Main Graph Track in Timing Insights. */
-	bool bTimingViewMainGraphShowPointsWithBorder;
+	/** If enabled, shows rendering frames in the Main Graph Track in Timing Insights. */
+	bool bTimingViewMainGraphShowRenderingFrames = true;
 
-	//** If enabled, values will be displayed as connected lines in the Main Graph Track in Timing Insights. */
-	bool bTimingViewMainGraphShowConnectedLines;
+	//////////////////////////////////////////////////
+	// [Insights.TimingProfiler.TimersView]
+
+	/** The list of visible columns in the Timers view in the Instance mode. */
+	TArray<FString> TimersViewInstanceVisibleColumns;
+
+	/** The list of visible columns in the Timers view in the Game Frame mode. */
+	TArray<FString> TimersViewGameFrameVisibleColumns;
+
+	/** The list of visible columns in the Timers view in the Rendering Frame mode. */
+	TArray<FString> TimersViewRenderingFrameVisibleColumns;
+
+	/**
+	 * The mode for the timers panel.
+	 * See ETraceFrameType in MiscTrace.h.
+	 */
+	int32 TimersViewMode = 2; // (int32)TraceFrameType_Count
+
+	/** The grouping mode for the timers panel. */
+	int32 TimersViewGroupingMode = 3; // ByType
+
+	/** If enabled, CPU timers will be displayed in the Timing View. */
+	bool bTimersViewShowCpuTimers = true;
 
-	//** If enabled, values will be displayed as polygons in the Main Graph Track in Timing Insights. */
-	bool bTimingViewMainGraphShowPolygons;
+	/** If enabled, GPU timers will be displayed in the Timing View. */
+	bool bTimersViewShowGpuTimers = true;
 
-	//** If enabled, uses duration of timing events for connected lines and polygons in the Main Graph Track in Timing Insights. */
-	bool bTimingViewMainGraphShowEventDuration;
+	/** If enabled, timers with no instances in the selected interval will still be displayed in the Timers View. */
+	bool bTimersViewShowZeroCountTimers = true;
 
-	//** If enabled, shows bars corresponding to the duration of the timing events in the Main Graph Track in Timing Insights. */
-	bool bTimingViewMainGraphShowBars;
+	//////////////////////////////////////////////////
+	// [Insights.MemoryProfiler]
 
-	//** If enabled, shows game frames in the Main Graph Track in Timing Insights. */
-	bool bTimingViewMainGraphShowGameFrames;
+	/** List of search paths to look for symbol files */
+	TArray<FString> SymbolSearchPaths;
 
-	//** If enabled, shows rendering frames in the Main Graph Track in Timing Insights. */
-	bool bTimingViewMainGraphShowRenderingFrames;
+	//////////////////////////////////////////////////
 };
