[Insights] Persistent settings for the Frames panel.

diff --git a/Engine/Programs/UnrealInsights/Config/DefaultUnrealInsightsSettings.ini b/Engine/Programs/UnrealInsights/Config/DefaultUnrealInsightsSettings.ini
--- a/Engine/Programs/UnrealInsights/Config/DefaultUnrealInsightsSettings.ini
+++ b/Engine/Programs/UnrealInsights/Config/DefaultUnrealInsightsSettings.ini
@@ -23,6 +23,16 @@ AutoScrollFrameAlignment=game
 AutoScrollViewportOffsetPercent=0.1
 ; Minimum time between two auto-scroll updates, in [seconds].
 AutoScrollMinDelay=0.3
+
+[Insights.TimingProfiler.FramesView]
+; If enabled, the upper threshold line is visible. The frame coloring by threshold is enabled regardless of this setting.
+bShowUpperThresholdLine=False
+; If enabled, the lower threshold line is visible. The frame coloring by threshold is enabled regardless of this setting.
+bShowLowerThresholdLine=False
+; The upper threshold for frames. Can be specified as a frame duration ([0.001 .. 1.0] seconds; ex.: "0.010" for 10 ms) or as a framerate ([1 fps .. 1000 fps]; ex: "100 fps").
+UpperThreshold=30 fps
+; The lower threshold for frames. Can be specified as a frame duration ([0.001 .. 1.0] seconds; ex.: "0.010" for 10 ms) or as a framerate ([1 fps .. 1000 fps]; ex: "100 fps").
+LowerThreshold=60 fps
 ; If enabled, the Timing view will also be zoomed when a frame is selected in the Frames track.
 ; This option can be changed from the context menu of the Frames track.
 bAutoZoomOnFrameSelection=False
diff --git a/Engine/Source/Developer/TraceInsights/Private/Insights/InsightsSettings.cpp b/Engine/Source/Developer/TraceInsights/Private/Insights/InsightsSettings.cpp
--- a/Engine/Source/Developer/TraceInsights/Private/Insights/InsightsSettings.cpp
+++ b/Engine/Source/Developer/TraceInsights/Private/Insights/InsightsSettings.cpp
@@ -99,7 +99,58 @@ void FInsightsSettings::LoadFromConfig()
 	GConfig->GetDouble(TEXT("Insights.TimingProfiler"), TEXT("AutoScrollViewportOffsetPercent"), AutoScrollViewportOffsetPercent, SettingsIni);
 	GConfig->GetDouble(TEXT("Insights.TimingProfiler"), TEXT("AutoScrollMinDelay"), AutoScrollMinDelay, SettingsIni);
 
-	GConfig->GetBool(TEXT("Insights.TimingProfiler"), TEXT("bAutoZoomOnFrameSelection"), bAutoZoomOnFrameSelection, SettingsIni);
+	//////////////////////////////////////////////////
+	// [Insights.TimingProfiler.FramesView]
+
+	GConfig->GetBool(TEXT("Insights.TimingProfiler.FramesView"), TEXT("bShowUpperThresholdLine"), bShowUpperThresholdLine, SettingsIni);
+	GConfig->GetBool(TEXT("Insights.TimingProfiler.FramesView"), TEXT("bShowLowerThresholdLine"), bShowLowerThresholdLine, SettingsIni);
+
+	static constexpr double MinThresholdTime = 0.001; // == 1ms == 1000 fps
+	static constexpr double MaxThresholdTime = 1.0; // == 1s == 1 fps
+
+	FString UpperThreshold;
+	if (GConfig->GetString(TEXT("Insights.TimingProfiler.FramesView"), TEXT("UpperThreshold"), UpperThreshold, SettingsIni))
+	{
+		if (UpperThreshold.IsEmpty())
+		{
+			UpperThreshold = TEXT("30 fps");
+		}
+		if (UpperThreshold.EndsWith(TEXT("fps")))
+		{
+			double FPS = FCString::Atof(*UpperThreshold);
+			UpperThresholdTime = 1.0 / FMath::Clamp(FPS, 1.0 / MaxThresholdTime, 1.0 / MinThresholdTime);
+			bShowUpperThresholdAsFps = true;
+		}
+		else
+		{
+			double Time = FCString::Atof(*UpperThreshold);
+			UpperThresholdTime = FMath::Clamp(Time, MinThresholdTime, MaxThresholdTime);
+			bShowUpperThresholdAsFps = false;
+		}
+	}
+
+	FString LowerThreshold;
+	if (GConfig->GetString(TEXT("Insights.TimingProfiler.FramesView"), TEXT("LowerThreshold"), LowerThreshold, SettingsIni))
+	{
+		if (LowerThreshold.IsEmpty())
+		{
+			LowerThreshold = TEXT("60 fps");
+		}
+		if (LowerThreshold.EndsWith(TEXT("fps")))
+		{
+			double FPS = FCString::Atof(*LowerThreshold);
+			LowerThresholdTime = 1.0 / FMath::Clamp(FPS, 1.0 / MaxThresholdTime, 1.0 / MinThresholdTime);
+			bShowLowerThresholdAsFps = true;
+		}
+		else
+		{
+			double Time = FCString::Atof(*LowerThreshold);
+			LowerThresholdTime = FMath::Clamp(Time, MinThresholdTime, MaxThresholdTime);
+			bShowLowerThresholdAsFps = false;
+		}
+	}
+
+	GConfig->GetBool(TEXT("Insights.TimingProfiler.FramesView"), TEXT("bAutoZoomOnFrameSelection"), bAutoZoomOnFrameSelection, SettingsIni);
 
 	//////////////////////////////////////////////////
 	// [Insights.TimingProfiler.MainGraph]
@@ -151,7 +202,35 @@ void FInsightsSettings::SaveToConfig()
 	GConfig->SetDouble(TEXT("Insights.TimingProfiler"), TEXT("AutoScrollViewportOffsetPercent"), AutoScrollViewportOffsetPercent, SettingsIni);
 	GConfig->SetDouble(TEXT("Insights.TimingProfiler"), TEXT("AutoScrollMinDelay"), AutoScrollMinDelay, SettingsIni);
 
-	GConfig->SetBool(TEXT("Insights.TimingProfiler"), TEXT("bAutoZoomOnFrameSelection"), bAutoZoomOnFrameSelection, SettingsIni);
+	//////////////////////////////////////////////////
+	// [Insights.TimingProfiler.FramesView]
+
+	GConfig->SetBool(TEXT("Insights.TimingProfiler.FramesView"), TEXT("bShowUpperThresholdLine"), bShowUpperThresholdLine, SettingsIni);
+	GConfig->SetBool(TEXT("Insights.TimingProfiler.FramesView"), TEXT("bShowLowerThresholdLine"), bShowLowerThresholdLine, SettingsIni);
+
+	if (bShowUpperThresholdAsFps)
+	{
+		FString UpperThreshold = FString::Printf(TEXT("%g fps"), 1.0 / UpperThresholdTime);
+		GConfig->SetString(TEXT("Insights.TimingProfiler.FramesView"), TEXT("UpperThreshold"), *UpperThreshold, SettingsIni);
+	}
+	else
+	{
+		FString UpperThreshold = FString::Printf(TEXT("%g"), UpperThresholdTime);
+		GConfig->SetString(TEXT("Insights.TimingProfiler.FramesView"), TEXT("UpperThreshold"), *UpperThreshold, SettingsIni);
+	}
+
+	if (bShowLowerThresholdAsFps)
+	{
+		FString LowerThreshold = FString::Printf(TEXT("%g fps"), 1.0 / LowerThresholdTime);
+		GConfig->SetString(TEXT("Insights.TimingProfiler.FramesView"), TEXT("LowerThreshold"), *LowerThreshold, SettingsIni);
+	}
+	else
+	{
+		FString LowerThreshold = FString::Printf(TEXT("%g"), LowerThresholdTime);
+		GConfig->SetString(TEXT("Insights.TimingProfiler.FramesView"), TEXT("LowerThreshold"), *LowerThreshold, SettingsIni);
+	}
+
+	GConfig->SetBool(TEXT("Insights.TimingProfiler.FramesView"), TEXT("bAutoZoomOnFrameSelection"), bAutoZoomOnFrameSelection, SettingsIni);
 
 	//////////////////////////////////////////////////
 	// [Insights.TimingProfiler.MainGraph]
diff --git a/Engine/Source/Developer/TraceInsights/Private/Insights/InsightsSettings.h b/Engine/Source/Developer/TraceInsights/Private/Insights/InsightsSettings.h
--- a/Engine/Source/Developer/TraceInsights/Private/Insights/InsightsSettings.h
+++ b/Engine/Source/Developer/TraceInsights/Private/Insights/InsightsSettings.h
@@ -72,6 +72,62 @@ public:
 	void SetAutoScrollMinDelay(double Delay) { AutoScrollMinDelay = Delay; }
 	void SetAndSaveAutoScrollMinDelay(double Delay) { SET_AND_SAVE(AutoScrollMinDelay, Delay); }
 
+	//////////////////////////////////////////////////
+	// [Insights.TimingProfiler.FramesView]
+
+	bool IsShowUpperThresholdLineEnabled() const { return bShowUpperThresholdLine; }
+	void SetShowUpperThresholdLineEnabled(bool bOnOff) { bShowUpperThresholdLine = bOnOff; }
+	void SetAndSaveShowUpperThresholdLineEnabled(bool bOnOff) { SET_AND_SAVE(bShowUpperThresholdLine, bOnOff); }
+
+	bool IsShowLowerThresholdLineEnabled() const { return bShowLowerThresholdLine; }
+	void SetShowLowerThresholdLineEnabled(bool bOnOff) { bShowLowerThresholdLine = bOnOff; }
+	void SetAndSaveShowLowerThresholdLineEnabled(bool bOnOff) { SET_AND_SAVE(bShowLowerThresholdLine, bOnOff); }
+
+	double GetUpperThresholdTime() const { return UpperThresholdTime; }
+	void SetUpperThresholdTime(double InUpperThresholdTime) { UpperThresholdTime = InUpperThresholdTime; }
+	void SetAndSaveUpperThresholdTime(double InUpperThresholdTime) { SET_AND_SAVE(UpperThresholdTime, InUpperThresholdTime); }
+
+	double GetLowerThresholdTime() const { return LowerThresholdTime; }
+	void SetLowerThresholdTime(double InLowerThresholdTime) { LowerThresholdTime = InLowerThresholdTime; }
+	void SetAndSaveLowerThresholdTime(double InLowerThresholdTime) { SET_AND_SAVE(LowerThresholdTime, InLowerThresholdTime); }
+
+	bool IsShowUpperThresholdAsFpsEnabled() const { return bShowUpperThresholdAsFps; }
+	void SetShowUpperThresholdAsFpsEnabled(bool bOnOff) { bShowUpperThresholdAsFps = bOnOff; }
+	void SetAndSaveShowUpperThresholdAsFpsEnabled(bool bOnOff) { SET_AND_SAVE(bShowUpperThresholdAsFps, bOnOff); }
+
+	bool IsShowLowerThresholdAsFpsEnabled() const { return bShowLowerThresholdAsFps; }
+	void SetShowLowerThresholdAsFpsEnabled(bool bOnOff) { bShowLowerThresholdAsFps = bOnOff; }
+	void SetAndSaveShowLowerThresholdAsFpsEnabled(bool bOnOff) { SET_AND_SAVE(bShowLowerThresholdAsFps, bOnOff); }
+
+	void SetAndSaveThresholds(double InUpperThresholdTime, double InLowerThresholdTime, bool bInShowUpperThresholdAsFps, bool bInShowLowerThresholdAsFps)
+	{
+		bool bChanged = false;
+		if (UpperThresholdTime != InUpperThresholdTime)
+		{
+			UpperThresholdTime = InUpperThresholdTime;
+			bChanged = true;
+		}
+		if (LowerThresholdTime != InLowerThresholdTime)
+		{
+			LowerThresholdTime = InLowerThresholdTime;
+			bChanged = true;
+		}
+		if (bShowUpperThresholdAsFps != bInShowUpperThresholdAsFps)
+		{
+			bShowUpperThresholdAsFps = bInShowUpperThresholdAsFps;
+			bChanged = true;
+		}
+		if (bShowLowerThresholdAsFps != bInShowLowerThresholdAsFps)
+		{
+			bShowLowerThresholdAsFps = bInShowLowerThresholdAsFps;
+			bChanged = true;
+		}
+		if (bChanged)
+		{
+			SaveToConfig();
+		}
+	}
+
 	bool IsAutoZoomOnFrameSelectionEnabled() const { return bAutoZoomOnFrameSelection; }
 	void SetAutoZoomOnFrameSelection(bool bOnOff) { bAutoZoomOnFrameSelection = bOnOff; }
 	void SetAndSaveAutoZoomOnFrameSelection(bool bOnOff) { SET_AND_SAVE(bAutoZoomOnFrameSelection, bOnOff); }
@@ -198,6 +254,29 @@ private:
 	/** Minimum time between two auto-scroll updates, in [seconds]. */
 	double AutoScrollMinDelay = 0.3; // [seconds]
 
+	//////////////////////////////////////////////////
+	// [Insights.TimingProfiler.FramesView]
+
+	/** If enabled, the upper threshold line is visible.The frame coloring by threshold is enabled regardless of this setting. */
+	bool bShowUpperThresholdLine = false;
+
+	/** If enabled, the lower threshold line is visible.The frame coloring by threshold is enabled regardless of this setting. */
+	bool bShowLowerThresholdLine = false;
+
+	/**
+	 * The upper threshold for frames.
+	 * Can be specified as a frame duration([0.001 .. 1.0] seconds; ex.: "0.010" for 10 ms) or as a framerate([1 fps .. 1000 fps]; ex: "100 fps").
+	 */
+	double UpperThresholdTime = 1.0 / 30.0;
+	/**
+	 * The lower threshold for frames.
+	 * Can be specified as a frame duration([0.001 .. 1.0] seconds; ex.: "0.010" for 10 ms) or as a framerate([1 fps .. 1000 fps]; ex: "100 fps").
+	 */
+	double LowerThresholdTime = 1.0 / 60.0;
+
+	bool bShowUpperThresholdAsFps = true;
+	bool bShowLowerThresholdAsFps = true;
+
 	/** If enabled, the Timing View will also be zoomed when a new frame is selected in the Frames track. */
 	bool bAutoZoomOnFrameSelection = false;
 
diff --git a/Engine/Source/Developer/TraceInsights/Private/Insights/Widgets/SFrameTrack.cpp b/Engine/Source/Developer/TraceInsights/Private/Insights/Widgets/SFrameTrack.cpp
--- a/Engine/Source/Developer/TraceInsights/Private/Insights/Widgets/SFrameTrack.cpp
+++ b/Engine/Source/Developer/TraceInsights/Private/Insights/Widgets/SFrameTrack.cpp
@@ -118,6 +118,15 @@ void SFrameTrack::Reset()
 
 	AnalysisSyncNextTimestamp = 0;
 
+	bShowUpperThresholdLine = Settings.IsShowUpperThresholdLineEnabled();
+	bShowLowerThresholdLine = Settings.IsShowLowerThresholdLineEnabled();
+
+	UpperThresholdTime = FMath::Clamp(Settings.GetUpperThresholdTime(), MinThresholdTime, MaxThresholdTime);
+	LowerThresholdTime = FMath::Clamp(Settings.GetLowerThresholdTime(), MinThresholdTime, MaxThresholdTime);
+
+	bShowUpperThresholdAsFps = Settings.IsShowUpperThresholdAsFpsEnabled();
+	bShowLowerThresholdAsFps = Settings.IsShowLowerThresholdAsFpsEnabled();
+
 	MousePosition = FVector2D::ZeroVector;
 
 	MousePositionOnButtonDown = FVector2D::ZeroVector;
@@ -1501,6 +1510,10 @@ void SFrameTrack::CreateThresholdsMenu(FMenuBuilder& MenuBuilder)
 				[this]()
 				{
 					bShowUpperThresholdLine = !bShowUpperThresholdLine;
+
+					// Persistent option. Save it to the config file.
+					FInsightsSettings& Settings = FInsightsManager::Get()->GetSettings();
+					Settings.SetAndSaveShowUpperThresholdLineEnabled(bShowUpperThresholdLine);
 				}),
 			FCanExecuteAction(),
 			FIsActionChecked::CreateLambda(
@@ -1522,6 +1535,10 @@ void SFrameTrack::CreateThresholdsMenu(FMenuBuilder& MenuBuilder)
 				[this]()
 				{
 					bShowLowerThresholdLine = !bShowLowerThresholdLine;
+
+					// Persistent option. Save it to the config file.
+					FInsightsSettings& Settings = FInsightsManager::Get()->GetSettings();
+					Settings.SetAndSaveShowLowerThresholdLineEnabled(bShowLowerThresholdLine);
 				}),
 			FCanExecuteAction(),
 			FIsActionChecked::CreateLambda(
@@ -1625,6 +1642,7 @@ TSharedRef<SWidget> SFrameTrack::CreateUpperThresholdWidget()
 			{
 				LowerThresholdTime = UpperThresholdTime;
 			}
+			SaveThresholds();
 		})
 	]
 
@@ -1707,6 +1725,7 @@ TSharedRef<SWidget> SFrameTrack::CreateLowerThresholdWidget()
 			{
 				UpperThresholdTime = LowerThresholdTime;
 			}
+			SaveThresholds();
 		})
 	]
 
@@ -1733,12 +1752,22 @@ TSharedRef<SWidget> SFrameTrack::CreateLowerThresholdWidget()
 
 ////////////////////////////////////////////////////////////////////////////////////////////////////
 
+void SFrameTrack::SaveThresholds()
+{
+	// Persistent option. Save it to the config file.
+	FInsightsSettings& Settings = FInsightsManager::Get()->GetSettings();
+	Settings.SetAndSaveThresholds(UpperThresholdTime, LowerThresholdTime, bShowUpperThresholdAsFps, bShowLowerThresholdAsFps);
+}
+
+////////////////////////////////////////////////////////////////////////////////////////////////////
+
 void SFrameTrack::SetThresholdsFPS(double InUpperThresholdFPS, double InLowerThresholdFPS)
 {
 	UpperThresholdTime = 1.0 / InUpperThresholdFPS;
 	LowerThresholdTime = 1.0 / InLowerThresholdFPS;
 	bShowUpperThresholdAsFps = true;
 	bShowLowerThresholdAsFps = true;
+	SaveThresholds();
 }
 
 ////////////////////////////////////////////////////////////////////////////////////////////////////
diff --git a/Engine/Source/Developer/TraceInsights/Private/Insights/Widgets/SFrameTrack.h b/Engine/Source/Developer/TraceInsights/Private/Insights/Widgets/SFrameTrack.h
--- a/Engine/Source/Developer/TraceInsights/Private/Insights/Widgets/SFrameTrack.h
+++ b/Engine/Source/Developer/TraceInsights/Private/Insights/Widgets/SFrameTrack.h
@@ -152,6 +152,7 @@ protected:
 	void CreateThresholdsMenu(FMenuBuilder& MenuBuilder);
 	TSharedRef<SWidget> CreateUpperThresholdWidget();
 	TSharedRef<SWidget> CreateLowerThresholdWidget();
+	void SaveThresholds();
 	void SetThresholdsFPS(double InUpperThresholdFPS, double InLowerThresholdFPS);
 	TSharedRef<SWidget> CreateThresholdPresetsWidget();
 	void CreateSelectedFrameMenu(FMenuBuilder& MenuBuilder);
@@ -215,6 +216,9 @@ protected:
 	//////////////////////////////////////////////////
 	// Thresholds
 
+	bool bShowUpperThresholdLine = false;
+	bool bShowLowerThresholdLine = false;
+
 	static constexpr double MinThresholdTime = 0.001; // == 1ms == 1000 fps
 	static constexpr double MaxThresholdTime = 1.0; // == 1s == 1 fps
 
@@ -224,9 +228,6 @@ protected:
 	bool bShowUpperThresholdAsFps = true;
 	bool bShowLowerThresholdAsFps = true;
 
-	bool bShowUpperThresholdLine = false;
-	bool bShowLowerThresholdLine = false;
-
 	//////////////////////////////////////////////////
 
 	TSharedPtr<SScrollBar> HorizontalScrollBar;
