Allow Unreal Insights to connect to non-standard control ports

When running multiple instances on the same machine, TraceLog listening port will be different for each process. Allow the Unreal Insights connect UI to use non-standard ports by specifying "ip:port" format.

diff --git a/Engine/Source/Developer/TraceInsights/Private/Insights/Widgets/SStartPageWindow.cpp b/Engine/Source/Developer/TraceInsights/Private/Insights/Widgets/SStartPageWindow.cpp
--- a/Engine/Source/Developer/TraceInsights/Private/Insights/Widgets/SStartPageWindow.cpp
+++ b/Engine/Source/Developer/TraceInsights/Private/Insights/Widgets/SStartPageWindow.cpp
@@ -4770,7 +4770,18 @@ FReply SConnectionWindow::Connect_OnClicked()
 			UE_LOG(TraceInsights, Log, TEXT("[Connection] Try connecting to \"%s\"..."), *RunningInstanceAddressStr);
 
 			UE::Trace::FControlClient ControlClient;
-			if (ControlClient.Connect(*RunningInstanceAddressStr))
+			FString IPStr, PortStr;
+			RunningInstanceAddressStr.Split(TEXT(":"), &IPStr, &PortStr);
+			uint16 Port = 1985;
+			if (!PortStr.IsEmpty())
+			{
+				Port = uint16(FCString::Atoi(*PortStr));
+			}
+			else
+			{
+				IPStr = RunningInstanceAddressStr;
+			}
+			if (ControlClient.Connect(*IPStr, Port))
 			{
 				UE_LOG(TraceInsights, Log, TEXT("[Connection] SendSendTo(\"%s\")..."), *TraceRecorderAddressStr);
 				ControlClient.SendSendTo(*TraceRecorderAddressStr);
