[Insights] Created IInsightsManager to expose Insights Manager delegates + Moved FAvailabilityCheck class to its own file so it is publicly accesible + Misc. Insights namespace code adjustments.

diff --git a/Engine/Source/Developer/TraceInsights/Private/Insights/Common/AvailabilityCheck.cpp b/Engine/Source/Developer/TraceInsights/Private/Insights/Common/AvailabilityCheck.cpp
new file mode 100644
--- /dev/null
+++ b/Engine/Source/Developer/TraceInsights/Private/Insights/Common/AvailabilityCheck.cpp
@@ -0,0 +1,52 @@
+// Copyright Epic Games, Inc. All Rights Reserved.
+
+#include "Insights/Common/AvailabilityCheck.h"
+
+#include "HAL/PlatformTime.h"
+
+namespace Insights
+{
+
+////////////////////////////////////////////////////////////////////////////////////////////////////
+// FAvailabilityCheck
+////////////////////////////////////////////////////////////////////////////////////////////////////
+
+bool FAvailabilityCheck::Tick()
+{
+	if (NextTimestamp != (uint64)-1)
+	{
+		const uint64 Time = FPlatformTime::Cycles64();
+		if (Time > NextTimestamp)
+		{
+			// Increase wait time with 0.1s, but at no more than 3s.
+			WaitTime = FMath::Min(WaitTime + 0.1, 3.0);
+			const uint64 WaitTimeCycles64 = static_cast<uint64>(WaitTime / FPlatformTime::GetSecondsPerCycle64());
+			NextTimestamp = Time + WaitTimeCycles64;
+
+			return true; // yes, manager can check for (slow) availability conditions
+		}
+	}
+
+	return false; // no, manager should not check for (slow) availability conditions during this tick
+}
+
+////////////////////////////////////////////////////////////////////////////////////////////////////
+
+void FAvailabilityCheck::Disable()
+{
+	WaitTime = 0.0;
+	NextTimestamp = (uint64)-1;
+}
+
+////////////////////////////////////////////////////////////////////////////////////////////////////
+
+void FAvailabilityCheck::Enable(double InWaitTime)
+{
+	WaitTime = InWaitTime;
+	const uint64 WaitTimeCycles64 = static_cast<uint64>(WaitTime / FPlatformTime::GetSecondsPerCycle64());
+	NextTimestamp = FPlatformTime::Cycles64() + WaitTimeCycles64;
+}
+
+////////////////////////////////////////////////////////////////////////////////////////////////////
+
+} // namespace Insights
diff --git a/Engine/Source/Developer/TraceInsights/Private/Insights/InsightsManager.cpp b/Engine/Source/Developer/TraceInsights/Private/Insights/InsightsManager.cpp
--- a/Engine/Source/Developer/TraceInsights/Private/Insights/InsightsManager.cpp
+++ b/Engine/Source/Developer/TraceInsights/Private/Insights/InsightsManager.cpp
@@ -66,46 +66,6 @@ const FName FInsightsManagerTabs::MemoryProfilerTabId(TEXT("MemoryProfiler"));
 const FName FInsightsManagerTabs::AutomationWindowTabId(TEXT("AutomationWindow"));
 const FName FInsightsManagerTabs::MessageLogTabId(TEXT("MessageLog"));
 
-////////////////////////////////////////////////////////////////////////////////////////////////////
-// FAvailabilityCheck
-////////////////////////////////////////////////////////////////////////////////////////////////////
-
-bool FAvailabilityCheck::Tick()
-{
-	if (NextTimestamp != (uint64)-1)
-	{
-		const uint64 Time = FPlatformTime::Cycles64();
-		if (Time > NextTimestamp)
-		{
-			// Increase wait time with 0.1s, but at no more than 3s.
-			WaitTime = FMath::Min(WaitTime + 0.1, 3.0);
-			const uint64 WaitTimeCycles64 = static_cast<uint64>(WaitTime / FPlatformTime::GetSecondsPerCycle64());
-			NextTimestamp = Time + WaitTimeCycles64;
-
-			return true; // yes, manager can check for (slow) availability conditions
-		}
-	}
-
-	return false; // no, manager should not check for (slow) availability conditions during this tick
-}
-
-////////////////////////////////////////////////////////////////////////////////////////////////////
-
-void FAvailabilityCheck::Disable()
-{
-	WaitTime = 0.0;
-	NextTimestamp = (uint64)-1;
-}
-
-////////////////////////////////////////////////////////////////////////////////////////////////////
-
-void FAvailabilityCheck::Enable(double InWaitTime)
-{
-	WaitTime = InWaitTime;
-	const uint64 WaitTimeCycles64 = static_cast<uint64>(WaitTime / FPlatformTime::GetSecondsPerCycle64());
-	NextTimestamp = FPlatformTime::Cycles64() + WaitTimeCycles64;
-}
-
 ////////////////////////////////////////////////////////////////////////////////////////////////////
 // FInsightsManager
 ////////////////////////////////////////////////////////////////////////////////////////////////////
diff --git a/Engine/Source/Developer/TraceInsights/Private/Insights/InsightsManager.h b/Engine/Source/Developer/TraceInsights/Private/Insights/InsightsManager.h
--- a/Engine/Source/Developer/TraceInsights/Private/Insights/InsightsManager.h
+++ b/Engine/Source/Developer/TraceInsights/Private/Insights/InsightsManager.h
@@ -8,7 +8,9 @@
 #include "Input/DragAndDrop.h"
 
 // Insights
+#include "Insights/Common/AvailabilityCheck.h"
 #include "Insights/Common/Stopwatch.h"
+#include "Insights/IInsightsManager.h"
 #include "Insights/InsightsCommands.h"
 #include "Insights/InsightsSettings.h"
 #include "Insights/InsightsSessionBrowserSettings.h"
@@ -37,26 +39,6 @@ class SSessionInfoWindow;
 class FInsightsTestRunner;
 class FInsightsMenuBuilder;
 
-/**
- * Utility class used by profiler managers to limit how often they check for availability conditions.
- */
-class FAvailabilityCheck
-{
-public:
-	/** Returns true if managers are allowed to do (slow) availability check during this tick. */
-	bool Tick();
-
-	/** Disables the "availability check" (i.e. Tick() calls will return false when disabled). */
-	void Disable();
-
-	/** Enables the "availability check" with a specified initial delay. */
-	void Enable(double InWaitTime);
-
-private:
-	double WaitTime = 0.0;
-	uint64 NextTimestamp = (uint64)-1;
-};
-
 /**
  * Struct that holds data about in progress async operations
  */
@@ -78,7 +60,7 @@ struct FAsyncTaskData
  *     Connecting/disconnecting to source trace
  *     Global Unreal Insights application state and settings
  */
-class FInsightsManager : public TSharedFromThis<FInsightsManager>, public IInsightsComponent
+class FInsightsManager : public TSharedFromThis<FInsightsManager>, public Insights::IInsightsManager
 {
 	friend class FInsightsActionManager;
 
@@ -340,22 +322,20 @@ public:
 
 public:
 	/** The event to execute when the session has changed. */
-	DECLARE_EVENT(FTimingProfilerManager, FSessionChangedEvent);
-	FSessionChangedEvent& GetSessionChangedEvent() { return SessionChangedEvent; }
+	virtual Insights::IInsightsManager::FSessionChangedEvent& GetSessionChangedEvent() override { return SessionChangedEvent; }
 private:
 	/** The event to execute when the session has changed. */
-	FSessionChangedEvent SessionChangedEvent;
+	Insights::IInsightsManager::FSessionChangedEvent SessionChangedEvent;
 
 	////////////////////////////////////////////////////////////////////////////////////////////////////
 	// SessionAnalysisCompletedEvent
 
 public:
 	/** The event to execute when session analysis is complete. */
-	DECLARE_EVENT(FTimingProfilerManager, FSessionAnalysisCompletedEvent);
-	FSessionAnalysisCompletedEvent& GetSessionAnalysisCompletedEvent() { return SessionAnalysisCompletedEvent; }
+	virtual Insights::IInsightsManager::FSessionAnalysisCompletedEvent& GetSessionAnalysisCompletedEvent() override { return SessionAnalysisCompletedEvent; }
 private:
 	/** The event to execute when session analysis is completed. */
-	FSessionAnalysisCompletedEvent SessionAnalysisCompletedEvent;
+	Insights::IInsightsManager::FSessionAnalysisCompletedEvent SessionAnalysisCompletedEvent;
 
 	////////////////////////////////////////////////////////////////////////////////////////////////////
 
diff --git a/Engine/Source/Developer/TraceInsights/Private/Insights/LoadingProfiler/LoadingProfilerManager.h b/Engine/Source/Developer/TraceInsights/Private/Insights/LoadingProfiler/LoadingProfilerManager.h
--- a/Engine/Source/Developer/TraceInsights/Private/Insights/LoadingProfiler/LoadingProfilerManager.h
+++ b/Engine/Source/Developer/TraceInsights/Private/Insights/LoadingProfiler/LoadingProfilerManager.h
@@ -139,7 +139,7 @@ private:
 private:
 	bool bIsInitialized;
 	bool bIsAvailable;
-	FAvailabilityCheck AvailabilityCheck;
+	Insights::FAvailabilityCheck AvailabilityCheck;
 
 	/** The delegate to be invoked when this manager ticks. */
 	FTickerDelegate OnTick;
diff --git a/Engine/Source/Developer/TraceInsights/Private/Insights/MemoryProfiler/MemoryProfilerManager.h b/Engine/Source/Developer/TraceInsights/Private/Insights/MemoryProfiler/MemoryProfilerManager.h
--- a/Engine/Source/Developer/TraceInsights/Private/Insights/MemoryProfiler/MemoryProfilerManager.h
+++ b/Engine/Source/Developer/TraceInsights/Private/Insights/MemoryProfiler/MemoryProfilerManager.h
@@ -128,7 +128,7 @@ private:
 private:
 	bool bIsInitialized;
 	bool bIsAvailable;
-	FAvailabilityCheck AvailabilityCheck;
+	Insights::FAvailabilityCheck AvailabilityCheck;
 
 	/** The delegate to be invoked when this manager ticks. */
 	FTickerDelegate OnTick;
diff --git a/Engine/Source/Developer/TraceInsights/Private/Insights/NetworkingProfiler/NetworkingProfilerManager.h b/Engine/Source/Developer/TraceInsights/Private/Insights/NetworkingProfiler/NetworkingProfilerManager.h
--- a/Engine/Source/Developer/TraceInsights/Private/Insights/NetworkingProfiler/NetworkingProfilerManager.h
+++ b/Engine/Source/Developer/TraceInsights/Private/Insights/NetworkingProfiler/NetworkingProfilerManager.h
@@ -113,7 +113,7 @@ private:
 private:
 	bool bIsInitialized;
 	bool bIsAvailable;
-	FAvailabilityCheck AvailabilityCheck;
+	Insights::FAvailabilityCheck AvailabilityCheck;
 
 	/** The delegate to be invoked when this manager ticks. */
 	FTickerDelegate OnTick;
diff --git a/Engine/Source/Developer/TraceInsights/Private/Insights/TimingProfilerManager.h b/Engine/Source/Developer/TraceInsights/Private/Insights/TimingProfilerManager.h
--- a/Engine/Source/Developer/TraceInsights/Private/Insights/TimingProfilerManager.h
+++ b/Engine/Source/Developer/TraceInsights/Private/Insights/TimingProfilerManager.h
@@ -187,7 +187,7 @@ private:
 private:
 	bool bIsInitialized;
 	bool bIsAvailable;
-	FAvailabilityCheck AvailabilityCheck;
+	Insights::FAvailabilityCheck AvailabilityCheck;
 
 	/** The delegate to be invoked when this manager ticks. */
 	FTickerDelegate OnTick;
diff --git a/Engine/Source/Developer/TraceInsights/Private/Insights/TraceInsightsModule.cpp b/Engine/Source/Developer/TraceInsights/Private/Insights/TraceInsightsModule.cpp
--- a/Engine/Source/Developer/TraceInsights/Private/Insights/TraceInsightsModule.cpp
+++ b/Engine/Source/Developer/TraceInsights/Private/Insights/TraceInsightsModule.cpp
@@ -175,6 +175,13 @@ void FTraceInsightsModule::UnregisterComponent(TSharedPtr<IInsightsComponent> Co
 
 ////////////////////////////////////////////////////////////////////////////////////////////////////
 
+TSharedPtr<Insights::IInsightsManager> FTraceInsightsModule::GetInsightsManager()
+{
+	return FInsightsManager::Get();
+}
+
+////////////////////////////////////////////////////////////////////////////////////////////////////
+
 void FTraceInsightsModule::CreateDefaultStore()
 {
 	ConnectToStore(TEXT("127.0.0.1"));
diff --git a/Engine/Source/Developer/TraceInsights/Private/Insights/TraceInsightsModule.h b/Engine/Source/Developer/TraceInsights/Private/Insights/TraceInsightsModule.h
--- a/Engine/Source/Developer/TraceInsights/Private/Insights/TraceInsightsModule.h
+++ b/Engine/Source/Developer/TraceInsights/Private/Insights/TraceInsightsModule.h
@@ -7,6 +7,11 @@
 #include "Framework/Docking/TabManager.h"
 #include "Trace/StoreService.h"
 
+namespace Insights
+{
+	class IInsightsManager;
+}
+
 namespace TraceServices
 {
 	class IAnalysisService;
@@ -51,6 +56,8 @@ public:
 	virtual void RegisterComponent(TSharedPtr<IInsightsComponent> Component) override;
 	virtual void UnregisterComponent(TSharedPtr<IInsightsComponent> Component) override;
 
+	virtual TSharedPtr<Insights::IInsightsManager> GetInsightsManager() override;
+
 	virtual void RegisterMajorTabConfig(const FName& InMajorTabId, const FInsightsMajorTabConfig& InConfig) override;
 	virtual void UnregisterMajorTabConfig(const FName& InMajorTabId) override;
 	virtual FOnInsightsMajorTabCreated& OnMajorTabCreated() override { return OnInsightsMajorTabCreatedDelegate; }
diff --git a/Engine/Source/Developer/TraceInsights/Public/Insights/Common/AvailabilityCheck.h b/Engine/Source/Developer/TraceInsights/Public/Insights/Common/AvailabilityCheck.h
new file mode 100644
--- /dev/null
+++ b/Engine/Source/Developer/TraceInsights/Public/Insights/Common/AvailabilityCheck.h
@@ -0,0 +1,34 @@
+// Copyright Epic Games, Inc. All Rights Reserved.
+
+#pragma once
+
+#include "HAL/Platform.h"
+
+namespace Insights
+{
+
+////////////////////////////////////////////////////////////////////////////////////////////////////
+
+/**
+ * Utility class used by profiler managers to limit how often they check for availability conditions.
+ */
+class TRACEINSIGHTS_API FAvailabilityCheck
+{
+public:
+	/** Returns true if managers are allowed to do (slow) availability check during this tick. */
+	bool Tick();
+
+	/** Disables the "availability check" (i.e. Tick() calls will return false when disabled). */
+	void Disable();
+
+	/** Enables the "availability check" with a specified initial delay. */
+	void Enable(double InWaitTime);
+
+private:
+	double WaitTime = 0.0;
+	uint64 NextTimestamp = (uint64)-1;
+};
+
+////////////////////////////////////////////////////////////////////////////////////////////////////
+
+} // namespace Insights
diff --git a/Engine/Source/Developer/TraceInsights/Public/Insights/IInsightsManager.h b/Engine/Source/Developer/TraceInsights/Public/Insights/IInsightsManager.h
new file mode 100644
--- /dev/null
+++ b/Engine/Source/Developer/TraceInsights/Public/Insights/IInsightsManager.h
@@ -0,0 +1,27 @@
+// Copyright Epic Games, Inc. All Rights Reserved.
+
+#pragma once
+
+#include "Insights/IUnrealInsightsModule.h"
+
+namespace Insights
+{
+
+////////////////////////////////////////////////////////////////////////////////////////////////////
+
+/** Interface for an Insights manager. */
+class TRACEINSIGHTS_API IInsightsManager : public IInsightsComponent
+{
+public:
+	/** The event to execute when the session has changed. */
+	DECLARE_MULTICAST_DELEGATE(FSessionChangedEvent);
+	virtual FSessionChangedEvent& GetSessionChangedEvent() = 0;
+
+	/** The event to execute when session analysis is complete. */
+	DECLARE_MULTICAST_DELEGATE(FSessionAnalysisCompletedEvent);
+	virtual FSessionAnalysisCompletedEvent& GetSessionAnalysisCompletedEvent() = 0;
+};
+
+////////////////////////////////////////////////////////////////////////////////////////////////////
+
+} // namespace Insights
diff --git a/Engine/Source/Developer/TraceInsights/Public/Insights/IUnrealInsightsModule.h b/Engine/Source/Developer/TraceInsights/Public/Insights/IUnrealInsightsModule.h
--- a/Engine/Source/Developer/TraceInsights/Public/Insights/IUnrealInsightsModule.h
+++ b/Engine/Source/Developer/TraceInsights/Public/Insights/IUnrealInsightsModule.h
@@ -19,6 +19,11 @@ namespace Trace
 }
 }
 
+namespace Insights
+{
+	class IInsightsManager;
+}
+
 namespace TraceServices
 {
 	class IAnalysisSession;
@@ -201,6 +206,13 @@ public:
 
 	//////////////////////////////////////////////////
 
+	/**
+	 * Gets the Insights manager.
+	*/
+	virtual TSharedPtr<Insights::IInsightsManager> GetInsightsManager() = 0;
+
+	//////////////////////////////////////////////////
+
 	/**
 	 * Gets the current analysis session.
 	 */
