Added branch name to Unreal Insights window title, to more easily identify branches when working on multiple ones simultaneously

diff --git a/Engine/Source/Developer/TraceInsights/Private/Insights/InsightsManager.cpp b/Engine/Source/Developer/TraceInsights/Private/Insights/InsightsManager.cpp
--- a/Engine/Source/Developer/TraceInsights/Private/Insights/InsightsManager.cpp
+++ b/Engine/Source/Developer/TraceInsights/Private/Insights/InsightsManager.cpp
@@ -1101,8 +1101,9 @@ void FInsightsManager::UpdateAppTitle()
 					TraceServices::FSessionInfo SessionInfo = DiagnosticsProvider->GetSessionInfo();
 
 					const FString SessionName = FPaths::GetBaseFilename(CurrentTraceFilename);
-					const FText AppTitle = FText::Format(LOCTEXT("UnrealInsightsAppNameFmt2", "{0} - {1} - {2} - {3} - {4} - Unreal Insights"),
+					const FText AppTitle = FText::Format(LOCTEXT("UnrealInsightsAppNameFmt2", "{0}{1} - {2} - {3} - {4} - {5} Unreal Insights"),
 						FText::FromString(SessionName),
+						!SessionInfo.Branch.IsEmpty() ? FText::FromString(TEXT(" - ") + SessionInfo.Branch) : FText::GetEmpty(),
 						FText::FromString(SessionInfo.Platform),
 						FText::FromString(SessionInfo.AppName),
 						FText::FromString(LexToString(SessionInfo.ConfigurationType)),
