[Insights]
  - Fixed grouping for lazy created child nodes in STabelTreeView. The parent node (which lazy creates child nodes) requires to implement GetAuthorGrouping().
  - Added names to Packages and Tasks tables.
  - Minor code refactoring.

#rnx

diff --git a/Engine/Source/Developer/TraceInsights/Private/Insights/CookProfiler/CookProfilerManager.cpp b/Engine/Source/Developer/TraceInsights/Private/Insights/CookProfiler/CookProfilerManager.cpp
--- a/Engine/Source/Developer/TraceInsights/Private/Insights/CookProfiler/CookProfilerManager.cpp
+++ b/Engine/Source/Developer/TraceInsights/Private/Insights/CookProfiler/CookProfilerManager.cpp
@@ -187,7 +187,7 @@ void FCookProfilerManager::RegisterTimingProfilerLayoutExtensions(FInsightsMajor
 	FMinorTabConfig& MinorTabConfig = InOutExtender.AddMinorTabConfig();
 	MinorTabConfig.TabId = FCookProfilerTabs::PackageTableTreeViewTabID;
 	MinorTabConfig.TabLabel = LOCTEXT("PackageTableTreeViewTabTitle", "Packages");
-	MinorTabConfig.TabTooltip = LOCTEXT("PackageTableTreeViewTabTitleTooltip", "Opens the Packages Tree View tab, that allows cook profilling.");
+	MinorTabConfig.TabTooltip = LOCTEXT("PackageTableTreeViewTabTitleTooltip", "Opens the Packages Tree View tab, that allows cook profiling.");
 	MinorTabConfig.TabIcon = FSlateIcon(FInsightsStyle::GetStyleSetName(), "Icons.PackagesView");
 	MinorTabConfig.OnSpawnTab = FOnSpawnTab::CreateRaw(this, &FCookProfilerManager::SpawnTab_PackageTableTreeView);
 	MinorTabConfig.CanSpawnTab = FCanSpawnTab::CreateRaw(this, &FCookProfilerManager::CanSpawnTab_PackageTableTreeView);
@@ -204,6 +204,7 @@ TSharedRef<SDockTab> FCookProfilerManager::SpawnTab_PackageTableTreeView(const F
 {
 	TSharedRef<FPackageTable> PackageTable = MakeShared<FPackageTable>();
 	PackageTable->Reset();
+	PackageTable->SetDisplayName(LOCTEXT("PackageTableTreeViewTabTitle", "Packages"));
 
 	const TSharedRef<SDockTab> DockTab = SNew(SDockTab)
 		.ShouldAutosize(false)
diff --git a/Engine/Source/Developer/TraceInsights/Private/Insights/Table/ViewModels/TableTreeNode.h b/Engine/Source/Developer/TraceInsights/Private/Insights/Table/ViewModels/TableTreeNode.h
--- a/Engine/Source/Developer/TraceInsights/Private/Insights/Table/ViewModels/TableTreeNode.h
+++ b/Engine/Source/Developer/TraceInsights/Private/Insights/Table/ViewModels/TableTreeNode.h
@@ -12,6 +12,7 @@ namespace Insights
 {
 
 class FTable;
+class FTreeNodeGrouping;
 class STableTreeView;
 
 struct FTableRowId
@@ -148,6 +149,13 @@ public:
 
 	virtual bool OnLazyCreateChildren(TSharedPtr<STableTreeView> InTableTreeView) { return false; }
 
+	/**
+	 * The grouping that has generated this node.
+	 * This is used to correctly apply further groupings for the lazy created children.
+	 * If this returns nullptr, grouping is not applied for lazy created children.
+	 */
+	virtual const FTreeNodeGrouping* GetAuthorGrouping() const { return nullptr; }
+
 protected:
 	TWeakPtr<FTable> ParentTable;
 	FTableRowId RowId;
diff --git a/Engine/Source/Developer/TraceInsights/Private/Insights/Table/Widgets/STableTreeView.cpp b/Engine/Source/Developer/TraceInsights/Private/Insights/Table/Widgets/STableTreeView.cpp
--- a/Engine/Source/Developer/TraceInsights/Private/Insights/Table/Widgets/STableTreeView.cpp
+++ b/Engine/Source/Developer/TraceInsights/Private/Insights/Table/Widgets/STableTreeView.cpp
@@ -1148,7 +1148,9 @@ void STableTreeView::UpdateFilterContext(const FFilterConfigurator& InFilterConf
 {
 	for (const TSharedRef<FTableColumn>& Column : Table->GetColumns())
 	{
-		if (!Column->CanBeFiltered() || !InFilterConfigurator.IsKeyUsed(Column->GetIndex()) || !Column->GetValue(InNode).IsSet())
+		if (!Column->CanBeFiltered() ||
+			!InFilterConfigurator.IsKeyUsed(Column->GetIndex()) ||
+			!Column->GetValue(InNode).IsSet())
 		{
 			continue;
 		}
@@ -1330,14 +1332,12 @@ void STableTreeView::ApplyHierarchyFiltering()
 
 	// The Root node is always hidden. The tree shows the filtered children of the Root node.
 	const TArray<FBaseTreeNodePtr>& RootChildren = Root->GetFilteredChildren();
-	const int32 NumRootChildren = RootChildren.Num();
-	FilteredGroupNodes.Reset(NumRootChildren);
-	for (int32 Cx = 0; Cx < NumRootChildren; ++Cx)
+	FilteredGroupNodes.Reset(RootChildren.Num());
+	for (const FBaseTreeNodePtr& ChildNodePtr : RootChildren)
 	{
-		const FTableTreeNodePtr& ChildNodePtr = StaticCastSharedPtr<FTableTreeNode>(RootChildren[Cx]);
 		if (ChildNodePtr->IsGroup())
 		{
-			FilteredGroupNodes.Add(ChildNodePtr);
+			FilteredGroupNodes.Add(StaticCastSharedPtr<FTableTreeNode>(ChildNodePtr));
 		}
 	}
 
@@ -1355,9 +1355,8 @@ void STableTreeView::ApplyHierarchyFiltering()
 				bExpansionSaved = true;
 			}
 
-			for (int32 Fx = 0; Fx < FilteredGroupNodes.Num(); Fx++)
+			for (const FTableTreeNodePtr& GroupPtr : FilteredGroupNodes)
 			{
-				const FTableTreeNodePtr& GroupPtr = FilteredGroupNodes[Fx];
 				TreeView->SetItemExpansion(GroupPtr, GroupPtr->IsExpanded());
 			}
 		}
@@ -1420,12 +1419,11 @@ bool STableTreeView::ApplyHierarchyFilteringRec(FTableTreeNodePtr NodePtr)
 		}
 
 		const TArray<FBaseTreeNodePtr>& GroupChildren = NodePtr->GetChildren();
-		const int32 NumChildren = GroupChildren.Num();
 
 		NodePtr->ClearFilteredChildren(0);
 
 		int32 NumVisibleChildren = 0;
-		for (int32 Cx = 0; Cx < NumChildren; ++Cx)
+		for (const FBaseTreeNodePtr& GroupChildNodePtr : GroupChildren)
 		{
 			if (AsyncOperationProgress.ShouldCancelAsyncOp())
 			{
@@ -1433,7 +1431,7 @@ bool STableTreeView::ApplyHierarchyFilteringRec(FTableTreeNodePtr NodePtr)
 			}
 
 			// Add a child.
-			const FTableTreeNodePtr& ChildNodePtr = StaticCastSharedPtr<FTableTreeNode>(GroupChildren[Cx]);
+			const FTableTreeNodePtr& ChildNodePtr = StaticCastSharedPtr<FTableTreeNode>(GroupChildNodePtr);
 			if (ApplyHierarchyFilteringRec(ChildNodePtr))
 			{
 				NodePtr->AddFilteredChild(ChildNodePtr);
@@ -1470,20 +1468,18 @@ bool STableTreeView::MakeSubtreeVisible(FTableTreeNodePtr NodePtr, bool bFilterI
 	if (NodePtr->IsGroup())
 	{
 		const TArray<FBaseTreeNodePtr>& GroupChildren = NodePtr->GetChildren();
-		const int32 NumChildren = GroupChildren.Num();
 
-		NodePtr->ClearFilteredChildren(bFilterIsEmpty ? NumChildren : 0);
+		NodePtr->ClearFilteredChildren(bFilterIsEmpty ? GroupChildren.Num() : 0);
 
-		int32 NumVisibleChildren = 0;
 		bool bShouldExpand = false;
-		for (int32 Cx = 0; Cx < NumChildren; ++Cx)
+		for (const FBaseTreeNodePtr& GroupChildNodePtr : GroupChildren)
 		{
 			if (AsyncOperationProgress.ShouldCancelAsyncOp())
 			{
 				break;
 			}
 
-			const FTableTreeNodePtr& ChildNodePtr = StaticCastSharedPtr<FTableTreeNode>(GroupChildren[Cx]);
+			const FTableTreeNodePtr& ChildNodePtr = StaticCastSharedPtr<FTableTreeNode>(GroupChildNodePtr);
 			bShouldExpand |= MakeSubtreeVisible(ChildNodePtr, bFilterIsEmpty);
 			NodePtr->AddFilteredChild(ChildNodePtr);
 			NodePtr->SetExpansion(true);
@@ -1564,7 +1560,39 @@ void STableTreeView::TreeView_OnGetChildren(FTableTreeNodePtr InParent, TArray<F
 	if (InParent->OnLazyCreateChildren(SharedThis(this)))
 	{
 		// Node filtering does not apply to lazy expanded nodes.
-		// Grouping is ignored for lazy expanded nodes.
+
+		// Apply grouping only if the node specifies the grouping that created the node.
+		// In this case, it will apply groupings after the identified one.
+		const FTreeNodeGrouping* AuthorGrouping = InParent->GetAuthorGrouping();
+		if (AuthorGrouping)
+		{
+			TArray<TSharedPtr<FTreeNodeGrouping>> Groupings;
+
+			bool bFound = false;
+			for (TSharedPtr<FTreeNodeGrouping> Grouping : CurrentGroupings)
+			{
+				if (bFound)
+				{
+					Groupings.Add(Grouping);
+				}
+				else
+				{
+					if (Grouping.Get() == AuthorGrouping)
+					{
+						bFound = true;
+					}
+				}
+			}
+
+			if (Groupings.Num() > 0)
+			{
+				// Extract children nodes.
+				TArray<FTableTreeNodePtr> ChildNodes;
+				InParent->SwapChildrenFast(reinterpret_cast<TArray<FBaseTreeNodePtr>&>(ChildNodes));
+
+				GroupNodesRec(ChildNodes, *InParent, 0, Groupings);
+			}
+		}
 
 		// Update aggregation.
 		UpdateAggregatedValuesRec(*InParent);
@@ -2656,6 +2684,11 @@ void STableTreeView::UpdateAggregatedValuesSingleNode(FTableTreeNode& InOutGroup
 
 void STableTreeView::UpdateAggregatedValuesRec(FTableTreeNode& InOutGroupNode)
 {
+	if (AsyncOperationProgress.ShouldCancelAsyncOp())
+	{
+		return;
+	}
+
 	STableTreeView::UpdateAggregatedValues<true>(Table, InOutGroupNode);
 }
 
@@ -2992,7 +3025,6 @@ void STableTreeView::ShowColumn(FTableColumn& Column)
 			.VAlignCell(VAlign_Fill)
 			.InitialSortMode(Column.GetInitialSortMode())
 			.SortMode(this, &STableTreeView::GetSortModeForColumn, Column.GetId())
-			.OnSort(this, &STableTreeView::OnSortModeChanged)
 			.FillWidth(Column.GetInitialWidth())
 			//.FixedWidth(Column.IsFixedWidth() ? Column.GetInitialWidth() : TOptional<float>())
 			.HeaderContent()
@@ -3011,6 +3043,11 @@ void STableTreeView::ShowColumn(FTableColumn& Column)
 				TreeViewHeaderRow_GenerateColumnMenu(Column)
 			];
 
+		if (Column.CanBeSorted())
+		{
+			ColumnArgs.OnSort(this, &STableTreeView::OnSortModeChanged);
+		}
+
 		int32 ColumnIndex = 0;
 		const int32 NewColumnPosition = Table->GetColumnPositionIndex(Column.GetId());
 		const int32 NumColumns = TreeViewHeaderRow->GetColumns().Num();
diff --git a/Engine/Source/Developer/TraceInsights/Private/Insights/Table/Widgets/STableTreeView.h b/Engine/Source/Developer/TraceInsights/Private/Insights/Table/Widgets/STableTreeView.h
--- a/Engine/Source/Developer/TraceInsights/Private/Insights/Table/Widgets/STableTreeView.h
+++ b/Engine/Source/Developer/TraceInsights/Private/Insights/Table/Widgets/STableTreeView.h
@@ -179,8 +179,8 @@ public:
 
 	/** Gets the table row nodes. Each node corresponds to a table row. Index in this array corresponds to RowIndex in source table. */
 	const TArray<FTableTreeNodePtr>& GetTableRowNodes() const { return TableRowNodes; }
-	
-	/** Gets the avaiable grouping. */
+
+	/** Gets the available groupings. */
 	const TArray<TSharedPtr<FTreeNodeGrouping>>& GetAvailableGroupings() const { return AvailableGroupings; }
 
 	/** Sets the current groupings. */
@@ -525,7 +525,7 @@ protected:
 
 	static const FName RootNodeName;
 
-	/** The root node of the tree. */
+	/** The root node of the tree. It is invisible; only its children are displayed in the tree view. */
 	FTableTreeNodePtr Root;
 
 	/** Table row nodes. Each node corresponds to a table row. Index in this array corresponds to RowIndex in source table. */
@@ -543,13 +543,13 @@ protected:
 	/** Currently expanded group nodes. */
 	TSet<FTableTreeNodePtr> ExpandedNodes;
 
-	/** If true, the expanded nodes have been saved before applying a text filter. */
+	/** If true, the expanded nodes have been saved before applying a text filter. Only used when bRunInAsyncMode==false. */
 	bool bExpansionSaved = false;
 
-	static constexpr int32 MaxNodesToAutoExpand = 1000;
-	static constexpr int32 MaxDepthToAutoExpand = 4;
-	static constexpr int32 MaxNodesToExpand = 1000000;
-	static constexpr int32 MaxDepthToExpand = 100;
+	int32 MaxNodesToAutoExpand = 1000;
+	int32 MaxDepthToAutoExpand = 4;
+	int32 MaxNodesToExpand = 1000000;
+	int32 MaxDepthToExpand = 100;
 
 	//////////////////////////////////////////////////
 	// Search box & the hierarchy filtering
@@ -650,8 +650,8 @@ class FTableTreeViewNodeFilteringAsyncTask
 {
 public:
 	FTableTreeViewNodeFilteringAsyncTask(STableTreeView* InPtr)
+		: TableTreeViewPtr(InPtr)
 	{
-		TableTreeViewPtr = InPtr;
 	}
 
 	FORCEINLINE TStatId GetStatId() const { RETURN_QUICK_DECLARE_CYCLE_STAT(FTableTreeViewNodeFilteringAsyncTask, STATGROUP_TaskGraphTasks); }
@@ -676,8 +676,8 @@ class FTableTreeViewHierarchyFilteringAsyncTask
 {
 public:
 	FTableTreeViewHierarchyFilteringAsyncTask(STableTreeView* InPtr)
+		: TableTreeViewPtr(InPtr)
 	{
-		TableTreeViewPtr = InPtr;
 	}
 
 	FORCEINLINE TStatId GetStatId() const { RETURN_QUICK_DECLARE_CYCLE_STAT(FTableTreeViewHierarchyFilteringAsyncTask, STATGROUP_TaskGraphTasks); }
@@ -702,10 +702,10 @@ class FTableTreeViewSortingAsyncTask
 {
 public:
 	FTableTreeViewSortingAsyncTask(STableTreeView* InPtr, ITableCellValueSorter* InSorter, EColumnSortMode::Type InColumnSortMode)
+		: TableTreeViewPtr(InPtr)
+		, Sorter(InSorter)
+		, ColumnSortMode(InColumnSortMode)
 	{
-		TableTreeViewPtr = InPtr;
-		Sorter = InSorter;
-		ColumnSortMode = InColumnSortMode;
 	}
 
 	FORCEINLINE TStatId GetStatId() const { RETURN_QUICK_DECLARE_CYCLE_STAT(FTableTreeViewSortingAsyncTask, STATGROUP_TaskGraphTasks); }
@@ -721,9 +721,9 @@ public:
 	}
 
 private:
-	STableTreeView* TableTreeViewPtr;
-	ITableCellValueSorter* Sorter;
-	EColumnSortMode::Type ColumnSortMode;
+	STableTreeView* TableTreeViewPtr = nullptr;
+	ITableCellValueSorter* Sorter = nullptr;
+	EColumnSortMode::Type ColumnSortMode = EColumnSortMode::Type::None;
 };
 
 ////////////////////////////////////////////////////////////////////////////////////////////////////
@@ -732,9 +732,9 @@ class FTableTreeViewGroupingAsyncTask
 {
 public:
 	FTableTreeViewGroupingAsyncTask(STableTreeView* InPtr, TArray<TSharedPtr<FTreeNodeGrouping>>* InGroupings)
+		: TableTreeViewPtr(InPtr)
+		, Groupings(InGroupings)
 	{
-		TableTreeViewPtr = InPtr;
-		Groupings = InGroupings;
 	}
 
 	FORCEINLINE TStatId GetStatId() const { RETURN_QUICK_DECLARE_CYCLE_STAT(FTableTreeViewGroupingAsyncTask, STATGROUP_TaskGraphTasks); }
@@ -751,7 +751,7 @@ public:
 
 private:
 	STableTreeView* TableTreeViewPtr = nullptr;
-	TArray<TSharedPtr<FTreeNodeGrouping>>* Groupings;
+	TArray<TSharedPtr<FTreeNodeGrouping>>* Groupings = nullptr;
 };
 
 ////////////////////////////////////////////////////////////////////////////////////////////////////
@@ -762,7 +762,8 @@ public:
 	FSearchForItemToSelectTask(TSharedPtr<FTableTaskCancellationToken> InToken, TSharedPtr<STableTreeView> InPtr)
 		: CancellationToken(InToken)
 		, TableTreeViewPtr(InPtr)
-	{}
+	{
+	}
 
 	FORCEINLINE TStatId GetStatId() const { RETURN_QUICK_DECLARE_CYCLE_STAT(FSearchForItemToSelectTask, STATGROUP_TaskGraphTasks); }
 	ENamedThreads::Type GetDesiredThread() { return ENamedThreads::Type::AnyThread; }
@@ -770,7 +771,13 @@ public:
 
 	void DoTask(ENamedThreads::Type CurrentThread, const FGraphEventRef& MyCompletionGraphEvent)
 	{
-		TableTreeViewPtr->SearchForItem(CancellationToken);
+		if (!CancellationToken.IsValid() || !CancellationToken->ShouldCancel())
+		{
+			if (TableTreeViewPtr.IsValid())
+			{
+				TableTreeViewPtr->SearchForItem(CancellationToken);
+			}
+		}
 	}
 
 private:
@@ -787,7 +794,8 @@ public:
 		: CancellationToken(InToken)
 		, TableTreeViewPtr(InPtr)
 		, RowIndex(InRowIndex)
-		{}
+	{
+	}
 
 	FORCEINLINE TStatId GetStatId() const { RETURN_QUICK_DECLARE_CYCLE_STAT(FSelectNodeByTableRowIndexTask, STATGROUP_TaskGraphTasks); }
 	ENamedThreads::Type GetDesiredThread() { return ENamedThreads::Type::GameThread; }
@@ -795,16 +803,19 @@ public:
 
 	void DoTask(ENamedThreads::Type CurrentThread, const FGraphEventRef& MyCompletionGraphEvent)
 	{
-		if (!CancellationToken->ShouldCancel())
+		if (!CancellationToken.IsValid() || !CancellationToken->ShouldCancel())
 		{
-			TableTreeViewPtr->SelectNodeByTableRowIndex(RowIndex);
+			if (TableTreeViewPtr.IsValid())
+			{
+				TableTreeViewPtr->SelectNodeByTableRowIndex(RowIndex);
+			}
 		}
 	}
 
 private:
 	TSharedPtr< FTableTaskCancellationToken> CancellationToken;
 	TSharedPtr<STableTreeView> TableTreeViewPtr;
-	uint32 RowIndex;
+	uint32 RowIndex = 0;
 };
 
 ////////////////////////////////////////////////////////////////////////////////////////////////////
diff --git a/Engine/Source/Developer/TraceInsights/Private/Insights/TaskGraphProfiler/TaskGraphProfilerManager.cpp b/Engine/Source/Developer/TraceInsights/Private/Insights/TaskGraphProfiler/TaskGraphProfilerManager.cpp
--- a/Engine/Source/Developer/TraceInsights/Private/Insights/TaskGraphProfiler/TaskGraphProfilerManager.cpp
+++ b/Engine/Source/Developer/TraceInsights/Private/Insights/TaskGraphProfiler/TaskGraphProfilerManager.cpp
@@ -214,7 +214,7 @@ void FTaskGraphProfilerManager::RegisterTimingProfilerLayoutExtensions(FInsights
 	FMinorTabConfig& MinorTabConfig = InOutExtender.AddMinorTabConfig();
 	MinorTabConfig.TabId = FTaskGraphProfilerTabs::TaskTableTreeViewTabID;
 	MinorTabConfig.TabLabel = LOCTEXT("TaskTableTreeViewTabTitle", "Tasks");
-	MinorTabConfig.TabTooltip = LOCTEXT("TaskTableTreeViewTabTitleTooltip", "Opens the Task Table Tree View tab, that allows Task Graph profilling.");
+	MinorTabConfig.TabTooltip = LOCTEXT("TaskTableTreeViewTabTitleTooltip", "Opens the Task Table Tree View tab, that allows Task Graph profiling.");
 	MinorTabConfig.TabIcon = FSlateIcon(FInsightsStyle::GetStyleSetName(), "Icons.TasksView");
 	MinorTabConfig.OnSpawnTab = FOnSpawnTab::CreateRaw(this, &FTaskGraphProfilerManager::SpawnTab_TaskTableTreeView);
 	MinorTabConfig.CanSpawnTab = FCanSpawnTab::CreateRaw(this, &FTaskGraphProfilerManager::CanSpawnTab_TaskTableTreeView);
@@ -231,6 +231,7 @@ TSharedRef<SDockTab> FTaskGraphProfilerManager::SpawnTab_TaskTableTreeView(const
 {
 	TSharedRef<FTaskTable> TaskTable = MakeShared<FTaskTable>();
 	TaskTable->Reset();
+	TaskTable->SetDisplayName(LOCTEXT("TaskTableTreeViewTabTitle", "Tasks"));
 
 	const TSharedRef<SDockTab> DockTab = SNew(SDockTab)
 		.ShouldAutosize(false)
diff --git a/Engine/Source/Developer/TraceInsights/Private/Insights/ViewModels/FilterConfiguratorNode.h b/Engine/Source/Developer/TraceInsights/Private/Insights/ViewModels/FilterConfiguratorNode.h
--- a/Engine/Source/Developer/TraceInsights/Private/Insights/ViewModels/FilterConfiguratorNode.h
+++ b/Engine/Source/Developer/TraceInsights/Private/Insights/ViewModels/FilterConfiguratorNode.h
@@ -13,20 +13,6 @@ namespace Insights
 
 ////////////////////////////////////////////////////////////////////////////////////////////////////
 
-enum class EFilterConfiguratorNodeType
-{
-	/** A Filter node. */
-	Filter,
-
-	/** A group node. */
-	Group,
-
-	/** Invalid enum type, may be used as a number of enumerations. */
-	InvalidOrMax,
-};
-
-////////////////////////////////////////////////////////////////////////////////////////////////////
-
 /** Type definition for shared pointers to instances of FFilterConfiguratorNode. */
 typedef TSharedPtr<class FFilterConfiguratorNode> FFilterConfiguratorNodePtr;
 
@@ -70,7 +56,7 @@ public:
 	void SetSelectedFilterGroupOperator(TSharedPtr<FFilterGroupOperator> InSelectedFilterGroupOperator) { SelectedFilterGroupOperator = InSelectedFilterGroupOperator; }
 	TSharedPtr<FFilterGroupOperator> GetSelectedFilterGroupOperator() const { return SelectedFilterGroupOperator; }
 
-	TSharedPtr<TArray<TSharedPtr<IFilterOperator>>> GetAvailableFilterOperators() const {	return AvailableFilterOperators;	}
+	TSharedPtr<TArray<TSharedPtr<IFilterOperator>>> GetAvailableFilterOperators() const { return AvailableFilterOperators; }
 
 	const FString& GetTextBoxValue() { return TextBoxValue; }
 	void SetTextBoxValue(const FString& InValue);
