Missed a couple of test renaming
#rnx

diff --git a/Engine/Source/Developer/TraceInsights/Private/Insights/Tests/FunctionalTests/UITests.cpp b/Engine/Source/Developer/TraceInsights/Private/Insights/Tests/FunctionalTests/UITests.cpp
--- a/Engine/Source/Developer/TraceInsights/Private/Insights/Tests/FunctionalTests/UITests.cpp
+++ b/Engine/Source/Developer/TraceInsights/Private/Insights/Tests/FunctionalTests/UITests.cpp
@@ -14,7 +14,7 @@
 
 DECLARE_LOG_CATEGORY_EXTERN(UITests, Log, All);
 
-IMPLEMENT_SIMPLE_AUTOMATION_TEST(FHideAndShowAllTimingViewTabs, "System.Insights.Analysis.TimingInsights.HideAndShowAllTimingViewTabs", EAutomationTestFlags::ApplicationContextMask | EAutomationTestFlags::EngineFilter)
+IMPLEMENT_SIMPLE_AUTOMATION_TEST(FHideAndShowAllTimingViewTabs, "System.Insights.Analysis.TimingInsights.HideAndShowAllTimingViewTabs", EAutomationTestFlags::ProgramContext | EAutomationTestFlags::EngineFilter)
 bool FHideAndShowAllTimingViewTabs::RunTest(const FString& Parameters)
 {
 	TSharedPtr<FTimingProfilerManager> TimingProfilerManager = FTimingProfilerManager::Get();
