[Insights] Massive refactoring of TraceInsights module:
  * Added new TraceInsightsCore module.
  * Moved Common, Filter and Table API to TraceInsightsCore.
  * Moved MemoryProfiler code into UE::Insights::MemoryProfiler namespace.
  * Moved NetworkingProfiler code into UE::Insights::NetworkingProfiler namespace.
  * Moved some of the TimingProfiler code into UE::Insights::TimingProfiler namespace.
  * Moved some of the public API into UE::Insights namespace.
  * Updated localization namespaces too.
  * Fixed many typos in comments.
  * Added backward compatibility types / workarounds for types moved into new namespaces.
  * Added dependency to TraceInsightsCore module to each Insights plugin.

Notes about updated public API:
  * Redirected "Insights/IFilterExecutor.h" to "InsightsCore/Filter/ViewModels/IFilterExecutor.h". (Insights::IFilterExecutor --> UE::Insights::IFilterExecutor)
  * Removed "Insights/AvailabilityCheck.h" (added after UE 5.4 release, so we do not need backward compatibility). Use "InsightsCore/Common/AvailabilityCheck.h" instead. (Insights::FAvailabilityCheck --> UE::Insights::FAvailabilityCheck)
  * Redirected "Insights/Common/PaintUtils.h" to "InsightsCore/Common/PaintUtils.h". (struct ::FDrawContext --> class UE::Insights::FDrawContext)
  * Redirected "Insights/Common/SimpleRtti.h" to "InsightsCore/Common/SimpleRtti.h".
  * Redirected "Insights/Common/TimeUtils.h" to "InsightsCore/Common/TimeUtils.h".
      * Moved API from TimeUtils namespace to UE::Insights namespace (ex.: UE::Insights::FormatTimeAuto())
      * Moved time constants to struct UE::Insights::FTimeValue (ex.: UE::Insights::FTimeValue::Second)
      * struct TimeUtils::FTimeSplit --> UE::Insights::FTimeSplit
      * Fixed typo for FTimeValue::Milisecond --> FTimeValue::Millisecond
      * Fixed typo for FTimeSplit::Miliseconds --> FTimeSplit::Milliseconds
  * Insights::ITimingViewSession moved to UE::Insights::Timing::ITimingViewSession.
  * Insights::ITimingViewExtender moved to UE::Insights::Timing::ITimingViewExtender.

diff --git a/Engine/Plugins/Animation/GameplayInsights/Source/GameplayInsights/GameplayInsights.Build.cs b/Engine/Plugins/Animation/GameplayInsights/Source/GameplayInsights/GameplayInsights.Build.cs
--- a/Engine/Plugins/Animation/GameplayInsights/Source/GameplayInsights/GameplayInsights.Build.cs
+++ b/Engine/Plugins/Animation/GameplayInsights/Source/GameplayInsights/GameplayInsights.Build.cs
@@ -18,6 +18,7 @@ namespace UnrealBuildTool.Rules
 				"TraceAnalysis",
 				"TraceServices",
 				"TraceInsights",
+				"TraceInsightsCore",
 				"AssetRegistry",
 				"ApplicationCore",
 				"RewindDebuggerInterface",
diff --git a/Engine/Plugins/RenderGraphInsights/Source/RenderGraphInsights/RenderGraphInsights.Build.cs b/Engine/Plugins/RenderGraphInsights/Source/RenderGraphInsights/RenderGraphInsights.Build.cs
--- a/Engine/Plugins/RenderGraphInsights/Source/RenderGraphInsights/RenderGraphInsights.Build.cs
+++ b/Engine/Plugins/RenderGraphInsights/Source/RenderGraphInsights/RenderGraphInsights.Build.cs
@@ -19,6 +19,7 @@ namespace UnrealBuildTool.Rules
 				"TraceAnalysis",
 				"TraceServices",
 				"TraceInsights",
+				"TraceInsightsCore",
 				"InputCore"
 			});
 		}
diff --git a/Engine/Plugins/Runtime/NetworkPredictionInsights/Source/NetworkPredictionInsights/NetworkPredictionInsights.Build.cs b/Engine/Plugins/Runtime/NetworkPredictionInsights/Source/NetworkPredictionInsights/NetworkPredictionInsights.Build.cs
--- a/Engine/Plugins/Runtime/NetworkPredictionInsights/Source/NetworkPredictionInsights/NetworkPredictionInsights.Build.cs
+++ b/Engine/Plugins/Runtime/NetworkPredictionInsights/Source/NetworkPredictionInsights/NetworkPredictionInsights.Build.cs
@@ -18,12 +18,12 @@ namespace UnrealBuildTool.Rules
 				"TraceAnalysis",
 				"TraceServices",
 				"TraceInsights",
+				"TraceInsightsCore",
 				"AssetRegistry",
 				"ApplicationCore",
-				
 			});
 
-            if (Target.bCompileAgainstEngine)
+			if (Target.bCompileAgainstEngine)
 			{
 				PrivateDependencyModuleNames.AddRange(
 				new string[]
diff --git a/Engine/Source/Developer/TraceInsights/Private/Insights/ContextSwitches/ContextSwitchesProfilerManager.cpp b/Engine/Source/Developer/TraceInsights/Private/Insights/ContextSwitches/ContextSwitchesProfilerManager.cpp
--- a/Engine/Source/Developer/TraceInsights/Private/Insights/ContextSwitches/ContextSwitchesProfilerManager.cpp
+++ b/Engine/Source/Developer/TraceInsights/Private/Insights/ContextSwitches/ContextSwitchesProfilerManager.cpp
@@ -5,9 +5,13 @@
 #include "Features/IModularFeatures.h"
 #include "Framework/Commands/UICommandList.h"
 #include "Modules/ModuleManager.h"
+
+// TraceServices
 #include "TraceServices/Model/ContextSwitches.h"
 
+// TraceInsights
 #include "Insights/ContextSwitches/ViewModels/ContextSwitchesSharedState.h"
+#include "Insights/ITimingViewExtender.h"
 #include "Insights/TimingProfilerManager.h"
 #include "Insights/Widgets/STimingProfilerWindow.h"
 #include "Insights/Widgets/STimingView.h"
@@ -104,7 +108,7 @@ FContextSwitchesProfilerManager::~FContextSwitchesProfilerManager()
 
 	if (ContextSwitchesSharedState.IsValid())
 	{
-		IModularFeatures::Get().UnregisterModularFeature(Insights::TimingViewExtenderFeatureName, ContextSwitchesSharedState.Get());
+		IModularFeatures::Get().UnregisterModularFeature(UE::Insights::Timing::TimingViewExtenderFeatureName, ContextSwitchesSharedState.Get());
 	}
 }
 
@@ -137,6 +141,8 @@ bool FContextSwitchesProfilerManager::Tick(float DeltaTime)
 			const TraceServices::IContextSwitchesProvider* ContextSwitchesProvider = TraceServices::ReadContextSwitchesProvider(*Session.Get());
 			if (ContextSwitchesProvider && ContextSwitchesProvider->HasData())
 			{
+				using namespace UE::Insights::TimingProfiler;
+
 				TSharedPtr<STimingProfilerWindow> Window = FTimingProfilerManager::Get()->GetProfilerWindow();
 				if (!Window.IsValid())
 				{
@@ -155,7 +161,7 @@ bool FContextSwitchesProfilerManager::Tick(float DeltaTime)
 				{
 					ContextSwitchesSharedState = MakeShared<FContextSwitchesSharedState>(TimingView.Get());
 					ContextSwitchesSharedState->AddCommands();
-					IModularFeatures::Get().RegisterModularFeature(Insights::TimingViewExtenderFeatureName, ContextSwitchesSharedState.Get());
+					IModularFeatures::Get().RegisterModularFeature(UE::Insights::Timing::TimingViewExtenderFeatureName, ContextSwitchesSharedState.Get());
 				}
 			}
 
diff --git a/Engine/Source/Developer/TraceInsights/Private/Insights/ContextSwitches/ContextSwitchesProfilerManager.h b/Engine/Source/Developer/TraceInsights/Private/Insights/ContextSwitches/ContextSwitchesProfilerManager.h
--- a/Engine/Source/Developer/TraceInsights/Private/Insights/ContextSwitches/ContextSwitchesProfilerManager.h
+++ b/Engine/Source/Developer/TraceInsights/Private/Insights/ContextSwitches/ContextSwitchesProfilerManager.h
@@ -2,10 +2,11 @@
 
 #pragma once
 
+#include "CoreTypes.h"
+
 #include "Containers/Ticker.h"
-#include "CoreMinimal.h"
 
-// Insights
+// TraceInsights
 #include "Insights/InsightsManager.h"
 #include "Insights/IUnrealInsightsModule.h"
 
@@ -51,7 +52,7 @@ public:
 
 	void OnSessionChanged();
 
-	TSharedPtr<Insights::FContextSwitchesSharedState> GetContextSwitchesSharedState() { return ContextSwitchesSharedState;	}
+	TSharedPtr<FContextSwitchesSharedState> GetContextSwitchesSharedState() { return ContextSwitchesSharedState; }
 
 private:
 	/** Updates this manager, done through FCoreTicker. */
@@ -62,7 +63,7 @@ private:
 private:
 	bool bIsInitialized;
 	bool bIsAvailable;
-	FAvailabilityCheck AvailabilityCheck;
+	UE::Insights::FAvailabilityCheck AvailabilityCheck;
 
 	/** The delegate to be invoked when this manager ticks. */
 	FTickerDelegate OnTick;
diff --git a/Engine/Source/Developer/TraceInsights/Private/Insights/ContextSwitches/ViewModels/ContextSwitchTimingEvent.h b/Engine/Source/Developer/TraceInsights/Private/Insights/ContextSwitches/ViewModels/ContextSwitchTimingEvent.h
--- a/Engine/Source/Developer/TraceInsights/Private/Insights/ContextSwitches/ViewModels/ContextSwitchTimingEvent.h
+++ b/Engine/Source/Developer/TraceInsights/Private/Insights/ContextSwitches/ViewModels/ContextSwitchTimingEvent.h
@@ -4,7 +4,7 @@
 
 #include "CoreMinimal.h"
 
-// Insights
+// TraceInsights
 #include "Insights/ViewModels/TimingEvent.h"
 
 ////////////////////////////////////////////////////////////////////////////////////////////////////
diff --git a/Engine/Source/Developer/TraceInsights/Private/Insights/ContextSwitches/ViewModels/ContextSwitchesFilterConverters.cpp b/Engine/Source/Developer/TraceInsights/Private/Insights/ContextSwitches/ViewModels/ContextSwitchesFilterConverters.cpp
--- a/Engine/Source/Developer/TraceInsights/Private/Insights/ContextSwitches/ViewModels/ContextSwitchesFilterConverters.cpp
+++ b/Engine/Source/Developer/TraceInsights/Private/Insights/ContextSwitches/ViewModels/ContextSwitchesFilterConverters.cpp
@@ -2,9 +2,11 @@
 
 #include "ContextSwitchesFilterConverters.h"
 
+// TraceServices
 #include "TraceServices/Model/ContextSwitches.h"
 #include "TraceServices/Model/Threads.h"
 
+// TraceInsights
 #include "Insights/InsightsManager.h"
 
 #define LOCTEXT_NAMESPACE "Insights::ContextSwitchesFilterValueConverter"
diff --git a/Engine/Source/Developer/TraceInsights/Private/Insights/ContextSwitches/ViewModels/ContextSwitchesFilterConverters.h b/Engine/Source/Developer/TraceInsights/Private/Insights/ContextSwitches/ViewModels/ContextSwitchesFilterConverters.h
--- a/Engine/Source/Developer/TraceInsights/Private/Insights/ContextSwitches/ViewModels/ContextSwitchesFilterConverters.h
+++ b/Engine/Source/Developer/TraceInsights/Private/Insights/ContextSwitches/ViewModels/ContextSwitchesFilterConverters.h
@@ -2,14 +2,15 @@
 
 #pragma once
 
-#include "Insights/ViewModels/Filters.h"
+// TraceInsightsCore
+#include "InsightsCore/Filter/ViewModels/Filters.h"
 
 namespace Insights
 {
 
 ////////////////////////////////////////////////////////////////////////////////////////////////////
 
-class FCoreEventNameFilterValueConverter : public IFilterValueConverter
+class FCoreEventNameFilterValueConverter : public UE::Insights::IFilterValueConverter
 {
 public:
 	virtual bool Convert(const FString& Input, int64& Output, FText& OutError) const override;
diff --git a/Engine/Source/Developer/TraceInsights/Private/Insights/ContextSwitches/ViewModels/ContextSwitchesSharedState.cpp b/Engine/Source/Developer/TraceInsights/Private/Insights/ContextSwitches/ViewModels/ContextSwitchesSharedState.cpp
--- a/Engine/Source/Developer/TraceInsights/Private/Insights/ContextSwitches/ViewModels/ContextSwitchesSharedState.cpp
+++ b/Engine/Source/Developer/TraceInsights/Private/Insights/ContextSwitches/ViewModels/ContextSwitchesSharedState.cpp
@@ -5,10 +5,15 @@
 #include "Framework/Commands/Commands.h"
 #include "Framework/Commands/UICommandList.h"
 #include "Framework/MultiBox/MultiBoxBuilder.h"
+
+// TraceServices
 #include "TraceServices/Model/Threads.h"
 #include "TraceServices/Model/ContextSwitches.h"
 
-// Insights
+// TraceInsightsCore
+#include "InsightsCore/Filter/ViewModels/FilterConfigurator.h"
+
+// TraceInsights
 #include "Insights/ContextSwitches/ContextSwitchesProfilerManager.h"
 #include "Insights/ContextSwitches/ViewModels/ContextSwitchesFilterConverters.h"
 #include "Insights/ContextSwitches/ViewModels/ContextSwitchesTimingTrack.h"
@@ -18,7 +23,6 @@
 #include "Insights/InsightsManager.h"
 #include "Insights/InsightsStyle.h"
 #include "Insights/TimingProfilerManager.h"
-#include "Insights/ViewModels/FilterConfigurator.h"
 #include "Insights/ViewModels/ThreadTimingTrack.h"
 #include "Insights/Widgets/STimingProfilerWindow.h"
 #include "Insights/Widgets/STimingView.h"
@@ -113,7 +117,7 @@ UE_ENABLE_OPTIMIZATION_SHIP
 // FContextSwitchesSharedState
 ////////////////////////////////////////////////////////////////////////////////////////////////////
 
-FContextSwitchesSharedState::FContextSwitchesSharedState(STimingView* InTimingView)
+FContextSwitchesSharedState::FContextSwitchesSharedState(UE::Insights::TimingProfiler::STimingView* InTimingView)
 	: TimingViewSession(InTimingView)
 	, ThreadsSerial(0)
 	, CpuCoresSerial(0)
@@ -128,7 +132,7 @@ FContextSwitchesSharedState::FContextSwitchesSharedState(STimingView* InTimingVi
 
 ////////////////////////////////////////////////////////////////////////////////////////////////////
 
-void FContextSwitchesSharedState::OnBeginSession(Insights::ITimingViewSession& InSession)
+void FContextSwitchesSharedState::OnBeginSession(UE::Insights::Timing::ITimingViewSession& InSession)
 {
 	if (&InSession != TimingViewSession)
 	{
@@ -162,7 +166,7 @@ void FContextSwitchesSharedState::OnBeginSession(Insights::ITimingViewSession& I
 
 ////////////////////////////////////////////////////////////////////////////////////////////////////
 
-void FContextSwitchesSharedState::OnEndSession(Insights::ITimingViewSession& InSession)
+void FContextSwitchesSharedState::OnEndSession(UE::Insights::Timing::ITimingViewSession& InSession)
 {
 	if (&InSession != TimingViewSession)
 	{
@@ -188,7 +192,7 @@ void FContextSwitchesSharedState::OnEndSession(Insights::ITimingViewSession& InS
 
 ////////////////////////////////////////////////////////////////////////////////////////////////////
 
-void FContextSwitchesSharedState::Tick(Insights::ITimingViewSession& InSession, const TraceServices::IAnalysisSession& InAnalysisSession)
+void FContextSwitchesSharedState::Tick(UE::Insights::Timing::ITimingViewSession& InSession, const TraceServices::IAnalysisSession& InAnalysisSession)
 {
 	if (&InSession != TimingViewSession)
 	{
@@ -217,7 +221,7 @@ void FContextSwitchesSharedState::Tick(Insights::ITimingViewSession& InSession,
 
 ////////////////////////////////////////////////////////////////////////////////////////////////////
 
-void FContextSwitchesSharedState::ExtendCpuTracksFilterMenu(Insights::ITimingViewSession& InSession, FMenuBuilder& InMenuBuilder)
+void FContextSwitchesSharedState::ExtendCpuTracksFilterMenu(UE::Insights::Timing::ITimingViewSession& InSession, FMenuBuilder& InMenuBuilder)
 {
 	if (&InSession != TimingViewSession)
 	{
@@ -256,8 +260,10 @@ void FContextSwitchesSharedState::BuildSubMenu(FMenuBuilder& InMenuBuilder)
 
 ////////////////////////////////////////////////////////////////////////////////////////////////////
 
-void FContextSwitchesSharedState::AddQuickFindFilters(TSharedPtr<FFilterConfigurator> FilterConfigurator)
+void FContextSwitchesSharedState::AddQuickFindFilters(TSharedPtr<UE::Insights::FFilterConfigurator> FilterConfigurator)
 {
+	using namespace UE::Insights;
+
 	TSharedPtr<TArray<TSharedPtr<IFilterOperator>>> CoreEventNameFilterOperators = MakeShared<TArray<TSharedPtr<IFilterOperator>>>();
 	CoreEventNameFilterOperators->Add(StaticCastSharedRef<IFilterOperator>(MakeShared<FFilterOperator<int64>>(EFilterOperator::Eq, TEXT("Is"), [](int64 lhs, int64 rhs) { return lhs == rhs; })));
 
@@ -302,7 +308,7 @@ void FContextSwitchesSharedState::AddCommands()
 {
 	FContextSwitchesStateCommands::Register();
 
-	TSharedPtr<STimingView> TimingView = GetTimingView();
+	TSharedPtr<UE::Insights::TimingProfiler::STimingView> TimingView = GetTimingView();
 	
 	if (!TimingView.IsValid())
 	{
@@ -373,7 +379,7 @@ bool FContextSwitchesSharedState::AreContextSwitchesAvailable() const
 
 void FContextSwitchesSharedState::AddCoreTracks()
 {
-	TSharedPtr<STimingView> TimingView = GetTimingView();
+	TSharedPtr<UE::Insights::TimingProfiler::STimingView> TimingView = GetTimingView();
 
 	if (!TimingView.IsValid())
 	{
@@ -423,7 +429,7 @@ void FContextSwitchesSharedState::AddCoreTracks()
 
 void FContextSwitchesSharedState::RemoveCoreTracks()
 {
-	TSharedPtr<STimingView> TimingView = GetTimingView();
+	TSharedPtr<UE::Insights::TimingProfiler::STimingView> TimingView = GetTimingView();
 	CpuCoresSerial = 0;
 
 	if (TimingView)
@@ -468,7 +474,7 @@ void FContextSwitchesSharedState::AddContextSwitchesChildTracks()
 	//TODO: Create "Cpu Thread" timing tracks also for threads without cpu timing events (i.e. only with context switch events).
 #endif
 
-	TSharedPtr<STimingView> TimingView = GetTimingView();
+	TSharedPtr<UE::Insights::TimingProfiler::STimingView> TimingView = GetTimingView();
 	if (!TimingView.IsValid())
 	{
 		return;
@@ -501,7 +507,7 @@ void FContextSwitchesSharedState::RemoveContextSwitchesChildTracks()
 {
 	ThreadsSerial = 0;
 	
-	TSharedPtr<STimingView> TimingView = GetTimingView();
+	TSharedPtr<UE::Insights::TimingProfiler::STimingView> TimingView = GetTimingView();
 	if (!TimingView.IsValid())
 	{
 		return;
@@ -583,7 +589,7 @@ void FContextSwitchesSharedState::SetNonTargetProcessEventsVisible(bool bOnOff)
 {
 	if (bAreNonTargetProcessEventsVisible != bOnOff)
 	{
-		TSharedPtr<STimingView> TimingView = GetTimingView();
+		TSharedPtr<UE::Insights::TimingProfiler::STimingView> TimingView = GetTimingView();
 		const TSharedPtr<ITimingEventFilter> EventFilter = TimingView->GetEventFilter();
 		for (auto Track : CpuCoreTimingTracks)
 		{
@@ -605,7 +611,7 @@ void FContextSwitchesSharedState::SetNonTargetProcessEventsVisible(bool bOnOff)
 
 bool FContextSwitchesSharedState::IsValidCpuCoreEventSelected() const
 {
-	TSharedPtr<STimingView> TimingView = GetTimingView();
+	TSharedPtr<UE::Insights::TimingProfiler::STimingView> TimingView = GetTimingView();
 	if (TimingView == nullptr)
 	{
 		return false;
@@ -621,7 +627,7 @@ bool FContextSwitchesSharedState::IsValidCpuCoreEventSelected() const
 
 bool FContextSwitchesSharedState::IsValidContextSwitchEventSelected() const
 {
-	TSharedPtr<STimingView> TimingView = GetTimingView();
+	TSharedPtr<UE::Insights::TimingProfiler::STimingView> TimingView = GetTimingView();
 	if (TimingView == nullptr)
 	{
 		return false;
@@ -637,7 +643,7 @@ bool FContextSwitchesSharedState::IsValidContextSwitchEventSelected() const
 
 void FContextSwitchesSharedState::Command_NavigateToCpuThreadEvent_Execute()
 {
-	TSharedPtr<STimingView> TimingView = GetTimingView();
+	TSharedPtr<UE::Insights::TimingProfiler::STimingView> TimingView = GetTimingView();
 
 	if (TimingView && ensure(IsValidCpuCoreEventSelected()))
 	{
@@ -689,7 +695,7 @@ bool FContextSwitchesSharedState::Command_NavigateToCpuThreadEvent_CanExecute()
 		return false;
 	}
 
-	TSharedPtr<STimingView> TimingView = GetTimingView();
+	TSharedPtr<UE::Insights::TimingProfiler::STimingView> TimingView = GetTimingView();
 
 	if (TimingView && IsValidCpuCoreEventSelected())
 	{
@@ -718,7 +724,7 @@ bool FContextSwitchesSharedState::Command_NavigateToCpuThreadEvent_CanExecute()
 
 void FContextSwitchesSharedState::Command_DockCpuThreadTrackToBottom_Execute()
 {
-	TSharedPtr<STimingView> TimingView = GetTimingView();
+	TSharedPtr<UE::Insights::TimingProfiler::STimingView> TimingView = GetTimingView();
 
 	if (TimingView && ensure(IsValidCpuCoreEventSelected()))
 	{
@@ -755,7 +761,7 @@ void FContextSwitchesSharedState::Command_DockCpuThreadTrackToBottom_Execute()
 
 void FContextSwitchesSharedState::Command_NavigateToCpuCoreEvent_Execute()
 {
-	TSharedPtr<STimingView> TimingView = GetTimingView();
+	TSharedPtr<UE::Insights::TimingProfiler::STimingView> TimingView = GetTimingView();
 
 	if (TimingView && ensure(IsValidContextSwitchEventSelected()))
 	{
@@ -791,7 +797,7 @@ void FContextSwitchesSharedState::Command_NavigateToCpuCoreEvent_Execute()
 
 void FContextSwitchesSharedState::Command_DockCpuCoreTrackToTop_Execute()
 {
-	TSharedPtr<STimingView> TimingView = GetTimingView();
+	TSharedPtr<UE::Insights::TimingProfiler::STimingView> TimingView = GetTimingView();
 
 	if (TimingView && ensure(IsValidContextSwitchEventSelected()))
 	{
@@ -852,7 +858,7 @@ void FContextSwitchesSharedState::GetThreadInfo(uint32 InSystemThreadId, uint32&
 TSharedPtr<FThreadTimingTrack> FContextSwitchesSharedState::GetThreadTimingTrack(uint32 ThreadId) const
 {
 	TSharedPtr<FThreadTimingTrack> FoundTrack;
-	TSharedPtr<STimingView> TimingView = GetTimingView();
+	TSharedPtr<UE::Insights::TimingProfiler::STimingView> TimingView = GetTimingView();
 
 	if (TimingView)
 	{
@@ -873,19 +879,19 @@ TSharedPtr<FThreadTimingTrack> FContextSwitchesSharedState::GetThreadTimingTrack
 
 ////////////////////////////////////////////////////////////////////////////////////////////////////
 
-TSharedPtr<Insights::FCpuCoreTimingTrack> FContextSwitchesSharedState::GetCpuCoreTimingTrack(uint32 CoreNumber) const
+TSharedPtr<FCpuCoreTimingTrack> FContextSwitchesSharedState::GetCpuCoreTimingTrack(uint32 CoreNumber) const
 {
 	TSharedPtr<FCpuCoreTimingTrack> FoundTrack;
-	TSharedPtr<STimingView> TimingView = GetTimingView();
+	TSharedPtr<UE::Insights::TimingProfiler::STimingView> TimingView = GetTimingView();
 
 	if (TimingView)
 	{
 		TimingView->EnumerateAllTracks([&FoundTrack, CoreNumber](TSharedPtr<FBaseTimingTrack>& Track) -> bool
 		{
-			if (Track->Is<Insights::FCpuCoreTimingTrack>() &&
-				Track->As<Insights::FCpuCoreTimingTrack>().GetCoreNumber() == CoreNumber)
+			if (Track->Is<FCpuCoreTimingTrack>() &&
+				Track->As<FCpuCoreTimingTrack>().GetCoreNumber() == CoreNumber)
 			{
-				FoundTrack = StaticCastSharedPtr<Insights::FCpuCoreTimingTrack>(Track);
+				FoundTrack = StaticCastSharedPtr<FCpuCoreTimingTrack>(Track);
 				return false;
 			}
 			return true;
@@ -897,16 +903,11 @@ TSharedPtr<Insights::FCpuCoreTimingTrack> FContextSwitchesSharedState::GetCpuCor
 
 ////////////////////////////////////////////////////////////////////////////////////////////////////
 
-TSharedPtr<STimingView> FContextSwitchesSharedState::GetTimingView()
+TSharedPtr<UE::Insights::TimingProfiler::STimingView> FContextSwitchesSharedState::GetTimingView()
 {
-	TSharedPtr<STimingView> TimingView;
+	using namespace UE::Insights::TimingProfiler;
 	TSharedPtr<STimingProfilerWindow> Window = FTimingProfilerManager::Get()->GetProfilerWindow();
-	if (Window.IsValid())
-	{
-		TimingView = Window->GetTimingView();
-	}
-
-	return TimingView;
+	return Window.IsValid() ? Window->GetTimingView() : nullptr;
 }
 
 ////////////////////////////////////////////////////////////////////////////////////////////////////
diff --git a/Engine/Source/Developer/TraceInsights/Private/Insights/ContextSwitches/ViewModels/ContextSwitchesSharedState.h b/Engine/Source/Developer/TraceInsights/Private/Insights/ContextSwitches/ViewModels/ContextSwitchesSharedState.h
--- a/Engine/Source/Developer/TraceInsights/Private/Insights/ContextSwitches/ViewModels/ContextSwitchesSharedState.h
+++ b/Engine/Source/Developer/TraceInsights/Private/Insights/ContextSwitches/ViewModels/ContextSwitchesSharedState.h
@@ -2,10 +2,11 @@
 
 #pragma once
 
-#include "CoreMinimal.h"
+#include "CoreTypes.h"
+
 #include "Framework/Commands/Commands.h"
 
-// Insights
+// TraceInsights
 #include "Insights/ITimingViewExtender.h"
 
 ////////////////////////////////////////////////////////////////////////////////////////////////////
@@ -14,19 +15,16 @@ class FBaseTimingTrack;
 class FMenuBuilder;
 class FThreadTimingTrack;
 class ITimingEvent;
-class STimingView;
 
-namespace TraceServices
-{
-	class IAnalysisSession;
-}
+namespace TraceServices { class IAnalysisSession; }
+namespace UE::Insights::Timing { class ITimingViewSession; }
+namespace UE::Insights::TimingProfiler { class STimingView; }
 
 namespace Insights
 {
 
 class FCpuCoreTimingTrack;
 class FContextSwitchesTimingTrack;
-class ITimingViewSession;
 
 ////////////////////////////////////////////////////////////////////////////////////////////////////
 
@@ -55,24 +53,24 @@ public:
 
 ////////////////////////////////////////////////////////////////////////////////////////////////////
 
-class FContextSwitchesSharedState : public Insights::ITimingViewExtender, public TSharedFromThis<FContextSwitchesSharedState>
+class FContextSwitchesSharedState : public UE::Insights::Timing::ITimingViewExtender, public TSharedFromThis<FContextSwitchesSharedState>
 {
 public:
-	FContextSwitchesSharedState(STimingView* InTimingView);
+	FContextSwitchesSharedState(UE::Insights::TimingProfiler::STimingView* InTimingView);
 	virtual ~FContextSwitchesSharedState() = default;
 
 	//////////////////////////////////////////////////
 	// ITimingViewExtender interface
 
-	virtual void OnBeginSession(Insights::ITimingViewSession& InSession) override;
-	virtual void OnEndSession(Insights::ITimingViewSession& InSession) override;
-	virtual void Tick(Insights::ITimingViewSession& InSession, const TraceServices::IAnalysisSession& InAnalysisSession) override;
-	virtual void ExtendCpuTracksFilterMenu(Insights::ITimingViewSession& InSession, FMenuBuilder& InMenuBuilder) override;
-	virtual void AddQuickFindFilters(TSharedPtr<class FFilterConfigurator> FilterConfigurator) override;
+	virtual void OnBeginSession(UE::Insights::Timing::ITimingViewSession& InSession) override;
+	virtual void OnEndSession(UE::Insights::Timing::ITimingViewSession& InSession) override;
+	virtual void Tick(UE::Insights::Timing::ITimingViewSession& InSession, const TraceServices::IAnalysisSession& InAnalysisSession) override;
+	virtual void ExtendCpuTracksFilterMenu(UE::Insights::Timing::ITimingViewSession& InSession, FMenuBuilder& InMenuBuilder) override;
+	virtual void AddQuickFindFilters(TSharedPtr<UE::Insights::FFilterConfigurator> FilterConfigurator) override;
 
 	//////////////////////////////////////////////////
 
-	static TSharedPtr<STimingView> GetTimingView();
+	static TSharedPtr<UE::Insights::TimingProfiler::STimingView> GetTimingView();
 
 	bool AreContextSwitchesAvailable() const;
 
@@ -115,7 +113,7 @@ public:
 
 	void GetThreadInfo(uint32 InSystemThreadId, uint32& OutThreadId, const TCHAR*& OutThreadName) const;
 	TSharedPtr<FThreadTimingTrack> GetThreadTimingTrack(uint32 ThreadId) const;
-	TSharedPtr<Insights::FCpuCoreTimingTrack> GetCpuCoreTimingTrack(uint32 CoreNumber) const;
+	TSharedPtr<FCpuCoreTimingTrack> GetCpuCoreTimingTrack(uint32 CoreNumber) const;
 
 private:
 	void AddCoreTracks();
@@ -167,7 +165,7 @@ private:
 	TMap<uint32, TSharedPtr<FCpuCoreTimingTrack>> CpuCoreTimingTracks;
 	TMap<uint32, TSharedPtr<FContextSwitchesTimingTrack>> ContextSwitchesTimingTracks;
 
-	Insights::ITimingViewSession* TimingViewSession;
+	UE::Insights::Timing::ITimingViewSession* TimingViewSession;
 
 	uint64 ThreadsSerial;
 	uint64 CpuCoresSerial;
diff --git a/Engine/Source/Developer/TraceInsights/Private/Insights/ContextSwitches/ViewModels/ContextSwitchesTimingTrack.cpp b/Engine/Source/Developer/TraceInsights/Private/Insights/ContextSwitches/ViewModels/ContextSwitchesTimingTrack.cpp
--- a/Engine/Source/Developer/TraceInsights/Private/Insights/ContextSwitches/ViewModels/ContextSwitchesTimingTrack.cpp
+++ b/Engine/Source/Developer/TraceInsights/Private/Insights/ContextSwitches/ViewModels/ContextSwitchesTimingTrack.cpp
@@ -3,10 +3,14 @@
 #include "ContextSwitchesTimingTrack.h"
 
 #include "Framework/MultiBox/MultiBoxBuilder.h"
+
+// TraceServices
 #include "TraceServices/Model/ContextSwitches.h"
 
-// Insights
-#include "Insights/Common/TimeUtils.h"
+// TraceInsightsCore
+#include "InsightsCore/Common/TimeUtils.h"
+
+// TraceInsights
 #include "Insights/ContextSwitches/ViewModels/ContextSwitchesSharedState.h"
 #include "Insights/ContextSwitches/ViewModels/ContextSwitchTimingEvent.h"
 #include "Insights/InsightsManager.h"
@@ -242,9 +246,10 @@ void FContextSwitchesTimingTrack::InitTooltip(FTooltipDrawState& InOutTooltip, c
 	const FContextSwitchTimingEvent& ContextSwitchEvent = InTooltipEvent.As<FContextSwitchTimingEvent>();
 	InOutTooltip.AddTitle(FString::Printf(TEXT("Core %d"), ContextSwitchEvent.GetCoreNumber()));
 
-	InOutTooltip.AddNameValueTextLine(TEXT("Start Time:"), TimeUtils::FormatTimeAuto(InTooltipEvent.GetStartTime(), 6));
-	InOutTooltip.AddNameValueTextLine(TEXT("End Time:"), TimeUtils::FormatTimeAuto(InTooltipEvent.GetEndTime(), 6));
-	InOutTooltip.AddNameValueTextLine(TEXT("Duration:"), TimeUtils::FormatTimeAuto(InTooltipEvent.GetDuration()));
+	using namespace UE::Insights;
+	InOutTooltip.AddNameValueTextLine(TEXT("Start Time:"), FormatTimeAuto(InTooltipEvent.GetStartTime(), 6));
+	InOutTooltip.AddNameValueTextLine(TEXT("End Time:"), FormatTimeAuto(InTooltipEvent.GetEndTime(), 6));
+	InOutTooltip.AddNameValueTextLine(TEXT("Duration:"), FormatTimeAuto(InTooltipEvent.GetDuration()));
 
 	InOutTooltip.UpdateLayout();
 }
diff --git a/Engine/Source/Developer/TraceInsights/Private/Insights/ContextSwitches/ViewModels/ContextSwitchesTimingTrack.h b/Engine/Source/Developer/TraceInsights/Private/Insights/ContextSwitches/ViewModels/ContextSwitchesTimingTrack.h
--- a/Engine/Source/Developer/TraceInsights/Private/Insights/ContextSwitches/ViewModels/ContextSwitchesTimingTrack.h
+++ b/Engine/Source/Developer/TraceInsights/Private/Insights/ContextSwitches/ViewModels/ContextSwitchesTimingTrack.h
@@ -4,7 +4,7 @@
 
 #include "CoreMinimal.h"
 
-// Insights
+// TraceInsights
 #include "Insights/ViewModels/TimingEventsTrack.h"
 
 ////////////////////////////////////////////////////////////////////////////////////////////////////
diff --git a/Engine/Source/Developer/TraceInsights/Private/Insights/ContextSwitches/ViewModels/CpuCoreTimingTrack.cpp b/Engine/Source/Developer/TraceInsights/Private/Insights/ContextSwitches/ViewModels/CpuCoreTimingTrack.cpp
--- a/Engine/Source/Developer/TraceInsights/Private/Insights/ContextSwitches/ViewModels/CpuCoreTimingTrack.cpp
+++ b/Engine/Source/Developer/TraceInsights/Private/Insights/ContextSwitches/ViewModels/CpuCoreTimingTrack.cpp
@@ -3,16 +3,20 @@
 #include "CpuCoreTimingTrack.h"
 
 #include "Framework/MultiBox/MultiBoxBuilder.h"
+
+// TraceServices
 #include "TraceServices/Model/ContextSwitches.h"
 #include "TraceServices/Model/Threads.h"
 
-// Insights
-#include "Insights/Common/TimeUtils.h"
+// TraceInsightsCore
+#include "InsightsCore/Common/TimeUtils.h"
+#include "InsightsCore/Filter/ViewModels/FilterConfigurator.h"
+#include "InsightsCore/Filter/ViewModels/Filters.h"
+
+// TraceInsights
 #include "Insights/ContextSwitches/ViewModels/ContextSwitchesSharedState.h"
 #include "Insights/ContextSwitches/ViewModels/ContextSwitchTimingEvent.h"
 #include "Insights/InsightsManager.h"
-#include "Insights/ViewModels/FilterConfigurator.h"
-#include "Insights/ViewModels/Filters.h"
 #include "Insights/ViewModels/TimingTrackViewport.h"
 #include "Insights/ViewModels/TimingViewDrawHelper.h"
 #include "Insights/ViewModels/TooltipDrawState.h"
@@ -143,19 +147,20 @@ void FCpuCoreTimingTrack::BuildFilteredDrawState(ITimingEventsTrackDrawStateBuil
 
 	if (HasCustomFilter())
 	{
-		TSharedPtr<STimingView> TimingView = SharedState.GetTimingView();
+		TSharedPtr<UE::Insights::TimingProfiler::STimingView> TimingView = SharedState.GetTimingView();
 		if (!TimingView)
 		{
 			return;
 		}
 
-		TSharedPtr<FFilterConfigurator> FilterConfigurator = TimingView->GetFilterConfigurator();
+		TSharedPtr<UE::Insights::FFilterConfigurator> FilterConfigurator = TimingView->GetFilterConfigurator();
 		if (!FilterConfigurator.IsValid())
 		{
 			return;
 		}
 
-		FFilterContext FilterContext;
+		using EFilterField = UE::Insights::EFilterField;
+		UE::Insights::FFilterContext FilterContext;
 		FilterContext.SetReturnValueForUnsetFilters(false);
 
 		FilterContext.AddFilterData<double>(static_cast<int32>(EFilterField::StartTime), 0.0f);
@@ -251,7 +256,7 @@ void FCpuCoreTimingTrack::PostDraw(const ITimingTrackDrawContext& Context) const
 
 		FString Str = FString::Printf(TEXT("%s (Duration.: %s)"),
 			*GetThreadName(static_cast<uint32>(SelectedEvent.GetType())),
-			*TimeUtils::FormatTimeAuto(SelectedEvent.GetDuration()));
+			*UE::Insights::FormatTimeAuto(SelectedEvent.GetDuration()));
 
 		DrawSelectedEventInfo(Str, Context.GetViewport(), Context.GetDrawContext(), Helper.GetWhiteBrush(), Helper.GetEventFont());
 	}
@@ -333,7 +338,8 @@ const TSharedPtr<const ITimingEvent> FCpuCoreTimingTrack::SearchEvent(const FTim
 {
 	TSharedPtr<const ITimingEvent> FoundEvent;
 
-	FFilterContext FilterConfiguratorContext;
+	using EFilterField = UE::Insights::EFilterField;
+	UE::Insights::FFilterContext FilterConfiguratorContext;
 	FilterConfiguratorContext.SetReturnValueForUnsetFilters(false);
 	FilterConfiguratorContext.AddFilterData<double>(static_cast<int32>(EFilterField::StartTime), 0.0f);
 	FilterConfiguratorContext.AddFilterData<double>(static_cast<int32>(EFilterField::EndTime), 0.0f);
@@ -437,9 +443,10 @@ void FCpuCoreTimingTrack::InitTooltip(FTooltipDrawState& InOutTooltip, const ITi
 		InOutTooltip.AddNameValueTextLine(TEXT("Thread Id:"), FString::Printf(TEXT("%d"), ThreadId));
 	}
 
-	InOutTooltip.AddNameValueTextLine(TEXT("Start Time:"), TimeUtils::FormatTimeAuto(InTooltipEvent.GetStartTime(), 6));
-	InOutTooltip.AddNameValueTextLine(TEXT("End Time:"), TimeUtils::FormatTimeAuto(InTooltipEvent.GetEndTime(), 6));
-	InOutTooltip.AddNameValueTextLine(TEXT("Duration:"), TimeUtils::FormatTimeAuto(InTooltipEvent.GetDuration()));
+	using namespace UE::Insights;
+	InOutTooltip.AddNameValueTextLine(TEXT("Start Time:"), FormatTimeAuto(InTooltipEvent.GetStartTime(), 6));
+	InOutTooltip.AddNameValueTextLine(TEXT("End Time:"), FormatTimeAuto(InTooltipEvent.GetEndTime(), 6));
+	InOutTooltip.AddNameValueTextLine(TEXT("Duration:"), FormatTimeAuto(InTooltipEvent.GetDuration()));
 
 	InOutTooltip.UpdateLayout();
 }
@@ -482,13 +489,13 @@ void FCpuCoreTimingTrack::BuildContextMenu(FMenuBuilder& InOutMenuBuilder)
 
 bool FCpuCoreTimingTrack::HasCustomFilter() const
 {
-	TSharedPtr<STimingView> TimingView = SharedState.GetTimingView();
+	TSharedPtr<UE::Insights::TimingProfiler::STimingView> TimingView = SharedState.GetTimingView();
 	if (!TimingView)
 	{
 		return false;
 	}
 
-	TSharedPtr<FFilterConfigurator> FilterConfigurator = TimingView->GetFilterConfigurator();
+	TSharedPtr<UE::Insights::FFilterConfigurator> FilterConfigurator = TimingView->GetFilterConfigurator();
 	return FilterConfigurator.IsValid() && !FilterConfigurator->IsEmpty();
 }
 
diff --git a/Engine/Source/Developer/TraceInsights/Private/Insights/CookProfiler/CookProfilerManager.cpp b/Engine/Source/Developer/TraceInsights/Private/Insights/CookProfiler/CookProfilerManager.cpp
--- a/Engine/Source/Developer/TraceInsights/Private/Insights/CookProfiler/CookProfilerManager.cpp
+++ b/Engine/Source/Developer/TraceInsights/Private/Insights/CookProfiler/CookProfilerManager.cpp
@@ -2,14 +2,16 @@
 
 #include "CookProfilerManager.h"
 
-#include "Common/ProviderLock.h" // TraceServices
 #include "Features/IModularFeatures.h"
 #include "Framework/Docking/TabManager.h"
 #include "Modules/ModuleManager.h"
-#include "TraceServices/Model/CookProfilerProvider.h"
 #include "Widgets/Docking/SDockTab.h"
 
-// Insights
+// TraceServices
+#include "Common/ProviderLock.h"
+#include "TraceServices/Model/CookProfilerProvider.h"
+
+// TraceInsights
 #include "Insights/InsightsStyle.h"
 #include "Insights/CookProfiler/ViewModels/PackageTable.h"
 #include "Insights/CookProfiler/Widgets/SPackageTableTreeView.h"
diff --git a/Engine/Source/Developer/TraceInsights/Private/Insights/CookProfiler/CookProfilerManager.h b/Engine/Source/Developer/TraceInsights/Private/Insights/CookProfiler/CookProfilerManager.h
--- a/Engine/Source/Developer/TraceInsights/Private/Insights/CookProfiler/CookProfilerManager.h
+++ b/Engine/Source/Developer/TraceInsights/Private/Insights/CookProfiler/CookProfilerManager.h
@@ -2,11 +2,12 @@
 
 #pragma once
 
-#include "Containers/Ticker.h"
 #include "CoreMinimal.h"
+
+#include "Containers/Ticker.h"
 #include "Framework/Commands/UICommandList.h"
 
-// Insights
+// TraceInsights
 #include "Insights/InsightsManager.h"
 #include "Insights/IUnrealInsightsModule.h"
 
@@ -76,7 +77,7 @@ private:
 private:
 	bool bIsInitialized;
 	bool bIsAvailable;
-	FAvailabilityCheck AvailabilityCheck;
+	UE::Insights::FAvailabilityCheck AvailabilityCheck;
 
 	/** The delegate to be invoked when this manager ticks. */
 	FTickerDelegate OnTick;
diff --git a/Engine/Source/Developer/TraceInsights/Private/Insights/CookProfiler/ViewModels/PackageEntry.h b/Engine/Source/Developer/TraceInsights/Private/Insights/CookProfiler/ViewModels/PackageEntry.h
--- a/Engine/Source/Developer/TraceInsights/Private/Insights/CookProfiler/ViewModels/PackageEntry.h
+++ b/Engine/Source/Developer/TraceInsights/Private/Insights/CookProfiler/ViewModels/PackageEntry.h
@@ -4,6 +4,7 @@
 
 #include "CoreMinimal.h"
 
+// TraceServices
 #include "TraceServices/Model/CookProfilerProvider.h"
 
 namespace Insights
diff --git a/Engine/Source/Developer/TraceInsights/Private/Insights/CookProfiler/ViewModels/PackageNode.h b/Engine/Source/Developer/TraceInsights/Private/Insights/CookProfiler/ViewModels/PackageNode.h
--- a/Engine/Source/Developer/TraceInsights/Private/Insights/CookProfiler/ViewModels/PackageNode.h
+++ b/Engine/Source/Developer/TraceInsights/Private/Insights/CookProfiler/ViewModels/PackageNode.h
@@ -4,28 +4,18 @@
 
 #include "CoreMinimal.h"
 
-// Insights
+// TraceInsightsCore
+#include "InsightsCore/Table/ViewModels/TableTreeNode.h"
+
+// TraceInsights
 #include "Insights/CookProfiler/ViewModels/PackageTable.h"
 #include "Insights/CookProfiler/ViewModels/PackageEntry.h"
-#include "Insights/Table/ViewModels/TableTreeNode.h"
 
 ////////////////////////////////////////////////////////////////////////////////////////////////////
 
 namespace Insights
 {
 
-enum class EPackageNodeType
-{
-	/** The TaskNode is an allocation node. */
-	Package,
-
-	/** The TaskNode is a group node. */
-	Group,
-
-	/** Invalid enum type, may be used as a number of enumerations. */
-	InvalidOrMax,
-};
-
 ////////////////////////////////////////////////////////////////////////////////////////////////////
 
 class FPackageNode;
@@ -46,33 +36,26 @@ typedef TWeakPtr<class FTaskNode> FPackageNodeWeak;
 /**
  * Class used to store information about a package node (used in the SPackageTableTreeView).
  */
-class FPackageNode : public FTableTreeNode
+class FPackageNode : public UE::Insights::FTableTreeNode
 {
-	INSIGHTS_DECLARE_RTTI(FPackageNode, FTableTreeNode)
+	INSIGHTS_DECLARE_RTTI(FPackageNode, UE::Insights::FTableTreeNode)
 
 public:
 	/** Initialization constructor for the Task node. */
 	explicit FPackageNode(const FName InName, TWeakPtr<FPackageTable> InParentTable, int32 InRowIndex)
-		: FTableTreeNode(InName, InParentTable, InRowIndex)
-		, Type(EPackageNodeType::Package)
+		: UE::Insights::FTableTreeNode(InName, InParentTable, InRowIndex)
 	{
 	}
 
 	/** Initialization constructor for the group node. */
 	explicit FPackageNode(const FName InGroupName, TWeakPtr<FPackageTable> InParentTable)
-		: FTableTreeNode(InGroupName, InParentTable)
-		, Type(EPackageNodeType::Group)
+		: UE::Insights::FTableTreeNode(InGroupName, InParentTable)
 	{
 	}
 
-	/**
-	 * @return a type of this Task node or ETaskNodeType::Group for group nodes.
-	 */
-	EPackageNodeType GetType() const { return Type; }
-
 	FPackageTable& GetPackageTableChecked() const
 	{
-		const TSharedPtr<FTable>& TablePin = GetParentTable().Pin();
+		const TSharedPtr<UE::Insights::FTable>& TablePin = GetParentTable().Pin();
 		check(TablePin.IsValid());
 		return *StaticCastSharedPtr<FPackageTable>(TablePin);
 	}
@@ -80,9 +63,6 @@ public:
 	bool IsValidPackage() const { return GetPackageTableChecked().IsValidRowIndex(GetRowIndex()); }
 	const FPackageEntry* GetPackage() const { return GetPackageTableChecked().GetPackage(GetRowIndex()); }
 	const FPackageEntry& GetPackageChecked() const { return GetPackageTableChecked().GetPackageChecked(GetRowIndex()); }
-
-private:
-	const EPackageNodeType Type;
 };
 
 ////////////////////////////////////////////////////////////////////////////////////////////////////
diff --git a/Engine/Source/Developer/TraceInsights/Private/Insights/CookProfiler/ViewModels/PackageTable.cpp b/Engine/Source/Developer/TraceInsights/Private/Insights/CookProfiler/ViewModels/PackageTable.cpp
--- a/Engine/Source/Developer/TraceInsights/Private/Insights/CookProfiler/ViewModels/PackageTable.cpp
+++ b/Engine/Source/Developer/TraceInsights/Private/Insights/CookProfiler/ViewModels/PackageTable.cpp
@@ -2,13 +2,15 @@
 
 #include "PackageTable.h"
 
-// Insights
+// TraceInsightsCore
+#include "InsightsCore/Table/ViewModels/TableCellValueFormatter.h"
+#include "InsightsCore/Table/ViewModels/TableCellValueGetter.h"
+#include "InsightsCore/Table/ViewModels/TableCellValueSorter.h"
+#include "InsightsCore/Table/ViewModels/TableColumn.h"
+
+// TraceInsights
 #include "Insights/CookProfiler/ViewModels/PackageNode.h"
 #include "Insights/CookProfiler/ViewModels/PackageTable.h"
-#include "Insights/Table/ViewModels/TableCellValueFormatter.h"
-#include "Insights/Table/ViewModels/TableCellValueGetter.h"
-#include "Insights/Table/ViewModels/TableCellValueSorter.h"
-#include "Insights/Table/ViewModels/TableColumn.h"
 
 #define LOCTEXT_NAMESPACE "Insights::FPackageTable"
 
@@ -36,17 +38,18 @@ const FName FPackageTableColumns::GetIsCachedCookedPlatformDataLoadedExclColumnI
 const FName FPackageTableColumns::PackageAssetClassColumnId(TEXT("AssetClass"));
 
 ////////////////////////////////////////////////////////////////////////////////////////////////////
-typedef FTableCellValue(*PackageFieldGetter)(const FTableColumn&, const FPackageEntry&);
+
+typedef UE::Insights::FTableCellValue (*PackageFieldGetter) (const UE::Insights::FTableColumn&, const FPackageEntry&);
 
 template<PackageFieldGetter Getter>
-class FPackageColumnValueGetter : public FTableCellValueGetter
+class FPackageColumnValueGetter : public UE::Insights::FTableCellValueGetter
 {
 public:
-	virtual const TOptional<FTableCellValue> GetValue(const FTableColumn& Column, const FBaseTreeNode& Node) const override
+	virtual const TOptional<UE::Insights::FTableCellValue> GetValue(const UE::Insights::FTableColumn& Column, const UE::Insights::FBaseTreeNode& Node) const override
 	{
 		if (Node.IsGroup())
 		{
-			const FTableTreeNode& NodePtr = static_cast<const FTableTreeNode&>(Node);
+			const UE::Insights::FTableTreeNode& NodePtr = static_cast<const UE::Insights::FTableTreeNode&>(Node);
 			if (NodePtr.HasAggregatedValue(Column.GetId()))
 			{
 				return NodePtr.GetAggregatedValue(Column.GetId());
@@ -62,7 +65,7 @@ public:
 			}
 		}
 
-		return TOptional<FTableCellValue>();
+		return TOptional<UE::Insights::FTableCellValue>();
 	}
 };
 
@@ -70,6 +73,9 @@ public:
 
 struct DefaultPackageFieldGetterFuncts
 {
+	using FTableCellValue = UE::Insights::FTableCellValue;
+	using FTableColumn = UE::Insights::FTableColumn;
+
 	static FTableCellValue GetId(const FTableColumn& Column, const FPackageEntry& Package) { return FTableCellValue((int64)Package.GetId());	}
 	static FTableCellValue GetName(const FTableColumn& Column, const FPackageEntry& Package) { return FTableCellValue((const TCHAR*)Package.GetName());	}
 
@@ -115,6 +121,8 @@ void FPackageTable::Reset()
 
 void FPackageTable::AddDefaultColumns()
 {
+	using namespace UE::Insights;
+
 	//////////////////////////////////////////////////
 	// Hierarchy Column
 	{
diff --git a/Engine/Source/Developer/TraceInsights/Private/Insights/CookProfiler/ViewModels/PackageTable.h b/Engine/Source/Developer/TraceInsights/Private/Insights/CookProfiler/ViewModels/PackageTable.h
--- a/Engine/Source/Developer/TraceInsights/Private/Insights/CookProfiler/ViewModels/PackageTable.h
+++ b/Engine/Source/Developer/TraceInsights/Private/Insights/CookProfiler/ViewModels/PackageTable.h
@@ -2,17 +2,22 @@
 
 #pragma once
 
-#include "CoreMinimal.h"
+#include "CoreTypes.h"
 
-// Insights
+#include "Containers/Array.h"
+#include "UObject/NameTypes.h"
+
+// TraceInsightsCore
+#include "InsightsCore/Table/ViewModels/Table.h"
+
+// TraceInsights
 #include "Insights/CookProfiler/ViewModels/PackageEntry.h"
-#include "Insights/Table/ViewModels/Table.h"
+
+namespace UE::Insights { class FTableColumn; }
 
 namespace Insights
 {
 
-class FTableColumn;
-
 ////////////////////////////////////////////////////////////////////////////////////////////////////
 
 // Column identifiers
@@ -38,7 +43,7 @@ struct FPackageTableColumns
 
 ////////////////////////////////////////////////////////////////////////////////////////////////////
 
-class FPackageTable : public FTable
+class FPackageTable : public UE::Insights::FTable
 {
 public:
 	FPackageTable();
diff --git a/Engine/Source/Developer/TraceInsights/Private/Insights/CookProfiler/Widgets/SPackageTableTreeView.cpp b/Engine/Source/Developer/TraceInsights/Private/Insights/CookProfiler/Widgets/SPackageTableTreeView.cpp
--- a/Engine/Source/Developer/TraceInsights/Private/Insights/CookProfiler/Widgets/SPackageTableTreeView.cpp
+++ b/Engine/Source/Developer/TraceInsights/Private/Insights/CookProfiler/Widgets/SPackageTableTreeView.cpp
@@ -2,20 +2,23 @@
 
 #include "SPackageTableTreeView.h"
 
-#include "Common/ProviderLock.h" // TraceServices
 #include "Framework/MultiBox/MultiBoxBuilder.h"
 #include "SlateOptMacros.h"
 #include "Widgets/Input/STextComboBox.h"
 
-// Insights
+// TraceServices
+#include "Common/ProviderLock.h"
+
+// TraceInsightsCore
+#include "InsightsCore/Table/ViewModels/TableColumn.h"
+#include "InsightsCore/Table/ViewModels/TreeNodeGrouping.h"
+
+// TraceInsights
 #include "Insights/InsightsStyle.h"
 #include "Insights/CookProfiler/CookProfilerManager.h"
 #include "Insights/CookProfiler/ViewModels/PackageEntry.h"
 #include "Insights/CookProfiler/ViewModels/PackageNode.h"
-#include "Insights/Table/ViewModels/TableColumn.h"
-#include "Insights/Table/ViewModels/TreeNodeGrouping.h"
 #include "Insights/TimingProfilerManager.h"
-#include "Insights/ViewModels/FilterConfigurator.h"
 #include "Insights/Widgets/STimingProfilerWindow.h"
 #include "Insights/Widgets/STimingView.h"
 
@@ -140,7 +143,7 @@ void SPackageTableTreeView::RebuildTree(bool bResync)
 		Packages.Reserve(NumPackages);
 		TableRowNodes.Reserve(NumPackages);
 
-		TArray<FTableTreeNodePtr>* Nodes = &TableRowNodes;
+		TArray<UE::Insights::FTableTreeNodePtr>* Nodes = &TableRowNodes;
 		for(const TraceServices::FPackageData& Package : PackageAggreagation)
 		{
 			Packages.Emplace(Package);
@@ -201,7 +204,7 @@ TSharedPtr<SWidget> SPackageTableTreeView::ConstructToolbar()
 			SNew(SBox)
 			.MinDesiredWidth(150.0f)
 			[
-				SAssignNew(PresetComboBox, SComboBox<TSharedRef<ITableTreeViewPreset>>)
+				SAssignNew(PresetComboBox, SComboBox<TSharedRef<UE::Insights::ITableTreeViewPreset>>)
 				.ToolTipText(this, &SPackageTableTreeView::ViewPreset_GetSelectedToolTipText)
 				.OptionsSource(GetAvailableViewPresets())
 				.OnSelectionChanged(this, &SPackageTableTreeView::ViewPreset_OnSelectionChanged)
@@ -232,11 +235,11 @@ void SPackageTableTreeView::InternalCreateGroupings()
 	STableTreeView::InternalCreateGroupings();
 
 	AvailableGroupings.RemoveAll(
-		[](TSharedPtr<FTreeNodeGrouping>& Grouping)
+		[](TSharedPtr<UE::Insights::FTreeNodeGrouping>& Grouping)
 		{
-			if (Grouping->Is<FTreeNodeGroupingByUniqueValue>())
+			if (Grouping->Is<UE::Insights::FTreeNodeGroupingByUniqueValue>())
 			{
-				const FName ColumnId = Grouping->As<FTreeNodeGroupingByUniqueValue>().GetColumnId();
+				const FName ColumnId = Grouping->As<UE::Insights::FTreeNodeGroupingByUniqueValue>().GetColumnId();
 				if (ColumnId == FPackageTableColumns::BeginCacheForCookedPlatformDataTimeInclColumnId ||
 					ColumnId == FPackageTableColumns::BeginCacheForCookedPlatformDataTimeExclColumnId ||
 					ColumnId == FPackageTableColumns::GetIsCachedCookedPlatformDataLoadedInclColumnId ||
@@ -251,9 +254,9 @@ void SPackageTableTreeView::InternalCreateGroupings()
 					return true;
 				}
 			}
-			else if (Grouping->Is<FTreeNodeGroupingByPathBreakdown>())
+			else if (Grouping->Is<UE::Insights::FTreeNodeGroupingByPathBreakdown>())
 			{
-				const FName ColumnId = Grouping->As<FTreeNodeGroupingByPathBreakdown>().GetColumnId();
+				const FName ColumnId = Grouping->As<UE::Insights::FTreeNodeGroupingByPathBreakdown>().GetColumnId();
 				if (ColumnId == FPackageTableColumns::PackageAssetClassColumnId)
 				{
 					return true;
@@ -265,7 +268,7 @@ void SPackageTableTreeView::InternalCreateGroupings()
 
 ////////////////////////////////////////////////////////////////////////////////////////////////////
 
-void SPackageTableTreeView::TreeView_OnMouseButtonDoubleClick(FTableTreeNodePtr TreeNode)
+void SPackageTableTreeView::TreeView_OnMouseButtonDoubleClick(UE::Insights::FTableTreeNodePtr TreeNode)
 {
 	STableTreeView::TreeView_OnMouseButtonDoubleClick(TreeNode);
 }
@@ -274,6 +277,8 @@ void SPackageTableTreeView::TreeView_OnMouseButtonDoubleClick(FTableTreeNodePtr
 
 void SPackageTableTreeView::InitAvailableViewPresets()
 {
+	using namespace UE::Insights;
+
 	//////////////////////////////////////////////////
 	// Default View
 
diff --git a/Engine/Source/Developer/TraceInsights/Private/Insights/CookProfiler/Widgets/SPackageTableTreeView.h b/Engine/Source/Developer/TraceInsights/Private/Insights/CookProfiler/Widgets/SPackageTableTreeView.h
--- a/Engine/Source/Developer/TraceInsights/Private/Insights/CookProfiler/Widgets/SPackageTableTreeView.h
+++ b/Engine/Source/Developer/TraceInsights/Private/Insights/CookProfiler/Widgets/SPackageTableTreeView.h
@@ -3,9 +3,11 @@
 #pragma once
 
 #include "CoreMinimal.h"
+
+// TraceServices
 #include "TraceServices/Model/CookProfilerProvider.h"
 
-// Insights
+// TraceInsights
 #include "Insights/CookProfiler/ViewModels/PackageTable.h"
 #include "Insights/Table/Widgets/SSessionTableTreeView.h"
 
@@ -16,7 +18,7 @@ namespace Insights
 
 ////////////////////////////////////////////////////////////////////////////////////////////////////
 
-class SPackageTableTreeView : public SSessionTableTreeView
+class SPackageTableTreeView : public UE::Insights::SSessionTableTreeView
 {
 public:
 	/** Default constructor. */
@@ -66,7 +68,7 @@ public:
 
 	////////////////////////////////////////////////////////////////////////////////////////////////////
 
-	virtual void TreeView_OnMouseButtonDoubleClick(FTableTreeNodePtr TreeNode) override;
+	virtual void TreeView_OnMouseButtonDoubleClick(UE::Insights::FTableTreeNodePtr TreeNode) override;
 
 protected:
 	virtual void InternalCreateGroupings() override;
diff --git a/Engine/Source/Developer/TraceInsights/Private/Insights/ImportTool/TableImportTool.cpp b/Engine/Source/Developer/TraceInsights/Private/Insights/ImportTool/TableImportTool.cpp
--- a/Engine/Source/Developer/TraceInsights/Private/Insights/ImportTool/TableImportTool.cpp
+++ b/Engine/Source/Developer/TraceInsights/Private/Insights/ImportTool/TableImportTool.cpp
@@ -10,11 +10,15 @@
 #include "MessageLogModule.h"
 #include "Misc/FileHelper.h"
 #include "SlateOptMacros.h"
-#include "TraceServices/Model/Threads.h"
 #include "Widgets/Docking/SDockTab.h"
 
-// Insights
-#include "Insights/Common/Stopwatch.h"
+// TraceServices
+#include "TraceServices/Model/Threads.h"
+
+// TraceInsightsCore
+#include "InsightsCore/Common/Stopwatch.h"
+
+// TraceInsights
 #include "Insights/InsightsManager.h"
 #include "Insights/InsightsStyle.h"
 #include "Insights/Log.h"
diff --git a/Engine/Source/Developer/TraceInsights/Private/Insights/ImportTool/TableImportTool.h b/Engine/Source/Developer/TraceInsights/Private/Insights/ImportTool/TableImportTool.h
--- a/Engine/Source/Developer/TraceInsights/Private/Insights/ImportTool/TableImportTool.h
+++ b/Engine/Source/Developer/TraceInsights/Private/Insights/ImportTool/TableImportTool.h
@@ -4,9 +4,11 @@
 
 #include "CoreMinimal.h"
 
+// TraceServices
 #include "TraceServices/Model/TableImport.h"
 #include "TraceServices/Containers/Tables.h"
 
+// TraceInsights
 #include "Insights/IUnrealInsightsModule.h"
 
 class SDockTab;
diff --git a/Engine/Source/Developer/TraceInsights/Private/Insights/InsightsManager.cpp b/Engine/Source/Developer/TraceInsights/Private/Insights/InsightsManager.cpp
--- a/Engine/Source/Developer/TraceInsights/Private/Insights/InsightsManager.cpp
+++ b/Engine/Source/Developer/TraceInsights/Private/Insights/InsightsManager.cpp
@@ -15,16 +15,23 @@
 #include "Misc/Paths.h"
 #include "Modules/ModuleManager.h"
 #include "Templates/UniquePtr.h"
+#include "WorkspaceMenuStructure.h"
+#include "WorkspaceMenuStructureModule.h"
+
+// TraceAnalysis
 #include "Trace/StoreClient.h"
+
+// TraceServices
 #include "TraceServices/AnalysisService.h"
 #include "TraceServices/Model/Diagnostics.h"
 #include "TraceServices/Model/NetProfiler.h"
-#include "WorkspaceMenuStructure.h"
-#include "WorkspaceMenuStructureModule.h"
 
-// Insights
+// TraceInsightsCore
+#include "InsightsCore/Common/TimeUtils.h"
+#include "InsightsCore/Filter/ViewModels/Filters.h"
+
+// TraceInsights
 #include "Insights/Common/InsightsMenuBuilder.h"
-#include "Insights/Common/TimeUtils.h"
 #include "Insights/InsightsStyle.h"
 #include "Insights/IUnrealInsightsModule.h"
 #include "Insights/LoadingProfiler/LoadingProfilerManager.h"
@@ -32,9 +39,8 @@
 #include "Insights/NetworkingProfiler/NetworkingProfilerManager.h"
 #include "Insights/Tests/InsightsTestRunner.h"
 #include "Insights/TimingProfilerManager.h"
-#include "Insights/ViewModels/Filters.h"
-#include "Insights/Widgets/SStartPageWindow.h"
 #include "Insights/Widgets/SSessionInfoWindow.h"
+#include "Insights/Widgets/SStartPageWindow.h"
 #include "Insights/Widgets/STimingProfilerWindow.h"
 
 ////////////////////////////////////////////////////////////////////////////////////////////////////
@@ -124,7 +130,7 @@ void FInsightsManager::Initialize(IUnrealInsightsModule& InsightsModule)
 
 	InsightsMenuBuilder = MakeShared<FInsightsMenuBuilder>();
 
-	Insights::FFilterService::Initialize();
+	UE::Insights::FFilterService::Initialize();
 
 	FMessageLogModule& MessageLogModule = FModuleManager::LoadModuleChecked<FMessageLogModule>("MessageLog");
 	MessageLogModule.RegisterLogListing(GetLogListingName(), LOCTEXT("UnrealInsights", "Unreal Insights"));
@@ -166,7 +172,7 @@ void FInsightsManager::Shutdown()
 		}
 	}
 
-	Insights::FFilterService::Shutdown();
+	UE::Insights::FFilterService::Shutdown();
 
 	FInsightsManager::Instance.Reset();
 }
@@ -541,6 +547,8 @@ bool FInsightsManager::Tick(float DeltaTime)
 
 void FInsightsManager::UpdateSessionDuration()
 {
+	using namespace UE::Insights;
+
 	if (Session.IsValid())
 	{
 		bool bLocalIsAnalysisComplete = false;
@@ -552,7 +560,6 @@ void FInsightsManager::UpdateSessionDuration()
 
 		}
 
-
 		if (LocalSessionDuration != SessionDuration)
 		{
 			SessionDuration = LocalSessionDuration;
@@ -562,7 +569,7 @@ void FInsightsManager::UpdateSessionDuration()
 			if (bIsAnalysisComplete)
 			{
 				UE_LOG(TraceInsights, Warning, TEXT("The session duration was updated (%s) after the analysis has been completed."),
-					*TimeUtils::FormatTimeAuto(GetSessionDuration(), 2));
+					*FormatTimeAuto(GetSessionDuration(), 2));
 			}
 		}
 
@@ -572,9 +579,9 @@ void FInsightsManager::UpdateSessionDuration()
 			SessionAnalysisCompletedEvent.Broadcast();
 
 			UE_LOG(TraceInsights, Log, TEXT("Analysis has completed in %s (%.1fX speed; session duration: %s)."),
-				*TimeUtils::FormatTimeAuto(AnalysisDuration, 2),
+				*FormatTimeAuto(AnalysisDuration, 2),
 				AnalysisSpeedFactor,
-				*TimeUtils::FormatTimeAuto(SessionDuration, 2));
+				*FormatTimeAuto(SessionDuration, 2));
 
 			OnSessionAnalysisCompleted();
 		}
@@ -745,6 +752,8 @@ void FInsightsManager::ActivateTimingInsightsTab()
 	{
 		TimingInsightsTab->ActivateInParent(ETabActivationCause::SetDirectly);
 
+		using namespace UE::Insights::TimingProfiler;
+
 		//TODO: FTimingProfilerManager::Get()->ActivateWindow();
 		TSharedPtr<class STimingProfilerWindow> Wnd = FTimingProfilerManager::Get()->GetProfilerWindow();
 		if (Wnd)
@@ -1151,7 +1160,7 @@ void FInsightsManager::OnSessionAnalysisCompleted()
 	{
 		FOutputDevice& Ar = *GLog;
 		Ar.Logf(TEXT("Executing commands on analysis completed..."));
-		FStopwatch Stopwatch;
+		UE::Insights::FStopwatch Stopwatch;
 		Stopwatch.Start();
 		IUnrealInsightsModule& TraceInsightsModule = FModuleManager::LoadModuleChecked<IUnrealInsightsModule>("TraceInsights");
 		TraceInsightsModule.Exec(*SessionAnalysisCompletedCmd, Ar);
diff --git a/Engine/Source/Developer/TraceInsights/Private/Insights/InsightsManager.h b/Engine/Source/Developer/TraceInsights/Private/Insights/InsightsManager.h
--- a/Engine/Source/Developer/TraceInsights/Private/Insights/InsightsManager.h
+++ b/Engine/Source/Developer/TraceInsights/Private/Insights/InsightsManager.h
@@ -2,14 +2,17 @@
 
 #pragma once
 
+#include "CoreTypes.h"
+
 #include "Containers/Ticker.h"
-#include "CoreMinimal.h"
 #include "Framework/Commands/UICommandList.h"
 #include "Input/DragAndDrop.h"
 
-// Insights
-#include "Insights/Common/AvailabilityCheck.h"
-#include "Insights/Common/Stopwatch.h"
+// TraceInsightsCore
+#include "InsightsCore/Common/AvailabilityCheck.h"
+#include "InsightsCore/Common/Stopwatch.h"
+
+// TraceInsights
 #include "Insights/IInsightsManager.h"
 #include "Insights/InsightsCommands.h"
 #include "Insights/InsightsSettings.h"
@@ -18,13 +21,10 @@
 
 ////////////////////////////////////////////////////////////////////////////////////////////////////
 
-namespace UE
-{
-namespace Trace
+namespace UE::Trace
 {
 	class FStoreClient;
 }
-}
 
 namespace TraceServices
 {
@@ -60,7 +60,7 @@ struct FAsyncTaskData
  *     Connecting/disconnecting to source trace
  *     Global Unreal Insights application state and settings
  */
-class FInsightsManager : public TSharedFromThis<FInsightsManager>, public Insights::IInsightsManager
+class FInsightsManager : public TSharedFromThis<FInsightsManager>, public ::Insights::IInsightsManager
 {
 	friend class FInsightsActionManager;
 
@@ -322,20 +322,20 @@ public:
 
 public:
 	/** The event to execute when the session has changed. */
-	virtual Insights::IInsightsManager::FSessionChangedEvent& GetSessionChangedEvent() override { return SessionChangedEvent; }
+	virtual ::Insights::IInsightsManager::FSessionChangedEvent& GetSessionChangedEvent() override { return SessionChangedEvent; }
 private:
 	/** The event to execute when the session has changed. */
-	Insights::IInsightsManager::FSessionChangedEvent SessionChangedEvent;
+	::Insights::IInsightsManager::FSessionChangedEvent SessionChangedEvent;
 
 	////////////////////////////////////////////////////////////////////////////////////////////////////
 	// SessionAnalysisCompletedEvent
 
 public:
 	/** The event to execute when session analysis is complete. */
-	virtual Insights::IInsightsManager::FSessionAnalysisCompletedEvent& GetSessionAnalysisCompletedEvent() override { return SessionAnalysisCompletedEvent; }
+	virtual ::Insights::IInsightsManager::FSessionAnalysisCompletedEvent& GetSessionAnalysisCompletedEvent() override { return SessionAnalysisCompletedEvent; }
 private:
 	/** The event to execute when session analysis is completed. */
-	Insights::IInsightsManager::FSessionAnalysisCompletedEvent SessionAnalysisCompletedEvent;
+	::Insights::IInsightsManager::FSessionAnalysisCompletedEvent SessionAnalysisCompletedEvent;
 
 	////////////////////////////////////////////////////////////////////////////////////////////////////
 
@@ -460,7 +460,7 @@ private:
 	bool bIsAutoLoadLiveSessionEnabled = false;
 	TSet<uint32> AutoLoadedTraceIds; // list of trace ids for the auto loaded live sessions
 
-	FStopwatch AnalysisStopwatch;
+	UE::Insights::FStopwatch AnalysisStopwatch;
 	double SessionDuration = 0.0;
 	double AnalysisDuration = 0.0;
 	double AnalysisSpeedFactor = 0.0;
diff --git a/Engine/Source/Developer/TraceInsights/Private/Insights/LoadingProfiler/LoadingProfilerManager.cpp b/Engine/Source/Developer/TraceInsights/Private/Insights/LoadingProfiler/LoadingProfilerManager.cpp
--- a/Engine/Source/Developer/TraceInsights/Private/Insights/LoadingProfiler/LoadingProfilerManager.cpp
+++ b/Engine/Source/Developer/TraceInsights/Private/Insights/LoadingProfiler/LoadingProfilerManager.cpp
@@ -4,17 +4,21 @@
 
 #include "MessageLogModule.h"
 #include "Modules/ModuleManager.h"
-#include "TraceServices/Model/LoadTimeProfiler.h"
 #include "Widgets/Docking/SDockTab.h"
 #include "WorkspaceMenuStructure.h"
 #include "WorkspaceMenuStructureModule.h"
 
-// Insights
+// TraceServices
+#include "TraceServices/Model/LoadTimeProfiler.h"
+
+// TraceInsightsCore
+#include "InsightsCore/Table/Widgets/STableTreeView.h"
+
+// TraceInsights
 #include "Insights/Common/InsightsMenuBuilder.h"
 #include "Insights/InsightsManager.h"
 #include "Insights/InsightsStyle.h"
 #include "Insights/LoadingProfiler/Widgets/SLoadingProfilerWindow.h"
-#include "Insights/Table/Widgets/STableTreeView.h"
 #include "Insights/Widgets/STimingView.h"
 
 ////////////////////////////////////////////////////////////////////////////////////////////////////
@@ -384,7 +388,7 @@ void FLoadingProfilerManager::OnWindowClosedEvent()
 	TSharedPtr<SLoadingProfilerWindow> Wnd = GetProfilerWindow();
 	if (Wnd)
 	{
-		TSharedPtr<STimingView> TimingView = Wnd->GetTimingView();
+		TSharedPtr<UE::Insights::TimingProfiler::STimingView> TimingView = Wnd->GetTimingView();
 		if (TimingView.IsValid())
 		{
 			TimingView->CloseQuickFindTab();
diff --git a/Engine/Source/Developer/TraceInsights/Private/Insights/LoadingProfiler/LoadingProfilerManager.h b/Engine/Source/Developer/TraceInsights/Private/Insights/LoadingProfiler/LoadingProfilerManager.h
--- a/Engine/Source/Developer/TraceInsights/Private/Insights/LoadingProfiler/LoadingProfilerManager.h
+++ b/Engine/Source/Developer/TraceInsights/Private/Insights/LoadingProfiler/LoadingProfilerManager.h
@@ -7,7 +7,7 @@
 #include "Framework/Commands/UICommandList.h"
 #include "Logging/LogMacros.h"
 
-// Insights
+// TraceInsights
 #include "Insights/InsightsManager.h"
 #include "Insights/IUnrealInsightsModule.h"
 #include "Insights/LoadingProfiler/LoadingProfilerCommands.h"
@@ -139,7 +139,7 @@ private:
 private:
 	bool bIsInitialized;
 	bool bIsAvailable;
-	Insights::FAvailabilityCheck AvailabilityCheck;
+	UE::Insights::FAvailabilityCheck AvailabilityCheck;
 
 	/** The delegate to be invoked when this manager ticks. */
 	FTickerDelegate OnTick;
diff --git a/Engine/Source/Developer/TraceInsights/Private/Insights/LoadingProfiler/Widgets/SLoadingProfilerWindow.cpp b/Engine/Source/Developer/TraceInsights/Private/Insights/LoadingProfiler/Widgets/SLoadingProfilerWindow.cpp
--- a/Engine/Source/Developer/TraceInsights/Private/Insights/LoadingProfiler/Widgets/SLoadingProfilerWindow.cpp
+++ b/Engine/Source/Developer/TraceInsights/Private/Insights/LoadingProfiler/Widgets/SLoadingProfilerWindow.cpp
@@ -6,10 +6,12 @@
 #include "Framework/Docking/TabManager.h"
 #include "Framework/Docking/WorkspaceItem.h"
 #include "SlateOptMacros.h"
-#include "TraceServices/Model/LoadTimeProfiler.h"
 #include "Widgets/Docking/SDockTab.h"
 
-// Insights
+// TraceServices
+#include "TraceServices/Model/LoadTimeProfiler.h"
+
+// TraceInsights
 #include "Insights/InsightsStyle.h"
 #include "Insights/LoadingProfiler/LoadingProfilerManager.h"
 #include "Insights/LoadingProfiler/Widgets/SLoadingProfilerToolbar.h"
@@ -228,7 +230,7 @@ TSharedRef<SDockTab> SLoadingProfilerWindow::SpawnTab_TimingView(const FSpawnTab
 		.ShouldAutosize(false)
 		.TabRole(ETabRole::PanelTab)
 		[
-			SAssignNew(TimingView, STimingView, FInsightsManagerTabs::LoadingProfilerTabId)
+			SAssignNew(TimingView, UE::Insights::TimingProfiler::STimingView, FInsightsManagerTabs::LoadingProfilerTabId)
 		];
 
 	TimingView->Reset(true);
@@ -542,9 +544,9 @@ TSharedRef<SWidget> SLoadingProfilerWindow::CreateToolbar(TSharedPtr<FExtender>
 
 ////////////////////////////////////////////////////////////////////////////////////////////////////
 
-void SLoadingProfilerWindow::OnTimeSelectionChanged(Insights::ETimeChangedFlags InFlags, double InStartTime, double InEndTime)
+void SLoadingProfilerWindow::OnTimeSelectionChanged(UE::Insights::Timing::ETimeChangedFlags InFlags, double InStartTime, double InEndTime)
 {
-	if (InFlags != Insights::ETimeChangedFlags::Interactive)
+	if (InFlags != UE::Insights::Timing::ETimeChangedFlags::Interactive)
 	{
 		if (InStartTime < InEndTime)
 		{
diff --git a/Engine/Source/Developer/TraceInsights/Private/Insights/LoadingProfiler/Widgets/SLoadingProfilerWindow.h b/Engine/Source/Developer/TraceInsights/Private/Insights/LoadingProfiler/Widgets/SLoadingProfilerWindow.h
--- a/Engine/Source/Developer/TraceInsights/Private/Insights/LoadingProfiler/Widgets/SLoadingProfilerWindow.h
+++ b/Engine/Source/Developer/TraceInsights/Private/Insights/LoadingProfiler/Widgets/SLoadingProfilerWindow.h
@@ -4,13 +4,13 @@
 
 #include "CoreMinimal.h"
 
-// Insights
-#include "Insights/ITimingViewSession.h" // for Insights::ETimeChangedFlags
+// TraceInsights
+#include "Insights/ITimingViewSession.h" // for ETimeChangedFlags
 #include "Insights/Widgets/SMajorTabWindow.h"
 
 ////////////////////////////////////////////////////////////////////////////////////////////////////
 
-class STimingView;
+namespace UE::Insights::TimingProfiler { class STimingView; }
 
 namespace Insights
 {
@@ -57,7 +57,7 @@ public:
 	/** Constructs this widget. */
 	void Construct(const FArguments& InArgs, const TSharedRef<SDockTab>& ConstructUnderMajorTab, const TSharedPtr<SWindow>& ConstructUnderWindow);
 
-	TSharedPtr<STimingView> GetTimingView() const { return TimingView; }
+	TSharedPtr<UE::Insights::TimingProfiler::STimingView> GetTimingView() const { return TimingView; }
 	TSharedPtr<Insights::SUntypedTableTreeView> GetEventAggregationTreeView() const { return EventAggregationTreeView; }
 	TSharedPtr<Insights::SUntypedTableTreeView> GetObjectTypeAggregationTreeView() const { return ObjectTypeAggregationTreeView; }
 	TSharedPtr<Insights::SUntypedTableTreeView> GetPackageDetailsTreeView() const { return PackageDetailsTreeView; }
@@ -90,11 +90,11 @@ private:
 	TSharedRef<SDockTab> SpawnTab_RequestsTreeView(const FSpawnTabArgs& Args);
 	void OnRequestsTreeViewTabClosed(TSharedRef<SDockTab> TabBeingClosed);
 
-	void OnTimeSelectionChanged(Insights::ETimeChangedFlags InFlags, double InStartTime, double InEndTime);
+	void OnTimeSelectionChanged(UE::Insights::Timing::ETimeChangedFlags InFlags, double InStartTime, double InEndTime);
 
 private:
 	/** The Timing view (multi-track) widget */
-	TSharedPtr<STimingView> TimingView;
+	TSharedPtr<UE::Insights::TimingProfiler::STimingView> TimingView;
 
 	/** The Event Aggregation tree view widget */
 	TSharedPtr<Insights::SUntypedTableTreeView> EventAggregationTreeView;
diff --git a/Engine/Source/Developer/TraceInsights/Private/Insights/MemoryProfiler/Common/SymbolSearchPathsHelper.h b/Engine/Source/Developer/TraceInsights/Private/Insights/MemoryProfiler/Common/SymbolSearchPathsHelper.h
--- a/Engine/Source/Developer/TraceInsights/Private/Insights/MemoryProfiler/Common/SymbolSearchPathsHelper.h
+++ b/Engine/Source/Developer/TraceInsights/Private/Insights/MemoryProfiler/Common/SymbolSearchPathsHelper.h
@@ -1,9 +1,16 @@
 ﻿// Copyright Epic Games, Inc. All Rights Reserved.
+
 #pragma once
+
 #include "Internationalization/Text.h"
+
+// TraceServices
 #include "TraceServices/Model/Modules.h"
 
-#define LOCTEXT_NAMESPACE "SymbolSearchPathsHelper"
+#define LOCTEXT_NAMESPACE "UE::Insights::FSymbolSearchPathsHelper"
+
+namespace UE::Insights
+{
 
 class FSymbolSearchPathsHelper
 {
@@ -27,4 +34,6 @@ public:
 	}
 };
 
+} // namespace UE::Insights
+
 #undef LOCTEXT_NAMESPACE
diff --git a/Engine/Source/Developer/TraceInsights/Private/Insights/MemoryProfiler/MemoryProfilerCommands.cpp b/Engine/Source/Developer/TraceInsights/Private/Insights/MemoryProfiler/MemoryProfilerCommands.cpp
--- a/Engine/Source/Developer/TraceInsights/Private/Insights/MemoryProfiler/MemoryProfilerCommands.cpp
+++ b/Engine/Source/Developer/TraceInsights/Private/Insights/MemoryProfiler/MemoryProfilerCommands.cpp
@@ -6,14 +6,17 @@
 #include "Framework/Application/SlateApplication.h"
 #include "Framework/MultiBox/MultiBoxBuilder.h"
 
-// Insights
+// TraceInsights
 #include "Insights/InsightsManager.h"
 #include "Insights/InsightsStyle.h"
 #include "Insights/MemoryProfiler/MemoryProfilerManager.h"
 
 ////////////////////////////////////////////////////////////////////////////////////////////////////
 
-#define LOCTEXT_NAMESPACE "FMemoryProfilerCommands"
+#define LOCTEXT_NAMESPACE "UE::Insights::MemoryProfiler"
+
+namespace UE::Insights::MemoryProfiler
+{
 
 ////////////////////////////////////////////////////////////////////////////////////////////////////
 // FMemoryProfilerMenuBuilder
@@ -129,4 +132,6 @@ IMPLEMENT_TOGGLE_COMMAND(ToggleModulesViewVisibility, IsModulesViewVisible, Show
 
 ////////////////////////////////////////////////////////////////////////////////////////////////////
 
+} // namespace UE::Insights::MemoryProfiler
+
 #undef LOCTEXT_NAMESPACE
diff --git a/Engine/Source/Developer/TraceInsights/Private/Insights/MemoryProfiler/MemoryProfilerCommands.h b/Engine/Source/Developer/TraceInsights/Private/Insights/MemoryProfiler/MemoryProfilerCommands.h
--- a/Engine/Source/Developer/TraceInsights/Private/Insights/MemoryProfiler/MemoryProfilerCommands.h
+++ b/Engine/Source/Developer/TraceInsights/Private/Insights/MemoryProfiler/MemoryProfilerCommands.h
@@ -2,12 +2,16 @@
 
 #pragma once
 
-#include "CoreMinimal.h"
+#include "CoreTypes.h"
+
 #include "Framework/Commands/UIAction.h"
 #include "Framework/Commands/Commands.h"
 
 class FMenuBuilder;
 
+namespace UE::Insights::MemoryProfiler
+{
+
 ////////////////////////////////////////////////////////////////////////////////////////////////////
 /**
  * Class that holds all profiler commands.
@@ -101,3 +105,7 @@ protected:
 	/** Reference to the global instance of the profiler manager. */
 	class FMemoryProfilerManager* This;
 };
+
+////////////////////////////////////////////////////////////////////////////////////////////////////
+
+} // namespace UE::Insights::MemoryProfiler
diff --git a/Engine/Source/Developer/TraceInsights/Private/Insights/MemoryProfiler/MemoryProfilerManager.cpp b/Engine/Source/Developer/TraceInsights/Private/Insights/MemoryProfiler/MemoryProfilerManager.cpp
--- a/Engine/Source/Developer/TraceInsights/Private/Insights/MemoryProfiler/MemoryProfilerManager.cpp
+++ b/Engine/Source/Developer/TraceInsights/Private/Insights/MemoryProfiler/MemoryProfilerManager.cpp
@@ -13,20 +13,23 @@
 #include "TraceServices/Model/AllocationsProvider.h"
 #include "TraceServices/Model/Memory.h"
 
-// Insights
+// TraceInsightsCore
+#include "InsightsCore/Table/Widgets/STableTreeView.h"
+
+// TraceInsights
 #include "Insights/Common/InsightsMenuBuilder.h"
 #include "Insights/InsightsManager.h"
 #include "Insights/InsightsStyle.h"
 #include "Insights/MemoryProfiler/ViewModels/MemorySharedState.h"
 #include "Insights/MemoryProfiler/Widgets/SMemoryProfilerWindow.h"
-#include "Insights/Table/Widgets/STableTreeView.h"
 #include "Insights/Widgets/STimingView.h"
 
-////////////////////////////////////////////////////////////////////////////////////////////////////
+#define LOCTEXT_NAMESPACE "UE::Insights::MemoryProfiler"
 
-#define LOCTEXT_NAMESPACE "MemoryProfilerManager"
+namespace UE::Insights::MemoryProfiler
+{
 
-DEFINE_LOG_CATEGORY(MemoryProfiler);
+DEFINE_LOG_CATEGORY(LogMemoryProfiler);
 
 TSharedPtr<FMemoryProfilerManager> FMemoryProfilerManager::Instance = nullptr;
 
@@ -79,7 +82,7 @@ void FMemoryProfilerManager::Initialize(IUnrealInsightsModule& InsightsModule)
 	}
 	bIsInitialized = true;
 
-	UE_LOG(MemoryProfiler, Log, TEXT("Initialize"));
+	UE_LOG(LogMemoryProfiler, Log, TEXT("Initialize"));
 
 	// Register tick functions.
 	OnTick = FTickerDelegate::CreateSP(this, &FMemoryProfilerManager::Tick);
@@ -121,7 +124,7 @@ void FMemoryProfilerManager::Shutdown()
 
 	FMemoryProfilerManager::Instance.Reset();
 
-	UE_LOG(MemoryProfiler, Log, TEXT("Shutdown"));
+	UE_LOG(LogMemoryProfiler, Log, TEXT("Shutdown"));
 }
 
 ////////////////////////////////////////////////////////////////////////////////////////////////////
@@ -301,7 +304,7 @@ bool FMemoryProfilerManager::Tick(float DeltaTime)
 			const FName& TabId = FInsightsManagerTabs::MemoryProfilerTabId;
 			if (FGlobalTabmanager::Get()->HasTabSpawner(TabId))
 			{
-				UE_LOG(MemoryProfiler, Log, TEXT("Opening the \"Memory Insights\" tab..."));
+				UE_LOG(LogMemoryProfiler, Log, TEXT("Opening the \"Memory Insights\" tab..."));
 				FGlobalTabmanager::Get()->TryInvokeTab(TabId);
 			}
 #endif
@@ -315,7 +318,7 @@ bool FMemoryProfilerManager::Tick(float DeltaTime)
 
 void FMemoryProfilerManager::OnSessionChanged()
 {
-	UE_LOG(MemoryProfiler, Log, TEXT("OnSessionChanged"));
+	UE_LOG(LogMemoryProfiler, Log, TEXT("OnSessionChanged"));
 
 	bIsAvailable = false;
 	if (FInsightsManager::Get()->GetSession().IsValid())
@@ -396,7 +399,7 @@ void FMemoryProfilerManager::OnWindowClosedEvent()
 	{
 		Wnd->CloseMemAllocTableTreeTabs();
 
-		TSharedPtr<STimingView> TimingView = Wnd->GetTimingView();
+		TSharedPtr<TimingProfiler::STimingView> TimingView = Wnd->GetTimingView();
 		if (TimingView.IsValid())
 		{
 			TimingView->CloseQuickFindTab();
@@ -406,4 +409,6 @@ void FMemoryProfilerManager::OnWindowClosedEvent()
 
 ////////////////////////////////////////////////////////////////////////////////////////////////////
 
+} // namespace UE::Insights::MemoryProfiler
+
 #undef LOCTEXT_NAMESPACE
diff --git a/Engine/Source/Developer/TraceInsights/Private/Insights/MemoryProfiler/MemoryProfilerManager.h b/Engine/Source/Developer/TraceInsights/Private/Insights/MemoryProfiler/MemoryProfilerManager.h
--- a/Engine/Source/Developer/TraceInsights/Private/Insights/MemoryProfiler/MemoryProfilerManager.h
+++ b/Engine/Source/Developer/TraceInsights/Private/Insights/MemoryProfiler/MemoryProfilerManager.h
@@ -2,20 +2,24 @@
 
 #pragma once
 
+#include "CoreTypes.h"
+
 #include "Containers/Ticker.h"
-#include "CoreMinimal.h"
 #include "Framework/Commands/UICommandList.h"
 #include "Logging/LogMacros.h"
 
-// Insights
+// TraceInsights
 #include "Insights/InsightsManager.h"
 #include "Insights/IUnrealInsightsModule.h"
 #include "Insights/MemoryProfiler/MemoryProfilerCommands.h"
 
+namespace UE::Insights::MemoryProfiler
+{
+
 class FMemorySharedState;
 class SMemoryProfilerWindow;
 
-DECLARE_LOG_CATEGORY_EXTERN(MemoryProfiler, Log, All);
+DECLARE_LOG_CATEGORY_EXTERN(LogMemoryProfiler, Log, All);
 
 ////////////////////////////////////////////////////////////////////////////////////////////////////
 /**
@@ -128,7 +132,7 @@ private:
 private:
 	bool bIsInitialized;
 	bool bIsAvailable;
-	Insights::FAvailabilityCheck AvailabilityCheck;
+	FAvailabilityCheck AvailabilityCheck;
 
 	/** The delegate to be invoked when this manager ticks. */
 	FTickerDelegate OnTick;
@@ -163,3 +167,5 @@ private:
 	/** A shared pointer to the global instance of the Memory Profiler manager. */
 	static TSharedPtr<FMemoryProfilerManager> Instance;
 };
+
+} // namespace UE::Insights::MemoryProfiler
diff --git a/Engine/Source/Developer/TraceInsights/Private/Insights/MemoryProfiler/ViewModels/CallstackFormatting.cpp b/Engine/Source/Developer/TraceInsights/Private/Insights/MemoryProfiler/ViewModels/CallstackFormatting.cpp
--- a/Engine/Source/Developer/TraceInsights/Private/Insights/MemoryProfiler/ViewModels/CallstackFormatting.cpp
+++ b/Engine/Source/Developer/TraceInsights/Private/Insights/MemoryProfiler/ViewModels/CallstackFormatting.cpp
@@ -2,9 +2,13 @@
 
 #include "CallstackFormatting.h"
 
-#define LOCTEXT_NAMESPACE "Insights::CallstackFormatting"
+// TraceServices
+#include "TraceServices/Model/Callstack.h"
+#include "TraceServices/Model/Modules.h"
 
-namespace Insights
+#define LOCTEXT_NAMESPACE "UE::Insights::CallstackFormatting"
+
+namespace UE::Insights
 {
 
 ////////////////////////////////////////////////////////////////////////////////////////////////////
@@ -104,6 +108,6 @@ void FormatStackFrame(const TraceServices::FStackFrame& Frame, FStringBuilderBas
 
 ////////////////////////////////////////////////////////////////////////////////////////////////////
 
-} // namespace Insights
+} // namespace UE::Insights
 
 #undef LOCTEXT_NAMESPACE
diff --git a/Engine/Source/Developer/TraceInsights/Private/Insights/MemoryProfiler/ViewModels/CallstackFormatting.h b/Engine/Source/Developer/TraceInsights/Private/Insights/MemoryProfiler/ViewModels/CallstackFormatting.h
--- a/Engine/Source/Developer/TraceInsights/Private/Insights/MemoryProfiler/ViewModels/CallstackFormatting.h
+++ b/Engine/Source/Developer/TraceInsights/Private/Insights/MemoryProfiler/ViewModels/CallstackFormatting.h
@@ -2,12 +2,14 @@
 
 #pragma once
 
+#include "CoreTypes.h"
+
 #include "Internationalization/Text.h"
 #include "Misc/EnumClassFlags.h"
-#include "TraceServices/Model/Callstack.h"
-#include "TraceServices/Model/Modules.h"
 
-namespace Insights
+namespace TraceServices { struct FStackFrame; }
+
+namespace UE::Insights
 {
 
 enum class EStackFrameFormatFlags : uint8
@@ -30,4 +32,4 @@ FText GetEmptyCallstackString();
 
 void FormatStackFrame(const TraceServices::FStackFrame& Frame, FStringBuilderBase& OutString, EStackFrameFormatFlags FormatFlags);
 
-} // namespace Insights
+} // namespace UE::Insights
diff --git a/Engine/Source/Developer/TraceInsights/Private/Insights/MemoryProfiler/ViewModels/MemAllocFilterValueConverter.cpp b/Engine/Source/Developer/TraceInsights/Private/Insights/MemoryProfiler/ViewModels/MemAllocFilterValueConverter.cpp
--- a/Engine/Source/Developer/TraceInsights/Private/Insights/MemoryProfiler/ViewModels/MemAllocFilterValueConverter.cpp
+++ b/Engine/Source/Developer/TraceInsights/Private/Insights/MemoryProfiler/ViewModels/MemAllocFilterValueConverter.cpp
@@ -2,11 +2,12 @@
 
 #include "MemAllocFilterValueConverter.h"
 
+#include "Internationalization/Internationalization.h"
 #include "Misc/Char.h"
 
-#define LOCTEXT_NAMESPACE "Insights::FMemoryFilterValueConverter"
+#define LOCTEXT_NAMESPACE "UE::Insights::MemoryProfiler::FMemAllocNode"
 
-namespace Insights
+namespace UE::Insights::MemoryProfiler
 {
 
 ////////////////////////////////////////////////////////////////////////////////////////////////////
@@ -84,6 +85,6 @@ FText FMemoryFilterValueConverter::GetHintText() const
 
 ////////////////////////////////////////////////////////////////////////////////////////////////////
 
-} // namespace Insights
+} // namespace UE::Insights::MemoryProfiler
 
 #undef LOCTEXT_NAMESPACE
diff --git a/Engine/Source/Developer/TraceInsights/Private/Insights/MemoryProfiler/ViewModels/MemAllocFilterValueConverter.h b/Engine/Source/Developer/TraceInsights/Private/Insights/MemoryProfiler/ViewModels/MemAllocFilterValueConverter.h
--- a/Engine/Source/Developer/TraceInsights/Private/Insights/MemoryProfiler/ViewModels/MemAllocFilterValueConverter.h
+++ b/Engine/Source/Developer/TraceInsights/Private/Insights/MemoryProfiler/ViewModels/MemAllocFilterValueConverter.h
@@ -2,10 +2,10 @@
 
 #pragma once
 
-// Insights
-#include "Insights/ViewModels/Filters.h"
+// TraceInsightsCore
+#include "InsightsCore/Filter/ViewModels/Filters.h"
 
-namespace Insights
+namespace UE::Insights::MemoryProfiler
 {
 
 ////////////////////////////////////////////////////////////////////////////////////////////////////
@@ -20,4 +20,4 @@ public:
 
 ////////////////////////////////////////////////////////////////////////////////////////////////////
 
-} // namespace Insights
+} // namespace UE::Insights::MemoryProfiler
diff --git a/Engine/Source/Developer/TraceInsights/Private/Insights/MemoryProfiler/ViewModels/MemAllocGroupingByCallstack.cpp b/Engine/Source/Developer/TraceInsights/Private/Insights/MemoryProfiler/ViewModels/MemAllocGroupingByCallstack.cpp
--- a/Engine/Source/Developer/TraceInsights/Private/Insights/MemoryProfiler/ViewModels/MemAllocGroupingByCallstack.cpp
+++ b/Engine/Source/Developer/TraceInsights/Private/Insights/MemoryProfiler/ViewModels/MemAllocGroupingByCallstack.cpp
@@ -1,17 +1,21 @@
 // Copyright Epic Games, Inc. All Rights Reserved.
 
 #include "MemAllocGroupingByCallstack.h"
+
+// TraceServices
 #include "TraceServices/Model/Callstack.h"
 
-// Insights
-#include "Insights/Common/AsyncOperationProgress.h"
+// TraceInsightsCore
+#include "InsightsCore/Common/AsyncOperationProgress.h"
+
+// TraceInsights
 #include "Insights/InsightsStyle.h"
 #include "Insights/MemoryProfiler/ViewModels/CallstackFormatting.h"
 #include "Insights/MemoryProfiler/ViewModels/MemAllocNode.h"
 
-#define LOCTEXT_NAMESPACE "Insights::FMemAllocGroupingByCallstack"
+#define LOCTEXT_NAMESPACE "UE::Insights::MemoryProfiler::FMemAllocNode"
 
-namespace Insights
+namespace UE::Insights::MemoryProfiler
 {
 
 ////////////////////////////////////////////////////////////////////////////////////////////////////
@@ -85,7 +89,11 @@ FMemAllocGroupingByCallstack::~FMemAllocGroupingByCallstack()
 
 ////////////////////////////////////////////////////////////////////////////////////////////////////
 
-void FMemAllocGroupingByCallstack::GroupNodes(const TArray<FTableTreeNodePtr>& Nodes, FTableTreeNode& ParentGroup, TWeakPtr<FTable> InParentTable, IAsyncOperationProgress& InAsyncOperationProgress) const
+void FMemAllocGroupingByCallstack::GroupNodes(
+	const TArray<FTableTreeNodePtr>& Nodes,
+	FTableTreeNode& ParentGroup,
+	TWeakPtr<FTable> InParentTable,
+	IAsyncOperationProgress& InAsyncOperationProgress) const
 {
 	const bool bLocalIsGroupingByFunction = bIsGroupingByFunction;
 	const bool bLocalShouldSkipFilteredFrames = bShouldSkipFilteredFrames;
@@ -322,6 +330,6 @@ FTableTreeNode* FMemAllocGroupingByCallstack::CreateEmptyCallstackGroup(TWeakPtr
 
 ////////////////////////////////////////////////////////////////////////////////////////////////////
 
-} // namespace Insights
+} // namespace UE::Insights::MemoryProfiler
 
 #undef LOCTEXT_NAMESPACE
diff --git a/Engine/Source/Developer/TraceInsights/Private/Insights/MemoryProfiler/ViewModels/MemAllocGroupingByCallstack.h b/Engine/Source/Developer/TraceInsights/Private/Insights/MemoryProfiler/ViewModels/MemAllocGroupingByCallstack.h
--- a/Engine/Source/Developer/TraceInsights/Private/Insights/MemoryProfiler/ViewModels/MemAllocGroupingByCallstack.h
+++ b/Engine/Source/Developer/TraceInsights/Private/Insights/MemoryProfiler/ViewModels/MemAllocGroupingByCallstack.h
@@ -2,22 +2,18 @@
 
 #pragma once
 
-#include "CoreMinimal.h"
+#include "CoreTypes.h"
 
-// Insights
-#include "Insights/Table/ViewModels/BaseTreeNode.h"
-#include "Insights/Table/ViewModels/TreeNodeGrouping.h"
+// TraceInsightsCore
+#include "InsightsCore/Table/ViewModels/BaseTreeNode.h"
+#include "InsightsCore/Table/ViewModels/TreeNodeGrouping.h"
 
-namespace TraceServices
-{
-	struct FStackFrame;
-}
+namespace TraceServices { struct FStackFrame; }
+namespace UE::Insights { class IAsyncOperationProgress; }
 
-namespace Insights
+namespace UE::Insights::MemoryProfiler
 {
 
-class IAsyncOperationProgress;
-
 ////////////////////////////////////////////////////////////////////////////////////////////////////
 
 class FCallstackFrameGroupNode : public FTableTreeNode
@@ -102,4 +98,4 @@ private:
 
 ////////////////////////////////////////////////////////////////////////////////////////////////////
 
-} // namespace Insights
+} // namespace UE::Insights::MemoryProfiler
diff --git a/Engine/Source/Developer/TraceInsights/Private/Insights/MemoryProfiler/ViewModels/MemAllocGroupingByHeap.cpp b/Engine/Source/Developer/TraceInsights/Private/Insights/MemoryProfiler/ViewModels/MemAllocGroupingByHeap.cpp
--- a/Engine/Source/Developer/TraceInsights/Private/Insights/MemoryProfiler/ViewModels/MemAllocGroupingByHeap.cpp
+++ b/Engine/Source/Developer/TraceInsights/Private/Insights/MemoryProfiler/ViewModels/MemAllocGroupingByHeap.cpp
@@ -2,17 +2,21 @@
 
 #include "MemAllocGroupingByHeap.h"
 
-#include "Common/ProviderLock.h" // TraceServices
 #include "Internationalization/Internationalization.h"
 
-// Insights
-#include "Insights/Common/AsyncOperationProgress.h"
+// TraceServices
+#include "Common/ProviderLock.h"
+
+// TraceInsightsCore
+#include "InsightsCore/Common/AsyncOperationProgress.h"
+
+// TraceInsights
 #include "Insights/MemoryProfiler/ViewModels/CallstackFormatting.h"
 #include "Insights/MemoryProfiler/ViewModels/MemAllocNode.h"
 
-#define LOCTEXT_NAMESPACE "Insights::FMemAllocGroupingByHeap"
+#define LOCTEXT_NAMESPACE "UE::Insights::MemoryProfiler::FMemAllocNode"
 
-namespace Insights
+namespace UE::Insights::MemoryProfiler
 {
 
 INSIGHTS_IMPLEMENT_RTTI(FMemAllocGroupingByHeap)
@@ -77,8 +81,11 @@ FTableTreeNodePtr MakeGroupNodeHierarchy(const TraceServices::IAllocationsProvid
 
 ////////////////////////////////////////////////////////////////////////////////////////////////////
 
-void FMemAllocGroupingByHeap::GroupNodes(const TArray<FTableTreeNodePtr>& Nodes, FTableTreeNode& ParentGroup,
-	TWeakPtr<FTable> InParentTable, IAsyncOperationProgress& InAsyncOperationProgress) const
+void FMemAllocGroupingByHeap::GroupNodes(
+	const TArray<FTableTreeNodePtr>& Nodes,
+	FTableTreeNode& ParentGroup,
+	TWeakPtr<FTable> InParentTable,
+	IAsyncOperationProgress& InAsyncOperationProgress) const
 {
 	ParentGroup.ClearChildren();
 
@@ -135,6 +142,6 @@ void FMemAllocGroupingByHeap::GroupNodes(const TArray<FTableTreeNodePtr>& Nodes,
 
 ////////////////////////////////////////////////////////////////////////////////////////////////////
 
-} // namespace Insights
+} // namespace UE::Insights::MemoryProfiler
 
 #undef LOCTEXT_NAMESPACE
diff --git a/Engine/Source/Developer/TraceInsights/Private/Insights/MemoryProfiler/ViewModels/MemAllocGroupingByHeap.h b/Engine/Source/Developer/TraceInsights/Private/Insights/MemoryProfiler/ViewModels/MemAllocGroupingByHeap.h
--- a/Engine/Source/Developer/TraceInsights/Private/Insights/MemoryProfiler/ViewModels/MemAllocGroupingByHeap.h
+++ b/Engine/Source/Developer/TraceInsights/Private/Insights/MemoryProfiler/ViewModels/MemAllocGroupingByHeap.h
@@ -2,24 +2,18 @@
 
 #pragma once
 
-#include "CoreMinimal.h"
+#include "CoreTypes.h"
 
-// Insights
-#include "Insights/Table/ViewModels/BaseTreeNode.h"
-#include "Insights/Table/ViewModels/TreeNodeGrouping.h"
-#include "TraceServices/Model/AllocationsProvider.h"
+// TraceInsightsCore
+#include "InsightsCore/Table/ViewModels/BaseTreeNode.h"
+#include "InsightsCore/Table/ViewModels/TreeNodeGrouping.h"
 
-namespace TraceServices
-{
-	struct FStackFrame;
-}
-
-namespace Insights
-{
+namespace TraceServices { class IAllocationsProvider; }
 
-class IAsyncOperationProgress;
+namespace UE::Insights { class IAsyncOperationProgress; }
 
-////////////////////////////////////////////////////////////////////////////////////////////////////
+namespace UE::Insights::MemoryProfiler
+{
 
 class FMemAllocGroupingByHeap : public FTreeNodeGrouping
 {
@@ -30,10 +24,9 @@ public:
 	virtual ~FMemAllocGroupingByHeap() override;
 
 	virtual void GroupNodes(const TArray<FTableTreeNodePtr>& Nodes, FTableTreeNode& ParentGroup, TWeakPtr<FTable> InParentTable, IAsyncOperationProgress& InAsyncOperationProgress) const override;
+
 private:
 	const TraceServices::IAllocationsProvider& AllocProvider;
 };
 
-////////////////////////////////////////////////////////////////////////////////////////////////////
-
-} // namespace Insights
+} // namespace UE::Insights::MemoryProfiler
diff --git a/Engine/Source/Developer/TraceInsights/Private/Insights/MemoryProfiler/ViewModels/MemAllocGroupingBySize.cpp b/Engine/Source/Developer/TraceInsights/Private/Insights/MemoryProfiler/ViewModels/MemAllocGroupingBySize.cpp
--- a/Engine/Source/Developer/TraceInsights/Private/Insights/MemoryProfiler/ViewModels/MemAllocGroupingBySize.cpp
+++ b/Engine/Source/Developer/TraceInsights/Private/Insights/MemoryProfiler/ViewModels/MemAllocGroupingBySize.cpp
@@ -2,13 +2,15 @@
 
 #include "MemAllocGroupingBySize.h"
 
-// Insights
-#include "Insights/Common/AsyncOperationProgress.h"
+// TraceInsightsCore
+#include "InsightsCore/Common/AsyncOperationProgress.h"
+
+// TraceInsights
 #include "Insights/MemoryProfiler/ViewModels/MemAllocNode.h"
 
-#define LOCTEXT_NAMESPACE "Insights::FMemAllocGroupingBySize"
+#define LOCTEXT_NAMESPACE "UE::Insights::MemoryProfiler::FMemAllocNode"
 
-namespace Insights
+namespace UE::Insights::MemoryProfiler
 {
 
 ////////////////////////////////////////////////////////////////////////////////////////////////////
@@ -73,7 +75,11 @@ void FMemAllocGroupingBySize::ResetThresholdsPow2()
 
 ////////////////////////////////////////////////////////////////////////////////////////////////////
 
-void FMemAllocGroupingBySize::GroupNodes(const TArray<FTableTreeNodePtr>& Nodes, FTableTreeNode& ParentGroup, TWeakPtr<FTable> InParentTable, IAsyncOperationProgress& InAsyncOperationProgress) const
+void FMemAllocGroupingBySize::GroupNodes(
+	const TArray<FTableTreeNodePtr>& Nodes,
+	FTableTreeNode& ParentGroup,
+	TWeakPtr<FTable> InParentTable,
+	IAsyncOperationProgress& InAsyncOperationProgress) const
 {
 	TArray<FTableTreeNodePtr> GroupMap; // Node for each ThresholdIndex
 	GroupMap.AddDefaulted(Thresholds.Num());
@@ -144,6 +150,6 @@ void FMemAllocGroupingBySize::GroupNodes(const TArray<FTableTreeNodePtr>& Nodes,
 
 ////////////////////////////////////////////////////////////////////////////////////////////////////
 
-} // namespace Insights
+} // namespace UE::Insights::MemoryProfiler
 
 #undef LOCTEXT_NAMESPACE
diff --git a/Engine/Source/Developer/TraceInsights/Private/Insights/MemoryProfiler/ViewModels/MemAllocGroupingBySize.h b/Engine/Source/Developer/TraceInsights/Private/Insights/MemoryProfiler/ViewModels/MemAllocGroupingBySize.h
--- a/Engine/Source/Developer/TraceInsights/Private/Insights/MemoryProfiler/ViewModels/MemAllocGroupingBySize.h
+++ b/Engine/Source/Developer/TraceInsights/Private/Insights/MemoryProfiler/ViewModels/MemAllocGroupingBySize.h
@@ -2,18 +2,16 @@
 
 #pragma once
 
-#include "CoreMinimal.h"
+#include "CoreTypes.h"
 
-// Insights
-#include "Insights/Table/ViewModels/BaseTreeNode.h"
-#include "Insights/Table/ViewModels/TreeNodeGrouping.h"
+// TraceInsightsCore
+#include "InsightsCore/Table/ViewModels/BaseTreeNode.h"
+#include "InsightsCore/Table/ViewModels/TreeNodeGrouping.h"
 
-namespace Insights
-{
-
-class IAsyncOperationProgress;
+namespace UE::Insights { class IAsyncOperationProgress; }
 
-////////////////////////////////////////////////////////////////////////////////////////////////////
+namespace UE::Insights::MemoryProfiler
+{
 
 class FMemAllocGroupingBySize : public FTreeNodeGrouping
 {
@@ -43,6 +41,4 @@ private:
 	bool bIsPow2; // thresholds are automatically set as power of two
 };
 
-////////////////////////////////////////////////////////////////////////////////////////////////////
-
-} // namespace Insights
+} // namespace UE::Insights::MemoryProfiler
diff --git a/Engine/Source/Developer/TraceInsights/Private/Insights/MemoryProfiler/ViewModels/MemAllocGroupingBySwapPage.cpp b/Engine/Source/Developer/TraceInsights/Private/Insights/MemoryProfiler/ViewModels/MemAllocGroupingBySwapPage.cpp
--- a/Engine/Source/Developer/TraceInsights/Private/Insights/MemoryProfiler/ViewModels/MemAllocGroupingBySwapPage.cpp
+++ b/Engine/Source/Developer/TraceInsights/Private/Insights/MemoryProfiler/ViewModels/MemAllocGroupingBySwapPage.cpp
@@ -1,18 +1,22 @@
-﻿// Copyright Epic Games, Inc. All Rights Reserved.
+// Copyright Epic Games, Inc. All Rights Reserved.
 
 #include "MemAllocGroupingBySwapPage.h"
 
-#include "Common/ProviderLock.h" // TraceServices
 #include "Internationalization/Internationalization.h"
 
-// Insights
-#include "Insights/Common/AsyncOperationProgress.h"
+// TraceServices
+#include "Common/ProviderLock.h"
+
+// TraceInsightsCore
+#include "InsightsCore/Common/AsyncOperationProgress.h"
+
+// TraceInsights
 #include "Insights/MemoryProfiler/ViewModels/CallstackFormatting.h"
 #include "Insights/MemoryProfiler/ViewModels/MemAllocInSwapNode.h"
 
-#define LOCTEXT_NAMESPACE "Insights::FMemAllocGroupingBySwap"
+#define LOCTEXT_NAMESPACE "UE::Insights::MemoryProfiler::FMemAllocNode"
 
-namespace Insights
+namespace UE::Insights::MemoryProfiler
 {
 
 INSIGHTS_IMPLEMENT_RTTI(FMemAllocGroupingBySwapPage)
@@ -57,8 +61,11 @@ struct TMemoryPageMapKeyFuncs : BaseKeyFuncs<TPair<uint64, FTableTreeNode*>, uin
 
 ////////////////////////////////////////////////////////////////////////////////////////////////////
 
-void FMemAllocGroupingBySwapPage::GroupNodes(const TArray<FTableTreeNodePtr>& Nodes, FTableTreeNode& ParentGroup,
-	TWeakPtr<FTable> InParentTable, IAsyncOperationProgress& InAsyncOperationProgress) const
+void FMemAllocGroupingBySwapPage::GroupNodes(
+	const TArray<FTableTreeNodePtr>& Nodes,
+	FTableTreeNode& ParentGroup,
+	TWeakPtr<FTable> InParentTable,
+	IAsyncOperationProgress& InAsyncOperationProgress) const
 {
 	ParentGroup.ClearChildren();
 
@@ -164,6 +171,6 @@ void FMemAllocGroupingBySwapPage::GroupNodes(const TArray<FTableTreeNodePtr>& No
 
 ////////////////////////////////////////////////////////////////////////////////////////////////////
 
-} // namespace Insights
+} // namespace UE::Insights::MemoryProfiler
 
 #undef LOCTEXT_NAMESPACE
diff --git a/Engine/Source/Developer/TraceInsights/Private/Insights/MemoryProfiler/ViewModels/MemAllocGroupingBySwapPage.h b/Engine/Source/Developer/TraceInsights/Private/Insights/MemoryProfiler/ViewModels/MemAllocGroupingBySwapPage.h
--- a/Engine/Source/Developer/TraceInsights/Private/Insights/MemoryProfiler/ViewModels/MemAllocGroupingBySwapPage.h
+++ b/Engine/Source/Developer/TraceInsights/Private/Insights/MemoryProfiler/ViewModels/MemAllocGroupingBySwapPage.h
@@ -1,24 +1,19 @@
-﻿// Copyright Epic Games, Inc. All Rights Reserved.
+// Copyright Epic Games, Inc. All Rights Reserved.
 
 #pragma once
 
-#include "CoreMinimal.h"
+#include "CoreTypes.h"
 
-// Insights
-#include "Insights/Table/ViewModels/BaseTreeNode.h"
-#include "Insights/Table/ViewModels/TreeNodeGrouping.h"
-#include "TraceServices/Model/AllocationsProvider.h"
+// TraceInsightsCore
+#include "InsightsCore/Table/ViewModels/BaseTreeNode.h"
+#include "InsightsCore/Table/ViewModels/TreeNodeGrouping.h"
 
-namespace TraceServices
-{
-	struct FStackFrame;
-}
+namespace TraceServices { class IAllocationsProvider; }
 
-namespace Insights
-{
-class IAsyncOperationProgress;
+namespace UE::Insights { class IAsyncOperationProgress; }
 
-////////////////////////////////////////////////////////////////////////////////////////////////////
+namespace UE::Insights::MemoryProfiler
+{
 
 class FMemAllocGroupingBySwapPage : public FTreeNodeGrouping
 {
@@ -29,10 +24,9 @@ public:
 	virtual ~FMemAllocGroupingBySwapPage() override;
 
 	virtual void GroupNodes(const TArray<FTableTreeNodePtr>& Nodes, FTableTreeNode& ParentGroup, TWeakPtr<FTable> InParentTable, IAsyncOperationProgress& InAsyncOperationProgress) const override;
+
 private:
 	const TraceServices::IAllocationsProvider& AllocProvider;
 };
 
-////////////////////////////////////////////////////////////////////////////////////////////////////
-
-} // namespace Insights
+} // namespace UE::Insights::MemoryProfiler
diff --git a/Engine/Source/Developer/TraceInsights/Private/Insights/MemoryProfiler/ViewModels/MemAllocGroupingByTag.cpp b/Engine/Source/Developer/TraceInsights/Private/Insights/MemoryProfiler/ViewModels/MemAllocGroupingByTag.cpp
--- a/Engine/Source/Developer/TraceInsights/Private/Insights/MemoryProfiler/ViewModels/MemAllocGroupingByTag.cpp
+++ b/Engine/Source/Developer/TraceInsights/Private/Insights/MemoryProfiler/ViewModels/MemAllocGroupingByTag.cpp
@@ -2,19 +2,22 @@
 
 #include "MemAllocGroupingByTag.h"
 
-#include "Common/ProviderLock.h" // TraceServices
+// TraceServices
+#include "Common/ProviderLock.h"
 
-// Insights
-#include "Insights/Common/AsyncOperationProgress.h"
+// TraceInsightsCore
+#include "InsightsCore/Common/AsyncOperationProgress.h"
+
+// TraceInsights
 #include "Insights/InsightsManager.h"
 #include "Insights/MemoryProfiler/ViewModels/MemAllocNode.h"
 #include "Insights/MemoryProfiler/MemoryProfilerManager.h"
 
-#define LOCTEXT_NAMESPACE "Insights::FMemAllocGroupingByTag"
+#define LOCTEXT_NAMESPACE "UE::Insights::MemoryProfiler::FMemAllocNode"
 
 #define INSIGHTS_MERGE_MEM_TAGS_BY_NAME 1
 
-namespace Insights
+namespace UE::Insights::MemoryProfiler
 {
 
 ////////////////////////////////////////////////////////////////////////////////////////////////////
@@ -39,11 +42,13 @@ FMemAllocGroupingByTag::FMemAllocGroupingByTag(const TraceServices::IAllocations
 
 ////////////////////////////////////////////////////////////////////////////////////////////////////
 
-void FMemAllocGroupingByTag::GroupNodes(const TArray<FTableTreeNodePtr>& Nodes,
-                                        FTableTreeNode& ParentGroup, TWeakPtr<FTable> InParentTable,
-                                        IAsyncOperationProgress& InAsyncOperationProgress) const
+void FMemAllocGroupingByTag::GroupNodes(
+	const TArray<FTableTreeNodePtr>& Nodes,
+	FTableTreeNode& ParentGroup,
+	TWeakPtr<FTable> InParentTable,
+	IAsyncOperationProgress& InAsyncOperationProgress) const
 {
-	using namespace TraceServices;
+	using TagIdType = TraceServices::TagIdType;
 
 	ParentGroup.ClearChildren();
 
@@ -291,7 +296,7 @@ void FMemTagTableTreeNode::UpdateLLMSize() const
 		return;
 	}
 
-	::FMemorySharedState* SharedState = FMemoryProfilerManager::Get()->GetSharedState();
+	FMemorySharedState* SharedState = FMemoryProfilerManager::Get()->GetSharedState();
 	if (!SharedState)
 	{
 		return;
@@ -300,8 +305,8 @@ void FMemTagTableTreeNode::UpdateLLMSize() const
 	TraceServices::FMemoryTrackerId TrackerId = FMemoryTracker::InvalidTrackerId;
 	TraceServices::FMemoryTagId TagId = FMemoryTag::InvalidTagId;
 
-	const Insights::FMemoryTagList& TagList = SharedState->GetTagList();
-	for (const Insights::FMemoryTag* MemTag : TagList.GetTags())
+	const FMemoryTagList& TagList = SharedState->GetTagList();
+	for (const FMemoryTag* MemTag : TagList.GetTags())
 	{
 		if (MemTag)
 		{
@@ -350,7 +355,7 @@ void FMemTagTableTreeNode::UpdateLLMSize() const
 
 ////////////////////////////////////////////////////////////////////////////////////////////////////
 
-} // namespace Insights
+} // namespace UE::Insights::MemoryProfiler
 
 #undef INSIGHTS_MERGE_MEM_TAGS_BY_NAME
 #undef LOCTEXT_NAMESPACE
diff --git a/Engine/Source/Developer/TraceInsights/Private/Insights/MemoryProfiler/ViewModels/MemAllocGroupingByTag.h b/Engine/Source/Developer/TraceInsights/Private/Insights/MemoryProfiler/ViewModels/MemAllocGroupingByTag.h
--- a/Engine/Source/Developer/TraceInsights/Private/Insights/MemoryProfiler/ViewModels/MemAllocGroupingByTag.h
+++ b/Engine/Source/Developer/TraceInsights/Private/Insights/MemoryProfiler/ViewModels/MemAllocGroupingByTag.h
@@ -1,17 +1,20 @@
 ﻿// Copyright Epic Games, Inc. All Rights Reserved.
 #pragma once
 
-#include "CoreMinimal.h"
-#include "Insights/Common/SimpleRtti.h"
-#include "Insights/Table/ViewModels/TableTreeNode.h"
-#include "Insights/Table/ViewModels/TreeNodeGrouping.h"
-#include "TraceServices/Model/AllocationsProvider.h"
+#include "CoreTypes.h"
+
+// TraceServices
 #include "TraceServices/Model/Memory.h"
 
-namespace Insights
-{
+// TraceInsightsCore
+#include "InsightsCore/Common/SimpleRtti.h"
+#include "InsightsCore/Table/ViewModels/TableTreeNode.h"
+#include "InsightsCore/Table/ViewModels/TreeNodeGrouping.h"
 
-class InAsyncOperationProgress;
+namespace TraceServices { class IAllocationsProvider; }
+
+namespace UE::Insights::MemoryProfiler
+{
 
 ////////////////////////////////////////////////////////////////////////////////////////////////////
 
@@ -79,4 +82,5 @@ private:
 };
 
 ////////////////////////////////////////////////////////////////////////////////////////////////////
-}
+
+} // namespace UE::Insights::MemoryProfiler
diff --git a/Engine/Source/Developer/TraceInsights/Private/Insights/MemoryProfiler/ViewModels/MemAllocInSwapNode.cpp b/Engine/Source/Developer/TraceInsights/Private/Insights/MemoryProfiler/ViewModels/MemAllocInSwapNode.cpp
--- a/Engine/Source/Developer/TraceInsights/Private/Insights/MemoryProfiler/ViewModels/MemAllocInSwapNode.cpp
+++ b/Engine/Source/Developer/TraceInsights/Private/Insights/MemoryProfiler/ViewModels/MemAllocInSwapNode.cpp
@@ -1,10 +1,8 @@
-﻿// Copyright Epic Games, Inc. All Rights Reserved.
+// Copyright Epic Games, Inc. All Rights Reserved.
 
 #include "MemAllocInSwapNode.h"
 
-#define LOCTEXT_NAMESPACE "Insights::FMemAllocSwapNode"
-
-namespace Insights
+namespace UE::Insights::MemoryProfiler
 {
 
 INSIGHTS_IMPLEMENT_RTTI(FMemAllocInSwapNode)
@@ -18,6 +16,4 @@ uint64 FMemAllocInSwapNode::GetBytesInSwapPage() const
 
 ////////////////////////////////////////////////////////////////////////////////////////////////////
 
-} // namespace Insights
-
-#undef LOCTEXT_NAMESPACE
+} // namespace UE::Insights::MemoryProfiler
diff --git a/Engine/Source/Developer/TraceInsights/Private/Insights/MemoryProfiler/ViewModels/MemAllocInSwapNode.h b/Engine/Source/Developer/TraceInsights/Private/Insights/MemoryProfiler/ViewModels/MemAllocInSwapNode.h
--- a/Engine/Source/Developer/TraceInsights/Private/Insights/MemoryProfiler/ViewModels/MemAllocInSwapNode.h
+++ b/Engine/Source/Developer/TraceInsights/Private/Insights/MemoryProfiler/ViewModels/MemAllocInSwapNode.h
@@ -1,15 +1,12 @@
-﻿// Copyright Epic Games, Inc. All Rights Reserved.
+// Copyright Epic Games, Inc. All Rights Reserved.
 
 #pragma once
 
 #include "Insights/MemoryProfiler/ViewModels/MemAllocNode.h"
 
-////////////////////////////////////////////////////////////////////////////////////////////////////
-
-namespace Insights
+namespace UE::Insights::MemoryProfiler
 {
 
-////////////////////////////////////////////////////////////////////////////////////////////////////
 /**
 * Class used to store information about an allocation node in a particular swap page.
 */
@@ -31,6 +28,4 @@ private:
 	const uint64 BytesInSwapPage;
 };
 
-////////////////////////////////////////////////////////////////////////////////////////////////////
-
-} // namespace Insights
+} // namespace UE::Insights::MemoryProfiler
