[Insights] Massive refactoring of TraceInsights module:

diff --git a/Engine/Source/Developer/TraceInsights/Private/Insights/MemoryProfiler/ViewModels/MemAllocNode.cpp b/Engine/Source/Developer/TraceInsights/Private/Insights/MemoryProfiler/ViewModels/MemAllocNode.cpp
--- a/Engine/Source/Developer/TraceInsights/Private/Insights/MemoryProfiler/ViewModels/MemAllocNode.cpp
+++ b/Engine/Source/Developer/TraceInsights/Private/Insights/MemoryProfiler/ViewModels/MemAllocNode.cpp
@@ -2,12 +2,15 @@
 
 #include "MemAllocNode.h"
 
-// Insights
+// TraceServices
+#include "TraceServices/Model/Callstack.h"
+
+// TraceInsights
 #include "Insights/MemoryProfiler/ViewModels/CallstackFormatting.h"
 
-#define LOCTEXT_NAMESPACE "Insights::FMemAllocNode"
+#define LOCTEXT_NAMESPACE "UE::Insights::MemoryProfiler::FMemAllocNode"
 
-namespace Insights
+namespace UE::Insights::MemoryProfiler
 {
 
 INSIGHTS_IMPLEMENT_RTTI(FMemAllocNode)
@@ -77,7 +80,7 @@ FText FMemAllocNode::GetFullCallstackOrSourceFiles(ECallstackType InCallstackTyp
 		return FText();
 	}
 
-	const Insights::FMemoryAlloc& Alloc = GetMemAllocChecked();
+	const FMemoryAlloc& Alloc = GetMemAllocChecked();
 	const TraceServices::FCallstack* Callstack =
 		InCallstackType == ECallstackType::AllocCallstack
 			? Alloc.GetAllocCallstack()
@@ -125,7 +128,7 @@ FText FMemAllocNode::GetTopFunctionOrSourceFile(ECallstackType InCallstackType,
 		return FText();
 	}
 
-	const Insights::FMemoryAlloc& Alloc = GetMemAllocChecked();
+	const FMemoryAlloc& Alloc = GetMemAllocChecked();
 	const TraceServices::FCallstack* Callstack =
 		InCallstackType == ECallstackType::AllocCallstack
 			? Alloc.GetAllocCallstack()
@@ -172,6 +175,6 @@ FText FMemAllocNode::GetTopFunctionOrSourceFile(ECallstackType InCallstackType,
 
 ////////////////////////////////////////////////////////////////////////////////////////////////////
 
-} // namespace Insights
+} // namespace UE::Insights::MemoryProfiler
 
 #undef LOCTEXT_NAMESPACE
diff --git a/Engine/Source/Developer/TraceInsights/Private/Insights/MemoryProfiler/ViewModels/MemAllocNode.h b/Engine/Source/Developer/TraceInsights/Private/Insights/MemoryProfiler/ViewModels/MemAllocNode.h
--- a/Engine/Source/Developer/TraceInsights/Private/Insights/MemoryProfiler/ViewModels/MemAllocNode.h
+++ b/Engine/Source/Developer/TraceInsights/Private/Insights/MemoryProfiler/ViewModels/MemAllocNode.h
@@ -2,18 +2,24 @@
 
 #pragma once
 
-#include "CoreMinimal.h"
+#include "CoreTypes.h"
+
+#include "Templates/SharedPointer.h"
+
+// TraceServices
 #include "TraceServices/Model/Memory.h"
 
-// Insights
+// TraceInsightsCore
+#include "InsightsCore/Table/ViewModels/TableTreeNode.h"
+
+// TraceInsights
 #include "Insights/MemoryProfiler/ViewModels/MemAllocTable.h"
 #include "Insights/MemoryProfiler/ViewModels/MemoryAlloc.h"
 #include "Insights/MemoryProfiler/ViewModels/MemorySharedState.h"
-#include "Insights/Table/ViewModels/TableTreeNode.h"
 
 ////////////////////////////////////////////////////////////////////////////////////////////////////
 
-namespace Insights
+namespace UE::Insights::MemoryProfiler
 {
 
 ////////////////////////////////////////////////////////////////////////////////////////////////////
@@ -88,4 +94,4 @@ private:
 
 ////////////////////////////////////////////////////////////////////////////////////////////////////
 
-} // namespace Insights
+} // namespace UE::Insights::MemoryProfiler
diff --git a/Engine/Source/Developer/TraceInsights/Private/Insights/MemoryProfiler/ViewModels/MemAllocTable.cpp b/Engine/Source/Developer/TraceInsights/Private/Insights/MemoryProfiler/ViewModels/MemAllocTable.cpp
--- a/Engine/Source/Developer/TraceInsights/Private/Insights/MemoryProfiler/ViewModels/MemAllocTable.cpp
+++ b/Engine/Source/Developer/TraceInsights/Private/Insights/MemoryProfiler/ViewModels/MemAllocTable.cpp
@@ -7,22 +7,27 @@
 #include "Widgets/SToolTip.h"
 #include "Widgets/Text/STextBlock.h"
 
-// Insights
+// TraceServices
+#include "TraceServices/Model/Callstack.h"
+
+// TraceInsightsCore
+#include "InsightsCore/Filter/ViewModels/TimeFilterValueConverter.h"
+#include "InsightsCore/Table/ViewModels/TableCellValueFormatter.h"
+#include "InsightsCore/Table/ViewModels/TableCellValueGetter.h"
+#include "InsightsCore/Table/ViewModels/TableCellValueSorter.h"
+#include "InsightsCore/Table/ViewModels/TableColumn.h"
+
+// TraceInsights
 #include "Insights/MemoryProfiler/ViewModels/CallstackFormatting.h"
 #include "Insights/MemoryProfiler/ViewModels/MemAllocFilterValueConverter.h"
 #include "Insights/MemoryProfiler/ViewModels/MemAllocGroupingByTag.h"
 #include "Insights/MemoryProfiler/ViewModels/MemAllocNode.h"
 #include "Insights/MemoryProfiler/ViewModels/MemAllocInSwapNode.h"
 #include "Insights/MemoryProfiler/ViewModels/MemAllocTable.h"
-#include "Insights/Table/ViewModels/TableCellValueFormatter.h"
-#include "Insights/Table/ViewModels/TableCellValueGetter.h"
-#include "Insights/Table/ViewModels/TableCellValueSorter.h"
-#include "Insights/Table/ViewModels/TableColumn.h"
-#include "Insights/ViewModels/TimeFilterValueConverter.h"
 
-#define LOCTEXT_NAMESPACE "Insights::FMemAllocTable"
+#define LOCTEXT_NAMESPACE "UE::Insights::MemoryProfiler::FMemAllocTable"
 
-namespace Insights
+namespace UE::Insights::MemoryProfiler
 {
 
 ////////////////////////////////////////////////////////////////////////////////////////////////////
@@ -2014,6 +2019,6 @@ void FMemAllocTable::AddDefaultColumns()
 
 ////////////////////////////////////////////////////////////////////////////////////////////////////
 
-} // namespace Insights
+} // namespace UE::Insights::MemoryProfiler
 
 #undef LOCTEXT_NAMESPACE
diff --git a/Engine/Source/Developer/TraceInsights/Private/Insights/MemoryProfiler/ViewModels/MemAllocTable.h b/Engine/Source/Developer/TraceInsights/Private/Insights/MemoryProfiler/ViewModels/MemAllocTable.h
--- a/Engine/Source/Developer/TraceInsights/Private/Insights/MemoryProfiler/ViewModels/MemAllocTable.h
+++ b/Engine/Source/Developer/TraceInsights/Private/Insights/MemoryProfiler/ViewModels/MemAllocTable.h
@@ -2,16 +2,21 @@
 
 #pragma once
 
-#include "CoreMinimal.h"
+#include "CoreTypes.h"
 
-// Insights
+#include "Containers/Array.h"
+#include "UObject/NameTypes.h"
+
+// TraceInsightsCore
+#include "InsightsCore/Table/ViewModels/Table.h"
+
+// TraceInsights
 #include "Insights/MemoryProfiler/ViewModels/MemoryAlloc.h"
-#include "Insights/Table/ViewModels/Table.h"
 
-namespace Insights
-{
+namespace UE::Insights { class FTableColumn; }
 
-class FTableColumn;
+namespace UE::Insights::MemoryProfiler
+{
 
 ////////////////////////////////////////////////////////////////////////////////////////////////////
 
@@ -84,4 +89,4 @@ private:
 
 ////////////////////////////////////////////////////////////////////////////////////////////////////
 
-} // namespace Insights
+} // namespace UE::Insights::MemoryProfiler
diff --git a/Engine/Source/Developer/TraceInsights/Private/Insights/MemoryProfiler/ViewModels/MemTagNode.cpp b/Engine/Source/Developer/TraceInsights/Private/Insights/MemoryProfiler/ViewModels/MemTagNode.cpp
--- a/Engine/Source/Developer/TraceInsights/Private/Insights/MemoryProfiler/ViewModels/MemTagNode.cpp
+++ b/Engine/Source/Developer/TraceInsights/Private/Insights/MemoryProfiler/ViewModels/MemTagNode.cpp
@@ -2,10 +2,13 @@
 
 #include "MemTagNode.h"
 
-// Insights
+// TraceInsights
 #include "Insights/MemoryProfiler/MemoryProfilerManager.h"
 
-#define LOCTEXT_NAMESPACE "FMemTagNode"
+#define LOCTEXT_NAMESPACE "UE::Insights::MemoryProfiler::FMemTagNode"
+
+namespace UE::Insights::MemoryProfiler
+{
 
 INSIGHTS_IMPLEMENT_RTTI(FMemTagNode)
 
@@ -16,8 +19,8 @@ FText FMemTagNode::GetTrackerText() const
 	FMemorySharedState* SharedState = FMemoryProfilerManager::Get()->GetSharedState();
 	if (SharedState)
 	{
-		Insights::FMemoryTrackerId TrackerId = GetMemTrackerId();
-		const Insights::FMemoryTracker* Tracker = SharedState->GetTrackerById(TrackerId);
+		FMemoryTrackerId TrackerId = GetMemTrackerId();
+		const FMemoryTracker* Tracker = SharedState->GetTrackerById(TrackerId);
 		if (Tracker)
 		{
 			return FText::FromString(Tracker->GetName());
@@ -43,4 +46,6 @@ void FMemTagNode::SetAggregatedStats(const TraceServices::FMemoryProfilerAggrega
 */
 ////////////////////////////////////////////////////////////////////////////////////////////////////
 
+} // namespace UE::Insights::MemoryProfiler
+
 #undef LOCTEXT_NAMESPACE
diff --git a/Engine/Source/Developer/TraceInsights/Private/Insights/MemoryProfiler/ViewModels/MemTagNode.h b/Engine/Source/Developer/TraceInsights/Private/Insights/MemoryProfiler/ViewModels/MemTagNode.h
--- a/Engine/Source/Developer/TraceInsights/Private/Insights/MemoryProfiler/ViewModels/MemTagNode.h
+++ b/Engine/Source/Developer/TraceInsights/Private/Insights/MemoryProfiler/ViewModels/MemTagNode.h
@@ -3,13 +3,17 @@
 #pragma once
 
 #include "CoreMinimal.h"
+
+// TraceServices
 #include "TraceServices/Model/Memory.h"
 
-// Insights
+// TraceInsightsCore
+#include "InsightsCore/Table/ViewModels/BaseTreeNode.h"
+
+// TraceInsights
 #include "Insights/MemoryProfiler/ViewModels/MemorySharedState.h"
 #include "Insights/MemoryProfiler/ViewModels/MemoryTag.h"
 #include "Insights/MemoryProfiler/ViewModels/MemoryTracker.h"
-#include "Insights/Table/ViewModels/BaseTreeNode.h"
 
 namespace TraceServices
 {
@@ -24,6 +28,9 @@ namespace TraceServices
 	};
 }
 
+namespace UE::Insights::MemoryProfiler
+{
+
 ////////////////////////////////////////////////////////////////////////////////////////////////////
 
 enum class EMemTagNodeType
@@ -58,13 +65,13 @@ typedef TWeakPtr<class FMemTagNode> FMemTagNodeWeak;
 /**
  * Class used to store information about an llm tag node (used in the SMemTagTreeView).
  */
-class FMemTagNode : public Insights::FBaseTreeNode
+class FMemTagNode : public FBaseTreeNode
 {
 	INSIGHTS_DECLARE_RTTI(FMemTagNode, FBaseTreeNode)
 
 public:
 	/** Initialization constructor for the MemTag node. */
-	explicit FMemTagNode(Insights::FMemoryTag* InMemTag)
+	explicit FMemTagNode(FMemoryTag* InMemTag)
 		: FBaseTreeNode(FName(InMemTag->GetStatFullName(), 0), false)
 		, Type(EMemTagNodeType::MemTag)
 		, MemTag(InMemTag)
@@ -87,18 +94,18 @@ public:
 	EMemTagNodeType GetType() const { return Type; }
 
 	bool IsValidStat() const { return MemTag != nullptr; }
-	Insights::FMemoryTag* GetMemTag() const { return MemTag; }
+	FMemoryTag* GetMemTag() const { return MemTag; }
 
-	Insights::FMemoryTagId GetMemTagId() const { return MemTag ? MemTag->GetId() : Insights::FMemoryTag::InvalidTagId; }
+	FMemoryTagId GetMemTagId() const { return MemTag ? MemTag->GetId() : FMemoryTag::InvalidTagId; }
 
-	Insights::FMemoryTrackerId GetMemTrackerId() const { return MemTag ? MemTag->GetTrackerId() : Insights::FMemoryTracker::InvalidTrackerId; }
+	FMemoryTrackerId GetMemTrackerId() const { return MemTag ? MemTag->GetTrackerId() : FMemoryTracker::InvalidTrackerId; }
 	FText GetTrackerText() const;
 
 	FLinearColor GetColor() const { return MemTag ? MemTag->GetColor() : FLinearColor(0.5f, 0.5f, 0.5f, 1.0f); }
 	bool IsAddedToGraph() const { return MemTag ? MemTag->IsAddedToGraph() : false; }
 
 	FMemTagNodePtr GetParentTagNode() const { return ParentTagNode; }
-	Insights::FMemoryTag* GetParentMemTag() const { return ParentTagNode.IsValid() ? ParentTagNode->GetMemTag() : nullptr; }
+	FMemoryTag* GetParentMemTag() const { return ParentTagNode.IsValid() ? ParentTagNode->GetMemTag() : nullptr; }
 	void SetParentTagNode(FMemTagNodePtr NodePtr) { ParentTagNode = NodePtr; }
 
 	/**
@@ -111,9 +118,11 @@ public:
 
 private:
 	const EMemTagNodeType Type;
-	Insights::FMemoryTag* MemTag;
+	FMemoryTag* MemTag;
 	FMemTagNodePtr ParentTagNode;
 	TraceServices::FMemoryProfilerAggregatedStats AggregatedStats;
 };
 
 ////////////////////////////////////////////////////////////////////////////////////////////////////
+
+} // namespace UE::Insights::MemoryProfiler
diff --git a/Engine/Source/Developer/TraceInsights/Private/Insights/MemoryProfiler/ViewModels/MemTagNodeGroupingAndSorting.cpp b/Engine/Source/Developer/TraceInsights/Private/Insights/MemoryProfiler/ViewModels/MemTagNodeGroupingAndSorting.cpp
--- a/Engine/Source/Developer/TraceInsights/Private/Insights/MemoryProfiler/ViewModels/MemTagNodeGroupingAndSorting.cpp
+++ b/Engine/Source/Developer/TraceInsights/Private/Insights/MemoryProfiler/ViewModels/MemTagNodeGroupingAndSorting.cpp
@@ -2,10 +2,13 @@
 
 #include "MemTagNodeGroupingAndSorting.h"
 
-#include "Insights/Table/ViewModels/TableColumn.h"
+// TraceInsightsCore
+#include "InsightsCore/Table/ViewModels/TableColumn.h"
+
+// TraceInsights
 #include "Insights/MemoryProfiler/ViewModels/MemTagNodeHelper.h"
 
-#define LOCTEXT_NAMESPACE "MemTagNode"
+#define LOCTEXT_NAMESPACE "UE::Insights::MemoryProfiler::FMemTagNode"
 
 #define INSIGHTS_ENSURE ensure
 //#define INSIGHTS_ENSURE(...)
@@ -27,12 +30,15 @@
 #define INSIGHTS_DEFAULT_SORTING_NODES(A, B) return A->GetName().LexicalLess(B->GetName());
 //#define INSIGHTS_DEFAULT_SORTING_NODES(A, B) return A->GetDefaultSortOrder() < B->GetDefaultSortOrder();
 
+namespace UE::Insights::MemoryProfiler
+{
+
 ////////////////////////////////////////////////////////////////////////////////////////////////////
 // Sorting by Type
 ////////////////////////////////////////////////////////////////////////////////////////////////////
 
-FMemTagNodeSortingByType::FMemTagNodeSortingByType(TSharedRef<Insights::FTableColumn> InColumnRef)
-	: Insights::FTableCellValueSorter(
+FMemTagNodeSortingByType::FMemTagNodeSortingByType(TSharedRef<FTableColumn> InColumnRef)
+	: FTableCellValueSorter(
 		FName(TEXT("ByType")),
 		LOCTEXT("Sorting_ByType_Name", "By Type"),
 		LOCTEXT("Sorting_ByType_Title", "Sort By Type"),
@@ -43,19 +49,19 @@ FMemTagNodeSortingByType::FMemTagNodeSortingByType(TSharedRef<Insights::FTableCo
 
 ////////////////////////////////////////////////////////////////////////////////////////////////////
 
-void FMemTagNodeSortingByType::Sort(TArray<Insights::FBaseTreeNodePtr>& NodesToSort, Insights::ESortMode SortMode) const
+void FMemTagNodeSortingByType::Sort(TArray<FBaseTreeNodePtr>& NodesToSort, ESortMode SortMode) const
 {
-	if (SortMode == Insights::ESortMode::Ascending)
+	if (SortMode == ESortMode::Ascending)
 	{
-		NodesToSort.Sort([this](const Insights::FBaseTreeNodePtr& A, const Insights::FBaseTreeNodePtr& B) -> bool
+		NodesToSort.Sort([this](const FBaseTreeNodePtr& A, const FBaseTreeNodePtr& B) -> bool
 		{
 			INSIGHTS_DEFAULT_PRESORTING_NODES(A, B);
 
 			INSIGHTS_ENSURE(A.IsValid() && A->Is<FMemTagNode>());
-			const FMemTagNodePtr MemTagNodeA = StaticCastSharedPtr<FMemTagNode, Insights::FBaseTreeNode>(A);
+			const FMemTagNodePtr MemTagNodeA = StaticCastSharedPtr<FMemTagNode, FBaseTreeNode>(A);
 
 			INSIGHTS_ENSURE(B.IsValid() && B->Is<FMemTagNode>());
-			const FMemTagNodePtr MemTagNodeB = StaticCastSharedPtr<FMemTagNode, Insights::FBaseTreeNode>(B);
+			const FMemTagNodePtr MemTagNodeB = StaticCastSharedPtr<FMemTagNode, FBaseTreeNode>(B);
 
 			if (MemTagNodeA->GetType() == MemTagNodeB->GetType())
 			{
@@ -70,15 +76,15 @@ void FMemTagNodeSortingByType::Sort(TArray<Insights::FBaseTreeNodePtr>& NodesToS
 	}
 	else // if (SortMode == ESortMode::Descending)
 	{
-		NodesToSort.Sort([this](const Insights::FBaseTreeNodePtr& A, const Insights::FBaseTreeNodePtr& B) -> bool
+		NodesToSort.Sort([this](const FBaseTreeNodePtr& A, const FBaseTreeNodePtr& B) -> bool
 		{
 			INSIGHTS_DEFAULT_PRESORTING_NODES(A, B);
 
 			INSIGHTS_ENSURE(A.IsValid() && A->Is<FMemTagNode>());
-			const FMemTagNodePtr MemTagNodeA = StaticCastSharedPtr<FMemTagNode, Insights::FBaseTreeNode>(A);
+			const FMemTagNodePtr MemTagNodeA = StaticCastSharedPtr<FMemTagNode, FBaseTreeNode>(A);
 
 			INSIGHTS_ENSURE(B.IsValid() && B->Is<FMemTagNode>());
-			const FMemTagNodePtr MemTagNodeB = StaticCastSharedPtr<FMemTagNode, Insights::FBaseTreeNode>(B);
+			const FMemTagNodePtr MemTagNodeB = StaticCastSharedPtr<FMemTagNode, FBaseTreeNode>(B);
 
 			if (MemTagNodeA->GetType() == MemTagNodeB->GetType())
 			{
@@ -97,8 +103,8 @@ void FMemTagNodeSortingByType::Sort(TArray<Insights::FBaseTreeNodePtr>& NodesToS
 // Sorting by Tracker(s)
 ////////////////////////////////////////////////////////////////////////////////////////////////////
 
-FMemTagNodeSortingByTracker::FMemTagNodeSortingByTracker(TSharedRef<Insights::FTableColumn> InColumnRef)
-	: Insights::FTableCellValueSorter(
+FMemTagNodeSortingByTracker::FMemTagNodeSortingByTracker(TSharedRef<FTableColumn> InColumnRef)
+	: FTableCellValueSorter(
 		FName(TEXT("ByTracker")),
 		LOCTEXT("Sorting_ByTracker_Name", "By Tracker"),
 		LOCTEXT("Sorting_ByTracker_Title", "Sort By Tracker"),
@@ -109,19 +115,19 @@ FMemTagNodeSortingByTracker::FMemTagNodeSortingByTracker(TSharedRef<Insights::FT
 
 ////////////////////////////////////////////////////////////////////////////////////////////////////
 
-void FMemTagNodeSortingByTracker::Sort(TArray<Insights::FBaseTreeNodePtr>& NodesToSort, Insights::ESortMode SortMode) const
+void FMemTagNodeSortingByTracker::Sort(TArray<FBaseTreeNodePtr>& NodesToSort, ESortMode SortMode) const
 {
-	if (SortMode == Insights::ESortMode::Ascending)
+	if (SortMode == ESortMode::Ascending)
 	{
-		NodesToSort.Sort([this](const Insights::FBaseTreeNodePtr& A, const Insights::FBaseTreeNodePtr& B) -> bool
+		NodesToSort.Sort([this](const FBaseTreeNodePtr& A, const FBaseTreeNodePtr& B) -> bool
 		{
 			INSIGHTS_DEFAULT_PRESORTING_NODES(A, B);
 
 			INSIGHTS_ENSURE(A.IsValid() && A->Is<FMemTagNode>());
-			const FMemTagNodePtr MemTagNodeA = StaticCastSharedPtr<FMemTagNode, Insights::FBaseTreeNode>(A);
+			const FMemTagNodePtr MemTagNodeA = StaticCastSharedPtr<FMemTagNode, FBaseTreeNode>(A);
 
 			INSIGHTS_ENSURE(B.IsValid() && B->Is<FMemTagNode>());
-			const FMemTagNodePtr MemTagNodeB = StaticCastSharedPtr<FMemTagNode, Insights::FBaseTreeNode>(B);
+			const FMemTagNodePtr MemTagNodeB = StaticCastSharedPtr<FMemTagNode, FBaseTreeNode>(B);
 
 			if (MemTagNodeA->GetMemTrackerId() == MemTagNodeB->GetMemTrackerId())
 			{
@@ -136,15 +142,15 @@ void FMemTagNodeSortingByTracker::Sort(TArray<Insights::FBaseTreeNodePtr>& Nodes
 	}
 	else // if (SortMode == ESortMode::Descending)
 	{
-		NodesToSort.Sort([this](const Insights::FBaseTreeNodePtr& A, const Insights::FBaseTreeNodePtr& B) -> bool
+		NodesToSort.Sort([this](const FBaseTreeNodePtr& A, const FBaseTreeNodePtr& B) -> bool
 		{
 			INSIGHTS_DEFAULT_PRESORTING_NODES(A, B);
 
 			INSIGHTS_ENSURE(A.IsValid() && A->Is<FMemTagNode>());
-			const FMemTagNodePtr MemTagNodeA = StaticCastSharedPtr<FMemTagNode, Insights::FBaseTreeNode>(A);
+			const FMemTagNodePtr MemTagNodeA = StaticCastSharedPtr<FMemTagNode, FBaseTreeNode>(A);
 
 			INSIGHTS_ENSURE(B.IsValid() && B->Is<FMemTagNode>());
-			const FMemTagNodePtr MemTagNodeB = StaticCastSharedPtr<FMemTagNode, Insights::FBaseTreeNode>(B);
+			const FMemTagNodePtr MemTagNodeB = StaticCastSharedPtr<FMemTagNode, FBaseTreeNode>(B);
 
 			if (MemTagNodeA->GetMemTrackerId() == MemTagNodeB->GetMemTrackerId())
 			{
@@ -163,8 +169,8 @@ void FMemTagNodeSortingByTracker::Sort(TArray<Insights::FBaseTreeNodePtr>& Nodes
 // Sort by Instance Count
 ////////////////////////////////////////////////////////////////////////////////////////////////////
 
-FMemTagNodeSortingByInstanceCount::FMemTagNodeSortingByInstanceCount(TSharedRef<Insights::FTableColumn> InColumnRef)
-	: Insights::FTableCellValueSorter(
+FMemTagNodeSortingByInstanceCount::FMemTagNodeSortingByInstanceCount(TSharedRef<FTableColumn> InColumnRef)
+	: FTableCellValueSorter(
 		FName(TEXT("ByInstanceCount")),
 		LOCTEXT("Sorting_ByInstanceCount_Name", "By Instance Count"),
 		LOCTEXT("Sorting_ByInstanceCount_Title", "Sort By Instance Count"),
@@ -175,20 +181,20 @@ FMemTagNodeSortingByInstanceCount::FMemTagNodeSortingByInstanceCount(TSharedRef<
 
 ////////////////////////////////////////////////////////////////////////////////////////////////////
 
-void FMemTagNodeSortingByInstanceCount::Sort(TArray<Insights::FBaseTreeNodePtr>& NodesToSort, Insights::ESortMode SortMode) const
+void FMemTagNodeSortingByInstanceCount::Sort(TArray<FBaseTreeNodePtr>& NodesToSort, ESortMode SortMode) const
 {
-	if (SortMode == Insights::ESortMode::Ascending)
+	if (SortMode == ESortMode::Ascending)
 	{
-		NodesToSort.Sort([this](const Insights::FBaseTreeNodePtr& A, const Insights::FBaseTreeNodePtr& B) -> bool
+		NodesToSort.Sort([this](const FBaseTreeNodePtr& A, const FBaseTreeNodePtr& B) -> bool
 		{
 			INSIGHTS_DEFAULT_PRESORTING_NODES(A, B);
 
 			INSIGHTS_ENSURE(A.IsValid() && A->Is<FMemTagNode>());
-			const FMemTagNodePtr MemTagNodeA = StaticCastSharedPtr<FMemTagNode, Insights::FBaseTreeNode>(A);
+			const FMemTagNodePtr MemTagNodeA = StaticCastSharedPtr<FMemTagNode, FBaseTreeNode>(A);
 			const uint64 ValueA = MemTagNodeA->GetAggregatedStats().InstanceCount;
 
 			INSIGHTS_ENSURE(B.IsValid() && B->Is<FMemTagNode>());
-			const FMemTagNodePtr MemTagNodeB = StaticCastSharedPtr<FMemTagNode, Insights::FBaseTreeNode>(B);
+			const FMemTagNodePtr MemTagNodeB = StaticCastSharedPtr<FMemTagNode, FBaseTreeNode>(B);
 			const uint64 ValueB = MemTagNodeB->GetAggregatedStats().InstanceCount;
 
 			if (ValueA == ValueB)
@@ -204,16 +210,16 @@ void FMemTagNodeSortingByInstanceCount::Sort(TArray<Insights::FBaseTreeNodePtr>&
 	}
 	else // if (SortMode == ESortMode::Descending)
 	{
-		NodesToSort.Sort([this](const Insights::FBaseTreeNodePtr& A, const Insights::FBaseTreeNodePtr& B) -> bool
+		NodesToSort.Sort([this](const FBaseTreeNodePtr& A, const FBaseTreeNodePtr& B) -> bool
 		{
 			INSIGHTS_DEFAULT_PRESORTING_NODES(A, B);
 
 			INSIGHTS_ENSURE(A.IsValid() && A->Is<FMemTagNode>());
-			const FMemTagNodePtr MemTagNodeA = StaticCastSharedPtr<FMemTagNode, Insights::FBaseTreeNode>(A);
+			const FMemTagNodePtr MemTagNodeA = StaticCastSharedPtr<FMemTagNode, FBaseTreeNode>(A);
 			const uint64 ValueA = MemTagNodeA->GetAggregatedStats().InstanceCount;
 
 			INSIGHTS_ENSURE(B.IsValid() && B->Is<FMemTagNode>());
-			const FMemTagNodePtr MemTagNodeB = StaticCastSharedPtr<FMemTagNode, Insights::FBaseTreeNode>(B);
+			const FMemTagNodePtr MemTagNodeB = StaticCastSharedPtr<FMemTagNode, FBaseTreeNode>(B);
 			const uint64 ValueB = MemTagNodeB->GetAggregatedStats().InstanceCount;
 
 			if (ValueA == ValueB)
@@ -231,6 +237,8 @@ void FMemTagNodeSortingByInstanceCount::Sort(TArray<Insights::FBaseTreeNodePtr>&
 
 ////////////////////////////////////////////////////////////////////////////////////////////////////
 
+} // namespace UE::Insights::MemoryProfiler
+
 #undef INSIGHTS_DEFAULT_SORTING_NODES
 #undef INSIGHTS_DEFAULT_PRESORTING_NODES
 #undef INSIGHTS_ENSURE
diff --git a/Engine/Source/Developer/TraceInsights/Private/Insights/MemoryProfiler/ViewModels/MemTagNodeGroupingAndSorting.h b/Engine/Source/Developer/TraceInsights/Private/Insights/MemoryProfiler/ViewModels/MemTagNodeGroupingAndSorting.h
--- a/Engine/Source/Developer/TraceInsights/Private/Insights/MemoryProfiler/ViewModels/MemTagNodeGroupingAndSorting.h
+++ b/Engine/Source/Developer/TraceInsights/Private/Insights/MemoryProfiler/ViewModels/MemTagNodeGroupingAndSorting.h
@@ -2,64 +2,71 @@
 
 #pragma once
 
-#include "CoreMinimal.h"
+#include "CoreTypes.h"
+#include "Containers/Array.h"
+#include "Templates/SharedPointer.h"
 
-// Insights
+// TraceInsightsCore
+#include "InsightsCore/Table/ViewModels/TableColumn.h"
+#include "InsightsCore/Table/ViewModels/TableCellValueSorter.h"
+#include "InsightsCore/Table/ViewModels/TreeNodeGrouping.h"
+
+// TraceInsights
 #include "Insights/MemoryProfiler/ViewModels/MemTagNode.h"
-#include "Insights/Table/ViewModels/TableColumn.h"
-#include "Insights/Table/ViewModels/TableCellValueSorter.h"
-#include "Insights/Table/ViewModels/TreeNodeGrouping.h"
+
+namespace UE::Insights::MemoryProfiler
+{
 
 ////////////////////////////////////////////////////////////////////////////////////////////////////
 // Sorters
 ////////////////////////////////////////////////////////////////////////////////////////////////////
 
-class FMemTagNodeSortingByType: public Insights::FTableCellValueSorter
+class FMemTagNodeSortingByType: public FTableCellValueSorter
 {
 public:
-	FMemTagNodeSortingByType(TSharedRef<Insights::FTableColumn> InColumnRef);
+	FMemTagNodeSortingByType(TSharedRef<FTableColumn> InColumnRef);
 
-	virtual void Sort(TArray<Insights::FBaseTreeNodePtr>& NodesToSort, Insights::ESortMode SortMode) const override;
+	virtual void Sort(TArray<FBaseTreeNodePtr>& NodesToSort, ESortMode SortMode) const override;
 };
 
 ////////////////////////////////////////////////////////////////////////////////////////////////////
 
-class FMemTagNodeSortingByTracker : public Insights::FTableCellValueSorter
+class FMemTagNodeSortingByTracker : public FTableCellValueSorter
 {
 public:
-	FMemTagNodeSortingByTracker(TSharedRef<Insights::FTableColumn> InColumnRef);
+	FMemTagNodeSortingByTracker(TSharedRef<FTableColumn> InColumnRef);
 
-	virtual void Sort(TArray<Insights::FBaseTreeNodePtr>& NodesToSort, Insights::ESortMode SortMode) const override;
+	virtual void Sort(TArray<FBaseTreeNodePtr>& NodesToSort, ESortMode SortMode) const override;
 };
 
 ////////////////////////////////////////////////////////////////////////////////////////////////////
 
-class FMemTagNodeSortingByInstanceCount : public Insights::FTableCellValueSorter
+class FMemTagNodeSortingByInstanceCount : public FTableCellValueSorter
 {
 public:
-	FMemTagNodeSortingByInstanceCount(TSharedRef<Insights::FTableColumn> InColumnRef);
+	FMemTagNodeSortingByInstanceCount(TSharedRef<FTableColumn> InColumnRef);
 
-	virtual void Sort(TArray<Insights::FBaseTreeNodePtr>& NodesToSort, Insights::ESortMode SortMode) const override;
+	virtual void Sort(TArray<FBaseTreeNodePtr>& NodesToSort, ESortMode SortMode) const override;
 };
 
 ////////////////////////////////////////////////////////////////////////////////////////////////////
 
-class FMemTagNodeSortingByTotalInclusiveSize : public Insights::FTableCellValueSorter
+class FMemTagNodeSortingByTotalInclusiveSize : public FTableCellValueSorter
 {
 public:
-	FMemTagNodeSortingByTotalInclusiveSize(TSharedRef<Insights::FTableColumn> InColumnRef);
+	FMemTagNodeSortingByTotalInclusiveSize(TSharedRef<FTableColumn> InColumnRef);
 
-	virtual void Sort(TArray<Insights::FBaseTreeNodePtr>& NodesToSort, Insights::ESortMode SortMode) const override;
+	virtual void Sort(TArray<FBaseTreeNodePtr>& NodesToSort, ESortMode SortMode) const override;
 };
 
 ////////////////////////////////////////////////////////////////////////////////////////////////////
 
-class FMemTagNodeSortingByTotalExclusiveSize : public Insights::FTableCellValueSorter
+class FMemTagNodeSortingByTotalExclusiveSize : public FTableCellValueSorter
 {
 public:
-	FMemTagNodeSortingByTotalExclusiveSize(TSharedRef<Insights::FTableColumn> InColumnRef);
+	FMemTagNodeSortingByTotalExclusiveSize(TSharedRef<FTableColumn> InColumnRef);
 
-	virtual void Sort(TArray<Insights::FBaseTreeNodePtr>& NodesToSort, Insights::ESortMode SortMode) const override;
+	virtual void Sort(TArray<FBaseTreeNodePtr>& NodesToSort, ESortMode SortMode) const override;
 };
 
 ////////////////////////////////////////////////////////////////////////////////////////////////////
@@ -92,3 +99,5 @@ enum class EMemTagNodeGroupingMode
 typedef TSharedPtr<EMemTagNodeGroupingMode> EMemTagNodeGroupingModePtr;
 
 ////////////////////////////////////////////////////////////////////////////////////////////////////
+
+} // namespace UE::Insights::MemoryProfiler
diff --git a/Engine/Source/Developer/TraceInsights/Private/Insights/MemoryProfiler/ViewModels/MemTagNodeHelper.cpp b/Engine/Source/Developer/TraceInsights/Private/Insights/MemoryProfiler/ViewModels/MemTagNodeHelper.cpp
--- a/Engine/Source/Developer/TraceInsights/Private/Insights/MemoryProfiler/ViewModels/MemTagNodeHelper.cpp
+++ b/Engine/Source/Developer/TraceInsights/Private/Insights/MemoryProfiler/ViewModels/MemTagNodeHelper.cpp
@@ -2,10 +2,13 @@
 
 #include "MemTagNodeHelper.h"
 
-// Insights
+// TraceInsights
 #include "Insights/InsightsStyle.h"
 
-#define LOCTEXT_NAMESPACE "MemTagNode"
+#define LOCTEXT_NAMESPACE "UE::Insights::MemoryProfiler::FMemTagNode"
+
+namespace UE::Insights::MemoryProfiler
+{
 
 ////////////////////////////////////////////////////////////////////////////////////////////////////
 // MemTagNode Type Helper
@@ -91,4 +94,6 @@ FText MemTagNodeGroupingHelper::ToDescription(const EMemTagNodeGroupingMode Grou
 
 ////////////////////////////////////////////////////////////////////////////////////////////////////
 
+} // namespace UE::Insights::MemoryProfiler
+
 #undef LOCTEXT_NAMESPACE
diff --git a/Engine/Source/Developer/TraceInsights/Private/Insights/MemoryProfiler/ViewModels/MemTagNodeHelper.h b/Engine/Source/Developer/TraceInsights/Private/Insights/MemoryProfiler/ViewModels/MemTagNodeHelper.h
--- a/Engine/Source/Developer/TraceInsights/Private/Insights/MemoryProfiler/ViewModels/MemTagNodeHelper.h
+++ b/Engine/Source/Developer/TraceInsights/Private/Insights/MemoryProfiler/ViewModels/MemTagNodeHelper.h
@@ -2,13 +2,19 @@
 
 #pragma once
 
-#include "CoreMinimal.h"
-#include "Styling/SlateBrush.h"
+#include "CoreTypes.h"
 
-// Insights
+#include "Internationalization/Text.h"
+
+// TraceInsights
 #include "Insights/MemoryProfiler/ViewModels/MemTagNode.h"
 #include "Insights/MemoryProfiler/ViewModels/MemTagNodeGroupingAndSorting.h"
 
+struct FSlateBrush;
+
+namespace UE::Insights::MemoryProfiler
+{
+
 ////////////////////////////////////////////////////////////////////////////////////////////////////
 
 /** Helper struct that contains helper functions and classes for EMemTagNodeType enum. */
@@ -53,3 +59,5 @@ struct MemTagNodeGroupingHelper
 };
 
 ////////////////////////////////////////////////////////////////////////////////////////////////////
+
+} // namespace UE::Insights::MemoryProfiler
diff --git a/Engine/Source/Developer/TraceInsights/Private/Insights/MemoryProfiler/ViewModels/MemTagTreeViewColumnFactory.cpp b/Engine/Source/Developer/TraceInsights/Private/Insights/MemoryProfiler/ViewModels/MemTagTreeViewColumnFactory.cpp
--- a/Engine/Source/Developer/TraceInsights/Private/Insights/MemoryProfiler/ViewModels/MemTagTreeViewColumnFactory.cpp
+++ b/Engine/Source/Developer/TraceInsights/Private/Insights/MemoryProfiler/ViewModels/MemTagTreeViewColumnFactory.cpp
@@ -2,14 +2,19 @@
 
 #include "MemTagTreeViewColumnFactory.h"
 
-// Insights
-#include "Insights/Table/ViewModels/TableCellValueFormatter.h"
-#include "Insights/Table/ViewModels/TableCellValueGetter.h"
-#include "Insights/Table/ViewModels/TableCellValueSorter.h"
+// TraceInsightsCore
+#include "InsightsCore/Table/ViewModels/TableCellValueFormatter.h"
+#include "InsightsCore/Table/ViewModels/TableCellValueGetter.h"
+#include "InsightsCore/Table/ViewModels/TableCellValueSorter.h"
+
+// TraceInsights
 #include "Insights/MemoryProfiler/ViewModels/MemTagNodeGroupingAndSorting.h"
 #include "Insights/MemoryProfiler/ViewModels/MemTagNodeHelper.h"
 
-#define LOCTEXT_NAMESPACE "SMemTagTreeView"
+#define LOCTEXT_NAMESPACE "UE::Insights::MemoryProfiler::SMemTagTreeView"
+
+namespace UE::Insights::MemoryProfiler
+{
 
 ////////////////////////////////////////////////////////////////////////////////////////////////////
 // Column identifiers
@@ -24,7 +29,7 @@ const FName FMemTagTreeViewColumns::AverageValueColumnID(TEXT("Average"));
 
 ////////////////////////////////////////////////////////////////////////////////////////////////////
 
-void FMemTagTreeViewColumnFactory::CreateMemTagTreeViewColumns(TArray<TSharedRef<Insights::FTableColumn>>& Columns)
+void FMemTagTreeViewColumnFactory::CreateMemTagTreeViewColumns(TArray<TSharedRef<FTableColumn>>& Columns)
 {
 	Columns.Reset();
 
@@ -39,10 +44,8 @@ void FMemTagTreeViewColumnFactory::CreateMemTagTreeViewColumns(TArray<TSharedRef
 
 ////////////////////////////////////////////////////////////////////////////////////////////////////
 
-TSharedRef<Insights::FTableColumn> FMemTagTreeViewColumnFactory::CreateNameColumn()
+TSharedRef<FTableColumn> FMemTagTreeViewColumnFactory::CreateNameColumn()
 {
-	using namespace Insights;
-
 	TSharedRef<FTableColumn> ColumnRef = MakeShared<FTableColumn>(FMemTagTreeViewColumns::NameColumnID);
 	FTableColumn& Column = *ColumnRef;
 
@@ -74,10 +77,8 @@ TSharedRef<Insights::FTableColumn> FMemTagTreeViewColumnFactory::CreateNameColum
 
 ////////////////////////////////////////////////////////////////////////////////////////////////////
 
-TSharedRef<Insights::FTableColumn> FMemTagTreeViewColumnFactory::CreateTypeColumn()
+TSharedRef<FTableColumn> FMemTagTreeViewColumnFactory::CreateTypeColumn()
 {
-	using namespace Insights;
-
 	TSharedRef<FTableColumn> ColumnRef = MakeShared<FTableColumn>(FMemTagTreeViewColumns::TypeColumnID);
 	FTableColumn& Column = *ColumnRef;
 
@@ -120,10 +121,8 @@ TSharedRef<Insights::FTableColumn> FMemTagTreeViewColumnFactory::CreateTypeColum
 
 ////////////////////////////////////////////////////////////////////////////////////////////////////
 
-TSharedRef<Insights::FTableColumn> FMemTagTreeViewColumnFactory::CreateTrackerColumn()
+TSharedRef<FTableColumn> FMemTagTreeViewColumnFactory::CreateTrackerColumn()
 {
-	using namespace Insights;
-
 	TSharedRef<FTableColumn> ColumnRef = MakeShared<FTableColumn>(FMemTagTreeViewColumns::TrackerColumnID);
 	FTableColumn& Column = *ColumnRef;
 
@@ -165,10 +164,8 @@ TSharedRef<Insights::FTableColumn> FMemTagTreeViewColumnFactory::CreateTrackerCo
 
 ////////////////////////////////////////////////////////////////////////////////////////////////////
 
-TSharedRef<Insights::FTableColumn> FMemTagTreeViewColumnFactory::CreateInstanceCountColumn()
+TSharedRef<FTableColumn> FMemTagTreeViewColumnFactory::CreateInstanceCountColumn()
 {
-	using namespace Insights;
-
 	TSharedRef<FTableColumn> ColumnRef = MakeShared<FTableColumn>(FMemTagTreeViewColumns::InstanceCountColumnID);
 	FTableColumn& Column = *ColumnRef;
 
@@ -211,10 +208,8 @@ TSharedRef<Insights::FTableColumn> FMemTagTreeViewColumnFactory::CreateInstanceC
 
 ////////////////////////////////////////////////////////////////////////////////////////////////////
 
-TSharedRef<Insights::FTableColumn> FMemTagTreeViewColumnFactory::CreateMinValueColumn()
+TSharedRef<FTableColumn> FMemTagTreeViewColumnFactory::CreateMinValueColumn()
 {
-	using namespace Insights;
-
 	TSharedRef<FTableColumn> ColumnRef = MakeShared<FTableColumn>(FMemTagTreeViewColumns::MinValueColumnID);
 	FTableColumn& Column = *ColumnRef;
 
@@ -256,10 +251,8 @@ TSharedRef<Insights::FTableColumn> FMemTagTreeViewColumnFactory::CreateMinValueC
 
 ////////////////////////////////////////////////////////////////////////////////////////////////////
 
-TSharedRef<Insights::FTableColumn> FMemTagTreeViewColumnFactory::CreateMaxValueColumn()
+TSharedRef<FTableColumn> FMemTagTreeViewColumnFactory::CreateMaxValueColumn()
 {
-	using namespace Insights;
-
 	TSharedRef<FTableColumn> ColumnRef = MakeShared<FTableColumn>(FMemTagTreeViewColumns::MaxValueColumnID);
 	FTableColumn& Column = *ColumnRef;
 
@@ -301,10 +294,8 @@ TSharedRef<Insights::FTableColumn> FMemTagTreeViewColumnFactory::CreateMaxValueC
 
 ////////////////////////////////////////////////////////////////////////////////////////////////////
 
-TSharedRef<Insights::FTableColumn> FMemTagTreeViewColumnFactory::CreateAverageValueColumn()
+TSharedRef<FTableColumn> FMemTagTreeViewColumnFactory::CreateAverageValueColumn()
 {
-	using namespace Insights;
-
 	TSharedRef<FTableColumn> ColumnRef = MakeShared<FTableColumn>(FMemTagTreeViewColumns::AverageValueColumnID);
 	FTableColumn& Column = *ColumnRef;
 
@@ -346,4 +337,6 @@ TSharedRef<Insights::FTableColumn> FMemTagTreeViewColumnFactory::CreateAverageVa
 
 ////////////////////////////////////////////////////////////////////////////////////////////////////
 
+} // namespace UE::Insights::MemoryProfiler
+
 #undef LOCTEXT_NAMESPACE
diff --git a/Engine/Source/Developer/TraceInsights/Private/Insights/MemoryProfiler/ViewModels/MemTagTreeViewColumnFactory.h b/Engine/Source/Developer/TraceInsights/Private/Insights/MemoryProfiler/ViewModels/MemTagTreeViewColumnFactory.h
--- a/Engine/Source/Developer/TraceInsights/Private/Insights/MemoryProfiler/ViewModels/MemTagTreeViewColumnFactory.h
+++ b/Engine/Source/Developer/TraceInsights/Private/Insights/MemoryProfiler/ViewModels/MemTagTreeViewColumnFactory.h
@@ -2,15 +2,19 @@
 
 #pragma once
 
-#include "CoreMinimal.h"
+#include "CoreTypes.h"
+#include "Containers/Array.h"
+#include "Templates/SharedPointer.h"
+#include "UObject/NameTypes.h"
 
-////////////////////////////////////////////////////////////////////////////////////////////////////
+namespace UE::Insights { class FTableColumn; }
 
-struct FMemTagTreeViewColumns
+namespace UE::Insights::MemoryProfiler
 {
-	//////////////////////////////////////////////////
-	// Column identifiers
 
+// Column identifiers
+struct FMemTagTreeViewColumns
+{
 	static const FName NameColumnID;
 	static const FName TypeColumnID;
 	static const FName TrackerColumnID;
@@ -20,30 +24,21 @@ struct FMemTagTreeViewColumns
 	static const FName AverageValueColumnID;
 };
 
-////////////////////////////////////////////////////////////////////////////////////////////////////
-
-namespace Insights
-{
-	class FTableColumn;
-}
-
-////////////////////////////////////////////////////////////////////////////////////////////////////
-
 struct FMemTagTreeViewColumnFactory
 {
 public:
-	static void CreateMemTagTreeViewColumns(TArray<TSharedRef<Insights::FTableColumn>>& Columns);
+	static void CreateMemTagTreeViewColumns(TArray<TSharedRef<FTableColumn>>& Columns);
 
-	static TSharedRef<Insights::FTableColumn> CreateNameColumn();
-	static TSharedRef<Insights::FTableColumn> CreateTypeColumn();
-	static TSharedRef<Insights::FTableColumn> CreateTrackerColumn();
-	static TSharedRef<Insights::FTableColumn> CreateInstanceCountColumn();
-	static TSharedRef<Insights::FTableColumn> CreateMinValueColumn();
-	static TSharedRef<Insights::FTableColumn> CreateMaxValueColumn();
-	static TSharedRef<Insights::FTableColumn> CreateAverageValueColumn();
+	static TSharedRef<FTableColumn> CreateNameColumn();
+	static TSharedRef<FTableColumn> CreateTypeColumn();
+	static TSharedRef<FTableColumn> CreateTrackerColumn();
+	static TSharedRef<FTableColumn> CreateInstanceCountColumn();
+	static TSharedRef<FTableColumn> CreateMinValueColumn();
+	static TSharedRef<FTableColumn> CreateMaxValueColumn();
+	static TSharedRef<FTableColumn> CreateAverageValueColumn();
 
 private:
 	static constexpr float ValueColumnInitialWidth = 50.0f;
 };
 
-////////////////////////////////////////////////////////////////////////////////////////////////////
+} // namespace UE::Insights::MemoryProfiler
diff --git a/Engine/Source/Developer/TraceInsights/Private/Insights/MemoryProfiler/ViewModels/MemoryAlloc.cpp b/Engine/Source/Developer/TraceInsights/Private/Insights/MemoryProfiler/ViewModels/MemoryAlloc.cpp
--- a/Engine/Source/Developer/TraceInsights/Private/Insights/MemoryProfiler/ViewModels/MemoryAlloc.cpp
+++ b/Engine/Source/Developer/TraceInsights/Private/Insights/MemoryProfiler/ViewModels/MemoryAlloc.cpp
@@ -2,9 +2,7 @@
 
 #include "MemoryAlloc.h"
 
-#include "Insights/MemoryProfiler/ViewModels/CallstackFormatting.h"
-
-namespace Insights
+namespace UE::Insights::MemoryProfiler
 {
 
 ////////////////////////////////////////////////////////////////////////////////////////////////////
@@ -21,4 +19,4 @@ FMemoryAlloc::~FMemoryAlloc()
 
 ////////////////////////////////////////////////////////////////////////////////////////////////////
 
-} // namespace Insights
+} // namespace UE::Insights::MemoryProfiler
diff --git a/Engine/Source/Developer/TraceInsights/Private/Insights/MemoryProfiler/ViewModels/MemoryAlloc.h b/Engine/Source/Developer/TraceInsights/Private/Insights/MemoryProfiler/ViewModels/MemoryAlloc.h
--- a/Engine/Source/Developer/TraceInsights/Private/Insights/MemoryProfiler/ViewModels/MemoryAlloc.h
+++ b/Engine/Source/Developer/TraceInsights/Private/Insights/MemoryProfiler/ViewModels/MemoryAlloc.h
@@ -2,12 +2,14 @@
 
 #pragma once
 
-#include "CoreMinimal.h"
+#include "CoreTypes.h"
 
-#include "TraceServices/Model/AllocationsProvider.h"
-#include "TraceServices/Model/Callstack.h"
+// TraceServices
+#include "TraceServices/Model/AllocationsProvider.h" // for TagIdType
 
-namespace Insights
+namespace TraceServices { struct FCallstack; }
+
+namespace UE::Insights::MemoryProfiler
 {
 
 ////////////////////////////////////////////////////////////////////////////////////////////////////
@@ -79,4 +81,4 @@ private:
 
 ////////////////////////////////////////////////////////////////////////////////////////////////////
 
-} // namespace Insights
+} // namespace UE::Insights::MemoryProfiler
diff --git a/Engine/Source/Developer/TraceInsights/Private/Insights/MemoryProfiler/ViewModels/MemoryGraphTrack.cpp b/Engine/Source/Developer/TraceInsights/Private/Insights/MemoryProfiler/ViewModels/MemoryGraphTrack.cpp
--- a/Engine/Source/Developer/TraceInsights/Private/Insights/MemoryProfiler/ViewModels/MemoryGraphTrack.cpp
+++ b/Engine/Source/Developer/TraceInsights/Private/Insights/MemoryProfiler/ViewModels/MemoryGraphTrack.cpp
@@ -13,9 +13,11 @@
 #include "TraceServices/Model/Counters.h"
 #include "TraceServices/Model/Memory.h"
 
-// Insights
-#include "Insights/Common/PaintUtils.h"
-#include "Insights/Common/TimeUtils.h"
+// TraceInsightsCore
+#include "InsightsCore/Common/PaintUtils.h"
+#include "InsightsCore/Common/TimeUtils.h"
+
+// TraceInsights
 #include "Insights/InsightsManager.h"
 #include "Insights/MemoryProfiler/ViewModels/MemorySharedState.h"
 #include "Insights/ViewModels/AxisViewportDouble.h"
@@ -26,7 +28,10 @@
 
 #include <limits>
 
-#define LOCTEXT_NAMESPACE "MemoryGraphTrack"
+#define LOCTEXT_NAMESPACE "UE::Insights::MemoryProfiler::FMemoryGraphTrack"
+
+namespace UE::Insights::MemoryProfiler
+{
 
 ////////////////////////////////////////////////////////////////////////////////////////////////////
 // FMemoryGraphSeries
@@ -236,7 +241,7 @@ void FMemoryGraphTrack::Update(const ITimingTrackUpdateContext& Context)
 // LLM Tag Series
 ////////////////////////////////////////////////////////////////////////////////////////////////////
 
-TSharedPtr<FMemoryGraphSeries> FMemoryGraphTrack::GetMemTagSeries(Insights::FMemoryTrackerId InMemTrackerId, Insights::FMemoryTagId InMemTagId)
+TSharedPtr<FMemoryGraphSeries> FMemoryGraphTrack::GetMemTagSeries(FMemoryTrackerId InMemTrackerId, FMemoryTagId InMemTagId)
 {
 	TSharedPtr<FGraphSeries>* Ptr = AllSeries.FindByPredicate([InMemTrackerId, InMemTagId](const TSharedPtr<FGraphSeries>& Series)
 	{
@@ -251,7 +256,7 @@ TSharedPtr<FMemoryGraphSeries> FMemoryGraphTrack::GetMemTagSeries(Insights::FMem
 
 ////////////////////////////////////////////////////////////////////////////////////////////////////
 
-TSharedPtr<FMemoryGraphSeries> FMemoryGraphTrack::AddMemTagSeries(Insights::FMemoryTrackerId InMemTrackerId, Insights::FMemoryTagId InMemTagId)
+TSharedPtr<FMemoryGraphSeries> FMemoryGraphTrack::AddMemTagSeries(FMemoryTrackerId InMemTrackerId, FMemoryTagId InMemTagId)
 {
 	TSharedPtr<FMemoryGraphSeries> Series = GetMemTagSeries(InMemTrackerId, InMemTagId);
 
@@ -278,7 +283,7 @@ TSharedPtr<FMemoryGraphSeries> FMemoryGraphTrack::AddMemTagSeries(Insights::FMem
 
 ////////////////////////////////////////////////////////////////////////////////////////////////////
 
-int32 FMemoryGraphTrack::RemoveMemTagSeries(Insights::FMemoryTrackerId InMemTrackerId, Insights::FMemoryTagId InMemTagId)
+int32 FMemoryGraphTrack::RemoveMemTagSeries(FMemoryTrackerId InMemTrackerId, FMemoryTagId InMemTagId)
 {
 	SetDirtyFlag();
 	return AllSeries.RemoveAll([InMemTrackerId, InMemTagId](const TSharedPtr<FGraphSeries>& GraphSeries)
@@ -949,7 +954,7 @@ void FMemoryGraphTrack::DrawVerticalAxisGrid(const ITimingTrackDrawContext& Cont
 			const float MY = static_cast<float>(Context.GetMousePosition().Y);
 
 			//constexpr float MX1 = 80.0f; // start fading out
-			constexpr float MX2 = 120.0f; // completly faded out
+			constexpr float MX2 = 120.0f; // completely faded out
 
 			if (MX > ViewWidth - MX2 && MY >= MaxValueY && MY < MinValueY + TextH)
 			{
@@ -1242,11 +1247,11 @@ void FMemoryGraphTrack::InitTooltip(FTooltipDrawState& InOutTooltip, const ITimi
 			InOutTooltip.AddTitle(SubTitle, Series->GetColor());
 		}
 
-		const double Precision = FMath::Max(1.0 / TimeScaleX, TimeUtils::Nanosecond);
-		InOutTooltip.AddNameValueTextLine(TEXT("Time:"), TimeUtils::FormatTime(TooltipEvent.GetStartTime(), Precision));
+		const double Precision = FMath::Max(1.0 / TimeScaleX, FTimeValue::Nanosecond);
+		InOutTooltip.AddNameValueTextLine(TEXT("Time:"), FormatTime(TooltipEvent.GetStartTime(), Precision));
 		if (Series->HasEventDuration())
 		{
-			InOutTooltip.AddNameValueTextLine(TEXT("Duration:"), TimeUtils::FormatTimeAuto(TooltipEvent.GetDuration()));
+			InOutTooltip.AddNameValueTextLine(TEXT("Duration:"), FormatTimeAuto(TooltipEvent.GetDuration()));
 		}
 		InOutTooltip.AddNameValueTextLine(TEXT("Value:"), Series->FormatValue(TooltipEvent.GetValue()));
 		InOutTooltip.UpdateLayout();
@@ -1255,4 +1260,6 @@ void FMemoryGraphTrack::InitTooltip(FTooltipDrawState& InOutTooltip, const ITimi
 
 ////////////////////////////////////////////////////////////////////////////////////////////////////
 
+} // namespace UE::Insights::MemoryProfiler
+
 #undef LOCTEXT_NAMESPACE
diff --git a/Engine/Source/Developer/TraceInsights/Private/Insights/MemoryProfiler/ViewModels/MemoryGraphTrack.h b/Engine/Source/Developer/TraceInsights/Private/Insights/MemoryProfiler/ViewModels/MemoryGraphTrack.h
--- a/Engine/Source/Developer/TraceInsights/Private/Insights/MemoryProfiler/ViewModels/MemoryGraphTrack.h
+++ b/Engine/Source/Developer/TraceInsights/Private/Insights/MemoryProfiler/ViewModels/MemoryGraphTrack.h
@@ -2,20 +2,24 @@
 
 #pragma once
 
-#include "CoreMinimal.h"
+#include "CoreTypes.h"
+
 #include "Math/Color.h"
 
-// Insights
+// TraceInsights
 #include "Insights/MemoryProfiler/ViewModels/MemoryTag.h"
 #include "Insights/MemoryProfiler/ViewModels/MemoryTracker.h"
 #include "Insights/ViewModels/GraphSeries.h"
 #include "Insights/ViewModels/GraphTrack.h"
 
-class FMemorySharedState;
 class FSlateFontMeasure;
-
 struct FSlateBrush;
 
+namespace UE::Insights::MemoryProfiler
+{
+
+class FMemorySharedState;
+
 ////////////////////////////////////////////////////////////////////////////////////////////////////
 
 enum class EGraphTrackLabelUnit
@@ -58,11 +62,11 @@ public:
 public:
 	virtual FString FormatValue(double Value) const override;
 
-	Insights::FMemoryTrackerId GetTrackerId() const { return TrackerId; }
-	void SetTrackerId(Insights::FMemoryTrackerId InTrackerId) { TrackerId = InTrackerId; }
+	FMemoryTrackerId GetTrackerId() const { return TrackerId; }
+	void SetTrackerId(FMemoryTrackerId InTrackerId) { TrackerId = InTrackerId; }
 
-	Insights::FMemoryTagId GetTagId() const { return TagId; }
-	void SetTagId(Insights::FMemoryTagId InTagId) { TagId = InTagId; }
+	FMemoryTagId GetTagId() const { return TagId; }
+	void SetTagId(FMemoryTagId InTagId) { TagId = InTagId; }
 
 	ETimelineType GetTimelineType() const { return TimelineType; }
 	void SetTimelineType(ETimelineType InTimelineType) { TimelineType = InTimelineType; }
@@ -72,8 +76,8 @@ public:
 	void SetValueRange(double Min, double Max) { MinValue = Min; MaxValue = Max; }
 
 private:
-	Insights::FMemoryTrackerId TrackerId = Insights::FMemoryTracker::InvalidTrackerId; // LLM tracker id
-	Insights::FMemoryTagId TagId = Insights::FMemoryTag::InvalidTagId; // LLM tag id
+	FMemoryTrackerId TrackerId = FMemoryTracker::InvalidTrackerId; // LLM tracker id
+	FMemoryTagId TagId = FMemoryTag::InvalidTagId; // LLM tag id
 	double MinValue = 0.0;
 	double MaxValue = 0.0;
 	ETimelineType TimelineType = ETimelineType::MemTag;
@@ -118,9 +122,9 @@ public:
 	virtual void Update(const ITimingTrackUpdateContext& Context) override;
 	virtual void InitTooltip(FTooltipDrawState& InOutTooltip, const ITimingEvent& InTooltipEvent) const override;
 
-	TSharedPtr<FMemoryGraphSeries> GetMemTagSeries(Insights::FMemoryTrackerId InMemTrackerId, Insights::FMemoryTagId InMemTagId);
-	TSharedPtr<FMemoryGraphSeries> AddMemTagSeries(Insights::FMemoryTrackerId InMemTrackerId, Insights::FMemoryTagId InMemTagId);
-	int32 RemoveMemTagSeries(Insights::FMemoryTrackerId InMemTrackerId, Insights::FMemoryTagId InMemTagId);
+	TSharedPtr<FMemoryGraphSeries> GetMemTagSeries(FMemoryTrackerId InMemTrackerId, FMemoryTagId InMemTagId);
+	TSharedPtr<FMemoryGraphSeries> AddMemTagSeries(FMemoryTrackerId InMemTrackerId, FMemoryTagId InMemTagId);
+	int32 RemoveMemTagSeries(FMemoryTrackerId InMemTrackerId, FMemoryTagId InMemTagId);
 	int32 RemoveAllMemTagSeries();
 
 	TSharedPtr<FMemoryGraphSeries> GetTimelineSeries(FMemoryGraphSeries::ETimelineType InTimelineType);
@@ -169,7 +173,7 @@ protected:
 
 	/**
 	 * Number of decimal digits for labels.
-	 * Specifies the number of decimal digits to use when formating labels of the vertical axis grid.
+	 * Specifies the number of decimal digits to use when formatting labels of the vertical axis grid.
 	 * If negative, the formatting will use maximum the number of decimal digits specified (trims trailing 0s),
 	 * otherwise, it will use exactly the number of decimal digits specified.
 	 */
@@ -189,3 +193,5 @@ protected:
 };
 
 ////////////////////////////////////////////////////////////////////////////////////////////////////
+
+} // namespace UE::Insights::MemoryProfiler
diff --git a/Engine/Source/Developer/TraceInsights/Private/Insights/MemoryProfiler/ViewModels/MemorySharedState.cpp b/Engine/Source/Developer/TraceInsights/Private/Insights/MemoryProfiler/ViewModels/MemorySharedState.cpp
--- a/Engine/Source/Developer/TraceInsights/Private/Insights/MemoryProfiler/ViewModels/MemorySharedState.cpp
+++ b/Engine/Source/Developer/TraceInsights/Private/Insights/MemoryProfiler/ViewModels/MemorySharedState.cpp
@@ -1,16 +1,21 @@
 // Copyright Epic Games, Inc. All Rights Reserved.
 
 #include "MemorySharedState.h"
+
 #include "Framework/Commands/Commands.h"
 #include "Framework/Commands/UICommandList.h"
 #include "Framework/MultiBox/MultiBoxBuilder.h"
 #include "Logging/MessageLog.h"
+
+// TraceServices
 #include "TraceServices/Model/Memory.h"
 
-// Insights
-#include "Insights/Common/PaintUtils.h"
-#include "Insights/Common/Stopwatch.h"
-#include "Insights/Common/TimeUtils.h"
+// TraceInsightsCore
+#include "InsightsCore/Common/PaintUtils.h"
+#include "InsightsCore/Common/Stopwatch.h"
+#include "InsightsCore/Common/TimeUtils.h"
+
+// TraceInsights
 #include "Insights/InsightsManager.h"
 #include "Insights/InsightsStyle.h"
 #include "Insights/MemoryProfiler/MemoryProfilerManager.h"
@@ -19,20 +24,23 @@
 #include "Insights/MemoryProfiler/ViewModels/Report.h"
 #include "Insights/MemoryProfiler/ViewModels/ReportXmlParser.h"
 #include "Insights/ViewModels/TimingEvent.h"
-#include "Insights/ViewModels/TimingTrackViewport.h"
+#include "Insights/ViewModels/TimingEventSearch.h"
 #include "Insights/ViewModels/TimingGraphTrack.h"
+#include "Insights/ViewModels/TimingTrackViewport.h"
 #include "Insights/ViewModels/TooltipDrawState.h"
-#include "Insights/ViewModels/TimingEventSearch.h"
 #include "Insights/Widgets/STimingView.h"
 
 #include <limits>
 
-static_assert(Insights::FMemoryTracker::InvalidTrackerId == TraceServices::FMemoryTrackerInfo::InvalidTrackerId, "InvalidTrackerId");
-static_assert(Insights::FMemoryTag::InvalidTagId == TraceServices::FMemoryTagInfo::InvalidTagId, "InvalidTagId");
+#define LOCTEXT_NAMESPACE "UE::Insights::MemoryProfiler::FMemorySharedState"
 
-#define LOCTEXT_NAMESPACE "MemorySharedState"
+namespace UE::Insights::MemoryProfiler
+{
 
-const FName Insights::FQueryTargetWindowSpec::NewWindow = TEXT("New Window");
+static_assert(FMemoryTracker::InvalidTrackerId == TraceServices::FMemoryTrackerInfo::InvalidTrackerId, "InvalidTrackerId");
+static_assert(FMemoryTag::InvalidTagId == TraceServices::FMemoryTagInfo::InvalidTagId, "InvalidTagId");
+
+const FName FQueryTargetWindowSpec::NewWindow = TEXT("New Window");
 
 ////////////////////////////////////////////////////////////////////////////////////////////////////
 // FMemoryTimingViewCommands
@@ -87,7 +95,7 @@ FMemorySharedState::FMemorySharedState()
 {
 	InitMemoryRules();
 
-	CurrentQueryTarget = MakeShared<Insights::FQueryTargetWindowSpec>(Insights::FQueryTargetWindowSpec::NewWindow, LOCTEXT("NewWindow", "New Window"));
+	CurrentQueryTarget = MakeShared<FQueryTargetWindowSpec>(FQueryTargetWindowSpec::NewWindow, LOCTEXT("NewWindow", "New Window"));
 	QueryTargetSpecs.Add(CurrentQueryTarget);
 }
 
@@ -99,7 +107,7 @@ FMemorySharedState::~FMemorySharedState()
 
 ////////////////////////////////////////////////////////////////////////////////////////////////////
 
-void FMemorySharedState::OnBeginSession(Insights::ITimingViewSession& InSession)
+void FMemorySharedState::OnBeginSession(Timing::ITimingViewSession& InSession)
 {
 	if (&InSession != TimingView.Get())
 	{
@@ -126,7 +134,7 @@ void FMemorySharedState::OnBeginSession(Insights::ITimingViewSession& InSession)
 
 ////////////////////////////////////////////////////////////////////////////////////////////////////
 
-void FMemorySharedState::OnEndSession(Insights::ITimingViewSession& InSession)
+void FMemorySharedState::OnEndSession(Timing::ITimingViewSession& InSession)
 {
 	if (&InSession != TimingView.Get())
 	{
@@ -153,7 +161,7 @@ void FMemorySharedState::OnEndSession(Insights::ITimingViewSession& InSession)
 
 ////////////////////////////////////////////////////////////////////////////////////////////////////
 
-void FMemorySharedState::Tick(Insights::ITimingViewSession& InSession, const TraceServices::IAnalysisSession& InAnalysisSession)
+void FMemorySharedState::Tick(Timing::ITimingViewSession& InSession, const TraceServices::IAnalysisSession& InAnalysisSession)
 {
 	if (&InSession != TimingView.Get())
 	{
@@ -298,7 +306,7 @@ void FMemorySharedState::CreateDefaultTracks()
 		return;
 	}
 
-	Insights::FMemoryTrackerId DefaultTrackerId = DefaultTracker->GetId();
+	FMemoryTrackerId DefaultTrackerId = DefaultTracker->GetId();
 
 	static const TCHAR* DefaultTags[] =
 	{
@@ -329,7 +337,7 @@ void FMemorySharedState::CreateDefaultTracks()
 	{
 		if (!CreatedDefaultTracks[DefaultTagIndex])
 		{
-			for (const Insights::FMemoryTag* Tag : Tags)
+			for (const FMemoryTag* Tag : Tags)
 			{
 				if (Tag->GetTrackerId() == DefaultTrackerId && // is it used by the default tracker?
 					Tag->GetGraphTracks().Num() == 0 && // a graph isn't already added for this llm tag?
@@ -350,9 +358,9 @@ FString FMemorySharedState::TrackersToString(uint64 Flags, const TCHAR* Conjunct
 	FString Str;
 	if (Flags != 0)
 	{
-		for (const TSharedPtr<Insights::FMemoryTracker>& Tracker : Trackers)
+		for (const TSharedPtr<FMemoryTracker>& Tracker : Trackers)
 		{
-			const uint64 TrackerFlag = Insights::FMemoryTracker::AsFlag(Tracker->GetId());
+			const uint64 TrackerFlag = FMemoryTracker::AsFlag(Tracker->GetId());
 			if ((Flags & TrackerFlag) != 0)
 			{
 				if (!Str.IsEmpty())
@@ -373,9 +381,9 @@ FString FMemorySharedState::TrackersToString(uint64 Flags, const TCHAR* Conjunct
 
 ////////////////////////////////////////////////////////////////////////////////////////////////////
 
-const Insights::FMemoryTracker* FMemorySharedState::GetTrackerById(Insights::FMemoryTrackerId InMemTrackerId) const
+const FMemoryTracker* FMemorySharedState::GetTrackerById(FMemoryTrackerId InMemTrackerId) const
 {
-	const TSharedPtr<Insights::FMemoryTracker>* TrackerPtr = Trackers.FindByPredicate([InMemTrackerId](TSharedPtr<Insights::FMemoryTracker>& Tracker) { return Tracker->GetId() == InMemTrackerId; });
+	const TSharedPtr<FMemoryTracker>* TrackerPtr = Trackers.FindByPredicate([InMemTrackerId](TSharedPtr<FMemoryTracker>& Tracker) { return Tracker->GetId() == InMemTrackerId; });
 	return TrackerPtr ? TrackerPtr->Get() : nullptr;
 }
 
@@ -396,16 +404,16 @@ void FMemorySharedState::SyncTrackers()
 		{
 			MemoryProvider->EnumerateTrackers([this](const TraceServices::FMemoryTrackerInfo& Tracker)
 			{
-				Trackers.Add(MakeShared<Insights::FMemoryTracker>(Tracker.Id, Tracker.Name));
+				Trackers.Add(MakeShared<FMemoryTracker>(Tracker.Id, Tracker.Name));
 			});
 
-			Trackers.Sort([](const TSharedPtr<Insights::FMemoryTracker>& A, const TSharedPtr<Insights::FMemoryTracker>& B) { return A->GetId() < B->GetId(); });
+			Trackers.Sort([](const TSharedPtr<FMemoryTracker>& A, const TSharedPtr<FMemoryTracker>& B) { return A->GetId() < B->GetId(); });
 		}
 	}
 
 	if (Trackers.Num() > 0)
 	{
-		for (const TSharedPtr<Insights::FMemoryTracker>& Tracker : Trackers)
+		for (const TSharedPtr<FMemoryTracker>& Tracker : Trackers)
 		{
 			if (FCString::Stricmp(*Tracker->GetName(), TEXT("Default")) == 0)
 			{
@@ -433,7 +441,7 @@ void FMemorySharedState::SetTrackHeightMode(EMemoryTrackHeightMode InTrackHeight
 
 ////////////////////////////////////////////////////////////////////////////////////////////////////
 
-void FMemorySharedState::ExtendOtherTracksFilterMenu(Insights::ITimingViewSession& InSession, FMenuBuilder& InOutMenuBuilder)
+void FMemorySharedState::ExtendOtherTracksFilterMenu(Timing::ITimingViewSession& InSession, FMenuBuilder& InOutMenuBuilder)
 {
 	if (&InSession != TimingView.Get())
 	{
@@ -562,7 +570,7 @@ void FMemorySharedState::RemoveTrackFromMemTags(TSharedPtr<FMemoryGraphTrack>& G
 	{
 		//TODO: if (Series->Is<FMemoryGraphSeries>())
 		TSharedPtr<FMemoryGraphSeries> MemorySeries = StaticCastSharedPtr<FMemoryGraphSeries>(Series);
-		Insights::FMemoryTag* TagPtr = TagList.GetTagById(MemorySeries->GetTrackerId(), MemorySeries->GetTagId());
+		FMemoryTag* TagPtr = TagList.GetTagById(MemorySeries->GetTrackerId(), MemorySeries->GetTagId());
 		if (TagPtr)
 		{
 			TagPtr->RemoveTrack(GraphTrack);
@@ -572,14 +580,14 @@ void FMemorySharedState::RemoveTrackFromMemTags(TSharedPtr<FMemoryGraphTrack>& G
 
 ////////////////////////////////////////////////////////////////////////////////////////////////////
 
-TSharedPtr<FMemoryGraphTrack> FMemorySharedState::GetMemTagGraphTrack(Insights::FMemoryTrackerId InMemTrackerId, Insights::FMemoryTagId InMemTagId)
+TSharedPtr<FMemoryGraphTrack> FMemorySharedState::GetMemTagGraphTrack(FMemoryTrackerId InMemTrackerId, FMemoryTagId InMemTagId)
 {
 	if (!TimingView.IsValid())
 	{
 		return nullptr;
 	}
 
-	Insights::FMemoryTag* TagPtr = TagList.GetTagById(InMemTrackerId, InMemTagId);
+	FMemoryTag* TagPtr = TagList.GetTagById(InMemTrackerId, InMemTagId);
 	if (TagPtr)
 	{
 		for (TSharedPtr<FMemoryGraphTrack> MemoryGraph : TagPtr->GetGraphTracks())
@@ -596,19 +604,19 @@ TSharedPtr<FMemoryGraphTrack> FMemorySharedState::GetMemTagGraphTrack(Insights::
 
 ////////////////////////////////////////////////////////////////////////////////////////////////////
 
-TSharedPtr<FMemoryGraphTrack> FMemorySharedState::CreateMemTagGraphTrack(Insights::FMemoryTrackerId InMemTrackerId, Insights::FMemoryTagId InMemTagId)
+TSharedPtr<FMemoryGraphTrack> FMemorySharedState::CreateMemTagGraphTrack(FMemoryTrackerId InMemTrackerId, FMemoryTagId InMemTagId)
 {
 	if (!TimingView.IsValid())
 	{
 		return nullptr;
 	}
 
-	Insights::FMemoryTag* TagPtr = TagList.GetTagById(InMemTrackerId, InMemTagId);
+	FMemoryTag* TagPtr = TagList.GetTagById(InMemTrackerId, InMemTagId);
 
 	FString SeriesName;
 	if (TagPtr)
 	{
-		const Insights::FMemoryTracker* Tracker = GetTrackerById(InMemTrackerId);
+		const FMemoryTracker* Tracker = GetTrackerById(InMemTrackerId);
 		if (Tracker && Tracker != DefaultTracker.Get())
 		{
 			SeriesName = FString::Printf(TEXT("LLM %s (%s)"), *TagPtr->GetStatFullName(), *Tracker->GetName());
@@ -691,7 +699,7 @@ TSharedPtr<FMemoryGraphTrack> FMemorySharedState::CreateMemTagGraphTrack(Insight
 
 ////////////////////////////////////////////////////////////////////////////////////////////////////
 
-int32 FMemorySharedState::RemoveMemTagGraphTrack(Insights::FMemoryTrackerId InMemTrackerId, Insights::FMemoryTagId InMemTagId)
+int32 FMemorySharedState::RemoveMemTagGraphTrack(FMemoryTrackerId InMemTrackerId, FMemoryTagId InMemTagId)
 {
 	if (!TimingView.IsValid())
 	{
@@ -700,7 +708,7 @@ int32 FMemorySharedState::RemoveMemTagGraphTrack(Insights::FMemoryTrackerId InMe
 
 	int32 TrackCount = 0;
 
-	Insights::FMemoryTag* TagPtr = TagList.GetTagById(InMemTrackerId, InMemTagId);
+	FMemoryTag* TagPtr = TagList.GetTagById(InMemTrackerId, InMemTagId);
 	if (TagPtr)
 	{
 		for (TSharedPtr<FMemoryGraphTrack> GraphTrack : TagPtr->GetGraphTracks())
@@ -762,7 +770,7 @@ int32 FMemorySharedState::RemoveAllMemTagGraphTracks()
 		AllTracks.Remove(GraphTrack);
 	}
 
-	for (Insights::FMemoryTag* TagPtr : TagList.GetTags())
+	for (FMemoryTag* TagPtr : TagList.GetTags())
 	{
 		TagPtr->RemoveAllTracks();
 	}
@@ -772,14 +780,14 @@ int32 FMemorySharedState::RemoveAllMemTagGraphTracks()
 
 ////////////////////////////////////////////////////////////////////////////////////////////////////
 
-TSharedPtr<FMemoryGraphSeries> FMemorySharedState::ToggleMemTagGraphSeries(TSharedPtr<FMemoryGraphTrack> InGraphTrack, Insights::FMemoryTrackerId InMemTrackerId, Insights::FMemoryTagId InMemTagId)
+TSharedPtr<FMemoryGraphSeries> FMemorySharedState::ToggleMemTagGraphSeries(TSharedPtr<FMemoryGraphTrack> InGraphTrack, FMemoryTrackerId InMemTrackerId, FMemoryTagId InMemTagId)
 {
 	if (!InGraphTrack.IsValid())
 	{
 		return nullptr;
 	}
 
-	Insights::FMemoryTag* TagPtr = TagList.GetTagById(InMemTrackerId, InMemTagId);
+	FMemoryTag* TagPtr = TagList.GetTagById(InMemTrackerId, InMemTagId);
 
 	TSharedPtr<FMemoryGraphSeries> Series = InGraphTrack->GetMemTagSeries(InMemTrackerId, InMemTagId);
 	if (Series.IsValid())
@@ -820,12 +828,12 @@ TSharedPtr<FMemoryGraphSeries> FMemorySharedState::ToggleMemTagGraphSeries(TShar
 // Create graphs from LLMReportTypes.xml file
 void FMemorySharedState::CreateTracksFromReport(const FString& Filename)
 {
-	Insights::FReportConfig ReportConfig;
+	FReportConfig ReportConfig;
 
-	Insights::FReportXmlParser ReportXmlParser;
+	FReportXmlParser ReportXmlParser;
 
 	ReportXmlParser.LoadReportTypesXML(ReportConfig, Filename);
-	if (ReportXmlParser.GetStatus() != Insights::FReportXmlParser::EStatus::Completed)
+	if (ReportXmlParser.GetStatus() != FReportXmlParser::EStatus::Completed)
 	{
 		FMessageLog ReportMessageLog(FMemoryProfilerManager::Get()->GetLogListingName());
 		ReportMessageLog.AddMessages(ReportXmlParser.GetErrorMessages());
@@ -837,9 +845,9 @@ void FMemorySharedState::CreateTracksFromReport(const FString& Filename)
 
 ////////////////////////////////////////////////////////////////////////////////////////////////////
 
-void FMemorySharedState::CreateTracksFromReport(const Insights::FReportConfig& ReportConfig)
+void FMemorySharedState::CreateTracksFromReport(const FReportConfig& ReportConfig)
 {
-	for (const Insights::FReportTypeConfig& ReportTypeConfig : ReportConfig.ReportTypes)
+	for (const FReportTypeConfig& ReportTypeConfig : ReportConfig.ReportTypes)
 	{
 		CreateTracksFromReport(ReportTypeConfig);
 	}
@@ -847,14 +855,14 @@ void FMemorySharedState::CreateTracksFromReport(const Insights::FReportConfig& R
 
 ////////////////////////////////////////////////////////////////////////////////////////////////////
 
-void FMemorySharedState::CreateTracksFromReport(const Insights::FReportTypeConfig& ReportTypeConfig)
+void FMemorySharedState::CreateTracksFromReport(const FReportTypeConfig& ReportTypeConfig)
 {
 	int32 Order = GetNextMemoryGraphTrackOrder();
 	int32 NumAddedTracks = 0;
 
 	const bool bIsPlatformTracker = ReportTypeConfig.Name.StartsWith(TEXT("LLMPlatform"));
 
-	for (const Insights::FReportTypeGraphConfig& ReportTypeGraphConfig : ReportTypeConfig.Graphs)
+	for (const FReportTypeGraphConfig& ReportTypeGraphConfig : ReportTypeConfig.Graphs)
 	{
 		TSharedPtr<FMemoryGraphTrack> GraphTrack = CreateGraphTrack(ReportTypeGraphConfig, bIsPlatformTracker);
 		if (GraphTrack)
@@ -875,7 +883,7 @@ void FMemorySharedState::CreateTracksFromReport(const Insights::FReportTypeConfi
 
 ////////////////////////////////////////////////////////////////////////////////////////////////////
 
-TSharedPtr<FMemoryGraphTrack> FMemorySharedState::CreateGraphTrack(const Insights::FReportTypeGraphConfig& ReportTypeGraphConfig, bool bIsPlatformTracker)
+TSharedPtr<FMemoryGraphTrack> FMemorySharedState::CreateGraphTrack(const FReportTypeGraphConfig& ReportTypeGraphConfig, bool bIsPlatformTracker)
 {
 	if (ReportTypeGraphConfig.GraphConfig == nullptr)
 	{
@@ -888,7 +896,7 @@ TSharedPtr<FMemoryGraphTrack> FMemorySharedState::CreateGraphTrack(const Insight
 		return nullptr;
 	}
 
-	const Insights::FGraphConfig& GraphConfig = *ReportTypeGraphConfig.GraphConfig;
+	const FGraphConfig& GraphConfig = *ReportTypeGraphConfig.GraphConfig;
 
 	int32 CharIndex;
 	const TCHAR* DelimStr;
@@ -929,12 +937,12 @@ TSharedPtr<FMemoryGraphTrack> FMemorySharedState::CreateGraphTrack(const Insight
 	TArray<FString> IgnoreStats;
 	GraphConfig.IgnoreStats.ParseIntoArray(IgnoreStats, DelimStr);
 
-	TArray<Insights::FMemoryTag*> Tags;
+	TArray<FMemoryTag*> Tags;
 	TagList.FilterTags(IncludeStats, IgnoreStats, Tags);
 
-	Insights::FMemoryTrackerId MemTrackerId = bIsPlatformTracker ?
-		(PlatformTracker ? PlatformTracker->GetId() : Insights::FMemoryTracker::InvalidTrackerId) :
-		(DefaultTracker ? DefaultTracker->GetId() : Insights::FMemoryTracker::InvalidTrackerId);
+	FMemoryTrackerId MemTrackerId = bIsPlatformTracker ?
+		(PlatformTracker ? PlatformTracker->GetId() : FMemoryTracker::InvalidTrackerId) :
+		(DefaultTracker ? DefaultTracker->GetId() : FMemoryTracker::InvalidTrackerId);
 
 	TSharedPtr<FMemoryGraphTrack> GraphTrack = CreateMemoryGraphTrack();
 	if (GraphTrack)
@@ -952,7 +960,7 @@ TSharedPtr<FMemoryGraphTrack> FMemorySharedState::CreateGraphTrack(const Insight
 		const double MaxValue = GraphConfig.MaxY * 1024.0 * 1024.0;
 		GraphTrack->SetDefaultValueRange(MinValue, MaxValue);
 
-		UE_LOG(MemoryProfiler, Log, TEXT("[LLM Tags] Created graph \"%s\" (H=%.1f%s, MainStat=%s, Stats=%s)"),
+		UE_LOG(LogMemoryProfiler, Log, TEXT("[LLM Tags] Created graph \"%s\" (H=%.1f%s, MainStat=%s, Stats=%s)"),
 			*ReportTypeGraphConfig.Title,
 			GraphTrack->GetHeight(),
 			GraphConfig.bStacked ? TEXT(", stacked") : TEXT(""),
@@ -961,9 +969,9 @@ TSharedPtr<FMemoryGraphTrack> FMemorySharedState::CreateGraphTrack(const Insight
 
 		TSharedPtr<FMemoryGraphSeries> MainSeries;
 
-		for (Insights::FMemoryTag* TagPtr : Tags)
+		for (FMemoryTag* TagPtr : Tags)
 		{
-			Insights::FMemoryTag& Tag = *TagPtr;
+			FMemoryTag& Tag = *TagPtr;
 
 			TSharedPtr<FMemoryGraphSeries> Series = GraphTrack->AddMemTagSeries(MemTrackerId, Tag.GetId());
 			Series->SetName(FText::FromString(FString::Printf(TEXT("LLM %s"), *Tag.GetStatFullName())));
@@ -997,118 +1005,118 @@ void FMemorySharedState::InitMemoryRules()
 
 	MemoryRules.Reset();
 
-	MemoryRules.Add(MakeShared<Insights::FMemoryRuleSpec>(
+	MemoryRules.Add(MakeShared<FMemoryRuleSpec>(
 		ERule::aAf, 1,
 		LOCTEXT("MemRule_aAf_Short", "*A*"),
 		LOCTEXT("MemRule_aAf_Verbose", "Active Allocs"),
 		LOCTEXT("MemRule_aAf_Desc", "Identifies active allocations at time A.\n(a ≤ A ≤ f)")));
 
-	MemoryRules.Add(MakeShared<Insights::FMemoryRuleSpec>(
+	MemoryRules.Add(MakeShared<FMemoryRuleSpec>(
 		ERule::afA, 1,
 		LOCTEXT("MemRule_afA_Short", "**A"),
 		LOCTEXT("MemRule_afA_Verbose", "Before"),
 		LOCTEXT("MemRule_afA_Desc", "Identifies allocations allocated and freed before time A.\n(a ≤ f ≤ A)")));
 
-	MemoryRules.Add(MakeShared<Insights::FMemoryRuleSpec>(
+	MemoryRules.Add(MakeShared<FMemoryRuleSpec>(
 		ERule::Aaf, 1,
 		LOCTEXT("MemRule_Aaf_Short", "A**"),
 		LOCTEXT("MemRule_Aaf_Verbose", "After"),
 		LOCTEXT("MemRule_Aaf_Desc", "Identifies allocations allocated after time A.\n(A ≤ a ≤ f)")));
 
-	MemoryRules.Add(MakeShared<Insights::FMemoryRuleSpec>(
+	MemoryRules.Add(MakeShared<FMemoryRuleSpec>(
 		ERule::aAfB, 2,
 		LOCTEXT("MemRule_aAfB_Short", "*A*B"),
 		LOCTEXT("MemRule_aAfB_Verbose", "Decline"),
 		LOCTEXT("MemRule_aAfB_Desc", "Identifies allocations allocated before time A and freed between time A and time B.\n(a ≤ A ≤ f ≤ B)")));
 
-	MemoryRules.Add(MakeShared<Insights::FMemoryRuleSpec>(
+	MemoryRules.Add(MakeShared<FMemoryRuleSpec>(
 		ERule::AaBf, 2,
 		LOCTEXT("MemRule_AaBf_Short", "A*B*"),
 		LOCTEXT("MemRule_AaBf_Verbose", "Growth"),
 		LOCTEXT("MemRule_AaBf_Desc", "Identifies allocations allocated between time A and time B and not freed until at least time B.\n(A ≤ a ≤ B ≤ f)")));
 
-	MemoryRules.Add(MakeShared<Insights::FMemoryRuleSpec>(
+	MemoryRules.Add(MakeShared<FMemoryRuleSpec>(
 		ERule::aAfaBf, 2,
 		LOCTEXT("MemRule_aAfaBf_Short", "A*B*/*A*B"),
 		LOCTEXT("MemRule_aAfaBf_Verbose", "Growth vs. Decline"),
 		LOCTEXT("MemRule_aAfaBf_Desc", "Identifies \"growth\" allocations, allocated between time A and time B and not freed until at least time B (A ≤ a ≤ B ≤ f)\nand \"decline\" allocations, allocated before time A and freed between time A and time B (a ≤ A ≤ f ≤ B).\nThe \"decline\" allocations are changed to have negative size, so the size aggregation shows variation between A and B.")));
 
-	MemoryRules.Add(MakeShared<Insights::FMemoryRuleSpec>(
+	MemoryRules.Add(MakeShared<FMemoryRuleSpec>(
 		ERule::AfB, 2,
 		LOCTEXT("MemRule_AfB_Short", "*A**B"),
 		LOCTEXT("MemRule_AfB_Verbose", "Free Events"),
 		LOCTEXT("MemRule_AfB_Desc", "Identifies allocations freed between time A and time B.\n(A ≤ f ≤ B)")));
 
-	MemoryRules.Add(MakeShared<Insights::FMemoryRuleSpec>(
+	MemoryRules.Add(MakeShared<FMemoryRuleSpec>(
 		ERule::AaB, 2,
 		LOCTEXT("MemRule_AaB_Short", "A**B*"),
 		LOCTEXT("MemRule_AaB_Verbose", "Alloc Events"),
 		LOCTEXT("MemRule_AaB_Desc", "Identifies allocations allocated between time A and time B.\n(A ≤ a ≤ B)")));
 
-	MemoryRules.Add(MakeShared<Insights::FMemoryRuleSpec>(
+	MemoryRules.Add(MakeShared<FMemoryRuleSpec>(
 		ERule::AafB, 2,
 		LOCTEXT("MemRule_AafB_Short", "A**B"),
 		LOCTEXT("MemRule_AafB_Verbose", "Short Living Allocs"),
 		LOCTEXT("MemRule_AafB_Desc", "Identifies allocations allocated and freed between time A and time B.\n(A ≤ a ≤ f ≤ B)")));
 
-	MemoryRules.Add(MakeShared<Insights::FMemoryRuleSpec>(
+	MemoryRules.Add(MakeShared<FMemoryRuleSpec>(
 		ERule::aABf, 2,
 		LOCTEXT("MemRule_aABf_Short", "*A B*"),
 		LOCTEXT("MemRule_aABf_Verbose", "Long Living Allocs"),
 		LOCTEXT("MemRule_aABf_Desc", "Identifies allocations allocated before time A and not freed until at least time B.\n(a ≤ A ≤ B ≤ f)")));
 
-	MemoryRules.Add(MakeShared<Insights::FMemoryRuleSpec>(
+	MemoryRules.Add(MakeShared<FMemoryRuleSpec>(
 		ERule::AaBCf, 3,
 		LOCTEXT("MemRule_AaBCf_Short", "A*B C*"),
 		LOCTEXT("MemRule_AaBCf_Verbose", "Memory Leaks"),
 		LOCTEXT("MemRule_AaBCf_Desc", "Identifies allocations allocated between time A and time B and not freed until at least time C.\n(A ≤ a ≤ B ≤ C ≤ f)")));
 
-	MemoryRules.Add(MakeShared<Insights::FMemoryRuleSpec>(
+	MemoryRules.Add(MakeShared<FMemoryRuleSpec>(
 		ERule::AaBfC, 3,
 		LOCTEXT("MemRule_AaBfC_Short", "A*B*C"),
 		LOCTEXT("MemRule_AaBfC_Verbose", "Limited Lifetime"),
 		LOCTEXT("MemRule_AaBfC_Desc", "Identifies allocations allocated between time A and time B and freed between time B and time C.\n(A ≤ a ≤ B ≤ f ≤ C)")));
 
-	MemoryRules.Add(MakeShared<Insights::FMemoryRuleSpec>(
+	MemoryRules.Add(MakeShared<FMemoryRuleSpec>(
 		ERule::aABfC, 3,
 		LOCTEXT("MemRule_aABfC_Short", "*A B*C"),
 		LOCTEXT("MemRule_aABfC_Verbose", "Decline of Long Living Allocs"),
 		LOCTEXT("MemRule_aABfC_Desc", "Identifies allocations allocated before time A and freed between time B and time C.\n(a ≤ A ≤ B ≤ f ≤ C)")));
 
-	MemoryRules.Add(MakeShared<Insights::FMemoryRuleSpec>(
+	MemoryRules.Add(MakeShared<FMemoryRuleSpec>(
 		ERule::AaBCfD, 4,
 		LOCTEXT("MemRule_AaBCfD_Short", "A*B C*D"),
 		LOCTEXT("MemRule_AaBCfD_Verbose", "Specific Lifetime"),
 		LOCTEXT("MemRule_AaBCfD_Desc", "Identifies allocations allocated between time A and time B and freed between time C and time D.\n(A ≤ a ≤ B ≤ C ≤ f ≤ D)")));
 
-	MemoryRules.Add(MakeShared<Insights::FMemoryRuleSpec>(
+	MemoryRules.Add(MakeShared<FMemoryRuleSpec>(
 		ERule::AoB, 2,
 		LOCTEXT("MemRule_AoB_Short", "A ↓ B"),
 		LOCTEXT("MemRule_AoB_Verbose", "Paged-Out Allocs"),
 		LOCTEXT("MemRule_AoB_Desc", "Identifies allocations paged-out (swapped-out) between time A and time B.\n(A ≤ page-out ≤ B)")));
 
-	MemoryRules.Add(MakeShared<Insights::FMemoryRuleSpec>(
+	MemoryRules.Add(MakeShared<FMemoryRuleSpec>(
 		ERule::AiB, 2,
 		LOCTEXT("MemRule_AiB_Short", "A ↑ B"),
 		LOCTEXT("MemRule_AiB_Verbose", "Paged-In Allocs"),
 		LOCTEXT("MemRule_AiB_Desc", "Identifies allocations paged-in (swapped-in) between time A and time B.\n(A ≤ page-in ≤ B)")));
 
 	//TODO
-	//MemoryRules.Add(MakeShared<Insights::FMemoryRuleSpec>(
+	//MemoryRules.Add(MakeShared<FMemoryRuleSpec>(
 	//	ERule::A_vs_B, 2,
 	//	LOCTEXT("MemRule_A_vs_B_Short", "*A* + *B*"),
 	//	LOCTEXT("MemRule_A_vs_B_Verbose", "Compare A vs. B"),
 	//	LOCTEXT("MemRule_A_vs_B_Desc", "Compares live allocations at time A with live allocations at time B.\n(*A* vs. *B*)")));
 
 	//TODO
-	//MemoryRules.Add(MakeShared<Insights::FMemoryRuleSpec>(
+	//MemoryRules.Add(MakeShared<FMemoryRuleSpec>(
 	//	ERule::A_or_B, 2,
 	//	LOCTEXT("MemRule_A_or_B_Short", "*A* | *B*"),
 	//	LOCTEXT("MemRule_A_or_B_Verbose", "A or B"),
 	//	LOCTEXT("MemRule_A_or_B_Desc", "Identifies allocations live at time A or at time B.\n(a ≤ A ≤ f OR a ≤ B ≤ f)\n{*A*} U {*B*}")));
 
 	//TODO
-	//MemoryRules.Add(MakeShared<Insights::FMemoryRuleSpec>(
+	//MemoryRules.Add(MakeShared<FMemoryRuleSpec>(
 	//	ERule::A_xor_B, 2,
 	//	LOCTEXT("MemRule_A_xor_B_Short", "*A* ^ *B*"),
 	//	LOCTEXT("MemRule_A_xor_B_Verbose", "A xor B"),
@@ -1125,7 +1133,7 @@ void FMemorySharedState::OnMemoryRuleChanged()
 
 ////////////////////////////////////////////////////////////////////////////////////////////////////
 
-void FMemorySharedState::AddQueryTarget(TSharedPtr<Insights::FQueryTargetWindowSpec> InPtr)
+void FMemorySharedState::AddQueryTarget(TSharedPtr<FQueryTargetWindowSpec> InPtr)
 {
 	QueryTargetSpecs.Add(InPtr);
 }
@@ -1133,11 +1141,13 @@ void FMemorySharedState::AddQueryTarget(TSharedPtr<Insights::FQueryTargetWindowS
 ////////////////////////////////////////////////////////////////////////////////////////////////////
 
 
-void FMemorySharedState::RemoveQueryTarget(TSharedPtr<Insights::FQueryTargetWindowSpec> InPtr)
+void FMemorySharedState::RemoveQueryTarget(TSharedPtr<FQueryTargetWindowSpec> InPtr)
 {
 	QueryTargetSpecs.Remove(InPtr);
 }
 
 ////////////////////////////////////////////////////////////////////////////////////////////////////
 
+} // namespace UE::Insights::MemoryProfiler
+
 #undef LOCTEXT_NAMESPACE
diff --git a/Engine/Source/Developer/TraceInsights/Private/Insights/MemoryProfiler/ViewModels/MemorySharedState.h b/Engine/Source/Developer/TraceInsights/Private/Insights/MemoryProfiler/ViewModels/MemorySharedState.h
--- a/Engine/Source/Developer/TraceInsights/Private/Insights/MemoryProfiler/ViewModels/MemorySharedState.h
+++ b/Engine/Source/Developer/TraceInsights/Private/Insights/MemoryProfiler/ViewModels/MemorySharedState.h
@@ -2,11 +2,14 @@
 
 #pragma once
 
-#include "CoreMinimal.h"
+#include "CoreTypes.h"
+
 #include "Framework/Commands/Commands.h"
+
+// TraceServices
 #include "TraceServices/Model/AllocationsProvider.h"
 
-// Insights
+// TraceInsights
 #include "Insights/ITimingViewExtender.h"
 #include "Insights/MemoryProfiler/ViewModels/MemoryGraphTrack.h"
 #include "Insights/MemoryProfiler/ViewModels/MemoryTag.h"
@@ -14,9 +17,10 @@
 class FTimingEventSearchParameters;
 class FTimingGraphSeries;
 class FTimingGraphTrack;
-class STimingView;
 
-namespace Insights
+namespace UE::Insights::TimingProfiler { class STimingView; }
+
+namespace UE::Insights::MemoryProfiler
 {
 
 struct FReportConfig;
@@ -77,10 +81,6 @@ private:
 
 ////////////////////////////////////////////////////////////////////////////////////////////////////
 
-} // namespace Insights
-
-////////////////////////////////////////////////////////////////////////////////////////////////////
-
 class FMemoryTimingViewCommands : public TCommands<FMemoryTimingViewCommands>
 {
 public:
@@ -94,20 +94,20 @@ public:
 
 ////////////////////////////////////////////////////////////////////////////////////////////////////
 
-class FMemorySharedState : public Insights::ITimingViewExtender, public TSharedFromThis<FMemorySharedState>
+class FMemorySharedState : public Timing::ITimingViewExtender, public TSharedFromThis<FMemorySharedState>
 {
 public:
 	FMemorySharedState();
 	virtual ~FMemorySharedState();
 
-	TSharedPtr<STimingView> GetTimingView() const { return TimingView; }
-	void SetTimingView(TSharedPtr<STimingView> InTimingView) { TimingView = InTimingView; }
+	TSharedPtr<TimingProfiler::STimingView> GetTimingView() const { return TimingView; }
+	void SetTimingView(TSharedPtr<TimingProfiler::STimingView> InTimingView) { TimingView = InTimingView; }
 
-	const Insights::FMemoryTagList& GetTagList() { return TagList; }
+	const FMemoryTagList& GetTagList() { return TagList; }
 
-	const TArray<TSharedPtr<Insights::FMemoryTracker>>& GetTrackers()  const { return Trackers; }
+	const TArray<TSharedPtr<FMemoryTracker>>& GetTrackers()  const { return Trackers; }
 	FString TrackersToString(uint64 Flags, const TCHAR* Conjunction) const;
-	const Insights::FMemoryTracker* GetTrackerById(Insights::FMemoryTrackerId InMemTrackerId) const;
+	const FMemoryTracker* GetTrackerById(FMemoryTrackerId InMemTrackerId) const;
 
 	TSharedPtr<FMemoryGraphTrack> GetMainGraphTrack() const { return MainGraphTrack; }
 
@@ -117,10 +117,10 @@ public:
 	//////////////////////////////////////////////////
 	// ITimingViewExtender interface
 
-	virtual void OnBeginSession(Insights::ITimingViewSession& InSession) override;
-	virtual void OnEndSession(Insights::ITimingViewSession& InSession) override;
-	virtual void Tick(Insights::ITimingViewSession& InSession, const TraceServices::IAnalysisSession& InAnalysisSession) override;
-	virtual void ExtendOtherTracksFilterMenu(Insights::ITimingViewSession& InSession, FMenuBuilder& InOutMenuBuilder) override;
+	virtual void OnBeginSession(Timing::ITimingViewSession& InSession) override;
+	virtual void OnEndSession(Timing::ITimingViewSession& InSession) override;
+	virtual void Tick(Timing::ITimingViewSession& InSession, const TraceServices::IAnalysisSession& InAnalysisSession) override;
+	virtual void ExtendOtherTracksFilterMenu(Timing::ITimingViewSession& InSession, FMenuBuilder& InOutMenuBuilder) override;
 
 	//////////////////////////////////////////////////
 
@@ -137,45 +137,45 @@ public:
 	TSharedPtr<FMemoryGraphTrack> CreateMemoryGraphTrack();
 	int32 RemoveMemoryGraphTrack(TSharedPtr<FMemoryGraphTrack> GraphTrack);
 
-	TSharedPtr<FMemoryGraphTrack> GetMemTagGraphTrack(Insights::FMemoryTrackerId InMemTrackerId, Insights::FMemoryTagId InMemTagId);
-	TSharedPtr<FMemoryGraphTrack> CreateMemTagGraphTrack(Insights::FMemoryTrackerId InMemTrackerId, Insights::FMemoryTagId InMemTagId);
+	TSharedPtr<FMemoryGraphTrack> GetMemTagGraphTrack(FMemoryTrackerId InMemTrackerId, FMemoryTagId InMemTagId);
+	TSharedPtr<FMemoryGraphTrack> CreateMemTagGraphTrack(FMemoryTrackerId InMemTrackerId, FMemoryTagId InMemTagId);
 	void RemoveTrackFromMemTags(TSharedPtr<FMemoryGraphTrack>& GraphTrack);
-	int32 RemoveMemTagGraphTrack(Insights::FMemoryTrackerId InMemTrackerId, Insights::FMemoryTagId InMemTagId);
+	int32 RemoveMemTagGraphTrack(FMemoryTrackerId InMemTrackerId, FMemoryTagId InMemTagId);
 	int32 RemoveAllMemTagGraphTracks();
 
-	TSharedPtr<FMemoryGraphSeries> ToggleMemTagGraphSeries(TSharedPtr<FMemoryGraphTrack> InGraphTrack, Insights::FMemoryTrackerId InMemTrackerId, Insights::FMemoryTagId InMemTagId);
+	TSharedPtr<FMemoryGraphSeries> ToggleMemTagGraphSeries(TSharedPtr<FMemoryGraphTrack> InGraphTrack, FMemoryTrackerId InMemTrackerId, FMemoryTagId InMemTagId);
 
 	void CreateTracksFromReport(const FString& Filename);
-	void CreateTracksFromReport(const Insights::FReportConfig& ReportConfig);
-	void CreateTracksFromReport(const Insights::FReportTypeConfig& ReportTypeConfig);
+	void CreateTracksFromReport(const FReportConfig& ReportConfig);
+	void CreateTracksFromReport(const FReportTypeConfig& ReportTypeConfig);
 
-	const TArray<TSharedPtr<Insights::FMemoryRuleSpec>>& GetMemoryRules() const { return MemoryRules; }
+	const TArray<TSharedPtr<FMemoryRuleSpec>>& GetMemoryRules() const { return MemoryRules; }
 	
-	TSharedPtr<Insights::FMemoryRuleSpec> GetCurrentMemoryRule() const { return CurrentMemoryRule; }
-	void SetCurrentMemoryRule(TSharedPtr<Insights::FMemoryRuleSpec> InRule) { CurrentMemoryRule = InRule; OnMemoryRuleChanged(); }
+	TSharedPtr<FMemoryRuleSpec> GetCurrentMemoryRule() const { return CurrentMemoryRule; }
+	void SetCurrentMemoryRule(TSharedPtr<FMemoryRuleSpec> InRule) { CurrentMemoryRule = InRule; OnMemoryRuleChanged(); }
 
-	const TArray<TSharedPtr<Insights::FQueryTargetWindowSpec>>& GetQueryTargets() const { return QueryTargetSpecs; }
+	const TArray<TSharedPtr<FQueryTargetWindowSpec>>& GetQueryTargets() const { return QueryTargetSpecs; }
 
-	TSharedPtr<Insights::FQueryTargetWindowSpec> GetCurrentQueryTarget() const { return CurrentQueryTarget; }
-	void SetCurrentQueryTarget(TSharedPtr<Insights::FQueryTargetWindowSpec> InTarget) { CurrentQueryTarget = InTarget; }
-	void AddQueryTarget(TSharedPtr<Insights::FQueryTargetWindowSpec> InPtr);
-	void RemoveQueryTarget(TSharedPtr<Insights::FQueryTargetWindowSpec> InPtr);
+	TSharedPtr<FQueryTargetWindowSpec> GetCurrentQueryTarget() const { return CurrentQueryTarget; }
+	void SetCurrentQueryTarget(TSharedPtr<FQueryTargetWindowSpec> InTarget) { CurrentQueryTarget = InTarget; }
+	void AddQueryTarget(TSharedPtr<FQueryTargetWindowSpec> InPtr);
+	void RemoveQueryTarget(TSharedPtr<FQueryTargetWindowSpec> InPtr);
 
 private:
 	void SyncTrackers();
 	int32 GetNextMemoryGraphTrackOrder();
-	TSharedPtr<FMemoryGraphTrack> CreateGraphTrack(const Insights::FReportTypeGraphConfig& ReportTypeGraphConfig, bool bIsPlatformTracker);
+	TSharedPtr<FMemoryGraphTrack> CreateGraphTrack(const FReportTypeGraphConfig& ReportTypeGraphConfig, bool bIsPlatformTracker);
 	void InitMemoryRules();
 	void OnMemoryRuleChanged();
 
 private:
-	TSharedPtr<STimingView> TimingView;
+	TSharedPtr<TimingProfiler::STimingView> TimingView;
 
-	Insights::FMemoryTagList TagList;
+	FMemoryTagList TagList;
 
-	TArray<TSharedPtr<Insights::FMemoryTracker>> Trackers;
-	TSharedPtr<Insights::FMemoryTracker> DefaultTracker;
-	TSharedPtr<Insights::FMemoryTracker> PlatformTracker;
+	TArray<TSharedPtr<FMemoryTracker>> Trackers;
+	TSharedPtr<FMemoryTracker> DefaultTracker;
+	TSharedPtr<FMemoryTracker> PlatformTracker;
 
 	TSharedPtr<FMemoryGraphTrack> MainGraphTrack; // the Main Memory Graph track; also hosts the Total Allocated Memory series
 	TSharedPtr<FMemoryGraphTrack> LiveAllocsGraphTrack; // the graph track for the Live Allocation Count series
@@ -190,11 +190,13 @@ private:
 
 	TBitArray<> CreatedDefaultTracks;
 
-	TArray<TSharedPtr<Insights::FMemoryRuleSpec>> MemoryRules;
-	TSharedPtr<Insights::FMemoryRuleSpec> CurrentMemoryRule;
+	TArray<TSharedPtr<FMemoryRuleSpec>> MemoryRules;
+	TSharedPtr<FMemoryRuleSpec> CurrentMemoryRule;
 
-	TSharedPtr<Insights::FQueryTargetWindowSpec> CurrentQueryTarget;
-	TArray<TSharedPtr<Insights::FQueryTargetWindowSpec>> QueryTargetSpecs;
+	TSharedPtr<FQueryTargetWindowSpec> CurrentQueryTarget;
+	TArray<TSharedPtr<FQueryTargetWindowSpec>> QueryTargetSpecs;
 };
 
 ////////////////////////////////////////////////////////////////////////////////////////////////////
+
+} // namespace UE::Insights::MemoryProfiler
diff --git a/Engine/Source/Developer/TraceInsights/Private/Insights/MemoryProfiler/ViewModels/MemoryTag.cpp b/Engine/Source/Developer/TraceInsights/Private/Insights/MemoryProfiler/ViewModels/MemoryTag.cpp
--- a/Engine/Source/Developer/TraceInsights/Private/Insights/MemoryProfiler/ViewModels/MemoryTag.cpp
+++ b/Engine/Source/Developer/TraceInsights/Private/Insights/MemoryProfiler/ViewModels/MemoryTag.cpp
@@ -1,12 +1,14 @@
 // Copyright Epic Games, Inc. All Rights Reserved.
 
 #include "MemoryTag.h"
+
+// TraceServices
 #include "TraceServices/Model/Memory.h"
 
-// Insights
+// TraceInsights
 #include "Insights/InsightsManager.h"
 
-namespace Insights
+namespace UE::Insights::MemoryProfiler
 {
 
 ////////////////////////////////////////////////////////////////////////////////////////////////////
@@ -161,7 +163,7 @@ void FMemoryTagList::UpdateInternal()
 					{
 						if (TrackerFlags & 1)
 						{
-							TMap<FMemoryTagId, FMemoryTag*>& TagsMap = TrackersAndTagsMap.FindOrAdd(Insights::FMemoryTrackerId(TrackerId));
+							TMap<FMemoryTagId, FMemoryTag*>& TagsMap = TrackersAndTagsMap.FindOrAdd(FMemoryTrackerId(TrackerId));
 							FMemoryTag* TagPtr = TagsMap.FindRef(TagId);
 							if (!TagPtr)
 							{
@@ -173,7 +175,7 @@ void FMemoryTagList::UpdateInternal()
 								Tag.ParentId = static_cast<FMemoryTagId>(TraceTag.ParentId);
 								Tag.StatName = TraceTag.Name;
 								Tag.StatFullName = TraceTag.Name;
-								Tag.TrackerId = Insights::FMemoryTrackerId(TrackerId);
+								Tag.TrackerId = FMemoryTrackerId(TrackerId);
 								Tag.SetColorAuto();
 
 								// Skip the parent prefix if it is already included in the name.
@@ -270,4 +272,4 @@ int32 FMemoryTagList::FilterTags(const TArray<FString>& InIncludeStats, const TA
 
 ////////////////////////////////////////////////////////////////////////////////////////////////////
 
-} // namespace Insights
+} // namespace UE::Insights::MemoryProfiler
diff --git a/Engine/Source/Developer/TraceInsights/Private/Insights/MemoryProfiler/ViewModels/MemoryTag.h b/Engine/Source/Developer/TraceInsights/Private/Insights/MemoryProfiler/ViewModels/MemoryTag.h
--- a/Engine/Source/Developer/TraceInsights/Private/Insights/MemoryProfiler/ViewModels/MemoryTag.h
+++ b/Engine/Source/Developer/TraceInsights/Private/Insights/MemoryProfiler/ViewModels/MemoryTag.h
@@ -2,15 +2,21 @@
 
 #pragma once
 
-#include "CoreMinimal.h"
+#include "CoreTypes.h"
 
-#include "Insights/MemoryProfiler/ViewModels/MemoryTracker.h"
+#include "Containers/Array.h"
+#include "Containers/Map.h"
+#include "Containers/Set.h"
+#include "Math/Color.h"
+#include "Templates/SharedPointer.h"
 
-class FMemoryGraphTrack;
+#include "Insights/MemoryProfiler/ViewModels/MemoryTracker.h"
 
-namespace Insights
+namespace UE::Insights::MemoryProfiler
 {
 
+class FMemoryGraphTrack;
+
 typedef int64 FMemoryTagId;
 
 ////////////////////////////////////////////////////////////////////////////////////////////////////
@@ -58,7 +64,7 @@ private:
 	FMemoryTagId ParentId = InvalidTagId;
 	FString StatName;
 	FString StatFullName; // includes the parent prefix
-	Insights::FMemoryTrackerId TrackerId = Insights::FMemoryTracker::InvalidTrackerId;
+	FMemoryTrackerId TrackerId = FMemoryTracker::InvalidTrackerId;
 	FLinearColor Color;
 	FMemoryTag* Parent = nullptr;
 	TSet<FMemoryTag*> Children;
@@ -101,4 +107,4 @@ private:
 
 ////////////////////////////////////////////////////////////////////////////////////////////////////
 
-} // namespace Insights
+} // namespace UE::Insights::MemoryProfiler
diff --git a/Engine/Source/Developer/TraceInsights/Private/Insights/MemoryProfiler/ViewModels/MemoryTracker.cpp b/Engine/Source/Developer/TraceInsights/Private/Insights/MemoryProfiler/ViewModels/MemoryTracker.cpp
--- a/Engine/Source/Developer/TraceInsights/Private/Insights/MemoryProfiler/ViewModels/MemoryTracker.cpp
+++ b/Engine/Source/Developer/TraceInsights/Private/Insights/MemoryProfiler/ViewModels/MemoryTracker.cpp
@@ -2,11 +2,11 @@
 
 #include "MemoryTracker.h"
 
-// Insights
+// TraceInsights
 #include "Insights/InsightsManager.h"
 #include "Insights/MemoryProfiler/ViewModels/MemoryTag.h"
 
-namespace Insights
+namespace UE::Insights::MemoryProfiler
 {
 
 ////////////////////////////////////////////////////////////////////////////////////////////////////
@@ -28,4 +28,4 @@ FMemoryTracker::~FMemoryTracker()
 
 ////////////////////////////////////////////////////////////////////////////////////////////////////
 
-} // namespace Insights
+} // namespace UE::Insights::MemoryProfiler
diff --git a/Engine/Source/Developer/TraceInsights/Private/Insights/MemoryProfiler/ViewModels/MemoryTracker.h b/Engine/Source/Developer/TraceInsights/Private/Insights/MemoryProfiler/ViewModels/MemoryTracker.h
--- a/Engine/Source/Developer/TraceInsights/Private/Insights/MemoryProfiler/ViewModels/MemoryTracker.h
+++ b/Engine/Source/Developer/TraceInsights/Private/Insights/MemoryProfiler/ViewModels/MemoryTracker.h
@@ -2,9 +2,10 @@
 
 #pragma once
 
-#include "CoreMinimal.h"
+#include "CoreTypes.h"
+#include "Containers/UnrealString.h"
 
-namespace Insights
+namespace UE::Insights::MemoryProfiler
 {
 
 typedef int32 FMemoryTrackerId;
@@ -36,4 +37,4 @@ private:
 
 ////////////////////////////////////////////////////////////////////////////////////////////////////
 
-} // namespace Insights
+} // namespace UE::Insights::MemoryProfiler
diff --git a/Engine/Source/Developer/TraceInsights/Private/Insights/MemoryProfiler/ViewModels/Report.h b/Engine/Source/Developer/TraceInsights/Private/Insights/MemoryProfiler/ViewModels/Report.h
--- a/Engine/Source/Developer/TraceInsights/Private/Insights/MemoryProfiler/ViewModels/Report.h
+++ b/Engine/Source/Developer/TraceInsights/Private/Insights/MemoryProfiler/ViewModels/Report.h
@@ -4,7 +4,7 @@
 
 #include "CoreMinimal.h"
 
-namespace Insights
+namespace UE::Insights::MemoryProfiler
 {
 
 struct FGraphConfig
@@ -102,4 +102,4 @@ private:
 	FReportConfig(const FReportConfig&) = delete;
 };
 
-} // namespace Insights
+} // namespace UE::Insights::MemoryProfiler
diff --git a/Engine/Source/Developer/TraceInsights/Private/Insights/MemoryProfiler/ViewModels/ReportXmlParser.cpp b/Engine/Source/Developer/TraceInsights/Private/Insights/MemoryProfiler/ViewModels/ReportXmlParser.cpp
--- a/Engine/Source/Developer/TraceInsights/Private/Insights/MemoryProfiler/ViewModels/ReportXmlParser.cpp
+++ b/Engine/Source/Developer/TraceInsights/Private/Insights/MemoryProfiler/ViewModels/ReportXmlParser.cpp
@@ -7,14 +7,14 @@
 #include "XmlNode.h"
 #include "XmlFile.h"
 
-// Insights
+// TraceInsights
 #include "Insights/Log.h"
 #include "Insights/MemoryProfiler/ViewModels/Report.h"
 
-namespace Insights
-{
+#define LOCTEXT_NAMESPACE "UE::Insights::MemoryProfiler::FReportXmlParser"
 
-#define LOCTEXT_NAMESPACE "Insights.ReportXmlParser"
+namespace UE::Insights::MemoryProfiler
+{
 
 ////////////////////////////////////////////////////////////////////////////////////////////////////
 
@@ -524,6 +524,6 @@ void FReportXmlParser::LogWarning(const FText& Text)
 
 ////////////////////////////////////////////////////////////////////////////////////////////////////
 
-#undef LOCTEXT_NAMESPACE
+} // namespace UE::Insights::MemoryProfiler
 
-} // namespace Insights
+#undef LOCTEXT_NAMESPACE
diff --git a/Engine/Source/Developer/TraceInsights/Private/Insights/MemoryProfiler/ViewModels/ReportXmlParser.h b/Engine/Source/Developer/TraceInsights/Private/Insights/MemoryProfiler/ViewModels/ReportXmlParser.h
--- a/Engine/Source/Developer/TraceInsights/Private/Insights/MemoryProfiler/ViewModels/ReportXmlParser.h
+++ b/Engine/Source/Developer/TraceInsights/Private/Insights/MemoryProfiler/ViewModels/ReportXmlParser.h
@@ -8,7 +8,7 @@ class FXmlNode;
 class FXmlAttribute;
 class FTokenizedMessage;
 
-namespace Insights
+namespace UE::Insights::MemoryProfiler
 {
 
 struct FGraphConfig;
@@ -49,4 +49,4 @@ private:
 	EStatus Status = EStatus::Completed;
 };
 
-} // namespace Insights
+} // namespace UE::Insights::MemoryProfiler
