[Insights] Massive refactoring of TraceInsights module:

diff --git a/Engine/Source/Developer/TraceInsights/Private/Insights/MemoryProfiler/Widgets/SMemAllocTableTreeView.cpp b/Engine/Source/Developer/TraceInsights/Private/Insights/MemoryProfiler/Widgets/SMemAllocTableTreeView.cpp
--- a/Engine/Source/Developer/TraceInsights/Private/Insights/MemoryProfiler/Widgets/SMemAllocTableTreeView.cpp
+++ b/Engine/Source/Developer/TraceInsights/Private/Insights/MemoryProfiler/Widgets/SMemAllocTableTreeView.cpp
@@ -3,10 +3,13 @@
 #include "SMemAllocTableTreeView.h"
 
 #include "Containers/Set.h"
+#include "DesktopPlatformModule.h"
 #include "Framework/MultiBox/MultiBoxBuilder.h"
 #include "GenericPlatform/GenericPlatformMisc.h"
+#include "HAL/PlatformFileManager.h"
 #include "ISourceCodeAccessModule.h"
 #include "ISourceCodeAccessor.h"
+#include "Logging/MessageLog.h"
 #include "Modules/ModuleManager.h"
 #include "SlateOptMacros.h"
 #include "Styling/AppStyle.h"
@@ -14,9 +17,6 @@
 #include "Widgets/Images/SImage.h"
 #include "Widgets/Input/SCheckBox.h"
 #include "Widgets/SToolTip.h"
-#include "DesktopPlatformModule.h"
-#include "HAL/PlatformFileManager.h"
-#include "Logging/MessageLog.h"
 
 // TraceServices
 #include "Common/ProviderLock.h"
@@ -27,32 +27,34 @@
 #include "TraceServices/Model/Modules.h"
 #include "TraceServices/Model/Strings.h"
 
-// Insights
+// TraceInsightsCore
+#include "InsightsCore/Filter/ViewModels/FilterConfigurator.h"
+#include "InsightsCore/Filter/ViewModels/Filters.h"
+#include "InsightsCore/Table/ViewModels/TableCellValueFormatter.h"
+#include "InsightsCore/Table/ViewModels/TableColumn.h"
+
+// TraceInsights
 #include "Insights/InsightsStyle.h"
 #include "Insights/MemoryProfiler/Common/SymbolSearchPathsHelper.h"
 #include "Insights/MemoryProfiler/MemoryProfilerManager.h"
 #include "Insights/MemoryProfiler/ViewModels/CallstackFormatting.h"
 #include "Insights/MemoryProfiler/ViewModels/MemAllocGroupingByCallstack.h"
 #include "Insights/MemoryProfiler/ViewModels/MemAllocGroupingByHeap.h"
-#include "Insights/MemoryProfiler/ViewModels/MemAllocGroupingBySwapPage.h"
 #include "Insights/MemoryProfiler/ViewModels/MemAllocGroupingBySize.h"
+#include "Insights/MemoryProfiler/ViewModels/MemAllocGroupingBySwapPage.h"
 #include "Insights/MemoryProfiler/ViewModels/MemAllocGroupingByTag.h"
 #include "Insights/MemoryProfiler/ViewModels/MemAllocNode.h"
 #include "Insights/MemoryProfiler/ViewModels/MemAllocTable.h"
 #include "Insights/MemoryProfiler/ViewModels/MemorySharedState.h"
 #include "Insights/MemoryProfiler/Widgets/SMemoryProfilerWindow.h"
-#include "Insights/Table/ViewModels/TableCellValueFormatter.h"
-#include "Insights/Table/ViewModels/TableColumn.h"
 #include "Insights/TimingProfilerCommon.h"
-#include "Insights/ViewModels/FilterConfigurator.h"
-#include "Insights/ViewModels/Filters.h"
 
 #include <limits>
 #include <memory>
 
-#define LOCTEXT_NAMESPACE "SMemAllocTableTreeView"
+#define LOCTEXT_NAMESPACE "UE::Insights::MemoryProfiler::SMemAllocTableTreeView"
 
-namespace Insights
+namespace UE::Insights::MemoryProfiler
 {
 
 const int32 SMemAllocTableTreeView::FullCallStackIndex = 0x0000FFFFF;
@@ -73,7 +75,7 @@ SMemAllocTableTreeView::~SMemAllocTableTreeView()
 
 ////////////////////////////////////////////////////////////////////////////////////////////////////
 
-void SMemAllocTableTreeView::Construct(const FArguments& InArgs, TSharedPtr<Insights::FMemAllocTable> InTablePtr)
+void SMemAllocTableTreeView::Construct(const FArguments& InArgs, TSharedPtr<FMemAllocTable> InTablePtr)
 {
 	ConstructWidget(InTablePtr);
 }
@@ -121,7 +123,7 @@ void SMemAllocTableTreeView::RebuildTree(bool bResync)
 
 	const int32 PreviousNodeCount = TableRowNodes.Num();
 
-	TSharedPtr<Insights::FMemAllocTable> MemAllocTable = GetMemAllocTable();
+	TSharedPtr<FMemAllocTable> MemAllocTable = GetMemAllocTable();
 
 	if (Session.IsValid() && MemAllocTable.IsValid())
 	{
@@ -130,13 +132,13 @@ void SMemAllocTableTreeView::RebuildTree(bool bResync)
 
 		if (QueryStatus == TraceServices::IAllocationsProvider::EQueryStatus::Done)
 		{
-			UE_LOG(MemoryProfiler, Log, TEXT("[MemAlloc] Rebuilding tree..."));
+			UE_LOG(LogMemoryProfiler, Log, TEXT("[MemAlloc] Rebuilding tree..."));
 			TArray<FMemoryAlloc>& Allocs = MemAllocTable->GetAllocs();
 
 			const int32 TotalAllocCount = Allocs.Num();
 			if (TotalAllocCount != TableRowNodes.Num())
 			{
-				UE_LOG(MemoryProfiler, Log, TEXT("[MemAlloc] Creating nodes (%d nodes --> %d allocs)..."), TableRowNodes.Num(), TotalAllocCount);
+				UE_LOG(LogMemoryProfiler, Log, TEXT("[MemAlloc] Creating nodes (%d nodes --> %d allocs)..."), TableRowNodes.Num(), TotalAllocCount);
 
 				if (TableRowNodes.Num() > TotalAllocCount)
 				{
@@ -214,7 +216,7 @@ void SMemAllocTableTreeView::RebuildTree(bool bResync)
 	if (TotalTime > 0.01)
 	{
 		const double SyncTime = SyncStopwatch.GetAccumulatedTime();
-		UE_LOG(MemoryProfiler, Log, TEXT("[MemAlloc] Tree view rebuilt in %.4fs (sync: %.4fs + update: %.4fs) --> %d nodes (%d added)"),
+		UE_LOG(LogMemoryProfiler, Log, TEXT("[MemAlloc] Tree view rebuilt in %.4fs (sync: %.4fs + update: %.4fs) --> %d nodes (%d added)"),
 			TotalTime, SyncTime, TotalTime - SyncTime, TableRowNodes.Num(), TableRowNodes.Num() - PreviousNodeCount);
 	}
 }
@@ -243,7 +245,7 @@ void SMemAllocTableTreeView::ResetAndStartQuery()
 {
 	TableRowNodes.Reset();
 
-	TSharedPtr<Insights::FMemAllocTable> MemAllocTable = GetMemAllocTable();
+	TSharedPtr<FMemAllocTable> MemAllocTable = GetMemAllocTable();
 	if (MemAllocTable)
 	{
 		MemAllocTable->SetTimeMarkerA(TimeMarkers[0]); // to be used by LLM Size and LLM Delta Size columns
@@ -265,20 +267,20 @@ void SMemAllocTableTreeView::StartQuery()
 
 	if (!Rule)
 	{
-		UE_LOG(MemoryProfiler, Warning, TEXT("[MemAlloc] Invalid query rule!"));
+		UE_LOG(LogMemoryProfiler, Warning, TEXT("[MemAlloc] Invalid query rule!"));
 		return;
 	}
 
 	if (!Session.IsValid())
 	{
-		UE_LOG(MemoryProfiler, Warning, TEXT("[MemAlloc] Invalid analysis session!"));
+		UE_LOG(LogMemoryProfiler, Warning, TEXT("[MemAlloc] Invalid analysis session!"));
 		return;
 	}
 
 	const TraceServices::IAllocationsProvider* AllocationsProvider = TraceServices::ReadAllocationsProvider(*Session.Get());
 	if (!AllocationsProvider)
 	{
-		UE_LOG(MemoryProfiler, Warning, TEXT("[MemAlloc] Invalid allocations provider!"));
+		UE_LOG(LogMemoryProfiler, Warning, TEXT("[MemAlloc] Invalid allocations provider!"));
 		return;
 	}
 
@@ -286,7 +288,7 @@ void SMemAllocTableTreeView::StartQuery()
 		const TraceServices::IAllocationsProvider& Provider = *AllocationsProvider;
 		TraceServices::FProviderReadScopeLock _(Provider);
 
-		TSharedPtr<Insights::FMemAllocTable> MemAllocTable = GetMemAllocTable();
+		TSharedPtr<FMemAllocTable> MemAllocTable = GetMemAllocTable();
 		if (MemAllocTable)
 		{
 			MemAllocTable->SetPlatformPageSize(Provider.GetPlatformPageSize());
@@ -298,7 +300,7 @@ void SMemAllocTableTreeView::StartQuery()
 
 	if (Query == 0)
 	{
-		UE_LOG(MemoryProfiler, Error, TEXT("[MemAlloc] Unsupported query rule (%s)!"), *Rule->GetShortName().ToString());
+		UE_LOG(LogMemoryProfiler, Error, TEXT("[MemAlloc] Unsupported query rule (%s)!"), *Rule->GetShortName().ToString());
 	}
 	else
 	{
@@ -321,14 +323,14 @@ void SMemAllocTableTreeView::UpdateQuery(TraceServices::IAllocationsProvider::EQ
 
 	if (!Session.IsValid())
 	{
-		UE_LOG(MemoryProfiler, Warning, TEXT("[MemAlloc] Invalid analysis session!"));
+		UE_LOG(LogMemoryProfiler, Warning, TEXT("[MemAlloc] Invalid analysis session!"));
 		return;
 	}
 
 	const TraceServices::IAllocationsProvider* AllocationsProvider = TraceServices::ReadAllocationsProvider(*Session.Get());
 	if (!AllocationsProvider)
 	{
-		UE_LOG(MemoryProfiler, Warning, TEXT("[MemAlloc] Invalid allocations provider!"));
+		UE_LOG(LogMemoryProfiler, Warning, TEXT("[MemAlloc] Invalid allocations provider!"));
 		return;
 	}
 	const TraceServices::IAllocationsProvider& Provider = *AllocationsProvider;
@@ -367,7 +369,7 @@ void SMemAllocTableTreeView::UpdateQuery(TraceServices::IAllocationsProvider::EQ
 
 		if (Status.Status <= TraceServices::IAllocationsProvider::EQueryStatus::Done)
 		{
-			UE_LOG(MemoryProfiler, Log, TEXT("[MemAlloc] Query completed."));
+			UE_LOG(LogMemoryProfiler, Log, TEXT("[MemAlloc] Query completed."));
 			Query = 0;
 			QueryStopwatch.Stop();
 			return;
@@ -380,7 +382,7 @@ void SMemAllocTableTreeView::UpdateQuery(TraceServices::IAllocationsProvider::EQ
 
 		check(Status.Status == TraceServices::IAllocationsProvider::EQueryStatus::Available);
 
-		TSharedPtr<Insights::FMemAllocTable> MemAllocTable = GetMemAllocTable();
+		TSharedPtr<FMemAllocTable> MemAllocTable = GetMemAllocTable();
 		if (MemAllocTable)
 		{
 			TraceServices::FProviderReadScopeLock _(Provider);
@@ -398,7 +400,7 @@ void SMemAllocTableTreeView::UpdateQuery(TraceServices::IAllocationsProvider::EQ
 			TraceServices::IAllocationsProvider::FQueryResult Result = Status.NextResult();
 			while (Result.IsValid())
 			{
-				UE_LOG(MemoryProfiler, Log, TEXT("[MemAlloc] Page with %u allocs..."), Result->Num());
+				UE_LOG(LogMemoryProfiler, Log, TEXT("[MemAlloc] Page with %u allocs..."), Result->Num());
 
 				++PageCount;
 				PageStopwatch.Restart();
@@ -505,7 +507,7 @@ void SMemAllocTableTreeView::UpdateQuery(TraceServices::IAllocationsProvider::EQ
 				if (PageTime > 0.01)
 				{
 					const double Speed = (PageTime * 1000000.0) / AllocCount;
-					UE_LOG(MemoryProfiler, Log, TEXT("[MemAlloc] Query result for page %u (%u allocs, slack=%u) retrieved in %.3fs (speed: %.3f seconds per 1M allocs)."), PageCount, AllocCount, Allocs.GetSlack(), PageTime, Speed);
+					UE_LOG(LogMemoryProfiler, Log, TEXT("[MemAlloc] Query result for page %u (%u allocs, slack=%u) retrieved in %.3fs (speed: %.3f seconds per 1M allocs)."), PageCount, AllocCount, Allocs.GetSlack(), PageTime, Speed);
 				}
 
 				Result = Status.NextResult();
@@ -516,7 +518,7 @@ void SMemAllocTableTreeView::UpdateQuery(TraceServices::IAllocationsProvider::EQ
 			if (TotalTime > 0.01)
 			{
 				const double Speed = (TotalTime * 1000000.0) / TotalAllocCount;
-				UE_LOG(MemoryProfiler, Log, TEXT("[MemAlloc] Query results (%u pages, %u allocs, slack=%u) retrieved in %.3fs (speed: %.3f seconds per 1M allocs)."), PageCount, TotalAllocCount, Allocs.GetSlack(), TotalTime, Speed);
+				UE_LOG(LogMemoryProfiler, Log, TEXT("[MemAlloc] Query results (%u pages, %u allocs, slack=%u) retrieved in %.3fs (speed: %.3f seconds per 1M allocs)."), PageCount, TotalAllocCount, Allocs.GetSlack(), TotalTime, Speed);
 			}
 		}
 
@@ -539,7 +541,7 @@ void SMemAllocTableTreeView::CancelQuery()
 			if (AllocationsProvider)
 			{
 				AllocationsProvider->CancelQuery(Query);
-				UE_LOG(MemoryProfiler, Log, TEXT("[MemAlloc] Query canceled."));
+				UE_LOG(LogMemoryProfiler, Log, TEXT("[MemAlloc] Query canceled."));
 			}
 		}
 
@@ -1364,7 +1366,7 @@ void SMemAllocTableTreeView::TreeView_OnSelectionChanged(FTableTreeNodePtr Selec
 		TSharedRef<FTableColumn> CountColumn = Table->FindColumnChecked(FMemAllocTableColumns::CountColumnId);
 		TSharedRef<FTableColumn> SizeColumn = Table->FindColumnChecked(FMemAllocTableColumns::SizeColumnId);
 
-		for (const Insights::FTableTreeNodePtr& Node : SelectedNodes)
+		for (const FTableTreeNodePtr& Node : SelectedNodes)
 		{
 			TOptional<FTableCellValue> CountValue = CountColumn->GetValue(*Node.Get());
 			if (CountValue.IsSet())
@@ -1395,7 +1397,7 @@ void SMemAllocTableTreeView::TreeView_OnSelectionChanged(FTableTreeNodePtr Selec
 
 void SMemAllocTableTreeView::InternalCreateGroupings()
 {
-	Insights::STableTreeView::InternalCreateGroupings();
+	STableTreeView::InternalCreateGroupings();
 
 	AvailableGroupings.RemoveAll(
 		[](TSharedPtr<FTreeNodeGrouping>& Grouping)
@@ -2313,6 +2315,6 @@ FText SMemAllocTableTreeView::GetSelectedCallstackFrameFileName() const
 
 ////////////////////////////////////////////////////////////////////////////////////////////////////
 
-} // namespace Insights
+} // namespace UE::Insights::MemoryProfiler
 
 #undef LOCTEXT_NAMESPACE
diff --git a/Engine/Source/Developer/TraceInsights/Private/Insights/MemoryProfiler/Widgets/SMemAllocTableTreeView.h b/Engine/Source/Developer/TraceInsights/Private/Insights/MemoryProfiler/Widgets/SMemAllocTableTreeView.h
--- a/Engine/Source/Developer/TraceInsights/Private/Insights/MemoryProfiler/Widgets/SMemAllocTableTreeView.h
+++ b/Engine/Source/Developer/TraceInsights/Private/Insights/MemoryProfiler/Widgets/SMemAllocTableTreeView.h
@@ -3,14 +3,18 @@
 #pragma once
 
 #include "CoreMinimal.h"
+
+// TraceServices
 #include "TraceServices/Model/AllocationsProvider.h"
 
-// Insights
-#include "Insights/Common/Stopwatch.h"
+// TraceInsightsCore
+#include "InsightsCore/Common/Stopwatch.h"
+
+// TraceInsights
 #include "Insights/MemoryProfiler/ViewModels/MemAllocTable.h"
 #include "Insights/Table/Widgets/SSessionTableTreeView.h"
 
-namespace Insights
+namespace UE::Insights::MemoryProfiler
 {
 
 class FCallstackFrameGroupNode;
@@ -131,7 +135,7 @@ private:
 	FText GetFooterLeftText() const;
 	FText GetFooterCenterText() const;
 
-	virtual void TreeView_OnSelectionChanged(Insights::FTableTreeNodePtr SelectedItem, ESelectInfo::Type SelectInfo) override;
+	virtual void TreeView_OnSelectionChanged(FTableTreeNodePtr SelectedItem, ESelectInfo::Type SelectInfo) override;
 
 	virtual void UpdateFilterContext(const FFilterConfigurator& InFilterConfigurator, const FTableTreeNode& InNode) const override;
 	virtual void InitFilterConfigurator(FFilterConfigurator& InOutFilterConfigurator) override;
@@ -162,4 +166,4 @@ private:
 
 ////////////////////////////////////////////////////////////////////////////////////////////////////
 
-} // namespace Insights
+} // namespace UE::Insights::MemoryProfiler
diff --git a/Engine/Source/Developer/TraceInsights/Private/Insights/MemoryProfiler/Widgets/SMemInvestigationView.cpp b/Engine/Source/Developer/TraceInsights/Private/Insights/MemoryProfiler/Widgets/SMemInvestigationView.cpp
--- a/Engine/Source/Developer/TraceInsights/Private/Insights/MemoryProfiler/Widgets/SMemInvestigationView.cpp
+++ b/Engine/Source/Developer/TraceInsights/Private/Insights/MemoryProfiler/Widgets/SMemInvestigationView.cpp
@@ -4,24 +4,29 @@
 
 #include "DesktopPlatformModule.h"
 #include "SlateOptMacros.h"
+#include "Widgets/Images/SImage.h"
 #include "Widgets/Input/SButton.h"
 #include "Widgets/Input/SCheckBox.h"
 #include "Widgets/Input/SEditableTextBox.h"
-#include "Widgets/Images/SImage.h"
 
-// Insights
-#include "Insights/Common/Stopwatch.h"
-#include "Insights/Common/TimeUtils.h"
+// TraceInsightsCore
+#include "InsightsCore/Common/Stopwatch.h"
+#include "InsightsCore/Common/TimeUtils.h"
+
+// TraceInsights
 #include "Insights/InsightsStyle.h"
-#include "Insights/MemoryProfiler/MemoryProfilerManager.h"
 #include "Insights/MemoryProfiler/Common/SymbolSearchPathsHelper.h"
+#include "Insights/MemoryProfiler/MemoryProfilerManager.h"
 #include "Insights/MemoryProfiler/ViewModels/MemorySharedState.h"
 #include "Insights/MemoryProfiler/Widgets/SMemAllocTableTreeView.h"
 #include "Insights/MemoryProfiler/Widgets/SMemoryProfilerWindow.h"
 #include "Insights/ViewModels/TimeRulerTrack.h"
 #include "Insights/Widgets/STimingView.h"
 
-#define LOCTEXT_NAMESPACE "SMemInvestigationView"
+#define LOCTEXT_NAMESPACE "UE::Insights::MemoryProfiler::SMemInvestigationView"
+
+namespace UE::Insights::MemoryProfiler
+{
 
 ////////////////////////////////////////////////////////////////////////////////////////////////////
 
@@ -97,7 +102,7 @@ TSharedRef<SWidget> SMemInvestigationView::ConstructInvestigationWidgetArea()
 			.FillWidth(1.0f)
 			.VAlign(VAlign_Center)
 			[
-				SAssignNew(QueryRuleComboBox, SComboBox<TSharedPtr<Insights::FMemoryRuleSpec>>)
+				SAssignNew(QueryRuleComboBox, SComboBox<TSharedPtr<FMemoryRuleSpec>>)
 				.ToolTipText(this, &SMemInvestigationView::QueryRule_GetTooltipText)
 				.OptionsSource(GetAvailableQueryRules())
 				.OnSelectionChanged(this, &SMemInvestigationView::QueryRule_OnSelectionChanged)
@@ -204,7 +209,7 @@ TSharedRef<SWidget> SMemInvestigationView::ConstructInvestigationWidgetArea()
 			.FillWidth(1.0f)
 			.VAlign(VAlign_Center)
 			[
-				SAssignNew(QueryTargetComboBox, SComboBox<TSharedPtr<Insights::FQueryTargetWindowSpec>>)
+				SAssignNew(QueryTargetComboBox, SComboBox<TSharedPtr<FQueryTargetWindowSpec>>)
 				.ToolTipText(LOCTEXT("QueryTargetWindowTooltip", "Select an existing or new window where the query results will be displayed"))
 				.OptionsSource(GetAvailableQueryTargets())
 				.OnSelectionChanged(this, &SMemInvestigationView::QueryTarget_OnSelectionChanged)
@@ -252,7 +257,7 @@ END_SLATE_FUNCTION_BUILD_OPTIMIZATION
 ////////////////////////////////////////////////////////////////////////////////////////////////////
 
 BEGIN_SLATE_FUNCTION_BUILD_OPTIMIZATION
-TSharedRef<SWidget> SMemInvestigationView::QueryRule_OnGenerateWidget(TSharedPtr<Insights::FMemoryRuleSpec> InRule)
+TSharedRef<SWidget> SMemInvestigationView::QueryRule_OnGenerateWidget(TSharedPtr<FMemoryRuleSpec> InRule)
 {
 	const FText QueryRuleText = FText::Format(LOCTEXT("QueryRuleComboBox_TextFmt", "{0} ({1})"), InRule->GetVerboseName(), InRule->GetShortName());
 
@@ -300,8 +305,9 @@ TSharedRef<SWidget> SMemInvestigationView::ConstructTimeMarkerWidget(uint32 Time
 		return SNew(SBox);
 	}
 
-	const TSharedRef<Insights::FTimeMarker>& TimeMarker = ProfilerWindow->GetCustomTimeMarker(TimeMarkerIndex);
-	TSharedPtr<Insights::FTimeMarker> PreviousTimeMarker;
+	using FTimeMarker = TimingProfiler::FTimeMarker;
+	const TSharedRef<FTimeMarker>& TimeMarker = ProfilerWindow->GetCustomTimeMarker(TimeMarkerIndex);
+	TSharedPtr<FTimeMarker> PreviousTimeMarker;
 	if (TimeMarkerIndex > 0)
 	{
 		PreviousTimeMarker = ProfilerWindow->GetCustomTimeMarker(TimeMarkerIndex - 1);
@@ -342,7 +348,7 @@ TSharedRef<SWidget> SMemInvestigationView::ConstructTimeMarkerWidget(uint32 Time
 					TSharedPtr<SMemoryProfilerWindow> ProfilerWindow = GetProfilerWindow();
 					if (ProfilerWindow.IsValid())
 					{
-						ProfilerWindow->OnTimeMarkerChanged(Insights::ETimeChangedFlags::None, TimeMarker);
+						ProfilerWindow->OnTimeMarkerChanged(Timing::ETimeChangedFlags::None, TimeMarker);
 					}
 				})
 		]
@@ -358,12 +364,12 @@ TSharedRef<SWidget> SMemInvestigationView::ConstructTimeMarkerWidget(uint32 Time
 					if (PreviousTimeMarker.IsValid())
 					{
 						return FText::FromString(FString::Printf(TEXT("%s (+%s)"),
-							*TimeUtils::FormatTime(TimeMarker->GetTime(), 0.1),
-							*TimeUtils::FormatTime(TimeMarker->GetTime() - PreviousTimeMarker->GetTime(), 0.1)));
+							*FormatTime(TimeMarker->GetTime(), 0.1),
+							*FormatTime(TimeMarker->GetTime() - PreviousTimeMarker->GetTime(), 0.1)));
 					}
 					else
 					{
-						return FText::FromString(TimeUtils::FormatTime(TimeMarker->GetTime(), 0.1));
+						return FText::FromString(FormatTime(TimeMarker->GetTime(), 0.1));
 					}
 				})
 		];
@@ -402,7 +408,7 @@ void SMemInvestigationView::Tick(const FGeometry& AllottedGeometry, const double
 // Query Rules
 ////////////////////////////////////////////////////////////////////////////////////////////////////
 
-const TArray<TSharedPtr<Insights::FMemoryRuleSpec>>* SMemInvestigationView::GetAvailableQueryRules()
+const TArray<TSharedPtr<FMemoryRuleSpec>>* SMemInvestigationView::GetAvailableQueryRules()
 {
 	TSharedPtr<SMemoryProfilerWindow> ProfilerWindow = GetProfilerWindow();
 	if (ProfilerWindow.IsValid())
@@ -415,7 +421,7 @@ const TArray<TSharedPtr<Insights::FMemoryRuleSpec>>* SMemInvestigationView::GetA
 
 ////////////////////////////////////////////////////////////////////////////////////////////////////
 
-void SMemInvestigationView::QueryRule_OnSelectionChanged(TSharedPtr<Insights::FMemoryRuleSpec> InRule, ESelectInfo::Type SelectInfo)
+void SMemInvestigationView::QueryRule_OnSelectionChanged(TSharedPtr<FMemoryRuleSpec> InRule, ESelectInfo::Type SelectInfo)
 {
 	if (SelectInfo != ESelectInfo::Direct)
 	{
@@ -437,7 +443,7 @@ FText SMemInvestigationView::QueryRule_GetSelectedText() const
 	if (ProfilerWindow.IsValid())
 	{
 		FMemorySharedState& SharedState = ProfilerWindow->GetSharedState();
-		TSharedPtr<Insights::FMemoryRuleSpec> Rule = SharedState.GetCurrentMemoryRule();
+		TSharedPtr<FMemoryRuleSpec> Rule = SharedState.GetCurrentMemoryRule();
 		if (Rule)
 		{
 			return FText::Format(LOCTEXT("QueryRuleComboBox_TextFmt", "{0} ({1})"), Rule->GetVerboseName(), Rule->GetShortName());
@@ -454,7 +460,7 @@ FText SMemInvestigationView::QueryRule_GetTooltipText() const
 	if (ProfilerWindow.IsValid())
 	{
 		FMemorySharedState& SharedState = ProfilerWindow->GetSharedState();
-		TSharedPtr<Insights::FMemoryRuleSpec> Rule = SharedState.GetCurrentMemoryRule();
+		TSharedPtr<FMemoryRuleSpec> Rule = SharedState.GetCurrentMemoryRule();
 		if (Rule.IsValid())
 		{
 			return Rule->GetDescription();
@@ -465,7 +471,7 @@ FText SMemInvestigationView::QueryRule_GetTooltipText() const
 
 ////////////////////////////////////////////////////////////////////////////////////////////////////
 
-const TArray<TSharedPtr<Insights::FQueryTargetWindowSpec>>* SMemInvestigationView::GetAvailableQueryTargets()
+const TArray<TSharedPtr<FQueryTargetWindowSpec>>* SMemInvestigationView::GetAvailableQueryTargets()
 {
 	TSharedPtr<SMemoryProfilerWindow> ProfilerWindow = GetProfilerWindow();
 	if (ProfilerWindow.IsValid())
@@ -478,7 +484,7 @@ const TArray<TSharedPtr<Insights::FQueryTargetWindowSpec>>* SMemInvestigationVie
 
 ////////////////////////////////////////////////////////////////////////////////////////////////////
 
-void SMemInvestigationView::QueryTarget_OnSelectionChanged(TSharedPtr<Insights::FQueryTargetWindowSpec> InTarget, ESelectInfo::Type SelectInfo)
+void SMemInvestigationView::QueryTarget_OnSelectionChanged(TSharedPtr<FQueryTargetWindowSpec> InTarget, ESelectInfo::Type SelectInfo)
 {
 	if (SelectInfo != ESelectInfo::Type::Direct)
 	{
@@ -510,7 +516,7 @@ void SMemInvestigationView::UpdateSymbolPathsText() const
 
 ////////////////////////////////////////////////////////////////////////////////////////////////////
 
-TSharedRef<SWidget> SMemInvestigationView::QueryTarget_OnGenerateWidget(TSharedPtr<Insights::FQueryTargetWindowSpec> InTarget)
+TSharedRef<SWidget> SMemInvestigationView::QueryTarget_OnGenerateWidget(TSharedPtr<FQueryTargetWindowSpec> InTarget)
 {
 	return SNew(STextBlock)
 		.Text(InTarget->GetText())
@@ -525,7 +531,7 @@ FText SMemInvestigationView::QueryTarget_GetSelectedText() const
 	if (ProfilerWindow.IsValid())
 	{
 		FMemorySharedState& SharedState = ProfilerWindow->GetSharedState();
-		TSharedPtr<Insights::FQueryTargetWindowSpec> Target = SharedState.GetCurrentQueryTarget();
+		TSharedPtr<FQueryTargetWindowSpec> Target = SharedState.GetCurrentQueryTarget();
 		if (Target)
 		{
 			return Target->GetText();
@@ -541,15 +547,15 @@ FReply SMemInvestigationView::RunQuery()
 	TSharedPtr<SMemoryProfilerWindow> ProfilerWindow = GetProfilerWindow();
 	if (!ProfilerWindow.IsValid())
 	{
-		UE_LOG(MemoryProfiler, Error, TEXT("[MemQuery] Invalid Profiler Window!"));
+		UE_LOG(LogMemoryProfiler, Error, TEXT("[MemQuery] Invalid Profiler Window!"));
 		return FReply::Handled();
 	}
 
 	FMemorySharedState& SharedState = ProfilerWindow->GetSharedState();
-	TSharedPtr<Insights::FMemoryRuleSpec> Rule = SharedState.GetCurrentMemoryRule();
+	TSharedPtr<FMemoryRuleSpec> Rule = SharedState.GetCurrentMemoryRule();
 	if (!Rule)
 	{
-		UE_LOG(MemoryProfiler, Error, TEXT("[MemQuery] Invalid Rule!"));
+		UE_LOG(LogMemoryProfiler, Error, TEXT("[MemQuery] Invalid Rule!"));
 		return FReply::Handled();
 	}
 
@@ -557,7 +563,7 @@ FReply SMemInvestigationView::RunQuery()
 	const uint32 RuleNumTimeMarkers = Rule->GetNumTimeMarkers();
 	if (RuleNumTimeMarkers > NumTimeMarkers)
 	{
-		UE_LOG(MemoryProfiler, Error, TEXT("[MemQuery] Only %d time markers available. Current rule (%s) requires %u time markers!"),
+		UE_LOG(LogMemoryProfiler, Error, TEXT("[MemQuery] Only %d time markers available. Current rule (%s) requires %u time markers!"),
 			NumTimeMarkers, *Rule->GetShortName().ToString(), RuleNumTimeMarkers);
 		return FReply::Handled();
 	}
@@ -574,18 +580,19 @@ FReply SMemInvestigationView::RunQuery()
 				Builder.Append(TEXT(", "));
 			}
 			Builder.AppendChar((TCHAR)(TEXT('A') + TimeMarkerIndex));
-			const TSharedRef<Insights::FTimeMarker>& TimeMarker = ProfilerWindow->GetCustomTimeMarker(TimeMarkerIndex);
+			using FTimeMarker = TimingProfiler::FTimeMarker;
+			const TSharedRef<FTimeMarker>& TimeMarker = ProfilerWindow->GetCustomTimeMarker(TimeMarkerIndex);
 			Builder.Appendf(TEXT("=%.9f"), TimeMarker->GetTime());
 		}
 		Builder.Append(TEXT(")"));
 	}
-	UE_LOG(MemoryProfiler, Log, TEXT("[MemQuery] Run Query %s%s..."), *Rule->GetShortName().ToString(), Builder.ToString());
+	UE_LOG(LogMemoryProfiler, Log, TEXT("[MemQuery] Run Query %s%s..."), *Rule->GetShortName().ToString(), Builder.ToString());
 #endif
 
-	TSharedPtr<Insights::SMemAllocTableTreeView> MemAllocTableTreeView = ProfilerWindow->ShowMemAllocTableTreeViewTab();
+	TSharedPtr<SMemAllocTableTreeView> MemAllocTableTreeView = ProfilerWindow->ShowMemAllocTableTreeViewTab();
 	if (MemAllocTableTreeView)
 	{
-		Insights::SMemAllocTableTreeView::FQueryParams QueryParams;
+		SMemAllocTableTreeView::FQueryParams QueryParams;
 		QueryParams.Rule = Rule;
 		QueryParams.TimeMarkers[0] = (RuleNumTimeMarkers > 0) ? ProfilerWindow->GetCustomTimeMarker(0)->GetTime() : 0.0;
 		QueryParams.TimeMarkers[1] = (RuleNumTimeMarkers > 1) ? ProfilerWindow->GetCustomTimeMarker(1)->GetTime() : 0.0;
@@ -609,8 +616,8 @@ FReply SMemInvestigationView::OnTimeMarkerLabelDoubleClicked(const FGeometry& My
 		const uint32 NumTimeMarkers = ProfilerWindow->GetNumCustomTimeMarkers();
 		if (TimeMarkerIndex < NumTimeMarkers)
 		{
-			const TSharedRef<Insights::FTimeMarker>& TimeMarker = ProfilerWindow->GetCustomTimeMarker(TimeMarkerIndex);
-			TSharedPtr<STimingView> TimingView = ProfilerWindow->GetTimingView();
+			const TSharedRef<TimingProfiler::FTimeMarker>& TimeMarker = ProfilerWindow->GetCustomTimeMarker(TimeMarkerIndex);
+			TSharedPtr<TimingProfiler::STimingView> TimingView = ProfilerWindow->GetTimingView();
 			if (TimingView.IsValid())
 			{
 				// Move timer to the center of the timing view.
@@ -624,4 +631,6 @@ FReply SMemInvestigationView::OnTimeMarkerLabelDoubleClicked(const FGeometry& My
 
 ////////////////////////////////////////////////////////////////////////////////////////////////////
 
+} // namespace UE::Insights::MemoryProfiler
+
 #undef LOCTEXT_NAMESPACE
diff --git a/Engine/Source/Developer/TraceInsights/Private/Insights/MemoryProfiler/Widgets/SMemInvestigationView.h b/Engine/Source/Developer/TraceInsights/Private/Insights/MemoryProfiler/Widgets/SMemInvestigationView.h
--- a/Engine/Source/Developer/TraceInsights/Private/Insights/MemoryProfiler/Widgets/SMemInvestigationView.h
+++ b/Engine/Source/Developer/TraceInsights/Private/Insights/MemoryProfiler/Widgets/SMemInvestigationView.h
@@ -3,6 +3,7 @@
 #pragma once
 
 #include "CoreMinimal.h"
+
 #include "SlateFwd.h"
 #include "Widgets/DeclarativeSyntaxSupport.h"
 #include "Widgets/Input/SComboBox.h"
@@ -12,18 +13,18 @@
 ////////////////////////////////////////////////////////////////////////////////////////////////////
 
 class FMenuBuilder;
-class SMemoryProfilerWindow;
 
 namespace TraceServices
 {
 	class IAnalysisSession;
 }
 
-namespace Insights
+namespace UE::Insights::MemoryProfiler
 {
-	class FMemoryRuleSpec;
-	class FQueryTargetWindowSpec;
-}
+
+class FMemoryRuleSpec;
+class FQueryTargetWindowSpec;
+class SMemoryProfilerWindow;
 
 ////////////////////////////////////////////////////////////////////////////////////////////////////
 /**
@@ -58,7 +59,7 @@ public:
 
 	void Reset();
 
-	void QueryTarget_OnSelectionChanged(TSharedPtr<Insights::FQueryTargetWindowSpec> InRule, ESelectInfo::Type SelectInfo);
+	void QueryTarget_OnSelectionChanged(TSharedPtr<FQueryTargetWindowSpec> InRule, ESelectInfo::Type SelectInfo);
 
 private:
 	void UpdateSymbolPathsText() const;
@@ -68,12 +69,12 @@ private:
 	/** Called when the analysis session has changed. */
 	void InsightsManager_OnSessionChanged();
 
-	const TArray<TSharedPtr<Insights::FQueryTargetWindowSpec>>* GetAvailableQueryTargets();
-	TSharedRef<SWidget> QueryTarget_OnGenerateWidget(TSharedPtr<Insights::FQueryTargetWindowSpec> InRule);
+	const TArray<TSharedPtr<FQueryTargetWindowSpec>>* GetAvailableQueryTargets();
+	TSharedRef<SWidget> QueryTarget_OnGenerateWidget(TSharedPtr<FQueryTargetWindowSpec> InRule);
 	FText QueryTarget_GetSelectedText() const;
-	const TArray<TSharedPtr<Insights::FMemoryRuleSpec>>* GetAvailableQueryRules();
-	void QueryRule_OnSelectionChanged(TSharedPtr<Insights::FMemoryRuleSpec> InRule, ESelectInfo::Type SelectInfo);
-	TSharedRef<SWidget> QueryRule_OnGenerateWidget(TSharedPtr<Insights::FMemoryRuleSpec> InRule);
+	const TArray<TSharedPtr<FMemoryRuleSpec>>* GetAvailableQueryRules();
+	void QueryRule_OnSelectionChanged(TSharedPtr<FMemoryRuleSpec> InRule, ESelectInfo::Type SelectInfo);
+	TSharedRef<SWidget> QueryRule_OnGenerateWidget(TSharedPtr<FMemoryRuleSpec> InRule);
 	FText QueryRule_GetSelectedText() const;
 	FText QueryRule_GetTooltipText() const;
 	FReply RunQuery();
@@ -96,14 +97,16 @@ private:
 	/** The analysis session used to populate this widget. */
 	TSharedPtr<const TraceServices::IAnalysisSession> Session;
 
-	TSharedPtr<SComboBox<TSharedPtr<Insights::FMemoryRuleSpec>>> QueryRuleComboBox;
+	TSharedPtr<SComboBox<TSharedPtr<FMemoryRuleSpec>>> QueryRuleComboBox;
 
 	bool bIncludeHeapAllocs;
 	bool bIncludeSwapAllocs;
 
-	TSharedPtr<SComboBox<TSharedPtr<Insights::FQueryTargetWindowSpec>>> QueryTargetComboBox;
+	TSharedPtr<SComboBox<TSharedPtr<FQueryTargetWindowSpec>>> QueryTargetComboBox;
 	
 	TSharedPtr<STextBlock> SymbolPathsTextBlock;
 };
 
 ////////////////////////////////////////////////////////////////////////////////////////////////////
+
+} // namespace UE::Insights::MemoryProfiler
diff --git a/Engine/Source/Developer/TraceInsights/Private/Insights/MemoryProfiler/Widgets/SMemTagTreeView.cpp b/Engine/Source/Developer/TraceInsights/Private/Insights/MemoryProfiler/Widgets/SMemTagTreeView.cpp
--- a/Engine/Source/Developer/TraceInsights/Private/Insights/MemoryProfiler/Widgets/SMemTagTreeView.cpp
+++ b/Engine/Source/Developer/TraceInsights/Private/Insights/MemoryProfiler/Widgets/SMemTagTreeView.cpp
@@ -16,9 +16,13 @@
 #include "Widgets/SToolTip.h"
 #include "Widgets/Views/STableViewBase.h"
 
-// Insights
-#include "Insights/Common/Stopwatch.h"
-#include "Insights/Common/TimeUtils.h"
+// TraceInsightsCore
+#include "InsightsCore/Common/Stopwatch.h"
+#include "InsightsCore/Common/TimeUtils.h"
+#include "InsightsCore/Table/ViewModels/Table.h"
+#include "InsightsCore/Table/ViewModels/TableColumn.h"
+
+// TraceInsights
 #include "Insights/InsightsStyle.h"
 #include "Insights/MemoryProfiler/MemoryProfilerManager.h"
 #include "Insights/MemoryProfiler/ViewModels/MemoryTag.h"
@@ -28,16 +32,17 @@
 #include "Insights/MemoryProfiler/Widgets/SMemTagTreeViewTooltip.h"
 #include "Insights/MemoryProfiler/Widgets/SMemTagTreeViewTableRow.h"
 #include "Insights/MemoryProfiler/Widgets/SMemoryProfilerWindow.h"
-#include "Insights/Table/ViewModels/Table.h"
-#include "Insights/Table/ViewModels/TableColumn.h"
 
-#define LOCTEXT_NAMESPACE "SMemTagTreeView"
+#define LOCTEXT_NAMESPACE "UE::Insights::MemoryProfiler::SMemTagTreeView"
+
+namespace UE::Insights::MemoryProfiler
+{
 
 ////////////////////////////////////////////////////////////////////////////////////////////////////
 
 SMemTagTreeView::SMemTagTreeView()
 	: ProfilerWindowWeakPtr()
-	, Table(MakeShared<Insights::FTable>())
+	, Table(MakeShared<FTable>())
 	, bExpansionSaved(false)
 	, bFilterOutZeroCountMemTags(false)
 	, TrackersFilter(uint64(-1))
@@ -192,10 +197,10 @@ void SMemTagTreeView::CreateTrackersMenuSection(FMenuBuilder& MenuBuilder)
 	if (ProfilerWindow.IsValid())
 	{
 		FMemorySharedState& SharedState = ProfilerWindow->GetSharedState();
-		const TArray<TSharedPtr<Insights::FMemoryTracker>>& Trackers = SharedState.GetTrackers();
-		for (const TSharedPtr<Insights::FMemoryTracker>& Tracker : Trackers)
+		const TArray<TSharedPtr<FMemoryTracker>>& Trackers = SharedState.GetTrackers();
+		for (const TSharedPtr<FMemoryTracker>& Tracker : Trackers)
 		{
-			const Insights::FMemoryTrackerId TrackerId = Tracker->GetId();
+			const FMemoryTrackerId TrackerId = Tracker->GetId();
 			MenuBuilder.AddMenuEntry(
 				FText::FromString(Tracker->GetName()),
 				TAttribute<FText>(),
@@ -213,17 +218,17 @@ void SMemTagTreeView::CreateTrackersMenuSection(FMenuBuilder& MenuBuilder)
 
 ////////////////////////////////////////////////////////////////////////////////////////////////////
 
-void SMemTagTreeView::ToggleTracker(Insights::FMemoryTrackerId InTrackerId)
+void SMemTagTreeView::ToggleTracker(FMemoryTrackerId InTrackerId)
 {
-	TrackersFilter ^= Insights::FMemoryTracker::AsFlag(InTrackerId);
+	TrackersFilter ^= FMemoryTracker::AsFlag(InTrackerId);
 	ApplyFiltering();
 }
 
 ////////////////////////////////////////////////////////////////////////////////////////////////////
 
-bool SMemTagTreeView::IsTrackerChecked(Insights::FMemoryTrackerId InTrackerId) const
+bool SMemTagTreeView::IsTrackerChecked(FMemoryTrackerId InTrackerId) const
 {
-	return (TrackersFilter & Insights::FMemoryTracker::AsFlag(InTrackerId)) != 0;
+	return (TrackersFilter & FMemoryTracker::AsFlag(InTrackerId)) != 0;
 }
 
 ////////////////////////////////////////////////////////////////////////////////////////////////////
@@ -682,9 +687,9 @@ void SMemTagTreeView::TreeView_BuildSortByMenu(FMenuBuilder& MenuBuilder)
 {
 	MenuBuilder.BeginSection("SortColumn", LOCTEXT("ContextMenu_Section_SortColumn", "Sort Column"));
 
-	for (const TSharedRef<Insights::FTableColumn>& ColumnRef : Table->GetColumns())
+	for (const TSharedRef<FTableColumn>& ColumnRef : Table->GetColumns())
 	{
-		const Insights::FTableColumn& Column = *ColumnRef;
+		const FTableColumn& Column = *ColumnRef;
 
 		if (Column.IsVisible() && Column.CanBeSorted())
 		{
@@ -751,9 +756,9 @@ void SMemTagTreeView::TreeView_BuildViewColumnMenu(FMenuBuilder& MenuBuilder)
 {
 	MenuBuilder.BeginSection("Columns", LOCTEXT("ContextMenu_Section_Columns", "Columns"));
 
-	for (const TSharedRef<Insights::FTableColumn>& ColumnRef : Table->GetColumns())
+	for (const TSharedRef<FTableColumn>& ColumnRef : Table->GetColumns())
 	{
-		const Insights::FTableColumn& Column = *ColumnRef;
+		const FTableColumn& Column = *ColumnRef;
 
 		FUIAction Action_ToggleColumn
 		(
@@ -780,12 +785,12 @@ void SMemTagTreeView::TreeView_BuildViewColumnMenu(FMenuBuilder& MenuBuilder)
 void SMemTagTreeView::InitializeAndShowHeaderColumns()
 {
 	// Create columns.
-	TArray<TSharedRef<Insights::FTableColumn>> Columns;
+	TArray<TSharedRef<FTableColumn>> Columns;
 	FMemTagTreeViewColumnFactory::CreateMemTagTreeViewColumns(Columns);
 	Table->SetColumns(Columns);
 
 	// Show columns.
-	for (const TSharedRef<Insights::FTableColumn>& ColumnRef : Table->GetColumns())
+	for (const TSharedRef<FTableColumn>& ColumnRef : Table->GetColumns())
 	{
 		if (ColumnRef->ShouldBeVisible())
 		{
@@ -798,13 +803,13 @@ void SMemTagTreeView::InitializeAndShowHeaderColumns()
 
 FText SMemTagTreeView::GetColumnHeaderText(const FName ColumnId) const
 {
-	const Insights::FTableColumn& Column = *Table->FindColumnChecked(ColumnId);
+	const FTableColumn& Column = *Table->FindColumnChecked(ColumnId);
 	return Column.GetShortName();
 }
 
 ////////////////////////////////////////////////////////////////////////////////////////////////////
 
-TSharedRef<SWidget> SMemTagTreeView::TreeViewHeaderRow_GenerateColumnMenu(const Insights::FTableColumn& Column)
+TSharedRef<SWidget> SMemTagTreeView::TreeViewHeaderRow_GenerateColumnMenu(const FTableColumn& Column)
 {
 	const bool bShouldCloseWindowAfterMenuSelection = true;
 	FMenuBuilder MenuBuilder(bShouldCloseWindowAfterMenuSelection, NULL);
@@ -975,7 +980,7 @@ void SMemTagTreeView::UpdateTree()
 	const double TotalTime = Stopwatch.GetAccumulatedTime();
 	if (TotalTime > 0.1)
 	{
-		UE_LOG(MemoryProfiler, Log, TEXT("[LLM Tags] Tree view updated in %.3fs (%d counters) --> G:%.3fs + S:%.3fs + F:%.3fs"),
+		UE_LOG(LogMemoryProfiler, Log, TEXT("[LLM Tags] Tree view updated in %.3fs (%d counters) --> G:%.3fs + S:%.3fs + F:%.3fs"),
 			TotalTime, MemTagNodes.Num(), Time1, Time2 - Time1, TotalTime - Time2);
 	}
 }
@@ -996,12 +1001,12 @@ void SMemTagTreeView::ApplyFiltering()
 		const bool bIsGroupVisible = Filters->PassesAllFilters(GroupPtr);
 		int32 NumVisibleChildren = 0;
 
-		const TArray<Insights::FBaseTreeNodePtr>& GroupChildren = GroupPtr->GetChildren();
-		for (const Insights::FBaseTreeNodePtr& Child : GroupChildren)
+		const TArray<FBaseTreeNodePtr>& GroupChildren = GroupPtr->GetChildren();
+		for (const FBaseTreeNodePtr& Child : GroupChildren)
 		{
-			const FMemTagNodePtr& NodePtr = StaticCastSharedPtr<FMemTagNode, Insights::FBaseTreeNode>(Child);
+			const FMemTagNodePtr& NodePtr = StaticCastSharedPtr<FMemTagNode, FBaseTreeNode>(Child);
 			const bool bIsChildVisible = (!bFilterOutZeroCountMemTags || NodePtr->GetAggregatedStats().InstanceCount > 0)
-									  && (!Insights::FMemoryTracker::IsValidTrackerId(NodePtr->GetMemTrackerId()) || ((Insights::FMemoryTracker::AsFlag(NodePtr->GetMemTrackerId()) & TrackersFilter) != 0))
+									  && (!FMemoryTracker::IsValidTrackerId(NodePtr->GetMemTrackerId()) || ((FMemoryTracker::AsFlag(NodePtr->GetMemTrackerId()) & TrackersFilter) != 0))
 									  && Filters->PassesAllFilters(NodePtr);
 			if (bIsChildVisible)
 			{
@@ -1110,11 +1115,11 @@ void SMemTagTreeView::TreeView_OnSelectionChanged(FMemTagNodePtr SelectedItem, E
 
 void SMemTagTreeView::TreeView_OnGetChildren(FMemTagNodePtr InParent, TArray<FMemTagNodePtr>& OutChildren)
 {
-	const TArray<Insights::FBaseTreeNodePtr>& Children = InParent->GetFilteredChildren();
+	const TArray<FBaseTreeNodePtr>& Children = InParent->GetFilteredChildren();
 	OutChildren.Reset(Children.Num());
-	for (const Insights::FBaseTreeNodePtr& Child : Children)
+	for (const FBaseTreeNodePtr& Child : Children)
 	{
-		OutChildren.Add(StaticCastSharedPtr<FMemTagNode, Insights::FBaseTreeNode>(Child));
+		OutChildren.Add(StaticCastSharedPtr<FMemTagNode, FBaseTreeNode>(Child));
 	}
 }
 
@@ -1133,8 +1138,8 @@ void SMemTagTreeView::TreeView_OnMouseButtonDoubleClick(FMemTagNodePtr MemTagNod
 		if (ProfilerWindow.IsValid())
 		{
 			FMemorySharedState& SharedState = ProfilerWindow->GetSharedState();
-			const Insights::FMemoryTrackerId MemTrackerId = MemTagNodePtr->GetMemTrackerId();
-			const Insights::FMemoryTagId MemTagId = MemTagNodePtr->GetMemTagId();
+			const FMemoryTrackerId MemTrackerId = MemTagNodePtr->GetMemTrackerId();
+			const FMemoryTagId MemTagId = MemTagNodePtr->GetMemTagId();
 
 			TSharedPtr<FMemoryGraphTrack> GraphTrack = SharedState.GetMemTagGraphTrack(MemTrackerId, MemTagId);
 			if (!GraphTrack.IsValid())
@@ -1178,7 +1183,7 @@ bool SMemTagTreeView::TableRow_ShouldBeEnabled(FMemTagNodePtr NodePtr) const
 
 ////////////////////////////////////////////////////////////////////////////////////////////////////
 
-void SMemTagTreeView::TableRow_SetHoveredCell(TSharedPtr<Insights::FTable> InTablePtr, TSharedPtr<Insights::FTableColumn> InColumnPtr, FMemTagNodePtr InNodePtr)
+void SMemTagTreeView::TableRow_SetHoveredCell(TSharedPtr<FTable> InTablePtr, TSharedPtr<FTableColumn> InColumnPtr, FMemTagNodePtr InNodePtr)
 {
 	HoveredColumnId = InColumnPtr ? InColumnPtr->GetId() : FName();
 
@@ -1305,10 +1310,10 @@ void SMemTagTreeView::CreateGroups()
 	// Groups LLM tags by tracker.
 	else if (GroupingMode == EMemTagNodeGroupingMode::ByTracker)
 	{
-		TMap<Insights::FMemoryTrackerId, FMemTagNodePtr> GroupNodeSet;
+		TMap<FMemoryTrackerId, FMemTagNodePtr> GroupNodeSet;
 		for (const FMemTagNodePtr& NodePtr : MemTagNodes)
 		{
-			Insights::FMemoryTrackerId TrackerId = NodePtr->GetMemTrackerId();
+			FMemoryTrackerId TrackerId = NodePtr->GetMemTrackerId();
 			FMemTagNodePtr GroupPtr = GroupNodeSet.FindRef(TrackerId);
 			if (!GroupPtr)
 			{
@@ -1421,11 +1426,11 @@ void SMemTagTreeView::CreateSortings()
 	AvailableSorters.Reset();
 	CurrentSorter = nullptr;
 
-	for (const TSharedRef<Insights::FTableColumn>& ColumnRef : Table->GetColumns())
+	for (const TSharedRef<FTableColumn>& ColumnRef : Table->GetColumns())
 	{
 		if (ColumnRef->CanBeSorted())
 		{
-			TSharedPtr<Insights::ITableCellValueSorter> SorterPtr = ColumnRef->GetValueSorter();
+			TSharedPtr<ITableCellValueSorter> SorterPtr = ColumnRef->GetValueSorter();
 			if (ensure(SorterPtr.IsValid()))
 			{
 				AvailableSorters.Add(SorterPtr);
@@ -1440,7 +1445,7 @@ void SMemTagTreeView::CreateSortings()
 
 void SMemTagTreeView::UpdateCurrentSortingByColumn()
 {
-	TSharedPtr<Insights::FTableColumn> ColumnPtr = Table->FindColumn(ColumnBeingSorted);
+	TSharedPtr<FTableColumn> ColumnPtr = Table->FindColumn(ColumnBeingSorted);
 	CurrentSorter = ColumnPtr.IsValid() ? ColumnPtr->GetValueSorter() : nullptr;
 }
 
@@ -1451,21 +1456,21 @@ void SMemTagTreeView::SortTreeNodes()
 	if (CurrentSorter.IsValid())
 	{
 		// Sort groups (always by name).
-		TArray<Insights::FBaseTreeNodePtr> SortedGroupNodes;
+		TArray<FBaseTreeNodePtr> SortedGroupNodes;
 		for (const FMemTagNodePtr& NodePtr : GroupNodes)
 		{
 			SortedGroupNodes.Add(NodePtr);
 		}
-		TSharedPtr<Insights::ITableCellValueSorter> Sorter = CurrentSorter;
-		Insights::ESortMode SortMode = (ColumnSortMode == EColumnSortMode::Type::Descending) ? Insights::ESortMode::Descending : Insights::ESortMode::Ascending;
+		TSharedPtr<ITableCellValueSorter> Sorter = CurrentSorter;
+		ESortMode SortMode = (ColumnSortMode == EColumnSortMode::Type::Descending) ? ESortMode::Descending : ESortMode::Ascending;
 		if (CurrentSorter->GetName() != FName(TEXT("ByName")))
 		{
-			Sorter = MakeShared<Insights::FSorterByName>(Table->GetColumns()[0]);
-			SortMode = Insights::ESortMode::Ascending;
+			Sorter = MakeShared<FSorterByName>(Table->GetColumns()[0]);
+			SortMode = ESortMode::Ascending;
 		}
 		Sorter->Sort(SortedGroupNodes, SortMode);
 		GroupNodes.Reset();
-		for (const Insights::FBaseTreeNodePtr& NodePtr : SortedGroupNodes)
+		for (const FBaseTreeNodePtr& NodePtr : SortedGroupNodes)
 		{
 			GroupNodes.Add(StaticCastSharedPtr<FMemTagNode>(NodePtr));
 		}
@@ -1480,12 +1485,12 @@ void SMemTagTreeView::SortTreeNodes()
 
 ////////////////////////////////////////////////////////////////////////////////////////////////////
 
-void SMemTagTreeView::SortTreeNodesRec(FMemTagNode& Node, const Insights::ITableCellValueSorter& Sorter)
+void SMemTagTreeView::SortTreeNodesRec(FMemTagNode& Node, const ITableCellValueSorter& Sorter)
 {
-	Insights::ESortMode SortMode = (ColumnSortMode == EColumnSortMode::Type::Descending) ? Insights::ESortMode::Descending : Insights::ESortMode::Ascending;
+	ESortMode SortMode = (ColumnSortMode == EColumnSortMode::Type::Descending) ? ESortMode::Descending : ESortMode::Ascending;
 	Node.SortChildren(Sorter, SortMode);
 
-	for (Insights::FBaseTreeNodePtr ChildPtr : Node.GetChildren())
+	for (FBaseTreeNodePtr ChildPtr : Node.GetChildren())
 	{
 		//if (ChildPtr->IsGroup())
 		if (ChildPtr->GetChildrenCount() > 0)
@@ -1540,7 +1545,7 @@ bool SMemTagTreeView::HeaderMenu_SortMode_IsChecked(const FName ColumnId, const
 
 bool SMemTagTreeView::HeaderMenu_SortMode_CanExecute(const FName ColumnId, const EColumnSortMode::Type InSortMode) const
 {
-	const Insights::FTableColumn& Column = *Table->FindColumnChecked(ColumnId);
+	const FTableColumn& Column = *Table->FindColumnChecked(ColumnId);
 	return Column.CanBeSorted();
 }
 
@@ -1613,7 +1618,7 @@ bool SMemTagTreeView::CanShowColumn(const FName ColumnId) const
 
 void SMemTagTreeView::ShowColumn(const FName ColumnId)
 {
-	Insights::FTableColumn& Column = *Table->FindColumnChecked(ColumnId);
+	FTableColumn& Column = *Table->FindColumnChecked(ColumnId);
 	Column.Show();
 
 	SHeaderRow::FColumn::FArguments ColumnArgs;
@@ -1668,7 +1673,7 @@ void SMemTagTreeView::ShowColumn(const FName ColumnId)
 
 bool SMemTagTreeView::CanHideColumn(const FName ColumnId) const
 {
-	const Insights::FTableColumn& Column = *Table->FindColumnChecked(ColumnId);
+	const FTableColumn& Column = *Table->FindColumnChecked(ColumnId);
 	return Column.CanBeHidden();
 }
 
@@ -1676,7 +1681,7 @@ bool SMemTagTreeView::CanHideColumn(const FName ColumnId) const
 
 void SMemTagTreeView::HideColumn(const FName ColumnId)
 {
-	Insights::FTableColumn& Column = *Table->FindColumnChecked(ColumnId);
+	FTableColumn& Column = *Table->FindColumnChecked(ColumnId);
 	Column.Hide();
 
 	TreeViewHeaderRow->RemoveColumn(ColumnId);
@@ -1688,7 +1693,7 @@ void SMemTagTreeView::HideColumn(const FName ColumnId)
 
 bool SMemTagTreeView::IsColumnVisible(const FName ColumnId)
 {
-	const Insights::FTableColumn& Column = *Table->FindColumnChecked(ColumnId);
+	const FTableColumn& Column = *Table->FindColumnChecked(ColumnId);
 	return Column.IsVisible();
 }
 
@@ -1696,7 +1701,7 @@ bool SMemTagTreeView::IsColumnVisible(const FName ColumnId)
 
 bool SMemTagTreeView::CanToggleColumnVisibility(const FName ColumnId) const
 {
-	const Insights::FTableColumn& Column = *Table->FindColumnChecked(ColumnId);
+	const FTableColumn& Column = *Table->FindColumnChecked(ColumnId);
 	return !Column.IsVisible() || Column.CanBeHidden();
 }
 
@@ -1704,7 +1709,7 @@ bool SMemTagTreeView::CanToggleColumnVisibility(const FName ColumnId) const
 
 void SMemTagTreeView::ToggleColumnVisibility(const FName ColumnId)
 {
-	const Insights::FTableColumn& Column = *Table->FindColumnChecked(ColumnId);
+	const FTableColumn& Column = *Table->FindColumnChecked(ColumnId);
 	if (Column.IsVisible())
 	{
 		HideColumn(ColumnId);
@@ -1732,9 +1737,9 @@ void SMemTagTreeView::ContextMenu_ShowAllColumns_Execute()
 	ColumnSortMode = GetDefaultColumnSortMode();
 	UpdateCurrentSortingByColumn();
 
-	for (const TSharedRef<Insights::FTableColumn>& ColumnRef : Table->GetColumns())
+	for (const TSharedRef<FTableColumn>& ColumnRef : Table->GetColumns())
 	{
-		const Insights::FTableColumn& Column = *ColumnRef;
+		const FTableColumn& Column = *ColumnRef;
 
 		if (!Column.IsVisible())
 		{
@@ -1774,9 +1779,9 @@ void SMemTagTreeView::ContextMenu_ShowMinMaxMedColumns_Execute()
 	ColumnSortMode = EColumnSortMode::Descending;
 	UpdateCurrentSortingByColumn();
 
-	for (const TSharedRef<Insights::FTableColumn>& ColumnRef : Table->GetColumns())
+	for (const TSharedRef<FTableColumn>& ColumnRef : Table->GetColumns())
 	{
-		const Insights::FTableColumn& Column = *ColumnRef;
+		const FTableColumn& Column = *ColumnRef;
 
 		const bool bShouldBeVisible = Preset.Contains(Column.GetId());
 
@@ -1808,9 +1813,9 @@ void SMemTagTreeView::ContextMenu_ResetColumns_Execute()
 	ColumnSortMode = GetDefaultColumnSortMode();
 	UpdateCurrentSortingByColumn();
 
-	for (const TSharedRef<Insights::FTableColumn>& ColumnRef : Table->GetColumns())
+	for (const TSharedRef<FTableColumn>& ColumnRef : Table->GetColumns())
 	{
-		const Insights::FTableColumn& Column = *ColumnRef;
+		const FTableColumn& Column = *ColumnRef;
 
 		if (Column.ShouldBeVisible() && !Column.IsVisible())
 		{
@@ -1880,20 +1885,20 @@ void SMemTagTreeView::RebuildTree(bool bResync)
 	if (ProfilerWindow.IsValid())
 	{
 		FMemorySharedState& SharedState = ProfilerWindow->GetSharedState();
-		const Insights::FMemoryTagList& TagList = SharedState.GetTagList();
+		const FMemoryTagList& TagList = SharedState.GetTagList();
 
 		if (LastMemoryTagListSerialNumber != TagList.GetSerialNumber())
 		{
 			LastMemoryTagListSerialNumber = TagList.GetSerialNumber();
 
-			const TArray<Insights::FMemoryTag*>& MemTags = TagList.GetTags();
+			const TArray<FMemoryTag*>& MemTags = TagList.GetTags();
 			const int32 MemTagCount = MemTags.Num();
 
 			MemTagNodes.Empty(MemTagCount);
 			MemTagNodesIdMap.Empty(MemTagCount);
 			bListHasChanged = true;
 
-			for (Insights::FMemoryTag* MemTagPtr : MemTags)
+			for (FMemoryTag* MemTagPtr : MemTags)
 			{
 				FMemTagNodePtr MemTagNodePtr = MakeShared<FMemTagNode>(MemTagPtr);
 				MemTagNodes.Add(MemTagNodePtr);
@@ -1904,7 +1909,7 @@ void SMemTagTreeView::RebuildTree(bool bResync)
 			for (FMemTagNodePtr& NodePtr : MemTagNodes)
 			{
 				check(NodePtr->GetMemTag() != nullptr);
-				Insights::FMemoryTag& MemTag = *NodePtr->GetMemTag();
+				FMemoryTag& MemTag = *NodePtr->GetMemTag();
 
 				FMemTagNodePtr ParentNodePtr = MemTagNodesIdMap.FindRef(MemTag.GetParentId());
 				if (ParentNodePtr)
@@ -1951,7 +1956,7 @@ void SMemTagTreeView::RebuildTree(bool bResync)
 	if (TotalTime > 0.01)
 	{
 		const double SyncTime = SyncStopwatch.GetAccumulatedTime();
-		UE_LOG(MemoryProfiler, Log, TEXT("[LLM Tags] Tree view rebuilt in %.4fs (sync: %.4fs + update: %.4fs) --> %d LLM tags (%d added)"),
+		UE_LOG(LogMemoryProfiler, Log, TEXT("[LLM Tags] Tree view rebuilt in %.4fs (sync: %.4fs + update: %.4fs) --> %d LLM tags (%d added)"),
 			TotalTime, SyncTime, TotalTime - SyncTime, MemTagNodes.Num(), MemTagNodes.Num() - PreviousNodeCount);
 	}
 }
@@ -2044,13 +2049,13 @@ void SMemTagTreeView::UpdateStatsInternal()
 	Stopwatch.Stop();
 	const double TotalTime = Stopwatch.GetAccumulatedTime();
 	const double AggregationTime = AggregationStopwatch.GetAccumulatedTime();
-	UE_LOG(MemoryProfiler, Log, TEXT("[LLM Tags] Aggregated stats updated in %.4fs (%.4fs + %.4fs)"),
+	UE_LOG(LogMemoryProfiler, Log, TEXT("[LLM Tags] Aggregated stats updated in %.4fs (%.4fs + %.4fs)"),
 		TotalTime, AggregationTime, TotalTime - AggregationTime);
 }
 
 ////////////////////////////////////////////////////////////////////////////////////////////////////
 
-void SMemTagTreeView::SelectMemTagNode(Insights::FMemoryTagId MemTagId)
+void SMemTagTreeView::SelectMemTagNode(FMemoryTagId MemTagId)
 {
 	FMemTagNodePtr NodePtr = GetMemTagNode(MemTagId);
 	if (NodePtr)
@@ -2144,19 +2149,19 @@ void SMemTagTreeView::CreateGraphTracksForSelectedMemTags()
 		{
 			if (SelectedMemTagNode->IsGroup())
 			{
-				const TArray<Insights::FBaseTreeNodePtr>& Children = SelectedMemTagNode->GetFilteredChildren();
-				for (const Insights::FBaseTreeNodePtr& Child : Children)
+				const TArray<FBaseTreeNodePtr>& Children = SelectedMemTagNode->GetFilteredChildren();
+				for (const FBaseTreeNodePtr& Child : Children)
 				{
-					FMemTagNodePtr MemTagNode = StaticCastSharedPtr<FMemTagNode, Insights::FBaseTreeNode>(Child);
-					const Insights::FMemoryTrackerId MemTrackerId = MemTagNode->GetMemTrackerId();
-					const Insights::FMemoryTagId MemTagId = MemTagNode->GetMemTagId();
+					FMemTagNodePtr MemTagNode = StaticCastSharedPtr<FMemTagNode, FBaseTreeNode>(Child);
+					const FMemoryTrackerId MemTrackerId = MemTagNode->GetMemTrackerId();
+					const FMemoryTagId MemTagId = MemTagNode->GetMemTagId();
 					SharedState.CreateMemTagGraphTrack(MemTrackerId, MemTagId);
 				}
 			}
 			else
 			{
-				const Insights::FMemoryTrackerId MemTrackerId = SelectedMemTagNode->GetMemTrackerId();
-				const Insights::FMemoryTagId MemTagId = SelectedMemTagNode->GetMemTagId();
+				const FMemoryTrackerId MemTrackerId = SelectedMemTagNode->GetMemTrackerId();
+				const FMemoryTagId MemTagId = SelectedMemTagNode->GetMemTagId();
 				SharedState.CreateMemTagGraphTrack(MemTrackerId, MemTagId);
 			}
 		}
@@ -2192,12 +2197,12 @@ void SMemTagTreeView::CreateGraphTracksForFilteredMemTags()
 		FMemorySharedState& SharedState = ProfilerWindow->GetSharedState();
 		for (const FMemTagNodePtr& GroupNode : FilteredGroupNodes)
 		{
-			const TArray<Insights::FBaseTreeNodePtr>& Children = GroupNode->GetFilteredChildren();
-			for (const Insights::FBaseTreeNodePtr& Child : Children)
+			const TArray<FBaseTreeNodePtr>& Children = GroupNode->GetFilteredChildren();
+			for (const FBaseTreeNodePtr& Child : Children)
 			{
-				FMemTagNodePtr MemTagNode = StaticCastSharedPtr<FMemTagNode, Insights::FBaseTreeNode>(Child);
-				const Insights::FMemoryTrackerId MemTrackerId = MemTagNode->GetMemTrackerId();
-				const Insights::FMemoryTagId MemTagId = MemTagNode->GetMemTagId();
+				FMemTagNodePtr MemTagNode = StaticCastSharedPtr<FMemTagNode, FBaseTreeNode>(Child);
+				const FMemoryTrackerId MemTrackerId = MemTagNode->GetMemTrackerId();
+				const FMemoryTagId MemTagId = MemTagNode->GetMemTagId();
 				SharedState.CreateMemTagGraphTrack(MemTrackerId, MemTagId);
 			}
 		}
@@ -2205,7 +2210,7 @@ void SMemTagTreeView::CreateGraphTracksForFilteredMemTags()
 }
 
 ////////////////////////////////////////////////////////////////////////////////////////////////////
-// Create all mem memory graph tracks for selected LLM tag(s)
+// Create all memory graph tracks for selected LLM tag(s)
 ////////////////////////////////////////////////////////////////////////////////////////////////////
 
 bool SMemTagTreeView::CanCreateAllGraphTracks() const
@@ -2224,8 +2229,8 @@ void SMemTagTreeView::CreateAllGraphTracks()
 		FMemorySharedState& SharedState = ProfilerWindow->GetSharedState();
 		for (const FMemTagNodePtr& MemTagNode : MemTagNodes)
 		{
-			const Insights::FMemoryTrackerId MemTrackerId = MemTagNode->GetMemTrackerId();
-			const Insights::FMemoryTagId MemTagId = MemTagNode->GetMemTagId();
+			const FMemoryTrackerId MemTrackerId = MemTagNode->GetMemTrackerId();
+			const FMemoryTagId MemTagId = MemTagNode->GetMemTagId();
 			SharedState.CreateMemTagGraphTrack(MemTrackerId, MemTagId);
 		}
 	}
@@ -2254,19 +2259,19 @@ void SMemTagTreeView::RemoveGraphTracksForSelectedMemTags()
 		{
 			if (SelectedMemTagNode->IsGroup())
 			{
-				const TArray<Insights::FBaseTreeNodePtr>& Children = SelectedMemTagNode->GetFilteredChildren();
-				for (const Insights::FBaseTreeNodePtr& Child : Children)
+				const TArray<FBaseTreeNodePtr>& Children = SelectedMemTagNode->GetFilteredChildren();
+				for (const FBaseTreeNodePtr& Child : Children)
 				{
-					FMemTagNodePtr MemTagNode = StaticCastSharedPtr<FMemTagNode, Insights::FBaseTreeNode>(Child);
-					const Insights::FMemoryTrackerId MemTrackerId = MemTagNode->GetMemTrackerId();
-					const Insights::FMemoryTagId MemTagId = MemTagNode->GetMemTagId();
+					FMemTagNodePtr MemTagNode = StaticCastSharedPtr<FMemTagNode, FBaseTreeNode>(Child);
+					const FMemoryTrackerId MemTrackerId = MemTagNode->GetMemTrackerId();
+					const FMemoryTagId MemTagId = MemTagNode->GetMemTagId();
 					SharedState.RemoveMemTagGraphTrack(MemTrackerId, MemTagId);
 				}
 			}
 			else
 			{
-				const Insights::FMemoryTrackerId MemTrackerId = SelectedMemTagNode->GetMemTrackerId();
-				const Insights::FMemoryTagId MemTagId = SelectedMemTagNode->GetMemTagId();
+				const FMemoryTrackerId MemTrackerId = SelectedMemTagNode->GetMemTrackerId();
+				const FMemoryTagId MemTagId = SelectedMemTagNode->GetMemTagId();
 				SharedState.RemoveMemTagGraphTrack(MemTrackerId, MemTagId);
 			}
 		}
@@ -2327,16 +2332,16 @@ void SMemTagTreeView::SetColorToNode(const FMemTagNodePtr& MemTagNode, FLinearCo
 {
 	if (MemTagNode->IsGroup())
 	{
-		const TArray<Insights::FBaseTreeNodePtr>& Children = MemTagNode->GetFilteredChildren();
-		for (const Insights::FBaseTreeNodePtr& Child : Children)
+		const TArray<FBaseTreeNodePtr>& Children = MemTagNode->GetFilteredChildren();
+		for (const FBaseTreeNodePtr& Child : Children)
 		{
-			const FMemTagNodePtr ChildMemTagNode = StaticCastSharedPtr<FMemTagNode, Insights::FBaseTreeNode>(Child);
+			const FMemTagNodePtr ChildMemTagNode = StaticCastSharedPtr<FMemTagNode, FBaseTreeNode>(Child);
 			SetColorToNode(ChildMemTagNode, Color, bSetRandomColor);
 		}
 		return;
 	}
 
-	Insights::FMemoryTag* MemTag = MemTagNode->GetMemTag();
+	FMemoryTag* MemTag = MemTagNode->GetMemTag();
 	if (!MemTag)
 	{
 		return;
@@ -2354,7 +2359,7 @@ void SMemTagTreeView::SetColorToNode(const FMemTagNodePtr& MemTagNode, FLinearCo
 
 	const FLinearColor BorderColor(FMath::Min(Color.R + 0.4f, 1.0f), FMath::Min(Color.G + 0.4f, 1.0f), FMath::Min(Color.B + 0.4f, 1.0f), 1.0f);
 
-	const Insights::FMemoryTagId MemTagId = MemTagNode->GetMemTagId();
+	const FMemoryTagId MemTagId = MemTagNode->GetMemTagId();
 
 	TSharedPtr<FMemoryGraphTrack> MainGraphTrack;
 	TSharedPtr<SMemoryProfilerWindow> ProfilerWindow = GetProfilerWindow();
@@ -2459,4 +2464,6 @@ void SMemTagTreeView::ColorPickerCancelled(FLinearColor OriginalColor)
 
 ////////////////////////////////////////////////////////////////////////////////////////////////////
 
+} // namespace UE::Insights::MemoryProfiler
+
 #undef LOCTEXT_NAMESPACE
diff --git a/Engine/Source/Developer/TraceInsights/Private/Insights/MemoryProfiler/Widgets/SMemTagTreeView.h b/Engine/Source/Developer/TraceInsights/Private/Insights/MemoryProfiler/Widgets/SMemTagTreeView.h
--- a/Engine/Source/Developer/TraceInsights/Private/Insights/MemoryProfiler/Widgets/SMemTagTreeView.h
+++ b/Engine/Source/Developer/TraceInsights/Private/Insights/MemoryProfiler/Widgets/SMemTagTreeView.h
@@ -3,6 +3,7 @@
 #pragma once
 
 #include "CoreMinimal.h"
+
 #include "Misc/FilterCollection.h"
 #include "Misc/TextFilter.h"
 #include "SlateFwd.h"
@@ -15,28 +16,30 @@
 #include "Widgets/Views/STableViewBase.h"
 #include "Widgets/Views/STreeView.h"
 
-// Insights
+// TraceInsights
 #include "Insights/MemoryProfiler/ViewModels/MemTagNodeGroupingAndSorting.h"
 #include "Insights/MemoryProfiler/ViewModels/MemTagNode.h"
 
-////////////////////////////////////////////////////////////////////////////////////////////////////
-
 class FMenuBuilder;
-class FMemoryGraphTrack;
-class SMemoryProfilerWindow;
 
 namespace TraceServices
 {
 	class IAnalysisSession;
 }
 
-namespace Insights
+namespace UE::Insights
 {
 	class FTable;
 	class FTableColumn;
 	class ITableCellValueSorter;
 }
 
+namespace UE::Insights::MemoryProfiler
+{
+
+class FMemoryGraphTrack;
+class SMemoryProfilerWindow;
+
 ////////////////////////////////////////////////////////////////////////////////////////////////////
 
 /** The filter collection - used for updating the list of tree nodes. */
@@ -47,7 +50,7 @@ typedef TTextFilter<const FMemTagNodePtr&> FMemTagNodeTextFilter;
 
 ////////////////////////////////////////////////////////////////////////////////////////////////////
 /**
- * A custom widget used to display the list of LLM tags and thier aggregated stats.
+ * A custom widget used to display the list of LLM tags and their aggregated stats.
  */
 class SMemTagTreeView : public SCompoundWidget
 {
@@ -76,7 +79,7 @@ public:
 	 */
 	void Construct(const FArguments& InArgs, TSharedPtr<SMemoryProfilerWindow> InProfilerWindow);
 
-	TSharedPtr<Insights::FTable> GetTable() const { return Table; }
+	TSharedPtr<FTable> GetTable() const { return Table; }
 
 	void Reset();
 
@@ -89,8 +92,8 @@ public:
 	void ResetStats();
 	void UpdateStats(double StartTime, double EndTime);
 
-	FMemTagNodePtr GetMemTagNode(Insights::FMemoryTagId MemTagId) const { return MemTagNodesIdMap.FindRef(MemTagId); }
-	void SelectMemTagNode(Insights::FMemoryTagId MemTagId);
+	FMemTagNodePtr GetMemTagNode(FMemoryTagId MemTagId) const { return MemTagNodesIdMap.FindRef(MemTagId); }
+	void SelectMemTagNode(FMemoryTagId MemTagId);
 
 private:
 	TSharedRef<SWidget> MakeTrackersMenu();
@@ -131,7 +134,7 @@ private:
 
 	FText GetColumnHeaderText(const FName ColumnId) const;
 
-	TSharedRef<SWidget> TreeViewHeaderRow_GenerateColumnMenu(const Insights::FTableColumn& Column);
+	TSharedRef<SWidget> TreeViewHeaderRow_GenerateColumnMenu(const FTableColumn& Column);
 
 	////////////////////////////////////////////////////////////////////////////////////////////////////
 	// Tree View - Misc
@@ -159,7 +162,7 @@ private:
 
 	bool TableRow_ShouldBeEnabled(FMemTagNodePtr NodePtr) const;
 
-	void TableRow_SetHoveredCell(TSharedPtr<Insights::FTable> TablePtr, TSharedPtr<Insights::FTableColumn> ColumnPtr, FMemTagNodePtr NodePtr);
+	void TableRow_SetHoveredCell(TSharedPtr<FTable> TablePtr, TSharedPtr<FTableColumn> ColumnPtr, FMemTagNodePtr NodePtr);
 	EHorizontalAlignment TableRow_GetColumnOutlineHAlignment(const FName ColumnId) const;
 
 	FText TableRow_GetHighlightText() const;
@@ -177,8 +180,8 @@ private:
 	bool SearchBox_IsEnabled() const;
 	void SearchBox_OnTextChanged(const FText& InFilterText);
 
-	void ToggleTracker(Insights::FMemoryTrackerId InTrackerId);
-	bool IsTrackerChecked(Insights::FMemoryTrackerId InTrackerId) const;
+	void ToggleTracker(FMemoryTrackerId InTrackerId);
+	bool IsTrackerChecked(FMemoryTrackerId InTrackerId) const;
 
 	////////////////////////////////////////////////////////////////////////////////////////////////////
 	// Grouping
@@ -204,7 +207,7 @@ private:
 
 	void UpdateCurrentSortingByColumn();
 	void SortTreeNodes();
-	void SortTreeNodesRec(FMemTagNode& Node, const Insights::ITableCellValueSorter& Sorter);
+	void SortTreeNodesRec(FMemTagNode& Node, const ITableCellValueSorter& Sorter);
 
 	EColumnSortMode::Type GetSortModeForColumn(const FName ColumnId) const;
 	void SetSortModeForColumn(const FName& ColumnId, EColumnSortMode::Type SortMode);
@@ -313,7 +316,7 @@ private:
 	TWeakPtr<SMemoryProfilerWindow> ProfilerWindowWeakPtr;
 
 	/** Table view model. */
-	TSharedPtr<Insights::FTable> Table;
+	TSharedPtr<FTable> Table;
 
 	/** The analysis session used to populate this widget. */
 	TSharedPtr<const TraceServices::IAnalysisSession> Session;
@@ -358,7 +361,7 @@ private:
 	TSet<FMemTagNodePtr> MemTagNodes;
 
 	/** All LLM tag nodes, stored as NodeId -> FMemTagNodePtr. */
-	TMap<Insights::FMemoryTagId, FMemTagNodePtr> MemTagNodesIdMap;
+	TMap<FMemoryTagId, FMemTagNodePtr> MemTagNodesIdMap;
 
 	/** Currently expanded group nodes. */
 	TSet<FMemTagNodePtr> ExpandedNodes;
@@ -404,10 +407,10 @@ private:
 	//bool bUseSorting;
 
 	/** All available sorters. */
-	TArray<TSharedPtr<Insights::ITableCellValueSorter>> AvailableSorters;
+	TArray<TSharedPtr<ITableCellValueSorter>> AvailableSorters;
 
 	/** Current sorter. It is nullptr if sorting is disabled. */
-	TSharedPtr<Insights::ITableCellValueSorter> CurrentSorter;
+	TSharedPtr<ITableCellValueSorter> CurrentSorter;
 
 	/** Name of the column currently being sorted. Can be NAME_None if sorting is disabled (CurrentSorting == nullptr) or if a complex sorting is used (CurrentSorting != nullptr). */
 	FName ColumnBeingSorted;
@@ -417,7 +420,7 @@ private:
 
 	//////////////////////////////////////////////////
 
-	TSharedPtr<SComboBox<TSharedPtr<Insights::FMemoryTracker>>> TrackerComboBox;
+	TSharedPtr<SComboBox<TSharedPtr<FMemoryTracker>>> TrackerComboBox;
 
 	FLinearColor EditableColorValue;
 
@@ -428,3 +431,5 @@ private:
 };
 
 ////////////////////////////////////////////////////////////////////////////////////////////////////
+
+} // namespace UE::Insights::MemoryProfiler
diff --git a/Engine/Source/Developer/TraceInsights/Private/Insights/MemoryProfiler/Widgets/SMemTagTreeViewTableCell.cpp b/Engine/Source/Developer/TraceInsights/Private/Insights/MemoryProfiler/Widgets/SMemTagTreeViewTableCell.cpp
--- a/Engine/Source/Developer/TraceInsights/Private/Insights/MemoryProfiler/Widgets/SMemTagTreeViewTableCell.cpp
+++ b/Engine/Source/Developer/TraceInsights/Private/Insights/MemoryProfiler/Widgets/SMemTagTreeViewTableCell.cpp
@@ -11,13 +11,18 @@
 #include "Widgets/Text/STextBlock.h"
 #include "Widgets/Views/SExpanderArrow.h"
 
-// Insights
+// TraceInsightsCore
+#include "InsightsCore/Table/ViewModels/Table.h"
+#include "InsightsCore/Table/ViewModels/TableColumn.h"
+
+// TraceInsights
 #include "Insights/InsightsStyle.h"
-#include "Insights/Table/ViewModels/Table.h"
-#include "Insights/Table/ViewModels/TableColumn.h"
 #include "Insights/MemoryProfiler/Widgets/SMemTagTreeViewTableRow.h"
 
-#define LOCTEXT_NAMESPACE "SMemTagTreeViewView"
+#define LOCTEXT_NAMESPACE "UE::Insights::MemoryProfiler::SMemTagTreeView"
+
+namespace UE::Insights::MemoryProfiler
+{
 
 ////////////////////////////////////////////////////////////////////////////////////////////////////
 
@@ -175,4 +180,6 @@ END_SLATE_FUNCTION_BUILD_OPTIMIZATION
 
 ////////////////////////////////////////////////////////////////////////////////////////////////////
 
+} // namespace UE::Insights::MemoryProfiler
+
 #undef LOCTEXT_NAMESPACE
diff --git a/Engine/Source/Developer/TraceInsights/Private/Insights/MemoryProfiler/Widgets/SMemTagTreeViewTableCell.h b/Engine/Source/Developer/TraceInsights/Private/Insights/MemoryProfiler/Widgets/SMemTagTreeViewTableCell.h
--- a/Engine/Source/Developer/TraceInsights/Private/Insights/MemoryProfiler/Widgets/SMemTagTreeViewTableCell.h
+++ b/Engine/Source/Developer/TraceInsights/Private/Insights/MemoryProfiler/Widgets/SMemTagTreeViewTableCell.h
@@ -3,21 +3,25 @@
 #pragma once
 
 #include "CoreMinimal.h"
+
 #include "Widgets/DeclarativeSyntaxSupport.h"
 #include "Widgets/SCompoundWidget.h"
 
-// Insights
+// TraceInsights
 #include "Insights/MemoryProfiler/ViewModels/MemTagNode.h"
 
 class ITableRow;
 
-namespace Insights
+namespace UE::Insights
 {
 	class FTable;
 	class FTableColumn;
 }
 
-DECLARE_DELEGATE_ThreeParams(FSetHoveredMemTagTreeViewTableCell, TSharedPtr<Insights::FTable> /*TablePtr*/, TSharedPtr<Insights::FTableColumn> /*ColumnPtr*/, FMemTagNodePtr /*MemTagNodePtr*/);
+namespace UE::Insights::MemoryProfiler
+{
+
+DECLARE_DELEGATE_ThreeParams(FSetHoveredMemTagTreeViewTableCell, TSharedPtr<FTable> /*TablePtr*/, TSharedPtr<FTableColumn> /*ColumnPtr*/, FMemTagNodePtr /*MemTagNodePtr*/);
 
 class SMemTagTreeViewTableCell : public SCompoundWidget
 {
@@ -25,8 +29,8 @@ public:
 	SLATE_BEGIN_ARGS(SMemTagTreeViewTableCell) {}
 		SLATE_EVENT(FSetHoveredMemTagTreeViewTableCell, OnSetHoveredCell)
 		SLATE_ATTRIBUTE(FText, HighlightText)
-		SLATE_ARGUMENT(TSharedPtr<Insights::FTable>, TablePtr)
-		SLATE_ARGUMENT(TSharedPtr<Insights::FTableColumn>, ColumnPtr)
+		SLATE_ARGUMENT(TSharedPtr<FTable>, TablePtr)
+		SLATE_ARGUMENT(TSharedPtr<FTableColumn>, ColumnPtr)
 		SLATE_ARGUMENT(FMemTagNodePtr, MemTagNodePtr)
 		SLATE_ARGUMENT(bool, IsNameColumn)
 	SLATE_END_ARGS()
@@ -169,13 +173,15 @@ protected:
 
 protected:
 	/** A shared pointer to the table view model. */
-	TSharedPtr<Insights::FTable> TablePtr; // TODO: TSharedRef
+	TSharedPtr<FTable> TablePtr; // TODO: TSharedRef
 
 	/** A shared pointer to the table column view model. */
-	TSharedPtr<Insights::FTableColumn> ColumnPtr; // TODO: TSharedRef
+	TSharedPtr<FTableColumn> ColumnPtr; // TODO: TSharedRef
 
 	/** A shared pointer to the timer node. */
 	FMemTagNodePtr MemTagNodePtr; // TODO: TSharedRef
 
 	FSetHoveredMemTagTreeViewTableCell SetHoveredCellDelegate;
 };
+
+} // namespace UE::Insights::MemoryProfiler
diff --git a/Engine/Source/Developer/TraceInsights/Private/Insights/MemoryProfiler/Widgets/SMemTagTreeViewTableRow.cpp b/Engine/Source/Developer/TraceInsights/Private/Insights/MemoryProfiler/Widgets/SMemTagTreeViewTableRow.cpp
--- a/Engine/Source/Developer/TraceInsights/Private/Insights/MemoryProfiler/Widgets/SMemTagTreeViewTableRow.cpp
+++ b/Engine/Source/Developer/TraceInsights/Private/Insights/MemoryProfiler/Widgets/SMemTagTreeViewTableRow.cpp
@@ -6,15 +6,20 @@
 #include "Widgets/SOverlay.h"
 #include "Widgets/Images/SImage.h"
 
-// Insights
-#include "Insights/Common/TimeUtils.h"
+// TraceInsightsCore
+#include "InsightsCore/Common/TimeUtils.h"
+#include "InsightsCore/Table/ViewModels/Table.h"
+#include "InsightsCore/Table/ViewModels/TableColumn.h"
+
+// TraceInsights
 #include "Insights/InsightsStyle.h"
-#include "Insights/Table/ViewModels/Table.h"
-#include "Insights/Table/ViewModels/TableColumn.h"
 #include "Insights/MemoryProfiler/Widgets/SMemTagTreeViewTableCell.h"
 #include "Insights/MemoryProfiler/Widgets/SMemTagTreeViewTooltip.h"
 
-#define LOCTEXT_NAMESPACE "SMemTagTreeViewView"
+#define LOCTEXT_NAMESPACE "UE::Insights::MemoryProfiler::SMemTagTreeView"
+
+namespace UE::Insights::MemoryProfiler
+{
 
 ////////////////////////////////////////////////////////////////////////////////////////////////////
 
@@ -44,7 +49,7 @@ void SMemTagTreeViewTableRow::Construct(const FArguments& InArgs, const TSharedR
 
 TSharedRef<SWidget> SMemTagTreeViewTableRow::GenerateWidgetForColumn(const FName& ColumnId)
 {
-	TSharedPtr<Insights::FTableColumn> ColumnPtr = TablePtr->FindColumnChecked(ColumnId);
+	TSharedPtr<FTableColumn> ColumnPtr = TablePtr->FindColumnChecked(ColumnId);
 
 	return
 		SNew(SOverlay)
@@ -182,11 +187,13 @@ EVisibility SMemTagTreeViewTableRow::IsColumnVisible(const FName ColumnId) const
 
 ////////////////////////////////////////////////////////////////////////////////////////////////////
 
-void SMemTagTreeViewTableRow::OnSetHoveredCell(TSharedPtr<Insights::FTable> InTablePtr, TSharedPtr<Insights::FTableColumn> InColumnPtr, FMemTagNodePtr InMemTagNodePtr)
+void SMemTagTreeViewTableRow::OnSetHoveredCell(TSharedPtr<FTable> InTablePtr, TSharedPtr<FTableColumn> InColumnPtr, FMemTagNodePtr InMemTagNodePtr)
 {
 	SetHoveredCellDelegate.ExecuteIfBound(InTablePtr, InColumnPtr, InMemTagNodePtr);
 }
 
 ////////////////////////////////////////////////////////////////////////////////////////////////////
 
+} // namespace UE::Insights::MemoryProfiler
+
 #undef LOCTEXT_NAMESPACE
diff --git a/Engine/Source/Developer/TraceInsights/Private/Insights/MemoryProfiler/Widgets/SMemTagTreeViewTableRow.h b/Engine/Source/Developer/TraceInsights/Private/Insights/MemoryProfiler/Widgets/SMemTagTreeViewTableRow.h
--- a/Engine/Source/Developer/TraceInsights/Private/Insights/MemoryProfiler/Widgets/SMemTagTreeViewTableRow.h
+++ b/Engine/Source/Developer/TraceInsights/Private/Insights/MemoryProfiler/Widgets/SMemTagTreeViewTableRow.h
@@ -3,28 +3,36 @@
 #pragma once
 
 #include "CoreMinimal.h"
+
+#include "Delegates/Delegate.h"
+#include "Delegates/DelegateCombinations.h"
 #include "Widgets/DeclarativeSyntaxSupport.h"
 #include "Widgets/Views/SHeaderRow.h"
 #include "Widgets/Views/STableViewBase.h"
 #include "Widgets/Views/STableRow.h"
 #include "Widgets/Views/STreeView.h"
 
-// Insights
+// TraceInsights
+#include "Insights/MemoryProfiler/ViewModels/MemTagNode.h"
 #include "Insights/MemoryProfiler/ViewModels/MemTagNodeHelper.h"
 
 class IToolTip;
-class SMemCounterTableRowToolTip;
 
-namespace Insights
+namespace UE::Insights
 {
 	class FTable;
 	class FTableColumn;
 }
 
+namespace UE::Insights::MemoryProfiler
+{
+
+class SMemCounterTableRowToolTip;
+
 DECLARE_DELEGATE_RetVal_OneParam(bool, FShouldBeEnabledDelegate, FMemTagNodePtr /*NodePtr*/);
 DECLARE_DELEGATE_RetVal_OneParam(bool, FIsColumnVisibleDelegate, const FName /*ColumnId*/);
 DECLARE_DELEGATE_RetVal_OneParam(EHorizontalAlignment, FGetColumnOutlineHAlignmentDelegate, const FName /*ColumnId*/);
-DECLARE_DELEGATE_ThreeParams(FSetHoveredMemTagTreeViewTableCell, TSharedPtr<Insights::FTable> /*TablePtr*/, TSharedPtr<Insights::FTableColumn> /*ColumnPtr*/, FMemTagNodePtr /*MemTagNodePtr*/);
+DECLARE_DELEGATE_ThreeParams(FSetHoveredMemTagTreeViewTableCell, TSharedPtr<FTable> /*TablePtr*/, TSharedPtr<FTableColumn> /*ColumnPtr*/, FMemTagNodePtr /*MemTagNodePtr*/);
 
 /** Widget that represents a table row in the tree control. Generates widgets for each column on demand. */
 class SMemTagTreeViewTableRow : public SMultiColumnTableRow<FMemTagNodePtr>
@@ -37,7 +45,7 @@ public:
 		SLATE_EVENT(FSetHoveredMemTagTreeViewTableCell, OnSetHoveredCell)
 		SLATE_ATTRIBUTE(FText, HighlightText)
 		SLATE_ATTRIBUTE(FName, HighlightedNodeName)
-		SLATE_ARGUMENT(TSharedPtr<Insights::FTable>, TablePtr)
+		SLATE_ARGUMENT(TSharedPtr<FTable>, TablePtr)
 		SLATE_ARGUMENT(FMemTagNodePtr, MemTagNodePtr)
 	SLATE_END_ARGS()
 
@@ -72,11 +80,11 @@ protected:
 	const FSlateBrush* GetOutlineBrush(const FName ColumnId) const;
 	bool HandleShouldBeEnabled() const;
 	EVisibility IsColumnVisible(const FName ColumnId) const;
-	void OnSetHoveredCell(TSharedPtr<Insights::FTable> InTablePtr, TSharedPtr<Insights::FTableColumn> InColumnPtr, FMemTagNodePtr InMemTagNodePtr);
+	void OnSetHoveredCell(TSharedPtr<FTable> InTablePtr, TSharedPtr<FTableColumn> InColumnPtr, FMemTagNodePtr InMemTagNodePtr);
 
 protected:
 	/** A shared pointer to the table view model. */
-	TSharedPtr<Insights::FTable> TablePtr;
+	TSharedPtr<FTable> TablePtr;
 
 	/** Data context for this table row. */
 	FMemTagNodePtr MemTagNodePtr;
@@ -94,3 +102,5 @@ protected:
 
 	TSharedPtr<SMemCounterTableRowToolTip> RowToolTip;
 };
+
+} // namespace UE::Insights::MemoryProfiler
diff --git a/Engine/Source/Developer/TraceInsights/Private/Insights/MemoryProfiler/Widgets/SMemTagTreeViewTooltip.cpp b/Engine/Source/Developer/TraceInsights/Private/Insights/MemoryProfiler/Widgets/SMemTagTreeViewTooltip.cpp
--- a/Engine/Source/Developer/TraceInsights/Private/Insights/MemoryProfiler/Widgets/SMemTagTreeViewTooltip.cpp
+++ b/Engine/Source/Developer/TraceInsights/Private/Insights/MemoryProfiler/Widgets/SMemTagTreeViewTooltip.cpp
@@ -4,29 +4,34 @@
 
 #include "SlateOptMacros.h"
 #include "Styling/StyleColors.h"
-#include "TraceServices/Model/NetProfiler.h"
 #include "Widgets/Layout/SGridPanel.h"
 #include "Widgets/Layout/SSeparator.h"
 #include "Widgets/SBoxPanel.h"
 #include "Widgets/SToolTip.h"
 #include "Widgets/Text/STextBlock.h"
 
-// Insights
+// TraceServices
+#include "TraceServices/Model/NetProfiler.h"
+
+// TraceInsightsCore
+#include "InsightsCore/Table/ViewModels/Table.h"
+#include "InsightsCore/Table/ViewModels/TableColumn.h"
+
+// TraceInsights
 #include "Insights/InsightsStyle.h"
-#include "Insights/Table/ViewModels/Table.h"
-#include "Insights/Table/ViewModels/TableColumn.h"
 #include "Insights/MemoryProfiler/ViewModels/MemTagNode.h"
 #include "Insights/MemoryProfiler/ViewModels/MemTagNodeHelper.h"
 
-////////////////////////////////////////////////////////////////////////////////////////////////////
+#define LOCTEXT_NAMESPACE "UE::Insights::MemoryProfiler::SMemTagTreeView"
 
-#define LOCTEXT_NAMESPACE "SMemTagTreeView"
+namespace UE::Insights::MemoryProfiler
+{
 
 ////////////////////////////////////////////////////////////////////////////////////////////////////
 
 BEGIN_SLATE_FUNCTION_BUILD_OPTIMIZATION
 
-TSharedPtr<SToolTip> SMemTagTreeViewTooltip::GetTableTooltip(const Insights::FTable& Table)
+TSharedPtr<SToolTip> SMemTagTreeViewTooltip::GetTableTooltip(const FTable& Table)
 {
 	TSharedPtr<SToolTip> ColumnTooltip =
 		SNew(SToolTip)
@@ -57,7 +62,7 @@ TSharedPtr<SToolTip> SMemTagTreeViewTooltip::GetTableTooltip(const Insights::FTa
 
 ////////////////////////////////////////////////////////////////////////////////////////////////////
 
-TSharedPtr<SToolTip> SMemTagTreeViewTooltip::GetColumnTooltip(const Insights::FTableColumn& Column)
+TSharedPtr<SToolTip> SMemTagTreeViewTooltip::GetColumnTooltip(const FTableColumn& Column)
 {
 	TSharedPtr<SToolTip> ColumnTooltip =
 		SNew(SToolTip)
@@ -96,7 +101,7 @@ TSharedPtr<SToolTip> SMemTagTreeViewTooltip::GetRowTooltip(const TSharedPtr<FMem
 
 	FText TagText = FText::GetEmpty();
 	FText TagTextEx = LOCTEXT("MemTagNA", "N/A");
-	Insights::FMemoryTag* MemTag = MemTagNodePtr.IsValid() ? MemTagNodePtr->GetMemTag() : nullptr;
+	FMemoryTag* MemTag = MemTagNodePtr.IsValid() ? MemTagNodePtr->GetMemTag() : nullptr;
 	if (MemTag)
 	{
 		TagText = FText::FromString(MemTag->GetStatName());
@@ -105,7 +110,7 @@ TSharedPtr<SToolTip> SMemTagTreeViewTooltip::GetRowTooltip(const TSharedPtr<FMem
 
 	FText ParentTagText = FText::GetEmpty();
 	FText ParentTagTextEx = LOCTEXT("MemTagNA", "N/A");
-	Insights::FMemoryTag* ParentMemTag = MemTagNodePtr.IsValid() ? MemTagNodePtr->GetParentMemTag() : nullptr;
+	FMemoryTag* ParentMemTag = MemTagNodePtr.IsValid() ? MemTagNodePtr->GetParentMemTag() : nullptr;
 	if (ParentMemTag)
 	{
 		ParentTagText = FText::FromString(ParentMemTag->GetStatName());
@@ -358,4 +363,6 @@ END_SLATE_FUNCTION_BUILD_OPTIMIZATION
 
 ////////////////////////////////////////////////////////////////////////////////////////////////////
 
+} // namespace UE::Insights::MemoryProfiler
+
 #undef LOCTEXT_NAMESPACE
diff --git a/Engine/Source/Developer/TraceInsights/Private/Insights/MemoryProfiler/Widgets/SMemTagTreeViewTooltip.h b/Engine/Source/Developer/TraceInsights/Private/Insights/MemoryProfiler/Widgets/SMemTagTreeViewTooltip.h
--- a/Engine/Source/Developer/TraceInsights/Private/Insights/MemoryProfiler/Widgets/SMemTagTreeViewTooltip.h
+++ b/Engine/Source/Developer/TraceInsights/Private/Insights/MemoryProfiler/Widgets/SMemTagTreeViewTooltip.h
@@ -3,17 +3,21 @@
 #pragma once
 
 #include "CoreMinimal.h"
+
 #include "Widgets/IToolTip.h"
 #include "Widgets/SToolTip.h"
 
 class SGridPanel;
 
-namespace Insights
+namespace UE::Insights
 {
 	class FTable;
 	class FTableColumn;
 }
 
+namespace UE::Insights::MemoryProfiler
+{
+
 class FMemTagNode;
 
 ////////////////////////////////////////////////////////////////////////////////////////////////////
@@ -24,8 +28,8 @@ class SMemTagTreeViewTooltip
 public:
 	SMemTagTreeViewTooltip() = delete;
 
-	static TSharedPtr<SToolTip> GetTableTooltip(const Insights::FTable& Table);
-	static TSharedPtr<SToolTip> GetColumnTooltip(const Insights::FTableColumn& Column);
+	static TSharedPtr<SToolTip> GetTableTooltip(const FTable& Table);
+	static TSharedPtr<SToolTip> GetColumnTooltip(const FTableColumn& Column);
 	static TSharedPtr<SToolTip> GetRowTooltip(const TSharedPtr<FMemTagNode> TreeNodePtr);
 
 private:
@@ -83,3 +87,5 @@ private:
 };
 
 ////////////////////////////////////////////////////////////////////////////////////////////////////
+
+} // namespace UE::Insights::MemoryProfiler
diff --git a/Engine/Source/Developer/TraceInsights/Private/Insights/MemoryProfiler/Widgets/SMemoryProfilerToolbar.cpp b/Engine/Source/Developer/TraceInsights/Private/Insights/MemoryProfiler/Widgets/SMemoryProfilerToolbar.cpp
--- a/Engine/Source/Developer/TraceInsights/Private/Insights/MemoryProfiler/Widgets/SMemoryProfilerToolbar.cpp
+++ b/Engine/Source/Developer/TraceInsights/Private/Insights/MemoryProfiler/Widgets/SMemoryProfilerToolbar.cpp
@@ -6,7 +6,7 @@
 #include "Framework/MultiBox/MultiBoxBuilder.h"
 #include "Widgets/SBoxPanel.h"
 
-// Insights
+// TraceInsights
 #include "Insights/InsightsCommands.h"
 #include "Insights/InsightsManager.h"
 #include "Insights/InsightsStyle.h"
@@ -14,7 +14,10 @@
 
 ////////////////////////////////////////////////////////////////////////////////////////////////////
 
-#define LOCTEXT_NAMESPACE "SMemoryProfilerToolbar"
+#define LOCTEXT_NAMESPACE "UE::Insights::MemoryProfiler"
+
+namespace UE::Insights::MemoryProfiler
+{
 
 ////////////////////////////////////////////////////////////////////////////////////////////////////
 
@@ -112,4 +115,6 @@ void SMemoryProfilerToolbar::Construct(const FArguments& InArgs)
 
 ////////////////////////////////////////////////////////////////////////////////////////////////////
 
+} // namespace UE::Insights::MemoryProfiler
+
 #undef LOCTEXT_NAMESPACE
diff --git a/Engine/Source/Developer/TraceInsights/Private/Insights/MemoryProfiler/Widgets/SMemoryProfilerToolbar.h b/Engine/Source/Developer/TraceInsights/Private/Insights/MemoryProfiler/Widgets/SMemoryProfilerToolbar.h
--- a/Engine/Source/Developer/TraceInsights/Private/Insights/MemoryProfiler/Widgets/SMemoryProfilerToolbar.h
+++ b/Engine/Source/Developer/TraceInsights/Private/Insights/MemoryProfiler/Widgets/SMemoryProfilerToolbar.h
@@ -7,6 +7,9 @@
 
 class FExtender;
 
+namespace UE::Insights::MemoryProfiler
+{
+
 /** Ribbon based toolbar used as a main menu in the Profiler window. */
 class SMemoryProfilerToolbar : public SCompoundWidget
 {
@@ -28,3 +31,5 @@ public:
 	 */
 	void Construct(const FArguments& InArgs);
 };
+
+} // namespace UE::Insights::MemoryProfiler
diff --git a/Engine/Source/Developer/TraceInsights/Private/Insights/MemoryProfiler/Widgets/SMemoryProfilerWindow.cpp b/Engine/Source/Developer/TraceInsights/Private/Insights/MemoryProfiler/Widgets/SMemoryProfilerWindow.cpp
--- a/Engine/Source/Developer/TraceInsights/Private/Insights/MemoryProfiler/Widgets/SMemoryProfilerWindow.cpp
+++ b/Engine/Source/Developer/TraceInsights/Private/Insights/MemoryProfiler/Widgets/SMemoryProfilerWindow.cpp
@@ -9,7 +9,7 @@
 #include "SlateOptMacros.h"
 #include "Widgets/Docking/SDockTab.h"
 
-// Insights
+// TraceInsights
 #include "Insights/InsightsStyle.h"
 #include "Insights/MemoryProfiler/MemoryProfilerManager.h"
 #include "Insights/MemoryProfiler/ViewModels/MemorySharedState.h"
@@ -18,11 +18,13 @@
 #include "Insights/MemoryProfiler/Widgets/SMemoryProfilerToolbar.h"
 #include "Insights/MemoryProfiler/Widgets/SMemTagTreeView.h"
 #include "Insights/ViewModels/TimeRulerTrack.h"
+#include "Insights/Widgets/SModulesView.h"
 #include "Insights/Widgets/STimingView.h"
 
-////////////////////////////////////////////////////////////////////////////////////////////////////
+#define LOCTEXT_NAMESPACE "UE::Insights::MemoryProfiler"
 
-#define LOCTEXT_NAMESPACE "SMemoryProfilerWindow"
+namespace UE::Insights::MemoryProfiler
+{
 
 ////////////////////////////////////////////////////////////////////////////////////////////////////
 
@@ -97,12 +99,12 @@ TSharedRef<SDockTab> SMemoryProfilerWindow::SpawnTab_TimingView(const FSpawnTabA
 		.ShouldAutosize(false)
 		.TabRole(ETabRole::PanelTab)
 		[
-			SAssignNew(TimingView, STimingView, FInsightsManagerTabs::MemoryProfilerTabId)
+			SAssignNew(TimingView, TimingProfiler::STimingView, FInsightsManagerTabs::MemoryProfilerTabId)
 		];
 
 	SharedState->SetTimingView(TimingView);
 	SharedState->BindCommands();
-	IModularFeatures::Get().RegisterModularFeature(Insights::TimingViewExtenderFeatureName, &SharedState.Get());
+	IModularFeatures::Get().RegisterModularFeature(Timing::TimingViewExtenderFeatureName, &SharedState.Get());
 
 	TimingView->Reset(true);
 	TimingView->OnSelectionChanged().AddSP(this, &SMemoryProfilerWindow::OnTimeSelectionChanged);
@@ -130,7 +132,7 @@ void SMemoryProfilerWindow::OnTimingViewTabClosed(TSharedRef<SDockTab> TabBeingC
 		TimingView = nullptr;
 	}
 
-	IModularFeatures::Get().UnregisterModularFeature(Insights::TimingViewExtenderFeatureName, &SharedState.Get());
+	IModularFeatures::Get().UnregisterModularFeature(Timing::TimingViewExtenderFeatureName, &SharedState.Get());
 	SharedState->SetTimingView(nullptr);
 
 	RemoveOpenTab(TabBeingClosed);
@@ -205,17 +207,17 @@ TSharedRef<SDockTab> SMemoryProfilerWindow::SpawnTab_MemAllocTableTreeView(const
 {
 	//FMemoryProfilerManager::Get()->SetMemAllocTableTreeViewVisible(TabIndex, true);
 
-	TSharedRef<Insights::FMemAllocTable> MemAllocTable = MakeShared<Insights::FMemAllocTable>();
+	TSharedRef<FMemAllocTable> MemAllocTable = MakeShared<FMemAllocTable>();
 	MemAllocTable->Reset();
 	MemAllocTable->SetDisplayName(FText::FromString(TEXT("MemAllocs")));
 
-	TSharedPtr<Insights::SMemAllocTableTreeView> MemAllocTableTreeView;
+	TSharedPtr<SMemAllocTableTreeView> MemAllocTableTreeView;
 
 	const TSharedRef<SDockTab> DockTab = SNew(SDockTab)
 		.ShouldAutosize(false)
 		.TabRole(ETabRole::PanelTab)
 		[
-			SAssignNew(MemAllocTableTreeView, Insights::SMemAllocTableTreeView, MemAllocTable)
+			SAssignNew(MemAllocTableTreeView, SMemAllocTableTreeView, MemAllocTable)
 		];
 
 	MemAllocTableTreeView->SetLogListingName(FMemoryProfilerManager::Get()->GetLogListingName());
@@ -233,13 +235,13 @@ END_SLATE_FUNCTION_BUILD_OPTIMIZATION
 
 void SMemoryProfilerWindow::OnMemAllocTableTreeViewTabClosed(TSharedRef<SDockTab> TabBeingClosed)
 {
-	TSharedRef<Insights::SMemAllocTableTreeView> MemAllocTableTreeView = StaticCastSharedRef<Insights::SMemAllocTableTreeView>(TabBeingClosed->GetContent());
+	TSharedRef<SMemAllocTableTreeView> MemAllocTableTreeView = StaticCastSharedRef<SMemAllocTableTreeView>(TabBeingClosed->GetContent());
 
 	FName ClosingTabId = FMemoryProfilerTabs::MemAllocTableTreeViewID;
 	ClosingTabId.SetNumber(MemAllocTableTreeView->GetTabIndex());
 
-	TSharedPtr<Insights::FQueryTargetWindowSpec> TargetToDelete;
-	const TArray<TSharedPtr<Insights::FQueryTargetWindowSpec>>& Targets = this->GetSharedState().GetQueryTargets();
+	TSharedPtr<FQueryTargetWindowSpec> TargetToDelete;
+	const TArray<TSharedPtr<FQueryTargetWindowSpec>>& Targets = this->GetSharedState().GetQueryTargets();
 	for (int32 Index = 0; Index < Targets.Num(); ++Index)
 	{
 		if (Targets[Index]->GetName() == ClosingTabId)
@@ -256,7 +258,7 @@ void SMemoryProfilerWindow::OnMemAllocTableTreeViewTabClosed(TSharedRef<SDockTab
 
 	if (Targets.Num() > 0)
 	{
-		TSharedPtr<Insights::FQueryTargetWindowSpec> NewSelection = Targets[0];
+		TSharedPtr<FQueryTargetWindowSpec> NewSelection = Targets[0];
 		this->GetSharedState().SetCurrentQueryTarget(NewSelection);
 		if (MemInvestigationView.IsValid())
 		{
@@ -276,13 +278,13 @@ void SMemoryProfilerWindow::OnMemAllocTableTreeViewTabClosed(TSharedRef<SDockTab
 
 void SMemoryProfilerWindow::CloseMemAllocTableTreeTabs()
 {
-	const TArray<TSharedPtr<Insights::FQueryTargetWindowSpec>>& Targets = this->GetSharedState().GetQueryTargets();
+	const TArray<TSharedPtr<FQueryTargetWindowSpec>>& Targets = this->GetSharedState().GetQueryTargets();
 	while (Targets.Num() > 0)
 	{
 		FName Name = Targets[0]->GetName();
 		this->GetSharedState().RemoveQueryTarget(Targets[0]);
 
-		if (Name != Insights::FQueryTargetWindowSpec::NewWindow)
+		if (Name != FQueryTargetWindowSpec::NewWindow)
 		{
 			HideTab(Name);
 		}
@@ -292,7 +294,7 @@ void SMemoryProfilerWindow::CloseMemAllocTableTreeTabs()
 ////////////////////////////////////////////////////////////////////////////////////////////////////
 
 BEGIN_SLATE_FUNCTION_BUILD_OPTIMIZATION
-TSharedPtr<Insights::SMemAllocTableTreeView> SMemoryProfilerWindow::ShowMemAllocTableTreeViewTab()
+TSharedPtr<SMemAllocTableTreeView> SMemoryProfilerWindow::ShowMemAllocTableTreeViewTab()
 {
 	if (!GetTabManager().IsValid())
 	{
@@ -300,7 +302,7 @@ TSharedPtr<Insights::SMemAllocTableTreeView> SMemoryProfilerWindow::ShowMemAlloc
 	}
 	FTabManager* TabManagerPtr = GetTabManager().Get();
 
-	if (SharedState->GetCurrentQueryTarget()->GetName() == Insights::FQueryTargetWindowSpec::NewWindow)
+	if (SharedState->GetCurrentQueryTarget()->GetName() == FQueryTargetWindowSpec::NewWindow)
 	{
 		++LastMemAllocTableTreeViewIndex;
 		FName TabId = FMemoryProfilerTabs::MemAllocTableTreeViewID;
@@ -315,7 +317,7 @@ TSharedPtr<Insights::SMemAllocTableTreeView> SMemoryProfilerWindow::ShowMemAlloc
 			.SetIcon(FSlateIcon(FInsightsStyle::GetStyleSetName(), "Icons.MemAllocTableTreeView"))
 			.SetGroup(Group);
 
-		TSharedPtr<Insights::FQueryTargetWindowSpec> NewTarget = MakeShared<Insights::FQueryTargetWindowSpec>(TabId, MemAllocTableTreeViewTabDisplayName);
+		TSharedPtr<FQueryTargetWindowSpec> NewTarget = MakeShared<FQueryTargetWindowSpec>(TabId, MemAllocTableTreeViewTabDisplayName);
 		SharedState->AddQueryTarget(NewTarget);
 		SharedState->SetCurrentQueryTarget(NewTarget);
 		MemInvestigationView->QueryTarget_OnSelectionChanged(NewTarget, ESelectInfo::Type::Direct);
@@ -327,9 +329,9 @@ TSharedPtr<Insights::SMemAllocTableTreeView> SMemoryProfilerWindow::ShowMemAlloc
 		TSharedPtr<SDockTab> Tab = TabManagerPtr->TryInvokeTab(TabId);
 		if (Tab)
 		{
-			TSharedRef<Insights::SMemAllocTableTreeView> MemAllocTableTreeView = StaticCastSharedRef<Insights::SMemAllocTableTreeView>(Tab->GetContent());
+			TSharedRef<SMemAllocTableTreeView> MemAllocTableTreeView = StaticCastSharedRef<SMemAllocTableTreeView>(Tab->GetContent());
 
-			if (SharedState->GetCurrentQueryTarget()->GetName() == Insights::FQueryTargetWindowSpec::NewWindow)
+			if (SharedState->GetCurrentQueryTarget()->GetName() == FQueryTargetWindowSpec::NewWindow)
 			{
 				Tab->SetOnTabClosed(SDockTab::FOnTabClosedCallback::CreateRaw(this, &SMemoryProfilerWindow::OnMemAllocTableTreeViewTabClosed));
 			}
@@ -352,7 +354,7 @@ TSharedRef<SDockTab> SMemoryProfilerWindow::SpawnTab_ModulesView(const FSpawnTab
 		.ShouldAutosize(false)
 		.TabRole(ETabRole::PanelTab)
 		[
-			SAssignNew(ModulesView, Insights::SModulesView)
+			SAssignNew(ModulesView, SModulesView)
 		];
 
 	DockTab->SetOnTabClosed(SDockTab::FOnTabClosedCallback::CreateRaw(this, &SMemoryProfilerWindow::OnModulesViewClosed));
@@ -461,9 +463,9 @@ TSharedRef<SWidget> SMemoryProfilerWindow::CreateToolbar(TSharedPtr<FExtender> E
 
 ////////////////////////////////////////////////////////////////////////////////////////////////////
 
-void SMemoryProfilerWindow::OnTimeSelectionChanged(Insights::ETimeChangedFlags InFlags, double InStartTime, double InEndTime)
+void SMemoryProfilerWindow::OnTimeSelectionChanged(Timing::ETimeChangedFlags InFlags, double InStartTime, double InEndTime)
 {
-	if (InFlags != Insights::ETimeChangedFlags::Interactive)
+	if (InFlags != Timing::ETimeChangedFlags::Interactive)
 	{
 		if (InStartTime < InEndTime)
 		{
@@ -507,7 +509,7 @@ void SMemoryProfilerWindow::CreateTimingViewMarkers()
 
 	for (uint32 Index = 0; Index < MaxNumTimeMarkers; ++Index)
 	{
-		TSharedRef<Insights::FTimeMarker> TimeMarker = MakeShared<Insights::FTimeMarker>();
+		TSharedRef<TimingProfiler::FTimeMarker> TimeMarker = MakeShared<TimingProfiler::FTimeMarker>();
 
 		TimeMarkerName[0] = static_cast<TCHAR>(TEXT('A') + Index); // "A", "B", "C", etc.
 		TimeMarker->SetName(TimeMarkerName);
@@ -530,18 +532,18 @@ void SMemoryProfilerWindow::CreateTimingViewMarkers()
 
 void SMemoryProfilerWindow::ResetTimingViewMarkers()
 {
-	TSharedRef<FTimeRulerTrack> TimeRulerTrack = TimingView->GetTimeRulerTrack();
+	TSharedRef<TimingProfiler::FTimeRulerTrack> TimeRulerTrack = TimingView->GetTimeRulerTrack();
 
 	TimeRulerTrack->RemoveAllTimeMarkers();
 
 	// Hide the "Default Time Marker".
-	TSharedRef<Insights::FTimeMarker> DefaultTimeMarker = TimingView->GetDefaultTimeMarker();
+	TSharedRef<TimingProfiler::FTimeMarker> DefaultTimeMarker = TimingView->GetDefaultTimeMarker();
 	DefaultTimeMarker->SetVisibility(false);
 
 	const uint32 NumTimeMarkers = CustomTimeMarkers.Num();
 	for (uint32 Index = 0; Index < NumTimeMarkers; ++Index)
 	{
-		TSharedRef<Insights::FTimeMarker>& TimeMarker = CustomTimeMarkers[Index];
+		TSharedRef<TimingProfiler::FTimeMarker>& TimeMarker = CustomTimeMarkers[Index];
 		TimeMarker->SetTime((Index + 1) * 10.0); // 10s, 20s, 30s, etc.
 		TimeRulerTrack->AddTimeMarker(TimeMarker);
 	}
@@ -560,7 +562,7 @@ void SMemoryProfilerWindow::OnMemoryRuleChanged()
 
 void SMemoryProfilerWindow::UpdateTimingViewMarkers()
 {
-	TSharedPtr<Insights::FMemoryRuleSpec> Rule = SharedState->GetCurrentMemoryRule();
+	TSharedPtr<FMemoryRuleSpec> Rule = SharedState->GetCurrentMemoryRule();
 	const uint32 NumVisibleTimeMarkers = Rule ? Rule->GetNumTimeMarkers() : 0;
 
 	const uint32 NumTimeMarkers = CustomTimeMarkers.Num();
@@ -568,7 +570,7 @@ void SMemoryProfilerWindow::UpdateTimingViewMarkers()
 
 	for (uint32 Index = 0; Index < NumTimeMarkers; ++Index)
 	{
-		TSharedRef<Insights::FTimeMarker>& TimeMarker = CustomTimeMarkers[Index];
+		TSharedRef<TimingProfiler::FTimeMarker>& TimeMarker = CustomTimeMarkers[Index];
 		if (Index < NumVisibleTimeMarkers)
 		{
 			TimeMarker->SetVisibility(true);
@@ -582,7 +584,7 @@ void SMemoryProfilerWindow::UpdateTimingViewMarkers()
 
 ////////////////////////////////////////////////////////////////////////////////////////////////////
 
-void SMemoryProfilerWindow::OnTimeMarkerChanged(Insights::ETimeChangedFlags InFlags, TSharedRef<Insights::ITimeMarker> InTimeMarker)
+void SMemoryProfilerWindow::OnTimeMarkerChanged(Timing::ETimeChangedFlags InFlags, TSharedRef<Timing::ITimeMarker> InTimeMarker)
 {
 	const int32 NumTimeMarkers = CustomTimeMarkers.Num();
 
@@ -590,7 +592,7 @@ void SMemoryProfilerWindow::OnTimeMarkerChanged(Insights::ETimeChangedFlags InFl
 	int32 ChangedTimeMarkerIndex = -1;
 	for (int32 Index = 0; Index < NumTimeMarkers; ++Index)
 	{
-		TSharedRef<Insights::FTimeMarker>& TimeMarker = CustomTimeMarkers[Index];
+		TSharedRef<TimingProfiler::FTimeMarker>& TimeMarker = CustomTimeMarkers[Index];
 		if (TimeMarker == InTimeMarker)
 		{
 			ChangedTimeMarkerIndex = Index;
@@ -604,17 +606,17 @@ void SMemoryProfilerWindow::OnTimeMarkerChanged(Insights::ETimeChangedFlags InFl
 		TimingView.IsValid() &&
 		InTimeMarker == TimingView->GetDefaultTimeMarker())
 	{
-		TSharedRef<Insights::FTimeMarker>& TimeMarkerA = CustomTimeMarkers[0];
+		TSharedRef<TimingProfiler::FTimeMarker>& TimeMarkerA = CustomTimeMarkers[0];
 		TimeMarkerA->SetTime(InTimeMarker->GetTime());
 		ChangedTimeMarkerIndex = 0;
 	}
 
-	// Ensure the rest of time markers are orderd by time.
+	// Ensure the rest of time markers are ordered by time.
 	if (ChangedTimeMarkerIndex >= 0)
 	{
 		for (int32 Index = 0; Index < ChangedTimeMarkerIndex; ++Index)
 		{
-			TSharedRef<Insights::FTimeMarker>& TimeMarker = CustomTimeMarkers[Index];
+			TSharedRef<TimingProfiler::FTimeMarker>& TimeMarker = CustomTimeMarkers[Index];
 			if (TimeMarker->GetTime() > InTimeMarker->GetTime())
 			{
 				TimeMarker->SetTime(InTimeMarker->GetTime());
@@ -622,7 +624,7 @@ void SMemoryProfilerWindow::OnTimeMarkerChanged(Insights::ETimeChangedFlags InFl
 		}
 		for (int32 Index = ChangedTimeMarkerIndex + 1; Index < NumTimeMarkers; ++Index)
 		{
-			TSharedRef<Insights::FTimeMarker>& TimeMarker = CustomTimeMarkers[Index];
+			TSharedRef<TimingProfiler::FTimeMarker>& TimeMarker = CustomTimeMarkers[Index];
 			if (TimeMarker->GetTime() < InTimeMarker->GetTime())
 			{
 				TimeMarker->SetTime(InTimeMarker->GetTime());
@@ -633,4 +635,6 @@ void SMemoryProfilerWindow::OnTimeMarkerChanged(Insights::ETimeChangedFlags InFl
 
 ////////////////////////////////////////////////////////////////////////////////////////////////////
 
+} // namespace UE::Insights::MemoryProfiler
+
 #undef LOCTEXT_NAMESPACE
diff --git a/Engine/Source/Developer/TraceInsights/Private/Insights/MemoryProfiler/Widgets/SMemoryProfilerWindow.h b/Engine/Source/Developer/TraceInsights/Private/Insights/MemoryProfiler/Widgets/SMemoryProfilerWindow.h
--- a/Engine/Source/Developer/TraceInsights/Private/Insights/MemoryProfiler/Widgets/SMemoryProfilerWindow.h
+++ b/Engine/Source/Developer/TraceInsights/Private/Insights/MemoryProfiler/Widgets/SMemoryProfilerWindow.h
@@ -4,25 +4,32 @@
 
 #include "CoreMinimal.h"
 
-// Insights
-#include "Insights/ITimingViewSession.h" // for Insights::ETimeChangedFlags
-#include "Insights/Widgets/SMajorTabWindow.h"
-#include "Insights/Widgets/SModulesView.h"
+#include "UObject/NameTypes.h"
 
-////////////////////////////////////////////////////////////////////////////////////////////////////
+// TraceInsights
+#include "Insights/ITimingViewSession.h" // for ETimeChangedFlags and ITimeMarker
+#include "Insights/Widgets/SMajorTabWindow.h"
 
-class FMemorySharedState;
-class SMemInvestigationView;
-class SMemTagTreeView;
-class STimingView;
+namespace UE::Insights
+{
+	class STableTreeView;
+	class SModulesView;
+}
 
-namespace Insights
+namespace UE::Insights::TimingProfiler
 {
 	class FTimeMarker;
-	class STableTreeView;
-	class SMemAllocTableTreeView;
+	class STimingView;
 }
 
+namespace UE::Insights::MemoryProfiler
+{
+
+class FMemorySharedState;
+class SMemInvestigationView;
+class SMemTagTreeView;
+class SMemAllocTableTreeView;
+
 ////////////////////////////////////////////////////////////////////////////////////////////////////
 
 struct FMemoryProfilerTabs
@@ -38,7 +45,7 @@ struct FMemoryProfilerTabs
 ////////////////////////////////////////////////////////////////////////////////////////////////////
 
 /** Implements the Memory Insights window. */
-class SMemoryProfilerWindow : public Insights::SMajorTabWindow
+class SMemoryProfilerWindow : public ::Insights::SMajorTabWindow
 {
 public:
 	/** Default constructor. */
@@ -55,22 +62,22 @@ public:
 	/** Constructs this widget. */
 	void Construct(const FArguments& InArgs, const TSharedRef<SDockTab>& ConstructUnderMajorTab, const TSharedPtr<SWindow>& ConstructUnderWindow);
 
-	TSharedPtr<STimingView> GetTimingView() const { return TimingView; }
+	TSharedPtr<TimingProfiler::STimingView> GetTimingView() const { return TimingView; }
 	TSharedPtr<SMemInvestigationView> GetMemInvestigationView() const { return MemInvestigationView; }
 	TSharedPtr<SMemTagTreeView> GetMemTagTreeView() const { return MemTagTreeView; }
 
 	void CloseMemAllocTableTreeTabs();
-	TSharedPtr<Insights::SMemAllocTableTreeView> ShowMemAllocTableTreeViewTab();
+	TSharedPtr<SMemAllocTableTreeView> ShowMemAllocTableTreeViewTab();
 
 	uint32 GetNumCustomTimeMarkers() const { return (uint32)CustomTimeMarkers.Num(); }
-	const TSharedRef<Insights::FTimeMarker>& GetCustomTimeMarker(uint32 Index) const { return CustomTimeMarkers[Index]; }
-	const TArray<TSharedRef<Insights::FTimeMarker>>& GetCustomTimeMarkers() const { return CustomTimeMarkers; }
+	const TSharedRef<TimingProfiler::FTimeMarker>& GetCustomTimeMarker(uint32 Index) const { return CustomTimeMarkers[Index]; }
+	const TArray<TSharedRef<TimingProfiler::FTimeMarker>>& GetCustomTimeMarkers() const { return CustomTimeMarkers; }
 
 	FMemorySharedState& GetSharedState() { return *SharedState; }
 	const FMemorySharedState& GetSharedState() const { return *SharedState; }
 
 	void OnMemoryRuleChanged();
-	void OnTimeMarkerChanged(Insights::ETimeChangedFlags InFlags, TSharedRef<Insights::ITimeMarker> InTimeMarker);
+	void OnTimeMarkerChanged(Timing::ETimeChangedFlags InFlags, TSharedRef<Timing::ITimeMarker> InTimeMarker);
 
 protected:
 	virtual const TCHAR* GetAnalyticsEventName() const override;
@@ -95,7 +102,7 @@ private:
 	TSharedRef<SDockTab> SpawnTab_ModulesView(const FSpawnTabArgs& Args);
 	void OnModulesViewClosed(TSharedRef<SDockTab> TabBeingClosed);
 
-	void OnTimeSelectionChanged(Insights::ETimeChangedFlags InFlags, double InStartTime, double InEndTime);
+	void OnTimeSelectionChanged(Timing::ETimeChangedFlags InFlags, double InStartTime, double InEndTime);
 
 	void CreateTimingViewMarkers();
 	void ResetTimingViewMarkers();
@@ -105,9 +112,9 @@ private:
 	TSharedRef<FMemorySharedState> SharedState;
 
 	/** The Timing view (multi-track) widget */
-	TSharedPtr<STimingView> TimingView;
+	TSharedPtr<TimingProfiler::STimingView> TimingView;
 
-	TArray<TSharedRef<Insights::FTimeMarker>> CustomTimeMarkers;
+	TArray<TSharedRef<TimingProfiler::FTimeMarker>> CustomTimeMarkers;
 
 	/** The Memory Investigation (Allocation Queries) view widget */
 	TSharedPtr<SMemInvestigationView> MemInvestigationView;
@@ -116,11 +123,15 @@ private:
 	TSharedPtr<SMemTagTreeView> MemTagTreeView;
 
 	/** The list of Allocations table tree view widgets */
-	TArray<TSharedPtr<Insights::SMemAllocTableTreeView>> MemAllocTableTreeViews;
+	TArray<TSharedPtr<SMemAllocTableTreeView>> MemAllocTableTreeViews;
 
 	/** The Modules view widget. */
-	TSharedPtr<Insights::SModulesView> ModulesView;
+	TSharedPtr<SModulesView> ModulesView;
 
 	const int32 MaxMemAllocTableTreeViews = 4;
 	int32 LastMemAllocTableTreeViewIndex = -1;
 };
+
+////////////////////////////////////////////////////////////////////////////////////////////////////
+
+} // namespace UE::Insights::MemoryProfiler
diff --git a/Engine/Source/Developer/TraceInsights/Private/Insights/NetworkingProfiler/NetworkingProfilerCommands.h b/Engine/Source/Developer/TraceInsights/Private/Insights/NetworkingProfiler/NetworkingProfilerCommands.h
--- a/Engine/Source/Developer/TraceInsights/Private/Insights/NetworkingProfiler/NetworkingProfilerCommands.h
+++ b/Engine/Source/Developer/TraceInsights/Private/Insights/NetworkingProfiler/NetworkingProfilerCommands.h
@@ -3,11 +3,15 @@
 #pragma once
 
 #include "CoreMinimal.h"
+
 #include "Framework/Commands/UIAction.h"
 #include "Framework/Commands/Commands.h"
 
 class FMenuBuilder;
 
+namespace UE::Insights::NetworkingProfiler
+{
+
 ////////////////////////////////////////////////////////////////////////////////////////////////////
 /**
  * Class that holds all profiler commands.
@@ -100,3 +104,7 @@ protected:
 	/** Reference to the global instance of the profiler manager. */
 	class FNetworkingProfilerManager* This;
 };
+
+////////////////////////////////////////////////////////////////////////////////////////////////////
+
+} // namespace UE::Insights::NetworkingProfiler
