[Insights] Massive refactoring of TraceInsights module:

diff --git a/Engine/Source/Developer/TraceInsights/Private/Insights/NetworkingProfiler/NetworkingProfilerManager.cpp b/Engine/Source/Developer/TraceInsights/Private/Insights/NetworkingProfiler/NetworkingProfilerManager.cpp
--- a/Engine/Source/Developer/TraceInsights/Private/Insights/NetworkingProfiler/NetworkingProfilerManager.cpp
+++ b/Engine/Source/Developer/TraceInsights/Private/Insights/NetworkingProfiler/NetworkingProfilerManager.cpp
@@ -8,17 +8,18 @@
 #include "WorkspaceMenuStructure.h"
 #include "WorkspaceMenuStructureModule.h"
 
-// Insights
+// TraceInsights
 #include "Insights/Common/InsightsMenuBuilder.h"
 #include "Insights/InsightsManager.h"
 #include "Insights/InsightsStyle.h"
 #include "Insights/NetworkingProfiler/Widgets/SNetworkingProfilerWindow.h"
 
-////////////////////////////////////////////////////////////////////////////////////////////////////
+#define LOCTEXT_NAMESPACE "UE::Insights::NetworkingProfiler"
 
-#define LOCTEXT_NAMESPACE "NetworkingProfilerManager"
+namespace UE::Insights::NetworkingProfiler
+{
 
-DEFINE_LOG_CATEGORY(NetworkingProfiler);
+DEFINE_LOG_CATEGORY(LogNetworkingProfiler);
 
 TSharedPtr<FNetworkingProfilerManager> FNetworkingProfilerManager::Instance = nullptr;
 
@@ -67,7 +68,7 @@ void FNetworkingProfilerManager::Initialize(IUnrealInsightsModule& InsightsModul
 	}
 	bIsInitialized = true;
 
-	UE_LOG(NetworkingProfiler, Log, TEXT("Initialize"));
+	UE_LOG(LogNetworkingProfiler, Log, TEXT("Initialize"));
 
 	// Register tick functions.
 	OnTick = FTickerDelegate::CreateSP(this, &FNetworkingProfilerManager::Tick);
@@ -109,7 +110,7 @@ void FNetworkingProfilerManager::Shutdown()
 
 	FNetworkingProfilerManager::Instance.Reset();
 
-	UE_LOG(NetworkingProfiler, Log, TEXT("Shutdown"));
+	UE_LOG(LogNetworkingProfiler, Log, TEXT("Shutdown"));
 }
 
 ////////////////////////////////////////////////////////////////////////////////////////////////////
@@ -267,7 +268,7 @@ bool FNetworkingProfilerManager::Tick(float DeltaTime)
 			if (FGlobalTabmanager::Get()->HasTabSpawner(TabId))
 			{
 				// Spawn 2 tabs.
-				UE_LOG(NetworkingProfiler, Log, TEXT("Opening the \"Networking Insights\" tabs..."));
+				UE_LOG(LogNetworkingProfiler, Log, TEXT("Opening the \"Networking Insights\" tabs..."));
 				FGlobalTabmanager::Get()->TryInvokeTab(TabId);
 				FGlobalTabmanager::Get()->TryInvokeTab(TabId);
 			}
@@ -285,7 +286,7 @@ bool FNetworkingProfilerManager::Tick(float DeltaTime)
 			//	if (FGlobalTabmanager::Get()->HasTabSpawner(TabId) && 
 			//		!FGlobalTabmanager::Get()->FindExistingLiveTab(TabId).IsValid())
 			//	{
-			//		UE_LOG(NetworkingProfiler, Log, TEXT("Opening the \"Networking Insights\" tab..."));
+			//		UE_LOG(LogNetworkingProfiler, Log, TEXT("Opening the \"Networking Insights\" tab..."));
 			//		FGlobalTabmanager::Get()->TryInvokeTab(TabId);
 			//		--SpawnTabCount;
 			//	}
@@ -302,7 +303,7 @@ bool FNetworkingProfilerManager::Tick(float DeltaTime)
 
 void FNetworkingProfilerManager::OnSessionChanged()
 {
-	UE_LOG(NetworkingProfiler, Log, TEXT("OnSessionChanged"));
+	UE_LOG(LogNetworkingProfiler, Log, TEXT("OnSessionChanged"));
 
 	bIsAvailable = false;
 	if (FInsightsManager::Get()->GetSession().IsValid())
@@ -326,4 +327,6 @@ void FNetworkingProfilerManager::OnSessionChanged()
 
 ////////////////////////////////////////////////////////////////////////////////////////////////////
 
+} // namespace UE::Insights::NetworkingProfiler
+
 #undef LOCTEXT_NAMESPACE
diff --git a/Engine/Source/Developer/TraceInsights/Private/Insights/NetworkingProfiler/NetworkingProfilerManager.h b/Engine/Source/Developer/TraceInsights/Private/Insights/NetworkingProfiler/NetworkingProfilerManager.h
--- a/Engine/Source/Developer/TraceInsights/Private/Insights/NetworkingProfiler/NetworkingProfilerManager.h
+++ b/Engine/Source/Developer/TraceInsights/Private/Insights/NetworkingProfiler/NetworkingProfilerManager.h
@@ -2,19 +2,23 @@
 
 #pragma once
 
-#include "Containers/Ticker.h"
 #include "CoreMinimal.h"
+
+#include "Containers/Ticker.h"
 #include "Framework/Commands/UICommandList.h"
 #include "Logging/LogMacros.h"
 
-// Insights
+// TraceInsights
 #include "Insights/InsightsManager.h"
 #include "Insights/IUnrealInsightsModule.h"
 #include "Insights/NetworkingProfiler/NetworkingProfilerCommands.h"
 
+namespace UE::Insights::NetworkingProfiler
+{
+
 class SNetworkingProfilerWindow;
 
-DECLARE_LOG_CATEGORY_EXTERN(NetworkingProfiler, Log, All);
+DECLARE_LOG_CATEGORY_EXTERN(LogNetworkingProfiler, Log, All);
 
 ////////////////////////////////////////////////////////////////////////////////////////////////////
 /**
@@ -113,7 +117,7 @@ private:
 private:
 	bool bIsInitialized;
 	bool bIsAvailable;
-	Insights::FAvailabilityCheck AvailabilityCheck;
+	FAvailabilityCheck AvailabilityCheck;
 
 	/** The delegate to be invoked when this manager ticks. */
 	FTickerDelegate OnTick;
@@ -136,3 +140,5 @@ private:
 	/** The name of the Networking Profiler log listing. */
 	FName LogListingName;
 };
+
+} // namespace UE::Insights::NetworkingProfiler
diff --git a/Engine/Source/Developer/TraceInsights/Private/Insights/NetworkingProfiler/ViewModels/NetEventGroupingAndSorting.cpp b/Engine/Source/Developer/TraceInsights/Private/Insights/NetworkingProfiler/ViewModels/NetEventGroupingAndSorting.cpp
--- a/Engine/Source/Developer/TraceInsights/Private/Insights/NetworkingProfiler/ViewModels/NetEventGroupingAndSorting.cpp
+++ b/Engine/Source/Developer/TraceInsights/Private/Insights/NetworkingProfiler/ViewModels/NetEventGroupingAndSorting.cpp
@@ -2,10 +2,13 @@
 
 #include "NetEventGroupingAndSorting.h"
 
-#include "Insights/Table/ViewModels/TableColumn.h"
+// TraceInsightsCore
+#include "InsightsCore/Table/ViewModels/TableColumn.h"
+
+// TraceInsights
 #include "Insights/NetworkingProfiler/ViewModels/NetEventNodeHelper.h"
 
-#define LOCTEXT_NAMESPACE "NetEventNode"
+#define LOCTEXT_NAMESPACE "UE::Insights::NetworkingProfiler::FNetEventNode"
 
 #define INSIGHTS_ENSURE ensure
 //#define INSIGHTS_ENSURE(...)
@@ -14,12 +17,15 @@
 #define INSIGHTS_DEFAULT_SORTING_NODES(A, B) return A->GetName().LexicalLess(B->GetName());
 //#define INSIGHTS_DEFAULT_SORTING_NODES(A, B) return A->GetDefaultSortOrder() < B->GetDefaultSortOrder();
 
+namespace UE::Insights::NetworkingProfiler
+{
+
 ////////////////////////////////////////////////////////////////////////////////////////////////////
 // Sorting by Event Type
 ////////////////////////////////////////////////////////////////////////////////////////////////////
 
-FNetEventNodeSortingByEventType::FNetEventNodeSortingByEventType(TSharedRef<Insights::FTableColumn> InColumnRef)
-	: Insights::FTableCellValueSorter(
+FNetEventNodeSortingByEventType::FNetEventNodeSortingByEventType(TSharedRef<FTableColumn> InColumnRef)
+	: FTableCellValueSorter(
 		FName(TEXT("ByEventType")),
 		LOCTEXT("Sorting_ByEventType_Name", "By Type"),
 		LOCTEXT("Sorting_ByEventType_Title", "Sort By Type"),
@@ -30,17 +36,17 @@ FNetEventNodeSortingByEventType::FNetEventNodeSortingByEventType(TSharedRef<Insi
 
 ////////////////////////////////////////////////////////////////////////////////////////////////////
 
-void FNetEventNodeSortingByEventType::Sort(TArray<Insights::FBaseTreeNodePtr>& NodesToSort, Insights::ESortMode SortMode) const
+void FNetEventNodeSortingByEventType::Sort(TArray<FBaseTreeNodePtr>& NodesToSort, ESortMode SortMode) const
 {
-	if (SortMode == Insights::ESortMode::Ascending)
+	if (SortMode == ESortMode::Ascending)
 	{
-		NodesToSort.Sort([](const Insights::FBaseTreeNodePtr& A, const Insights::FBaseTreeNodePtr& B) -> bool
+		NodesToSort.Sort([](const FBaseTreeNodePtr& A, const FBaseTreeNodePtr& B) -> bool
 		{
 			INSIGHTS_ENSURE(A.IsValid() && A->Is<FNetEventNode>());
-			const FNetEventNodePtr NetEventNodeA = StaticCastSharedPtr<FNetEventNode, Insights::FBaseTreeNode>(A);
+			const FNetEventNodePtr NetEventNodeA = StaticCastSharedPtr<FNetEventNode, FBaseTreeNode>(A);
 
 			INSIGHTS_ENSURE(B.IsValid() && B->Is<FNetEventNode>());
-			const FNetEventNodePtr NetEventNodeB = StaticCastSharedPtr<FNetEventNode, Insights::FBaseTreeNode>(B);
+			const FNetEventNodePtr NetEventNodeB = StaticCastSharedPtr<FNetEventNode, FBaseTreeNode>(B);
 
 			if (NetEventNodeA->GetType() == NetEventNodeB->GetType())
 			{
@@ -55,13 +61,13 @@ void FNetEventNodeSortingByEventType::Sort(TArray<Insights::FBaseTreeNodePtr>& N
 	}
 	else // if (SortMode == ESortMode::Descending)
 	{
-		NodesToSort.Sort([](const Insights::FBaseTreeNodePtr& A, const Insights::FBaseTreeNodePtr& B) -> bool
+		NodesToSort.Sort([](const FBaseTreeNodePtr& A, const FBaseTreeNodePtr& B) -> bool
 		{
 			INSIGHTS_ENSURE(A.IsValid() && A->Is<FNetEventNode>());
-			const FNetEventNodePtr NetEventNodeA = StaticCastSharedPtr<FNetEventNode, Insights::FBaseTreeNode>(A);
+			const FNetEventNodePtr NetEventNodeA = StaticCastSharedPtr<FNetEventNode, FBaseTreeNode>(A);
 
 			INSIGHTS_ENSURE(B.IsValid() && B->Is<FNetEventNode>());
-			const FNetEventNodePtr NetEventNodeB = StaticCastSharedPtr<FNetEventNode, Insights::FBaseTreeNode>(B);
+			const FNetEventNodePtr NetEventNodeB = StaticCastSharedPtr<FNetEventNode, FBaseTreeNode>(B);
 
 			if (NetEventNodeA->GetType() == NetEventNodeB->GetType())
 			{
@@ -80,8 +86,8 @@ void FNetEventNodeSortingByEventType::Sort(TArray<Insights::FBaseTreeNodePtr>& N
 // Sort by Instance Count
 ////////////////////////////////////////////////////////////////////////////////////////////////////
 
-FNetEventNodeSortingByInstanceCount::FNetEventNodeSortingByInstanceCount(TSharedRef<Insights::FTableColumn> InColumnRef)
-	: Insights::FTableCellValueSorter(
+FNetEventNodeSortingByInstanceCount::FNetEventNodeSortingByInstanceCount(TSharedRef<FTableColumn> InColumnRef)
+	: FTableCellValueSorter(
 		FName(TEXT("ByInstanceCount")),
 		LOCTEXT("Sorting_ByInstanceCount_Name", "By Instance Count"),
 		LOCTEXT("Sorting_ByInstanceCount_Title", "Sort By Instance Count"),
@@ -92,18 +98,18 @@ FNetEventNodeSortingByInstanceCount::FNetEventNodeSortingByInstanceCount(TShared
 
 ////////////////////////////////////////////////////////////////////////////////////////////////////
 
-void FNetEventNodeSortingByInstanceCount::Sort(TArray<Insights::FBaseTreeNodePtr>& NodesToSort, Insights::ESortMode SortMode) const
+void FNetEventNodeSortingByInstanceCount::Sort(TArray<FBaseTreeNodePtr>& NodesToSort, ESortMode SortMode) const
 {
-	if (SortMode == Insights::ESortMode::Ascending)
+	if (SortMode == ESortMode::Ascending)
 	{
-		NodesToSort.Sort([](const Insights::FBaseTreeNodePtr& A, const Insights::FBaseTreeNodePtr& B) -> bool
+		NodesToSort.Sort([](const FBaseTreeNodePtr& A, const FBaseTreeNodePtr& B) -> bool
 		{
 			INSIGHTS_ENSURE(A.IsValid() && A->Is<FNetEventNode>());
-			const FNetEventNodePtr NetEventNodeA = StaticCastSharedPtr<FNetEventNode, Insights::FBaseTreeNode>(A);
+			const FNetEventNodePtr NetEventNodeA = StaticCastSharedPtr<FNetEventNode, FBaseTreeNode>(A);
 			const uint64 ValueA = NetEventNodeA->GetAggregatedStats().InstanceCount;
 
 			INSIGHTS_ENSURE(B.IsValid() && B->Is<FNetEventNode>());
-			const FNetEventNodePtr NetEventNodeB = StaticCastSharedPtr<FNetEventNode, Insights::FBaseTreeNode>(B);
+			const FNetEventNodePtr NetEventNodeB = StaticCastSharedPtr<FNetEventNode, FBaseTreeNode>(B);
 			const uint64 ValueB = NetEventNodeB->GetAggregatedStats().InstanceCount;
 
 			if (ValueA == ValueB)
@@ -119,14 +125,14 @@ void FNetEventNodeSortingByInstanceCount::Sort(TArray<Insights::FBaseTreeNodePtr
 	}
 	else // if (SortMode == ESortMode::Descending)
 	{
-		NodesToSort.Sort([](const Insights::FBaseTreeNodePtr& A, const Insights::FBaseTreeNodePtr& B) -> bool
+		NodesToSort.Sort([](const FBaseTreeNodePtr& A, const FBaseTreeNodePtr& B) -> bool
 		{
 			INSIGHTS_ENSURE(A.IsValid() && A->Is<FNetEventNode>());
-			const FNetEventNodePtr NetEventNodeA = StaticCastSharedPtr<FNetEventNode, Insights::FBaseTreeNode>(A);
+			const FNetEventNodePtr NetEventNodeA = StaticCastSharedPtr<FNetEventNode, FBaseTreeNode>(A);
 			const uint64 ValueA = NetEventNodeA->GetAggregatedStats().InstanceCount;
 
 			INSIGHTS_ENSURE(B.IsValid() && B->Is<FNetEventNode>());
-			const FNetEventNodePtr NetEventNodeB = StaticCastSharedPtr<FNetEventNode, Insights::FBaseTreeNode>(B);
+			const FNetEventNodePtr NetEventNodeB = StaticCastSharedPtr<FNetEventNode, FBaseTreeNode>(B);
 			const uint64 ValueB = NetEventNodeB->GetAggregatedStats().InstanceCount;
 
 			if (ValueA == ValueB)
@@ -146,8 +152,8 @@ void FNetEventNodeSortingByInstanceCount::Sort(TArray<Insights::FBaseTreeNodePtr
 // Sort by Total Inclusive Size
 ////////////////////////////////////////////////////////////////////////////////////////////////////
 
-FNetEventNodeSortingByTotalInclusiveSize::FNetEventNodeSortingByTotalInclusiveSize(TSharedRef<Insights::FTableColumn> InColumnRef)
-	: Insights::FTableCellValueSorter(
+FNetEventNodeSortingByTotalInclusiveSize::FNetEventNodeSortingByTotalInclusiveSize(TSharedRef<FTableColumn> InColumnRef)
+	: FTableCellValueSorter(
 		FName(TEXT("ByTotalInclusiveSize")),
 		LOCTEXT("Sorting_ByTotalInclusiveSize_Name", "By Total Inclusive Size"),
 		LOCTEXT("Sorting_ByTotalInclusiveSize_Title", "Sort By Total Inclusive Size"),
@@ -158,18 +164,18 @@ FNetEventNodeSortingByTotalInclusiveSize::FNetEventNodeSortingByTotalInclusiveSi
 
 ////////////////////////////////////////////////////////////////////////////////////////////////////
 
-void FNetEventNodeSortingByTotalInclusiveSize::Sort(TArray<Insights::FBaseTreeNodePtr>& NodesToSort, Insights::ESortMode SortMode) const
+void FNetEventNodeSortingByTotalInclusiveSize::Sort(TArray<FBaseTreeNodePtr>& NodesToSort, ESortMode SortMode) const
 {
-	if (SortMode == Insights::ESortMode::Ascending)
+	if (SortMode == ESortMode::Ascending)
 	{
-		NodesToSort.Sort([](const Insights::FBaseTreeNodePtr& A, const Insights::FBaseTreeNodePtr& B) -> bool
+		NodesToSort.Sort([](const FBaseTreeNodePtr& A, const FBaseTreeNodePtr& B) -> bool
 		{
 			INSIGHTS_ENSURE(A.IsValid() && A->Is<FNetEventNode>());
-			const FNetEventNodePtr NetEventNodeA = StaticCastSharedPtr<FNetEventNode, Insights::FBaseTreeNode>(A);
+			const FNetEventNodePtr NetEventNodeA = StaticCastSharedPtr<FNetEventNode, FBaseTreeNode>(A);
 			const uint32 ValueA = NetEventNodeA->GetAggregatedStats().TotalInclusive;
 
 			INSIGHTS_ENSURE(B.IsValid() && B->Is<FNetEventNode>());
-			const FNetEventNodePtr NetEventNodeB = StaticCastSharedPtr<FNetEventNode, Insights::FBaseTreeNode>(B);
+			const FNetEventNodePtr NetEventNodeB = StaticCastSharedPtr<FNetEventNode, FBaseTreeNode>(B);
 			const uint32 ValueB = NetEventNodeB->GetAggregatedStats().TotalInclusive;
 
 			if (ValueA == ValueB)
@@ -185,14 +191,14 @@ void FNetEventNodeSortingByTotalInclusiveSize::Sort(TArray<Insights::FBaseTreeNo
 	}
 	else // if (SortMode == ESortMode::Descending)
 	{
-		NodesToSort.Sort([](const Insights::FBaseTreeNodePtr& A, const Insights::FBaseTreeNodePtr& B) -> bool
+		NodesToSort.Sort([](const FBaseTreeNodePtr& A, const FBaseTreeNodePtr& B) -> bool
 		{
 			INSIGHTS_ENSURE(A.IsValid() && A->Is<FNetEventNode>());
-			const FNetEventNodePtr NetEventNodeA = StaticCastSharedPtr<FNetEventNode, Insights::FBaseTreeNode>(A);
+			const FNetEventNodePtr NetEventNodeA = StaticCastSharedPtr<FNetEventNode, FBaseTreeNode>(A);
 			const uint32 ValueA = NetEventNodeA->GetAggregatedStats().TotalInclusive;
 
 			INSIGHTS_ENSURE(B.IsValid() && B->Is<FNetEventNode>());
-			const FNetEventNodePtr NetEventNodeB = StaticCastSharedPtr<FNetEventNode, Insights::FBaseTreeNode>(B);
+			const FNetEventNodePtr NetEventNodeB = StaticCastSharedPtr<FNetEventNode, FBaseTreeNode>(B);
 			const uint32 ValueB = NetEventNodeB->GetAggregatedStats().TotalInclusive;
 
 			if (ValueA == ValueB)
@@ -212,8 +218,8 @@ void FNetEventNodeSortingByTotalInclusiveSize::Sort(TArray<Insights::FBaseTreeNo
 // Sort by Total Exclusive Size
 ////////////////////////////////////////////////////////////////////////////////////////////////////
 
-FNetEventNodeSortingByTotalExclusiveSize::FNetEventNodeSortingByTotalExclusiveSize(TSharedRef<Insights::FTableColumn> InColumnRef)
-	: Insights::FTableCellValueSorter(
+FNetEventNodeSortingByTotalExclusiveSize::FNetEventNodeSortingByTotalExclusiveSize(TSharedRef<FTableColumn> InColumnRef)
+	: FTableCellValueSorter(
 		FName(TEXT("ByTotalExclusiveSize")),
 		LOCTEXT("Sorting_ByTotalExclusiveSize_Name", "By Total Exclusive Size"),
 		LOCTEXT("Sorting_ByTotalExclusiveSize_Title", "Sort By Total Exclusive Size"),
@@ -224,18 +230,18 @@ FNetEventNodeSortingByTotalExclusiveSize::FNetEventNodeSortingByTotalExclusiveSi
 
 ////////////////////////////////////////////////////////////////////////////////////////////////////
 
-void FNetEventNodeSortingByTotalExclusiveSize::Sort(TArray<Insights::FBaseTreeNodePtr>& NodesToSort, Insights::ESortMode SortMode) const
+void FNetEventNodeSortingByTotalExclusiveSize::Sort(TArray<FBaseTreeNodePtr>& NodesToSort, ESortMode SortMode) const
 {
-	if (SortMode == Insights::ESortMode::Ascending)
+	if (SortMode == ESortMode::Ascending)
 	{
-		NodesToSort.Sort([](const Insights::FBaseTreeNodePtr& A, const Insights::FBaseTreeNodePtr& B) -> bool
+		NodesToSort.Sort([](const FBaseTreeNodePtr& A, const FBaseTreeNodePtr& B) -> bool
 		{
 			INSIGHTS_ENSURE(A.IsValid() && A->Is<FNetEventNode>());
-			const FNetEventNodePtr NetEventNodeA = StaticCastSharedPtr<FNetEventNode, Insights::FBaseTreeNode>(A);
+			const FNetEventNodePtr NetEventNodeA = StaticCastSharedPtr<FNetEventNode, FBaseTreeNode>(A);
 			const uint32 ValueA = NetEventNodeA->GetAggregatedStats().TotalExclusive;
 
 			INSIGHTS_ENSURE(B.IsValid() && B->Is<FNetEventNode>());
-			const FNetEventNodePtr NetEventNodeB = StaticCastSharedPtr<FNetEventNode, Insights::FBaseTreeNode>(B);
+			const FNetEventNodePtr NetEventNodeB = StaticCastSharedPtr<FNetEventNode, FBaseTreeNode>(B);
 			const uint32 ValueB = NetEventNodeB->GetAggregatedStats().TotalExclusive;
 
 			if (ValueA == ValueB)
@@ -251,14 +257,14 @@ void FNetEventNodeSortingByTotalExclusiveSize::Sort(TArray<Insights::FBaseTreeNo
 	}
 	else // if (SortMode == ESortMode::Descending)
 	{
-		NodesToSort.Sort([](const Insights::FBaseTreeNodePtr& A, const Insights::FBaseTreeNodePtr& B) -> bool
+		NodesToSort.Sort([](const FBaseTreeNodePtr& A, const FBaseTreeNodePtr& B) -> bool
 		{
 			INSIGHTS_ENSURE(A.IsValid() && A->Is<FNetEventNode>());
-			const FNetEventNodePtr NetEventNodeA = StaticCastSharedPtr<FNetEventNode, Insights::FBaseTreeNode>(A);
+			const FNetEventNodePtr NetEventNodeA = StaticCastSharedPtr<FNetEventNode, FBaseTreeNode>(A);
 			const uint32 ValueA = NetEventNodeA->GetAggregatedStats().TotalExclusive;
 
 			INSIGHTS_ENSURE(B.IsValid() && B->Is<FNetEventNode>());
-			const FNetEventNodePtr NetEventNodeB = StaticCastSharedPtr<FNetEventNode, Insights::FBaseTreeNode>(B);
+			const FNetEventNodePtr NetEventNodeB = StaticCastSharedPtr<FNetEventNode, FBaseTreeNode>(B);
 			const uint32 ValueB = NetEventNodeB->GetAggregatedStats().TotalExclusive;
 
 			if (ValueA == ValueB)
@@ -276,6 +282,8 @@ void FNetEventNodeSortingByTotalExclusiveSize::Sort(TArray<Insights::FBaseTreeNo
 
 ////////////////////////////////////////////////////////////////////////////////////////////////////
 
+} // namespace UE::Insights::NetworkingProfiler
+
 #undef INSIGHTS_DEFAULT_SORTING_NODES
 #undef INSIGHTS_ENSURE
 #undef LOCTEXT_NAMESPACE
diff --git a/Engine/Source/Developer/TraceInsights/Private/Insights/NetworkingProfiler/ViewModels/NetEventGroupingAndSorting.h b/Engine/Source/Developer/TraceInsights/Private/Insights/NetworkingProfiler/ViewModels/NetEventGroupingAndSorting.h
--- a/Engine/Source/Developer/TraceInsights/Private/Insights/NetworkingProfiler/ViewModels/NetEventGroupingAndSorting.h
+++ b/Engine/Source/Developer/TraceInsights/Private/Insights/NetworkingProfiler/ViewModels/NetEventGroupingAndSorting.h
@@ -4,52 +4,57 @@
 
 #include "CoreMinimal.h"
 
-// Insights
+// TraceInsightsCore
+#include "InsightsCore/Table/ViewModels/TableColumn.h"
+#include "InsightsCore/Table/ViewModels/TableCellValueSorter.h"
+#include "InsightsCore/Table/ViewModels/TreeNodeGrouping.h"
+
+// TraceInsights
 #include "Insights/NetworkingProfiler/ViewModels/NetEventNode.h"
-#include "Insights/Table/ViewModels/TableColumn.h"
-#include "Insights/Table/ViewModels/TableCellValueSorter.h"
-#include "Insights/Table/ViewModels/TreeNodeGrouping.h"
+
+namespace UE::Insights::NetworkingProfiler
+{
 
 ////////////////////////////////////////////////////////////////////////////////////////////////////
 // Sorters
 ////////////////////////////////////////////////////////////////////////////////////////////////////
 
-class FNetEventNodeSortingByEventType: public Insights::FTableCellValueSorter
+class FNetEventNodeSortingByEventType: public FTableCellValueSorter
 {
 public:
-	FNetEventNodeSortingByEventType(TSharedRef<Insights::FTableColumn> InColumnRef);
+	FNetEventNodeSortingByEventType(TSharedRef<FTableColumn> InColumnRef);
 
-	virtual void Sort(TArray<Insights::FBaseTreeNodePtr>& NodesToSort, Insights::ESortMode SortMode) const override;
+	virtual void Sort(TArray<FBaseTreeNodePtr>& NodesToSort, ESortMode SortMode) const override;
 };
 
 ////////////////////////////////////////////////////////////////////////////////////////////////////
 
-class FNetEventNodeSortingByInstanceCount : public Insights::FTableCellValueSorter
+class FNetEventNodeSortingByInstanceCount : public FTableCellValueSorter
 {
 public:
-	FNetEventNodeSortingByInstanceCount(TSharedRef<Insights::FTableColumn> InColumnRef);
+	FNetEventNodeSortingByInstanceCount(TSharedRef<FTableColumn> InColumnRef);
 
-	virtual void Sort(TArray<Insights::FBaseTreeNodePtr>& NodesToSort, Insights::ESortMode SortMode) const override;
+	virtual void Sort(TArray<FBaseTreeNodePtr>& NodesToSort, ESortMode SortMode) const override;
 };
 
 ////////////////////////////////////////////////////////////////////////////////////////////////////
 
-class FNetEventNodeSortingByTotalInclusiveSize : public Insights::FTableCellValueSorter
+class FNetEventNodeSortingByTotalInclusiveSize : public FTableCellValueSorter
 {
 public:
-	FNetEventNodeSortingByTotalInclusiveSize(TSharedRef<Insights::FTableColumn> InColumnRef);
+	FNetEventNodeSortingByTotalInclusiveSize(TSharedRef<FTableColumn> InColumnRef);
 
-	virtual void Sort(TArray<Insights::FBaseTreeNodePtr>& NodesToSort, Insights::ESortMode SortMode) const override;
+	virtual void Sort(TArray<FBaseTreeNodePtr>& NodesToSort, ESortMode SortMode) const override;
 };
 
 ////////////////////////////////////////////////////////////////////////////////////////////////////
 
-class FNetEventNodeSortingByTotalExclusiveSize : public Insights::FTableCellValueSorter
+class FNetEventNodeSortingByTotalExclusiveSize : public FTableCellValueSorter
 {
 public:
-	FNetEventNodeSortingByTotalExclusiveSize(TSharedRef<Insights::FTableColumn> InColumnRef);
+	FNetEventNodeSortingByTotalExclusiveSize(TSharedRef<FTableColumn> InColumnRef);
 
-	virtual void Sort(TArray<Insights::FBaseTreeNodePtr>& NodesToSort, Insights::ESortMode SortMode) const override;
+	virtual void Sort(TArray<FBaseTreeNodePtr>& NodesToSort, ESortMode SortMode) const override;
 };
 
 ////////////////////////////////////////////////////////////////////////////////////////////////////
@@ -79,3 +84,5 @@ enum class ENetEventGroupingMode
 typedef TSharedPtr<ENetEventGroupingMode> ENetEventGroupingModePtr;
 
 ////////////////////////////////////////////////////////////////////////////////////////////////////
+
+} // namespace UE::Insights::NetworkingProfiler
diff --git a/Engine/Source/Developer/TraceInsights/Private/Insights/NetworkingProfiler/ViewModels/NetEventNode.cpp b/Engine/Source/Developer/TraceInsights/Private/Insights/NetworkingProfiler/ViewModels/NetEventNode.cpp
--- a/Engine/Source/Developer/TraceInsights/Private/Insights/NetworkingProfiler/ViewModels/NetEventNode.cpp
+++ b/Engine/Source/Developer/TraceInsights/Private/Insights/NetworkingProfiler/ViewModels/NetEventNode.cpp
@@ -2,7 +2,10 @@
 
 #include "NetEventNode.h"
 
-#define LOCTEXT_NAMESPACE "FNetEventNode"
+#define LOCTEXT_NAMESPACE "UE::Insights::NetworkingProfiler::FNetEventNode"
+
+namespace UE::Insights::NetworkingProfiler
+{
 
 INSIGHTS_IMPLEMENT_RTTI(FNetEventNode)
 
@@ -22,4 +25,6 @@ void FNetEventNode::SetAggregatedStats(const TraceServices::FNetProfilerAggregat
 
 ////////////////////////////////////////////////////////////////////////////////////////////////////
 
+} // namespace UE::Insights::NetworkingProfiler
+
 #undef LOCTEXT_NAMESPACE
diff --git a/Engine/Source/Developer/TraceInsights/Private/Insights/NetworkingProfiler/ViewModels/NetEventNode.h b/Engine/Source/Developer/TraceInsights/Private/Insights/NetworkingProfiler/ViewModels/NetEventNode.h
--- a/Engine/Source/Developer/TraceInsights/Private/Insights/NetworkingProfiler/ViewModels/NetEventNode.h
+++ b/Engine/Source/Developer/TraceInsights/Private/Insights/NetworkingProfiler/ViewModels/NetEventNode.h
@@ -3,10 +3,15 @@
 #pragma once
 
 #include "CoreMinimal.h"
+
+// TraceServices
 #include "TraceServices/Model/NetProfiler.h"
 
-// Insights
-#include "Insights/Table/ViewModels/BaseTreeNode.h"
+// TraceInsightsCore
+#include "InsightsCore/Table/ViewModels/BaseTreeNode.h"
+
+namespace UE::Insights::NetworkingProfiler
+{
 
 ////////////////////////////////////////////////////////////////////////////////////////////////////
 
@@ -45,7 +50,7 @@ typedef TWeakPtr<class FNetEventNode> FNetEventNodeWeak;
 /**
  * Class used to store information about a timer node (used in the SNetStatsView).
  */
-class FNetEventNode : public Insights::FBaseTreeNode
+class FNetEventNode : public FBaseTreeNode
 {
 	INSIGHTS_DECLARE_RTTI(FNetEventNode, FBaseTreeNode)
 
@@ -105,3 +110,5 @@ private:
 };
 
 ////////////////////////////////////////////////////////////////////////////////////////////////////
+
+} // namespace UE::Insights::NetworkingProfiler
diff --git a/Engine/Source/Developer/TraceInsights/Private/Insights/NetworkingProfiler/ViewModels/NetEventNodeHelper.cpp b/Engine/Source/Developer/TraceInsights/Private/Insights/NetworkingProfiler/ViewModels/NetEventNodeHelper.cpp
--- a/Engine/Source/Developer/TraceInsights/Private/Insights/NetworkingProfiler/ViewModels/NetEventNodeHelper.cpp
+++ b/Engine/Source/Developer/TraceInsights/Private/Insights/NetworkingProfiler/ViewModels/NetEventNodeHelper.cpp
@@ -2,10 +2,13 @@
 
 #include "NetEventNodeHelper.h"
 
-// Insights
+// TraceInsights
 #include "Insights/InsightsStyle.h"
 
-#define LOCTEXT_NAMESPACE "NetEventNode"
+#define LOCTEXT_NAMESPACE "UE::Insights::NetworkingProfiler::FNetEventNode"
+
+namespace UE::Insights::NetworkingProfiler
+{
 
 ////////////////////////////////////////////////////////////////////////////////////////////////////
 // NetEventNode Type Helper
@@ -89,4 +92,6 @@ FText NetEventNodeGroupingHelper::ToDescription(const ENetEventGroupingMode Grou
 
 ////////////////////////////////////////////////////////////////////////////////////////////////////
 
+} // namespace UE::Insights::NetworkingProfiler
+
 #undef LOCTEXT_NAMESPACE
diff --git a/Engine/Source/Developer/TraceInsights/Private/Insights/NetworkingProfiler/ViewModels/NetEventNodeHelper.h b/Engine/Source/Developer/TraceInsights/Private/Insights/NetworkingProfiler/ViewModels/NetEventNodeHelper.h
--- a/Engine/Source/Developer/TraceInsights/Private/Insights/NetworkingProfiler/ViewModels/NetEventNodeHelper.h
+++ b/Engine/Source/Developer/TraceInsights/Private/Insights/NetworkingProfiler/ViewModels/NetEventNodeHelper.h
@@ -3,12 +3,16 @@
 #pragma once
 
 #include "CoreMinimal.h"
+
 #include "Styling/SlateBrush.h"
 
-// Insights
+// TraceInsights
 #include "Insights/NetworkingProfiler/ViewModels/NetEventNode.h"
 #include "Insights/NetworkingProfiler/ViewModels/NetEventGroupingAndSorting.h"
 
+namespace UE::Insights::NetworkingProfiler
+{
+
 ////////////////////////////////////////////////////////////////////////////////////////////////////
 
 /** Helper struct that contains helper functions and classes for ENetEventNodeType enum. */
@@ -53,3 +57,5 @@ struct NetEventNodeGroupingHelper
 };
 
 ////////////////////////////////////////////////////////////////////////////////////////////////////
+
+} // namespace UE::Insights::NetworkingProfiler
diff --git a/Engine/Source/Developer/TraceInsights/Private/Insights/NetworkingProfiler/ViewModels/NetStatsCounterGroupingAndSorting.cpp b/Engine/Source/Developer/TraceInsights/Private/Insights/NetworkingProfiler/ViewModels/NetStatsCounterGroupingAndSorting.cpp
--- a/Engine/Source/Developer/TraceInsights/Private/Insights/NetworkingProfiler/ViewModels/NetStatsCounterGroupingAndSorting.cpp
+++ b/Engine/Source/Developer/TraceInsights/Private/Insights/NetworkingProfiler/ViewModels/NetStatsCounterGroupingAndSorting.cpp
@@ -2,20 +2,26 @@
 
 #include "NetStatsCounterGroupingAndSorting.h"
 
-#include "Insights/Table/ViewModels/TableColumn.h"
+// TraceInsightsCore
+#include "InsightsCore/Table/ViewModels/TableColumn.h"
+
+// TraceInsights
 #include "Insights/NetworkingProfiler/ViewModels/NetStatsCounterNodeHelper.h"
 
-#define LOCTEXT_NAMESPACE "NetStatsCounterNode"
+#define LOCTEXT_NAMESPACE "UE::Insights::NetworkingProfiler::FNetStatsCounterNode"
 
 // Sort by name (ascending).
 #define INSIGHTS_DEFAULT_SORTING_NODES(A, B) return A->GetName().LexicalLess(B->GetName());
 
+namespace UE::Insights::NetworkingProfiler
+{
+
 ////////////////////////////////////////////////////////////////////////////////////////////////////
 // Sorting by Event Type
 ////////////////////////////////////////////////////////////////////////////////////////////////////
 
-FNetStatsCounterNodeSortingByEventType::FNetStatsCounterNodeSortingByEventType(TSharedRef<Insights::FTableColumn> InColumnRef)
-	: Insights::FTableCellValueSorter(
+FNetStatsCounterNodeSortingByEventType::FNetStatsCounterNodeSortingByEventType(TSharedRef<FTableColumn> InColumnRef)
+	: FTableCellValueSorter(
 		FName(TEXT("ByEventType")),
 		LOCTEXT("Sorting_ByEventType_Name", "By Type"),
 		LOCTEXT("Sorting_ByEventType_Title", "Sort By Type"),
@@ -26,17 +32,17 @@ FNetStatsCounterNodeSortingByEventType::FNetStatsCounterNodeSortingByEventType(T
 
 ////////////////////////////////////////////////////////////////////////////////////////////////////
 
-void FNetStatsCounterNodeSortingByEventType::Sort(TArray<Insights::FBaseTreeNodePtr>& NodesToSort, Insights::ESortMode SortMode) const
+void FNetStatsCounterNodeSortingByEventType::Sort(TArray<FBaseTreeNodePtr>& NodesToSort, ESortMode SortMode) const
 {
-	if (SortMode == Insights::ESortMode::Ascending)
+	if (SortMode == ESortMode::Ascending)
 	{
-		NodesToSort.Sort([](const Insights::FBaseTreeNodePtr& A, const Insights::FBaseTreeNodePtr& B) -> bool
+		NodesToSort.Sort([](const FBaseTreeNodePtr& A, const FBaseTreeNodePtr& B) -> bool
 		{
 			ensure(A.IsValid() && A->Is<FNetStatsCounterNode>());
-			const FNetStatsCounterNodePtr NetStatsCounterNodeA = StaticCastSharedPtr<FNetStatsCounterNode, Insights::FBaseTreeNode>(A);
+			const FNetStatsCounterNodePtr NetStatsCounterNodeA = StaticCastSharedPtr<FNetStatsCounterNode, FBaseTreeNode>(A);
 
 			ensure(B.IsValid() && B->Is<FNetStatsCounterNode>());
-			const FNetStatsCounterNodePtr NetStatsCounterNodeB = StaticCastSharedPtr<FNetStatsCounterNode, Insights::FBaseTreeNode>(B);
+			const FNetStatsCounterNodePtr NetStatsCounterNodeB = StaticCastSharedPtr<FNetStatsCounterNode, FBaseTreeNode>(B);
 
 			if (NetStatsCounterNodeA->GetType() == NetStatsCounterNodeB->GetType())
 			{
@@ -51,13 +57,13 @@ void FNetStatsCounterNodeSortingByEventType::Sort(TArray<Insights::FBaseTreeNode
 	}
 	else // if (SortMode == ESortMode::Descending)
 	{
-		NodesToSort.Sort([](const Insights::FBaseTreeNodePtr& A, const Insights::FBaseTreeNodePtr& B) -> bool
+		NodesToSort.Sort([](const FBaseTreeNodePtr& A, const FBaseTreeNodePtr& B) -> bool
 		{
 			ensure(A.IsValid() && A->Is<FNetStatsCounterNode>());
-			const FNetStatsCounterNodePtr NetStatsCounterNodeA = StaticCastSharedPtr<FNetStatsCounterNode, Insights::FBaseTreeNode>(A);
+			const FNetStatsCounterNodePtr NetStatsCounterNodeA = StaticCastSharedPtr<FNetStatsCounterNode, FBaseTreeNode>(A);
 
 			ensure(B.IsValid() && B->Is<FNetStatsCounterNode>());
-			const FNetStatsCounterNodePtr NetStatsCounterNodeB = StaticCastSharedPtr<FNetStatsCounterNode, Insights::FBaseTreeNode>(B);
+			const FNetStatsCounterNodePtr NetStatsCounterNodeB = StaticCastSharedPtr<FNetStatsCounterNode, FBaseTreeNode>(B);
 
 			if (NetStatsCounterNodeA->GetType() == NetStatsCounterNodeB->GetType())
 			{
@@ -76,8 +82,8 @@ void FNetStatsCounterNodeSortingByEventType::Sort(TArray<Insights::FBaseTreeNode
 // Sort by Sum
 ////////////////////////////////////////////////////////////////////////////////////////////////////
 
-FNetStatsCounterNodeSortingBySum::FNetStatsCounterNodeSortingBySum(TSharedRef<Insights::FTableColumn> InColumnRef)
-	: Insights::FTableCellValueSorter(
+FNetStatsCounterNodeSortingBySum::FNetStatsCounterNodeSortingBySum(TSharedRef<FTableColumn> InColumnRef)
+	: FTableCellValueSorter(
 		FName(TEXT("BySum")),
 		LOCTEXT("Sorting_BySum_Name", "By Sum"),
 		LOCTEXT("Sorting_BySum_Title", "Sort By Sum"),
@@ -88,18 +94,18 @@ FNetStatsCounterNodeSortingBySum::FNetStatsCounterNodeSortingBySum(TSharedRef<In
 
 //////////////////////////////////////////////////////////////////////////////////////////////////
 
-void FNetStatsCounterNodeSortingBySum::Sort(TArray<Insights::FBaseTreeNodePtr>& NodesToSort, Insights::ESortMode SortMode) const
+void FNetStatsCounterNodeSortingBySum::Sort(TArray<FBaseTreeNodePtr>& NodesToSort, ESortMode SortMode) const
 {
-	if (SortMode == Insights::ESortMode::Ascending)
+	if (SortMode == ESortMode::Ascending)
 	{
-		NodesToSort.Sort([](const Insights::FBaseTreeNodePtr& A, const Insights::FBaseTreeNodePtr& B) -> bool
+		NodesToSort.Sort([](const FBaseTreeNodePtr& A, const FBaseTreeNodePtr& B) -> bool
 		{
 			ensure(A.IsValid() && A->Is<FNetStatsCounterNode>());
-			const FNetStatsCounterNodePtr NetStatsCounterNodeA = StaticCastSharedPtr<FNetStatsCounterNode, Insights::FBaseTreeNode>(A);
+			const FNetStatsCounterNodePtr NetStatsCounterNodeA = StaticCastSharedPtr<FNetStatsCounterNode, FBaseTreeNode>(A);
 			const uint32 ValueA = NetStatsCounterNodeA->GetAggregatedStats().Sum;
 
 			ensure(B.IsValid() && B->Is<FNetStatsCounterNode>());
-			const FNetStatsCounterNodePtr NetStatsCounterNodeB = StaticCastSharedPtr<FNetStatsCounterNode, Insights::FBaseTreeNode>(B);
+			const FNetStatsCounterNodePtr NetStatsCounterNodeB = StaticCastSharedPtr<FNetStatsCounterNode, FBaseTreeNode>(B);
 			const uint32 ValueB = NetStatsCounterNodeB->GetAggregatedStats().Sum;
 
 			if (ValueA == ValueB)
@@ -115,14 +121,14 @@ void FNetStatsCounterNodeSortingBySum::Sort(TArray<Insights::FBaseTreeNodePtr>&
 	}
 	else // if (SortMode == ESortMode::Descending)
 	{
-		NodesToSort.Sort([](const Insights::FBaseTreeNodePtr& A, const Insights::FBaseTreeNodePtr& B) -> bool
+		NodesToSort.Sort([](const FBaseTreeNodePtr& A, const FBaseTreeNodePtr& B) -> bool
 		{
 			ensure(A.IsValid() && A->Is<FNetStatsCounterNode>());
-			const FNetStatsCounterNodePtr NetStatsCounterNodeA = StaticCastSharedPtr<FNetStatsCounterNode, Insights::FBaseTreeNode>(A);
+			const FNetStatsCounterNodePtr NetStatsCounterNodeA = StaticCastSharedPtr<FNetStatsCounterNode, FBaseTreeNode>(A);
 			const uint32 ValueA = NetStatsCounterNodeA->GetAggregatedStats().Sum;
 
 			ensure(B.IsValid() && B->Is<FNetStatsCounterNode>());
-			const FNetStatsCounterNodePtr NetStatsCounterNodeB = StaticCastSharedPtr<FNetStatsCounterNode, Insights::FBaseTreeNode>(B);
+			const FNetStatsCounterNodePtr NetStatsCounterNodeB = StaticCastSharedPtr<FNetStatsCounterNode, FBaseTreeNode>(B);
 			const uint32 ValueB = NetStatsCounterNodeB->GetAggregatedStats().Sum;
 
 			if (ValueA == ValueB)
@@ -140,5 +146,7 @@ void FNetStatsCounterNodeSortingBySum::Sort(TArray<Insights::FBaseTreeNodePtr>&
 
 ////////////////////////////////////////////////////////////////////////////////////////////////////
 
+} // namespace UE::Insights::NetworkingProfiler
+
 #undef INSIGHTS_DEFAULT_SORTING_NODES
 #undef LOCTEXT_NAMESPACE
diff --git a/Engine/Source/Developer/TraceInsights/Private/Insights/NetworkingProfiler/ViewModels/NetStatsCounterGroupingAndSorting.h b/Engine/Source/Developer/TraceInsights/Private/Insights/NetworkingProfiler/ViewModels/NetStatsCounterGroupingAndSorting.h
--- a/Engine/Source/Developer/TraceInsights/Private/Insights/NetworkingProfiler/ViewModels/NetStatsCounterGroupingAndSorting.h
+++ b/Engine/Source/Developer/TraceInsights/Private/Insights/NetworkingProfiler/ViewModels/NetStatsCounterGroupingAndSorting.h
@@ -4,32 +4,37 @@
 
 #include "CoreMinimal.h"
 
-// Insights
+// TraceInsightsCore
+#include "InsightsCore/Table/ViewModels/TableColumn.h"
+#include "InsightsCore/Table/ViewModels/TableCellValueSorter.h"
+#include "InsightsCore/Table/ViewModels/TreeNodeGrouping.h"
+
+// TraceInsights
 #include "Insights/NetworkingProfiler/ViewModels/NetStatsCounterNode.h"
-#include "Insights/Table/ViewModels/TableColumn.h"
-#include "Insights/Table/ViewModels/TableCellValueSorter.h"
-#include "Insights/Table/ViewModels/TreeNodeGrouping.h"
+
+namespace UE::Insights::NetworkingProfiler
+{
 
 ////////////////////////////////////////////////////////////////////////////////////////////////////
 // Sorters
 ////////////////////////////////////////////////////////////////////////////////////////////////////
 
-class FNetStatsCounterNodeSortingByEventType: public Insights::FTableCellValueSorter
+class FNetStatsCounterNodeSortingByEventType: public FTableCellValueSorter
 {
 public:
-	FNetStatsCounterNodeSortingByEventType(TSharedRef<Insights::FTableColumn> InColumnRef);
+	FNetStatsCounterNodeSortingByEventType(TSharedRef<FTableColumn> InColumnRef);
 
-	virtual void Sort(TArray<Insights::FBaseTreeNodePtr>& NodesToSort, Insights::ESortMode SortMode) const override;
+	virtual void Sort(TArray<FBaseTreeNodePtr>& NodesToSort, ESortMode SortMode) const override;
 };
 
 ////////////////////////////////////////////////////////////////////////////////////////////////////
 
-class FNetStatsCounterNodeSortingBySum : public Insights::FTableCellValueSorter
+class FNetStatsCounterNodeSortingBySum : public FTableCellValueSorter
 {
 public:
-	FNetStatsCounterNodeSortingBySum(TSharedRef<Insights::FTableColumn> InColumnRef);
+	FNetStatsCounterNodeSortingBySum(TSharedRef<FTableColumn> InColumnRef);
 
-	virtual void Sort(TArray<Insights::FBaseTreeNodePtr>& NodesToSort, Insights::ESortMode SortMode) const override;
+	virtual void Sort(TArray<FBaseTreeNodePtr>& NodesToSort, ESortMode SortMode) const override;
 };
 
 ////////////////////////////////////////////////////////////////////////////////////////////////////
@@ -56,3 +61,5 @@ enum class ENetStatsCounterGroupingMode
 typedef TSharedPtr<ENetStatsCounterGroupingMode> ENetStatsCounterGroupingModePtr;
 
 ////////////////////////////////////////////////////////////////////////////////////////////////////
+
+} // namespace UE::Insights::NetworkingProfiler
diff --git a/Engine/Source/Developer/TraceInsights/Private/Insights/NetworkingProfiler/ViewModels/NetStatsCounterNode.cpp b/Engine/Source/Developer/TraceInsights/Private/Insights/NetworkingProfiler/ViewModels/NetStatsCounterNode.cpp
--- a/Engine/Source/Developer/TraceInsights/Private/Insights/NetworkingProfiler/ViewModels/NetStatsCounterNode.cpp
+++ b/Engine/Source/Developer/TraceInsights/Private/Insights/NetworkingProfiler/ViewModels/NetStatsCounterNode.cpp
@@ -2,10 +2,13 @@
 
 #include "NetStatsCounterNode.h"
 
-// Insights
-#include "Insights/Common/TimeUtils.h"
+// TraceInsightsCore
+#include "InsightsCore/Common/TimeUtils.h"
 
-#define LOCTEXT_NAMESPACE "FNetStatsCounterNode"
+#define LOCTEXT_NAMESPACE "UE::Insights::NetworkingProfiler::FNetStatsCounterNode"
+
+namespace UE::Insights::NetworkingProfiler
+{
 
 INSIGHTS_IMPLEMENT_RTTI(FNetStatsCounterNode)
 
@@ -66,4 +69,6 @@ void FNetStatsCounterNode::SetAggregatedStats(const FNetProfilerAggregatedStatsC
 
 ////////////////////////////////////////////////////////////////////////////////////////////////////
 
+} // namespace UE::Insights::NetworkingProfiler
+
 #undef LOCTEXT_NAMESPACE
diff --git a/Engine/Source/Developer/TraceInsights/Private/Insights/NetworkingProfiler/ViewModels/NetStatsCounterNode.h b/Engine/Source/Developer/TraceInsights/Private/Insights/NetworkingProfiler/ViewModels/NetStatsCounterNode.h
--- a/Engine/Source/Developer/TraceInsights/Private/Insights/NetworkingProfiler/ViewModels/NetStatsCounterNode.h
+++ b/Engine/Source/Developer/TraceInsights/Private/Insights/NetworkingProfiler/ViewModels/NetStatsCounterNode.h
@@ -4,10 +4,15 @@
 
 #include "CoreMinimal.h"
 
-// Insights
-#include "Insights/Table/ViewModels/BaseTreeNode.h"
+// TraceServices
 #include "TraceServices/Model/NetProfiler.h"
 
+// TraceInsightsCore
+#include "InsightsCore/Table/ViewModels/BaseTreeNode.h"
+
+namespace UE::Insights::NetworkingProfiler
+{
+
 ////////////////////////////////////////////////////////////////////////////////////////////////////
 
 enum class ENetStatsCounterNodeType
@@ -46,7 +51,7 @@ typedef TWeakPtr<class FNetStatsCounterNode> FNetStatsCounterNodeWeak;
 /**
  * Class used to store information about a net stats counter node (used in SNetStatsCountersView).
  */
-class FNetStatsCounterNode : public Insights::FBaseTreeNode
+class FNetStatsCounterNode : public FBaseTreeNode
 {
 	INSIGHTS_DECLARE_RTTI(FNetStatsCounterNode, FBaseTreeNode)
 
@@ -111,3 +116,5 @@ private:
 };
 
 ////////////////////////////////////////////////////////////////////////////////////////////////////
+
+} // namespace UE::Insights::NetworkingProfiler
diff --git a/Engine/Source/Developer/TraceInsights/Private/Insights/NetworkingProfiler/ViewModels/NetStatsCounterNodeHelper.cpp b/Engine/Source/Developer/TraceInsights/Private/Insights/NetworkingProfiler/ViewModels/NetStatsCounterNodeHelper.cpp
--- a/Engine/Source/Developer/TraceInsights/Private/Insights/NetworkingProfiler/ViewModels/NetStatsCounterNodeHelper.cpp
+++ b/Engine/Source/Developer/TraceInsights/Private/Insights/NetworkingProfiler/ViewModels/NetStatsCounterNodeHelper.cpp
@@ -1,10 +1,14 @@
 // Copyright Epic Games, Inc. All Rights Reserved.
 
 #include "NetStatsCounterNodeHelper.h"
-// Insights
+
+// TraceInsights
 #include "Insights/InsightsStyle.h"
 
-#define LOCTEXT_NAMESPACE "NetStatsCounterNode"
+#define LOCTEXT_NAMESPACE "UE::Insights::NetworkingProfiler::FNetStatsCounterNode"
+
+namespace UE::Insights::NetworkingProfiler
+{
 
 ////////////////////////////////////////////////////////////////////////////////////////////////////
 // NetStatsCounterNode Type Helper
@@ -18,7 +22,7 @@ FText NetStatsCounterNodeTypeHelper::ToText(const ENetStatsCounterNodeType NodeT
 		case ENetStatsCounterNodeType::FrameStats:	return LOCTEXT("Stats_Name_FrameStats", "FrameStats");
 		case ENetStatsCounterNodeType::PacketStats:	return LOCTEXT("Stats_Name_PacketStats", "PacketStats");
 		case ENetStatsCounterNodeType::Group:		return LOCTEXT("Stats_Name_Group", "Group");
-		default:						return LOCTEXT("InvalidOrMax", "InvalidOrMax");
+		default:									return LOCTEXT("InvalidOrMax", "InvalidOrMax");
 	}
 }
 
@@ -32,7 +36,7 @@ FText NetStatsCounterNodeTypeHelper::ToDescription(const ENetStatsCounterNodeTyp
 		case ENetStatsCounterNodeType::FrameStats:	return LOCTEXT("Stats_Desc_FrameStats", "FrameStats node");
 		case ENetStatsCounterNodeType::PacketStats:	return LOCTEXT("Stats_Desc_PacketStats", "PacketStats node");
 		case ENetStatsCounterNodeType::Group:		return LOCTEXT("Stats_Desc_Group", "Group node");
-		default:						return LOCTEXT("InvalidOrMax", "InvalidOrMax");
+		default:									return LOCTEXT("InvalidOrMax", "InvalidOrMax");
 	}
 }
 
@@ -44,7 +48,7 @@ FName NetStatsCounterNodeTypeHelper::ToBrushName(const ENetStatsCounterNodeType
 	switch (NodeType)
 	{
 		case ENetStatsCounterNodeType::FrameStats:
-		case ENetStatsCounterNodeType::PacketStats:	
+		case ENetStatsCounterNodeType::PacketStats:
 			return TEXT("Profiler.FiltersAndPresets.StatTypeIcon");
 		case ENetStatsCounterNodeType::Group:
 			return TEXT("Profiler.Misc.GenericGroup");
@@ -86,10 +90,10 @@ FText NetStatsCounterNodeGroupingHelper::ToText(const ENetStatsCounterGroupingMo
 	static_assert(static_cast<int>(ENetStatsCounterGroupingMode::InvalidOrMax) == 3, "Not all cases are handled in switch below!?");
 	switch (GroupingMode)
 	{
-	case ENetStatsCounterGroupingMode::Flat:		return LOCTEXT("Grouping_Name_Flat", "Flat");
-	case ENetStatsCounterGroupingMode::ByName:		return LOCTEXT("Grouping_Name_ByName", "Name");
-	case ENetStatsCounterGroupingMode::ByType:		return LOCTEXT("Grouping_Name_ByType", "Event Type");
-	default:								return LOCTEXT("InvalidOrMax", "InvalidOrMax");
+		case ENetStatsCounterGroupingMode::Flat:	return LOCTEXT("Grouping_Name_Flat", "Flat");
+		case ENetStatsCounterGroupingMode::ByName:	return LOCTEXT("Grouping_Name_ByName", "Name");
+		case ENetStatsCounterGroupingMode::ByType:	return LOCTEXT("Grouping_Name_ByType", "Event Type");
+		default:									return LOCTEXT("InvalidOrMax", "InvalidOrMax");
 	}
 }
 
@@ -100,10 +104,10 @@ FText NetStatsCounterNodeGroupingHelper::ToDescription(const ENetStatsCounterGro
 	static_assert(static_cast<int>(ENetStatsCounterGroupingMode::InvalidOrMax) == 3, "Not all cases are handled in switch below!?");
 	switch (GroupingMode)
 	{
-	case ENetStatsCounterGroupingMode::Flat:		return LOCTEXT("Grouping_Desc_Flat", "Creates a single group. Includes all counters.");
-	case ENetStatsCounterGroupingMode::ByName:		return LOCTEXT("Grouping_Desc_ByName", "Creates one group for one letter.");
-	case ENetStatsCounterGroupingMode::ByType:		return LOCTEXT("Grouping_Desc_ByType", "Creates one group for each net event type.");
-	default:								return LOCTEXT("InvalidOrMax", "InvalidOrMax");
+		case ENetStatsCounterGroupingMode::Flat:	return LOCTEXT("Grouping_Desc_Flat", "Creates a single group. Includes all counters.");
+		case ENetStatsCounterGroupingMode::ByName:	return LOCTEXT("Grouping_Desc_ByName", "Creates one group for one letter.");
+		case ENetStatsCounterGroupingMode::ByType:	return LOCTEXT("Grouping_Desc_ByType", "Creates one group for each net event type.");
+		default:									return LOCTEXT("InvalidOrMax", "InvalidOrMax");
 	}
 }
 
@@ -114,13 +118,15 @@ FName NetStatsCounterNodeGroupingHelper::ToBrushName(const ENetStatsCounterGroup
 	static_assert(static_cast<int>(ENetStatsCounterGroupingMode::InvalidOrMax) == 3, "Not all cases are handled in switch below!?");
 	switch (GroupingMode)
 	{
-	case ENetStatsCounterGroupingMode::Flat:		return TEXT("Profiler.FiltersAndPresets.GroupNameIcon"); //TODO: "Icons.Grouping.Flat"
-	case ENetStatsCounterGroupingMode::ByName:		return TEXT("Profiler.FiltersAndPresets.GroupNameIcon"); //TODO: "Icons.Grouping.ByName"
-	case ENetStatsCounterGroupingMode::ByType:		return TEXT("Profiler.FiltersAndPresets.StatTypeIcon"); //TODO
-	default:								return NAME_None;
+		case ENetStatsCounterGroupingMode::Flat:	return TEXT("Profiler.FiltersAndPresets.GroupNameIcon"); //TODO: "Icons.Grouping.Flat"
+		case ENetStatsCounterGroupingMode::ByName:	return TEXT("Profiler.FiltersAndPresets.GroupNameIcon"); //TODO: "Icons.Grouping.ByName"
+		case ENetStatsCounterGroupingMode::ByType:	return TEXT("Profiler.FiltersAndPresets.StatTypeIcon"); //TODO
+		default:									return NAME_None;
 	}
 }
 
 ////////////////////////////////////////////////////////////////////////////////////////////////////
 
+} // namespace UE::Insights::NetworkingProfiler
+
 #undef LOCTEXT_NAMESPACE
diff --git a/Engine/Source/Developer/TraceInsights/Private/Insights/NetworkingProfiler/ViewModels/NetStatsCounterNodeHelper.h b/Engine/Source/Developer/TraceInsights/Private/Insights/NetworkingProfiler/ViewModels/NetStatsCounterNodeHelper.h
--- a/Engine/Source/Developer/TraceInsights/Private/Insights/NetworkingProfiler/ViewModels/NetStatsCounterNodeHelper.h
+++ b/Engine/Source/Developer/TraceInsights/Private/Insights/NetworkingProfiler/ViewModels/NetStatsCounterNodeHelper.h
@@ -3,12 +3,16 @@
 #pragma once
 
 #include "CoreMinimal.h"
+
 #include "Styling/SlateBrush.h"
 
-// Insights
+// TraceInsights
 #include "Insights/NetworkingProfiler/ViewModels/NetStatsCounterNode.h"
 #include "Insights/NetworkingProfiler/ViewModels/NetStatsCounterGroupingAndSorting.h"
 
+namespace UE::Insights::NetworkingProfiler
+{
+
 ////////////////////////////////////////////////////////////////////////////////////////////////////
 
 /** Helper struct that contains helper functions and classes for EStatsNodeType enum. */
@@ -66,3 +70,4 @@ struct NetStatsCounterNodeGroupingHelper
 
 ////////////////////////////////////////////////////////////////////////////////////////////////////
 
+} // namespace UE::Insights::NetworkingProfiler
diff --git a/Engine/Source/Developer/TraceInsights/Private/Insights/NetworkingProfiler/ViewModels/NetStatsCountersViewColumnFactory.cpp b/Engine/Source/Developer/TraceInsights/Private/Insights/NetworkingProfiler/ViewModels/NetStatsCountersViewColumnFactory.cpp
--- a/Engine/Source/Developer/TraceInsights/Private/Insights/NetworkingProfiler/ViewModels/NetStatsCountersViewColumnFactory.cpp
+++ b/Engine/Source/Developer/TraceInsights/Private/Insights/NetworkingProfiler/ViewModels/NetStatsCountersViewColumnFactory.cpp
@@ -2,14 +2,19 @@
 
 #include "NetStatsCountersViewColumnFactory.h"
 
-// Insights
-#include "Insights/Table/ViewModels/TableCellValueFormatter.h"
-#include "Insights/Table/ViewModels/TableCellValueGetter.h"
-#include "Insights/Table/ViewModels/TableCellValueSorter.h"
+// TraceInsightsCore
+#include "InsightsCore/Table/ViewModels/TableCellValueFormatter.h"
+#include "InsightsCore/Table/ViewModels/TableCellValueGetter.h"
+#include "InsightsCore/Table/ViewModels/TableCellValueSorter.h"
+
+// TraceInsights
 #include "Insights/NetworkingProfiler/ViewModels/NetStatsCounterGroupingAndSorting.h"
 #include "Insights/NetworkingProfiler/ViewModels/NetStatsCounterNodeHelper.h"
 
-#define LOCTEXT_NAMESPACE "SNetStatsCountersView"
+#define LOCTEXT_NAMESPACE "UE::Insights::NetworkingProfiler::SNetStatsCountersView"
+
+namespace UE::Insights::NetworkingProfiler
+{
 
 ////////////////////////////////////////////////////////////////////////////////////////////////////
 // Column identifiers
@@ -24,7 +29,7 @@ const FName FNetStatsCountersViewColumns::AverageCountColumnID(TEXT("Avg"));
 
 ////////////////////////////////////////////////////////////////////////////////////////////////////
 
-void FNetStatsCountersViewColumnFactory::CreateNetStatsCountersViewColumns(TArray<TSharedRef<Insights::FTableColumn>>& Columns)
+void FNetStatsCountersViewColumnFactory::CreateNetStatsCountersViewColumns(TArray<TSharedRef<FTableColumn>>& Columns)
 {
 	Columns.Reset();
 
@@ -39,10 +44,8 @@ void FNetStatsCountersViewColumnFactory::CreateNetStatsCountersViewColumns(TArra
 
 ////////////////////////////////////////////////////////////////////////////////////////////////////
 
-TSharedRef<Insights::FTableColumn> FNetStatsCountersViewColumnFactory::CreateNameColumn()
+TSharedRef<FTableColumn> FNetStatsCountersViewColumnFactory::CreateNameColumn()
 {
-	using namespace Insights;
-
 	TSharedRef<FTableColumn> ColumnRef = MakeShared<FTableColumn>(FNetStatsCountersViewColumns::NameColumnID);
 	FTableColumn& Column = *ColumnRef;
 
@@ -74,10 +77,8 @@ TSharedRef<Insights::FTableColumn> FNetStatsCountersViewColumnFactory::CreateNam
 
 ////////////////////////////////////////////////////////////////////////////////////////////////////
 
-TSharedRef<Insights::FTableColumn> FNetStatsCountersViewColumnFactory::CreateTypeColumn()
+TSharedRef<FTableColumn> FNetStatsCountersViewColumnFactory::CreateTypeColumn()
 {
-	using namespace Insights;
-
 	TSharedRef<FTableColumn> ColumnRef = MakeShared<FTableColumn>(FNetStatsCountersViewColumns::TypeColumnID);
 	FTableColumn& Column = *ColumnRef;
 
@@ -118,10 +119,8 @@ TSharedRef<Insights::FTableColumn> FNetStatsCountersViewColumnFactory::CreateTyp
 
 ////////////////////////////////////////////////////////////////////////////////////////////////////
 
-TSharedRef<Insights::FTableColumn> FNetStatsCountersViewColumnFactory::CreateInstanceCountColumn()
+TSharedRef<FTableColumn> FNetStatsCountersViewColumnFactory::CreateInstanceCountColumn()
 {
-	using namespace Insights;
-
 	TSharedRef<FTableColumn> ColumnRef = MakeShared<FTableColumn>(FNetStatsCountersViewColumns::InstanceCountColumnID);
 	FTableColumn& Column = *ColumnRef;
 
@@ -162,10 +161,8 @@ TSharedRef<Insights::FTableColumn> FNetStatsCountersViewColumnFactory::CreateIns
 // Inclusive  Columns
 ////////////////////////////////////////////////////////////////////////////////////////////////////
 
-TSharedRef<Insights::FTableColumn> FNetStatsCountersViewColumnFactory::CreateSumColumn()
+TSharedRef<FTableColumn> FNetStatsCountersViewColumnFactory::CreateSumColumn()
 {
-	using namespace Insights;
-
 	TSharedRef<FTableColumn> ColumnRef = MakeShared<FTableColumn>(FNetStatsCountersViewColumns::SumColumnID);
 	FTableColumn& Column = *ColumnRef;
 
@@ -218,10 +215,8 @@ TSharedRef<Insights::FTableColumn> FNetStatsCountersViewColumnFactory::CreateSum
 
 ////////////////////////////////////////////////////////////////////////////////////////////////////
 
-TSharedRef<Insights::FTableColumn> FNetStatsCountersViewColumnFactory::CreateMaxCountColumn()
+TSharedRef<FTableColumn> FNetStatsCountersViewColumnFactory::CreateMaxCountColumn()
 {
-	using namespace Insights;
-
 	TSharedRef<FTableColumn> ColumnRef = MakeShared<FTableColumn>(FNetStatsCountersViewColumns::MaxCountColumnID);
 	FTableColumn& Column = *ColumnRef;
 
@@ -273,10 +268,8 @@ TSharedRef<Insights::FTableColumn> FNetStatsCountersViewColumnFactory::CreateMax
 
 ////////////////////////////////////////////////////////////////////////////////////////////////////
 
-TSharedRef<Insights::FTableColumn> FNetStatsCountersViewColumnFactory::CreateAverageCountColumn()
+TSharedRef<FTableColumn> FNetStatsCountersViewColumnFactory::CreateAverageCountColumn()
 {
-	using namespace Insights;
-
 	TSharedRef<FTableColumn> ColumnRef = MakeShared<FTableColumn>(FNetStatsCountersViewColumns::AverageCountColumnID);
 	FTableColumn& Column = *ColumnRef;
 
@@ -328,4 +321,6 @@ TSharedRef<Insights::FTableColumn> FNetStatsCountersViewColumnFactory::CreateAve
 
 ////////////////////////////////////////////////////////////////////////////////////////////////////
 
+} // namespace UE::Insights::NetworkingProfiler
+
 #undef LOCTEXT_NAMESPACE
diff --git a/Engine/Source/Developer/TraceInsights/Private/Insights/NetworkingProfiler/ViewModels/NetStatsCountersViewColumnFactory.h b/Engine/Source/Developer/TraceInsights/Private/Insights/NetworkingProfiler/ViewModels/NetStatsCountersViewColumnFactory.h
--- a/Engine/Source/Developer/TraceInsights/Private/Insights/NetworkingProfiler/ViewModels/NetStatsCountersViewColumnFactory.h
+++ b/Engine/Source/Developer/TraceInsights/Private/Insights/NetworkingProfiler/ViewModels/NetStatsCountersViewColumnFactory.h
@@ -2,49 +2,43 @@
 
 #pragma once
 
-#include "CoreMinimal.h"
+#include "CoreTypes.h"
 
-////////////////////////////////////////////////////////////////////////////////////////////////////
+#include "Containers/Array.h"
+#include "Templates/SharedPointer.h"
+#include "UObject/NameTypes.h"
 
-struct FNetStatsCountersViewColumns
+namespace UE::Insights { class FTableColumn; }
+
+namespace UE::Insights::NetworkingProfiler
 {
-	//////////////////////////////////////////////////
-	// Column identifiers
 
+// Column identifiers
+struct FNetStatsCountersViewColumns
+{
 	static const FName NameColumnID;
 	static const FName TypeColumnID;
 	static const FName InstanceCountColumnID;
-
 	static const FName SumColumnID;
 	static const FName MaxCountColumnID;
 	static const FName AverageCountColumnID;
 };
 
-////////////////////////////////////////////////////////////////////////////////////////////////////
-
-namespace Insights
-{
-	class FTableColumn;
-}
-
-////////////////////////////////////////////////////////////////////////////////////////////////////
-
 struct FNetStatsCountersViewColumnFactory
 {
 public:
-	static void CreateNetStatsCountersViewColumns(TArray<TSharedRef<Insights::FTableColumn>>& Columns);
-
-	static TSharedRef<Insights::FTableColumn> CreateNameColumn();
-	static TSharedRef<Insights::FTableColumn> CreateTypeColumn();
-	static TSharedRef<Insights::FTableColumn> CreateInstanceCountColumn();
+	static void CreateNetStatsCountersViewColumns(TArray<TSharedRef<FTableColumn>>& Columns);
 
-	static TSharedRef<Insights::FTableColumn> CreateSumColumn();
-	static TSharedRef<Insights::FTableColumn> CreateMaxCountColumn();
-	static TSharedRef<Insights::FTableColumn> CreateAverageCountColumn();
+	static TSharedRef<FTableColumn> CreateNameColumn();
+	static TSharedRef<FTableColumn> CreateTypeColumn();
+	static TSharedRef<FTableColumn> CreateInstanceCountColumn();
+	static TSharedRef<FTableColumn> CreateSumColumn();
+	static TSharedRef<FTableColumn> CreateMaxCountColumn();
+	static TSharedRef<FTableColumn> CreateAverageCountColumn();
 
 private:
 	static constexpr float SumColumnInitialWidth = 60.0f;
 	static constexpr float CountColumnInitialWidth = 50.0f;
 };
 
-////////////////////////////////////////////////////////////////////////////////////////////////////
+} // namespace UE::Insights::NetworkingProfiler
diff --git a/Engine/Source/Developer/TraceInsights/Private/Insights/NetworkingProfiler/ViewModels/NetStatsViewColumnFactory.cpp b/Engine/Source/Developer/TraceInsights/Private/Insights/NetworkingProfiler/ViewModels/NetStatsViewColumnFactory.cpp
--- a/Engine/Source/Developer/TraceInsights/Private/Insights/NetworkingProfiler/ViewModels/NetStatsViewColumnFactory.cpp
+++ b/Engine/Source/Developer/TraceInsights/Private/Insights/NetworkingProfiler/ViewModels/NetStatsViewColumnFactory.cpp
@@ -2,14 +2,19 @@
 
 #include "NetStatsViewColumnFactory.h"
 
-// Insights
-#include "Insights/Table/ViewModels/TableCellValueFormatter.h"
-#include "Insights/Table/ViewModels/TableCellValueGetter.h"
-#include "Insights/Table/ViewModels/TableCellValueSorter.h"
+// TraceInsightsCore
+#include "InsightsCore/Table/ViewModels/TableCellValueFormatter.h"
+#include "InsightsCore/Table/ViewModels/TableCellValueGetter.h"
+#include "InsightsCore/Table/ViewModels/TableCellValueSorter.h"
+
+// TraceInsights
 #include "Insights/NetworkingProfiler/ViewModels/NetEventGroupingAndSorting.h"
 #include "Insights/NetworkingProfiler/ViewModels/NetEventNodeHelper.h"
 
-#define LOCTEXT_NAMESPACE "SNetStatsView"
+#define LOCTEXT_NAMESPACE "UE::Insights::NetworkingProfiler::SNetStatsView"
+
+namespace UE::Insights::NetworkingProfiler
+{
 
 ////////////////////////////////////////////////////////////////////////////////////////////////////
 // Column identifiers
@@ -31,7 +36,7 @@ const FName FNetStatsViewColumns::AverageExclusiveSizeColumnID(TEXT("AverageExcl
 
 ////////////////////////////////////////////////////////////////////////////////////////////////////
 
-void FNetStatsViewColumnFactory::CreateNetStatsViewColumns(TArray<TSharedRef<Insights::FTableColumn>>& Columns)
+void FNetStatsViewColumnFactory::CreateNetStatsViewColumns(TArray<TSharedRef<FTableColumn>>& Columns)
 {
 	Columns.Reset();
 
@@ -51,10 +56,8 @@ void FNetStatsViewColumnFactory::CreateNetStatsViewColumns(TArray<TSharedRef<Ins
 
 ////////////////////////////////////////////////////////////////////////////////////////////////////
 
-TSharedRef<Insights::FTableColumn> FNetStatsViewColumnFactory::CreateNameColumn()
+TSharedRef<FTableColumn> FNetStatsViewColumnFactory::CreateNameColumn()
 {
-	using namespace Insights;
-
 	TSharedRef<FTableColumn> ColumnRef = MakeShared<FTableColumn>(FNetStatsViewColumns::NameColumnID);
 	FTableColumn& Column = *ColumnRef;
 
@@ -86,10 +89,8 @@ TSharedRef<Insights::FTableColumn> FNetStatsViewColumnFactory::CreateNameColumn(
 
 ////////////////////////////////////////////////////////////////////////////////////////////////////
 
-TSharedRef<Insights::FTableColumn> FNetStatsViewColumnFactory::CreateTypeColumn()
+TSharedRef<FTableColumn> FNetStatsViewColumnFactory::CreateTypeColumn()
 {
-	using namespace Insights;
-
 	TSharedRef<FTableColumn> ColumnRef = MakeShared<FTableColumn>(FNetStatsViewColumns::TypeColumnID);
 	FTableColumn& Column = *ColumnRef;
 
@@ -132,10 +133,8 @@ TSharedRef<Insights::FTableColumn> FNetStatsViewColumnFactory::CreateTypeColumn(
 
 ////////////////////////////////////////////////////////////////////////////////////////////////////
 
-TSharedRef<Insights::FTableColumn> FNetStatsViewColumnFactory::CreateLevelColumn()
+TSharedRef<FTableColumn> FNetStatsViewColumnFactory::CreateLevelColumn()
 {
-	using namespace Insights;
-
 	TSharedRef<FTableColumn> ColumnRef = MakeShared<FTableColumn>(FNetStatsViewColumns::LevelColumnID);
 	FTableColumn& Column = *ColumnRef;
 
@@ -177,10 +176,8 @@ TSharedRef<Insights::FTableColumn> FNetStatsViewColumnFactory::CreateLevelColumn
 
 ////////////////////////////////////////////////////////////////////////////////////////////////////
 
-TSharedRef<Insights::FTableColumn> FNetStatsViewColumnFactory::CreateInstanceCountColumn()
+TSharedRef<FTableColumn> FNetStatsViewColumnFactory::CreateInstanceCountColumn()
 {
-	using namespace Insights;
-
 	TSharedRef<FTableColumn> ColumnRef = MakeShared<FTableColumn>(FNetStatsViewColumns::InstanceCountColumnID);
 	FTableColumn& Column = *ColumnRef;
 
@@ -226,10 +223,8 @@ TSharedRef<Insights::FTableColumn> FNetStatsViewColumnFactory::CreateInstanceCou
 // Inclusive  Columns
 ////////////////////////////////////////////////////////////////////////////////////////////////////
 
-TSharedRef<Insights::FTableColumn> FNetStatsViewColumnFactory::CreateTotalInclusiveSizeColumn()
+TSharedRef<FTableColumn> FNetStatsViewColumnFactory::CreateTotalInclusiveSizeColumn()
 {
-	using namespace Insights;
-
 	TSharedRef<FTableColumn> ColumnRef = MakeShared<FTableColumn>(FNetStatsViewColumns::TotalInclusiveSizeColumnID);
 	FTableColumn& Column = *ColumnRef;
 
@@ -274,10 +269,8 @@ TSharedRef<Insights::FTableColumn> FNetStatsViewColumnFactory::CreateTotalInclus
 
 ////////////////////////////////////////////////////////////////////////////////////////////////////
 
-TSharedRef<Insights::FTableColumn> FNetStatsViewColumnFactory::CreateMaxInclusiveSizeColumn()
+TSharedRef<FTableColumn> FNetStatsViewColumnFactory::CreateMaxInclusiveSizeColumn()
 {
-	using namespace Insights;
-
 	TSharedRef<FTableColumn> ColumnRef = MakeShared<FTableColumn>(FNetStatsViewColumns::MaxInclusiveSizeColumnID);
 	FTableColumn& Column = *ColumnRef;
 
@@ -320,10 +313,8 @@ TSharedRef<Insights::FTableColumn> FNetStatsViewColumnFactory::CreateMaxInclusiv
 
 ////////////////////////////////////////////////////////////////////////////////////////////////////
 
-TSharedRef<Insights::FTableColumn> FNetStatsViewColumnFactory::CreateAverageInclusiveSizeColumn()
+TSharedRef<FTableColumn> FNetStatsViewColumnFactory::CreateAverageInclusiveSizeColumn()
 {
-	using namespace Insights;
-
 	TSharedRef<FTableColumn> ColumnRef = MakeShared<FTableColumn>(FNetStatsViewColumns::AverageInclusiveSizeColumnID);
 	FTableColumn& Column = *ColumnRef;
 
@@ -368,10 +359,8 @@ TSharedRef<Insights::FTableColumn> FNetStatsViewColumnFactory::CreateAverageIncl
 // Exclusive  Columns
 ////////////////////////////////////////////////////////////////////////////////////////////////////
 
-TSharedRef<Insights::FTableColumn> FNetStatsViewColumnFactory::CreateTotalExclusiveSizeColumn()
+TSharedRef<FTableColumn> FNetStatsViewColumnFactory::CreateTotalExclusiveSizeColumn()
 {
-	using namespace Insights;
-
 	TSharedRef<FTableColumn> ColumnRef = MakeShared<FTableColumn>(FNetStatsViewColumns::TotalExclusiveSizeColumnID);
 	FTableColumn& Column = *ColumnRef;
 
@@ -416,10 +405,8 @@ TSharedRef<Insights::FTableColumn> FNetStatsViewColumnFactory::CreateTotalExclus
 
 ////////////////////////////////////////////////////////////////////////////////////////////////////
 
-TSharedRef<Insights::FTableColumn> FNetStatsViewColumnFactory::CreateMaxExclusiveSizeColumn()
+TSharedRef<FTableColumn> FNetStatsViewColumnFactory::CreateMaxExclusiveSizeColumn()
 {
-	using namespace Insights;
-
 	TSharedRef<FTableColumn> ColumnRef = MakeShared<FTableColumn>(FNetStatsViewColumns::MaxExclusiveSizeColumnID);
 	FTableColumn& Column = *ColumnRef;
 
@@ -462,10 +449,8 @@ TSharedRef<Insights::FTableColumn> FNetStatsViewColumnFactory::CreateMaxExclusiv
 
 ////////////////////////////////////////////////////////////////////////////////////////////////////
 
-TSharedRef<Insights::FTableColumn> FNetStatsViewColumnFactory::CreateAverageExclusiveSizeColumn()
+TSharedRef<FTableColumn> FNetStatsViewColumnFactory::CreateAverageExclusiveSizeColumn()
 {
-	using namespace Insights;
-
 	TSharedRef<FTableColumn> ColumnRef = MakeShared<FTableColumn>(FNetStatsViewColumns::AverageExclusiveSizeColumnID);
 	FTableColumn& Column = *ColumnRef;
 
@@ -509,4 +494,6 @@ TSharedRef<Insights::FTableColumn> FNetStatsViewColumnFactory::CreateAverageExcl
 
 ////////////////////////////////////////////////////////////////////////////////////////////////////
 
+} // namespace UE::Insights::NetworkingProfiler
+
 #undef LOCTEXT_NAMESPACE
diff --git a/Engine/Source/Developer/TraceInsights/Private/Insights/NetworkingProfiler/ViewModels/NetStatsViewColumnFactory.h b/Engine/Source/Developer/TraceInsights/Private/Insights/NetworkingProfiler/ViewModels/NetStatsViewColumnFactory.h
--- a/Engine/Source/Developer/TraceInsights/Private/Insights/NetworkingProfiler/ViewModels/NetStatsViewColumnFactory.h
+++ b/Engine/Source/Developer/TraceInsights/Private/Insights/NetworkingProfiler/ViewModels/NetStatsViewColumnFactory.h
@@ -2,15 +2,19 @@
 
 #pragma once
 
-#include "CoreMinimal.h"
+#include "CoreTypes.h"
+#include "Containers/Array.h"
+#include "Templates/SharedPointer.h"
+#include "UObject/NameTypes.h"
 
-////////////////////////////////////////////////////////////////////////////////////////////////////
+namespace UE::Insights { class FTableColumn; }
 
-struct FNetStatsViewColumns
+namespace UE::Insights::NetworkingProfiler
 {
-	//////////////////////////////////////////////////
-	// Column identifiers
 
+// Column identifiers
+struct FNetStatsViewColumns
+{
 	static const FName NameColumnID;
 	static const FName TypeColumnID;
 	static const FName LevelColumnID;
@@ -27,36 +31,27 @@ struct FNetStatsViewColumns
 	static const FName AverageExclusiveSizeColumnID;
 };
 
-////////////////////////////////////////////////////////////////////////////////////////////////////
-
-namespace Insights
-{
-	class FTableColumn;
-}
-
-////////////////////////////////////////////////////////////////////////////////////////////////////
-
 struct FNetStatsViewColumnFactory
 {
 public:
-	static void CreateNetStatsViewColumns(TArray<TSharedRef<Insights::FTableColumn>>& Columns);
+	static void CreateNetStatsViewColumns(TArray<TSharedRef<FTableColumn>>& Columns);
 
-	static TSharedRef<Insights::FTableColumn> CreateNameColumn();
-	static TSharedRef<Insights::FTableColumn> CreateTypeColumn();
-	static TSharedRef<Insights::FTableColumn> CreateLevelColumn();
-	static TSharedRef<Insights::FTableColumn> CreateInstanceCountColumn();
+	static TSharedRef<FTableColumn> CreateNameColumn();
+	static TSharedRef<FTableColumn> CreateTypeColumn();
+	static TSharedRef<FTableColumn> CreateLevelColumn();
+	static TSharedRef<FTableColumn> CreateInstanceCountColumn();
 
-	static TSharedRef<Insights::FTableColumn> CreateTotalInclusiveSizeColumn();
-	static TSharedRef<Insights::FTableColumn> CreateMaxInclusiveSizeColumn();
-	static TSharedRef<Insights::FTableColumn> CreateAverageInclusiveSizeColumn();
+	static TSharedRef<FTableColumn> CreateTotalInclusiveSizeColumn();
+	static TSharedRef<FTableColumn> CreateMaxInclusiveSizeColumn();
+	static TSharedRef<FTableColumn> CreateAverageInclusiveSizeColumn();
 
-	static TSharedRef<Insights::FTableColumn> CreateTotalExclusiveSizeColumn();
-	static TSharedRef<Insights::FTableColumn> CreateMaxExclusiveSizeColumn();
-	static TSharedRef<Insights::FTableColumn> CreateAverageExclusiveSizeColumn();
+	static TSharedRef<FTableColumn> CreateTotalExclusiveSizeColumn();
+	static TSharedRef<FTableColumn> CreateMaxExclusiveSizeColumn();
+	static TSharedRef<FTableColumn> CreateAverageExclusiveSizeColumn();
 
 private:
 	static constexpr float TotalSizeColumnInitialWidth = 60.0f;
 	static constexpr float SizeColumnInitialWidth = 50.0f;
 };
 
-////////////////////////////////////////////////////////////////////////////////////////////////////
+} // namespace UE::Insights::NetworkingProfiler
diff --git a/Engine/Source/Developer/TraceInsights/Private/Insights/NetworkingProfiler/ViewModels/PacketContentViewDrawHelper.cpp b/Engine/Source/Developer/TraceInsights/Private/Insights/NetworkingProfiler/ViewModels/PacketContentViewDrawHelper.cpp
--- a/Engine/Source/Developer/TraceInsights/Private/Insights/NetworkingProfiler/ViewModels/PacketContentViewDrawHelper.cpp
+++ b/Engine/Source/Developer/TraceInsights/Private/Insights/NetworkingProfiler/ViewModels/PacketContentViewDrawHelper.cpp
@@ -4,20 +4,25 @@
 
 #include "Fonts/FontMeasure.h"
 #include "Framework/Application/SlateApplication.h"
+#include "Misc/StringBuilder.h"
 #include "Rendering/DrawElements.h"
 #include "Styling/AppStyle.h"
 
-// Insights
-#include "Insights/Common/PaintUtils.h"
+// TraceInsightsCore
+#include "InsightsCore/Common/PaintUtils.h"
+
+// TraceInsights
 #include "Insights/InsightsStyle.h"
 #include "Insights/NetworkingProfiler/ViewModels/PacketContentViewport.h"
 #include "Insights/ViewModels/DrawHelpers.h"
 
 #include <limits>
-#include "Misc/StringBuilder.h"
 
 #define INSIGHTS_USE_LEGACY_BORDER 0
 
+namespace UE::Insights::NetworkingProfiler
+{
+
 ////////////////////////////////////////////////////////////////////////////////////////////////////
 // FPacketContentViewDrawStateBuilder
 ////////////////////////////////////////////////////////////////////////////////////////////////////
@@ -534,4 +539,6 @@ void FPacketContentViewDrawHelper::DrawEventHighlight(const FNetworkPacketEvent&
 
 ////////////////////////////////////////////////////////////////////////////////////////////////////
 
+} // namespace UE::Insights::NetworkingProfiler
+
 #undef INSIGHTS_USE_LEGACY_BORDER
diff --git a/Engine/Source/Developer/TraceInsights/Private/Insights/NetworkingProfiler/ViewModels/PacketContentViewDrawHelper.h b/Engine/Source/Developer/TraceInsights/Private/Insights/NetworkingProfiler/ViewModels/PacketContentViewDrawHelper.h
--- a/Engine/Source/Developer/TraceInsights/Private/Insights/NetworkingProfiler/ViewModels/PacketContentViewDrawHelper.h
+++ b/Engine/Source/Developer/TraceInsights/Private/Insights/NetworkingProfiler/ViewModels/PacketContentViewDrawHelper.h
@@ -3,18 +3,21 @@
 #pragma once
 
 #include "CoreMinimal.h"
+
 #include "Fonts/SlateFontInfo.h"
-#include "TraceServices/Model/NetProfiler.h"
 #include "Styling/WidgetStyle.h"
 
-enum class ESlateDrawEffect : uint8;
+// TraceServices
+#include "TraceServices/Model/NetProfiler.h"
 
-struct FDrawContext;
-struct FGeometry;
 struct FSlateBrush;
 
+namespace UE::Insights { class FDrawContext; }
+
+namespace UE::Insights::NetworkingProfiler
+{
+
 class FPacketContentViewport;
-class FSlateWindowElementList;
 
 ////////////////////////////////////////////////////////////////////////////////////////////////////
 
@@ -230,3 +233,5 @@ private:
 };
 
 ////////////////////////////////////////////////////////////////////////////////////////////////////
+
+} // namespace UE::Insights::NetworkingProfiler
diff --git a/Engine/Source/Developer/TraceInsights/Private/Insights/NetworkingProfiler/ViewModels/PacketContentViewport.h b/Engine/Source/Developer/TraceInsights/Private/Insights/NetworkingProfiler/ViewModels/PacketContentViewport.h
--- a/Engine/Source/Developer/TraceInsights/Private/Insights/NetworkingProfiler/ViewModels/PacketContentViewport.h
+++ b/Engine/Source/Developer/TraceInsights/Private/Insights/NetworkingProfiler/ViewModels/PacketContentViewport.h
@@ -4,9 +4,12 @@
 
 #include "CoreMinimal.h"
 
-// Insights
+// TraceInsights
 #include "Insights/ViewModels/AxisViewportDouble.h"
 
+namespace UE::Insights::NetworkingProfiler
+{
+
 ////////////////////////////////////////////////////////////////////////////////////////////////////
 
 class FPacketContentViewport
@@ -67,3 +70,5 @@ private:
 };
 
 ////////////////////////////////////////////////////////////////////////////////////////////////////
+
+} // namespace UE::Insights::NetworkingProfiler
diff --git a/Engine/Source/Developer/TraceInsights/Private/Insights/NetworkingProfiler/ViewModels/PacketViewDrawHelper.cpp b/Engine/Source/Developer/TraceInsights/Private/Insights/NetworkingProfiler/ViewModels/PacketViewDrawHelper.cpp
--- a/Engine/Source/Developer/TraceInsights/Private/Insights/NetworkingProfiler/ViewModels/PacketViewDrawHelper.cpp
+++ b/Engine/Source/Developer/TraceInsights/Private/Insights/NetworkingProfiler/ViewModels/PacketViewDrawHelper.cpp
@@ -5,8 +5,10 @@
 #include "Framework/Application/SlateApplication.h"
 #include "Rendering/DrawElements.h"
 
-// Insights
-#include "Insights/Common/PaintUtils.h"
+// TraceInsightsCore
+#include "InsightsCore/Common/PaintUtils.h"
+
+// TraceInsights
 #include "Insights/InsightsStyle.h"
 #include "Insights/NetworkingProfiler/ViewModels/PacketViewport.h"
 #include "Insights/NetworkingProfiler/ViewModels/PacketContentViewDrawHelper.h"
@@ -16,6 +18,9 @@
 
 #define INSIGHTS_USE_LEGACY_BORDER 0
 
+namespace UE::Insights::NetworkingProfiler
+{
+
 ////////////////////////////////////////////////////////////////////////////////////////////////////
 // FNetworkPacketAggregatedSample
 ////////////////////////////////////////////////////////////////////////////////////////////////////
@@ -349,4 +354,6 @@ void FPacketViewDrawHelper::DrawSelection(int32 StartPacketIndex, int32 EndPacke
 
 ////////////////////////////////////////////////////////////////////////////////////////////////////
 
+} // namespace UE::Insights::NetworkingProfiler
+
 #undef INSIGHTS_USE_LEGACY_BORDER
diff --git a/Engine/Source/Developer/TraceInsights/Private/Insights/NetworkingProfiler/ViewModels/PacketViewDrawHelper.h b/Engine/Source/Developer/TraceInsights/Private/Insights/NetworkingProfiler/ViewModels/PacketViewDrawHelper.h
--- a/Engine/Source/Developer/TraceInsights/Private/Insights/NetworkingProfiler/ViewModels/PacketViewDrawHelper.h
+++ b/Engine/Source/Developer/TraceInsights/Private/Insights/NetworkingProfiler/ViewModels/PacketViewDrawHelper.h
@@ -3,20 +3,23 @@
 #pragma once
 
 #include "CoreMinimal.h"
+
 #include "Fonts/SlateFontInfo.h"
 #include "Styling/WidgetStyle.h"
+
+// TraceServices
 #include "TraceServices/Model/NetProfiler.h"
 
 #include <limits>
 
-enum class ESlateDrawEffect : uint8;
-
-struct FDrawContext;
-struct FGeometry;
 struct FSlateBrush;
 
+namespace UE::Insights { class FDrawContext; }
+
+namespace UE::Insights::NetworkingProfiler
+{
+
 class FPacketViewport;
-class FSlateWindowElementList;
 
 ////////////////////////////////////////////////////////////////////////////////////////////////////
 
@@ -222,3 +225,5 @@ private:
 };
 
 ////////////////////////////////////////////////////////////////////////////////////////////////////
+
+} // namespace UE::Insights::NetworkingProfiler
diff --git a/Engine/Source/Developer/TraceInsights/Private/Insights/NetworkingProfiler/ViewModels/PacketViewport.h b/Engine/Source/Developer/TraceInsights/Private/Insights/NetworkingProfiler/ViewModels/PacketViewport.h
--- a/Engine/Source/Developer/TraceInsights/Private/Insights/NetworkingProfiler/ViewModels/PacketViewport.h
+++ b/Engine/Source/Developer/TraceInsights/Private/Insights/NetworkingProfiler/ViewModels/PacketViewport.h
@@ -4,10 +4,13 @@
 
 #include "CoreMinimal.h"
 
-// Insights
+// TraceInsights
 #include "Insights/ViewModels/AxisViewportInt32.h"
 #include "Insights/ViewModels/AxisViewportDouble.h"
 
+namespace UE::Insights::NetworkingProfiler
+{
+
 ////////////////////////////////////////////////////////////////////////////////////////////////////
 
 class FPacketViewport
@@ -63,3 +66,5 @@ private:
 };
 
 ////////////////////////////////////////////////////////////////////////////////////////////////////
+
+} // namespace UE::Insights::NetworkingProfiler
diff --git a/Engine/Source/Developer/TraceInsights/Private/Insights/NetworkingProfiler/Widgets/SNetStatsCountersTableCell.cpp b/Engine/Source/Developer/TraceInsights/Private/Insights/NetworkingProfiler/Widgets/SNetStatsCountersTableCell.cpp
--- a/Engine/Source/Developer/TraceInsights/Private/Insights/NetworkingProfiler/Widgets/SNetStatsCountersTableCell.cpp
+++ b/Engine/Source/Developer/TraceInsights/Private/Insights/NetworkingProfiler/Widgets/SNetStatsCountersTableCell.cpp
@@ -11,13 +11,18 @@
 #include "Widgets/Text/STextBlock.h"
 #include "Widgets/Views/SExpanderArrow.h"
 
-// Insights
+// TraceInsightsCore
+#include "InsightsCore/Table/ViewModels/Table.h"
+#include "InsightsCore/Table/ViewModels/TableColumn.h"
+
+// TraceInsights
 #include "Insights/InsightsStyle.h"
-#include "Insights/Table/ViewModels/Table.h"
-#include "Insights/Table/ViewModels/TableColumn.h"
 #include "Insights/NetworkingProfiler/Widgets/SNetStatsCountersTableRow.h"
 
-#define LOCTEXT_NAMESPACE "SNetStatsCountersView"
+#define LOCTEXT_NAMESPACE "UE::Insights::NetworkingProfiler::SNetStatsCountersView"
+
+namespace UE::Insights::NetworkingProfiler
+{
 
 ////////////////////////////////////////////////////////////////////////////////////////////////////
 
@@ -156,4 +161,6 @@ END_SLATE_FUNCTION_BUILD_OPTIMIZATION
 
 ////////////////////////////////////////////////////////////////////////////////////////////////////
 
+} // namespace UE::Insights::NetworkingProfiler
+
 #undef LOCTEXT_NAMESPACE
diff --git a/Engine/Source/Developer/TraceInsights/Private/Insights/NetworkingProfiler/Widgets/SNetStatsCountersTableCell.h b/Engine/Source/Developer/TraceInsights/Private/Insights/NetworkingProfiler/Widgets/SNetStatsCountersTableCell.h
--- a/Engine/Source/Developer/TraceInsights/Private/Insights/NetworkingProfiler/Widgets/SNetStatsCountersTableCell.h
+++ b/Engine/Source/Developer/TraceInsights/Private/Insights/NetworkingProfiler/Widgets/SNetStatsCountersTableCell.h
@@ -3,21 +3,25 @@
 #pragma once
 
 #include "CoreMinimal.h"
+
 #include "Widgets/DeclarativeSyntaxSupport.h"
 #include "Widgets/SCompoundWidget.h"
 
-// Insights
+// TraceInsights
 #include "Insights/NetworkingProfiler/ViewModels/NetStatsCounterNode.h"
 
 class ITableRow;
 
-namespace Insights
+namespace UE::Insights
 {
 	class FTable;
 	class FTableColumn;
 }
 
-DECLARE_DELEGATE_ThreeParams(FSetHoveredNetStatsCountersTableCell, TSharedPtr<Insights::FTable> /*TablePtr*/, TSharedPtr<Insights::FTableColumn> /*ColumnPtr*/, FNetStatsCounterNodePtr /*NetStatsCounterNodePtr*/);
+namespace UE::Insights::NetworkingProfiler
+{
+
+DECLARE_DELEGATE_ThreeParams(FSetHoveredNetStatsCountersTableCell, TSharedPtr<FTable> /*TablePtr*/, TSharedPtr<FTableColumn> /*ColumnPtr*/, FNetStatsCounterNodePtr /*NetStatsCounterNodePtr*/);
 
 class SNetStatsCountersTableCell : public SCompoundWidget
 {
@@ -25,8 +29,8 @@ public:
 	SLATE_BEGIN_ARGS(SNetStatsCountersTableCell) {}
 		SLATE_EVENT(FSetHoveredNetStatsCountersTableCell, OnSetHoveredCell)
 		SLATE_ATTRIBUTE(FText, HighlightText)
-		SLATE_ARGUMENT(TSharedPtr<Insights::FTable>, TablePtr)
-		SLATE_ARGUMENT(TSharedPtr<Insights::FTableColumn>, ColumnPtr)
+		SLATE_ARGUMENT(TSharedPtr<FTable>, TablePtr)
+		SLATE_ARGUMENT(TSharedPtr<FTableColumn>, ColumnPtr)
 		SLATE_ARGUMENT(FNetStatsCounterNodePtr, NetStatsCounterNodePtr)
 		SLATE_ARGUMENT(bool, IsNameColumn)
 	SLATE_END_ARGS()
@@ -159,13 +163,15 @@ protected:
 
 protected:
 	/** A shared pointer to the table view model. */
-	TSharedPtr<Insights::FTable> TablePtr; // TODO: TSharedRef
+	TSharedPtr<FTable> TablePtr; // TODO: TSharedRef
 
 	/** A shared pointer to the table column view model. */
-	TSharedPtr<Insights::FTableColumn> ColumnPtr; // TODO: TSharedRef
+	TSharedPtr<FTableColumn> ColumnPtr; // TODO: TSharedRef
 
 	/** A shared pointer to the timer node. */
 	FNetStatsCounterNodePtr NetStatsCounterNodePtr; // TODO: TSharedRef
 
 	FSetHoveredNetStatsCountersTableCell SetHoveredCellDelegate;
 };
+
+} // namespace UE::Insights::NetworkingProfiler
diff --git a/Engine/Source/Developer/TraceInsights/Private/Insights/NetworkingProfiler/Widgets/SNetStatsCountersTableRow.cpp b/Engine/Source/Developer/TraceInsights/Private/Insights/NetworkingProfiler/Widgets/SNetStatsCountersTableRow.cpp
--- a/Engine/Source/Developer/TraceInsights/Private/Insights/NetworkingProfiler/Widgets/SNetStatsCountersTableRow.cpp
+++ b/Engine/Source/Developer/TraceInsights/Private/Insights/NetworkingProfiler/Widgets/SNetStatsCountersTableRow.cpp
@@ -6,15 +6,20 @@
 #include "Widgets/SOverlay.h"
 #include "Widgets/Images/SImage.h"
 
-// Insights
+// TraceInsightsCore
+#include "InsightsCore/Common/TimeUtils.h"
+#include "InsightsCore/Table/ViewModels/Table.h"
+#include "InsightsCore/Table/ViewModels/TableColumn.h"
+
+// TraceInsights
 #include "Insights/InsightsStyle.h"
-#include "Insights/Common/TimeUtils.h"
-#include "Insights/Table/ViewModels/Table.h"
-#include "Insights/Table/ViewModels/TableColumn.h"
 #include "Insights/NetworkingProfiler/Widgets/SNetStatsCountersTableCell.h"
 #include "Insights/NetworkingProfiler/Widgets/SNetStatsCountersViewTooltip.h"
 
-#define LOCTEXT_NAMESPACE "SNetStatsView"
+#define LOCTEXT_NAMESPACE "UE::Insights::NetworkingProfiler::SNetStatsView"
+
+namespace UE::Insights::NetworkingProfiler
+{
 
 ////////////////////////////////////////////////////////////////////////////////////////////////////
 
@@ -44,7 +49,7 @@ void SNetStatsCountersTableRow::Construct(const FArguments& InArgs, const TShare
 
 TSharedRef<SWidget> SNetStatsCountersTableRow::GenerateWidgetForColumn(const FName& ColumnId)
 {
-	TSharedPtr<Insights::FTableColumn> ColumnPtr = TablePtr->FindColumnChecked(ColumnId);
+	TSharedPtr<FTableColumn> ColumnPtr = TablePtr->FindColumnChecked(ColumnId);
 
 	return
 		SNew(SOverlay)
@@ -181,11 +186,13 @@ EVisibility SNetStatsCountersTableRow::IsColumnVisible(const FName ColumnId) con
 
 ////////////////////////////////////////////////////////////////////////////////////////////////////
 
-void SNetStatsCountersTableRow::OnSetHoveredCell(TSharedPtr<Insights::FTable> InTablePtr, TSharedPtr<Insights::FTableColumn> InColumnPtr, FNetStatsCounterNodePtr InNetStatsCounterNodePtr)
+void SNetStatsCountersTableRow::OnSetHoveredCell(TSharedPtr<FTable> InTablePtr, TSharedPtr<FTableColumn> InColumnPtr, FNetStatsCounterNodePtr InNetStatsCounterNodePtr)
 {
 	SetHoveredCellDelegate.ExecuteIfBound(InTablePtr, InColumnPtr, InNetStatsCounterNodePtr);
 }
 
 ////////////////////////////////////////////////////////////////////////////////////////////////////
 
+} // namespace UE::Insights::NetworkingProfiler
+
 #undef LOCTEXT_NAMESPACE
diff --git a/Engine/Source/Developer/TraceInsights/Private/Insights/NetworkingProfiler/Widgets/SNetStatsCountersTableRow.h b/Engine/Source/Developer/TraceInsights/Private/Insights/NetworkingProfiler/Widgets/SNetStatsCountersTableRow.h
--- a/Engine/Source/Developer/TraceInsights/Private/Insights/NetworkingProfiler/Widgets/SNetStatsCountersTableRow.h
+++ b/Engine/Source/Developer/TraceInsights/Private/Insights/NetworkingProfiler/Widgets/SNetStatsCountersTableRow.h
@@ -3,28 +3,33 @@
 #pragma once
 
 #include "CoreMinimal.h"
+
 #include "Widgets/DeclarativeSyntaxSupport.h"
 #include "Widgets/Views/SHeaderRow.h"
 #include "Widgets/Views/STableViewBase.h"
 #include "Widgets/Views/STableRow.h"
 #include "Widgets/Views/STreeView.h"
 
-// Insights
+// TraceInsights
 #include "Insights/NetworkingProfiler/ViewModels/NetStatsCounterNodeHelper.h"
 
 class IToolTip;
-class SNetStatsCounterTableRowToolTip;
 
-namespace Insights
+namespace UE::Insights
 {
 	class FTable;
 	class FTableColumn;
 }
 
+namespace UE::Insights::NetworkingProfiler
+{
+
+class SNetStatsCounterTableRowToolTip;
+
 DECLARE_DELEGATE_RetVal_OneParam(bool, FNetStatsCounterNodeShouldBeEnabledDelegate, FNetStatsCounterNodePtr /*NodePtr*/);
 DECLARE_DELEGATE_RetVal_OneParam(bool, FIsColumnVisibleDelegate, const FName /*ColumnId*/);
 DECLARE_DELEGATE_RetVal_OneParam(EHorizontalAlignment, FGetColumnOutlineHAlignmentDelegate, const FName /*ColumnId*/);
-DECLARE_DELEGATE_ThreeParams(FSetHoveredNetStatsCounterTableCell, TSharedPtr<Insights::FTable> /*TablePtr*/, TSharedPtr<Insights::FTableColumn> /*ColumnPtr*/, FNetStatsCounterNodePtr /*NetStatsCounterNodePtr*/);
+DECLARE_DELEGATE_ThreeParams(FSetHoveredNetStatsCounterTableCell, TSharedPtr<FTable> /*TablePtr*/, TSharedPtr<FTableColumn> /*ColumnPtr*/, FNetStatsCounterNodePtr /*NetStatsCounterNodePtr*/);
 
 /** Widget that represents a table row in the tree control. Generates widgets for each column on demand. */
 class SNetStatsCountersTableRow : public SMultiColumnTableRow<FNetStatsCounterNodePtr>
@@ -37,7 +42,7 @@ public:
 		SLATE_EVENT(FSetHoveredNetStatsCounterTableCell, OnSetHoveredCell)
 		SLATE_ATTRIBUTE(FText, HighlightText)
 		SLATE_ATTRIBUTE(FName, HighlightedNodeName)
-		SLATE_ARGUMENT(TSharedPtr<Insights::FTable>, TablePtr)
+		SLATE_ARGUMENT(TSharedPtr<FTable>, TablePtr)
 		SLATE_ARGUMENT(FNetStatsCounterNodePtr, NetStatsCounterNodePtr)
 	SLATE_END_ARGS()
 
@@ -72,11 +77,11 @@ protected:
 	const FSlateBrush* GetOutlineBrush(const FName ColumnId) const;
 	bool HandleShouldBeEnabled() const;
 	EVisibility IsColumnVisible(const FName ColumnId) const;
-	void OnSetHoveredCell(TSharedPtr<Insights::FTable> InTablePtr, TSharedPtr<Insights::FTableColumn> InColumnPtr, FNetStatsCounterNodePtr InNetStatsCounterNodePtr);
+	void OnSetHoveredCell(TSharedPtr<FTable> InTablePtr, TSharedPtr<FTableColumn> InColumnPtr, FNetStatsCounterNodePtr InNetStatsCounterNodePtr);
 
 protected:
 	/** A shared pointer to the table view model. */
-	TSharedPtr<Insights::FTable> TablePtr;
+	TSharedPtr<FTable> TablePtr;
 
 	/** Data context for this table row. */
 	FNetStatsCounterNodePtr NetStatsCounterNodePtr;
@@ -94,3 +99,5 @@ protected:
 
 	TSharedPtr<SNetStatsCounterTableRowToolTip> RowToolTip;
 };
+
+} // namespace UE::Insights::NetworkingProfiler
