[Insights] Massive refactoring of TraceInsights module:

diff --git a/Engine/Source/Developer/TraceInsights/Private/Insights/NetworkingProfiler/Widgets/SNetStatsCountersView.cpp b/Engine/Source/Developer/TraceInsights/Private/Insights/NetworkingProfiler/Widgets/SNetStatsCountersView.cpp
--- a/Engine/Source/Developer/TraceInsights/Private/Insights/NetworkingProfiler/Widgets/SNetStatsCountersView.cpp
+++ b/Engine/Source/Developer/TraceInsights/Private/Insights/NetworkingProfiler/Widgets/SNetStatsCountersView.cpp
@@ -5,35 +5,42 @@
 #include "Framework/MultiBox/MultiBoxBuilder.h"
 #include "SlateOptMacros.h"
 #include "Templates/UniquePtr.h"
-#include "TraceServices/Model/NetProfiler.h"
+#include "Widgets/Images/SImage.h"
 #include "Widgets/Input/SCheckBox.h"
-#include "Widgets/Layout/SScrollBox.h"
 #include "Widgets/Input/SSearchBox.h"
-#include "Widgets/Images/SImage.h"
 #include "Widgets/Layout/SGridPanel.h"
+#include "Widgets/Layout/SScrollBox.h"
 #include "Widgets/Layout/SSeparator.h"
 #include "Widgets/SToolTip.h"
 #include "Widgets/Views/STableViewBase.h"
 
-// Insights
+// TraceServices
+#include "TraceServices/Model/NetProfiler.h"
+
+// TraceInsightsCore
+#include "InsightsCore/Table/ViewModels/Table.h"
+#include "InsightsCore/Table/ViewModels/TableColumn.h"
+
+// TraceInsights
 #include "Insights/InsightsStyle.h"
-#include "Insights/Table/ViewModels/Table.h"
-#include "Insights/Table/ViewModels/TableColumn.h"
 #include "Insights/NetworkingProfiler/NetworkingProfilerManager.h"
-#include "Insights/NetworkingProfiler/ViewModels/NetStatsCountersViewColumnFactory.h"
 #include "Insights/NetworkingProfiler/ViewModels/NetStatsCounterNodeHelper.h"
-#include "Insights/NetworkingProfiler/Widgets/SNetStatsCountersViewTooltip.h"
+#include "Insights/NetworkingProfiler/ViewModels/NetStatsCountersViewColumnFactory.h"
 #include "Insights/NetworkingProfiler/Widgets/SNetStatsCountersTableRow.h"
+#include "Insights/NetworkingProfiler/Widgets/SNetStatsCountersViewTooltip.h"
 #include "Insights/NetworkingProfiler/Widgets/SNetworkingProfilerWindow.h"
 #include "Insights/NetworkingProfiler/Widgets/SPacketContentView.h"
 
-#define LOCTEXT_NAMESPACE "SNetStatsCountersView"
+#define LOCTEXT_NAMESPACE "UE::Insights::NetworkingProfiler::SNetStatsCountersView"
+
+namespace UE::Insights::NetworkingProfiler
+{
 
 ////////////////////////////////////////////////////////////////////////////////////////////////////
 
 SNetStatsCountersView::SNetStatsCountersView()
 	: ProfilerWindow()
-	, Table(MakeShared<Insights::FTable>())
+	, Table(MakeShared<FTable>())
 	, bExpansionSaved(false)
 	, bFilterOutZeroCountStatsCounters(true)
 	, GroupingMode(ENetStatsCounterGroupingMode::ByType)
@@ -356,9 +363,9 @@ void SNetStatsCountersView::TreeView_BuildSortByMenu(FMenuBuilder& MenuBuilder)
 
 	MenuBuilder.BeginSection("ColumnName", LOCTEXT("ContextMenu_Header_Misc_ColumnName", "Column Name"));
 
-	for (const TSharedRef<Insights::FTableColumn>& ColumnRef : Table->GetColumns())
+	for (const TSharedRef<FTableColumn>& ColumnRef : Table->GetColumns())
 	{
-		const Insights::FTableColumn& Column = *ColumnRef;
+		const FTableColumn& Column = *ColumnRef;
 
 		if (Column.IsVisible() && Column.CanBeSorted())
 		{
@@ -427,9 +434,9 @@ void SNetStatsCountersView::TreeView_BuildViewColumnMenu(FMenuBuilder& MenuBuild
 {
 	MenuBuilder.BeginSection("ViewColumn", LOCTEXT("ContextMenu_Header_Columns_View", "View Column"));
 
-	for (const TSharedRef<Insights::FTableColumn>& ColumnRef : Table->GetColumns())
+	for (const TSharedRef<FTableColumn>& ColumnRef : Table->GetColumns())
 	{
-		const Insights::FTableColumn& Column = *ColumnRef;
+		const FTableColumn& Column = *ColumnRef;
 
 		FUIAction Action_ToggleColumn
 		(
@@ -456,12 +463,12 @@ void SNetStatsCountersView::TreeView_BuildViewColumnMenu(FMenuBuilder& MenuBuild
 void SNetStatsCountersView::InitializeAndShowHeaderColumns()
 {
 	// Create columns.
-	TArray<TSharedRef<Insights::FTableColumn>> Columns;
+	TArray<TSharedRef<FTableColumn>> Columns;
 	FNetStatsCountersViewColumnFactory::CreateNetStatsCountersViewColumns(Columns);
 	Table->SetColumns(Columns);
 
 	// Show columns.
-	for (const TSharedRef<Insights::FTableColumn>& ColumnRef : Table->GetColumns())
+	for (const TSharedRef<FTableColumn>& ColumnRef : Table->GetColumns())
 	{
 		if (ColumnRef->ShouldBeVisible())
 		{
@@ -474,13 +481,13 @@ void SNetStatsCountersView::InitializeAndShowHeaderColumns()
 
 FText SNetStatsCountersView::GetColumnHeaderText(const FName ColumnId) const
 {
-	const Insights::FTableColumn& Column = *Table->FindColumnChecked(ColumnId);
+	const FTableColumn& Column = *Table->FindColumnChecked(ColumnId);
 	return Column.GetShortName();
 }
 
 ////////////////////////////////////////////////////////////////////////////////////////////////////
 
-TSharedRef<SWidget> SNetStatsCountersView::TreeViewHeaderRow_GenerateColumnMenu(const Insights::FTableColumn& Column)
+TSharedRef<SWidget> SNetStatsCountersView::TreeViewHeaderRow_GenerateColumnMenu(const FTableColumn& Column)
 {
 	bool bIsMenuVisible = false;
 
@@ -610,7 +617,7 @@ void SNetStatsCountersView::UpdateTree()
 	const double TotalTime = Stopwatch.GetAccumulatedTime();
 	if (TotalTime > 0.1)
 	{
-		UE_LOG(NetworkingProfiler, Log, TEXT("[NetStats] Tree view updated in %.3fs (%d events) --> G:%.3fs + S:%.3fs + F:%.3fs"),
+		UE_LOG(LogNetworkingProfiler, Log, TEXT("[NetStats] Tree view updated in %.3fs (%d events) --> G:%.3fs + S:%.3fs + F:%.3fs"),
 			TotalTime, NetStatsCounterNodes.Num(), Time1, Time2 - Time1, TotalTime - Time2);
 	}
 }
@@ -629,11 +636,11 @@ void SNetStatsCountersView::ApplyFiltering()
 		GroupPtr->ClearFilteredChildren();
 		const bool bIsGroupVisible = Filters->PassesAllFilters(GroupPtr);
 
-		const TArray<Insights::FBaseTreeNodePtr>& GroupChildren = GroupPtr->GetChildren();
+		const TArray<FBaseTreeNodePtr>& GroupChildren = GroupPtr->GetChildren();
 		int32 NumVisibleChildren = 0;
-		for (const Insights::FBaseTreeNodePtr& ChildPtr : GroupChildren)
+		for (const FBaseTreeNodePtr& ChildPtr : GroupChildren)
 		{
-			const FNetStatsCounterNodePtr& NodePtr = StaticCastSharedPtr<FNetStatsCounterNode, Insights::FBaseTreeNode>(ChildPtr);
+			const FNetStatsCounterNodePtr& NodePtr = StaticCastSharedPtr<FNetStatsCounterNode, FBaseTreeNode>(ChildPtr);
 
 			const bool bIsChildVisible = (!bFilterOutZeroCountStatsCounters || NodePtr->GetAggregatedStats().Count > 0)
 									  && bStatsCounterTypeIsVisible[static_cast<int>(NodePtr->GetType())]
@@ -781,11 +788,11 @@ void SNetStatsCountersView::TreeView_OnSelectionChanged(FNetStatsCounterNodePtr
 
 void SNetStatsCountersView::TreeView_OnGetChildren(FNetStatsCounterNodePtr InParent, TArray<FNetStatsCounterNodePtr>& OutChildren)
 {
-	const TArray<Insights::FBaseTreeNodePtr>& Children = InParent->GetFilteredChildren();
+	const TArray<FBaseTreeNodePtr>& Children = InParent->GetFilteredChildren();
 	OutChildren.Reset(Children.Num());
-	for (const Insights::FBaseTreeNodePtr& Child : Children)
+	for (const FBaseTreeNodePtr& Child : Children)
 	{
-		OutChildren.Add(StaticCastSharedPtr<FNetStatsCounterNode, Insights::FBaseTreeNode>(Child));
+		OutChildren.Add(StaticCastSharedPtr<FNetStatsCounterNode, FBaseTreeNode>(Child));
 	}
 }
 
@@ -829,7 +836,7 @@ bool SNetStatsCountersView::TableRow_ShouldBeEnabled(FNetStatsCounterNodePtr Nod
 
 ////////////////////////////////////////////////////////////////////////////////////////////////////
 
-void SNetStatsCountersView::TableRow_SetHoveredCell(TSharedPtr<Insights::FTable> InTablePtr, TSharedPtr<Insights::FTableColumn> InColumnPtr, FNetStatsCounterNodePtr InNodePtr)
+void SNetStatsCountersView::TableRow_SetHoveredCell(TSharedPtr<FTable> InTablePtr, TSharedPtr<FTableColumn> InColumnPtr, FNetStatsCounterNodePtr InNodePtr)
 {
 	HoveredColumnId = InColumnPtr ? InColumnPtr->GetId() : FName();
 
@@ -1035,11 +1042,11 @@ void SNetStatsCountersView::CreateSortings()
 	AvailableSorters.Reset();
 	CurrentSorter = nullptr;
 
-	for (const TSharedRef<Insights::FTableColumn>& ColumnRef : Table->GetColumns())
+	for (const TSharedRef<FTableColumn>& ColumnRef : Table->GetColumns())
 	{
 		if (ColumnRef->CanBeSorted())
 		{
-			TSharedPtr<Insights::ITableCellValueSorter> SorterPtr = ColumnRef->GetValueSorter();
+			TSharedPtr<ITableCellValueSorter> SorterPtr = ColumnRef->GetValueSorter();
 			if (ensure(SorterPtr.IsValid()))
 			{
 				AvailableSorters.Add(SorterPtr);
@@ -1054,7 +1061,7 @@ void SNetStatsCountersView::CreateSortings()
 
 void SNetStatsCountersView::UpdateCurrentSortingByColumn()
 {
-	TSharedPtr<Insights::FTableColumn> ColumnPtr = Table->FindColumn(ColumnBeingSorted);
+	TSharedPtr<FTableColumn> ColumnPtr = Table->FindColumn(ColumnBeingSorted);
 	CurrentSorter = ColumnPtr.IsValid() ? ColumnPtr->GetValueSorter() : nullptr;
 }
 
@@ -1073,12 +1080,12 @@ void SNetStatsCountersView::SortTreeNodes()
 
 ////////////////////////////////////////////////////////////////////////////////////////////////////
 
-void SNetStatsCountersView::SortTreeNodesRec(FNetStatsCounterNode& Node, const Insights::ITableCellValueSorter& Sorter)
+void SNetStatsCountersView::SortTreeNodesRec(FNetStatsCounterNode& Node, const ITableCellValueSorter& Sorter)
 {
-	Insights::ESortMode SortMode = (ColumnSortMode == EColumnSortMode::Type::Descending) ? Insights::ESortMode::Descending : Insights::ESortMode::Ascending;
+	ESortMode SortMode = (ColumnSortMode == EColumnSortMode::Type::Descending) ? ESortMode::Descending : ESortMode::Ascending;
 	Node.SortChildren(Sorter, SortMode);
 
-	for (Insights::FBaseTreeNodePtr ChildPtr : Node.GetChildren())
+	for (FBaseTreeNodePtr ChildPtr : Node.GetChildren())
 	{
 		if (ChildPtr->GetChildrenCount() > 0)
 		{
@@ -1132,7 +1139,7 @@ bool SNetStatsCountersView::HeaderMenu_SortMode_IsChecked(const FName ColumnId,
 
 bool SNetStatsCountersView::HeaderMenu_SortMode_CanExecute(const FName ColumnId, const EColumnSortMode::Type InSortMode) const
 {
-	const Insights::FTableColumn& Column = *Table->FindColumnChecked(ColumnId);
+	const FTableColumn& Column = *Table->FindColumnChecked(ColumnId);
 	return Column.CanBeSorted();
 }
 
@@ -1205,7 +1212,7 @@ bool SNetStatsCountersView::CanShowColumn(const FName ColumnId) const
 
 void SNetStatsCountersView::ShowColumn(const FName ColumnId)
 {
-	Insights::FTableColumn& Column = *Table->FindColumnChecked(ColumnId);
+	FTableColumn& Column = *Table->FindColumnChecked(ColumnId);
 	Column.Show();
 
 	SHeaderRow::FColumn::FArguments ColumnArgs;
@@ -1260,7 +1267,7 @@ void SNetStatsCountersView::ShowColumn(const FName ColumnId)
 
 bool SNetStatsCountersView::CanHideColumn(const FName ColumnId) const
 {
-	const Insights::FTableColumn& Column = *Table->FindColumnChecked(ColumnId);
+	const FTableColumn& Column = *Table->FindColumnChecked(ColumnId);
 	return Column.CanBeHidden();
 }
 
@@ -1268,7 +1275,7 @@ bool SNetStatsCountersView::CanHideColumn(const FName ColumnId) const
 
 void SNetStatsCountersView::HideColumn(const FName ColumnId)
 {
-	Insights::FTableColumn& Column = *Table->FindColumnChecked(ColumnId);
+	FTableColumn& Column = *Table->FindColumnChecked(ColumnId);
 	Column.Hide();
 
 	TreeViewHeaderRow->RemoveColumn(ColumnId);
@@ -1280,7 +1287,7 @@ void SNetStatsCountersView::HideColumn(const FName ColumnId)
 
 bool SNetStatsCountersView::IsColumnVisible(const FName ColumnId) const
 {
-	const Insights::FTableColumn& Column = *Table->FindColumnChecked(ColumnId);
+	const FTableColumn& Column = *Table->FindColumnChecked(ColumnId);
 	return Column.IsVisible();
 }
 
@@ -1288,7 +1295,7 @@ bool SNetStatsCountersView::IsColumnVisible(const FName ColumnId) const
 
 bool SNetStatsCountersView::CanToggleColumnVisibility(const FName ColumnId) const
 {
-	const Insights::FTableColumn& Column = *Table->FindColumnChecked(ColumnId);
+	const FTableColumn& Column = *Table->FindColumnChecked(ColumnId);
 	return !Column.IsVisible() || Column.CanBeHidden();
 }
 
@@ -1296,7 +1303,7 @@ bool SNetStatsCountersView::CanToggleColumnVisibility(const FName ColumnId) cons
 
 void SNetStatsCountersView::ToggleColumnVisibility(const FName ColumnId)
 {
-	const Insights::FTableColumn& Column = *Table->FindColumnChecked(ColumnId);
+	const FTableColumn& Column = *Table->FindColumnChecked(ColumnId);
 	if (Column.IsVisible())
 	{
 		HideColumn(ColumnId);
@@ -1324,9 +1331,9 @@ void SNetStatsCountersView::ContextMenu_ShowAllColumns_Execute()
 	ColumnSortMode = GetDefaultColumnSortMode();
 	UpdateCurrentSortingByColumn();
 
-	for (const TSharedRef<Insights::FTableColumn>& ColumnRef : Table->GetColumns())
+	for (const TSharedRef<FTableColumn>& ColumnRef : Table->GetColumns())
 	{
-		const Insights::FTableColumn& Column = *ColumnRef;
+		const FTableColumn& Column = *ColumnRef;
 
 		if (!Column.IsVisible())
 		{
@@ -1352,9 +1359,9 @@ void SNetStatsCountersView::ContextMenu_ResetColumns_Execute()
 	ColumnSortMode = GetDefaultColumnSortMode();
 	UpdateCurrentSortingByColumn();
 
-	for (const TSharedRef<Insights::FTableColumn>& ColumnRef : Table->GetColumns())
+	for (const TSharedRef<FTableColumn>& ColumnRef : Table->GetColumns())
 	{
-		const Insights::FTableColumn& Column = *ColumnRef;
+		const FTableColumn& Column = *ColumnRef;
 
 		if (Column.ShouldBeVisible() && !Column.IsVisible())
 		{
@@ -1476,7 +1483,7 @@ void SNetStatsCountersView::RebuildTree(bool bResync)
 	if (TotalTime > 0.01)
 	{
 		const double SyncTime = SyncStopwatch.GetAccumulatedTime();
-		UE_LOG(NetworkingProfiler, Log, TEXT("[NetStatsCounters] Tree view rebuilt in %.4fs (%.4fs + %.4fs) --> %d net stats counters (%d added)"),
+		UE_LOG(LogNetworkingProfiler, Log, TEXT("[NetStatsCounters] Tree view rebuilt in %.4fs (%.4fs + %.4fs) --> %d net stats counters (%d added)"),
 			TotalTime, SyncTime, TotalTime - SyncTime, NetStatsCounterNodes.Num(), NetStatsCounterNodes.Num() - PreviousNodeCount);
 	}
 }
@@ -1578,7 +1585,7 @@ void SNetStatsCountersView::UpdateStatsInternal()
 	Stopwatch.Stop();
 	const double TotalTime = Stopwatch.GetAccumulatedTime();
 	const double AggregationTime = AggregationStopwatch.GetAccumulatedTime();
-	UE_LOG(NetworkingProfiler, Log, TEXT("[NetStats] Aggregated stats updated in %.4fs (%.4fs + %.4fs)"),
+	UE_LOG(LogNetworkingProfiler, Log, TEXT("[NetStats] Aggregated stats updated in %.4fs (%.4fs + %.4fs)"),
 		TotalTime, AggregationTime, TotalTime - AggregationTime);
 }
 
@@ -1603,4 +1610,6 @@ void SNetStatsCountersView::SelectNetStatsCounterNode(uint32 EventTypeIndex)
 
 ////////////////////////////////////////////////////////////////////////////////////////////////////
 
+} // namespace UE::Insights::NetworkingProfiler
+
 #undef LOCTEXT_NAMESPACE
diff --git a/Engine/Source/Developer/TraceInsights/Private/Insights/NetworkingProfiler/Widgets/SNetStatsCountersView.h b/Engine/Source/Developer/TraceInsights/Private/Insights/NetworkingProfiler/Widgets/SNetStatsCountersView.h
--- a/Engine/Source/Developer/TraceInsights/Private/Insights/NetworkingProfiler/Widgets/SNetStatsCountersView.h
+++ b/Engine/Source/Developer/TraceInsights/Private/Insights/NetworkingProfiler/Widgets/SNetStatsCountersView.h
@@ -15,28 +15,34 @@
 #include "Widgets/Views/STableViewBase.h"
 #include "Widgets/Views/STreeView.h"
 
-// Insights
+// TraceServices
+#include "TraceServices/Model/NetProfiler.h"
+
+// TraceInsights
 #include "Insights/NetworkingProfiler/ViewModels/NetStatsCounterGroupingAndSorting.h"
 #include "Insights/NetworkingProfiler/ViewModels/NetStatsCounterNode.h"
-#include "TraceServices/Model/NetProfiler.h"
 
 ////////////////////////////////////////////////////////////////////////////////////////////////////
 
 class FMenuBuilder;
-class SNetworkingProfilerWindow;
 
 namespace TraceServices
 {
 	class IAnalysisSession;
 }
 
-namespace Insights
+namespace UE::Insights
 {
 	class FTable;
 	class FTableColumn;
 	class ITableCellValueSorter;
 }
 
+namespace UE::Insights::NetworkingProfiler
+{
+
+class SNetworkingProfilerWindow;
+
 ////////////////////////////////////////////////////////////////////////////////////////////////////
 
 /** The filter collection - used for updating the list of tree nodes. */
@@ -67,7 +73,7 @@ public:
 	 */
 	void Construct(const FArguments& InArgs, TSharedPtr<SNetworkingProfilerWindow> InProfilerWindow);
 
-	TSharedPtr<Insights::FTable> GetTable() const { return Table; }
+	TSharedPtr<FTable> GetTable() const { return Table; }
 
 	/**
 	 * Ticks this widget. Override in derived classes, but always call the parent implementation.
@@ -123,7 +129,7 @@ protected:
 
 	FText GetColumnHeaderText(const FName ColumnId) const;
 
-	TSharedRef<SWidget> TreeViewHeaderRow_GenerateColumnMenu(const Insights::FTableColumn& Column);
+	TSharedRef<SWidget> TreeViewHeaderRow_GenerateColumnMenu(const FTableColumn& Column);
 
 	////////////////////////////////////////////////////////////////////////////////////////////////////
 	// Tree View - Misc
@@ -151,7 +157,7 @@ protected:
 
 	bool TableRow_ShouldBeEnabled(FNetStatsCounterNodePtr NodePtr) const;
 
-	void TableRow_SetHoveredCell(TSharedPtr<Insights::FTable> TablePtr, TSharedPtr<Insights::FTableColumn> ColumnPtr, FNetStatsCounterNodePtr NodePtr);
+	void TableRow_SetHoveredCell(TSharedPtr<FTable> TablePtr, TSharedPtr<FTableColumn> ColumnPtr, FNetStatsCounterNodePtr NodePtr);
 	EHorizontalAlignment TableRow_GetColumnOutlineHAlignment(const FName ColumnId) const;
 
 	FText TableRow_GetHighlightText() const;
@@ -197,7 +203,7 @@ protected:
 
 	void UpdateCurrentSortingByColumn();
 	void SortTreeNodes();
-	void SortTreeNodesRec(FNetStatsCounterNode& Node, const Insights::ITableCellValueSorter& Sorter);
+	void SortTreeNodesRec(FNetStatsCounterNode& Node, const ITableCellValueSorter& Sorter);
 
 	EColumnSortMode::Type GetSortModeForColumn(const FName ColumnId) const;
 	void SetSortModeForColumn(const FName& ColumnId, EColumnSortMode::Type SortMode);
@@ -255,7 +261,7 @@ protected:
 	TSharedPtr<SNetworkingProfilerWindow> ProfilerWindow;
 
 	/** Table view model. */
-	TSharedPtr<Insights::FTable> Table;
+	TSharedPtr<FTable> Table;
 
 	/** A weak pointer to the profiler session used to populate this widget. */
 	TSharedPtr<const TraceServices::IAnalysisSession>/*Weak*/ Session;
@@ -340,10 +346,10 @@ protected:
 	//bool bUseSorting;
 
 	/** All available sorters. */
-	TArray<TSharedPtr<Insights::ITableCellValueSorter>> AvailableSorters;
+	TArray<TSharedPtr<ITableCellValueSorter>> AvailableSorters;
 
 	/** Current sorter. It is nullptr if sorting is disabled. */
-	TSharedPtr<Insights::ITableCellValueSorter> CurrentSorter;
+	TSharedPtr<ITableCellValueSorter> CurrentSorter;
 
 	/** Name of the column currently being sorted. Can be NAME_None if sorting is disabled (CurrentSorting == nullptr) or if a complex sorting is used (CurrentSorting != nullptr). */
 	FName ColumnBeingSorted;
@@ -364,3 +370,5 @@ protected:
 };
 
 ////////////////////////////////////////////////////////////////////////////////////////////////////
+
+} // namespace UE::Insights::NetworkingProfiler
diff --git a/Engine/Source/Developer/TraceInsights/Private/Insights/NetworkingProfiler/Widgets/SNetStatsCountersViewTooltip.cpp b/Engine/Source/Developer/TraceInsights/Private/Insights/NetworkingProfiler/Widgets/SNetStatsCountersViewTooltip.cpp
--- a/Engine/Source/Developer/TraceInsights/Private/Insights/NetworkingProfiler/Widgets/SNetStatsCountersViewTooltip.cpp
+++ b/Engine/Source/Developer/TraceInsights/Private/Insights/NetworkingProfiler/Widgets/SNetStatsCountersViewTooltip.cpp
@@ -3,29 +3,36 @@
 #include "SNetStatsCountersViewTooltip.h"
 
 #include "SlateOptMacros.h"
-#include "TraceServices/Model/NetProfiler.h"
 #include "Widgets/Layout/SGridPanel.h"
 #include "Widgets/Layout/SSeparator.h"
 #include "Widgets/SBoxPanel.h"
 #include "Widgets/SToolTip.h"
 #include "Widgets/Text/STextBlock.h"
 
-// Insights
+// TraceServices
+#include "TraceServices/Model/NetProfiler.h"
+
+// TraceInsightsCore
+#include "InsightsCore/Table/ViewModels/Table.h"
+#include "InsightsCore/Table/ViewModels/TableColumn.h"
+
+// TraceInsights
 #include "Insights/InsightsStyle.h"
-#include "Insights/Table/ViewModels/Table.h"
-#include "Insights/Table/ViewModels/TableColumn.h"
 #include "Insights/NetworkingProfiler/ViewModels/NetStatsCounterNode.h"
 #include "Insights/NetworkingProfiler/ViewModels/NetStatsCounterNodeHelper.h"
 
 ////////////////////////////////////////////////////////////////////////////////////////////////////
 
-#define LOCTEXT_NAMESPACE "SNetStatsCountersView"
+#define LOCTEXT_NAMESPACE "UE::Insights::NetworkingProfiler::SNetStatsCountersView"
+
+namespace UE::Insights::NetworkingProfiler
+{
 
 ////////////////////////////////////////////////////////////////////////////////////////////////////
 
 BEGIN_SLATE_FUNCTION_BUILD_OPTIMIZATION
 
-TSharedPtr<SToolTip> SNetStatsCountersViewTooltip::GetTableTooltip(const Insights::FTable& Table)
+TSharedPtr<SToolTip> SNetStatsCountersViewTooltip::GetTableTooltip(const FTable& Table)
 {
 	TSharedPtr<SToolTip> ColumnTooltip =
 		SNew(SToolTip)
@@ -56,7 +63,7 @@ TSharedPtr<SToolTip> SNetStatsCountersViewTooltip::GetTableTooltip(const Insight
 
 ////////////////////////////////////////////////////////////////////////////////////////////////////
 
-TSharedPtr<SToolTip> SNetStatsCountersViewTooltip::GetColumnTooltip(const Insights::FTableColumn& Column)
+TSharedPtr<SToolTip> SNetStatsCountersViewTooltip::GetColumnTooltip(const FTableColumn& Column)
 {
 	TSharedPtr<SToolTip> ColumnTooltip =
 		SNew(SToolTip)
@@ -270,4 +277,6 @@ END_SLATE_FUNCTION_BUILD_OPTIMIZATION
 
 ////////////////////////////////////////////////////////////////////////////////////////////////////
 
+} // namespace UE::Insights::NetworkingProfiler
+
 #undef LOCTEXT_NAMESPACE
diff --git a/Engine/Source/Developer/TraceInsights/Private/Insights/NetworkingProfiler/Widgets/SNetStatsCountersViewTooltip.h b/Engine/Source/Developer/TraceInsights/Private/Insights/NetworkingProfiler/Widgets/SNetStatsCountersViewTooltip.h
--- a/Engine/Source/Developer/TraceInsights/Private/Insights/NetworkingProfiler/Widgets/SNetStatsCountersViewTooltip.h
+++ b/Engine/Source/Developer/TraceInsights/Private/Insights/NetworkingProfiler/Widgets/SNetStatsCountersViewTooltip.h
@@ -3,17 +3,21 @@
 #pragma once
 
 #include "CoreMinimal.h"
+
 #include "Widgets/IToolTip.h"
 #include "Widgets/SToolTip.h"
 
 class SGridPanel;
 
-namespace Insights
+namespace UE::Insights
 {
 	class FTable;
 	class FTableColumn;
 }
 
+namespace UE::Insights::NetworkingProfiler
+{
+
 class FNetStatsCounterNode;
 
 ////////////////////////////////////////////////////////////////////////////////////////////////////
@@ -24,8 +28,8 @@ class SNetStatsCountersViewTooltip
 public:
 	SNetStatsCountersViewTooltip() = delete;
 
-	static TSharedPtr<SToolTip> GetTableTooltip(const Insights::FTable& Table);
-	static TSharedPtr<SToolTip> GetColumnTooltip(const Insights::FTableColumn& Column);
+	static TSharedPtr<SToolTip> GetTableTooltip(const FTable& Table);
+	static TSharedPtr<SToolTip> GetColumnTooltip(const FTableColumn& Column);
 	static TSharedPtr<SToolTip> GetRowTooltip(const TSharedPtr<FNetStatsCounterNode> TreeNodePtr);
 
 private:
@@ -83,3 +87,5 @@ private:
 };
 
 ////////////////////////////////////////////////////////////////////////////////////////////////////
+
+} // namespace UE::Insights::NetworkingProfiler
diff --git a/Engine/Source/Developer/TraceInsights/Private/Insights/NetworkingProfiler/Widgets/SNetStatsTableCell.cpp b/Engine/Source/Developer/TraceInsights/Private/Insights/NetworkingProfiler/Widgets/SNetStatsTableCell.cpp
--- a/Engine/Source/Developer/TraceInsights/Private/Insights/NetworkingProfiler/Widgets/SNetStatsTableCell.cpp
+++ b/Engine/Source/Developer/TraceInsights/Private/Insights/NetworkingProfiler/Widgets/SNetStatsTableCell.cpp
@@ -11,13 +11,18 @@
 #include "Widgets/Text/STextBlock.h"
 #include "Widgets/Views/SExpanderArrow.h"
 
-// Insights
+// TraceInsightsCore
+#include "InsightsCore/Table/ViewModels/Table.h"
+#include "InsightsCore/Table/ViewModels/TableColumn.h"
+
+// TraceInsights
 #include "Insights/InsightsStyle.h"
-#include "Insights/Table/ViewModels/Table.h"
-#include "Insights/Table/ViewModels/TableColumn.h"
 #include "Insights/NetworkingProfiler/Widgets/SNetStatsTableRow.h"
 
-#define LOCTEXT_NAMESPACE "SNetStatsView"
+#define LOCTEXT_NAMESPACE "UE::Insights::NetworkingProfiler::SNetStatsView"
+
+namespace UE::Insights::NetworkingProfiler
+{
 
 ////////////////////////////////////////////////////////////////////////////////////////////////////
 
@@ -156,4 +161,6 @@ END_SLATE_FUNCTION_BUILD_OPTIMIZATION
 
 ////////////////////////////////////////////////////////////////////////////////////////////////////
 
+} // namespace UE::Insights::NetworkingProfiler
+
 #undef LOCTEXT_NAMESPACE
diff --git a/Engine/Source/Developer/TraceInsights/Private/Insights/NetworkingProfiler/Widgets/SNetStatsTableCell.h b/Engine/Source/Developer/TraceInsights/Private/Insights/NetworkingProfiler/Widgets/SNetStatsTableCell.h
--- a/Engine/Source/Developer/TraceInsights/Private/Insights/NetworkingProfiler/Widgets/SNetStatsTableCell.h
+++ b/Engine/Source/Developer/TraceInsights/Private/Insights/NetworkingProfiler/Widgets/SNetStatsTableCell.h
@@ -3,21 +3,25 @@
 #pragma once
 
 #include "CoreMinimal.h"
+
 #include "Widgets/DeclarativeSyntaxSupport.h"
 #include "Widgets/SCompoundWidget.h"
 
-// Insights
+// TraceInsights
 #include "Insights/NetworkingProfiler/ViewModels/NetEventNode.h"
 
 class ITableRow;
 
-namespace Insights
+namespace UE::Insights
 {
 	class FTable;
 	class FTableColumn;
 }
 
-DECLARE_DELEGATE_ThreeParams(FSetHoveredNetStatsTableCell, TSharedPtr<Insights::FTable> /*TablePtr*/, TSharedPtr<Insights::FTableColumn> /*ColumnPtr*/, FNetEventNodePtr /*NetEventNodePtr*/);
+namespace UE::Insights::NetworkingProfiler
+{
+
+DECLARE_DELEGATE_ThreeParams(FSetHoveredNetStatsTableCell, TSharedPtr<FTable> /*TablePtr*/, TSharedPtr<FTableColumn> /*ColumnPtr*/, FNetEventNodePtr /*NetEventNodePtr*/);
 
 class SNetStatsTableCell : public SCompoundWidget
 {
@@ -25,8 +29,8 @@ public:
 	SLATE_BEGIN_ARGS(SNetStatsTableCell) {}
 		SLATE_EVENT(FSetHoveredNetStatsTableCell, OnSetHoveredCell)
 		SLATE_ATTRIBUTE(FText, HighlightText)
-		SLATE_ARGUMENT(TSharedPtr<Insights::FTable>, TablePtr)
-		SLATE_ARGUMENT(TSharedPtr<Insights::FTableColumn>, ColumnPtr)
+		SLATE_ARGUMENT(TSharedPtr<FTable>, TablePtr)
+		SLATE_ARGUMENT(TSharedPtr<FTableColumn>, ColumnPtr)
 		SLATE_ARGUMENT(FNetEventNodePtr, NetEventNodePtr)
 		SLATE_ARGUMENT(bool, IsNameColumn)
 	SLATE_END_ARGS()
@@ -159,13 +163,15 @@ protected:
 
 protected:
 	/** A shared pointer to the table view model. */
-	TSharedPtr<Insights::FTable> TablePtr; // TODO: TSharedRef
+	TSharedPtr<FTable> TablePtr; // TODO: TSharedRef
 
 	/** A shared pointer to the table column view model. */
-	TSharedPtr<Insights::FTableColumn> ColumnPtr; // TODO: TSharedRef
+	TSharedPtr<FTableColumn> ColumnPtr; // TODO: TSharedRef
 
 	/** A shared pointer to the timer node. */
 	FNetEventNodePtr NetEventNodePtr; // TODO: TSharedRef
 
 	FSetHoveredNetStatsTableCell SetHoveredCellDelegate;
 };
+
+} // namespace UE::Insights::NetworkingProfiler
diff --git a/Engine/Source/Developer/TraceInsights/Private/Insights/NetworkingProfiler/Widgets/SNetStatsTableRow.cpp b/Engine/Source/Developer/TraceInsights/Private/Insights/NetworkingProfiler/Widgets/SNetStatsTableRow.cpp
--- a/Engine/Source/Developer/TraceInsights/Private/Insights/NetworkingProfiler/Widgets/SNetStatsTableRow.cpp
+++ b/Engine/Source/Developer/TraceInsights/Private/Insights/NetworkingProfiler/Widgets/SNetStatsTableRow.cpp
@@ -6,15 +6,20 @@
 #include "Widgets/SOverlay.h"
 #include "Widgets/Images/SImage.h"
 
-// Insights
-#include "Insights/Common/TimeUtils.h"
+// TraceInsightsCore
+#include "InsightsCore/Common/TimeUtils.h"
+#include "InsightsCore/Table/ViewModels/Table.h"
+#include "InsightsCore/Table/ViewModels/TableColumn.h"
+
+// TraceInsights
 #include "Insights/InsightsStyle.h"
-#include "Insights/Table/ViewModels/Table.h"
-#include "Insights/Table/ViewModels/TableColumn.h"
 #include "Insights/NetworkingProfiler/Widgets/SNetStatsTableCell.h"
 #include "Insights/NetworkingProfiler/Widgets/SNetStatsViewTooltip.h"
 
-#define LOCTEXT_NAMESPACE "SNetStatsView"
+#define LOCTEXT_NAMESPACE "UE::Insights::NetworkingProfiler::SNetStatsView"
+
+namespace UE::Insights::NetworkingProfiler
+{
 
 ////////////////////////////////////////////////////////////////////////////////////////////////////
 
@@ -44,7 +49,7 @@ void SNetStatsTableRow::Construct(const FArguments& InArgs, const TSharedRef<STa
 
 TSharedRef<SWidget> SNetStatsTableRow::GenerateWidgetForColumn(const FName& ColumnId)
 {
-	TSharedPtr<Insights::FTableColumn> ColumnPtr = TablePtr->FindColumnChecked(ColumnId);
+	TSharedPtr<FTableColumn> ColumnPtr = TablePtr->FindColumnChecked(ColumnId);
 
 	return
 		SNew(SOverlay)
@@ -180,11 +185,13 @@ EVisibility SNetStatsTableRow::IsColumnVisible(const FName ColumnId) const
 
 ////////////////////////////////////////////////////////////////////////////////////////////////////
 
-void SNetStatsTableRow::OnSetHoveredCell(TSharedPtr<Insights::FTable> InTablePtr, TSharedPtr<Insights::FTableColumn> InColumnPtr, FNetEventNodePtr InNetEventNodePtr)
+void SNetStatsTableRow::OnSetHoveredCell(TSharedPtr<FTable> InTablePtr, TSharedPtr<FTableColumn> InColumnPtr, FNetEventNodePtr InNetEventNodePtr)
 {
 	SetHoveredCellDelegate.ExecuteIfBound(InTablePtr, InColumnPtr, InNetEventNodePtr);
 }
 
 ////////////////////////////////////////////////////////////////////////////////////////////////////
 
+} // namespace UE::Insights::NetworkingProfiler
+
 #undef LOCTEXT_NAMESPACE
diff --git a/Engine/Source/Developer/TraceInsights/Private/Insights/NetworkingProfiler/Widgets/SNetStatsTableRow.h b/Engine/Source/Developer/TraceInsights/Private/Insights/NetworkingProfiler/Widgets/SNetStatsTableRow.h
--- a/Engine/Source/Developer/TraceInsights/Private/Insights/NetworkingProfiler/Widgets/SNetStatsTableRow.h
+++ b/Engine/Source/Developer/TraceInsights/Private/Insights/NetworkingProfiler/Widgets/SNetStatsTableRow.h
@@ -3,28 +3,34 @@
 #pragma once
 
 #include "CoreMinimal.h"
+
 #include "Widgets/DeclarativeSyntaxSupport.h"
 #include "Widgets/Views/SHeaderRow.h"
 #include "Widgets/Views/STableViewBase.h"
 #include "Widgets/Views/STableRow.h"
 #include "Widgets/Views/STreeView.h"
 
-// Insights
+// TraceInsights
+#include "Insights/NetworkingProfiler/ViewModels/NetEventNode.h"
 #include "Insights/NetworkingProfiler/ViewModels/NetEventNodeHelper.h"
 
 class IToolTip;
-class SNetEventTableRowToolTip;
 
-namespace Insights
+namespace UE::Insights
 {
 	class FTable;
 	class FTableColumn;
 }
 
+namespace UE::Insights::NetworkingProfiler
+{
+
+class SNetEventTableRowToolTip;
+
 DECLARE_DELEGATE_RetVal_OneParam(bool, FNetEventNodeShouldBeEnabledDelegate, FNetEventNodePtr /*NodePtr*/);
 DECLARE_DELEGATE_RetVal_OneParam(bool, FIsColumnVisibleDelegate, const FName /*ColumnId*/);
 DECLARE_DELEGATE_RetVal_OneParam(EHorizontalAlignment, FGetColumnOutlineHAlignmentDelegate, const FName /*ColumnId*/);
-DECLARE_DELEGATE_ThreeParams(FSetHoveredNetEventTableCell, TSharedPtr<Insights::FTable> /*TablePtr*/, TSharedPtr<Insights::FTableColumn> /*ColumnPtr*/, FNetEventNodePtr /*NetEventNodePtr*/);
+DECLARE_DELEGATE_ThreeParams(FSetHoveredNetEventTableCell, TSharedPtr<FTable> /*TablePtr*/, TSharedPtr<FTableColumn> /*ColumnPtr*/, FNetEventNodePtr /*NetEventNodePtr*/);
 
 /** Widget that represents a table row in the tree control. Generates widgets for each column on demand. */
 class SNetStatsTableRow : public SMultiColumnTableRow<FNetEventNodePtr>
@@ -37,7 +43,7 @@ public:
 		SLATE_EVENT(FSetHoveredNetEventTableCell, OnSetHoveredCell)
 		SLATE_ATTRIBUTE(FText, HighlightText)
 		SLATE_ATTRIBUTE(FName, HighlightedNodeName)
-		SLATE_ARGUMENT(TSharedPtr<Insights::FTable>, TablePtr)
+		SLATE_ARGUMENT(TSharedPtr<FTable>, TablePtr)
 		SLATE_ARGUMENT(FNetEventNodePtr, NetEventNodePtr)
 	SLATE_END_ARGS()
 
@@ -72,11 +78,11 @@ protected:
 	const FSlateBrush* GetOutlineBrush(const FName ColumnId) const;
 	bool HandleShouldBeEnabled() const;
 	EVisibility IsColumnVisible(const FName ColumnId) const;
-	void OnSetHoveredCell(TSharedPtr<Insights::FTable> InTablePtr, TSharedPtr<Insights::FTableColumn> InColumnPtr, FNetEventNodePtr InNetEventNodePtr);
+	void OnSetHoveredCell(TSharedPtr<FTable> InTablePtr, TSharedPtr<FTableColumn> InColumnPtr, FNetEventNodePtr InNetEventNodePtr);
 
 protected:
 	/** A shared pointer to the table view model. */
-	TSharedPtr<Insights::FTable> TablePtr;
+	TSharedPtr<FTable> TablePtr;
 
 	/** Data context for this table row. */
 	FNetEventNodePtr NetEventNodePtr;
@@ -94,3 +100,5 @@ protected:
 
 	TSharedPtr<SNetEventTableRowToolTip> RowToolTip;
 };
+
+} // namespace UE::Insights::NetworkingProfiler
diff --git a/Engine/Source/Developer/TraceInsights/Private/Insights/NetworkingProfiler/Widgets/SNetStatsView.cpp b/Engine/Source/Developer/TraceInsights/Private/Insights/NetworkingProfiler/Widgets/SNetStatsView.cpp
--- a/Engine/Source/Developer/TraceInsights/Private/Insights/NetworkingProfiler/Widgets/SNetStatsView.cpp
+++ b/Engine/Source/Developer/TraceInsights/Private/Insights/NetworkingProfiler/Widgets/SNetStatsView.cpp
@@ -5,7 +5,6 @@
 #include "Framework/MultiBox/MultiBoxBuilder.h"
 #include "SlateOptMacros.h"
 #include "Templates/UniquePtr.h"
-#include "TraceServices/Model/NetProfiler.h"
 #include "Widgets/Input/SCheckBox.h"
 #include "Widgets/Layout/SScrollBox.h"
 #include "Widgets/Input/SSearchBox.h"
@@ -15,10 +14,15 @@
 #include "Widgets/SToolTip.h"
 #include "Widgets/Views/STableViewBase.h"
 
-// Insights
+// TraceServices
+#include "TraceServices/Model/NetProfiler.h"
+
+// TraceInsightsCore
+#include "InsightsCore/Table/ViewModels/Table.h"
+#include "InsightsCore/Table/ViewModels/TableColumn.h"
+
+// TraceInsights
 #include "Insights/InsightsStyle.h"
-#include "Insights/Table/ViewModels/Table.h"
-#include "Insights/Table/ViewModels/TableColumn.h"
 #include "Insights/NetworkingProfiler/NetworkingProfilerManager.h"
 #include "Insights/NetworkingProfiler/ViewModels/NetStatsViewColumnFactory.h"
 #include "Insights/NetworkingProfiler/Widgets/SNetStatsViewTooltip.h"
@@ -26,13 +30,16 @@
 #include "Insights/NetworkingProfiler/Widgets/SNetworkingProfilerWindow.h"
 #include "Insights/NetworkingProfiler/Widgets/SPacketContentView.h"
 
-#define LOCTEXT_NAMESPACE "SNetStatsView"
+#define LOCTEXT_NAMESPACE "UE::Insights::NetworkingProfiler::SNetStatsView"
+
+namespace UE::Insights::NetworkingProfiler
+{
 
 ////////////////////////////////////////////////////////////////////////////////////////////////////
 
 SNetStatsView::SNetStatsView()
 	: ProfilerWindowWeakPtr()
-	, Table(MakeShared<Insights::FTable>())
+	, Table(MakeShared<FTable>())
 	, bExpansionSaved(false)
 	, bFilterOutZeroCountEvents(false)
 	, GroupingMode(ENetEventGroupingMode::Flat)
@@ -288,9 +295,9 @@ void SNetStatsView::TreeView_BuildSortByMenu(FMenuBuilder& MenuBuilder)
 {
 	MenuBuilder.BeginSection("SortColumn", LOCTEXT("ContextMenu_Section_SortColumn", "Sort Column"));
 
-	for (const TSharedRef<Insights::FTableColumn>& ColumnRef : Table->GetColumns())
+	for (const TSharedRef<FTableColumn>& ColumnRef : Table->GetColumns())
 	{
-		const Insights::FTableColumn& Column = *ColumnRef;
+		const FTableColumn& Column = *ColumnRef;
 
 		if (Column.IsVisible() && Column.CanBeSorted())
 		{
@@ -357,9 +364,9 @@ void SNetStatsView::TreeView_BuildViewColumnMenu(FMenuBuilder& MenuBuilder)
 {
 	MenuBuilder.BeginSection("Columns", LOCTEXT("ContextMenu_Section_Columns", "Columns"));
 
-	for (const TSharedRef<Insights::FTableColumn>& ColumnRef : Table->GetColumns())
+	for (const TSharedRef<FTableColumn>& ColumnRef : Table->GetColumns())
 	{
-		const Insights::FTableColumn& Column = *ColumnRef;
+		const FTableColumn& Column = *ColumnRef;
 
 		FUIAction Action_ToggleColumn
 		(
@@ -386,12 +393,12 @@ void SNetStatsView::TreeView_BuildViewColumnMenu(FMenuBuilder& MenuBuilder)
 void SNetStatsView::InitializeAndShowHeaderColumns()
 {
 	// Create columns.
-	TArray<TSharedRef<Insights::FTableColumn>> Columns;
+	TArray<TSharedRef<FTableColumn>> Columns;
 	FNetStatsViewColumnFactory::CreateNetStatsViewColumns(Columns);
 	Table->SetColumns(Columns);
 
 	// Show columns.
-	for (const TSharedRef<Insights::FTableColumn>& ColumnRef : Table->GetColumns())
+	for (const TSharedRef<FTableColumn>& ColumnRef : Table->GetColumns())
 	{
 		if (ColumnRef->ShouldBeVisible())
 		{
@@ -404,13 +411,13 @@ void SNetStatsView::InitializeAndShowHeaderColumns()
 
 FText SNetStatsView::GetColumnHeaderText(const FName ColumnId) const
 {
-	const Insights::FTableColumn& Column = *Table->FindColumnChecked(ColumnId);
+	const FTableColumn& Column = *Table->FindColumnChecked(ColumnId);
 	return Column.GetShortName();
 }
 
 ////////////////////////////////////////////////////////////////////////////////////////////////////
 
-TSharedRef<SWidget> SNetStatsView::TreeViewHeaderRow_GenerateColumnMenu(const Insights::FTableColumn& Column)
+TSharedRef<SWidget> SNetStatsView::TreeViewHeaderRow_GenerateColumnMenu(const FTableColumn& Column)
 {
 	const bool bShouldCloseWindowAfterMenuSelection = true;
 	FMenuBuilder MenuBuilder(bShouldCloseWindowAfterMenuSelection, NULL);
@@ -567,7 +574,7 @@ void SNetStatsView::UpdateTree()
 	const double TotalTime = Stopwatch.GetAccumulatedTime();
 	if (TotalTime > 0.1)
 	{
-		UE_LOG(NetworkingProfiler, Log, TEXT("[NetStats] Tree view updated in %.3fs (%d events) --> G:%.3fs + S:%.3fs + F:%.3fs"),
+		UE_LOG(LogNetworkingProfiler, Log, TEXT("[NetStats] Tree view updated in %.3fs (%d events) --> G:%.3fs + S:%.3fs + F:%.3fs"),
 			TotalTime, NetEventNodes.Num(), Time1, Time2 - Time1, TotalTime - Time2);
 	}
 }
@@ -586,11 +593,11 @@ void SNetStatsView::ApplyFiltering()
 		GroupPtr->ClearFilteredChildren();
 		const bool bIsGroupVisible = Filters->PassesAllFilters(GroupPtr);
 
-		const TArray<Insights::FBaseTreeNodePtr>& GroupChildren = GroupPtr->GetChildren();
+		const TArray<FBaseTreeNodePtr>& GroupChildren = GroupPtr->GetChildren();
 		int32 NumVisibleChildren = 0;
-		for (const Insights::FBaseTreeNodePtr& ChildPtr : GroupChildren)
+		for (const FBaseTreeNodePtr& ChildPtr : GroupChildren)
 		{
-			const FNetEventNodePtr& NodePtr = StaticCastSharedPtr<FNetEventNode, Insights::FBaseTreeNode>(ChildPtr);
+			const FNetEventNodePtr& NodePtr = StaticCastSharedPtr<FNetEventNode, FBaseTreeNode>(ChildPtr);
 
 			const bool bIsChildVisible = (!bFilterOutZeroCountEvents || NodePtr->GetAggregatedStats().InstanceCount > 0)
 									  && bNetEventTypeIsVisible[static_cast<int>(NodePtr->GetType())]
@@ -748,11 +755,11 @@ void SNetStatsView::TreeView_OnSelectionChanged(FNetEventNodePtr SelectedItem, E
 
 void SNetStatsView::TreeView_OnGetChildren(FNetEventNodePtr InParent, TArray<FNetEventNodePtr>& OutChildren)
 {
-	const TArray<Insights::FBaseTreeNodePtr>& Children = InParent->GetFilteredChildren();
+	const TArray<FBaseTreeNodePtr>& Children = InParent->GetFilteredChildren();
 	OutChildren.Reset(Children.Num());
-	for (const Insights::FBaseTreeNodePtr& Child : Children)
+	for (const FBaseTreeNodePtr& Child : Children)
 	{
-		OutChildren.Add(StaticCastSharedPtr<FNetEventNode, Insights::FBaseTreeNode>(Child));
+		OutChildren.Add(StaticCastSharedPtr<FNetEventNode, FBaseTreeNode>(Child));
 	}
 }
 
@@ -816,7 +823,7 @@ bool SNetStatsView::TableRow_ShouldBeEnabled(FNetEventNodePtr NodePtr) const
 
 ////////////////////////////////////////////////////////////////////////////////////////////////////
 
-void SNetStatsView::TableRow_SetHoveredCell(TSharedPtr<Insights::FTable> InTablePtr, TSharedPtr<Insights::FTableColumn> InColumnPtr, FNetEventNodePtr InNodePtr)
+void SNetStatsView::TableRow_SetHoveredCell(TSharedPtr<FTable> InTablePtr, TSharedPtr<FTableColumn> InColumnPtr, FNetEventNodePtr InNodePtr)
 {
 	HoveredColumnId = InColumnPtr ? InColumnPtr->GetId() : FName();
 
@@ -1042,11 +1049,11 @@ void SNetStatsView::CreateSortings()
 	AvailableSorters.Reset();
 	CurrentSorter = nullptr;
 
-	for (const TSharedRef<Insights::FTableColumn>& ColumnRef : Table->GetColumns())
+	for (const TSharedRef<FTableColumn>& ColumnRef : Table->GetColumns())
 	{
 		if (ColumnRef->CanBeSorted())
 		{
-			TSharedPtr<Insights::ITableCellValueSorter> SorterPtr = ColumnRef->GetValueSorter();
+			TSharedPtr<ITableCellValueSorter> SorterPtr = ColumnRef->GetValueSorter();
 			if (ensure(SorterPtr.IsValid()))
 			{
 				AvailableSorters.Add(SorterPtr);
@@ -1061,7 +1068,7 @@ void SNetStatsView::CreateSortings()
 
 void SNetStatsView::UpdateCurrentSortingByColumn()
 {
-	TSharedPtr<Insights::FTableColumn> ColumnPtr = Table->FindColumn(ColumnBeingSorted);
+	TSharedPtr<FTableColumn> ColumnPtr = Table->FindColumn(ColumnBeingSorted);
 	CurrentSorter = ColumnPtr.IsValid() ? ColumnPtr->GetValueSorter() : nullptr;
 }
 
@@ -1080,12 +1087,12 @@ void SNetStatsView::SortTreeNodes()
 
 ////////////////////////////////////////////////////////////////////////////////////////////////////
 
-void SNetStatsView::SortTreeNodesRec(FNetEventNode& Node, const Insights::ITableCellValueSorter& Sorter)
+void SNetStatsView::SortTreeNodesRec(FNetEventNode& Node, const ITableCellValueSorter& Sorter)
 {
-	Insights::ESortMode SortMode = (ColumnSortMode == EColumnSortMode::Type::Descending) ? Insights::ESortMode::Descending : Insights::ESortMode::Ascending;
+	ESortMode SortMode = (ColumnSortMode == EColumnSortMode::Type::Descending) ? ESortMode::Descending : ESortMode::Ascending;
 	Node.SortChildren(Sorter, SortMode);
 
-	for (Insights::FBaseTreeNodePtr ChildPtr : Node.GetChildren())
+	for (FBaseTreeNodePtr ChildPtr : Node.GetChildren())
 	{
 		if (ChildPtr->GetChildrenCount() > 0)
 		{
@@ -1139,7 +1146,7 @@ bool SNetStatsView::HeaderMenu_SortMode_IsChecked(const FName ColumnId, const EC
 
 bool SNetStatsView::HeaderMenu_SortMode_CanExecute(const FName ColumnId, const EColumnSortMode::Type InSortMode) const
 {
-	const Insights::FTableColumn& Column = *Table->FindColumnChecked(ColumnId);
+	const FTableColumn& Column = *Table->FindColumnChecked(ColumnId);
 	return Column.CanBeSorted();
 }
 
@@ -1212,7 +1219,7 @@ bool SNetStatsView::CanShowColumn(const FName ColumnId) const
 
 void SNetStatsView::ShowColumn(const FName ColumnId)
 {
-	Insights::FTableColumn& Column = *Table->FindColumnChecked(ColumnId);
+	FTableColumn& Column = *Table->FindColumnChecked(ColumnId);
 	Column.Show();
 
 	SHeaderRow::FColumn::FArguments ColumnArgs;
@@ -1267,7 +1274,7 @@ void SNetStatsView::ShowColumn(const FName ColumnId)
 
 bool SNetStatsView::CanHideColumn(const FName ColumnId) const
 {
-	const Insights::FTableColumn& Column = *Table->FindColumnChecked(ColumnId);
+	const FTableColumn& Column = *Table->FindColumnChecked(ColumnId);
 	return Column.CanBeHidden();
 }
 
@@ -1275,7 +1282,7 @@ bool SNetStatsView::CanHideColumn(const FName ColumnId) const
 
 void SNetStatsView::HideColumn(const FName ColumnId)
 {
-	Insights::FTableColumn& Column = *Table->FindColumnChecked(ColumnId);
+	FTableColumn& Column = *Table->FindColumnChecked(ColumnId);
 	Column.Hide();
 
 	TreeViewHeaderRow->RemoveColumn(ColumnId);
@@ -1287,7 +1294,7 @@ void SNetStatsView::HideColumn(const FName ColumnId)
 
 bool SNetStatsView::IsColumnVisible(const FName ColumnId) const
 {
-	const Insights::FTableColumn& Column = *Table->FindColumnChecked(ColumnId);
+	const FTableColumn& Column = *Table->FindColumnChecked(ColumnId);
 	return Column.IsVisible();
 }
 
@@ -1295,7 +1302,7 @@ bool SNetStatsView::IsColumnVisible(const FName ColumnId) const
 
 bool SNetStatsView::CanToggleColumnVisibility(const FName ColumnId) const
 {
-	const Insights::FTableColumn& Column = *Table->FindColumnChecked(ColumnId);
+	const FTableColumn& Column = *Table->FindColumnChecked(ColumnId);
 	return !Column.IsVisible() || Column.CanBeHidden();
 }
 
@@ -1303,7 +1310,7 @@ bool SNetStatsView::CanToggleColumnVisibility(const FName ColumnId) const
 
 void SNetStatsView::ToggleColumnVisibility(const FName ColumnId)
 {
-	const Insights::FTableColumn& Column = *Table->FindColumnChecked(ColumnId);
+	const FTableColumn& Column = *Table->FindColumnChecked(ColumnId);
 	if (Column.IsVisible())
 	{
 		HideColumn(ColumnId);
@@ -1331,9 +1338,9 @@ void SNetStatsView::ContextMenu_ShowAllColumns_Execute()
 	ColumnSortMode = GetDefaultColumnSortMode();
 	UpdateCurrentSortingByColumn();
 
-	for (const TSharedRef<Insights::FTableColumn>& ColumnRef : Table->GetColumns())
+	for (const TSharedRef<FTableColumn>& ColumnRef : Table->GetColumns())
 	{
-		const Insights::FTableColumn& Column = *ColumnRef;
+		const FTableColumn& Column = *ColumnRef;
 
 		if (!Column.IsVisible())
 		{
@@ -1359,9 +1366,9 @@ void SNetStatsView::ContextMenu_ResetColumns_Execute()
 	ColumnSortMode = GetDefaultColumnSortMode();
 	UpdateCurrentSortingByColumn();
 
-	for (const TSharedRef<Insights::FTableColumn>& ColumnRef : Table->GetColumns())
+	for (const TSharedRef<FTableColumn>& ColumnRef : Table->GetColumns())
 	{
-		const Insights::FTableColumn& Column = *ColumnRef;
+		const FTableColumn& Column = *ColumnRef;
 
 		if (Column.ShouldBeVisible() && !Column.IsVisible())
 		{
@@ -1489,7 +1496,7 @@ void SNetStatsView::RebuildTree(bool bResync)
 	if (TotalTime > 0.01)
 	{
 		const double SyncTime = SyncStopwatch.GetAccumulatedTime();
-		UE_LOG(NetworkingProfiler, Log, TEXT("[NetStats] Tree view rebuilt in %.4fs (sync: %.4fs + update: %.4fs) --> %d net events (%d added)"),
+		UE_LOG(LogNetworkingProfiler, Log, TEXT("[NetStats] Tree view rebuilt in %.4fs (sync: %.4fs + update: %.4fs) --> %d net events (%d added)"),
 			TotalTime, SyncTime, TotalTime - SyncTime, NetEventNodes.Num(), NetEventNodes.Num() - PreviousNodeCount);
 	}
 }
@@ -1596,7 +1603,7 @@ void SNetStatsView::UpdateStatsInternal()
 	Stopwatch.Stop();
 	const double TotalTime = Stopwatch.GetAccumulatedTime();
 	const double AggregationTime = AggregationStopwatch.GetAccumulatedTime();
-	UE_LOG(NetworkingProfiler, Log, TEXT("[NetStats] Aggregated stats updated in %.4fs (%.4fs + %.4fs)"),
+	UE_LOG(LogNetworkingProfiler, Log, TEXT("[NetStats] Aggregated stats updated in %.4fs (%.4fs + %.4fs)"),
 		TotalTime, AggregationTime, TotalTime - AggregationTime);
 }
 
@@ -1621,4 +1628,6 @@ void SNetStatsView::SelectNetEventNode(uint32 EventTypeIndex)
 
 ////////////////////////////////////////////////////////////////////////////////////////////////////
 
+} // namespace UE::Insights::NetworkingProfiler
+
 #undef LOCTEXT_NAMESPACE
diff --git a/Engine/Source/Developer/TraceInsights/Private/Insights/NetworkingProfiler/Widgets/SNetStatsView.h b/Engine/Source/Developer/TraceInsights/Private/Insights/NetworkingProfiler/Widgets/SNetStatsView.h
--- a/Engine/Source/Developer/TraceInsights/Private/Insights/NetworkingProfiler/Widgets/SNetStatsView.h
+++ b/Engine/Source/Developer/TraceInsights/Private/Insights/NetworkingProfiler/Widgets/SNetStatsView.h
@@ -3,6 +3,7 @@
 #pragma once
 
 #include "CoreMinimal.h"
+
 #include "Misc/FilterCollection.h"
 #include "Misc/TextFilter.h"
 #include "SlateFwd.h"
@@ -15,27 +16,31 @@
 #include "Widgets/Views/STableViewBase.h"
 #include "Widgets/Views/STreeView.h"
 
-// Insights
+// TraceInsights
 #include "Insights/NetworkingProfiler/ViewModels/NetEventGroupingAndSorting.h"
 #include "Insights/NetworkingProfiler/ViewModels/NetEventNode.h"
 
 ////////////////////////////////////////////////////////////////////////////////////////////////////
 
 class FMenuBuilder;
-class SNetworkingProfilerWindow;
 
 namespace TraceServices
 {
 	class IAnalysisSession;
 }
 
-namespace Insights
+namespace UE::Insights
 {
 	class FTable;
 	class FTableColumn;
 	class ITableCellValueSorter;
 }
 
+namespace UE::Insights::NetworkingProfiler
+{
+
+class SNetworkingProfilerWindow;
+
 ////////////////////////////////////////////////////////////////////////////////////////////////////
 
 /** The filter collection - used for updating the list of tree nodes. */
@@ -46,7 +51,7 @@ typedef TTextFilter<const FNetEventNodePtr&> FNetEventNodeTextFilter;
 
 ////////////////////////////////////////////////////////////////////////////////////////////////////
 /**
- * A custom widget used to display the list of net events and thier aggregated stats.
+ * A custom widget used to display the list of net events and their aggregated stats.
  */
 class SNetStatsView : public SCompoundWidget
 {
@@ -69,7 +74,7 @@ public:
 		return ProfilerWindowWeakPtr.Pin();
 	}
 
-	TSharedPtr<Insights::FTable> GetTable() const { return Table; }
+	TSharedPtr<FTable> GetTable() const { return Table; }
 
 	/**
 	 * Construct this widget
@@ -131,7 +136,7 @@ protected:
 
 	FText GetColumnHeaderText(const FName ColumnId) const;
 
-	TSharedRef<SWidget> TreeViewHeaderRow_GenerateColumnMenu(const Insights::FTableColumn& Column);
+	TSharedRef<SWidget> TreeViewHeaderRow_GenerateColumnMenu(const FTableColumn& Column);
 
 	////////////////////////////////////////////////////////////////////////////////////////////////////
 	// Tree View - Misc
@@ -159,7 +164,7 @@ protected:
 
 	bool TableRow_ShouldBeEnabled(FNetEventNodePtr NodePtr) const;
 
-	void TableRow_SetHoveredCell(TSharedPtr<Insights::FTable> TablePtr, TSharedPtr<Insights::FTableColumn> ColumnPtr, FNetEventNodePtr NodePtr);
+	void TableRow_SetHoveredCell(TSharedPtr<FTable> TablePtr, TSharedPtr<FTableColumn> ColumnPtr, FNetEventNodePtr NodePtr);
 	EHorizontalAlignment TableRow_GetColumnOutlineHAlignment(const FName ColumnId) const;
 
 	FText TableRow_GetHighlightText() const;
@@ -205,7 +210,7 @@ protected:
 
 	void UpdateCurrentSortingByColumn();
 	void SortTreeNodes();
-	void SortTreeNodesRec(FNetEventNode& Node, const Insights::ITableCellValueSorter& Sorter);
+	void SortTreeNodesRec(FNetEventNode& Node, const ITableCellValueSorter& Sorter);
 
 	EColumnSortMode::Type GetSortModeForColumn(const FName ColumnId) const;
 	void SetSortModeForColumn(const FName& ColumnId, EColumnSortMode::Type SortMode);
@@ -264,7 +269,7 @@ protected:
 	TWeakPtr<SNetworkingProfilerWindow> ProfilerWindowWeakPtr;
 
 	/** Table view model. */
-	TSharedPtr<Insights::FTable> Table;
+	TSharedPtr<FTable> Table;
 
 	/** The analysis session used to populate this widget. */
 	TSharedPtr<const TraceServices::IAnalysisSession> Session;
@@ -349,10 +354,10 @@ protected:
 	//bool bUseSorting;
 
 	/** All available sorters. */
-	TArray<TSharedPtr<Insights::ITableCellValueSorter>> AvailableSorters;
+	TArray<TSharedPtr<ITableCellValueSorter>> AvailableSorters;
 
 	/** Current sorter. It is nullptr if sorting is disabled. */
-	TSharedPtr<Insights::ITableCellValueSorter> CurrentSorter;
+	TSharedPtr<ITableCellValueSorter> CurrentSorter;
 
 	/** Name of the column currently being sorted. Can be NAME_None if sorting is disabled (CurrentSorting == nullptr) or if a complex sorting is used (CurrentSorting != nullptr). */
 	FName ColumnBeingSorted;
@@ -375,3 +380,5 @@ protected:
 };
 
 ////////////////////////////////////////////////////////////////////////////////////////////////////
+
+} // namespace UE::Insights::NetworkingProfiler
diff --git a/Engine/Source/Developer/TraceInsights/Private/Insights/NetworkingProfiler/Widgets/SNetStatsViewTooltip.cpp b/Engine/Source/Developer/TraceInsights/Private/Insights/NetworkingProfiler/Widgets/SNetStatsViewTooltip.cpp
--- a/Engine/Source/Developer/TraceInsights/Private/Insights/NetworkingProfiler/Widgets/SNetStatsViewTooltip.cpp
+++ b/Engine/Source/Developer/TraceInsights/Private/Insights/NetworkingProfiler/Widgets/SNetStatsViewTooltip.cpp
@@ -3,29 +3,36 @@
 #include "SNetStatsViewTooltip.h"
 
 #include "SlateOptMacros.h"
-#include "TraceServices/Model/NetProfiler.h"
 #include "Widgets/Layout/SGridPanel.h"
 #include "Widgets/Layout/SSeparator.h"
 #include "Widgets/SBoxPanel.h"
 #include "Widgets/SToolTip.h"
 #include "Widgets/Text/STextBlock.h"
 
-// Insights
+// TraceServices
+#include "TraceServices/Model/NetProfiler.h"
+
+// TraceInsightsCore
+#include "InsightsCore/Table/ViewModels/Table.h"
+#include "InsightsCore/Table/ViewModels/TableColumn.h"
+
+// TraceInsights
 #include "Insights/InsightsStyle.h"
-#include "Insights/Table/ViewModels/Table.h"
-#include "Insights/Table/ViewModels/TableColumn.h"
 #include "Insights/NetworkingProfiler/ViewModels/NetEventNode.h"
 #include "Insights/NetworkingProfiler/ViewModels/NetEventNodeHelper.h"
 
 ////////////////////////////////////////////////////////////////////////////////////////////////////
 
-#define LOCTEXT_NAMESPACE "SNetStatsView"
+#define LOCTEXT_NAMESPACE "UE::Insights::NetworkingProfiler::SNetStatsView"
+
+namespace UE::Insights::NetworkingProfiler
+{
 
 ////////////////////////////////////////////////////////////////////////////////////////////////////
 
 BEGIN_SLATE_FUNCTION_BUILD_OPTIMIZATION
 
-TSharedPtr<SToolTip> SNetStatsViewTooltip::GetTableTooltip(const Insights::FTable& Table)
+TSharedPtr<SToolTip> SNetStatsViewTooltip::GetTableTooltip(const FTable& Table)
 {
 	TSharedPtr<SToolTip> ColumnTooltip =
 		SNew(SToolTip)
@@ -56,7 +63,7 @@ TSharedPtr<SToolTip> SNetStatsViewTooltip::GetTableTooltip(const Insights::FTabl
 
 ////////////////////////////////////////////////////////////////////////////////////////////////////
 
-TSharedPtr<SToolTip> SNetStatsViewTooltip::GetColumnTooltip(const Insights::FTableColumn& Column)
+TSharedPtr<SToolTip> SNetStatsViewTooltip::GetColumnTooltip(const FTableColumn& Column)
 {
 	TSharedPtr<SToolTip> ColumnTooltip =
 		SNew(SToolTip)
@@ -317,4 +324,6 @@ END_SLATE_FUNCTION_BUILD_OPTIMIZATION
 
 ////////////////////////////////////////////////////////////////////////////////////////////////////
 
+} // namespace UE::Insights::NetworkingProfiler
+
 #undef LOCTEXT_NAMESPACE
diff --git a/Engine/Source/Developer/TraceInsights/Private/Insights/NetworkingProfiler/Widgets/SNetStatsViewTooltip.h b/Engine/Source/Developer/TraceInsights/Private/Insights/NetworkingProfiler/Widgets/SNetStatsViewTooltip.h
--- a/Engine/Source/Developer/TraceInsights/Private/Insights/NetworkingProfiler/Widgets/SNetStatsViewTooltip.h
+++ b/Engine/Source/Developer/TraceInsights/Private/Insights/NetworkingProfiler/Widgets/SNetStatsViewTooltip.h
@@ -8,12 +8,15 @@
 
 class SGridPanel;
 
-namespace Insights
+namespace UE::Insights
 {
 	class FTable;
 	class FTableColumn;
 }
 
+namespace UE::Insights::NetworkingProfiler
+{
+
 class FNetEventNode;
 
 ////////////////////////////////////////////////////////////////////////////////////////////////////
@@ -24,8 +27,8 @@ class SNetStatsViewTooltip
 public:
 	SNetStatsViewTooltip() = delete;
 
-	static TSharedPtr<SToolTip> GetTableTooltip(const Insights::FTable& Table);
-	static TSharedPtr<SToolTip> GetColumnTooltip(const Insights::FTableColumn& Column);
+	static TSharedPtr<SToolTip> GetTableTooltip(const FTable& Table);
+	static TSharedPtr<SToolTip> GetColumnTooltip(const FTableColumn& Column);
 	static TSharedPtr<SToolTip> GetRowTooltip(const TSharedPtr<FNetEventNode> TreeNodePtr);
 
 private:
@@ -83,3 +86,5 @@ private:
 };
 
 ////////////////////////////////////////////////////////////////////////////////////////////////////
+
+} // namespace UE::Insights::NetworkingProfiler
diff --git a/Engine/Source/Developer/TraceInsights/Private/Insights/NetworkingProfiler/Widgets/SNetworkingProfilerToolbar.cpp b/Engine/Source/Developer/TraceInsights/Private/Insights/NetworkingProfiler/Widgets/SNetworkingProfilerToolbar.cpp
--- a/Engine/Source/Developer/TraceInsights/Private/Insights/NetworkingProfiler/Widgets/SNetworkingProfilerToolbar.cpp
+++ b/Engine/Source/Developer/TraceInsights/Private/Insights/NetworkingProfiler/Widgets/SNetworkingProfilerToolbar.cpp
@@ -6,7 +6,7 @@
 #include "Framework/MultiBox/MultiBoxBuilder.h"
 #include "Widgets/SBoxPanel.h"
 
-// Insights
+// TraceInsights
 #include "Insights/InsightsCommands.h"
 #include "Insights/InsightsManager.h"
 #include "Insights/InsightsStyle.h"
@@ -15,7 +15,10 @@
 
 ////////////////////////////////////////////////////////////////////////////////////////////////////
 
-#define LOCTEXT_NAMESPACE "SNetworkingProfilerToolbar"
+#define LOCTEXT_NAMESPACE "UE::Insights::NetworkingProfiler"
+
+namespace UE::Insights::NetworkingProfiler
+{
 
 ////////////////////////////////////////////////////////////////////////////////////////////////////
 
@@ -132,4 +135,6 @@ void SNetworkingProfilerToolbar::Construct(const FArguments& InArgs, TSharedRef<
 
 ////////////////////////////////////////////////////////////////////////////////////////////////////
 
+} // namespace UE::Insights::NetworkingProfiler
+
 #undef LOCTEXT_NAMESPACE
diff --git a/Engine/Source/Developer/TraceInsights/Private/Insights/NetworkingProfiler/Widgets/SNetworkingProfilerToolbar.h b/Engine/Source/Developer/TraceInsights/Private/Insights/NetworkingProfiler/Widgets/SNetworkingProfilerToolbar.h
--- a/Engine/Source/Developer/TraceInsights/Private/Insights/NetworkingProfiler/Widgets/SNetworkingProfilerToolbar.h
+++ b/Engine/Source/Developer/TraceInsights/Private/Insights/NetworkingProfiler/Widgets/SNetworkingProfilerToolbar.h
@@ -6,6 +6,10 @@
 #include "Widgets/SCompoundWidget.h"
 
 class FExtender;
+
+namespace UE::Insights::NetworkingProfiler
+{
+
 class SNetworkingProfilerWindow;
 
 class SNetworkingProfilerToolbar : public SCompoundWidget
@@ -28,3 +32,5 @@ public:
 	 */
 	void Construct(const FArguments& InArgs, TSharedRef<SNetworkingProfilerWindow> InProfilerWindow);
 };
+
+} // namespace UE::Insights::NetworkingProfiler
diff --git a/Engine/Source/Developer/TraceInsights/Private/Insights/NetworkingProfiler/Widgets/SNetworkingProfilerWindow.cpp b/Engine/Source/Developer/TraceInsights/Private/Insights/NetworkingProfiler/Widgets/SNetworkingProfilerWindow.cpp
--- a/Engine/Source/Developer/TraceInsights/Private/Insights/NetworkingProfiler/Widgets/SNetworkingProfilerWindow.cpp
+++ b/Engine/Source/Developer/TraceInsights/Private/Insights/NetworkingProfiler/Widgets/SNetworkingProfilerWindow.cpp
@@ -8,7 +8,7 @@
 #include "SlateOptMacros.h"
 #include "Widgets/Docking/SDockTab.h"
 
-// Insights
+// TraceInsights
 #include "Insights/InsightsStyle.h"
 #include "Insights/NetworkingProfiler/NetworkingProfilerManager.h"
 #include "Insights/NetworkingProfiler/Widgets/SNetStatsView.h"
@@ -19,7 +19,10 @@
 
 ////////////////////////////////////////////////////////////////////////////////////////////////////
 
-#define LOCTEXT_NAMESPACE "SNetworkingProfilerWindow"
+#define LOCTEXT_NAMESPACE "UE::Insights::NetworkingProfiler"
+
+namespace UE::Insights::NetworkingProfiler
+{
 
 ////////////////////////////////////////////////////////////////////////////////////////////////////
 
@@ -977,4 +980,6 @@ void SNetworkingProfilerWindow::UpdateAggregatedNetStats()
 
 ////////////////////////////////////////////////////////////////////////////////////////////////////
 
+} // namespace UE::Insights::NetworkingProfiler
+
 #undef LOCTEXT_NAMESPACE
diff --git a/Engine/Source/Developer/TraceInsights/Private/Insights/NetworkingProfiler/Widgets/SNetworkingProfilerWindow.h b/Engine/Source/Developer/TraceInsights/Private/Insights/NetworkingProfiler/Widgets/SNetworkingProfilerWindow.h
--- a/Engine/Source/Developer/TraceInsights/Private/Insights/NetworkingProfiler/Widgets/SNetworkingProfilerWindow.h
+++ b/Engine/Source/Developer/TraceInsights/Private/Insights/NetworkingProfiler/Widgets/SNetworkingProfilerWindow.h
@@ -3,12 +3,17 @@
 #pragma once
 
 #include "CoreMinimal.h"
-#include "TraceServices/Model/NetProfiler.h"
 #include "Widgets/Input/SComboBox.h"
 
-// Insights
+// TraceServices
+#include "TraceServices/Model/NetProfiler.h"
+
+// TraceInsights
 #include "Insights/Widgets/SMajorTabWindow.h"
 
+namespace UE::Insights::NetworkingProfiler
+{
+
 ////////////////////////////////////////////////////////////////////////////////////////////////////
 
 class SPacketView;
@@ -30,7 +35,7 @@ struct FNetworkingProfilerTabs
 ////////////////////////////////////////////////////////////////////////////////////////////////////
 
 /** Implements the Networking Insights window. */
-class SNetworkingProfilerWindow : public Insights::SMajorTabWindow
+class SNetworkingProfilerWindow : public ::Insights::SMajorTabWindow
 {
 private:
 	struct FGameInstanceItem
@@ -248,3 +253,7 @@ private:
 	uint32 OldGameInstanceChangeCount;
 	uint32 OldConnectionChangeCount;
 };
+
+////////////////////////////////////////////////////////////////////////////////////////////////////
+
+} // namespace UE::Insights::NetworkingProfiler
diff --git a/Engine/Source/Developer/TraceInsights/Private/Insights/NetworkingProfiler/Widgets/SPacketContentView.cpp b/Engine/Source/Developer/TraceInsights/Private/Insights/NetworkingProfiler/Widgets/SPacketContentView.cpp
--- a/Engine/Source/Developer/TraceInsights/Private/Insights/NetworkingProfiler/Widgets/SPacketContentView.cpp
+++ b/Engine/Source/Developer/TraceInsights/Private/Insights/NetworkingProfiler/Widgets/SPacketContentView.cpp
@@ -10,19 +10,21 @@
 #include "Logging/MessageLog.h"
 #include "MessageLogModule.h"
 #include "Rendering/DrawElements.h"
+#include "SlateOptMacros.h"
 #include "Styling/AppStyle.h"
+#include "Templates/SharedPointer.h"
 #include "Widgets/Images/SImage.h"
-#include "Widgets/Input/SCheckBox.h"
 #include "Widgets/Input/SButton.h"
+#include "Widgets/Input/SCheckBox.h"
 #include "Widgets/Input/SEditableTextBox.h"
 #include "Widgets/Layout/SScrollBar.h"
-#include "SlateOptMacros.h"
-#include "Templates/SharedPointer.h"
 
-// Insights
-#include "Insights/Common/PaintUtils.h"
-#include "Insights/Common/TimeUtils.h"
-#include "Insights/Common/Stopwatch.h"
+// TraceInsightsCore
+#include "InsightsCore/Common/PaintUtils.h"
+#include "InsightsCore/Common/Stopwatch.h"
+#include "InsightsCore/Common/TimeUtils.h"
+
+// TraceInsights
 #include "Insights/InsightsManager.h"
 #include "Insights/InsightsStyle.h"
 #include "Insights/NetworkingProfiler/NetworkingProfilerManager.h"
@@ -31,7 +33,10 @@
 
 ////////////////////////////////////////////////////////////////////////////////////////////////////
 
-#define LOCTEXT_NAMESPACE "SPacketContentView"
+#define LOCTEXT_NAMESPACE "UE::Insights::NetworkingProfiler::SPacketContentView"
+
+namespace UE::Insights::NetworkingProfiler
+{
 
 ////////////////////////////////////////////////////////////////////////////////////////////////////
 
@@ -1172,7 +1177,7 @@ void SPacketContentView::UpdateHoveredEvent()
 			Tooltip.AddNameValueTextLine(TEXT("Net Id:"), FText::AsNumber(ObjectInstance.NetObjectId).ToString());
 			Tooltip.AddNameValueTextLine(TEXT("Type Id:"), FString::Printf(TEXT("0x%016" UINT64_x_FMT), ObjectInstance.TypeId));
 			Tooltip.AddNameValueTextLine(TEXT("Obj. LifeTime:"), FString::Format(TEXT("from {0} to {1}"),
-				{ TimeUtils::FormatTimeAuto(ObjectInstance.LifeTime.Begin), TimeUtils::FormatTimeAuto(ObjectInstance.LifeTime.End) }));
+				{ FormatTimeAuto(ObjectInstance.LifeTime.Begin), FormatTimeAuto(ObjectInstance.LifeTime.End)}));
 		}
 
 		Tooltip.AddNameValueTextLine(TEXT("Offset:"), FString::Format(TEXT("bit {0}"), { FText::AsNumber(Event.BitOffset).ToString() }));
@@ -1798,4 +1803,6 @@ void SPacketContentView::AdjustForSplitContent()
 
 ////////////////////////////////////////////////////////////////////////////////////////////////////
 
+} // namespace UE::Insights::NetworkingProfiler
+
 #undef LOCTEXT_NAMESPACE
diff --git a/Engine/Source/Developer/TraceInsights/Private/Insights/NetworkingProfiler/Widgets/SPacketContentView.h b/Engine/Source/Developer/TraceInsights/Private/Insights/NetworkingProfiler/Widgets/SPacketContentView.h
--- a/Engine/Source/Developer/TraceInsights/Private/Insights/NetworkingProfiler/Widgets/SPacketContentView.h
+++ b/Engine/Source/Developer/TraceInsights/Private/Insights/NetworkingProfiler/Widgets/SPacketContentView.h
@@ -3,23 +3,32 @@
 #pragma once
 
 #include "CoreMinimal.h"
+
 #include "Input/CursorReply.h"
 #include "Input/Reply.h"
 #include "Layout/Geometry.h"
 #include "Rendering/RenderingCommon.h"
 #include "Styling/SlateTypes.h"
-#include "TraceServices/Model/NetProfiler.h"
 #include "Widgets/DeclarativeSyntaxSupport.h"
-#include "Widgets/SCompoundWidget.h"
 #include "Widgets/Input/SComboBox.h"
+#include "Widgets/SCompoundWidget.h"
+
+// TraceServices
+#include "TraceServices/Model/NetProfiler.h"
+
+// TraceInsightsCore
+#include "InsightsCore/Common/FixedCircularBuffer.h"
 
-// Insights
-#include "Insights/Common/FixedCircularBuffer.h"
+// TraceInsights
 #include "Insights/NetworkingProfiler/ViewModels/PacketContentViewDrawHelper.h"
 #include "Insights/NetworkingProfiler/ViewModels/PacketContentViewport.h"
 #include "Insights/ViewModels/TooltipDrawState.h"
 
 class SScrollBar;
+
+namespace UE::Insights::NetworkingProfiler
+{
+
 class SNetworkingProfilerWindow;
 
 ////////////////////////////////////////////////////////////////////////////////////////////////////
@@ -310,3 +319,7 @@ private:
 	mutable TFixedCircularBuffer<uint64, 32> OnPaintDurationHistory;
 	mutable uint64 LastOnPaintTime;
 };
+
+////////////////////////////////////////////////////////////////////////////////////////////////////
+
+} // namespace UE::Insights::NetworkingProfiler
diff --git a/Engine/Source/Developer/TraceInsights/Private/Insights/NetworkingProfiler/Widgets/SPacketView.cpp b/Engine/Source/Developer/TraceInsights/Private/Insights/NetworkingProfiler/Widgets/SPacketView.cpp
--- a/Engine/Source/Developer/TraceInsights/Private/Insights/NetworkingProfiler/Widgets/SPacketView.cpp
+++ b/Engine/Source/Developer/TraceInsights/Private/Insights/NetworkingProfiler/Widgets/SPacketView.cpp
@@ -8,15 +8,19 @@
 #include "Framework/MultiBox/MultiBoxBuilder.h"
 #include "HAL/PlatformTime.h"
 #include "Rendering/DrawElements.h"
-#include "TraceServices/Model/Frames.h"
-#include "TraceServices/Model/NetProfiler.h"
 #include "Styling/AppStyle.h"
 #include "Widgets/Layout/SScrollBar.h"
 
-// Insights
-#include "Insights/Common/PaintUtils.h"
-#include "Insights/Common/Stopwatch.h"
-#include "Insights/Common/TimeUtils.h"
+// TraceServices
+#include "TraceServices/Model/Frames.h"
+#include "TraceServices/Model/NetProfiler.h"
+
+// TraceInsightsCore
+#include "InsightsCore/Common/PaintUtils.h"
+#include "InsightsCore/Common/Stopwatch.h"
+#include "InsightsCore/Common/TimeUtils.h"
+
+// TraceInsights
 #include "Insights/InsightsManager.h"
 #include "Insights/InsightsStyle.h"
 #include "Insights/NetworkingProfiler/NetworkingProfilerManager.h"
@@ -27,7 +31,10 @@
 
 ////////////////////////////////////////////////////////////////////////////////////////////////////
 
-#define LOCTEXT_NAMESPACE "SPacketView"
+#define LOCTEXT_NAMESPACE "UE::Insights::NetworkingProfiler::SPacketView"
+
+namespace UE::Insights::NetworkingProfiler
+{
 
 ////////////////////////////////////////////////////////////////////////////////////////////////////
 
@@ -837,8 +844,8 @@ int32 SPacketView::OnPaint(const FPaintArgs& Args, const FGeometry& AllottedGeom
 						FText::AsNumber(HoveredSample.Sample->LargestPacket.ContentSizeInBits).ToString(),
 						FText::AsNumber(HoveredSample.Sample->LargestPacket.TotalSizeInBytes).ToString(),
 						FText::AsNumber(UnusedBits).ToString(),
-						TimeUtils::FormatTimeHMS(HoveredSample.Sample->LargestPacket.TimeStamp, Precision),
-						::StatusToString(HoveredSample.Sample->LargestPacket.Status),
+						FormatTimeHMS(HoveredSample.Sample->LargestPacket.TimeStamp, Precision),
+						StatusToString(HoveredSample.Sample->LargestPacket.Status),
 						LexToString(HoveredSample.Sample->LargestPacket.ConnectionState),
 						EngineFrameNumber > 0 ?
 						FText::AsNumber(EngineFrameNumber).ToString() :
@@ -851,7 +858,7 @@ int32 SPacketView::OnPaint(const FPaintArgs& Args, const FGeometry& AllottedGeom
 				Text = FString::Format(TEXT("{0} network packets\n"
 											"({1})\n"
 											"Largest Packet\n"
-											"    Sequance Number: {2}\n"
+											"    Sequence Number: {2}\n"
 											"    Content Size: {3} bits\n"
 											"    Total Size: {4} bytes ({5} unused bits)\n"
 											"    Timestamp: {6}\n"
@@ -860,13 +867,13 @@ int32 SPacketView::OnPaint(const FPaintArgs& Args, const FGeometry& AllottedGeom
 											"    Engine Frame Number: {9}"),
 					{
 						HoveredSample.Sample->NumPackets,
-						::AggregatedStatusToString(HoveredSample.Sample->AggregatedStatus),
+						AggregatedStatusToString(HoveredSample.Sample->AggregatedStatus),
 						FText::AsNumber(HoveredSample.Sample->LargestPacket.SequenceNumber).ToString(),
 						FText::AsNumber(HoveredSample.Sample->LargestPacket.ContentSizeInBits).ToString(),
 						FText::AsNumber(HoveredSample.Sample->LargestPacket.TotalSizeInBytes).ToString(),
 						FText::AsNumber(UnusedBits).ToString(),
-						TimeUtils::FormatTimeHMS(HoveredSample.Sample->LargestPacket.TimeStamp, Precision),
-						::StatusToString(HoveredSample.Sample->LargestPacket.Status),
+						FormatTimeHMS(HoveredSample.Sample->LargestPacket.TimeStamp, Precision),
+						StatusToString(HoveredSample.Sample->LargestPacket.Status),
 						LexToString(HoveredSample.Sample->LargestPacket.ConnectionState),
 						EngineFrameNumber > 0 ?
 						FText::AsNumber(EngineFrameNumber).ToString() :
@@ -1859,4 +1866,6 @@ void SPacketView::ZoomHorizontally(const float Delta, const float X)
 
 ////////////////////////////////////////////////////////////////////////////////////////////////////
 
+} // namespace UE::Insights::NetworkingProfiler
+
 #undef LOCTEXT_NAMESPACE
diff --git a/Engine/Source/Developer/TraceInsights/Private/Insights/NetworkingProfiler/Widgets/SPacketView.h b/Engine/Source/Developer/TraceInsights/Private/Insights/NetworkingProfiler/Widgets/SPacketView.h
--- a/Engine/Source/Developer/TraceInsights/Private/Insights/NetworkingProfiler/Widgets/SPacketView.h
+++ b/Engine/Source/Developer/TraceInsights/Private/Insights/NetworkingProfiler/Widgets/SPacketView.h
@@ -3,20 +3,29 @@
 #pragma once
 
 #include "CoreMinimal.h"
+
 #include "Input/CursorReply.h"
 #include "Input/Reply.h"
 #include "Layout/Geometry.h"
 #include "Rendering/RenderingCommon.h"
-#include "TraceServices/Model/NetProfiler.h"
 #include "Widgets/DeclarativeSyntaxSupport.h"
 #include "Widgets/SCompoundWidget.h"
 
-// Insights
-#include "Insights/Common/FixedCircularBuffer.h"
+// TraceServices
+#include "TraceServices/Model/NetProfiler.h"
+
+// TraceInsightsCore
+#include "InsightsCore/Common/FixedCircularBuffer.h"
+
+// TraceInsights
 #include "Insights/NetworkingProfiler/ViewModels/PacketViewDrawHelper.h"
 #include "Insights/NetworkingProfiler/ViewModels/PacketViewport.h"
 
 class SScrollBar;
+
+namespace UE::Insights::NetworkingProfiler
+{
+
 class SNetworkingProfilerWindow;
 
 ////////////////////////////////////////////////////////////////////////////////////////////////////
@@ -254,3 +263,5 @@ private:
 	mutable TFixedCircularBuffer<uint64, 32> OnPaintDurationHistory;
 	mutable uint64 LastOnPaintTime;
 };
+
+} // namespace UE::Insights::NetworkingProfiler
diff --git a/Engine/Source/Developer/TraceInsights/Private/Insights/StoreService/DiagnosticsSessionAnalyzer.cpp b/Engine/Source/Developer/TraceInsights/Private/Insights/StoreService/DiagnosticsSessionAnalyzer.cpp
--- a/Engine/Source/Developer/TraceInsights/Private/Insights/StoreService/DiagnosticsSessionAnalyzer.cpp
+++ b/Engine/Source/Developer/TraceInsights/Private/Insights/StoreService/DiagnosticsSessionAnalyzer.cpp
@@ -3,6 +3,8 @@
 #include "DiagnosticsSessionAnalyzer.h"
 
 #include "HAL/LowLevelMemTracker.h"
+
+// TraceServices
 #include "TraceServices/Model/Diagnostics.h"
 
 namespace Insights
diff --git a/Engine/Source/Developer/TraceInsights/Private/Insights/StoreService/DiagnosticsSessionAnalyzer.h b/Engine/Source/Developer/TraceInsights/Private/Insights/StoreService/DiagnosticsSessionAnalyzer.h
--- a/Engine/Source/Developer/TraceInsights/Private/Insights/StoreService/DiagnosticsSessionAnalyzer.h
+++ b/Engine/Source/Developer/TraceInsights/Private/Insights/StoreService/DiagnosticsSessionAnalyzer.h
@@ -3,6 +3,8 @@
 #pragma once
 
 #include "CoreMinimal.h"
+
+// TraceAnalysis
 #include "Trace/Analyzer.h"
 
 namespace Insights
diff --git a/Engine/Source/Developer/TraceInsights/Private/Insights/StoreService/StoreBrowser.cpp b/Engine/Source/Developer/TraceInsights/Private/Insights/StoreService/StoreBrowser.cpp
--- a/Engine/Source/Developer/TraceInsights/Private/Insights/StoreService/StoreBrowser.cpp
+++ b/Engine/Source/Developer/TraceInsights/Private/Insights/StoreService/StoreBrowser.cpp
@@ -4,11 +4,15 @@
 
 #include "Async/ParallelFor.h"
 #include "Containers/ContainerAllocationPolicies.h"
+
+// TraceAnalysis
 #include "Trace/Analysis.h"
 #include "Trace/StoreClient.h"
 
-// Insights
-#include "Insights/Common/Stopwatch.h"
+// TraceInsightsCore
+#include "InsightsCore/Common/Stopwatch.h"
+
+// TraceInsights
 #include "Insights/InsightsManager.h"
 #include "Insights/Log.h"
 #include "Insights/StoreService/DiagnosticsSessionAnalyzer.h"
@@ -129,7 +133,7 @@ void FStoreBrowser::UpdateTraces()
 		return;
 	}
 
-	FStopwatch StopwatchTotal;
+	UE::Insights::FStopwatch StopwatchTotal;
 	StopwatchTotal.Start();
 
 	// Update store host.
@@ -469,7 +473,7 @@ void FStoreBrowser::UpdateTraces()
 
 	// Check to see if we need to update metadata.
 	{
-		FStopwatch Stopwatch;
+		UE::Insights::FStopwatch Stopwatch;
 		Stopwatch.Start();
 
 		TArray<TSharedPtr<FStoreBrowserTraceInfo>> TracesToUpdate;
@@ -612,7 +616,7 @@ void FStoreBrowser::UpdateMetadata(TSharedPtr<FStoreBrowserTraceInfo> TraceInfoP
 		UE::Trace::IInDataStream* Inner;
 		FStoreBrowserTraceInfo* Trace;
 		double TimeLimit = 1.0;
-		FStopwatch Stopwatch;
+		UE::Insights::FStopwatch Stopwatch;
 		int32 BytesRead = 0;
 		EReadStatus Status = EReadStatus::Ready;
 	};
diff --git a/Engine/Source/Developer/TraceInsights/Private/Insights/Table/ViewModels/UntypedTable.cpp b/Engine/Source/Developer/TraceInsights/Private/Insights/Table/ViewModels/UntypedTable.cpp
--- a/Engine/Source/Developer/TraceInsights/Private/Insights/Table/ViewModels/UntypedTable.cpp
+++ b/Engine/Source/Developer/TraceInsights/Private/Insights/Table/ViewModels/UntypedTable.cpp
@@ -1,14 +1,16 @@
 // Copyright Epic Games, Inc. All Rights Reserved.
 
 #include "UntypedTable.h"
+
+// TraceServices
 #include "TraceServices/Containers/Tables.h"
 
-// Insights
-#include "Insights/Table/ViewModels/TableCellValueFormatter.h"
-#include "Insights/Table/ViewModels/TableCellValueGetter.h"
-#include "Insights/Table/ViewModels/TableCellValueSorter.h"
-#include "Insights/Table/ViewModels/TableColumn.h"
-#include "Insights/Table/ViewModels/TableTreeNode.h"
+// TraceInsightsCore
+#include "InsightsCore/Table/ViewModels/TableCellValueFormatter.h"
+#include "InsightsCore/Table/ViewModels/TableCellValueGetter.h"
+#include "InsightsCore/Table/ViewModels/TableCellValueSorter.h"
+#include "InsightsCore/Table/ViewModels/TableColumn.h"
+#include "InsightsCore/Table/ViewModels/TableTreeNode.h"
 
 #define LOCTEXT_NAMESPACE "Insights::FUntypedTable"
 
@@ -19,13 +21,15 @@ namespace Insights
 // FUntypedTableTreeNodeValueGetter
 ////////////////////////////////////////////////////////////////////////////////////////////////////
 
-class FUntypedTableTreeNodeValueGetter : public FTableCellValueGetter
+class FUntypedTableTreeNodeValueGetter : public UE::Insights::FTableCellValueGetter
 {
 public:
-	FUntypedTableTreeNodeValueGetter(ETableCellDataType InDataType) : FTableCellValueGetter(), DataType(InDataType) {}
+	FUntypedTableTreeNodeValueGetter(UE::Insights::ETableCellDataType InDataType) : UE::Insights::FTableCellValueGetter(), DataType(InDataType) {}
 
-	virtual const TOptional<FTableCellValue> GetValue(const FTableColumn& Column, const FBaseTreeNode& Node) const
+	virtual const TOptional<UE::Insights::FTableCellValue> GetValue(const UE::Insights::FTableColumn& Column, const UE::Insights::FBaseTreeNode& Node) const
 	{
+		using namespace UE::Insights;
+
 		ensure(Node.Is<FTableTreeNode>());
 		const FTableTreeNode& TableTreeNode = static_cast<const FTableTreeNode&>(Node);
 
@@ -68,7 +72,7 @@ public:
 	}
 
 private:
-	ETableCellDataType DataType;
+	UE::Insights::ETableCellDataType DataType;
 };
 
 ////////////////////////////////////////////////////////////////////////////////////////////////////
@@ -158,6 +162,8 @@ bool FUntypedTable::UpdateSourceTable(TSharedPtr<TraceServices::IUntypedTable> I
 
 void FUntypedTable::CreateColumns(const TraceServices::ITableLayout& TableLayout)
 {
+	using namespace UE::Insights;
+
 	ensure(GetColumnCount() == 0);
 	const int32 ColumnCount = static_cast<int32>(TableLayout.GetColumnCount());
 
diff --git a/Engine/Source/Developer/TraceInsights/Private/Insights/Table/ViewModels/UntypedTable.h b/Engine/Source/Developer/TraceInsights/Private/Insights/Table/ViewModels/UntypedTable.h
--- a/Engine/Source/Developer/TraceInsights/Private/Insights/Table/ViewModels/UntypedTable.h
+++ b/Engine/Source/Developer/TraceInsights/Private/Insights/Table/ViewModels/UntypedTable.h
@@ -2,27 +2,26 @@
 
 #pragma once
 
-#include "CoreMinimal.h"
+#include "CoreTypes.h"
 
-#include "Insights/Table/ViewModels/Table.h"
+#include "Templates/SharedPointer.h"
+
+// TraceInsightsCore
+#include "InsightsCore/Table/ViewModels/Table.h"
 
 namespace TraceServices
 {
-
-class ITableLayout;
-class IUntypedTable;
-class IUntypedTableReader;
-
+	class ITableLayout;
+	class IUntypedTable;
+	class IUntypedTableReader;
 }
 
 namespace Insights
 {
 
-class FTableColumn;
-
 ////////////////////////////////////////////////////////////////////////////////////////////////////
 
-class FUntypedTable : public FTable
+class FUntypedTable : public UE::Insights::FTable
 {
 public:
 	FUntypedTable();
diff --git a/Engine/Source/Developer/TraceInsights/Private/Insights/Table/Widgets/SSessionTableTreeView.cpp b/Engine/Source/Developer/TraceInsights/Private/Insights/Table/Widgets/SSessionTableTreeView.cpp
--- a/Engine/Source/Developer/TraceInsights/Private/Insights/Table/Widgets/SSessionTableTreeView.cpp
+++ b/Engine/Source/Developer/TraceInsights/Private/Insights/Table/Widgets/SSessionTableTreeView.cpp
@@ -2,11 +2,13 @@
 
 #include "SSessionTableTreeView.h"
 
-// Insights
+// TraceInsights
 #include "Insights/InsightsManager.h"
 
 #define LOCTEXT_NAMESPACE "Insights::SSessionTableTreeView"
 
+namespace UE
+{
 namespace Insights
 {
 
@@ -29,7 +31,7 @@ SSessionTableTreeView::~SSessionTableTreeView()
 
 void SSessionTableTreeView::ConstructWidget(TSharedPtr<FTable> InTablePtr)
 {
-	STableTreeView::ConstructWidget(InTablePtr);
+	UE::Insights::STableTreeView::ConstructWidget(InTablePtr);
 
 	// Register ourselves with the Insights manager.
 	FInsightsManager::Get()->GetSessionChangedEvent().AddSP(this, &SSessionTableTreeView::InsightsManager_OnSessionChanged);
@@ -60,5 +62,6 @@ void SSessionTableTreeView::InsightsManager_OnSessionChanged()
 ////////////////////////////////////////////////////////////////////////////////////////////////////
 
 } // namespace Insights
+} // namespace UE
 
 #undef LOCTEXT_NAMESPACE
diff --git a/Engine/Source/Developer/TraceInsights/Private/Insights/Table/Widgets/SSessionTableTreeView.h b/Engine/Source/Developer/TraceInsights/Private/Insights/Table/Widgets/SSessionTableTreeView.h
--- a/Engine/Source/Developer/TraceInsights/Private/Insights/Table/Widgets/SSessionTableTreeView.h
+++ b/Engine/Source/Developer/TraceInsights/Private/Insights/Table/Widgets/SSessionTableTreeView.h
@@ -4,21 +4,21 @@
 
 #include "CoreMinimal.h"
 
-// Insights
-#include "Insights/Table/Widgets/STableTreeView.h"
+// TraceInsightsCore
+#include "InsightsCore/Table/Widgets/STableTreeView.h"
 
 namespace TraceServices
 {
 	class IAnalysisSession;
 }
 
+namespace UE
+{
 namespace Insights
 {
 
 class FTable;
 
-////////////////////////////////////////////////////////////////////////////////////////////////////
-
 class SSessionTableTreeView : public STableTreeView
 {
 public:
@@ -34,6 +34,5 @@ protected:
 	TSharedPtr<const TraceServices::IAnalysisSession> Session;
 };
 
-////////////////////////////////////////////////////////////////////////////////////////////////////
-
 } // namespace Insights
+} // namespace UE
diff --git a/Engine/Source/Developer/TraceInsights/Private/Insights/Table/Widgets/SUntypedDiffTableTreeView.cpp b/Engine/Source/Developer/TraceInsights/Private/Insights/Table/Widgets/SUntypedDiffTableTreeView.cpp
--- a/Engine/Source/Developer/TraceInsights/Private/Insights/Table/Widgets/SUntypedDiffTableTreeView.cpp
+++ b/Engine/Source/Developer/TraceInsights/Private/Insights/Table/Widgets/SUntypedDiffTableTreeView.cpp
@@ -4,8 +4,13 @@
 
 #include "SlateOptMacros.h"
 #include "Logging/MessageLog.h"
+
+// TraceServices
 #include "TraceServices/Model/TableMerge.h"
 
+// TraceInsights
+#include "Insights/InsightsManager.h"
+
 #define LOCTEXT_NAMESPACE "SUntypedDiffTableTreeView"
 
 namespace Insights
diff --git a/Engine/Source/Developer/TraceInsights/Private/Insights/Table/Widgets/SUntypedTableTreeView.cpp b/Engine/Source/Developer/TraceInsights/Private/Insights/Table/Widgets/SUntypedTableTreeView.cpp
--- a/Engine/Source/Developer/TraceInsights/Private/Insights/Table/Widgets/SUntypedTableTreeView.cpp
+++ b/Engine/Source/Developer/TraceInsights/Private/Insights/Table/Widgets/SUntypedTableTreeView.cpp
@@ -2,11 +2,14 @@
 
 #include "SUntypedTableTreeView.h"
 
+// TraceServices
 #include "TraceServices/Containers/Tables.h"
 #include "TraceServices/Model/Threads.h"
 
-// Insights
-#include "Insights/Common/Stopwatch.h"
+// TraceInsightsCore
+#include "InsightsCore/Common/Stopwatch.h"
+
+// TraceInsights
 #include "Insights/Log.h"
 #include "Insights/Table/ViewModels/UntypedTable.h"
 
@@ -62,10 +65,10 @@ void SUntypedTableTreeView::UpdateSourceTable(TSharedPtr<TraceServices::IUntyped
 
 void SUntypedTableTreeView::RebuildTree(bool bResync)
 {
-	FStopwatch Stopwatch;
+	UE::Insights::FStopwatch Stopwatch;
 	Stopwatch.Start();
 
-	FStopwatch SyncStopwatch;
+	UE::Insights::FStopwatch SyncStopwatch;
 	SyncStopwatch.Start();
 
 	if (bResync)
@@ -92,7 +95,7 @@ void SUntypedTableTreeView::RebuildTree(bool bResync)
 			{
 				TableReader->SetRowIndex(RowIndex);
 				FName NodeName(BaseNodeName, RowIndex + 1);
-				FTableTreeNodePtr NodePtr = MakeShared<FTableTreeNode>(NodeName, Table, RowIndex);
+				UE::Insights::FTableTreeNodePtr NodePtr = MakeShared<UE::Insights::FTableTreeNode>(NodeName, Table, RowIndex);
 				NodePtr->SetDefaultSortOrder(RowIndex + 1);
 				TableRowNodes.Add(NodePtr);
 			}
@@ -105,7 +108,7 @@ void SUntypedTableTreeView::RebuildTree(bool bResync)
 	if (bResync || TableRowNodes.Num() != PreviousNodeCount)
 	{
 		// Save selection.
-		TArray<FTableTreeNodePtr> SelectedItems;
+		TArray<UE::Insights::FTableTreeNodePtr> SelectedItems;
 		TreeView->GetSelectedItems(SelectedItems);
 
 		UpdateTree();
@@ -116,11 +119,11 @@ void SUntypedTableTreeView::RebuildTree(bool bResync)
 		if (SelectedItems.Num() > 0)
 		{
 			TreeView->ClearSelection();
-			for (FTableTreeNodePtr& NodePtr : SelectedItems)
+			for (UE::Insights::FTableTreeNodePtr& NodePtr : SelectedItems)
 			{
 				NodePtr = GetNodeByTableRowIndex(NodePtr->GetRowIndex());
 			}
-			SelectedItems.RemoveAll([](const FTableTreeNodePtr& NodePtr) { return !NodePtr.IsValid(); });
+			SelectedItems.RemoveAll([](const UE::Insights::FTableTreeNodePtr& NodePtr) { return !NodePtr.IsValid(); });
 			if (SelectedItems.Num() > 0)
 			{
 				TreeView->SetItemSelection(SelectedItems, true);
diff --git a/Engine/Source/Developer/TraceInsights/Private/Insights/Table/Widgets/SUntypedTableTreeView.h b/Engine/Source/Developer/TraceInsights/Private/Insights/Table/Widgets/SUntypedTableTreeView.h
--- a/Engine/Source/Developer/TraceInsights/Private/Insights/Table/Widgets/SUntypedTableTreeView.h
+++ b/Engine/Source/Developer/TraceInsights/Private/Insights/Table/Widgets/SUntypedTableTreeView.h
@@ -4,17 +4,19 @@
 
 #include "CoreMinimal.h"
 
-// Insights
-#include "Insights/Common/Stopwatch.h"
+// TraceInsightsCore
+#include "InsightsCore/Common/Stopwatch.h"
+#include "InsightsCore/Table/Widgets/STableTreeView.h"
+
+// TraceInsights
 #include "Insights/Table/ViewModels/UntypedTable.h"
-#include "Insights/Table/Widgets/STableTreeView.h"
 
 namespace Insights
 {
 
 ////////////////////////////////////////////////////////////////////////////////////////////////////
 
-class SUntypedTableTreeView : public STableTreeView
+class SUntypedTableTreeView : public UE::Insights::STableTreeView
 {
 public:
 	/** Default constructor. */
@@ -61,7 +63,7 @@ public:
 	virtual void RebuildTree(bool bResync);
 
 private:
-	FStopwatch CurrentOperationStopwatch;
+	UE::Insights::FStopwatch CurrentOperationStopwatch;
 	FText CurrentOperationNameOverride;
 };
 
diff --git a/Engine/Source/Developer/TraceInsights/Private/Insights/TaskGraphProfiler/TaskGraphProfilerManager.cpp b/Engine/Source/Developer/TraceInsights/Private/Insights/TaskGraphProfiler/TaskGraphProfilerManager.cpp
--- a/Engine/Source/Developer/TraceInsights/Private/Insights/TaskGraphProfiler/TaskGraphProfilerManager.cpp
+++ b/Engine/Source/Developer/TraceInsights/Private/Insights/TaskGraphProfiler/TaskGraphProfilerManager.cpp
@@ -7,10 +7,12 @@
 #include "Framework/Docking/TabManager.h"
 #include "Logging/MessageLog.h"
 #include "Modules/ModuleManager.h"
-#include "TraceServices/Model/TasksProfiler.h"
 #include "Widgets/Docking/SDockTab.h"
 
-// Insights
+// TraceServices
+#include "TraceServices/Model/TasksProfiler.h"
+
+// TraceInsights
 #include "Insights/InsightsStyle.h"
 #include "Insights/TaskGraphProfiler/ViewModels/TaskGraphRelation.h"
 #include "Insights/TaskGraphProfiler/ViewModels/TaskTable.h"
@@ -118,7 +120,7 @@ FTaskGraphProfilerManager::~FTaskGraphProfilerManager()
 
 	if (TaskTimingSharedState.IsValid())
 	{
-		IModularFeatures::Get().UnregisterModularFeature(Insights::TimingViewExtenderFeatureName, TaskTimingSharedState.Get());
+		IModularFeatures::Get().UnregisterModularFeature(UE::Insights::Timing::TimingViewExtenderFeatureName, TaskTimingSharedState.Get());
 	}
 }
 
@@ -152,13 +154,7 @@ bool FTaskGraphProfilerManager::Tick(float DeltaTime)
 			TSharedPtr<FTabManager> TabManagerShared = TimingTabManager.Pin();
 			if (TasksProvider && TasksProvider->GetNumTasks() > 0 && TabManagerShared.IsValid())
 			{
-				TSharedPtr<STimingProfilerWindow> Window = FTimingProfilerManager::Get()->GetProfilerWindow();
-				if (!Window.IsValid())
-				{
-					return true;
-				}
-
-				TSharedPtr<STimingView> TimingView = Window->GetTimingView();
+				TSharedPtr<UE::Insights::TimingProfiler::STimingView> TimingView = GetTimingView();
 				if (!TimingView.IsValid())
 				{
 					return true;
@@ -169,7 +165,7 @@ bool FTaskGraphProfilerManager::Tick(float DeltaTime)
 				if (!TaskTimingSharedState.IsValid())
 				{
 					TaskTimingSharedState = MakeShared<FTaskTimingSharedState>(TimingView.Get());
-					IModularFeatures::Get().RegisterModularFeature(Insights::TimingViewExtenderFeatureName, TaskTimingSharedState.Get());
+					IModularFeatures::Get().RegisterModularFeature(UE::Insights::Timing::TimingViewExtenderFeatureName, TaskTimingSharedState.Get());
 				}
 				TabManagerShared->TryInvokeTab(FTaskGraphProfilerTabs::TaskTableTreeViewTabID);
 			}
@@ -417,13 +413,7 @@ void FTaskGraphProfilerManager::GetRelationsOnCriticalPath(const TraceServices::
 			return;
 		}
 
-		TSharedPtr<class STimingProfilerWindow> TimingWindow = FTimingProfilerManager::Get()->GetProfilerWindow();
-		if (!TimingWindow.IsValid())
-		{
-			return;
-		}
-
-		TSharedPtr<STimingView> TimingView = TimingWindow->GetTimingView();
+		TSharedPtr<UE::Insights::TimingProfiler::STimingView> TimingView = GetTimingView();
 		if (!TimingView.IsValid())
 		{
 			return;
@@ -575,13 +565,7 @@ void FTaskGraphProfilerManager::AddRelation(const FThreadTrackEvent* InSelectedE
 		return;
 	}
 
-	TSharedPtr<class STimingProfilerWindow> TimingWindow = FTimingProfilerManager::Get()->GetProfilerWindow();
-	if (!TimingWindow.IsValid())
-	{
-		return;
-	}
-
-	TSharedPtr<STimingView> TimingView = TimingWindow->GetTimingView();
+	TSharedPtr<UE::Insights::TimingProfiler::STimingView> TimingView = GetTimingView();
 	if (!TimingView.IsValid())
 	{
 		return;
@@ -648,13 +632,7 @@ int32 FTaskGraphProfilerManager::GetRelationDisplayDepth(TSharedPtr<const FThrea
 
 void FTaskGraphProfilerManager::ClearTaskRelations()
 {
-	TSharedPtr<class STimingProfilerWindow> TimingWindow = FTimingProfilerManager::Get()->GetProfilerWindow();
-	if (!TimingWindow.IsValid())
-	{
-		return;
-	}
-
-	TSharedPtr<STimingView> TimingView = TimingWindow->GetTimingView();
+	TSharedPtr<UE::Insights::TimingProfiler::STimingView> TimingView = GetTimingView();
 	if (!TimingView.IsValid())
 	{
 		return;
@@ -674,13 +652,8 @@ void FTaskGraphProfilerManager::ClearTaskRelations()
 int32 FTaskGraphProfilerManager::GetDepthOfTaskExecution(double TaskStartedTime, double TaskFinishedTime, uint32 ThreadId)
 {
 	int32 Depth = -1;
-	TSharedPtr<class STimingProfilerWindow> TimingWindow = FTimingProfilerManager::Get()->GetProfilerWindow();
-	if (!TimingWindow.IsValid())
-	{
-		return Depth;
-	}
 
-	TSharedPtr<STimingView> TimingView = TimingWindow->GetTimingView();
+	TSharedPtr<UE::Insights::TimingProfiler::STimingView> TimingView = GetTimingView();
 	if (!TimingView.IsValid())
 	{
 		return Depth;
@@ -845,6 +818,8 @@ void FTaskGraphProfilerManager::RegisterOnWindowClosedEventHandle()
 {
 	if (!OnWindowClosedEventHandle.IsValid())
 	{
+		using namespace UE::Insights::TimingProfiler;
+
 		TSharedPtr<STimingProfilerWindow> Window = FTimingProfilerManager::Get()->GetProfilerWindow();
 		if (!Window.IsValid())
 		{
@@ -896,6 +871,7 @@ void FTaskGraphProfilerManager::OutputWarnings()
 									   FText::FromString(TrackList.ToString()));
 	}
 
+	using namespace UE::Insights::TimingProfiler;
 	FName LogListingName = FTimingProfilerManager::Get()->GetLogListingName();
 	FMessageLog ReportMessageLog(LogListingName);
 	ReportMessageLog.Warning(WarningMessage);
@@ -904,6 +880,15 @@ void FTaskGraphProfilerManager::OutputWarnings()
 
 ////////////////////////////////////////////////////////////////////////////////////////////////////
 
+TSharedPtr<UE::Insights::TimingProfiler::STimingView> FTaskGraphProfilerManager::GetTimingView()
+{
+	using namespace UE::Insights::TimingProfiler;
+	TSharedPtr<STimingProfilerWindow> Window = FTimingProfilerManager::Get()->GetProfilerWindow();
+	return Window.IsValid() ? Window->GetTimingView() : nullptr;
+}
+
+////////////////////////////////////////////////////////////////////////////////////////////////////
+
 } // namespace Insights
 
 #undef LOCTEXT_NAMESPACE
diff --git a/Engine/Source/Developer/TraceInsights/Private/Insights/TaskGraphProfiler/TaskGraphProfilerManager.h b/Engine/Source/Developer/TraceInsights/Private/Insights/TaskGraphProfiler/TaskGraphProfilerManager.h
--- a/Engine/Source/Developer/TraceInsights/Private/Insights/TaskGraphProfiler/TaskGraphProfilerManager.h
+++ b/Engine/Source/Developer/TraceInsights/Private/Insights/TaskGraphProfiler/TaskGraphProfilerManager.h
@@ -2,12 +2,13 @@
 
 #pragma once
 
+#include "CoreMinimal.h"
+
 #include "Async/TaskTrace.h"
 #include "Containers/Ticker.h"
-#include "CoreMinimal.h"
 #include "Framework/Commands/UICommandList.h"
 
-// Insights
+// TraceInsights
 #include "Insights/InsightsManager.h"
 #include "Insights/IUnrealInsightsModule.h"
 
@@ -20,6 +21,8 @@ namespace TraceServices
 class FThreadTimingTrack;
 class FThreadTrackEvent;
 
+namespace UE::Insights::TimingProfiler { class STimingView; }
+
 namespace Insights
 {
 
@@ -27,6 +30,8 @@ class FTaskGraphRelation;
 class FTaskTimingSharedState;
 class STaskTableTreeView;
 
+////////////////////////////////////////////////////////////////////////////////////////////////////
+
 struct FTaskGraphProfilerTabs
 {
 	// Tab identifiers
@@ -106,7 +111,7 @@ public:
 	FLinearColor GetColorForTaskEvent(ETaskEventType InEvent);
 	uint32 GetColorForTaskEventAsPackedARGB(ETaskEventType InEvent);
 
-	TSharedPtr<Insights::FTaskTimingSharedState> GetTaskTimingSharedState() { return TaskTimingSharedState;	}
+	TSharedPtr<::Insights::FTaskTimingSharedState> GetTaskTimingSharedState() { return TaskTimingSharedState;	}
 
 	bool GetShowCriticalPath() const { return bShowCriticalPath; }
 	void SetShowCriticalPath(bool bInValue) { bShowCriticalPath = bInValue; }
@@ -156,10 +161,12 @@ private:
 
 	void OutputWarnings();
 
+	static TSharedPtr<UE::Insights::TimingProfiler::STimingView> GetTimingView();
+
 private:
 	bool bIsInitialized;
 	bool bIsAvailable;
-	FAvailabilityCheck AvailabilityCheck;
+	UE::Insights::FAvailabilityCheck AvailabilityCheck;
 
 	/** The delegate to be invoked when this manager ticks. */
 	FTickerDelegate OnTick;
@@ -175,7 +182,7 @@ private:
 
 	TWeakPtr<FTabManager> TimingTabManager;
 
-	TSharedPtr<Insights::STaskTableTreeView> TaskTableTreeView;
+	TSharedPtr<::Insights::STaskTableTreeView> TaskTableTreeView;
 	FLinearColor ColorCode[static_cast<uint32>(ETaskEventType::NumTaskEventTypes)];
 	bool bShowCriticalPath = false;
 	bool bShowTransitions = true;
diff --git a/Engine/Source/Developer/TraceInsights/Private/Insights/TaskGraphProfiler/ViewModels/TaskEntry.h b/Engine/Source/Developer/TraceInsights/Private/Insights/TaskGraphProfiler/ViewModels/TaskEntry.h
--- a/Engine/Source/Developer/TraceInsights/Private/Insights/TaskGraphProfiler/ViewModels/TaskEntry.h
+++ b/Engine/Source/Developer/TraceInsights/Private/Insights/TaskGraphProfiler/ViewModels/TaskEntry.h
@@ -2,9 +2,11 @@
 
 #pragma once
 
-#include "Async/TaskTrace.h"
 #include "CoreMinimal.h"
 
+#include "Async/TaskTrace.h"
+
+// TraceServices
 #include "TraceServices/Model/TasksProfiler.h"
 
 namespace Insights
