[Insights] Massive refactoring of TraceInsights module:

diff --git a/Engine/Source/Developer/TraceInsights/Private/Insights/TaskGraphProfiler/ViewModels/TaskGraphRelation.cpp b/Engine/Source/Developer/TraceInsights/Private/Insights/TaskGraphProfiler/ViewModels/TaskGraphRelation.cpp
--- a/Engine/Source/Developer/TraceInsights/Private/Insights/TaskGraphProfiler/ViewModels/TaskGraphRelation.cpp
+++ b/Engine/Source/Developer/TraceInsights/Private/Insights/TaskGraphProfiler/ViewModels/TaskGraphRelation.cpp
@@ -2,10 +2,12 @@
 
 #include "TaskGraphRelation.h"
 
-// Insights
-#include "Insights/Common/PaintUtils.h"
-#include "Insights/TimingProfilerManager.h"
+// TraceInsightsCore
+#include "InsightsCore/Common/PaintUtils.h"
+
+// TraceInsights
 #include "Insights/TaskGraphProfiler/TaskGraphProfilerManager.h"
+#include "Insights/TimingProfilerManager.h"
 #include "Insights/ViewModels/TimingTrackViewport.h"
 #include "Insights/ViewModels/TimingViewDrawHelper.h"
 
@@ -29,7 +31,7 @@ FTaskGraphRelation::FTaskGraphRelation(double InSourceTime, int32 InSourceThread
 
 ////////////////////////////////////////////////////////////////////////////////////////////////////
 
-void FTaskGraphRelation::Draw(const FDrawContext& DrawContext, const FTimingTrackViewport& Viewport, const ITimingViewDrawHelper& Helper, const ITimingEventRelation::EDrawFilter Filter)
+void FTaskGraphRelation::Draw(const UE::Insights::FDrawContext& DrawContext, const FTimingTrackViewport& Viewport, const ITimingViewDrawHelper& Helper, const ITimingEventRelation::EDrawFilter Filter)
 {
 	int32 LayerId = Helper.GetRelationLayerId();
 
@@ -79,6 +81,7 @@ void FTaskGraphRelation::Draw(const FDrawContext& DrawContext, const FTimingTrac
 	float Y2 = 0.0f;
 	if (SourceTrackShared->IsVisible())
 	{
+		using namespace UE::Insights::TimingProfiler;
 		int32 ActualSourceDepth = FMath::Min(SourceDepth, (int32)FTimingProfilerManager::Get()->GetEventDepthLimit() - 1);
 		Y1 = SourceTrackShared->GetPosY();
 		Y1 += Viewport.GetLayout().GetLaneY(ActualSourceDepth) + Viewport.GetLayout().EventH / 2.0f;
@@ -94,6 +97,7 @@ void FTaskGraphRelation::Draw(const FDrawContext& DrawContext, const FTimingTrac
 
 	if (TargetTrackShared->IsVisible())
 	{
+		using namespace UE::Insights::TimingProfiler;
 		int32 ActualTargetDepth = FMath::Min(TargetDepth, (int32)FTimingProfilerManager::Get()->GetEventDepthLimit() - 1);
 		Y2 = TargetTrackShared->GetPosY();
 		Y2 += Viewport.GetLayout().GetLaneY(ActualTargetDepth) + Viewport.GetLayout().EventH / 2.0f;
@@ -130,8 +134,8 @@ void FTaskGraphRelation::Draw(const FDrawContext& DrawContext, const FTimingTrac
 	TArray<FVector2D> LinePoints;
 	LinePoints.Add(StartPoint + FVector2D(0.0, -LineHeightAtStart / 2.0));
 	LinePoints.Add(StartPoint + FVector2D(0.0, +LineHeightAtStart / 2.0));
-	DrawContext.DrawLines(OutlineLayerId, 0.0f, 0.0f, LinePoints, ESlateDrawEffect::None, OutlineColor, /*bAntialias=*/ true, OutlineThickness);
-	DrawContext.DrawLines(LayerId, 0.0f, 0.0f, LinePoints, ESlateDrawEffect::None, Color, /*bAntialias=*/ true, LineThickness);
+	DrawContext.DrawLines(OutlineLayerId, 0.0f, 0.0f, LinePoints, OutlineColor, /*bAntialias=*/ true, OutlineThickness);
+	DrawContext.DrawLines(LayerId, 0.0f, 0.0f, LinePoints, Color, /*bAntialias=*/ true, LineThickness);
 
 	constexpr double MinDistance = 1.5 * (LineLengthAtStart + LineLengthAtEnd);
 	constexpr double MaxDistance = 10000.0; // arbitrary limit to avoid stack overflow in recursive FLineBuilder::Subdivide when rendering splines
@@ -145,22 +149,22 @@ void FTaskGraphRelation::Draw(const FDrawContext& DrawContext, const FTimingTrac
 		LinePoints.Empty();
 		LinePoints.Add(StartPoint);
 		LinePoints.Add(SplineStart);
-		DrawContext.DrawLines(OutlineLayerId, 0.0f, 0.0f, LinePoints, ESlateDrawEffect::None, OutlineColor, /*bAntialias=*/ true, OutlineThickness);
-		DrawContext.DrawLines(LayerId, 0.0f, 0.0f, LinePoints, ESlateDrawEffect::None, Color, /*bAntialias=*/ true, LineThickness);
+		DrawContext.DrawLines(OutlineLayerId, 0.0f, 0.0f, LinePoints, OutlineColor, /*bAntialias=*/ true, OutlineThickness);
+		DrawContext.DrawLines(LayerId, 0.0f, 0.0f, LinePoints, Color, /*bAntialias=*/ true, LineThickness);
 
 		LinePoints.Empty();
 		LinePoints.Add(SplineEnd);
 		LinePoints.Add(EndPoint);
-		DrawContext.DrawLines(OutlineLayerId, 0.0f, 0.0f, LinePoints, ESlateDrawEffect::None, OutlineColor, /*bAntialias=*/ true, OutlineThickness);
-		DrawContext.DrawLines(LayerId, 0.0f, 0.0f, LinePoints, ESlateDrawEffect::None, Color, /*bAntialias=*/ true, LineThickness);
+		DrawContext.DrawLines(OutlineLayerId, 0.0f, 0.0f, LinePoints, OutlineColor, /*bAntialias=*/ true, OutlineThickness);
+		DrawContext.DrawLines(LayerId, 0.0f, 0.0f, LinePoints, Color, /*bAntialias=*/ true, LineThickness);
 	}
 	else
 	{
 		LinePoints.Empty();
 		LinePoints.Add(StartPoint);
 		LinePoints.Add(EndPoint);
-		DrawContext.DrawLines(OutlineLayerId, 0.0f, 0.0f, LinePoints, ESlateDrawEffect::None, OutlineColor, /*bAntialias=*/ true, OutlineThickness);
-		DrawContext.DrawLines(LayerId, 0.0f, 0.0f, LinePoints, ESlateDrawEffect::None, Color, /*bAntialias=*/ true, LineThickness);
+		DrawContext.DrawLines(OutlineLayerId, 0.0f, 0.0f, LinePoints, OutlineColor, /*bAntialias=*/ true, OutlineThickness);
+		DrawContext.DrawLines(LayerId, 0.0f, 0.0f, LinePoints, Color, /*bAntialias=*/ true, LineThickness);
 
 		ArrowDirection = StartPoint - EndPoint;
 		ArrowDirection.Normalize();
@@ -172,14 +176,14 @@ void FTaskGraphRelation::Draw(const FDrawContext& DrawContext, const FTimingTrac
 	LinePoints.Empty();
 	LinePoints.Add(ArrowOrigin);
 	LinePoints.Add(ArrowOrigin + ArrowDirection.GetRotated(-ArrowRotationAngle));
-	DrawContext.DrawLines(OutlineLayerId, 0.0f, 0.0f, LinePoints, ESlateDrawEffect::None, OutlineColor, /*bAntialias=*/ true, OutlineThickness);
-	DrawContext.DrawLines(LayerId, 0.0f, 0.0f, LinePoints, ESlateDrawEffect::None, Color, /*bAntialias=*/ true, LineThickness);
+	DrawContext.DrawLines(OutlineLayerId, 0.0f, 0.0f, LinePoints, OutlineColor, /*bAntialias=*/ true, OutlineThickness);
+	DrawContext.DrawLines(LayerId, 0.0f, 0.0f, LinePoints, Color, /*bAntialias=*/ true, LineThickness);
 
 	LinePoints.Empty();
 	LinePoints.Add(ArrowOrigin);
 	LinePoints.Add(ArrowOrigin + ArrowDirection.GetRotated(ArrowRotationAngle));
-	DrawContext.DrawLines(OutlineLayerId, 0.0f, 0.0f, LinePoints, ESlateDrawEffect::None, OutlineColor, /*bAntialias=*/ true, OutlineThickness);
-	DrawContext.DrawLines(LayerId, 0.0f, 0.0f, LinePoints, ESlateDrawEffect::None, Color, /*bAntialias=*/ true, LineThickness);
+	DrawContext.DrawLines(OutlineLayerId, 0.0f, 0.0f, LinePoints, OutlineColor, /*bAntialias=*/ true, OutlineThickness);
+	DrawContext.DrawLines(LayerId, 0.0f, 0.0f, LinePoints, Color, /*bAntialias=*/ true, LineThickness);
 }
 
 ////////////////////////////////////////////////////////////////////////////////////////////////////
diff --git a/Engine/Source/Developer/TraceInsights/Private/Insights/TaskGraphProfiler/ViewModels/TaskGraphRelation.h b/Engine/Source/Developer/TraceInsights/Private/Insights/TaskGraphProfiler/ViewModels/TaskGraphRelation.h
--- a/Engine/Source/Developer/TraceInsights/Private/Insights/TaskGraphProfiler/ViewModels/TaskGraphRelation.h
+++ b/Engine/Source/Developer/TraceInsights/Private/Insights/TaskGraphProfiler/ViewModels/TaskGraphRelation.h
@@ -18,7 +18,7 @@ class FTaskGraphRelation : public ITimingEventRelation
 public:
 	FTaskGraphRelation(double InSourceTime, int32 InSourceThreadId, double InTargetTime, int32 InTargetThreadId, ETaskEventType InType);
 
-	virtual void Draw(const FDrawContext& DrawContext, const FTimingTrackViewport& Viewport, const ITimingViewDrawHelper& Helper, const ITimingEventRelation::EDrawFilter Filter) override;
+	virtual void Draw(const UE::Insights::FDrawContext& DrawContext, const FTimingTrackViewport& Viewport, const ITimingViewDrawHelper& Helper, const ITimingEventRelation::EDrawFilter Filter) override;
 
 	void SetSourceTrack(TSharedPtr<const FBaseTimingTrack> InSourceTrack) { SourceTrack = InSourceTrack; }
 	TSharedPtr<const FBaseTimingTrack> GetSourceTrack() { return SourceTrack.Pin(); }
diff --git a/Engine/Source/Developer/TraceInsights/Private/Insights/TaskGraphProfiler/ViewModels/TaskNode.h b/Engine/Source/Developer/TraceInsights/Private/Insights/TaskGraphProfiler/ViewModels/TaskNode.h
--- a/Engine/Source/Developer/TraceInsights/Private/Insights/TaskGraphProfiler/ViewModels/TaskNode.h
+++ b/Engine/Source/Developer/TraceInsights/Private/Insights/TaskGraphProfiler/ViewModels/TaskNode.h
@@ -4,28 +4,18 @@
 
 #include "CoreMinimal.h"
 
-// Insights
+// TraceInsightsCore
+#include "InsightsCore/Table/ViewModels/TableTreeNode.h"
+
+// TraceInsights
 #include "Insights/TaskGraphProfiler/ViewModels/TaskTable.h"
 #include "Insights/TaskGraphProfiler/ViewModels/TaskEntry.h"
-#include "Insights/Table/ViewModels/TableTreeNode.h"
 
 ////////////////////////////////////////////////////////////////////////////////////////////////////
 
 namespace Insights
 {
 
-enum class ETaskNodeType
-{
-	/** The TaskNode is an allocation node. */
-	Task,
-
-	/** The TaskNode is a group node. */
-	Group,
-
-	/** Invalid enum type, may be used as a number of enumerations. */
-	InvalidOrMax,
-};
-
 ////////////////////////////////////////////////////////////////////////////////////////////////////
 
 class FTaskNode;
@@ -46,33 +36,26 @@ typedef TWeakPtr<class FTaskNode> FTaskNodeWeak;
 /**
  * Class used to store information about a task node (used in the STaskTreeView).
  */
-class FTaskNode : public FTableTreeNode
+class FTaskNode : public UE::Insights::FTableTreeNode
 {
-	INSIGHTS_DECLARE_RTTI(FTaskNode, FTableTreeNode)
+	INSIGHTS_DECLARE_RTTI(FTaskNode, UE::Insights::FTableTreeNode)
 
 public:
 	/** Initialization constructor for the Task node. */
 	explicit FTaskNode(const FName InName, TWeakPtr<FTaskTable> InParentTable, int32 InRowIndex)
-		: FTableTreeNode(InName, InParentTable, InRowIndex)
-		, Type(ETaskNodeType::Task)
+		: UE::Insights::FTableTreeNode(InName, InParentTable, InRowIndex)
 	{
 	}
 
 	/** Initialization constructor for the group node. */
 	explicit FTaskNode(const FName InGroupName, TWeakPtr<FTaskTable> InParentTable)
-		: FTableTreeNode(InGroupName, InParentTable)
-		, Type(ETaskNodeType::Group)
+		: UE::Insights::FTableTreeNode(InGroupName, InParentTable)
 	{
 	}
 
-	/**
-	 * @return a type of this Task node or ETaskNodeType::Group for group nodes.
-	 */
-	ETaskNodeType GetType() const { return Type; }
-
 	FTaskTable& GetTaskTableChecked() const
 	{
-		const TSharedPtr<FTable>& TablePin = GetParentTable().Pin();
+		const TSharedPtr<UE::Insights::FTable>& TablePin = GetParentTable().Pin();
 		check(TablePin.IsValid());
 		return *StaticCastSharedPtr<FTaskTable>(TablePin);
 	}
@@ -80,9 +63,6 @@ public:
 	bool IsValidTask() const { return GetTaskTableChecked().IsValidRowIndex(GetRowIndex()); }
 	const FTaskEntry* GetTask() const { return GetTaskTableChecked().GetTask(GetRowIndex()); }
 	const FTaskEntry& GetTaskChecked() const { return GetTaskTableChecked().GetTaskChecked(GetRowIndex()); }
-
-private:
-	const ETaskNodeType Type;
 };
 
 ////////////////////////////////////////////////////////////////////////////////////////////////////
diff --git a/Engine/Source/Developer/TraceInsights/Private/Insights/TaskGraphProfiler/ViewModels/TaskTable.cpp b/Engine/Source/Developer/TraceInsights/Private/Insights/TaskGraphProfiler/ViewModels/TaskTable.cpp
--- a/Engine/Source/Developer/TraceInsights/Private/Insights/TaskGraphProfiler/ViewModels/TaskTable.cpp
+++ b/Engine/Source/Developer/TraceInsights/Private/Insights/TaskGraphProfiler/ViewModels/TaskTable.cpp
@@ -2,14 +2,17 @@
 
 #include "TaskTable.h"
 
-// Insights
-#include "Insights/Common/TimeUtils.h"
+// TraceInsightsCore
+#include "InsightsCore/Common/TimeUtils.h"
+#include "InsightsCore/Table/ViewModels/TableCellValue.h"
+#include "InsightsCore/Table/ViewModels/TableCellValueFormatter.h"
+#include "InsightsCore/Table/ViewModels/TableCellValueGetter.h"
+#include "InsightsCore/Table/ViewModels/TableCellValueSorter.h"
+#include "InsightsCore/Table/ViewModels/TableColumn.h"
+
+// TraceInsights
 #include "Insights/TaskGraphProfiler/ViewModels/TaskNode.h"
 #include "Insights/TaskGraphProfiler/ViewModels/TaskTable.h"
-#include "Insights/Table/ViewModels/TableCellValueFormatter.h"
-#include "Insights/Table/ViewModels/TableCellValueGetter.h"
-#include "Insights/Table/ViewModels/TableCellValueSorter.h"
-#include "Insights/Table/ViewModels/TableColumn.h"
 
 #define LOCTEXT_NAMESPACE "Insights::FTaskTable"
 
@@ -41,17 +44,17 @@ const FName FTaskTableColumns::NumPrerequisitesColumnId(TEXT("NumPrerequisites")
 
 ////////////////////////////////////////////////////////////////////////////////////////////////////
 
-typedef FTableCellValue(*TaskFieldGetter)(const FTableColumn&, const FTaskEntry&);
+typedef UE::Insights::FTableCellValue (*TaskFieldGetter) (const UE::Insights::FTableColumn&, const FTaskEntry&);
 
 template<TaskFieldGetter Getter>
-class FTaskColumnValueGetter : public FTableCellValueGetter
+class FTaskColumnValueGetter : public UE::Insights::FTableCellValueGetter
 {
 public:
-	virtual const TOptional<FTableCellValue> GetValue(const FTableColumn& Column, const FBaseTreeNode& Node) const override
+	virtual const TOptional<UE::Insights::FTableCellValue> GetValue(const UE::Insights::FTableColumn& Column, const UE::Insights::FBaseTreeNode& Node) const override
 	{
 		if (Node.IsGroup())
 		{
-			const FTableTreeNode& NodePtr = static_cast<const FTableTreeNode&>(Node);
+			const UE::Insights::FTableTreeNode& NodePtr = static_cast<const UE::Insights::FTableTreeNode&>(Node);
 			if (NodePtr.HasAggregatedValue(Column.GetId()))
 			{
 				return NodePtr.GetAggregatedValue(Column.GetId());
@@ -67,7 +70,7 @@ public:
 			}
 		}
 
-		return TOptional<FTableCellValue>();
+		return TOptional<UE::Insights::FTableCellValue>();
 	}
 };
 
@@ -87,6 +90,9 @@ double GetRelativeValue(double ValueA, double ValueB)
 
 struct DefaultTaskFieldGetterFuncts
 {
+	using FTableCellValue = UE::Insights::FTableCellValue;
+	using FTableColumn = UE::Insights::FTableColumn;
+
 	static FTableCellValue GetDebugName(const FTableColumn& Column, const FTaskEntry& Task) { return FTableCellValue(Task.GetDebugName()); }
 	static FTableCellValue GetCreatedTimestamp(const FTableColumn& Column, const FTaskEntry& Task) { return FTableCellValue(Task.GetCreatedTimestamp());	}
 	static FTableCellValue GetCreatedThreadId(const FTableColumn& Column, const FTaskEntry& Task) { return FTableCellValue((int64)Task.GetCreatedThreadId());	}
@@ -120,6 +126,9 @@ struct DefaultTaskFieldGetterFuncts
 
 struct RelativeToPreviousTaskFieldGetterFuncts
 {
+	using FTableCellValue = UE::Insights::FTableCellValue;
+	using FTableColumn = UE::Insights::FTableColumn;
+
 	static FTableCellValue GetCreatedTimestamp(const FTableColumn& Column, const FTaskEntry& Task) { return FTableCellValue(Task.GetCreatedTimestamp()); }
 	static FTableCellValue GetLaunchedTimestamp(const FTableColumn& Column, const FTaskEntry& Task) { return FTableCellValue(GetRelativeValue(Task.GetLaunchedTimestamp(), Task.GetCreatedTimestamp())); }
 	static FTableCellValue GetScheduledTimestamp(const FTableColumn& Column, const FTaskEntry& Task) { return FTableCellValue(GetRelativeValue(Task.GetScheduledTimestamp(), Task.GetLaunchedTimestamp())); }
@@ -133,6 +142,9 @@ struct RelativeToPreviousTaskFieldGetterFuncts
 
 struct RelativeToCreatedTaskFieldGetterFuncts
 {
+	using FTableCellValue = UE::Insights::FTableCellValue;
+	using FTableColumn = UE::Insights::FTableColumn;
+
 	static FTableCellValue GetCreatedTimestamp(const FTableColumn& Column, const FTaskEntry& Task) { return FTableCellValue(Task.GetCreatedTimestamp()); }
 	static FTableCellValue GetLaunchedTimestamp(const FTableColumn& Column, const FTaskEntry& Task) { return FTableCellValue(GetRelativeValue(Task.GetLaunchedTimestamp(), Task.GetCreatedTimestamp())); }
 	static FTableCellValue GetScheduledTimestamp(const FTableColumn& Column, const FTaskEntry& Task) { return FTableCellValue(GetRelativeValue(Task.GetScheduledTimestamp(), Task.GetCreatedTimestamp())); }
@@ -144,24 +156,24 @@ struct RelativeToCreatedTaskFieldGetterFuncts
 
 ////////////////////////////////////////////////////////////////////////////////////////////////////
 
-class FTaskDoubleValueFormatterAsTimeAuto : public FTableCellValueFormatter
+class FTaskDoubleValueFormatterAsTimeAuto : public UE::Insights::FTableCellValueFormatter
 {
 public:
-	virtual FText FormatValue(const TOptional<FTableCellValue>& InValue) const override
+	virtual FText FormatValue(const TOptional<UE::Insights::FTableCellValue>& InValue) const override
 	{
 		if (InValue.IsSet())
 		{
 			const double Value = InValue.GetValue().Double;
 			if (Value != TraceServices::FTaskInfo::InvalidTimestamp)
 			{
-				return FText::FromString(TimeUtils::FormatTimeAuto(Value));
+				return FText::FromString(UE::Insights::FormatTimeAuto(Value));
 			}
 		}
 
 		return FText::FromString(TEXT("N/A"));
 	}
 
-	virtual FText FormatValueForTooltip(const TOptional<FTableCellValue>& InValue) const override
+	virtual FText FormatValueForTooltip(const TOptional<UE::Insights::FTableCellValue>& InValue) const override
 	{
 		if (InValue.IsSet())
 		{
@@ -172,7 +184,7 @@ public:
 			}
 			else if (Value != TraceServices::FTaskInfo::InvalidTimestamp)
 			{
-				return FText::FromString(FString::Printf(TEXT("%f (%s)"), Value, *TimeUtils::FormatTimeAuto(Value)));
+				return FText::FromString(FString::Printf(TEXT("%f (%s)"), Value, *UE::Insights::FormatTimeAuto(Value)));
 			}
 		}
 
@@ -207,6 +219,8 @@ void FTaskTable::Reset()
 
 void FTaskTable::AddDefaultColumns()
 {
+	using namespace UE::Insights;
+
 	//////////////////////////////////////////////////
 	// Hierarchy Column
 	{
diff --git a/Engine/Source/Developer/TraceInsights/Private/Insights/TaskGraphProfiler/ViewModels/TaskTable.h b/Engine/Source/Developer/TraceInsights/Private/Insights/TaskGraphProfiler/ViewModels/TaskTable.h
--- a/Engine/Source/Developer/TraceInsights/Private/Insights/TaskGraphProfiler/ViewModels/TaskTable.h
+++ b/Engine/Source/Developer/TraceInsights/Private/Insights/TaskGraphProfiler/ViewModels/TaskTable.h
@@ -2,17 +2,21 @@
 
 #pragma once
 
-#include "CoreMinimal.h"
+#include "CoreTypes.h"
+#include "Containers/Array.h"
+#include "UObject/NameTypes.h"
 
-// Insights
+// TraceInsightsCore
+#include "InsightsCore/Table/ViewModels/Table.h"
+
+// TraceInsights
 #include "Insights/TaskGraphProfiler/ViewModels/TaskEntry.h"
-#include "Insights/Table/ViewModels/Table.h"
+
+namespace UE::Insights { class FTableColumn; }
 
 namespace Insights
 {
 
-class FTableColumn;
-
 ////////////////////////////////////////////////////////////////////////////////////////////////////
 
 // Column identifiers
@@ -41,7 +45,7 @@ struct FTaskTableColumns
 
 ////////////////////////////////////////////////////////////////////////////////////////////////////
 
-class FTaskTable : public FTable
+class FTaskTable : public UE::Insights::FTable
 {
 public:
 	FTaskTable();
diff --git a/Engine/Source/Developer/TraceInsights/Private/Insights/TaskGraphProfiler/ViewModels/TaskTimingTrack.cpp b/Engine/Source/Developer/TraceInsights/Private/Insights/TaskGraphProfiler/ViewModels/TaskTimingTrack.cpp
--- a/Engine/Source/Developer/TraceInsights/Private/Insights/TaskGraphProfiler/ViewModels/TaskTimingTrack.cpp
+++ b/Engine/Source/Developer/TraceInsights/Private/Insights/TaskGraphProfiler/ViewModels/TaskTimingTrack.cpp
@@ -5,20 +5,24 @@
 #include "Framework/Commands/Commands.h"
 #include "Framework/Commands/UICommandList.h"
 #include "Framework/MultiBox/MultiBoxBuilder.h"
+
+// TraceServices
 #include "TraceServices/Model/TasksProfiler.h"
 #include "TraceServices/Model/Threads.h"
 
-// Insights
-#include "Insights/Common/TimeUtils.h"
+// TraceInsightsCore
+#include "InsightsCore/Common/TimeUtils.h"
+
+// TraceInsights
 #include "Insights/InsightsStyle.h"
 #include "Insights/ITimingViewSession.h"
 #include "Insights/TaskGraphProfiler/TaskGraphProfilerManager.h"
 #include "Insights/TaskGraphProfiler/ViewModels/TaskGraphRelation.h"
 #include "Insights/TaskGraphProfiler/ViewModels/TaskTrackEvent.h"
+#include "Insights/TimingProfilerManager.h"
 #include "Insights/ViewModels/ThreadTimingTrack.h"
 #include "Insights/ViewModels/ThreadTrackEvent.h"
 #include "Insights/ViewModels/TimingTrackViewport.h"
-#include "Insights/TimingProfilerManager.h"
 #include "Insights/ViewModels/TimingViewDrawHelper.h"
 #include "Insights/ViewModels/TooltipDrawState.h"
 #include "Insights/Widgets/STimingProfilerWindow.h"
@@ -122,14 +126,14 @@ public:
 // FTaskTimingSharedState
 ////////////////////////////////////////////////////////////////////////////////////////////////////
 
-FTaskTimingSharedState::FTaskTimingSharedState(STimingView* InTimingView)
+FTaskTimingSharedState::FTaskTimingSharedState(UE::Insights::TimingProfiler::STimingView* InTimingView)
 	: TimingViewSession(InTimingView)
 {
 }
 
 ////////////////////////////////////////////////////////////////////////////////////////////////////
 
-void FTaskTimingSharedState::OnBeginSession(Insights::ITimingViewSession& InSession)
+void FTaskTimingSharedState::OnBeginSession(UE::Insights::Timing::ITimingViewSession& InSession)
 {
 	if (&InSession != TimingViewSession)
 	{
@@ -148,7 +152,7 @@ void FTaskTimingSharedState::OnBeginSession(Insights::ITimingViewSession& InSess
 
 ////////////////////////////////////////////////////////////////////////////////////////////////////
 
-void FTaskTimingSharedState::OnEndSession(Insights::ITimingViewSession& InSession)
+void FTaskTimingSharedState::OnEndSession(UE::Insights::Timing::ITimingViewSession& InSession)
 {
 	if (&InSession != TimingViewSession)
 	{
@@ -162,7 +166,7 @@ void FTaskTimingSharedState::OnEndSession(Insights::ITimingViewSession& InSessio
 
 ////////////////////////////////////////////////////////////////////////////////////////////////////
 
-void FTaskTimingSharedState::Tick(Insights::ITimingViewSession& InSession, const TraceServices::IAnalysisSession& InAnalysisSession)
+void FTaskTimingSharedState::Tick(UE::Insights::Timing::ITimingViewSession& InSession, const TraceServices::IAnalysisSession& InAnalysisSession)
 {
 	if (&InSession != TimingViewSession)
 	{
@@ -171,7 +175,7 @@ void FTaskTimingSharedState::Tick(Insights::ITimingViewSession& InSession, const
 
 	if (!TaskTrack.IsValid() && FTaskGraphProfilerManager::Get()->GetIsAvailable())
 	{
-		TSharedPtr<STimingView> TimingView = GetTimingView();
+		TSharedPtr<UE::Insights::TimingProfiler::STimingView> TimingView = GetTimingView();
 
 		InitCommandList(TimingView);
 
@@ -188,10 +192,11 @@ void FTaskTimingSharedState::Tick(Insights::ITimingViewSession& InSession, const
 
 	if (bResetOnNextTick)
 	{
-		TSharedPtr<STimingView> TimingView = GetTimingView();
-
 		bResetOnNextTick = false;
-		if (!TimingView->GetSelectedEvent().IsValid() &&
+
+		TSharedPtr<UE::Insights::TimingProfiler::STimingView> TimingView = GetTimingView();
+		if (TimingView.IsValid() &&
+			!TimingView->GetSelectedEvent().IsValid() &&
 			(!TimingView->GetSelectedTrack().IsValid() || TimingView->GetSelectedTrack().Get() != TaskTrack.Get()))
 		{
 			SetTaskId(TaskTrace::InvalidId);
@@ -202,21 +207,16 @@ void FTaskTimingSharedState::Tick(Insights::ITimingViewSession& InSession, const
 
 ////////////////////////////////////////////////////////////////////////////////////////////////////
 
-TSharedPtr<STimingView> FTaskTimingSharedState::GetTimingView()
+TSharedPtr<UE::Insights::TimingProfiler::STimingView> FTaskTimingSharedState::GetTimingView()
 {
-	TSharedPtr<STimingView> TimingView;
+	using namespace UE::Insights::TimingProfiler;
 	TSharedPtr<STimingProfilerWindow> Window = FTimingProfilerManager::Get()->GetProfilerWindow();
-	if (Window.IsValid())
-	{
-		TimingView = Window->GetTimingView();
-	}
-
-	return TimingView;
+	return Window.IsValid() ? Window->GetTimingView() : nullptr;
 }
 
 ////////////////////////////////////////////////////////////////////////////////////////////////////
 
-void FTaskTimingSharedState::ExtendFilterMenu(Insights::ITimingViewSession& InSession, FMenuBuilder& InOutMenuBuilder)
+void FTaskTimingSharedState::ExtendFilterMenu(UE::Insights::Timing::ITimingViewSession& InSession, FMenuBuilder& InOutMenuBuilder)
 {
 	if (&InSession != TimingViewSession)
 	{
@@ -236,7 +236,7 @@ void FTaskTimingSharedState::SetTaskId(TaskTrace::FId InTaskId)
 
 ////////////////////////////////////////////////////////////////////////////////////////////////////
 
-void FTaskTimingSharedState::ExtendOtherTracksFilterMenu(ITimingViewSession& InSession, FMenuBuilder& InMenuBuilder)
+void FTaskTimingSharedState::ExtendOtherTracksFilterMenu(UE::Insights::Timing::ITimingViewSession& InSession, FMenuBuilder& InMenuBuilder)
 {
 	if (&InSession != TimingViewSession)
 	{
@@ -347,10 +347,18 @@ void FTaskTimingSharedState::BuildTasksSubMenu(FMenuBuilder& MenuBuilder)
 
 ////////////////////////////////////////////////////////////////////////////////////////////////////
 
-void FTaskTimingSharedState::InitCommandList(TSharedPtr<STimingView> TimingView)
+void FTaskTimingSharedState::InitCommandList(TSharedPtr<UE::Insights::TimingProfiler::STimingView> TimingView)
 {
+	if (!ensure(TimingView.IsValid()))
+	{
+		return;
+	}
+
 	TSharedPtr<FUICommandList> CommandList = TimingView->GetCommandList();
-	ensure(CommandList);
+	if (!ensure(CommandList.IsValid()))
+	{
+		return;
+	}
 
 	FTaskTimingStateCommands::Register();
 
@@ -677,8 +685,8 @@ void FTaskTimingSharedState::OnTaskSettingsChanged()
 	TaskTrack->SetTaskId(TaskTrace::InvalidId);
 
 	TSharedPtr<Insights::FTaskGraphProfilerManager> TaskGraphManager = Insights::FTaskGraphProfilerManager::Get();
-	TSharedPtr<STimingView> TimingView = GetTimingView();
 
+	TSharedPtr<UE::Insights::TimingProfiler::STimingView> TimingView = GetTimingView();
 	if (!TimingView)
 	{
 		return;
@@ -743,6 +751,8 @@ void FTaskTimingTrack::BuildDrawState(ITimingEventsTrackDrawStateBuilder& Builde
 
 	if (bShowDetailInfoOnTaskTrack)
 	{
+		using namespace UE::Insights::TimingProfiler;
+
 		uint32 Depth = 0;
 		for (TraceServices::FTaskInfo::FRelationInfo Relation : Task->Prerequisites)
 		{
@@ -1016,9 +1026,10 @@ void FTaskTimingTrack::InitTooltip(FTooltipDrawState& InOutTooltip, const ITimin
 	const FTaskTrackEvent& TaskTrackEvent = InTooltipEvent.As<FTaskTrackEvent>();
 	InOutTooltip.AddTitle(TaskTrackEvent.GetEventName());
 
-	InOutTooltip.AddNameValueTextLine(TaskTrackEvent.GetStartLabel(), TimeUtils::FormatTimeAuto(TaskTrackEvent.GetStartTime(), 6));
-	InOutTooltip.AddNameValueTextLine(TaskTrackEvent.GetEndLabel(), TimeUtils::FormatTimeAuto(TaskTrackEvent.GetEndTime(), 6));
-	InOutTooltip.AddNameValueTextLine(TEXT("Duration:"), TimeUtils::FormatTimeAuto(TaskTrackEvent.GetEndTime() - TaskTrackEvent.GetStartTime()));
+	using namespace UE::Insights;
+	InOutTooltip.AddNameValueTextLine(TaskTrackEvent.GetStartLabel(), FormatTimeAuto(TaskTrackEvent.GetStartTime(), 6));
+	InOutTooltip.AddNameValueTextLine(TaskTrackEvent.GetEndLabel(), FormatTimeAuto(TaskTrackEvent.GetEndTime(), 6));
+	InOutTooltip.AddNameValueTextLine(TEXT("Duration:"), FormatTimeAuto(TaskTrackEvent.GetEndTime() - TaskTrackEvent.GetStartTime()));
 
 	TSharedPtr<const TraceServices::IAnalysisSession> Session = FInsightsManager::Get()->GetSession();
 	if (!Session.IsValid())
@@ -1101,7 +1112,7 @@ void FTaskTimingTrack::GetEventRelations(const FThreadTrackEvent& InSelectedEven
 	const TraceServices::ITasksProvider* TasksProvider = TraceServices::ReadTasksProvider(*Session.Get());
 	if (TasksProvider)
 	{
-		TSharedPtr<STimingView> TimingView = SharedState.GetTimingView();
+		TSharedPtr<UE::Insights::TimingProfiler::STimingView> TimingView = SharedState.GetTimingView();
 		if (!TimingView.IsValid())
 		{
 			return;
diff --git a/Engine/Source/Developer/TraceInsights/Private/Insights/TaskGraphProfiler/ViewModels/TaskTimingTrack.h b/Engine/Source/Developer/TraceInsights/Private/Insights/TaskGraphProfiler/ViewModels/TaskTimingTrack.h
--- a/Engine/Source/Developer/TraceInsights/Private/Insights/TaskGraphProfiler/ViewModels/TaskTimingTrack.h
+++ b/Engine/Source/Developer/TraceInsights/Private/Insights/TaskGraphProfiler/ViewModels/TaskTimingTrack.h
@@ -2,29 +2,31 @@
 
 #pragma once
 
-#include "Async/TaskTrace.h"
 #include "CoreMinimal.h"
 
-// Insights
+#include "Async/TaskTrace.h"
+
+// TraceInsights
 #include "Insights/ITimingViewExtender.h"
 #include "Insights/ViewModels/TimingEventsTrack.h"
 
 class FUICommandList;
 
 class FThreadTrackEvent;
-class STimingView;
+namespace UE::Insights::TimingProfiler { class STimingView; }
 
 ////////////////////////////////////////////////////////////////////////////////////////////////////
 
 namespace Insights
 {
+
 class FTaskTimingTrack;
 
-class FTaskTimingSharedState : public Insights::ITimingViewExtender, public TSharedFromThis<FTaskTimingSharedState>
+class FTaskTimingSharedState : public UE::Insights::Timing::ITimingViewExtender, public TSharedFromThis<FTaskTimingSharedState>
 {
 
 public:
-	FTaskTimingSharedState(STimingView* InTimingView);
+	FTaskTimingSharedState(UE::Insights::TimingProfiler::STimingView* InTimingView);
 	virtual ~FTaskTimingSharedState() = default;
 
 	TSharedPtr<FTaskTimingTrack> GetTaskTrack() { return TaskTrack; }
@@ -34,11 +36,11 @@ public:
 	//////////////////////////////////////////////////
 	// ITimingViewExtender interface
 
-	virtual void OnBeginSession(Insights::ITimingViewSession& InSession) override;
-	virtual void OnEndSession(Insights::ITimingViewSession& InSession) override;
-	virtual void Tick(Insights::ITimingViewSession& InSession, const TraceServices::IAnalysisSession& InAnalysisSession) override;
-	virtual void ExtendFilterMenu(Insights::ITimingViewSession& InSession, FMenuBuilder& InMenuBuilder) override;
-	virtual void ExtendOtherTracksFilterMenu(ITimingViewSession& InSession, FMenuBuilder& InMenuBuilder) override;
+	virtual void OnBeginSession(UE::Insights::Timing::ITimingViewSession& InSession) override;
+	virtual void OnEndSession(UE::Insights::Timing::ITimingViewSession& InSession) override;
+	virtual void Tick(UE::Insights::Timing::ITimingViewSession& InSession, const TraceServices::IAnalysisSession& InAnalysisSession) override;
+	virtual void ExtendFilterMenu(UE::Insights::Timing::ITimingViewSession& InSession, FMenuBuilder& InMenuBuilder) override;
+	virtual void ExtendOtherTracksFilterMenu(UE::Insights::Timing::ITimingViewSession& InSession, FMenuBuilder& InMenuBuilder) override;
 
 	//////////////////////////////////////////////////
 
@@ -51,10 +53,10 @@ public:
 
 	void SetResetOnNextTick(bool bInValue) { bResetOnNextTick = bInValue; }
 
-	static TSharedPtr<STimingView> GetTimingView();
+	static TSharedPtr<UE::Insights::TimingProfiler::STimingView> GetTimingView();
 
 private:
-	void InitCommandList(TSharedPtr<STimingView> TimingView);
+	void InitCommandList(TSharedPtr<UE::Insights::TimingProfiler::STimingView> TimingView);
 
 	void BuildTasksSubMenu(FMenuBuilder& MenuBuilder);
 
@@ -97,7 +99,7 @@ private:
 	void OnTaskSettingsChanged();
 
 private:
-	Insights::ITimingViewSession* TimingViewSession;
+	UE::Insights::Timing::ITimingViewSession* TimingViewSession;
 
 	bool bShowHideTaskTrack;
 	bool bResetOnNextTick = false;
diff --git a/Engine/Source/Developer/TraceInsights/Private/Insights/TaskGraphProfiler/ViewModels/TaskTrackEvent.h b/Engine/Source/Developer/TraceInsights/Private/Insights/TaskGraphProfiler/ViewModels/TaskTrackEvent.h
--- a/Engine/Source/Developer/TraceInsights/Private/Insights/TaskGraphProfiler/ViewModels/TaskTrackEvent.h
+++ b/Engine/Source/Developer/TraceInsights/Private/Insights/TaskGraphProfiler/ViewModels/TaskTrackEvent.h
@@ -3,9 +3,10 @@
 #pragma once
 
 #include "CoreMinimal.h"
+
 #include "Async/TaskTrace.h"
 
-// Insights
+// TraceInsights
 #include "Insights/TaskGraphProfiler/TaskGraphProfilerManager.h"
 #include "Insights/ViewModels/TimingEvent.h"
 
diff --git a/Engine/Source/Developer/TraceInsights/Private/Insights/TaskGraphProfiler/Widgets/STaskTableTreeView.cpp b/Engine/Source/Developer/TraceInsights/Private/Insights/TaskGraphProfiler/Widgets/STaskTableTreeView.cpp
--- a/Engine/Source/Developer/TraceInsights/Private/Insights/TaskGraphProfiler/Widgets/STaskTableTreeView.cpp
+++ b/Engine/Source/Developer/TraceInsights/Private/Insights/TaskGraphProfiler/Widgets/STaskTableTreeView.cpp
@@ -9,19 +9,23 @@
 #include "Misc/Paths.h"
 #include "Modules/ModuleManager.h"
 #include "SlateOptMacros.h"
-#include "TraceServices/Model/TasksProfiler.h"
 #include "Widgets/Input/SComboBox.h"
 
-// Insights
+// TraceServices
+#include "TraceServices/Model/TasksProfiler.h"
+
+// TraceInsightsCore
+#include "InsightsCore/Filter/ViewModels/FilterConfigurator.h"
+#include "InsightsCore/Table/ViewModels/TableColumn.h"
+
+// TraceInsights
 #include "Insights/InsightsStyle.h"
 #include "Insights/TaskGraphProfiler/TaskGraphProfilerManager.h"
 #include "Insights/TaskGraphProfiler/ViewModels/TaskEntry.h"
 #include "Insights/TaskGraphProfiler/ViewModels/TaskGraphRelation.h"
 #include "Insights/TaskGraphProfiler/ViewModels/TaskNode.h"
 #include "Insights/TaskGraphProfiler/ViewModels/TaskTimingTrack.h"
-#include "Insights/Table/ViewModels/TableColumn.h"
 #include "Insights/TimingProfilerManager.h"
-#include "Insights/ViewModels/FilterConfigurator.h"
 #include "Insights/ViewModels/ThreadTimingTrack.h"
 #include "Insights/Widgets/STimingProfilerWindow.h"
 #include "Insights/Widgets/STimingView.h"
@@ -203,6 +207,8 @@ void STaskTableTreeView::Tick(const FGeometry& AllottedGeometry, const double In
 
 void STaskTableTreeView::RebuildTree(bool bResync)
 {
+	using namespace UE::Insights::TimingProfiler;
+
 	double NewQueryStartTime = FTimingProfilerManager::Get()->GetSelectionStartTime();
 	double NewQueryEndTime = FTimingProfilerManager::Get()->GetSelectionEndTime();
 
@@ -232,7 +238,7 @@ void STaskTableTreeView::RebuildTree(bool bResync)
 			{
 				FName BaseNodeName(TEXT("task"));
 
-				TArray<FTableTreeNodePtr>* Nodes = &TableRowNodes;
+				TArray<UE::Insights::FTableTreeNodePtr>* Nodes = &TableRowNodes;
 
 				TasksProvider->EnumerateTasks(QueryStartTime, QueryEndTime, SelectedTasksSelectionOption, [&Tasks, &TaskTable, &BaseNodeName, Nodes](const TraceServices::FTaskInfo& TaskInfo)
 				{
@@ -682,7 +688,7 @@ bool STaskTableTreeView::TasksSelectionOptions_IsEnabled() const
 
 bool STaskTableTreeView::ContextMenu_GoToTask_CanExecute() const
 {
-	TArray<FTableTreeNodePtr> SelectedItems;
+	TArray<UE::Insights::FTableTreeNodePtr> SelectedItems;
 	TreeView->GetSelectedItems(SelectedItems);
 
 	if (SelectedItems.Num() != 1)
@@ -704,7 +710,7 @@ bool STaskTableTreeView::ContextMenu_GoToTask_CanExecute() const
 
 void STaskTableTreeView::ContextMenu_GoToTask_Execute()
 {
-	TArray<FTableTreeNodePtr> SelectedItems;
+	TArray<UE::Insights::FTableTreeNodePtr> SelectedItems;
 	TreeView->GetSelectedItems(SelectedItems);
 
 	if (SelectedItems.Num() != 1)
@@ -720,6 +726,8 @@ void STaskTableTreeView::ContextMenu_GoToTask_Execute()
 		return;
 	}
 
+	using namespace UE::Insights::TimingProfiler;
+
 	TSharedPtr<STimingProfilerWindow> TimingWindow = FTimingProfilerManager::Get()->GetProfilerWindow();
 	if (!TimingWindow.IsValid())
 	{
@@ -800,7 +808,7 @@ bool STaskTableTreeView::ContextMenu_OpenInIDE_CanExecute() const
 
 bool STaskTableTreeView::GetSourceFileAndLineForSelectedTask(FString& OutFile, uint32& OutLine) const
 {
-	TArray<FTableTreeNodePtr> SelectedItems;
+	TArray<UE::Insights::FTableTreeNodePtr> SelectedItems;
 	TreeView->GetSelectedItems(SelectedItems);
 
 	if (SelectedItems.Num() != 1)
@@ -808,7 +816,7 @@ bool STaskTableTreeView::GetSourceFileAndLineForSelectedTask(FString& OutFile, u
 		return false;
 	}
 
-	FTableTreeNodePtr SelectedTreeNode = SelectedItems[0];
+	UE::Insights::FTableTreeNodePtr SelectedTreeNode = SelectedItems[0];
 	if (!SelectedTreeNode.IsValid() || !SelectedTreeNode->Is<FTaskNode>())
 	{
 		return false;
@@ -869,7 +877,7 @@ void STaskTableTreeView::ContextMenu_OpenInIDE_Execute()
 
 ////////////////////////////////////////////////////////////////////////////////////////////////////
 
-void STaskTableTreeView::TreeView_OnMouseButtonDoubleClick(FTableTreeNodePtr TreeNode)
+void STaskTableTreeView::TreeView_OnMouseButtonDoubleClick(UE::Insights::FTableTreeNodePtr TreeNode)
 {
 	if (!TreeNode->IsGroup())
 	{
@@ -884,12 +892,12 @@ void STaskTableTreeView::TreeView_OnMouseButtonDoubleClick(FTableTreeNodePtr Tre
 void STaskTableTreeView::SelectTaskEntry(TaskTrace::FId InId)
 {
 	TaskIdToSelect = InId;
-	StartTableDataTask<FSearchForItemToSelectTask>(SharedThis(this));
+	StartTableDataTask<UE::Insights::FSearchForItemToSelectTask>(SharedThis(this));
 }
 
 ////////////////////////////////////////////////////////////////////////////////////////////////////
 
-void STaskTableTreeView::SearchForItem(TSharedPtr<FTableTaskCancellationToken> CancellationToken)
+void STaskTableTreeView::SearchForItem(TSharedPtr<UE::Insights::FTableTaskCancellationToken> CancellationToken)
 {
 	TSharedPtr<FTaskTable> TaskTable = GetTaskTable();
 	TArray<FTaskEntry>& Tasks = TaskTable->GetTaskEntries();
@@ -904,7 +912,7 @@ void STaskTableTreeView::SearchForItem(TSharedPtr<FTableTaskCancellationToken> C
 
 		if (Tasks[Index].Id == TaskIdToSelect)
 		{
-			TGraphTask<FSelectNodeByTableRowIndexTask>::CreateTask().ConstructAndDispatchWhenReady(CancellationToken, SharedThis(this), Index);
+			TGraphTask<UE::Insights::FSelectNodeByTableRowIndexTask>::CreateTask().ConstructAndDispatchWhenReady(CancellationToken, SharedThis(this), Index);
 			break;
 		}
 	}
diff --git a/Engine/Source/Developer/TraceInsights/Private/Insights/TaskGraphProfiler/Widgets/STaskTableTreeView.h b/Engine/Source/Developer/TraceInsights/Private/Insights/TaskGraphProfiler/Widgets/STaskTableTreeView.h
--- a/Engine/Source/Developer/TraceInsights/Private/Insights/TaskGraphProfiler/Widgets/STaskTableTreeView.h
+++ b/Engine/Source/Developer/TraceInsights/Private/Insights/TaskGraphProfiler/Widgets/STaskTableTreeView.h
@@ -2,12 +2,16 @@
 
 #pragma once
 
+#include "CoreTypes.h"
+
 #include "Async/TaskTrace.h"
+
+// TraceServices
 #include "TraceServices/Model/TasksProfiler.h"
 
-// Insights
-#include "Insights/TaskGraphProfiler/ViewModels/TaskTable.h"
+// TraceInsights
 #include "Insights/Table/Widgets/SSessionTableTreeView.h"
+#include "Insights/TaskGraphProfiler/ViewModels/TaskTable.h"
 
 class FMenuBuilder;
 
@@ -16,7 +20,7 @@ namespace Insights
 
 ////////////////////////////////////////////////////////////////////////////////////////////////////
 
-class STaskTableTreeView : public SSessionTableTreeView
+class STaskTableTreeView : public UE::Insights::SSessionTableTreeView
 {
 private:
 	struct FColumnConfig
@@ -84,7 +88,7 @@ public:
 
 	////////////////////////////////////////////////////////////////////////////////////////////////////
 
-	virtual void TreeView_OnMouseButtonDoubleClick(FTableTreeNodePtr TreeNode) override;
+	virtual void TreeView_OnMouseButtonDoubleClick(UE::Insights::FTableTreeNodePtr TreeNode) override;
 
 	void SelectTaskEntry(TaskTrace::FId InId);
 
@@ -93,7 +97,7 @@ protected:
 
 	virtual void ExtendMenu(FMenuBuilder& MenuBuilder) override;
 
-	virtual void SearchForItem(TSharedPtr<FTableTaskCancellationToken> CancellationToken) override;
+	virtual void SearchForItem(TSharedPtr<UE::Insights::FTableTaskCancellationToken> CancellationToken) override;
 
 private:
 	void AddCommmands();
diff --git a/Engine/Source/Developer/TraceInsights/Private/Insights/Tests/FunctionalTests/MemoryInsightsTests.cpp b/Engine/Source/Developer/TraceInsights/Private/Insights/Tests/FunctionalTests/MemoryInsightsTests.cpp
--- a/Engine/Source/Developer/TraceInsights/Private/Insights/Tests/FunctionalTests/MemoryInsightsTests.cpp
+++ b/Engine/Source/Developer/TraceInsights/Private/Insights/Tests/FunctionalTests/MemoryInsightsTests.cpp
@@ -1,10 +1,22 @@
 // Copyright Epic Games, Inc. All Rights Reserved.
 
-#include "AutomationDriverCommon.h"
 #include "Algo/Find.h"
+#include "AutomationDriverCommon.h"
 #include "HAL/FileManager.h"
 #include "HAL/PlatformFileManager.h"
+#include "Misc/AutomationTest.h"
 #include "Misc/FileHelper.h"
+
+// TraceServices
+#include "TraceServices/Model/AllocationsProvider.h"
+#include "TraceServices/Model/Callstack.h"
+
+// TraceInsightsCore
+#include "InsightsCore/Table/ViewModels/TableTreeNode.h"
+#include "InsightsCore/Table/ViewModels/TreeNodeGrouping.h"
+#include "InsightsCore/Table/Widgets/STableTreeView.h"
+
+// TraceInsights
 #include "Insights/InsightsManager.h"
 #include "Insights/IUnrealInsightsModule.h"
 #include "Insights/MemoryProfiler/MemoryProfilerManager.h"
@@ -14,14 +26,10 @@
 #include "Insights/MemoryProfiler/ViewModels/MemorySharedState.h"
 #include "Insights/MemoryProfiler/Widgets/SMemAllocTableTreeView.h"
 #include "Insights/MemoryProfiler/Widgets/SMemoryProfilerWindow.h"
-#include "Insights/Table/ViewModels/TableTreeNode.h"
-#include "Insights/Table/ViewModels/TreeNodeGrouping.h"
-#include "Insights/Table/Widgets/STableTreeView.h"
 #include "Insights/Tests/InsightsTestUtils.h"
 #include "Insights/Widgets/SStartPageWindow.h"
 #include "Insights/Widgets/STimingView.h"
-#include "Misc/AutomationTest.h"
-#include "TraceServices/Model/AllocationsProvider.h"
+
 
 DECLARE_LOG_CATEGORY_EXTERN(MemoryInsightsTests, Log, All);
 
@@ -143,6 +151,8 @@ void FAutomationDriverUnrealInsightsHubMemoryInsightsTest::Define()
 IMPLEMENT_SIMPLE_AUTOMATION_TEST(FMemoryInsightsUploadLLMXMLReportsTraceTest, "System.Insights.Trace.Analysis.MemoryInsights.UploadMemoryInsightsLLMXMLReportsTrace", EAutomationTestFlags::ProgramContext | EAutomationTestFlags::EngineFilter)
 bool FMemoryInsightsUploadLLMXMLReportsTraceTest::RunTest(const FString& Parameters)
 {
+	using namespace UE::Insights::MemoryProfiler;
+
 	const FString ReportGraphsXMLPath = FPaths::RootDir() / TEXT("EngineTest/SourceAssets/Utrace/ReportGraphs.xml");
 	const FString LLMReportTypesXMLPath = FPaths::RootDir() / TEXT("EngineTest/SourceAssets/Utrace/LLMReportTypes.xml");
 
@@ -171,7 +181,7 @@ bool FMemoryInsightsUploadLLMXMLReportsTraceTest::RunTest(const FString& Paramet
 	return true;
 }
 
-DEFINE_LATENT_AUTOMATION_COMMAND_THREE_PARAMETER(FWaitForRunningQuieryFinishedCommand, TSharedPtr<Insights::SMemAllocTableTreeView>, MemAllocTableTreeView, double, Timeout, FAutomationTestBase*, Test);
+DEFINE_LATENT_AUTOMATION_COMMAND_THREE_PARAMETER(FWaitForRunningQuieryFinishedCommand, TSharedPtr<UE::Insights::MemoryProfiler::SMemAllocTableTreeView>, MemAllocTableTreeView, double, Timeout, FAutomationTestBase*, Test);
 bool FWaitForRunningQuieryFinishedCommand::Update()
 {
 	if (!MemAllocTableTreeView->IsRunning())
@@ -188,10 +198,10 @@ bool FWaitForRunningQuieryFinishedCommand::Update()
 	return false;
 }
 
-DEFINE_LATENT_AUTOMATION_COMMAND_TWO_PARAMETER(FChangeGroupingCommand, TSharedPtr<Insights::SMemAllocTableTreeView>, MemAllocTableTreeView, FAutomationTestBase*, Test);
+DEFINE_LATENT_AUTOMATION_COMMAND_TWO_PARAMETER(FChangeGroupingCommand, TSharedPtr<UE::Insights::MemoryProfiler::SMemAllocTableTreeView>, MemAllocTableTreeView, FAutomationTestBase*, Test);
 bool FChangeGroupingCommand::Update()
 {
-	TArray<TSharedPtr<Insights::FTreeNodeGrouping>> CurrentGroupings;
+	TArray<TSharedPtr<UE::Insights::FTreeNodeGrouping>> CurrentGroupings;
 	for (const auto& Grouping : MemAllocTableTreeView->GetAvailableGroupings())
 	{
 		if (Grouping->GetTitleName().ToString().Contains(TEXT("By Callstack")))
@@ -206,15 +216,17 @@ bool FChangeGroupingCommand::Update()
 	return true;
 }
 
-DEFINE_LATENT_AUTOMATION_COMMAND_FOUR_PARAMETER(FVerifyHierarchyCallStackCommand, TSharedPtr<Insights::SMemAllocTableTreeView>, MemAllocTableTreeView, FInsightsTestUtils, InsightsTestUtils, double, Timeout, FAutomationTestBase*, Test);
+DEFINE_LATENT_AUTOMATION_COMMAND_FOUR_PARAMETER(FVerifyHierarchyCallStackCommand, TSharedPtr<UE::Insights::MemoryProfiler::SMemAllocTableTreeView>, MemAllocTableTreeView, FInsightsTestUtils, InsightsTestUtils, double, Timeout, FAutomationTestBase*, Test);
 bool FVerifyHierarchyCallStackCommand::Update()
 {
+	using namespace UE::Insights::MemoryProfiler;
+
 	if (!MemAllocTableTreeView->IsRunningAsyncUpdate())
 	{
-		for (const TSharedPtr<Insights::FTableTreeNode>& Node : MemAllocTableTreeView->GetTableRowNodes())
+		for (const TSharedPtr<UE::Insights::FTableTreeNode>& Node : MemAllocTableTreeView->GetTableRowNodes())
 		{
-			const Insights::FMemAllocNode& MemAllocNode = static_cast<const Insights::FMemAllocNode&>(*Node);
-			const Insights::FMemoryAlloc Alloc = MemAllocNode.GetMemAllocChecked();
+			const FMemAllocNode& MemAllocNode = static_cast<const FMemAllocNode&>(*Node);
+			const FMemoryAlloc Alloc = MemAllocNode.GetMemAllocChecked();
 			if (!(!Alloc.GetAllocCallstack() || Alloc.GetAllocCallstack()->Num() == 0 || (Alloc.GetAllocCallstack()->Num() != 0 && Alloc.GetAllocCallstack()->Num() < 256)))
 			{
 				Test->AddError(TEXT("Resolved alloc callstack should be valid"));
@@ -236,7 +248,7 @@ bool FVerifyHierarchyCallStackCommand::Update()
 	return false;
 }
 
-const TMap<TraceServices::IAllocationsProvider::EQueryRule, Insights::SMemAllocTableTreeView::FQueryParams> AllocsTimeMarkerStandaloneGameGetterMap
+const TMap<TraceServices::IAllocationsProvider::EQueryRule, UE::Insights::MemoryProfiler::SMemAllocTableTreeView::FQueryParams> AllocsTimeMarkerStandaloneGameGetterMap
 {
 	{TraceServices::IAllocationsProvider::EQueryRule::aAf, {nullptr, {5.0, 0.0, 0.0, 0.0}}},
 	{TraceServices::IAllocationsProvider::EQueryRule::afA, {nullptr, {10.0, 0.0, 0.0, 0.0}}},
@@ -254,7 +266,7 @@ const TMap<TraceServices::IAllocationsProvider::EQueryRule, Insights::SMemAllocT
 	{TraceServices::IAllocationsProvider::EQueryRule::AaB, {nullptr, {50.0, 51.0, 0.0, 0.0}}},
 };
 
-const TMap<TraceServices::IAllocationsProvider::EQueryRule, Insights::SMemAllocTableTreeView::FQueryParams> AllocsTimeMarkerEditorPackageGetterMap
+const TMap<TraceServices::IAllocationsProvider::EQueryRule, UE::Insights::MemoryProfiler::SMemAllocTableTreeView::FQueryParams> AllocsTimeMarkerEditorPackageGetterMap
 {
 	{TraceServices::IAllocationsProvider::EQueryRule::aAf, {nullptr, {5.0, 0.0, 0.0, 0.0}}},
 	{TraceServices::IAllocationsProvider::EQueryRule::afA, {nullptr, {10.0, 0.0, 0.0, 0.0}}},
@@ -272,16 +284,18 @@ const TMap<TraceServices::IAllocationsProvider::EQueryRule, Insights::SMemAllocT
 	{TraceServices::IAllocationsProvider::EQueryRule::aABf, {nullptr, {2.0, 3.0, 0.0, 0.0}}},
 };
 
-bool MemoryInsightsAllocationsQueryTableTest(const FString& Parameters, const TMap <TraceServices::IAllocationsProvider::EQueryRule, Insights::SMemAllocTableTreeView::FQueryParams> AllocsTimeMarkerGetterMap, FAutomationTestBase* Test)
+bool MemoryInsightsAllocationsQueryTableTest(const FString& Parameters, const TMap <TraceServices::IAllocationsProvider::EQueryRule, UE::Insights::MemoryProfiler::SMemAllocTableTreeView::FQueryParams> AllocsTimeMarkerGetterMap, FAutomationTestBase* Test)
 {
+	using namespace UE::Insights::MemoryProfiler;
+
 	double Timeout = 30.0;
 	FInsightsTestUtils InsightsTestUtils(Test);
 	TSharedPtr<SMemoryProfilerWindow> ProfilerWindow = FMemoryProfilerManager::Get()->GetProfilerWindow();
 	TSharedPtr<FInsightsManager> InsightsManager = FInsightsManager::Get();
 	FMemorySharedState& SharedState = ProfilerWindow->GetSharedState();
 
-	TSharedPtr<Insights::FMemoryRuleSpec> MemoryRule = *Algo::FindByPredicate(SharedState.GetMemoryRules(),
-		[&Parameters](const TSharedPtr<Insights::FMemoryRuleSpec>& Rule)
+	TSharedPtr<FMemoryRuleSpec> MemoryRule = *Algo::FindByPredicate(SharedState.GetMemoryRules(),
+		[&Parameters](const TSharedPtr<FMemoryRuleSpec>& Rule)
 		{
 			return Rule->GetShortName().ToString().Contains(Parameters);
 		});
@@ -292,9 +306,9 @@ bool MemoryInsightsAllocationsQueryTableTest(const FString& Parameters, const TM
 		return false;
 	}
 
-	TSharedPtr<Insights::SMemAllocTableTreeView> MemAllocTableTreeView = ProfilerWindow->ShowMemAllocTableTreeViewTab();
+	TSharedPtr<SMemAllocTableTreeView> MemAllocTableTreeView = ProfilerWindow->ShowMemAllocTableTreeViewTab();
 
-	Insights::SMemAllocTableTreeView::FQueryParams QueryParams = AllocsTimeMarkerGetterMap.FindChecked(MemoryRule->GetValue());
+	SMemAllocTableTreeView::FQueryParams QueryParams = AllocsTimeMarkerGetterMap.FindChecked(MemoryRule->GetValue());
 	if (MemoryRule->GetValue() == TraceServices::IAllocationsProvider::EQueryRule::Aaf)
 	{
 		QueryParams.TimeMarkers[0] = InsightsManager->GetSessionDuration() - 10.0;
@@ -319,6 +333,8 @@ bool FMemoryInsightsAllocationsQueryTableEditorPackageTest::RunTest(const FStrin
 
 void FMemoryInsightsAllocationsQueryTableEditorPackageTest::GetTests(TArray<FString>& OutBeautifiedNames, TArray <FString>& OutTestCommands) const
 {
+	using namespace UE::Insights::MemoryProfiler;
+
 	if (!FMemoryProfilerManager::Get().IsValid())
 	{
 		return;
@@ -349,6 +365,8 @@ bool FMemoryInsightsAllocationsQueryTableStandaloneTest::RunTest(const FString&
 
 void FMemoryInsightsAllocationsQueryTableStandaloneTest::GetTests(TArray<FString>& OutBeautifiedNames, TArray <FString>& OutTestCommands) const
 {
+	using namespace UE::Insights::MemoryProfiler;
+
 	if (!FMemoryProfilerManager::Get().IsValid())
 	{
 		return;
@@ -360,7 +378,7 @@ void FMemoryInsightsAllocationsQueryTableStandaloneTest::GetTests(TArray<FString
 	}
 
 	FMemorySharedState& SharedState = ProfilerWindow->GetSharedState();
-	TSharedPtr<Insights::FMemoryRuleSpec> MemoryRule = SharedState.GetMemoryRules()[0];
+	TSharedPtr<FMemoryRuleSpec> MemoryRule = SharedState.GetMemoryRules()[0];
 
 	for (const auto& MemoryRules : SharedState.GetMemoryRules())
 	{
diff --git a/Engine/Source/Developer/TraceInsights/Private/Insights/Tests/FunctionalTests/SessionBrowserTests.cpp b/Engine/Source/Developer/TraceInsights/Private/Insights/Tests/FunctionalTests/SessionBrowserTests.cpp
--- a/Engine/Source/Developer/TraceInsights/Private/Insights/Tests/FunctionalTests/SessionBrowserTests.cpp
+++ b/Engine/Source/Developer/TraceInsights/Private/Insights/Tests/FunctionalTests/SessionBrowserTests.cpp
@@ -4,14 +4,18 @@
 #include "Framework/Application/SlateApplication.h"
 #include "HAL/FileManager.h"
 #include "HAL/PlatformFileManager.h"
+#include "Logging/LogMacros.h"
 #include "Misc/AutomationTest.h"
 #include "Misc/FileHelper.h"
+
+// TraceServices
+#include "TraceServices/Model/AllocationsProvider.h"
+
+// TraceInsights
 #include "Insights/InsightsManager.h"
 #include "Insights/IUnrealInsightsModule.h"
 #include "Insights/Tests/InsightsTestUtils.h"
 #include "Insights/Widgets/SStartPageWindow.h"
-#include "Logging/LogMacros.h"
-#include "TraceServices/Model/AllocationsProvider.h"
 
 DECLARE_LOG_CATEGORY_EXTERN(UnrealInsightsHubTests, Log, All);
 
@@ -47,7 +51,7 @@ void FAutomationDriverUnrealInsightsSessionBrowserTest::Define()
 				{
 					IPlatformFile& PlatformFile = FPlatformFileManager::Get().GetPlatformFile();
 					TSharedPtr<FInsightsManager> InsightsManager = FInsightsManager::Get();
-					TestTrue("Insigts manager should not be null", InsightsManager.IsValid());
+					TestTrue("Insights manager should not be null", InsightsManager.IsValid());
 					InsightsManager->GetTraceStoreWindow()->SetDeleteTraceConfirmationWindowVisibility(false);
 
 					const FString StoreDir = InsightsManager->GetStoreDir();
diff --git a/Engine/Source/Developer/TraceInsights/Private/Insights/Tests/FunctionalTests/TimingProfilerTests.cpp b/Engine/Source/Developer/TraceInsights/Private/Insights/Tests/FunctionalTests/TimingProfilerTests.cpp
--- a/Engine/Source/Developer/TraceInsights/Private/Insights/Tests/FunctionalTests/TimingProfilerTests.cpp
+++ b/Engine/Source/Developer/TraceInsights/Private/Insights/Tests/FunctionalTests/TimingProfilerTests.cpp
@@ -2,11 +2,16 @@
 
 #include "Insights/Tests/TimingProfilerTests.h"
 
-#include "TraceServices/Model/Threads.h"
 #include "Misc/AutomationTest.h"
 #include "Misc/FileHelper.h"
 
-#include "Insights/Common/Stopwatch.h"
+// TraceServices
+#include "TraceServices/Model/Threads.h"
+
+// TraceInsightsCore
+#include "InsightsCore/Common/Stopwatch.h"
+
+// TraceInsights
 #include "Insights/InsightsManager.h"
 #include "Insights/Tests/InsightsTestUtils.h"
 
@@ -339,7 +344,7 @@ void FTimingProfilerTests::RunEnumerateBenchmark(const FEnumerateTestParams& InP
 {
 	UE_LOG(TimingProfilerTests, Log, TEXT("RUNNING BENCHMARK..."));
 
-	FStopwatch Stopwatch;
+	UE::Insights::FStopwatch Stopwatch;
 	Stopwatch.Start();
 
 	double SessionTime = 0.0;
@@ -393,7 +398,7 @@ void FTimingProfilerTests::RunEnumerateAsyncBenchmark(const FEnumerateTestParams
 {
 	UE_LOG(TimingProfilerTests, Log, TEXT("RUNNING ASYNC ENUMERATE BENCHMARK..."));
 
-	FStopwatch Stopwatch;
+	UE::Insights::FStopwatch Stopwatch;
 	Stopwatch.Start();
 
 	double SessionTime = 0.0;
@@ -467,7 +472,7 @@ void FTimingProfilerTests::RunEnumerateAllTracksBenchmark(const FEnumerateTestPa
 {
 	UE_LOG(TimingProfilerTests, Log, TEXT("RUNNING ENUMERATE ALL TRACKS BENCHMARK..."));
 
-	FStopwatch Stopwatch;
+	UE::Insights::FStopwatch Stopwatch;
 	Stopwatch.Start();
 
 	double SessionTime = 0.0;
@@ -526,7 +531,7 @@ void FTimingProfilerTests::RunEnumerateAsyncAllTracksBenchmark(const FEnumerateT
 {
 	UE_LOG(TimingProfilerTests, Log, TEXT("RUNNING ASYNC ENUMERATE ALL TRACKS BENCHMARK..."));
 
-	FStopwatch Stopwatch;
+	UE::Insights::FStopwatch Stopwatch;
 	Stopwatch.Start();
 
 	double SessionTime = 0.0;
diff --git a/Engine/Source/Developer/TraceInsights/Private/Insights/Tests/FunctionalTests/UITests.cpp b/Engine/Source/Developer/TraceInsights/Private/Insights/Tests/FunctionalTests/UITests.cpp
--- a/Engine/Source/Developer/TraceInsights/Private/Insights/Tests/FunctionalTests/UITests.cpp
+++ b/Engine/Source/Developer/TraceInsights/Private/Insights/Tests/FunctionalTests/UITests.cpp
@@ -1,22 +1,30 @@
 // Copyright Epic Games, Inc. All Rights Reserved.
 
 #include "CoreMinimal.h"
-#include "Insights/Common/Stopwatch.h"
+
+#include "Logging/LogMacros.h"
+#include "Misc/AutomationTest.h"
+
+// TraceServices
+#include "TraceServices/Model/TimingProfiler.h"
+
+// TraceInsightsCore
+#include "InsightsCore/Common/Stopwatch.h"
+#include "InsightsCore/Filter/ViewModels/TimeFilterValueConverter.h"
+
+// TraceInsights
 #include "Insights/MemoryProfiler/ViewModels/MemAllocFilterValueConverter.h"
 #include "Insights/TaskGraphProfiler/ViewModels/TaskTimingTrack.h"
 #include "Insights/TimingProfilerManager.h"
-#include "Insights/ViewModels/TimeFilterValueConverter.h"
 #include "Insights/ViewModels/TimeRulerTrack.h"
 #include "Insights/Widgets/STimingView.h"
-#include "Logging/LogMacros.h"
-#include "Misc/AutomationTest.h"
-#include "TraceServices/Model/TimingProfiler.h"
 
 DECLARE_LOG_CATEGORY_EXTERN(UITests, Log, All);
 
 IMPLEMENT_SIMPLE_AUTOMATION_TEST(FHideAndShowAllTimingViewTabs, "System.Insights.Analysis.TimingInsights.HideAndShowAllTimingViewTabs", EAutomationTestFlags::ProgramContext | EAutomationTestFlags::EngineFilter)
 bool FHideAndShowAllTimingViewTabs::RunTest(const FString& Parameters)
 {
+	using namespace UE::Insights::TimingProfiler;
 	TSharedPtr<FTimingProfilerManager> TimingProfilerManager = FTimingProfilerManager::Get();
 
 	TimingProfilerManager->ShowHideTimingView(false);
@@ -39,7 +47,8 @@ bool FHideAndShowAllTimingViewTabs::RunTest(const FString& Parameters)
 IMPLEMENT_SIMPLE_AUTOMATION_TEST(FMemoryFilterValueConverterTest, "System.Insights.Analysis.MemoryFilterValueConverter", EAutomationTestFlags_ApplicationContextMask | EAutomationTestFlags::EngineFilter)
 bool FMemoryFilterValueConverterTest::RunTest(const FString& Parameters)
 {
-	Insights::FMemoryFilterValueConverter Converter;
+	using namespace UE::Insights::MemoryProfiler;
+	FMemoryFilterValueConverter Converter;
 
 	FText Error;
 	int64 Value;
@@ -75,7 +84,7 @@ bool FMemoryFilterValueConverterTest::RunTest(const FString& Parameters)
 IMPLEMENT_SIMPLE_AUTOMATION_TEST(FTimeFilterValueConverterTest, "System.Insights.Analysis.TimeFilterValueConverter", EAutomationTestFlags_ApplicationContextMask | EAutomationTestFlags::EngineFilter)
 bool FTimeFilterValueConverterTest::RunTest(const FString& Parameters)
 {
-	Insights::FTimeFilterValueConverter Converter;
+	UE::Insights::FTimeFilterValueConverter Converter;
 
 	FText Error;
 	double Value;
diff --git a/Engine/Source/Developer/TraceInsights/Private/Insights/Tests/InsightsTestUtils.cpp b/Engine/Source/Developer/TraceInsights/Private/Insights/Tests/InsightsTestUtils.cpp
--- a/Engine/Source/Developer/TraceInsights/Private/Insights/Tests/InsightsTestUtils.cpp
+++ b/Engine/Source/Developer/TraceInsights/Private/Insights/Tests/InsightsTestUtils.cpp
@@ -2,20 +2,27 @@
 
 #include "Insights/Tests/InsightsTestUtils.h"
 
-#include "Modules/ModuleManager.h"
-#include "Misc/FileHelper.h"
 #include "HAL/FileManagerGeneric.h"
+#include "Misc/AutomationTest.h"
+#include "Misc/FileHelper.h"
+#include "Modules/ModuleManager.h"
+
+// TraceAnalysis
+#include "Trace/StoreClient.h"
 
+// TraceServices
 #include "TraceServices/AnalysisService.h"
-#include "TraceServices/Model/AnalysisSession.h"
 #include "TraceServices/ITraceServicesModule.h"
+#include "TraceServices/Model/AnalysisSession.h"
 #include "TraceServices/ModuleService.h"
-#include "Trace/StoreClient.h"
-#include "Insights/Common/Stopwatch.h"
-#include "Insights/IUnrealInsightsModule.h"
+
+// TraceInsightsCore
+#include "InsightsCore/Common/Stopwatch.h"
+
+// TraceInsights
 #include "Insights/InsightsManager.h"
+#include "Insights/IUnrealInsightsModule.h"
 
-#include "Misc/AutomationTest.h"
 
 ////////////////////////////////////////////////////////////////////////////////////////////////////
 
@@ -49,7 +56,7 @@ bool FInsightsTestUtils::AnalyzeTrace(const TCHAR* Path) const
 		return false;
 	}
 
-	FStopwatch StopWatch;
+	UE::Insights::FStopwatch StopWatch;
 	StopWatch.Start();
 
 	double Duration = 0.0f;
diff --git a/Engine/Source/Developer/TraceInsights/Private/Insights/TimingProfilerCommands.cpp b/Engine/Source/Developer/TraceInsights/Private/Insights/TimingProfilerCommands.cpp
--- a/Engine/Source/Developer/TraceInsights/Private/Insights/TimingProfilerCommands.cpp
+++ b/Engine/Source/Developer/TraceInsights/Private/Insights/TimingProfilerCommands.cpp
@@ -6,14 +6,17 @@
 #include "Framework/Application/SlateApplication.h"
 #include "Framework/MultiBox/MultiBoxBuilder.h"
 
-// Insights
+// TraceInsights
 #include "Insights/InsightsManager.h"
 #include "Insights/InsightsStyle.h"
 #include "Insights/TimingProfilerManager.h"
 
 ////////////////////////////////////////////////////////////////////////////////////////////////////
 
-#define LOCTEXT_NAMESPACE "FTimingProfilerCommands"
+#define LOCTEXT_NAMESPACE "UE::Insights::TimingProfiler"
+
+namespace UE::Insights::TimingProfiler
+{
 
 ////////////////////////////////////////////////////////////////////////////////////////////////////
 // FTimingProfilerMenuBuilder
@@ -173,4 +176,6 @@ IMPLEMENT_TOGGLE_COMMAND(ToggleLogViewVisibility, IsLogViewVisible, ShowHideLogV
 
 ////////////////////////////////////////////////////////////////////////////////////////////////////
 
+} // namespace UE::Insights::TimingProfiler
+
 #undef LOCTEXT_NAMESPACE
diff --git a/Engine/Source/Developer/TraceInsights/Private/Insights/TimingProfilerCommands.h b/Engine/Source/Developer/TraceInsights/Private/Insights/TimingProfilerCommands.h
--- a/Engine/Source/Developer/TraceInsights/Private/Insights/TimingProfilerCommands.h
+++ b/Engine/Source/Developer/TraceInsights/Private/Insights/TimingProfilerCommands.h
@@ -2,12 +2,16 @@
 
 #pragma once
 
-#include "CoreMinimal.h"
+#include "CoreTypes.h"
+
 #include "Framework/Commands/UIAction.h"
 #include "Framework/Commands/Commands.h"
 
 class FMenuBuilder;
 
+namespace UE::Insights::TimingProfiler
+{
+
 ////////////////////////////////////////////////////////////////////////////////////////////////////
 /**
  * Class that holds all profiler commands.
@@ -149,3 +153,5 @@ protected:
 	/** Reference to the global instance of the profiler manager. */
 	class FTimingProfilerManager* This;
 };
+
+} // namespace UE::Insights::TimingProfiler
diff --git a/Engine/Source/Developer/TraceInsights/Private/Insights/TimingProfilerManager.cpp b/Engine/Source/Developer/TraceInsights/Private/Insights/TimingProfilerManager.cpp
--- a/Engine/Source/Developer/TraceInsights/Private/Insights/TimingProfilerManager.cpp
+++ b/Engine/Source/Developer/TraceInsights/Private/Insights/TimingProfilerManager.cpp
@@ -4,12 +4,14 @@
 
 #include "MessageLogModule.h"
 #include "Modules/ModuleManager.h"
-#include "TraceServices/Model/Counters.h"
 #include "Widgets/Docking/SDockTab.h"
 #include "WorkspaceMenuStructure.h"
 #include "WorkspaceMenuStructureModule.h"
 
-// Insights
+// TraceServices
+#include "TraceServices/Model/Counters.h"
+
+// TraceInsights
 #include "Insights/Common/InsightsMenuBuilder.h"
 #include "Insights/InsightsManager.h"
 #include "Insights/InsightsStyle.h"
@@ -26,9 +28,12 @@
 
 ////////////////////////////////////////////////////////////////////////////////////////////////////
 
-#define LOCTEXT_NAMESPACE "TimingProfilerManager"
+DEFINE_LOG_CATEGORY(LogTimingProfiler);
 
-DEFINE_LOG_CATEGORY(TimingProfiler);
+#define LOCTEXT_NAMESPACE "UE::Insights::TimingProfiler"
+
+namespace UE::Insights::TimingProfiler
+{
 
 TSharedPtr<FTimingProfilerManager> FTimingProfilerManager::Instance = nullptr;
 
@@ -72,7 +77,7 @@ FTimingProfilerManager::FTimingProfilerManager(TSharedRef<FUICommandList> InComm
 	, SelectionStartTime(0.0)
 	, SelectionEndTime(0.0)
 	, SelectedTimerId(InvalidTimerId)
-	, TimerButterflyAggregator(MakeShared<Insights::FTimerButterflyAggregator>())
+	, TimerButterflyAggregator(MakeShared<FTimerButterflyAggregator>())
 	, LogListingName(TEXT("TimingInsights"))
 {
 }
@@ -88,7 +93,7 @@ void FTimingProfilerManager::Initialize(IUnrealInsightsModule& InsightsModule)
 	}
 	bIsInitialized = true;
 
-	UE_LOG(TimingProfiler, Log, TEXT("Initialize"));
+	UE_LOG(LogTimingProfiler, Log, TEXT("Initialize"));
 
 	// Register tick functions.
 	OnTick = FTickerDelegate::CreateSP(this, &FTimingProfilerManager::Tick);
@@ -132,7 +137,7 @@ void FTimingProfilerManager::Shutdown()
 
 	FTimingProfilerManager::Instance.Reset();
 
-	UE_LOG(TimingProfiler, Log, TEXT("Shutdown"));
+	UE_LOG(LogTimingProfiler, Log, TEXT("Shutdown"));
 }
 
 ////////////////////////////////////////////////////////////////////////////////////////////////////
@@ -266,7 +271,7 @@ bool FTimingProfilerManager::Tick(float DeltaTime)
 			const FName& TabId = FInsightsManagerTabs::TimingProfilerTabId;
 			if (FGlobalTabmanager::Get()->HasTabSpawner(TabId))
 			{
-				UE_LOG(TimingProfiler, Log, TEXT("Opening the \"Timing Insights\" tab..."));
+				UE_LOG(LogTimingProfiler, Log, TEXT("Opening the \"Timing Insights\" tab..."));
 				FGlobalTabmanager::Get()->TryInvokeTab(TabId);
 			}
 #endif
@@ -314,7 +319,7 @@ void FTimingProfilerManager::FinishTimerButterflyAggregation()
 
 void FTimingProfilerManager::OnSessionChanged()
 {
-	UE_LOG(TimingProfiler, Log, TEXT("OnSessionChanged"));
+	UE_LOG(LogTimingProfiler, Log, TEXT("OnSessionChanged"));
 
 	bIsAvailable = false;
 	if (FInsightsManager::Get()->GetSession().IsValid())
@@ -656,12 +661,14 @@ void FTimingProfilerManager::OnWindowClosedEvent()
 
 bool FTimingProfilerManager::Exec(const TCHAR* Cmd, FOutputDevice& Ar)
 {
+	using FTimingExporter = ::Insights::FTimingExporter;
+
 	if (FParse::Command(&Cmd, TEXT("TimingInsights.ExportThreads")))
 	{
 		Ar.Logf(TEXT("TimingInsights.ExportThreads %s"), Cmd);
 		check(FInsightsManager::Get().IsValid() && FInsightsManager::Get()->GetSession().IsValid());
-		Insights::FTimingExporter Exporter(*FInsightsManager::Get()->GetSession().Get());
-		Insights::FTimingExporter::FExportThreadsParams Params; // default
+		FTimingExporter Exporter(*FInsightsManager::Get()->GetSession().Get());
+		FTimingExporter::FExportThreadsParams Params; // default
 
 		const bool bUseEscape = true;
 		FString Filename = FParse::Token(Cmd, bUseEscape);
@@ -676,8 +683,8 @@ bool FTimingProfilerManager::Exec(const TCHAR* Cmd, FOutputDevice& Ar)
 	{
 		Ar.Logf(TEXT("TimingInsights.ExportTimers %s"), Cmd);
 		check(FInsightsManager::Get().IsValid() && FInsightsManager::Get()->GetSession().IsValid());
-		Insights::FTimingExporter Exporter(*FInsightsManager::Get()->GetSession().Get());
-		Insights::FTimingExporter::FExportTimersParams Params; // default
+		FTimingExporter Exporter(*FInsightsManager::Get()->GetSession().Get());
+		FTimingExporter::FExportTimersParams Params; // default
 
 		const bool bUseEscape = true;
 		FString Filename = FParse::Token(Cmd, bUseEscape);
@@ -693,8 +700,8 @@ bool FTimingProfilerManager::Exec(const TCHAR* Cmd, FOutputDevice& Ar)
 		Ar.Logf(TEXT("TimingInsights.ExportTimingEvents %s"), Cmd);
 
 		check(FInsightsManager::Get().IsValid() && FInsightsManager::Get()->GetSession().IsValid());
-		Insights::FTimingExporter Exporter(*FInsightsManager::Get()->GetSession().Get());
-		Insights::FTimingExporter::FExportTimingEventsParams Params; // default (all timing events)
+		FTimingExporter Exporter(*FInsightsManager::Get()->GetSession().Get());
+		FTimingExporter::FExportTimingEventsParams Params; // default (all timing events)
 
 		// These variables needs to be in the same scope with the call to Exporter.ExportTimingEventsAsText().
 		TArray<FName> Columns; // referenced by Params.Columns
@@ -794,8 +801,8 @@ bool FTimingProfilerManager::Exec(const TCHAR* Cmd, FOutputDevice& Ar)
 		Ar.Logf(TEXT("TimingInsights.ExportTimerStatistics %s"), Cmd);
 
 		check(FInsightsManager::Get().IsValid() && FInsightsManager::Get()->GetSession().IsValid());
-		Insights::FTimingExporter Exporter(*FInsightsManager::Get()->GetSession().Get());
-		Insights::FTimingExporter::FExportTimerStatisticsParams Params; // default (all timing events)
+		FTimingExporter Exporter(*FInsightsManager::Get()->GetSession().Get());
+		FTimingExporter::FExportTimerStatisticsParams Params; // default (all timing events)
 
 		// These variables needs to be in the same scope with the call to Exporter.ExportTimerStatisticsAsText().
 		TArray<FName> Columns; // referenced by Params.Columns
@@ -899,8 +906,8 @@ bool FTimingProfilerManager::Exec(const TCHAR* Cmd, FOutputDevice& Ar)
 	{
 		Ar.Logf(TEXT("TimingInsights.ExportCounters %s"), Cmd);
 		check(FInsightsManager::Get().IsValid() && FInsightsManager::Get()->GetSession().IsValid());
-		Insights::FTimingExporter Exporter(*FInsightsManager::Get()->GetSession().Get());
-		Insights::FTimingExporter::FExportCountersParams Params; // default
+		FTimingExporter Exporter(*FInsightsManager::Get()->GetSession().Get());
+		FTimingExporter::FExportCountersParams Params; // default
 		Exporter.ExportCountersAsText(Cmd, Params);
 		return true;
 	}
@@ -909,8 +916,8 @@ bool FTimingProfilerManager::Exec(const TCHAR* Cmd, FOutputDevice& Ar)
 	{
 		Ar.Logf(TEXT("TimingInsights.ExportCounterValues %s"), Cmd);
 		check(FInsightsManager::Get().IsValid() && FInsightsManager::Get()->GetSession().IsValid());
-		Insights::FTimingExporter Exporter(*FInsightsManager::Get()->GetSession().Get());
-		Insights::FTimingExporter::FExportCounterParams Params; // default
+		FTimingExporter Exporter(*FInsightsManager::Get()->GetSession().Get());
+		FTimingExporter::FExportCounterParams Params; // default
 
 		TArray<FString> Counters; // list of counters to export
 
@@ -1033,4 +1040,6 @@ bool FTimingProfilerManager::Exec(const TCHAR* Cmd, FOutputDevice& Ar)
 
 ////////////////////////////////////////////////////////////////////////////////////////////////////
 
+} // namespace UE::Insights::TimingProfiler
+
 #undef LOCTEXT_NAMESPACE
diff --git a/Engine/Source/Developer/TraceInsights/Private/Insights/TimingProfilerManager.h b/Engine/Source/Developer/TraceInsights/Private/Insights/TimingProfilerManager.h
--- a/Engine/Source/Developer/TraceInsights/Private/Insights/TimingProfilerManager.h
+++ b/Engine/Source/Developer/TraceInsights/Private/Insights/TimingProfilerManager.h
@@ -2,20 +2,19 @@
 
 #pragma once
 
+#include "CoreTypes.h"
+
 #include "Containers/Ticker.h"
-#include "CoreMinimal.h"
 #include "Framework/Commands/UICommandList.h"
 
-// Insights
+// TraceInsights
 #include "Insights/InsightsManager.h"
 #include "Insights/IUnrealInsightsModule.h"
 #include "Insights/TimingProfilerCommands.h"
 #include "Insights/ViewModels/TimerNode.h"
 
-namespace Insights
+namespace UE::Insights
 {
-	class FTimerButterflyAggregator;
-
 	enum class ETimingEventsColoringMode : uint32
 	{
 		ByTimerName,
@@ -27,6 +26,10 @@ namespace Insights
 	};
 }
 
+namespace UE::Insights::TimingProfiler
+{
+
+class FTimerButterflyAggregator;
 class STimingProfilerWindow;
 
 ////////////////////////////////////////////////////////////////////////////////////////////////////
@@ -144,15 +147,15 @@ public:
 
 	void ResetCallersAndCallees();
 	void UpdateCallersAndCallees();
-	TSharedRef<Insights::FTimerButterflyAggregator> GetTimerButterflyAggregator() const { return TimerButterflyAggregator; }
+	TSharedRef<FTimerButterflyAggregator> GetTimerButterflyAggregator() const { return TimerButterflyAggregator; }
 
 	void UpdateAggregatedTimerStats();
 	void UpdateAggregatedCounterStats();
 
 	const FName& GetLogListingName() const { return LogListingName; }
 
-	Insights::ETimingEventsColoringMode GetColoringMode() const { return ColoringMode; }
-	void SetColoringMode(Insights::ETimingEventsColoringMode InColoringMode) { ColoringMode = InColoringMode; }
+	ETimingEventsColoringMode GetColoringMode() const { return ColoringMode; }
+	void SetColoringMode(ETimingEventsColoringMode InColoringMode) { ColoringMode = InColoringMode; }
 
 	uint32 GetEventDepthLimit() const { return EventDepthLimit; }
 	void SetEventDepthLimit(uint32 InEventDepthLimit) { EventDepthLimit = InEventDepthLimit; }
@@ -187,7 +190,7 @@ private:
 private:
 	bool bIsInitialized;
 	bool bIsAvailable;
-	Insights::FAvailabilityCheck AvailabilityCheck;
+	FAvailabilityCheck AvailabilityCheck;
 
 	/** The delegate to be invoked when this manager ticks. */
 	FTickerDelegate OnTick;
@@ -233,14 +236,16 @@ private:
 	static constexpr uint32 InvalidTimerId = uint32(-1);
 	uint32 SelectedTimerId;
 
-	TSharedRef<Insights::FTimerButterflyAggregator> TimerButterflyAggregator;
+	TSharedRef<FTimerButterflyAggregator> TimerButterflyAggregator;
 
 	/** The name of the Timing Insights log listing. */
 	FName LogListingName;
 
-	Insights::ETimingEventsColoringMode ColoringMode = Insights::ETimingEventsColoringMode::ByTimerName;
+	ETimingEventsColoringMode ColoringMode = ETimingEventsColoringMode::ByTimerName;
 	uint32 EventDepthLimit = UnlimitedEventDepth;
 
 	/** A shared pointer to the global instance of the Timing Profiler manager. */
 	static TSharedPtr<FTimingProfilerManager> Instance;
 };
+
+} // namespace UE::Insights::TimingProfiler
diff --git a/Engine/Source/Developer/TraceInsights/Private/Insights/TraceInsightsModule.cpp b/Engine/Source/Developer/TraceInsights/Private/Insights/TraceInsightsModule.cpp
--- a/Engine/Source/Developer/TraceInsights/Private/Insights/TraceInsightsModule.cpp
+++ b/Engine/Source/Developer/TraceInsights/Private/Insights/TraceInsightsModule.cpp
@@ -1,6 +1,7 @@
 // Copyright Epic Games, Inc. All Rights Reserved.
 
 #include "TraceInsightsModule.h"
+
 #include "Framework/Application/SlateApplication.h"
 #include "Framework/Docking/LayoutService.h"
 #include "Framework/Notifications/NotificationManager.h"
@@ -12,17 +13,23 @@
 #include "Misc/ConfigCacheIni.h"
 #include "Misc/Paths.h"
 #include "Modules/ModuleManager.h"
+
+// TraceAnalysis
 #include "Trace/StoreClient.h"
 #include "Trace/StoreService.h"
+
+// TraceServices
 #include "TraceServices/ITraceServicesModule.h"
 
-// Insights
+// TraceInsightsCore
+#include "InsightsCore/ITraceInsightsCoreModule.h"
+
+// TraceInsights
 #include "Insights/ContextSwitches/ContextSwitchesProfilerManager.h"
 #include "Insights/CookProfiler/CookProfilerManager.h"
 #include "Insights/ImportTool/TableImportTool.h"
 #include "Insights/InsightsManager.h"
 #include "Insights/InsightsStyle.h"
-#include "Insights/IUnrealInsightsModule.h"
 #include "Insights/LoadingProfiler/LoadingProfilerManager.h"
 #include "Insights/Log.h"
 #include "Insights/MemoryProfiler/MemoryProfilerManager.h"
@@ -36,8 +43,6 @@
 
 DEFINE_LOG_CATEGORY(TraceInsights);
 
-LLM_DEFINE_TAG(Insights);
-
 IMPLEMENT_MODULE(FTraceInsightsModule, TraceInsights);
 
 FString FTraceInsightsModule::UnrealInsightsLayoutIni;
@@ -50,6 +55,8 @@ void FTraceInsightsModule::StartupModule()
 {
 	LLM_SCOPE_BYTAG(Insights);
 
+	ITraceInsightsCoreModule& TraceInsightsCoreModule = FModuleManager::LoadModuleChecked<ITraceInsightsCoreModule>("TraceInsightsCore");
+
 	ITraceServicesModule& TraceServicesModule = FModuleManager::LoadModuleChecked<ITraceServicesModule>("TraceServices");
 	TraceAnalysisService = TraceServicesModule.GetAnalysisService();
 	TraceModuleService = TraceServicesModule.GetModuleService();
@@ -63,10 +70,10 @@ void FTraceInsightsModule::StartupModule()
 	RegisterComponent(FInsightsManager::CreateInstance(TraceAnalysisService.ToSharedRef(), TraceModuleService.ToSharedRef()));
 
 	// Register other default components.
-	RegisterComponent(FTimingProfilerManager::CreateInstance());
+	RegisterComponent(UE::Insights::TimingProfiler::FTimingProfilerManager::CreateInstance());
 	RegisterComponent(FLoadingProfilerManager::CreateInstance());
-	RegisterComponent(FNetworkingProfilerManager::CreateInstance());
-	RegisterComponent(FMemoryProfilerManager::CreateInstance());
+	RegisterComponent(UE::Insights::NetworkingProfiler::FNetworkingProfilerManager::CreateInstance());
+	RegisterComponent(UE::Insights::MemoryProfiler::FMemoryProfilerManager::CreateInstance());
 	RegisterComponent(Insights::FTaskGraphProfilerManager::CreateInstance());
 	RegisterComponent(Insights::FContextSwitchesProfilerManager::CreateInstance());
 	RegisterComponent(Insights::FCookProfilerManager::CreateInstance());
@@ -106,10 +113,10 @@ void FTraceInsightsModule::ShutdownModule()
 #define INSIGHTS_CHECK_SHARED_REFERENCES 1
 #if INSIGHTS_CHECK_SHARED_REFERENCES
 	TSharedPtr<const TraceServices::IAnalysisSession> Session = FInsightsManager::Get()->GetSession();
-	auto TimingInsightsWindow = FTimingProfilerManager::Get()->GetProfilerWindow();
+	auto TimingInsightsWindow = UE::Insights::TimingProfiler::FTimingProfilerManager::Get()->GetProfilerWindow();
 	auto AssetLoadingInsightsWindow = FLoadingProfilerManager::Get()->GetProfilerWindow();
-	auto NetworkingInsightsWindow0 = FNetworkingProfilerManager::Get()->GetProfilerWindow(0);
-	auto MemoryInsightsWindow = FMemoryProfilerManager::Get()->GetProfilerWindow();
+	auto NetworkingInsightsWindow0 = UE::Insights::NetworkingProfiler::FNetworkingProfilerManager::Get()->GetProfilerWindow(0);
+	auto MemoryInsightsWindow = UE::Insights::MemoryProfiler::FMemoryProfilerManager::Get()->GetProfilerWindow();
 #endif
 
 	// Unregister components. Shutdown in the reverse order they were registered.
@@ -328,7 +335,7 @@ void FTraceInsightsModule::CreateSessionBrowser(const FCreateSessionBrowserParam
 	//////////////////////////////////////////////////
 	// Set up command line parameter forwarding.
 
-	TSharedPtr<class STraceStoreWindow> TraceStoreWnd = FInsightsManager::Get()->GetTraceStoreWindow();
+	TSharedPtr<STraceStoreWindow> TraceStoreWnd = FInsightsManager::Get()->GetTraceStoreWindow();
 	if (TraceStoreWnd.IsValid())
 	{
 		TraceStoreWnd->SetEnableAutomaticTesting(Params.bInitializeTesting);
diff --git a/Engine/Source/Developer/TraceInsights/Private/Insights/TraceInsightsModule.h b/Engine/Source/Developer/TraceInsights/Private/Insights/TraceInsightsModule.h
--- a/Engine/Source/Developer/TraceInsights/Private/Insights/TraceInsightsModule.h
+++ b/Engine/Source/Developer/TraceInsights/Private/Insights/TraceInsightsModule.h
@@ -2,9 +2,12 @@
 
 #pragma once
 
-#include "CoreMinimal.h"
+#include "CoreTypes.h"
+
 #include "Insights/IUnrealInsightsModule.h"
 #include "Framework/Docking/TabManager.h"
+
+// TraceAnalysis
 #include "Trace/StoreService.h"
 
 namespace Insights
@@ -28,6 +31,10 @@ class SWindow;
 class FTraceInsightsModule : public IUnrealInsightsModule
 {
 public:
+	virtual ~FTraceInsightsModule()
+	{
+	}
+
 	virtual void StartupModule() override;
 	virtual void ShutdownModule() override;
 
diff --git a/Engine/Source/Developer/TraceInsights/Private/Insights/ViewModels/BaseTimingTrack.cpp b/Engine/Source/Developer/TraceInsights/Private/Insights/ViewModels/BaseTimingTrack.cpp
--- a/Engine/Source/Developer/TraceInsights/Private/Insights/ViewModels/BaseTimingTrack.cpp
+++ b/Engine/Source/Developer/TraceInsights/Private/Insights/ViewModels/BaseTimingTrack.cpp
@@ -1,12 +1,16 @@
 // Copyright Epic Games, Inc. All Rights Reserved.
 
 #include "Insights/ViewModels/BaseTimingTrack.h"
+
 #include "Framework/MultiBox/MultiBoxBuilder.h"
 
 #define LOCTEXT_NAMESPACE "BaseTimingTrack"
 
 INSIGHTS_IMPLEMENT_RTTI(FBaseTimingTrack)
 
+// start auto generated ids from a big number (MSB set to 1) to avoid collisions with ids for GPU/CPU tracks based on 32bit timeline index
+uint64 FBaseTimingTrack::IdGenerator = (1ULL << 63);
+
 void FBaseTimingTrack::BuildContextMenu(FMenuBuilder& MenuBuilder)
 {
 }
diff --git a/Engine/Source/Developer/TraceInsights/Private/Insights/ViewModels/CounterAggregation.cpp b/Engine/Source/Developer/TraceInsights/Private/Insights/ViewModels/CounterAggregation.cpp
--- a/Engine/Source/Developer/TraceInsights/Private/Insights/ViewModels/CounterAggregation.cpp
+++ b/Engine/Source/Developer/TraceInsights/Private/Insights/ViewModels/CounterAggregation.cpp
@@ -2,6 +2,7 @@
 
 #include "CounterAggregation.h"
 
+// TraceServices
 #include "TraceServices/Model/AnalysisSession.h"
 #include "TraceServices/Model/Counters.h"
 
diff --git a/Engine/Source/Developer/TraceInsights/Private/Insights/ViewModels/DrawHelpers.cpp b/Engine/Source/Developer/TraceInsights/Private/Insights/ViewModels/DrawHelpers.cpp
--- a/Engine/Source/Developer/TraceInsights/Private/Insights/ViewModels/DrawHelpers.cpp
+++ b/Engine/Source/Developer/TraceInsights/Private/Insights/ViewModels/DrawHelpers.cpp
@@ -5,15 +5,17 @@
 #include "Fonts/FontMeasure.h"
 #include "Fonts/SlateFontInfo.h"
 
-// Insights
-#include "Insights/Common/PaintUtils.h"
-#include "Insights/Common/TimeUtils.h"
+// TraceInsightsCore
+#include "InsightsCore/Common/PaintUtils.h"
+#include "InsightsCore/Common/TimeUtils.h"
+
+// TraceInsights
 #include "Insights/ViewModels/TimingTrackViewport.h"
 
 ////////////////////////////////////////////////////////////////////////////////////////////////////
 
 void FDrawHelpers::DrawBackground(
-	const FDrawContext& DrawContext,
+	const UE::Insights::FDrawContext& DrawContext,
 	const FSlateBrush* BackgroundAreaBrush,
 	const FLinearColor& ValidAreaColor,
 	const FLinearColor& InvalidAreaColor,
@@ -86,7 +88,7 @@ void FDrawHelpers::DrawBackground(
 
 ////////////////////////////////////////////////////////////////////////////////////////////////////
 
-void FDrawHelpers::DrawBackground(const FDrawContext& DrawContext,
+void FDrawHelpers::DrawBackground(const UE::Insights::FDrawContext& DrawContext,
 								  const FSlateBrush* BackgroundAreaBrush,
 								  const float X0,
 								  const float X1,
@@ -105,7 +107,7 @@ void FDrawHelpers::DrawBackground(const FDrawContext& DrawContext,
 
 ////////////////////////////////////////////////////////////////////////////////////////////////////
 
-void FDrawHelpers::DrawBackground(const FDrawContext& DrawContext,
+void FDrawHelpers::DrawBackground(const UE::Insights::FDrawContext& DrawContext,
 								  const FSlateBrush* BackgroundAreaBrush,
 								  const FTimingTrackViewport& Viewport,
 								  const float Y,
@@ -126,7 +128,7 @@ void FDrawHelpers::DrawBackground(const FDrawContext& DrawContext,
 
 ////////////////////////////////////////////////////////////////////////////////////////////////////
 
-void FDrawHelpers::DrawBackground(const FDrawContext& DrawContext,
+void FDrawHelpers::DrawBackground(const UE::Insights::FDrawContext& DrawContext,
 								  const FSlateBrush* BackgroundAreaBrush,
 								  const FTimingTrackViewport& Viewport,
 								  const float Y,
@@ -149,7 +151,7 @@ void FDrawHelpers::DrawBackground(const FDrawContext& DrawContext,
 ////////////////////////////////////////////////////////////////////////////////////////////////////
 
 void FDrawHelpers::DrawTimeRangeSelection(
-	const FDrawContext& DrawContext,
+	const UE::Insights::FDrawContext& DrawContext,
 	const FTimingTrackViewport& Viewport,
 	const double StartTime,
 	const double EndTime,
@@ -170,7 +172,7 @@ void FDrawHelpers::DrawTimeRangeSelection(
 			const float H = Viewport.GetHeight();
 
 			const double Duration = EndTime - StartTime;
-			const FString Text = TimeUtils::FormatTimeAuto(Duration);
+			const FString Text = UE::Insights::FormatTimeAuto(Duration);
 
 			DrawSelection(DrawContext, MinX, MaxX, SelectionX1, SelectionX2, Y, H, 6.0f, Text, Brush, Font);
 		}
@@ -180,7 +182,7 @@ void FDrawHelpers::DrawTimeRangeSelection(
 ////////////////////////////////////////////////////////////////////////////////////////////////////
 
 void FDrawHelpers::DrawSelection(
-	const FDrawContext& DrawContext,
+	const UE::Insights::FDrawContext& DrawContext,
 	const float MinX,
 	const float MaxX,
 	float SelectionX1,
diff --git a/Engine/Source/Developer/TraceInsights/Private/Insights/ViewModels/DrawHelpers.h b/Engine/Source/Developer/TraceInsights/Private/Insights/ViewModels/DrawHelpers.h
--- a/Engine/Source/Developer/TraceInsights/Private/Insights/ViewModels/DrawHelpers.h
+++ b/Engine/Source/Developer/TraceInsights/Private/Insights/ViewModels/DrawHelpers.h
@@ -3,17 +3,15 @@
 #pragma once
 
 #include "CoreMinimal.h"
+
 #include "Fonts/SlateFontInfo.h"
 #include "Math/Color.h"
 
-////////////////////////////////////////////////////////////////////////////////////////////////////
-
-struct FDrawContext;
 struct FSlateBrush;
 
-class FTimingTrackViewport;
+namespace UE::Insights { class FDrawContext; }
 
-////////////////////////////////////////////////////////////////////////////////////////////////////
+class FTimingTrackViewport;
 
 class FDrawHelpers
 {
@@ -27,7 +25,7 @@ public:
 	 *  ++++++|*****|++++++
 	 *  ++++++|*****|++++++
 	 */
-	static void DrawBackground(const FDrawContext& DrawContext,
+	static void DrawBackground(const UE::Insights::FDrawContext& DrawContext,
 							   const FSlateBrush* BackgroundAreaBrush,
 							   const FLinearColor& ValidAreaColor,
 							   const FLinearColor& InvalidAreaColor,
@@ -41,7 +39,7 @@ public:
 							   float& OutValidAreaX,
 							   float& OutValidAreaW);
 
-	static void DrawBackground(const FDrawContext& DrawContext,
+	static void DrawBackground(const UE::Insights::FDrawContext& DrawContext,
 							   const FSlateBrush* BackgroundAreaBrush,
 							   const float X0,
 							   const float X1,
@@ -50,13 +48,13 @@ public:
 							   const float Y,
 							   const float H);
 
-	static void DrawBackground(const FDrawContext& DrawContext,
+	static void DrawBackground(const UE::Insights::FDrawContext& DrawContext,
 							   const FSlateBrush* BackgroundAreaBrush,
 							   const FTimingTrackViewport& Viewport,
 							   const float Y,
 							   const float H);
 
-	static void DrawBackground(const FDrawContext& DrawContext,
+	static void DrawBackground(const UE::Insights::FDrawContext& DrawContext,
 							   const FSlateBrush* BackgroundAreaBrush,
 							   const FTimingTrackViewport& Viewport,
 							   const float Y,
@@ -67,14 +65,14 @@ public:
 	/**
 	 * Draw time range selection.
 	 */
-	static void DrawTimeRangeSelection(const FDrawContext& DrawContext,
+	static void DrawTimeRangeSelection(const UE::Insights::FDrawContext& DrawContext,
 									   const FTimingTrackViewport& Viewport,
 									   const double StartTime,
 									   const double EndTime,
 									   const FSlateBrush* Brush,
 									   const FSlateFontInfo& Font);
 
-	static void DrawSelection(const FDrawContext& DrawContext,
+	static void DrawSelection(const UE::Insights::FDrawContext& DrawContext,
 							  const float MinX,
 							  const float MaxX,
 							  float SelectionX1,
@@ -86,5 +84,3 @@ public:
 							  const FSlateBrush* Brush,
 							  const FSlateFontInfo& Font);
 };
-
-////////////////////////////////////////////////////////////////////////////////////////////////////
diff --git a/Engine/Source/Developer/TraceInsights/Private/Insights/ViewModels/EventNameFilterValueConverter.cpp b/Engine/Source/Developer/TraceInsights/Private/Insights/ViewModels/EventNameFilterValueConverter.cpp
--- a/Engine/Source/Developer/TraceInsights/Private/Insights/ViewModels/EventNameFilterValueConverter.cpp
+++ b/Engine/Source/Developer/TraceInsights/Private/Insights/ViewModels/EventNameFilterValueConverter.cpp
@@ -1,7 +1,11 @@
 // Copyright Epic Games, Inc. All Rights Reserved.
 
 #include "EventNameFilterValueConverter.h"
+
+// TraceServices
 #include "TraceServices/Model/TimingProfiler.h"
+
+// TraceInsights
 #include "Insights/InsightsManager.h"
 
 #define LOCTEXT_NAMESPACE "Insights::FEventNameFilterValueConverter"
diff --git a/Engine/Source/Developer/TraceInsights/Private/Insights/ViewModels/EventNameFilterValueConverter.h b/Engine/Source/Developer/TraceInsights/Private/Insights/ViewModels/EventNameFilterValueConverter.h
--- a/Engine/Source/Developer/TraceInsights/Private/Insights/ViewModels/EventNameFilterValueConverter.h
+++ b/Engine/Source/Developer/TraceInsights/Private/Insights/ViewModels/EventNameFilterValueConverter.h
@@ -2,12 +2,13 @@
 
 #pragma once
 
-#include "Insights/ViewModels/Filters.h"
+// TraceInsightsCore
+#include "InsightsCore/Filter/ViewModels/Filters.h"
 
 namespace Insights
 {
 
-class FEventNameFilterValueConverter : public IFilterValueConverter
+class FEventNameFilterValueConverter : public UE::Insights::IFilterValueConverter
 {
 public:
 	virtual bool Convert(const FString& Input, int64& Output, FText& OutError) const override;
diff --git a/Engine/Source/Developer/TraceInsights/Private/Insights/ViewModels/FileActivityTimingTrack.cpp b/Engine/Source/Developer/TraceInsights/Private/Insights/ViewModels/FileActivityTimingTrack.cpp
--- a/Engine/Source/Developer/TraceInsights/Private/Insights/ViewModels/FileActivityTimingTrack.cpp
+++ b/Engine/Source/Developer/TraceInsights/Private/Insights/ViewModels/FileActivityTimingTrack.cpp
@@ -8,21 +8,25 @@
 #include "Framework/Commands/UICommandList.h"
 #include "Framework/MultiBox/MultiBoxBuilder.h"
 #include "Styling/SlateBrush.h"
+
+// TraceServices
 #include "TraceServices/Model/LoadTimeProfiler.h"
 
-// Insights
+// TraceInsightsCore
+#include "InsightsCore/Common/PaintUtils.h"
+#include "InsightsCore/Common/Stopwatch.h"
+#include "InsightsCore/Common/TimeUtils.h"
+
+// TraceInsights
 #include "Insights/Common/InsightsMenuBuilder.h"
-#include "Insights/Common/PaintUtils.h"
-#include "Insights/Common/Stopwatch.h"
-#include "Insights/Common/TimeUtils.h"
 #include "Insights/InsightsManager.h"
 #include "Insights/InsightsStyle.h"
 #include "Insights/ITimingViewSession.h"
 #include "Insights/TimingProfilerCommon.h"
 #include "Insights/ViewModels/TimingEvent.h"
+#include "Insights/ViewModels/TimingEventSearch.h"
 #include "Insights/ViewModels/TimingTrackViewport.h"
 #include "Insights/ViewModels/TooltipDrawState.h"
-#include "Insights/ViewModels/TimingEventSearch.h"
 #include "Insights/Widgets/STimingView.h"
 
 #include <limits>
@@ -134,7 +138,7 @@ uint32 GetFileActivityTypeColor(TraceServices::EFileActivityType Type)
 
 const uint32 FFileActivitySharedState::MaxLanes = 10000;
 
-void FFileActivitySharedState::OnBeginSession(Insights::ITimingViewSession& InSession)
+void FFileActivitySharedState::OnBeginSession(UE::Insights::Timing::ITimingViewSession& InSession)
 {
 	if (&InSession != TimingView)
 	{
@@ -154,7 +158,7 @@ void FFileActivitySharedState::OnBeginSession(Insights::ITimingViewSession& InSe
 
 ////////////////////////////////////////////////////////////////////////////////////////////////////
 
-void FFileActivitySharedState::OnEndSession(Insights::ITimingViewSession& InSession)
+void FFileActivitySharedState::OnEndSession(UE::Insights::Timing::ITimingViewSession& InSession)
 {
 	if (&InSession != TimingView)
 	{
@@ -174,7 +178,7 @@ void FFileActivitySharedState::OnEndSession(Insights::ITimingViewSession& InSess
 
 ////////////////////////////////////////////////////////////////////////////////////////////////////
 
-void FFileActivitySharedState::Tick(Insights::ITimingViewSession& InSession, const TraceServices::IAnalysisSession& InAnalysisSession)
+void FFileActivitySharedState::Tick(UE::Insights::Timing::ITimingViewSession& InSession, const TraceServices::IAnalysisSession& InAnalysisSession)
 {
 	if (&InSession != TimingView)
 	{
@@ -210,7 +214,7 @@ void FFileActivitySharedState::Tick(Insights::ITimingViewSession& InSession, con
 		FileActivityMap.Reset();
 		AllIoEvents.Reset();
 
-		FStopwatch Stopwatch;
+		UE::Insights::FStopwatch Stopwatch;
 		Stopwatch.Start();
 
 		// Enumerate all IO events and cache them.
@@ -292,10 +296,10 @@ void FFileActivitySharedState::Tick(Insights::ITimingViewSession& InSession, con
 		}
 
 		Stopwatch.Stop();
-		UE_LOG(TimingProfiler, Log, TEXT("[IO] Enumerated %s events (%s file activities) in %s."),
+		UE_LOG(LogTimingProfiler, Log, TEXT("[IO] Enumerated %s events (%s file activities) in %s."),
 			*FText::AsNumber(AllIoEvents.Num()).ToString(),
 			*FText::AsNumber(FileActivities.Num()).ToString(),
-			*TimeUtils::FormatTimeAuto(Stopwatch.GetAccumulatedTime()));
+			*UE::Insights::FormatTimeAuto(Stopwatch.GetAccumulatedTime()));
 		Stopwatch.Restart();
 
 		// Sort cached IO file activities by Start Time.
@@ -305,7 +309,7 @@ void FFileActivitySharedState::Tick(Insights::ITimingViewSession& InSession, con
 		AllIoEvents.Sort([](const FIoTimingEvent& A, const FIoTimingEvent& B) { return A.StartTime < B.StartTime; });
 
 		Stopwatch.Stop();
-		UE_LOG(TimingProfiler, Log, TEXT("[IO] Sorted file activities and events in %s."), *TimeUtils::FormatTimeAuto(Stopwatch.GetAccumulatedTime()));
+		UE_LOG(LogTimingProfiler, Log, TEXT("[IO] Sorted file activities and events in %s."), *UE::Insights::FormatTimeAuto(Stopwatch.GetAccumulatedTime()));
 
 		if (FileActivities.Num() > 0)
 		{
@@ -382,7 +386,7 @@ void FFileActivitySharedState::Tick(Insights::ITimingViewSession& InSession, con
 			}
 
 			Stopwatch.Stop();
-			UE_LOG(TimingProfiler, Log, TEXT("[IO] Computed layout for file activities in %s."), *TimeUtils::FormatTimeAuto(Stopwatch.GetAccumulatedTime()));
+			UE_LOG(LogTimingProfiler, Log, TEXT("[IO] Computed layout for file activities in %s."), *UE::Insights::FormatTimeAuto(Stopwatch.GetAccumulatedTime()));
 
 			//////////////////////////////////////////////////
 
@@ -395,14 +399,14 @@ void FFileActivitySharedState::Tick(Insights::ITimingViewSession& InSession, con
 			}
 
 			Stopwatch.Stop();
-			UE_LOG(TimingProfiler, Log, TEXT("[IO] Updated depth for events in %s."), *TimeUtils::FormatTimeAuto(Stopwatch.GetAccumulatedTime()));
+			UE_LOG(LogTimingProfiler, Log, TEXT("[IO] Updated depth for events in %s."), *UE::Insights::FormatTimeAuto(Stopwatch.GetAccumulatedTime()));
 		}
 	}
 }
 
 ////////////////////////////////////////////////////////////////////////////////////////////////////
 
-void FFileActivitySharedState::ExtendOtherTracksFilterMenu(Insights::ITimingViewSession& InSession, FMenuBuilder& InOutMenuBuilder)
+void FFileActivitySharedState::ExtendOtherTracksFilterMenu(UE::Insights::Timing::ITimingViewSession& InSession, FMenuBuilder& InOutMenuBuilder)
 {
 	if (&InSession != TimingView)
 	{
@@ -706,7 +710,7 @@ void FFileActivityTimingTrack::InitTooltip(FTooltipDrawState& InOutTooltip, cons
 			}
 
 			const double Duration = InEvent.EndTime - InEvent.StartTime;
-			InOutTooltip.AddNameValueTextLine(TEXT("Duration:"), TimeUtils::FormatTimeAuto(Duration));
+			InOutTooltip.AddNameValueTextLine(TEXT("Duration:"), UE::Insights::FormatTimeAuto(Duration));
 
 			if (ActivityType == TraceServices::FileActivityType_Read || ActivityType == TraceServices::FileActivityType_Write)
 			{
diff --git a/Engine/Source/Developer/TraceInsights/Private/Insights/ViewModels/FileActivityTimingTrack.h b/Engine/Source/Developer/TraceInsights/Private/Insights/ViewModels/FileActivityTimingTrack.h
--- a/Engine/Source/Developer/TraceInsights/Private/Insights/ViewModels/FileActivityTimingTrack.h
+++ b/Engine/Source/Developer/TraceInsights/Private/Insights/ViewModels/FileActivityTimingTrack.h
@@ -3,14 +3,15 @@
 #pragma once
 
 #include "CoreMinimal.h"
+
 #include "Framework/Commands/Commands.h"
 
-// Insights
+// TraceInsights
 #include "Insights/ITimingViewExtender.h"
 #include "Insights/ViewModels/TimingEventsTrack.h"
 
 class FTimingEventSearchParameters;
-class STimingView;
+namespace UE::Insights::TimingProfiler { class STimingView; }
 
 class FOverviewFileActivityTimingTrack;
 class FDetailedFileActivityTimingTrack;
@@ -34,7 +35,7 @@ public:
 
 ////////////////////////////////////////////////////////////////////////////////////////////////////
 
-class FFileActivitySharedState : public Insights::ITimingViewExtender, public TSharedFromThis<FFileActivitySharedState>
+class FFileActivitySharedState : public UE::Insights::Timing::ITimingViewExtender, public TSharedFromThis<FFileActivitySharedState>
 {
 	friend class FOverviewFileActivityTimingTrack;
 	friend class FDetailedFileActivityTimingTrack;
@@ -70,16 +71,16 @@ public:
 	};
 
 public:
-	explicit FFileActivitySharedState(STimingView* InTimingView) : TimingView(InTimingView) {}
+	explicit FFileActivitySharedState(UE::Insights::TimingProfiler::STimingView* InTimingView) : TimingView(InTimingView) {}
 	virtual ~FFileActivitySharedState() = default;
 
 	//////////////////////////////////////////////////
 	// ITimingViewExtender interface
 
-	virtual void OnBeginSession(Insights::ITimingViewSession& InSession) override;
-	virtual void OnEndSession(Insights::ITimingViewSession& InSession) override;
-	virtual void Tick(Insights::ITimingViewSession& InSession, const TraceServices::IAnalysisSession& InAnalysisSession) override;
-	virtual void ExtendOtherTracksFilterMenu(Insights::ITimingViewSession& InSession, FMenuBuilder& InOutMenuBuilder) override;
+	virtual void OnBeginSession(UE::Insights::Timing::ITimingViewSession& InSession) override;
+	virtual void OnEndSession(UE::Insights::Timing::ITimingViewSession& InSession) override;
+	virtual void Tick(UE::Insights::Timing::ITimingViewSession& InSession, const TraceServices::IAnalysisSession& InAnalysisSession) override;
+	virtual void ExtendOtherTracksFilterMenu(UE::Insights::Timing::ITimingViewSession& InSession, FMenuBuilder& InOutMenuBuilder) override;
 
 	//////////////////////////////////////////////////
 
@@ -113,7 +114,7 @@ private:
 	void BuildSubMenu(FMenuBuilder& InOutMenuBuilder);
 
 private:
-	STimingView* TimingView;
+	UE::Insights::TimingProfiler::STimingView* TimingView;
 
 	TSharedPtr<FOverviewFileActivityTimingTrack> IoOverviewTrack;
 	TSharedPtr<FDetailedFileActivityTimingTrack> IoActivityTrack;
diff --git a/Engine/Source/Developer/TraceInsights/Private/Insights/ViewModels/FrameTimingTrack.cpp b/Engine/Source/Developer/TraceInsights/Private/Insights/ViewModels/FrameTimingTrack.cpp
--- a/Engine/Source/Developer/TraceInsights/Private/Insights/ViewModels/FrameTimingTrack.cpp
+++ b/Engine/Source/Developer/TraceInsights/Private/Insights/ViewModels/FrameTimingTrack.cpp
@@ -9,9 +9,11 @@
 #include "HAL/PlatformApplicationMisc.h"
 #include "Styling/SlateBrush.h"
 
-// Insights
-#include "Insights/Common/PaintUtils.h"
-#include "Insights/Common/TimeUtils.h"
+// TraceInsightsCore
+#include "InsightsCore/Common/PaintUtils.h"
+#include "InsightsCore/Common/TimeUtils.h"
+
+// TraceInsights
 #include "Insights/InsightsManager.h"
 #include "Insights/InsightsStyle.h"
 #include "Insights/ITimingViewSession.h"
@@ -78,7 +80,7 @@ bool FFrameSharedState::IsFrameTrackVisible(uint32 InFrameType) const
 
 ////////////////////////////////////////////////////////////////////////////////////////////////////
 
-void FFrameSharedState::OnBeginSession(Insights::ITimingViewSession& InSession)
+void FFrameSharedState::OnBeginSession(UE::Insights::Timing::ITimingViewSession& InSession)
 {
 	if (&InSession != TimingView)
 	{
@@ -91,7 +93,7 @@ void FFrameSharedState::OnBeginSession(Insights::ITimingViewSession& InSession)
 
 ////////////////////////////////////////////////////////////////////////////////////////////////////
 
-void FFrameSharedState::OnEndSession(Insights::ITimingViewSession& InSession)
+void FFrameSharedState::OnEndSession(UE::Insights::Timing::ITimingViewSession& InSession)
 {
 	if (&InSession != TimingView)
 	{
@@ -104,7 +106,7 @@ void FFrameSharedState::OnEndSession(Insights::ITimingViewSession& InSession)
 
 ////////////////////////////////////////////////////////////////////////////////////////////////////
 
-void FFrameSharedState::Tick(Insights::ITimingViewSession& InSession, const TraceServices::IAnalysisSession& InAnalysisSession)
+void FFrameSharedState::Tick(UE::Insights::Timing::ITimingViewSession& InSession, const TraceServices::IAnalysisSession& InAnalysisSession)
 {
 	if (&InSession != TimingView)
 	{
@@ -152,7 +154,7 @@ void FFrameSharedState::Tick(Insights::ITimingViewSession& InSession, const Trac
 
 ////////////////////////////////////////////////////////////////////////////////////////////////////
 
-void FFrameSharedState::ExtendOtherTracksFilterMenu(Insights::ITimingViewSession& InSession, FMenuBuilder& InOutMenuBuilder)
+void FFrameSharedState::ExtendOtherTracksFilterMenu(UE::Insights::Timing::ITimingViewSession& InSession, FMenuBuilder& InOutMenuBuilder)
 {
 	if (&InSession != TimingView)
 	{
@@ -344,7 +346,7 @@ void FFrameTimingTrack::PostDraw(const ITimingTrackDrawContext& Context) const
 
 ////////////////////////////////////////////////////////////////////////////////////////////////////
 
-void FFrameTimingTrack::DrawSelectedEventInfo(const FTimingEvent& SelectedEvent, const FTimingTrackViewport& Viewport, const FDrawContext& DrawContext, const FSlateBrush* WhiteBrush, const FSlateFontInfo& Font) const
+void FFrameTimingTrack::DrawSelectedEventInfo(const FTimingEvent& SelectedEvent, const FTimingTrackViewport& Viewport, const UE::Insights::FDrawContext& DrawContext, const FSlateBrush* WhiteBrush, const FSlateFontInfo& Font) const
 {
 	FindFrame(SelectedEvent, [this, &SelectedEvent, &Font, &Viewport, &DrawContext, &WhiteBrush](double InFoundStartTime, double InFoundEndTime, uint32 InFoundDepth, const TraceServices::FFrame& InFoundFrame)
 	{
@@ -413,7 +415,7 @@ void FFrameTimingTrack::InitTooltip(FTooltipDrawState& InOutTooltip, const ITimi
 
 			//const double Duration = TooltipEvent.GetDuration();
 			const double Duration = InFoundFrame.EndTime - InFoundFrame.StartTime;
-			InOutTooltip.AddNameValueTextLine(TEXT("Duration:"), TimeUtils::FormatTimeAuto(Duration));
+			InOutTooltip.AddNameValueTextLine(TEXT("Duration:"), UE::Insights::FormatTimeAuto(Duration));
 
 			InOutTooltip.UpdateLayout();
 		});
@@ -559,7 +561,7 @@ const FString FFrameTimingTrack::GetFrameName(const uint64 InFrameIndex) const
 
 const FString FFrameTimingTrack::GetCompleteFrameName(const uint64 InFrameIndex, const double InFrameDuration) const
 {
-	return FString::Printf(TEXT("%s Frame %" UINT64_FMT " (%s)"), FFrameTrackDrawHelper::FrameTypeToString(FrameType), InFrameIndex, *TimeUtils::FormatTimeAuto(InFrameDuration));
+	return FString::Printf(TEXT("%s Frame %" UINT64_FMT " (%s)"), FFrameTrackDrawHelper::FrameTypeToString(FrameType), InFrameIndex, *UE::Insights::FormatTimeAuto(InFrameDuration));
 }
 
 ////////////////////////////////////////////////////////////////////////////////////////////////////
