[Insights] Massive refactoring of TraceInsights module:

diff --git a/Engine/Source/Developer/TraceInsights/Private/Insights/ViewModels/FrameTimingTrack.h b/Engine/Source/Developer/TraceInsights/Private/Insights/ViewModels/FrameTimingTrack.h
--- a/Engine/Source/Developer/TraceInsights/Private/Insights/ViewModels/FrameTimingTrack.h
+++ b/Engine/Source/Developer/TraceInsights/Private/Insights/ViewModels/FrameTimingTrack.h
@@ -5,19 +5,22 @@
 #include "CoreMinimal.h"
 #include "Fonts/SlateFontInfo.h"
 #include "Framework/Commands/Commands.h"
+
+// TraceServices
 #include "TraceServices/Model/Frames.h"
 
-// Insights
+// TraceInsights
 #include "Insights/ITimingViewExtender.h"
 #include "Insights/ViewModels/TimingEventSearch.h" // for TTimingEventSearchCache
 #include "Insights/ViewModels/TimingEventsTrack.h"
 #include "Insights/ViewModels/TrackHeader.h"
 
+struct FSlateBrush;
+
 class FTimingEvent;
 class FTimingEventSearchParameters;
 class FFrameTimingTrack;
-class STimingView;
-struct FSlateBrush;
+namespace UE::Insights::TimingProfiler { class STimingView; }
 
 ////////////////////////////////////////////////////////////////////////////////////////////////////
 
@@ -34,10 +37,10 @@ public:
 
 ////////////////////////////////////////////////////////////////////////////////////////////////////
 
-class FFrameSharedState : public Insights::ITimingViewExtender, public TSharedFromThis<FFrameSharedState>
+class FFrameSharedState : public UE::Insights::Timing::ITimingViewExtender, public TSharedFromThis<FFrameSharedState>
 {
 public:
-	explicit FFrameSharedState(STimingView* InTimingView) : TimingView(InTimingView) {}
+	explicit FFrameSharedState(UE::Insights::TimingProfiler::STimingView* InTimingView) : TimingView(InTimingView) {}
 	virtual ~FFrameSharedState() = default;
 
 	TSharedPtr<FFrameTimingTrack> GetFrameTrack(uint32 InFrameType);
@@ -47,10 +50,10 @@ public:
 	//////////////////////////////////////////////////
 	// ITimingViewExtender interface
 
-	virtual void OnBeginSession(Insights::ITimingViewSession& InSession) override;
-	virtual void OnEndSession(Insights::ITimingViewSession& InSession) override;
-	virtual void Tick(Insights::ITimingViewSession& InSession, const TraceServices::IAnalysisSession& InAnalysisSession) override;
-	virtual void ExtendOtherTracksFilterMenu(Insights::ITimingViewSession& InSession, FMenuBuilder& InMenuBuilder) override;
+	virtual void OnBeginSession(UE::Insights::Timing::ITimingViewSession& InSession) override;
+	virtual void OnEndSession(UE::Insights::Timing::ITimingViewSession& InSession) override;
+	virtual void Tick(UE::Insights::Timing::ITimingViewSession& InSession, const TraceServices::IAnalysisSession& InAnalysisSession) override;
+	virtual void ExtendOtherTracksFilterMenu(UE::Insights::Timing::ITimingViewSession& InSession, FMenuBuilder& InMenuBuilder) override;
 
 	//////////////////////////////////////////////////
 
@@ -63,7 +66,7 @@ public:
 	void ShowHideAllFrameTracks() { SetAllFrameTracksToggle(!IsAllFrameTracksToggleOn()); }
 
 private:
-	STimingView* TimingView;
+	UE::Insights::TimingProfiler::STimingView* TimingView;
 
 	bool bShowHideAllFrameTracks;
 
@@ -130,7 +133,7 @@ public:
 	//////////////////////////////////////////////////
 
 private:
-	void DrawSelectedEventInfo(const FTimingEvent& SelectedEvent, const FTimingTrackViewport& Viewport, const FDrawContext& DrawContext, const FSlateBrush* WhiteBrush, const FSlateFontInfo& Font) const;
+	void DrawSelectedEventInfo(const FTimingEvent& SelectedEvent, const FTimingTrackViewport& Viewport, const UE::Insights::FDrawContext& DrawContext, const FSlateBrush* WhiteBrush, const FSlateFontInfo& Font) const;
 
 	bool FindFrame(const FTimingEvent& InTimingEvent, TFunctionRef<void(double, double, uint32, const TraceServices::FFrame&)> InFoundPredicate) const;
 	bool FindFrame(const FTimingEventSearchParameters& InParameters, TFunctionRef<void(double, double, uint32, const TraceServices::FFrame&)> InFoundPredicate) const;
diff --git a/Engine/Source/Developer/TraceInsights/Private/Insights/ViewModels/FrameTrackHelper.cpp b/Engine/Source/Developer/TraceInsights/Private/Insights/ViewModels/FrameTrackHelper.cpp
--- a/Engine/Source/Developer/TraceInsights/Private/Insights/ViewModels/FrameTrackHelper.cpp
+++ b/Engine/Source/Developer/TraceInsights/Private/Insights/ViewModels/FrameTrackHelper.cpp
@@ -6,10 +6,14 @@
 #include "Fonts/SlateFontInfo.h"
 #include "Framework/Application/SlateApplication.h"
 #include "Rendering/DrawElements.h"
+
+// TraceServices
 #include "TraceServices/Model/Frames.h"
 
-// Insights
-#include "Insights/Common/PaintUtils.h"
+// TraceInsightsCore
+#include "InsightsCore/Common/PaintUtils.h"
+
+// TraceInsights
 #include "Insights/InsightsStyle.h"
 #include "Insights/ViewModels/DrawHelpers.h"
 #include "Insights/ViewModels/FrameTrackViewport.h"
@@ -76,7 +80,7 @@ void FFrameTrackSeriesBuilder::AddFrame(const TraceServices::FFrame& Frame)
 // FFrameTrackDrawHelper
 ////////////////////////////////////////////////////////////////////////////////////////////////////
 
-FFrameTrackDrawHelper::FFrameTrackDrawHelper(const FDrawContext& InDrawContext, const FFrameTrackViewport& InViewport)
+FFrameTrackDrawHelper::FFrameTrackDrawHelper(const UE::Insights::FDrawContext& InDrawContext, const FFrameTrackViewport& InViewport)
 	: DrawContext(InDrawContext)
 	, Viewport(InViewport)
 	, WhiteBrush(FInsightsStyle::Get().GetBrush("WhiteBrush"))
diff --git a/Engine/Source/Developer/TraceInsights/Private/Insights/ViewModels/FrameTrackHelper.h b/Engine/Source/Developer/TraceInsights/Private/Insights/ViewModels/FrameTrackHelper.h
--- a/Engine/Source/Developer/TraceInsights/Private/Insights/ViewModels/FrameTrackHelper.h
+++ b/Engine/Source/Developer/TraceInsights/Private/Insights/ViewModels/FrameTrackHelper.h
@@ -3,14 +3,15 @@
 #pragma once
 
 #include "CoreMinimal.h"
+
 #include "Styling/WidgetStyle.h"
 
 enum class ESlateDrawEffect : uint8;
-
-struct FDrawContext;
 struct FGeometry;
 struct FSlateBrush;
 
+namespace UE::Insights { class FDrawContext; }
+
 class FFrameTrackViewport;
 class FSlateWindowElementList;
 
@@ -148,7 +149,7 @@ private:
 class FFrameTrackDrawHelper
 {
 public:
-	explicit FFrameTrackDrawHelper(const FDrawContext& InDrawContext, const FFrameTrackViewport& InViewport);
+	explicit FFrameTrackDrawHelper(const UE::Insights::FDrawContext& InDrawContext, const FFrameTrackViewport& InViewport);
 
 	/**
 	 * Non-copyable
@@ -176,7 +177,7 @@ public:
 	int32 GetNumDrawSamples() const { return NumDrawSamples; }
 
 private:
-	const FDrawContext& DrawContext;
+	const UE::Insights::FDrawContext& DrawContext;
 	const FFrameTrackViewport& Viewport;
 
 	const FSlateBrush* WhiteBrush;
diff --git a/Engine/Source/Developer/TraceInsights/Private/Insights/ViewModels/GraphTrack.cpp b/Engine/Source/Developer/TraceInsights/Private/Insights/ViewModels/GraphTrack.cpp
--- a/Engine/Source/Developer/TraceInsights/Private/Insights/ViewModels/GraphTrack.cpp
+++ b/Engine/Source/Developer/TraceInsights/Private/Insights/ViewModels/GraphTrack.cpp
@@ -7,9 +7,11 @@
 #include "Styling/AppStyle.h"
 #include "Widgets/Layout/SBox.h"
 
-// Insights
-#include "Insights/Common/PaintUtils.h"
-#include "Insights/Common/TimeUtils.h"
+// TraceInsightsCore
+#include "InsightsCore/Common/PaintUtils.h"
+#include "InsightsCore/Common/TimeUtils.h"
+
+// TraceInsights
 #include "Insights/InsightsStyle.h"
 #include "Insights/ViewModels/DrawHelpers.h"
 #include "Insights/ViewModels/GraphSeries.h"
@@ -134,7 +136,7 @@ void FGraphTrack::UpdateStats()
 
 void FGraphTrack::PreDraw(const ITimingTrackDrawContext& Context) const
 {
-	FDrawContext& DrawContext = Context.GetDrawContext();
+	UE::Insights::FDrawContext& DrawContext = Context.GetDrawContext();
 	const FTimingTrackViewport& Viewport = Context.GetViewport();
 
 	FDrawHelpers::DrawBackground(DrawContext, WhiteBrush, Viewport, GetPosY(), GetHeight());
@@ -144,7 +146,7 @@ void FGraphTrack::PreDraw(const ITimingTrackDrawContext& Context) const
 
 void FGraphTrack::Draw(const ITimingTrackDrawContext& Context) const
 {
-	FDrawContext& DrawContext = Context.GetDrawContext();
+	UE::Insights::FDrawContext& DrawContext = Context.GetDrawContext();
 	const FTimingTrackViewport& Viewport = Context.GetViewport();
 
 	// Set clipping area.
@@ -201,7 +203,7 @@ void FGraphTrack::Draw(const ITimingTrackDrawContext& Context) const
 
 ////////////////////////////////////////////////////////////////////////////////////////////////////
 
-void FGraphTrack::DrawSeries(const FGraphSeries& Series, FDrawContext& DrawContext, const FTimingTrackViewport& Viewport) const
+void FGraphTrack::DrawSeries(const FGraphSeries& Series, UE::Insights::FDrawContext& DrawContext, const FTimingTrackViewport& Viewport) const
 {
 	const float Scale = DrawContext.Geometry.GetAccumulatedLayoutTransform().GetScale();
 	const float PixelUnit = 1.0f / Scale;
@@ -539,7 +541,7 @@ void FGraphTrack::DrawEvent(const ITimingTrackDrawContext& Context, const ITimin
 	const TSharedPtr<const FGraphSeries> Series = GraphEvent.GetSeries();
 
 	const FTimingTrackViewport& Viewport = Context.GetViewport();
-	FDrawContext& DrawContext = Context.GetDrawContext();
+	UE::Insights::FDrawContext& DrawContext = Context.GetDrawContext();
 
 	const float EventX1 = Viewport.TimeToSlateUnitsRounded(GraphEvent.GetStartTime());
 	const float EventX2 = Viewport.TimeToSlateUnitsRounded(Viewport.RestrictEndTime(GraphEvent.GetEndTime()));
@@ -613,7 +615,7 @@ void FGraphTrack::DrawVerticalAxisGrid(const ITimingTrackDrawContext& Context) c
 void FGraphTrack::DrawHeader(const ITimingTrackDrawContext& Context) const
 {
 	const FTimingViewDrawHelper& Helper = *static_cast<const FTimingViewDrawHelper*>(&Context.GetHelper());
-	FDrawContext& DrawContext = Context.GetDrawContext();
+	UE::Insights::FDrawContext& DrawContext = Context.GetDrawContext();
 	Helper.DrawTrackHeader(*this, DrawContext.LayerId, DrawContext.LayerId + 1);
 	DrawContext.LayerId += 2;
 }
@@ -629,11 +631,13 @@ void FGraphTrack::InitTooltip(FTooltipDrawState& InOutTooltip, const ITimingEven
 
 		InOutTooltip.ResetContent();
 		InOutTooltip.AddTitle(Series->GetName().ToString(), Series->GetColor());
-		const double Precision = FMath::Max(1.0 / TimeScaleX, TimeUtils::Nanosecond);
-		InOutTooltip.AddNameValueTextLine(TEXT("Time:"), TimeUtils::FormatTime(TooltipEvent.GetStartTime(), Precision));
+
+		using namespace UE::Insights;
+		const double Precision = FMath::Max(1.0 / TimeScaleX, FTimeValue::Nanosecond);
+		InOutTooltip.AddNameValueTextLine(TEXT("Time:"), FormatTime(TooltipEvent.GetStartTime(), Precision));
 		if (Series->HasEventDuration())
 		{
-			InOutTooltip.AddNameValueTextLine(TEXT("Duration:"), TimeUtils::FormatTimeAuto(TooltipEvent.GetDuration()));
+			InOutTooltip.AddNameValueTextLine(TEXT("Duration:"), FormatTimeAuto(TooltipEvent.GetDuration()));
 		}
 		InOutTooltip.AddNameValueTextLine(TEXT("Value:"), Series->FormatValue(TooltipEvent.GetValue()));
 		InOutTooltip.UpdateLayout();
@@ -892,21 +896,21 @@ void FRandomGraphTrack::AddDefaultSeries()
 {
 	TSharedRef<FGraphSeries> Series0 = MakeShared<FGraphSeries>();
 	Series0->SetName(TEXT("Random Blue"));
-	Series0->SetDescription(TEXT("Random series; for debuging purposes"));
+	Series0->SetDescription(TEXT("Random series; for debugging purposes"));
 	Series0->SetColor(FLinearColor(0.1f, 0.5f, 1.0f, 1.0f), FLinearColor(0.4f, 0.8f, 1.0f, 1.0f));
 	Series0->SetVisibility(true);
 	AllSeries.Add(Series0);
 
 	TSharedRef<FGraphSeries> Series1 = MakeShared<FGraphSeries>();
 	Series1->SetName(TEXT("Random Yellow"));
-	Series1->SetDescription(TEXT("Random series; for debuging purposes"));
+	Series1->SetDescription(TEXT("Random series; for debugging purposes"));
 	Series1->SetColor(FLinearColor(0.9f, 0.9f, 0.1f, 1.0f), FLinearColor(1.0f, 1.0f, 0.4f, 1.0f));
 	Series1->SetVisibility(false);
 	AllSeries.Add(Series1);
 
 	TSharedRef<FGraphSeries> Series2 = MakeShared<FGraphSeries>();
 	Series2->SetName(TEXT("Random Red"));
-	Series2->SetDescription(TEXT("Random series; for debuging purposes"));
+	Series2->SetDescription(TEXT("Random series; for debugging purposes"));
 	Series2->SetColor(FLinearColor(1.0f, 0.1f, 0.2f, 1.0f), FLinearColor(1.0f, 0.4f, 0.5f, 1.0f));
 	Series2->SetVisibility(true);
 	AllSeries.Add(Series2);
diff --git a/Engine/Source/Developer/TraceInsights/Private/Insights/ViewModels/LoadingTimingTrack.cpp b/Engine/Source/Developer/TraceInsights/Private/Insights/ViewModels/LoadingTimingTrack.cpp
--- a/Engine/Source/Developer/TraceInsights/Private/Insights/ViewModels/LoadingTimingTrack.cpp
+++ b/Engine/Source/Developer/TraceInsights/Private/Insights/ViewModels/LoadingTimingTrack.cpp
@@ -7,16 +7,20 @@
 #include "Framework/Commands/UICommandList.h"
 #include "Framework/MultiBox/MultiBoxBuilder.h"
 #include "Styling/SlateBrush.h"
+
+// TraceServices
 #include "TraceServices/Model/Threads.h"
 
-// Insights
-#include "Insights/Common/PaintUtils.h"
-#include "Insights/Common/TimeUtils.h"
+// TraceInsightsCore
+#include "InsightsCore/Common/PaintUtils.h"
+#include "InsightsCore/Common/TimeUtils.h"
+#include "InsightsCore/Filter/ViewModels/FilterConfigurator.h"
+#include "InsightsCore/Filter/ViewModels/Filters.h"
+
+// TraceInsights
 #include "Insights/InsightsManager.h"
 #include "Insights/InsightsStyle.h"
 #include "Insights/ITimingViewSession.h"
-#include "Insights/ViewModels/Filters.h"
-#include "Insights/ViewModels/FilterConfigurator.h"
 #include "Insights/ViewModels/TimingEvent.h"
 #include "Insights/ViewModels/TimingTrackViewport.h"
 #include "Insights/ViewModels/TooltipDrawState.h"
@@ -62,7 +66,7 @@ UE_ENABLE_OPTIMIZATION_SHIP
 // FLoadingSharedState
 ////////////////////////////////////////////////////////////////////////////////////////////////////
 
-FLoadingSharedState::FLoadingSharedState(STimingView* InTimingView)
+FLoadingSharedState::FLoadingSharedState(UE::Insights::TimingProfiler::STimingView* InTimingView)
 	: TimingView(InTimingView)
 	, bShowHideAllLoadingTracks(false)
 	//, LoadingTracks
@@ -74,7 +78,7 @@ FLoadingSharedState::FLoadingSharedState(STimingView* InTimingView)
 
 ////////////////////////////////////////////////////////////////////////////////////////////////////
 
-void FLoadingSharedState::OnBeginSession(Insights::ITimingViewSession& InSession)
+void FLoadingSharedState::OnBeginSession(UE::Insights::Timing::ITimingViewSession& InSession)
 {
 	if (&InSession != TimingView)
 	{
@@ -99,7 +103,7 @@ void FLoadingSharedState::OnBeginSession(Insights::ITimingViewSession& InSession
 
 ////////////////////////////////////////////////////////////////////////////////////////////////////
 
-void FLoadingSharedState::OnEndSession(Insights::ITimingViewSession& InSession)
+void FLoadingSharedState::OnEndSession(UE::Insights::Timing::ITimingViewSession& InSession)
 {
 	if (&InSession != TimingView)
 	{
@@ -117,7 +121,7 @@ void FLoadingSharedState::OnEndSession(Insights::ITimingViewSession& InSession)
 
 ////////////////////////////////////////////////////////////////////////////////////////////////////
 
-void FLoadingSharedState::Tick(Insights::ITimingViewSession& InSession, const TraceServices::IAnalysisSession& InAnalysisSession)
+void FLoadingSharedState::Tick(UE::Insights::Timing::ITimingViewSession& InSession, const TraceServices::IAnalysisSession& InAnalysisSession)
 {
 	if (&InSession != TimingView)
 	{
@@ -161,7 +165,7 @@ void FLoadingSharedState::Tick(Insights::ITimingViewSession& InSession, const Tr
 
 ////////////////////////////////////////////////////////////////////////////////////////////////////
 
-void FLoadingSharedState::ExtendOtherTracksFilterMenu(Insights::ITimingViewSession& InSession, FMenuBuilder& InOutMenuBuilder)
+void FLoadingSharedState::ExtendOtherTracksFilterMenu(UE::Insights::Timing::ITimingViewSession& InSession, FMenuBuilder& InOutMenuBuilder)
 {
 	if (&InSession != TimingView)
 	{
@@ -340,9 +344,8 @@ void FLoadingTimingTrack::BuildFilteredDrawState(ITimingEventsTrackDrawStateBuil
 {
 	if (HasCustomFilter())
 	{
-		using namespace Insights;
-
-		FFilterContext FilterConfiguratorContext;
+		using EFilterField = UE::Insights::EFilterField;
+		UE::Insights::FFilterContext FilterConfiguratorContext;
 		FilterConfiguratorContext.SetReturnValueForUnsetFilters(false);
 		FilterConfiguratorContext.AddFilterData<double>(static_cast<int32>(EFilterField::StartTime), 0.0f);
 		FilterConfiguratorContext.AddFilterData<double>(static_cast<int32>(EFilterField::EndTime), 0.0f);
@@ -409,7 +412,7 @@ void FLoadingTimingTrack::InitTooltip(FTooltipDrawState& InOutTooltip, const ITi
 			const TraceServices::FPackageExportInfo* Export = InFoundEvent.Export;
 			const TraceServices::FPackageInfo* Package = InFoundEvent.Export ? InFoundEvent.Export->Package : InFoundEvent.Package;
 
-			InOutTooltip.AddNameValueTextLine(TEXT("Duration:"), TimeUtils::FormatTimeAuto(TooltipEvent.GetDuration()));
+			InOutTooltip.AddNameValueTextLine(TEXT("Duration:"), UE::Insights::FormatTimeAuto(TooltipEvent.GetDuration()));
 			InOutTooltip.AddNameValueTextLine(TEXT("Depth:"), FString::Printf(TEXT("%d"), TooltipEvent.GetDepth()));
 
 			if (Package)
@@ -454,9 +457,8 @@ const TSharedPtr<const ITimingEvent> FLoadingTimingTrack::SearchEvent(const FTim
 
 bool FLoadingTimingTrack::FindLoadTimeProfilerCpuEvent(const FTimingEventSearchParameters& InParameters, TFunctionRef<void(double, double, uint32, const TraceServices::FLoadTimeProfilerCpuEvent&)> InFoundPredicate) const
 {
-	using namespace Insights;
-
-	FFilterContext FilterConfiguratorContext;
+	using EFilterField = UE::Insights::EFilterField;
+	UE::Insights::FFilterContext FilterConfiguratorContext;
 	FilterConfiguratorContext.SetReturnValueForUnsetFilters(false);
 	FilterConfiguratorContext.AddFilterData<double>(static_cast<int32>(EFilterField::StartTime), 0.0f);
 	FilterConfiguratorContext.AddFilterData<double>(static_cast<int32>(EFilterField::EndTime), 0.0f);
@@ -547,7 +549,7 @@ bool FLoadingTimingTrack::HasCustomFilter() const
 
 ////////////////////////////////////////////////////////////////////////////////////////////////////
 
-void FLoadingTimingTrack::SetFilterConfigurator(TSharedPtr<Insights::FFilterConfigurator> InFilterConfigurator)
+void FLoadingTimingTrack::SetFilterConfigurator(TSharedPtr<UE::Insights::FFilterConfigurator> InFilterConfigurator)
 {
 	if (FilterConfigurator != InFilterConfigurator)
 	{
diff --git a/Engine/Source/Developer/TraceInsights/Private/Insights/ViewModels/LoadingTimingTrack.h b/Engine/Source/Developer/TraceInsights/Private/Insights/ViewModels/LoadingTimingTrack.h
--- a/Engine/Source/Developer/TraceInsights/Private/Insights/ViewModels/LoadingTimingTrack.h
+++ b/Engine/Source/Developer/TraceInsights/Private/Insights/ViewModels/LoadingTimingTrack.h
@@ -4,16 +4,19 @@
 
 #include "CoreMinimal.h"
 #include "Framework/Commands/Commands.h"
+
+// TraceServices
 #include "TraceServices/Model/LoadTimeProfiler.h"
 
-// Insights
+// TraceInsights
 #include "Insights/ITimingViewExtender.h"
 #include "Insights/ViewModels/TimingEventsTrack.h"
 
-class FTimingEventSearchParameters;
-class STimingView;
+namespace UE::Insights { class FFilterConfigurator; }
 
 class FLoadingTimingTrack;
+class FTimingEventSearchParameters;
+namespace UE::Insights::TimingProfiler { class STimingView; }
 
 ////////////////////////////////////////////////////////////////////////////////////////////////////
 
@@ -35,19 +38,19 @@ DECLARE_DELEGATE_RetVal_TwoParams(const TCHAR*, FLoadingTrackGetEventNameDelegat
 
 ////////////////////////////////////////////////////////////////////////////////////////////////////
 
-class FLoadingSharedState : public Insights::ITimingViewExtender, public TSharedFromThis<FLoadingSharedState>
+class FLoadingSharedState : public UE::Insights::Timing::ITimingViewExtender, public TSharedFromThis<FLoadingSharedState>
 {
 public:
-	explicit FLoadingSharedState(STimingView* InTimingView);
+	explicit FLoadingSharedState(UE::Insights::TimingProfiler::STimingView* InTimingView);
 	virtual ~FLoadingSharedState() = default;
 
 	//////////////////////////////////////////////////
 	// ITimingViewExtender interface
 
-	virtual void OnBeginSession(Insights::ITimingViewSession& InSession) override;
-	virtual void OnEndSession(Insights::ITimingViewSession& InSession) override;
-	virtual void Tick(Insights::ITimingViewSession& InSession, const TraceServices::IAnalysisSession& InAnalysisSession) override;
-	virtual void ExtendOtherTracksFilterMenu(Insights::ITimingViewSession& InSession, FMenuBuilder& InMenuBuilder) override;
+	virtual void OnBeginSession(UE::Insights::Timing::ITimingViewSession& InSession) override;
+	virtual void OnEndSession(UE::Insights::Timing::ITimingViewSession& InSession) override;
+	virtual void Tick(UE::Insights::Timing::ITimingViewSession& InSession, const TraceServices::IAnalysisSession& InAnalysisSession) override;
+	virtual void ExtendOtherTracksFilterMenu(UE::Insights::Timing::ITimingViewSession& InSession, FMenuBuilder& InMenuBuilder) override;
 
 	//////////////////////////////////////////////////
 
@@ -75,7 +78,7 @@ private:
 	const TCHAR* GetEventNameByPackageAndExportClassName(uint32 Depth, const TraceServices::FLoadTimeProfilerCpuEvent& Event) const;
 
 private:
-	STimingView* TimingView;
+	UE::Insights::TimingProfiler::STimingView* TimingView;
 
 	bool bShowHideAllLoadingTracks;
 
@@ -110,7 +113,7 @@ public:
 
 	virtual const TSharedPtr<const ITimingEvent> SearchEvent(const FTimingEventSearchParameters& InSearchParameters) const override;
 
-	virtual void SetFilterConfigurator(TSharedPtr<Insights::FFilterConfigurator> InFilterConfigurator) override;
+	virtual void SetFilterConfigurator(TSharedPtr<UE::Insights::FFilterConfigurator> InFilterConfigurator) override;
 
 protected:
 	// Helper function to find an event given search parameters
@@ -122,7 +125,7 @@ protected:
 	FLoadingSharedState& SharedState;
 	uint32 TimelineIndex;
 
-	TSharedPtr<Insights::FFilterConfigurator> FilterConfigurator;
+	TSharedPtr<UE::Insights::FFilterConfigurator> FilterConfigurator;
 };
 
 ////////////////////////////////////////////////////////////////////////////////////////////////////
diff --git a/Engine/Source/Developer/TraceInsights/Private/Insights/ViewModels/LogFilter.h b/Engine/Source/Developer/TraceInsights/Private/Insights/ViewModels/LogFilter.h
--- a/Engine/Source/Developer/TraceInsights/Private/Insights/ViewModels/LogFilter.h
+++ b/Engine/Source/Developer/TraceInsights/Private/Insights/ViewModels/LogFilter.h
@@ -2,13 +2,14 @@
 
 #pragma once
 
-#include "Async/AsyncWork.h"
 #include "CoreMinimal.h"
+
+#include "Async/AsyncWork.h"
 #include "Logging/LogVerbosity.h"
 #include "Misc/TextFilterExpressionEvaluator.h"
 #include "Stats/Stats.h"
 
-// Insights
+// TraceInsights
 #include "Insights/ViewModels/LogMessage.h"
 
 ////////////////////////////////////////////////////////////////////////////////////////////////////
diff --git a/Engine/Source/Developer/TraceInsights/Private/Insights/ViewModels/LogMessage.cpp b/Engine/Source/Developer/TraceInsights/Private/Insights/ViewModels/LogMessage.cpp
--- a/Engine/Source/Developer/TraceInsights/Private/Insights/ViewModels/LogMessage.cpp
+++ b/Engine/Source/Developer/TraceInsights/Private/Insights/ViewModels/LogMessage.cpp
@@ -4,10 +4,14 @@
 
 #include "Misc/OutputDeviceHelper.h"
 #include "Misc/ScopeLock.h"
+
+// TraceServices
 #include "TraceServices/Model/Log.h"
 
-// Insights
-#include "Insights/Common/TimeUtils.h"
+// TraceInsightsCore
+#include "InsightsCore/Common/TimeUtils.h"
+
+// TraceInsights
 #include "Insights/Log.h"
 
 ////////////////////////////////////////////////////////////////////////////////////////////////////
@@ -48,7 +52,7 @@ FText FLogMessageRecord::GetIndexAsText() const
 
 FText FLogMessageRecord::GetTimeAsText() const
 {
-	return FText::FromString(TimeUtils::FormatTimeHMS(Time, TimeUtils::Microsecond));
+	return FText::FromString(UE::Insights::FormatTimeHMS(Time, UE::Insights::FTimeValue::Microsecond));
 }
 
 ////////////////////////////////////////////////////////////////////////////////////////////////////
diff --git a/Engine/Source/Developer/TraceInsights/Private/Insights/ViewModels/MarkersTimingTrack.cpp b/Engine/Source/Developer/TraceInsights/Private/Insights/ViewModels/MarkersTimingTrack.cpp
--- a/Engine/Source/Developer/TraceInsights/Private/Insights/ViewModels/MarkersTimingTrack.cpp
+++ b/Engine/Source/Developer/TraceInsights/Private/Insights/ViewModels/MarkersTimingTrack.cpp
@@ -12,11 +12,15 @@
 #include "Misc/Paths.h"
 #include "Modules/ModuleManager.h"
 #include "Styling/AppStyle.h"
+
+// TraceServices
 #include "TraceServices/Model/Log.h"
 #include "TraceServices/Model/Screenshot.h"
 
-// Insights
-#include "Insights/Common/PaintUtils.h"
+// TraceInsightsCore
+#include "InsightsCore/Common/PaintUtils.h"
+
+// TraceInsights
 #include "Insights/InsightsManager.h"
 #include "Insights/InsightsStyle.h"
 #include "Insights/TimingProfilerManager.h"
@@ -171,7 +175,7 @@ void FMarkersTimingTrack::UpdateDrawState(const ITimingTrackUpdateContext& Conte
 
 void FMarkersTimingTrack::Draw(const ITimingTrackDrawContext& Context) const
 {
-	FDrawContext& DrawContext = Context.GetDrawContext();
+	UE::Insights::FDrawContext& DrawContext = Context.GetDrawContext();
 	const FTimingTrackViewport& Viewport = Context.GetViewport();
 
 	// Draw background.
@@ -186,7 +190,7 @@ void FMarkersTimingTrack::Draw(const ITimingTrackDrawContext& Context) const
 
 void FMarkersTimingTrack::PostDraw(const ITimingTrackDrawContext& Context) const
 {
-	FDrawContext& DrawContext = Context.GetDrawContext();
+	UE::Insights::FDrawContext& DrawContext = Context.GetDrawContext();
 	const FTimingTrackViewport& Viewport = Context.GetViewport();
 
 	//////////////////////////////////////////////////
@@ -609,6 +613,8 @@ bool FMarkersTimingTrack::TryGetHoveredEventScreenshotId(uint32& OutScreenshotId
 {
 	OutScreenshotId = TraceServices::FScreenshot::InvalidScreenshotId;
 
+	using namespace UE::Insights::TimingProfiler;
+
 	TSharedPtr<STimingProfilerWindow> Window = FTimingProfilerManager::Get()->GetProfilerWindow();
 	if (!Window.IsValid())
 	{
diff --git a/Engine/Source/Developer/TraceInsights/Private/Insights/ViewModels/MarkersTimingTrack.h b/Engine/Source/Developer/TraceInsights/Private/Insights/ViewModels/MarkersTimingTrack.h
--- a/Engine/Source/Developer/TraceInsights/Private/Insights/ViewModels/MarkersTimingTrack.h
+++ b/Engine/Source/Developer/TraceInsights/Private/Insights/ViewModels/MarkersTimingTrack.h
@@ -3,9 +3,10 @@
 #pragma once
 
 #include "CoreMinimal.h"
+
 #include "Fonts/SlateFontInfo.h"
 
-// Insights
+// TraceInsights
 #include "Insights/ViewModels/BaseTimingTrack.h"
 #include "Insights/ViewModels/TrackHeader.h"
 
@@ -16,8 +17,10 @@ namespace TraceServices
 	class ILogProvider;
 }
 
-struct FDrawContext;
 struct FSlateBrush;
+
+namespace UE::Insights { class FDrawContext; }
+
 class FTimingTrackViewport;
 
 ////////////////////////////////////////////////////////////////////////////////////////////////////
diff --git a/Engine/Source/Developer/TraceInsights/Private/Insights/ViewModels/QuickFind.cpp b/Engine/Source/Developer/TraceInsights/Private/Insights/ViewModels/QuickFind.cpp
--- a/Engine/Source/Developer/TraceInsights/Private/Insights/ViewModels/QuickFind.cpp
+++ b/Engine/Source/Developer/TraceInsights/Private/Insights/ViewModels/QuickFind.cpp
@@ -7,7 +7,7 @@ namespace Insights
 
 ////////////////////////////////////////////////////////////////////////////////////////////////////
 
-FQuickFind::FQuickFind(TSharedPtr<FFilterConfigurator> InFilterConfiguratorViewModel)
+FQuickFind::FQuickFind(TSharedPtr<UE::Insights::FFilterConfigurator> InFilterConfiguratorViewModel)
 {
 	FilterConfigurator = InFilterConfiguratorViewModel;
 }
diff --git a/Engine/Source/Developer/TraceInsights/Private/Insights/ViewModels/QuickFind.h b/Engine/Source/Developer/TraceInsights/Private/Insights/ViewModels/QuickFind.h
--- a/Engine/Source/Developer/TraceInsights/Private/Insights/ViewModels/QuickFind.h
+++ b/Engine/Source/Developer/TraceInsights/Private/Insights/ViewModels/QuickFind.h
@@ -2,9 +2,13 @@
 
 #pragma once
 
-#include "CoreMinimal.h"
-#include "Insights/ViewModels/Filters.h"
-#include "Insights/ViewModels/FilterConfigurator.h"
+#include "CoreTypes.h"
+
+#include "Delegates/Delegate.h"
+#include "Delegates/DelegateCombinations.h"
+#include "Templates/SharedPointer.h"
+
+namespace UE::Insights { class FFilterConfigurator; }
 
 namespace Insights
 {
@@ -14,11 +18,11 @@ namespace Insights
 class FQuickFind
 {
 public:
-	FQuickFind(TSharedPtr<FFilterConfigurator> InFilterConfiguratorViewModel);
+	FQuickFind(TSharedPtr<UE::Insights::FFilterConfigurator> InFilterConfiguratorViewModel);
 
 	~FQuickFind();
 
-	TSharedPtr<FFilterConfigurator> GetFilterConfigurator() { return FilterConfigurator; }
+	TSharedPtr<UE::Insights::FFilterConfigurator> GetFilterConfigurator() { return FilterConfigurator; }
 
 	////////////////////////////////////////////////////////////////////////////////////////////////////
 	// OnDestroyedEvent
@@ -93,7 +97,7 @@ private:
 	////////////////////////////////////////////////////////////////////////////////////////////////////
 
 private:
-	TSharedPtr<FFilterConfigurator> FilterConfigurator;
+	TSharedPtr<UE::Insights::FFilterConfigurator> FilterConfigurator;
 };
 
 ////////////////////////////////////////////////////////////////////////////////////////////////////
diff --git a/Engine/Source/Developer/TraceInsights/Private/Insights/ViewModels/RegionsTimingTrack.cpp b/Engine/Source/Developer/TraceInsights/Private/Insights/ViewModels/RegionsTimingTrack.cpp
--- a/Engine/Source/Developer/TraceInsights/Private/Insights/ViewModels/RegionsTimingTrack.cpp
+++ b/Engine/Source/Developer/TraceInsights/Private/Insights/ViewModels/RegionsTimingTrack.cpp
@@ -9,16 +9,18 @@
 #include "Common/ProviderLock.h"
 #include "TraceServices/Model/Regions.h"
 
+// TraceInsightsCore
+#include "InsightsCore/Common/TimeUtils.h"
+#include "InsightsCore/Filter/ViewModels/FilterConfigurator.h"
+#include "InsightsCore/Filter/ViewModels/Filters.h"
+
 // TraceInsights
 #include "Insights/Common/InsightsMenuBuilder.h"
-#include "Insights/Common/TimeUtils.h"
 #include "Insights/InsightsManager.h"
 #include "Insights/InsightsStyle.h"
 #include "Insights/ITimingViewSession.h"
 #include "Insights/Log.h"
 #include "Insights/TimingProfilerCommon.h"
-#include "Insights/ViewModels/FilterConfigurator.h"
-#include "Insights/ViewModels/Filters.h"
 #include "Insights/ViewModels/TimingTrackViewport.h"
 #include "Insights/ViewModels/TimingViewDrawHelper.h"
 #include "Insights/Widgets/STimingView.h"
@@ -61,9 +63,11 @@ void FTimingRegionsViewCommands::RegisterCommands()
 }
 UE_ENABLE_OPTIMIZATION_SHIP
 
+////////////////////////////////////////////////////////////////////////////////////////////////////
+// FTimingRegionsSharedState
 ////////////////////////////////////////////////////////////////////////////////////////////////////
 
-void FTimingRegionsSharedState::OnBeginSession(Insights::ITimingViewSession& InSession)
+void FTimingRegionsSharedState::OnBeginSession(UE::Insights::Timing::ITimingViewSession& InSession)
 {
 	if (&InSession != TimingView)
 	{
@@ -75,7 +79,7 @@ void FTimingRegionsSharedState::OnBeginSession(Insights::ITimingViewSession& InS
 
 ////////////////////////////////////////////////////////////////////////////////////////////////////
 
-void FTimingRegionsSharedState::OnEndSession(Insights::ITimingViewSession& InSession)
+void FTimingRegionsSharedState::OnEndSession(UE::Insights::Timing::ITimingViewSession& InSession)
 {
 	if (&InSession != TimingView)
 	{
@@ -87,7 +91,7 @@ void FTimingRegionsSharedState::OnEndSession(Insights::ITimingViewSession& InSes
 
 ////////////////////////////////////////////////////////////////////////////////////////////////////
 
-void FTimingRegionsSharedState::Tick(Insights::ITimingViewSession& InSession,
+void FTimingRegionsSharedState::Tick(UE::Insights::Timing::ITimingViewSession& InSession,
 	const TraceServices::IAnalysisSession& InAnalysisSession)
 {
 	if (&InSession != TimingView)
@@ -128,7 +132,7 @@ void FTimingRegionsSharedState::ShowHideRegionsTrack()
 
 ////////////////////////////////////////////////////////////////////////////////////////////////////
 
-void FTimingRegionsSharedState::ExtendOtherTracksFilterMenu(Insights::ITimingViewSession& InSession,
+void FTimingRegionsSharedState::ExtendOtherTracksFilterMenu(UE::Insights::Timing::ITimingViewSession& InSession,
                                                             FMenuBuilder& InOutMenuBuilder)
 {
 	InOutMenuBuilder.BeginSection("Timing Regions", LOCTEXT("ContextMenu_Section_Regions", "Timing Regions"));
@@ -165,7 +169,9 @@ void FTimingRegionsSharedState::BindCommands()
 		FIsActionChecked::CreateSP(this, &FTimingRegionsSharedState::IsRegionsTrackVisible));
 }
 
-//////////////////////////////////////////////////////////////////////////
+////////////////////////////////////////////////////////////////////////////////////////////////////
+// FTimingRegionsTrack
+////////////////////////////////////////////////////////////////////////////////////////////////////
 
 INSIGHTS_IMPLEMENT_RTTI(FTimingRegionsTrack)
 
@@ -178,7 +184,7 @@ void FTimingRegionsTrack::BuildContextMenu(FMenuBuilder& MenuBuilder)
 	FTimingEventsTrack::BuildContextMenu(MenuBuilder);
 }
 
-//////////////////////////////////////////////////////////////////////////
+////////////////////////////////////////////////////////////////////////////////////////////////////
 
 void FTimingRegionsTrack::InitTooltip(FTooltipDrawState& InOutTooltip, const ITimingEvent& InTooltipEvent) const
 {
@@ -198,13 +204,15 @@ void FTimingRegionsTrack::InitTooltip(FTooltipDrawState& InOutTooltip, const ITi
 		{
 			InOutTooltip.Reset();
 			InOutTooltip.AddTitle(InRegion.Text, FLinearColor::White);
-			InOutTooltip.AddNameValueTextLine(TEXT("Duration:"),  TimeUtils::FormatTimeAuto(InRegion.EndTime-InRegion.BeginTime));
+			InOutTooltip.AddNameValueTextLine(TEXT("Duration:"),  UE::Insights::FormatTimeAuto(InRegion.EndTime - InRegion.BeginTime));
 			InOutTooltip.AddNameValueTextLine(TEXT("Depth:"),  FString::FromInt(InRegion.Depth));
 			InOutTooltip.UpdateLayout();
 		});
 	}
 }
 
+////////////////////////////////////////////////////////////////////////////////////////////////////
+
 void FTimingRegionsTrack::BuildDrawState(ITimingEventsTrackDrawStateBuilder& Builder,
 	const ITimingTrackUpdateContext& Context)
 {
@@ -214,10 +222,10 @@ void FTimingRegionsTrack::BuildDrawState(ITimingEventsTrackDrawStateBuilder& Bui
 	const TraceServices::IRegionProvider& RegionProvider = TraceServices::ReadRegionProvider(*Session);
 	TraceServices::FProviderReadScopeLock RegionProviderScopedLock(RegionProvider);
 
-	FStopwatch Stopwatch;
+	UE::Insights::FStopwatch Stopwatch;
 	Stopwatch.Start();
 
-	// whe're counting only non-empty lanes, so we can collapse empty ones in the visualization.
+	// We are counting only non-empty lanes, so we can collapse empty ones in the visualization.
 	int32 CurDepth = 0;
 	RegionProvider.EnumerateLanes([this, Viewport, &CurDepth, &Builder](const TraceServices::FRegionLane& Lane, const int32 Depth)
 	{
@@ -234,10 +242,10 @@ void FTimingRegionsTrack::BuildDrawState(ITimingEventsTrackDrawStateBuilder& Bui
 
 	Stopwatch.Stop();
 	const double TotalTime = Stopwatch.GetAccumulatedTime();
-	UE_CLOG(TotalTime > 1.0,TimingProfiler, Verbose, TEXT("[Regions] Updated draw state in %s."), *TimeUtils::FormatTimeAuto(TotalTime));
+	UE_CLOG(TotalTime > 1.0, LogTimingProfiler, Verbose, TEXT("[Regions] Updated draw state in %s."), *UE::Insights::FormatTimeAuto(TotalTime));
 }
 
-//////////////////////////////////////////////////////////////////////////
+////////////////////////////////////////////////////////////////////////////////////////////////////
 
 void FTimingRegionsTrack::BuildFilteredDrawState(ITimingEventsTrackDrawStateBuilder& Builder, const ITimingTrackUpdateContext& Context)
 {
@@ -294,7 +302,8 @@ void FTimingRegionsTrack::BuildFilteredDrawState(ITimingEventsTrackDrawStateBuil
 			return;
 		}
 
-		FFilterContext FilterContext;
+		using EFilterField = UE::Insights::EFilterField;
+		UE::Insights::FFilterContext FilterContext;
 		FilterContext.SetReturnValueForUnsetFilters(false);
 
 		FilterContext.AddFilterData<double>(static_cast<int32>(EFilterField::StartTime), 0.0f);
@@ -335,7 +344,7 @@ void FTimingRegionsTrack::BuildFilteredDrawState(ITimingEventsTrackDrawStateBuil
 	}
 }
 
-//////////////////////////////////////////////////////////////////////////
+////////////////////////////////////////////////////////////////////////////////////////////////////
 
 const TSharedPtr<const ITimingEvent> FTimingRegionsTrack::SearchEvent(
 	const FTimingEventSearchParameters& InSearchParameters) const
@@ -350,7 +359,7 @@ const TSharedPtr<const ITimingEvent> FTimingRegionsTrack::SearchEvent(
 	return FoundEvent;
 }
 
-//////////////////////////////////////////////////////////////////////////
+////////////////////////////////////////////////////////////////////////////////////////////////////
 
 bool FTimingRegionsTrack::FindRegionEvent(const FTimingEventSearchParameters& InParameters,
 	TFunctionRef<void(double, double, uint32, const TraceServices::FTimeRegion&)> InFoundPredicate) const
@@ -365,7 +374,8 @@ bool FTimingRegionsTrack::FindRegionEvent(const FTimingEventSearchParameters& In
 		}
 	}
 
-	FFilterContext FilterConfiguratorContext;
+	using EFilterField = UE::Insights::EFilterField;
+	UE::Insights::FFilterContext FilterConfiguratorContext;
 	FilterConfiguratorContext.SetReturnValueForUnsetFilters(false);
 	FilterConfiguratorContext.AddFilterData<double>(static_cast<int32>(EFilterField::StartTime), 0.0f);
 	FilterConfiguratorContext.AddFilterData<double>(static_cast<int32>(EFilterField::EndTime), 0.0f);
@@ -422,9 +432,9 @@ bool FTimingRegionsTrack::FindRegionEvent(const FTimingEventSearchParameters& In
 	TTimingEventSearch<TraceServices::FTimeRegion>::NoMatch);
 }
 
-//////////////////////////////////////////////////////////////////////////
+////////////////////////////////////////////////////////////////////////////////////////////////////
 
-void FTimingRegionsTrack::SetFilterConfigurator(TSharedPtr<Insights::FFilterConfigurator> InFilterConfigurator)
+void FTimingRegionsTrack::SetFilterConfigurator(TSharedPtr<UE::Insights::FFilterConfigurator> InFilterConfigurator)
 {
 	if (FilterConfigurator != InFilterConfigurator)
 	{
@@ -433,14 +443,14 @@ void FTimingRegionsTrack::SetFilterConfigurator(TSharedPtr<Insights::FFilterConf
 	}
 }
 
-//////////////////////////////////////////////////////////////////////////
+////////////////////////////////////////////////////////////////////////////////////////////////////
 
 bool FTimingRegionsTrack::HasCustomFilter() const
 {
 	return FilterConfigurator.IsValid() && !FilterConfigurator->IsEmpty();
 }
 
-//////////////////////////////////////////////////////////////////////////
+////////////////////////////////////////////////////////////////////////////////////////////////////
 
 void FTimingRegionsTrack::OnClipboardCopyEvent(const ITimingEvent& InSelectedEvent) const
 {
@@ -456,7 +466,7 @@ void FTimingRegionsTrack::OnClipboardCopyEvent(const ITimingEvent& InSelectedEve
 	}
 }
 
-//////////////////////////////////////////////////////////////////////////
+////////////////////////////////////////////////////////////////////////////////////////////////////
 
 } // namespace Insights
 
diff --git a/Engine/Source/Developer/TraceInsights/Private/Insights/ViewModels/RegionsTimingTrack.h b/Engine/Source/Developer/TraceInsights/Private/Insights/ViewModels/RegionsTimingTrack.h
--- a/Engine/Source/Developer/TraceInsights/Private/Insights/ViewModels/RegionsTimingTrack.h
+++ b/Engine/Source/Developer/TraceInsights/Private/Insights/ViewModels/RegionsTimingTrack.h
@@ -3,11 +3,15 @@
 #pragma once
 
 #include "Framework/Commands/Commands.h"
+
+// TraceServices
+#include "TraceServices/Model/Regions.h"
+
+// TraceInsights
 #include "Insights/ITimingViewExtender.h"
 #include "Insights/ViewModels/TimingEventsTrack.h"
-#include "TraceServices/Model/Regions.h"
 
-class STimingView;
+namespace UE::Insights::TimingProfiler { class STimingView; }
 
 namespace Insights
 {
@@ -29,27 +33,27 @@ public:
 
 ////////////////////////////////////////////////////////////////////////////////////////////////////
 
-class FTimingRegionsSharedState : public Insights::ITimingViewExtender, public TSharedFromThis<FTimingRegionsSharedState>
+class FTimingRegionsSharedState : public UE::Insights::Timing::ITimingViewExtender, public TSharedFromThis<FTimingRegionsSharedState>
 {
 	friend class FTimingRegionsTrack;
 	
 public:
-	explicit FTimingRegionsSharedState(STimingView* InTimingView) : TimingView(InTimingView) {}
+	explicit FTimingRegionsSharedState(UE::Insights::TimingProfiler::STimingView* InTimingView) : TimingView(InTimingView) {}
 	virtual ~FTimingRegionsSharedState() override = default;
 
 	//////////////////////////////////////////////////
 	// ITimingViewExtender interface
 
-	virtual void OnBeginSession(Insights::ITimingViewSession& InSession) override;
-	virtual void OnEndSession(Insights::ITimingViewSession& InSession) override;
-	virtual void Tick(Insights::ITimingViewSession& InSession, const TraceServices::IAnalysisSession& InAnalysisSession) override;
+	virtual void OnBeginSession(UE::Insights::Timing::ITimingViewSession& InSession) override;
+	virtual void OnEndSession(UE::Insights::Timing::ITimingViewSession& InSession) override;
+	virtual void Tick(UE::Insights::Timing::ITimingViewSession& InSession, const TraceServices::IAnalysisSession& InAnalysisSession) override;
 	void ShowHideRegionsTrack();
 	bool IsRegionsTrackVisible() const {return bShowHideRegionsTrack;};
-	virtual void ExtendOtherTracksFilterMenu(Insights::ITimingViewSession& InSession, FMenuBuilder& InOutMenuBuilder) override;
+	virtual void ExtendOtherTracksFilterMenu(UE::Insights::Timing::ITimingViewSession& InSession, FMenuBuilder& InOutMenuBuilder) override;
 	void BindCommands();
 
 private:
-	STimingView* TimingView;
+	UE::Insights::TimingProfiler::STimingView* TimingView;
 	TSharedPtr<FTimingRegionsTrack> TimingRegionsTrack;
 
 	bool bShowHideRegionsTrack = true;
@@ -70,7 +74,7 @@ public:
 	virtual const TSharedPtr<const ITimingEvent> SearchEvent(const FTimingEventSearchParameters& InSearchParameters) const override;
 	virtual void BuildDrawState(ITimingEventsTrackDrawStateBuilder& Builder, const ITimingTrackUpdateContext& Context) override;
 	virtual void BuildFilteredDrawState(ITimingEventsTrackDrawStateBuilder& Builder, const ITimingTrackUpdateContext& Context) override;
-	virtual void SetFilterConfigurator(TSharedPtr<Insights::FFilterConfigurator> InFilterConfigurator) override;
+	virtual void SetFilterConfigurator(TSharedPtr<UE::Insights::FFilterConfigurator> InFilterConfigurator) override;
 	virtual bool HasCustomFilter() const override;
 	virtual void OnClipboardCopyEvent(const ITimingEvent& InSelectedEvent) const override;
 
@@ -78,7 +82,7 @@ protected:
 	bool FindRegionEvent(const FTimingEventSearchParameters& InParameters, TFunctionRef<void(double, double, uint32, const TraceServices::FTimeRegion&)> InFoundPredicate) const;
 
 private:
-	TSharedPtr<Insights::FFilterConfigurator> FilterConfigurator;
+	TSharedPtr<UE::Insights::FFilterConfigurator> FilterConfigurator;
 	FTimingRegionsSharedState& SharedState;
 };
 
diff --git a/Engine/Source/Developer/TraceInsights/Private/Insights/ViewModels/StatsAggregator.cpp b/Engine/Source/Developer/TraceInsights/Private/Insights/ViewModels/StatsAggregator.cpp
--- a/Engine/Source/Developer/TraceInsights/Private/Insights/ViewModels/StatsAggregator.cpp
+++ b/Engine/Source/Developer/TraceInsights/Private/Insights/ViewModels/StatsAggregator.cpp
@@ -94,7 +94,7 @@ void FStatsAggregator::Start()
 			OperationCount = 0;
 		}
 
-		UE_LOG(TimingProfiler, Log, TEXT("[%s] Request async aggregation (op %d) [%fs to %fs] (%fs)..."),
+		UE_LOG(LogTimingProfiler, Log, TEXT("[%s] Request async aggregation (op %d) [%fs to %fs] (%fs)..."),
 			*LogName, OperationCount + 1, IntervalStartTime, IntervalEndTime, IntervalEndTime - IntervalStartTime);
 		bIsStartRequested = true;
 	}
@@ -106,7 +106,7 @@ void FStatsAggregator::Cancel()
 {
 	if (AsyncTask)
 	{
-		UE_LOG(TimingProfiler, Log, TEXT("[%s] Cancel requested for async aggregation (op %d)..."), *LogName, OperationCount);
+		UE_LOG(LogTimingProfiler, Log, TEXT("[%s] Cancel requested for async aggregation (op %d)..."), *LogName, OperationCount);
 		CancellationToken->Cancel();
 	}
 }
@@ -124,7 +124,7 @@ void FStatsAggregator::Tick(TSharedPtr<const TraceServices::IAnalysisSession> In
 
 		if (!bIsStartRequested && !CancellationToken->ShouldCancel())
 		{
-			FStopwatch FinishedStopwatch;
+			UE::Insights::FStopwatch FinishedStopwatch;
 			FinishedStopwatch.Start();
 
 			bIsFinished = true;
@@ -146,7 +146,7 @@ void FStatsAggregator::Tick(TSharedPtr<const TraceServices::IAnalysisSession> In
 		AllOpsStopwatch.Update();
 		const double AllOpsDuration = AllOpsStopwatch.GetAccumulatedTime();
 
-		UE_LOG(TimingProfiler, Log, TEXT("[%s] Aggregated stats computed in %.4fs (%.4fs + %.4fs)%s%s - total: %.4fs (%d ops)"),
+		UE_LOG(LogTimingProfiler, Log, TEXT("[%s] Aggregated stats computed in %.4fs (%.4fs + %.4fs)%s%s - total: %.4fs (%d ops)"),
 			*LogName,
 			CurrentOpDuration, CurrentOpDuration - FinishedDuration, FinishedDuration,
 			CancellationToken->ShouldCancel() ? TEXT(" - CANCELED") : TEXT(""),
@@ -161,7 +161,7 @@ void FStatsAggregator::Tick(TSharedPtr<const TraceServices::IAnalysisSession> In
 			if (!CancellationToken->ShouldCancel())
 			{
 				// Cancel and wait for the previous async task to finish.
-				UE_LOG(TimingProfiler, Log, TEXT("[%s] Cancel previous async aggregation (op %d)..."), *LogName, OperationCount);
+				UE_LOG(LogTimingProfiler, Log, TEXT("[%s] Cancel previous async aggregation (op %d)..."), *LogName, OperationCount);
 				CancellationToken->Cancel();
 			}
 		}
@@ -173,7 +173,7 @@ void FStatsAggregator::Tick(TSharedPtr<const TraceServices::IAnalysisSession> In
 			bIsStartRequested = false;
 			CancellationToken->Reset();
 
-			UE_LOG(TimingProfiler, Log, TEXT("[%s] Start async aggregation (op %d)..."), *LogName, OperationCount);
+			UE_LOG(LogTimingProfiler, Log, TEXT("[%s] Start async aggregation (op %d)..."), *LogName, OperationCount);
 
 			IStatsAggregationWorker* Worker = CreateWorker(InSession);
 			check(Worker);
diff --git a/Engine/Source/Developer/TraceInsights/Private/Insights/ViewModels/StatsAggregator.h b/Engine/Source/Developer/TraceInsights/Private/Insights/ViewModels/StatsAggregator.h
--- a/Engine/Source/Developer/TraceInsights/Private/Insights/ViewModels/StatsAggregator.h
+++ b/Engine/Source/Developer/TraceInsights/Private/Insights/ViewModels/StatsAggregator.h
@@ -2,12 +2,16 @@
 
 #pragma once
 
-#include "TraceServices/Common/CancellationToken.h"
-
 #include "CoreMinimal.h"
+
 #include "Async/AsyncWork.h"
-#include "Insights/Common/InsightsAsyncWorkUtils.h"
-#include "Insights/Common/Stopwatch.h"
+
+// TraceServices
+#include "TraceServices/Common/CancellationToken.h"
+
+// TraceInsightsCore
+#include "InsightsCore/Common/IAsyncOperationStatusProvider.h"
+#include "InsightsCore/Common/Stopwatch.h"
 
 ////////////////////////////////////////////////////////////////////////////////////////////////////
 
@@ -21,7 +25,7 @@ namespace Insights
 
 ////////////////////////////////////////////////////////////////////////////////////////////////////
 
-class IStatsAggregator : public IAsyncOperationStatusProvider
+class IStatsAggregator : public UE::Insights::IAsyncOperationStatusProvider
 {
 public:
 	virtual void Start() = 0;
@@ -105,8 +109,8 @@ private:
 	bool bIsStartRequested;
 	bool bIsFinished;
 
-	mutable FStopwatch AllOpsStopwatch;
-	mutable FStopwatch CurrentOpStopwatch;
+	mutable UE::Insights::FStopwatch AllOpsStopwatch;
+	mutable UE::Insights::FStopwatch CurrentOpStopwatch;
 	uint32 OperationCount;
 };
 
diff --git a/Engine/Source/Developer/TraceInsights/Private/Insights/ViewModels/StatsGroupingAndSorting.cpp b/Engine/Source/Developer/TraceInsights/Private/Insights/ViewModels/StatsGroupingAndSorting.cpp
--- a/Engine/Source/Developer/TraceInsights/Private/Insights/ViewModels/StatsGroupingAndSorting.cpp
+++ b/Engine/Source/Developer/TraceInsights/Private/Insights/ViewModels/StatsGroupingAndSorting.cpp
@@ -2,7 +2,10 @@
 
 #include "StatsGroupingAndSorting.h"
 
-#include "Insights/Table/ViewModels/TableColumn.h"
+// TraceInsightsCore
+#include "InsightsCore/Table/ViewModels/TableColumn.h"
+
+// TraceInsights
 #include "Insights/ViewModels/StatsNodeHelper.h"
 
 #define LOCTEXT_NAMESPACE "StatsNode"
@@ -18,8 +21,8 @@
 // Sorting by Node Type
 ////////////////////////////////////////////////////////////////////////////////////////////////////
 
-FStatsNodeSortingByStatsType::FStatsNodeSortingByStatsType(TSharedRef<Insights::FTableColumn> InColumnRef)
-	: Insights::FTableCellValueSorter(
+FStatsNodeSortingByStatsType::FStatsNodeSortingByStatsType(TSharedRef<UE::Insights::FTableColumn> InColumnRef)
+	: UE::Insights::FTableCellValueSorter(
 		FName(TEXT("ByStatsType")),
 		LOCTEXT("Sorting_ByStatsType_Name", "By Type"),
 		LOCTEXT("Sorting_ByStatsType_Title", "Sort By Type"),
@@ -30,11 +33,11 @@ FStatsNodeSortingByStatsType::FStatsNodeSortingByStatsType(TSharedRef<Insights::
 
 ////////////////////////////////////////////////////////////////////////////////////////////////////
 
-void FStatsNodeSortingByStatsType::Sort(TArray<Insights::FBaseTreeNodePtr>& NodesToSort, Insights::ESortMode SortMode) const
+void FStatsNodeSortingByStatsType::Sort(TArray<UE::Insights::FBaseTreeNodePtr>& NodesToSort, UE::Insights::ESortMode SortMode) const
 {
-	if (SortMode == Insights::ESortMode::Ascending)
+	if (SortMode == UE::Insights::ESortMode::Ascending)
 	{
-		NodesToSort.Sort([](const Insights::FBaseTreeNodePtr& A, const Insights::FBaseTreeNodePtr& B) -> bool
+		NodesToSort.Sort([](const UE::Insights::FBaseTreeNodePtr& A, const UE::Insights::FBaseTreeNodePtr& B) -> bool
 		{
 			INSIGHTS_ENSURE(A.IsValid() && A->GetTypeName() == FStatsNode::TypeName);
 			const EStatsNodeType ValueA = reinterpret_cast<FStatsNode*>(A.Get())->GetType();
@@ -55,7 +58,7 @@ void FStatsNodeSortingByStatsType::Sort(TArray<Insights::FBaseTreeNodePtr>& Node
 	}
 	else // if (SortMode == ESortMode::Descending)
 	{
-		NodesToSort.Sort([](const Insights::FBaseTreeNodePtr& A, const Insights::FBaseTreeNodePtr& B) -> bool
+		NodesToSort.Sort([](const UE::Insights::FBaseTreeNodePtr& A, const UE::Insights::FBaseTreeNodePtr& B) -> bool
 		{
 			INSIGHTS_ENSURE(A.IsValid() && A->GetTypeName() == FStatsNode::TypeName);
 			const EStatsNodeType ValueA = reinterpret_cast<FStatsNode*>(A.Get())->GetType();
@@ -80,8 +83,8 @@ void FStatsNodeSortingByStatsType::Sort(TArray<Insights::FBaseTreeNodePtr>& Node
 // Sorting by Data Type
 ////////////////////////////////////////////////////////////////////////////////////////////////////
 
-FStatsNodeSortingByDataType::FStatsNodeSortingByDataType(TSharedRef<Insights::FTableColumn> InColumnRef)
-	: Insights::FTableCellValueSorter(
+FStatsNodeSortingByDataType::FStatsNodeSortingByDataType(TSharedRef<UE::Insights::FTableColumn> InColumnRef)
+	: UE::Insights::FTableCellValueSorter(
 		FName(TEXT("ByDataType")),
 		LOCTEXT("Sorting_ByDataType_Name", "By Data Type"),
 		LOCTEXT("Sorting_ByDataType_Title", "Sort By Data Type"),
@@ -92,11 +95,11 @@ FStatsNodeSortingByDataType::FStatsNodeSortingByDataType(TSharedRef<Insights::FT
 
 ////////////////////////////////////////////////////////////////////////////////////////////////////
 
-void FStatsNodeSortingByDataType::Sort(TArray<Insights::FBaseTreeNodePtr>& NodesToSort, Insights::ESortMode SortMode) const
+void FStatsNodeSortingByDataType::Sort(TArray<UE::Insights::FBaseTreeNodePtr>& NodesToSort, UE::Insights::ESortMode SortMode) const
 {
-	if (SortMode == Insights::ESortMode::Ascending)
+	if (SortMode == UE::Insights::ESortMode::Ascending)
 	{
-		NodesToSort.Sort([](const Insights::FBaseTreeNodePtr& A, const Insights::FBaseTreeNodePtr& B) -> bool
+		NodesToSort.Sort([](const UE::Insights::FBaseTreeNodePtr& A, const UE::Insights::FBaseTreeNodePtr& B) -> bool
 		{
 			INSIGHTS_ENSURE(A.IsValid() && A->GetTypeName() == FStatsNode::TypeName);
 			const EStatsNodeDataType ValueA = reinterpret_cast<FStatsNode*>(A.Get())->GetDataType();
@@ -117,7 +120,7 @@ void FStatsNodeSortingByDataType::Sort(TArray<Insights::FBaseTreeNodePtr>& Nodes
 	}
 	else // if (SortMode == ESortMode::Descending)
 	{
-		NodesToSort.Sort([](const Insights::FBaseTreeNodePtr& A, const Insights::FBaseTreeNodePtr& B) -> bool
+		NodesToSort.Sort([](const UE::Insights::FBaseTreeNodePtr& A, const UE::Insights::FBaseTreeNodePtr& B) -> bool
 		{
 			INSIGHTS_ENSURE(A.IsValid() && A->GetTypeName() == FStatsNode::TypeName);
 			const EStatsNodeDataType ValueA = reinterpret_cast<FStatsNode*>(A.Get())->GetDataType();
@@ -142,8 +145,8 @@ void FStatsNodeSortingByDataType::Sort(TArray<Insights::FBaseTreeNodePtr>& Nodes
 // Sort by Count (Aggregated Statistic)
 ////////////////////////////////////////////////////////////////////////////////////////////////////
 
-FStatsNodeSortingByCount::FStatsNodeSortingByCount(TSharedRef<Insights::FTableColumn> InColumnRef)
-	: Insights::FTableCellValueSorter(
+FStatsNodeSortingByCount::FStatsNodeSortingByCount(TSharedRef<UE::Insights::FTableColumn> InColumnRef)
+	: UE::Insights::FTableCellValueSorter(
 		FName(TEXT("ByCount")),
 		LOCTEXT("Sorting_ByCount_Name", "By Count"),
 		LOCTEXT("Sorting_ByCount_Title", "Sort By Count"),
@@ -154,11 +157,11 @@ FStatsNodeSortingByCount::FStatsNodeSortingByCount(TSharedRef<Insights::FTableCo
 
 ////////////////////////////////////////////////////////////////////////////////////////////////////
 
-void FStatsNodeSortingByCount::Sort(TArray<Insights::FBaseTreeNodePtr>& NodesToSort, Insights::ESortMode SortMode) const
+void FStatsNodeSortingByCount::Sort(TArray<UE::Insights::FBaseTreeNodePtr>& NodesToSort, UE::Insights::ESortMode SortMode) const
 {
-	if (SortMode == Insights::ESortMode::Ascending)
+	if (SortMode == UE::Insights::ESortMode::Ascending)
 	{
-		NodesToSort.Sort([](const Insights::FBaseTreeNodePtr& A, const Insights::FBaseTreeNodePtr& B) -> bool
+		NodesToSort.Sort([](const UE::Insights::FBaseTreeNodePtr& A, const UE::Insights::FBaseTreeNodePtr& B) -> bool
 		{
 			INSIGHTS_ENSURE(A.IsValid() && A->GetTypeName() == FStatsNode::TypeName);
 			const uint64 ValueA = reinterpret_cast<FStatsNode*>(A.Get())->GetAggregatedStats().Count;
@@ -179,7 +182,7 @@ void FStatsNodeSortingByCount::Sort(TArray<Insights::FBaseTreeNodePtr>& NodesToS
 	}
 	else // if (SortMode == ESortMode::Descending)
 	{
-		NodesToSort.Sort([](const Insights::FBaseTreeNodePtr& A, const Insights::FBaseTreeNodePtr& B) -> bool
+		NodesToSort.Sort([](const UE::Insights::FBaseTreeNodePtr& A, const UE::Insights::FBaseTreeNodePtr& B) -> bool
 		{
 			INSIGHTS_ENSURE(A.IsValid() && A->GetTypeName() == FStatsNode::TypeName);
 			const uint64 ValueA = reinterpret_cast<FStatsNode*>(A.Get())->GetAggregatedStats().Count;
diff --git a/Engine/Source/Developer/TraceInsights/Private/Insights/ViewModels/StatsGroupingAndSorting.h b/Engine/Source/Developer/TraceInsights/Private/Insights/ViewModels/StatsGroupingAndSorting.h
--- a/Engine/Source/Developer/TraceInsights/Private/Insights/ViewModels/StatsGroupingAndSorting.h
+++ b/Engine/Source/Developer/TraceInsights/Private/Insights/ViewModels/StatsGroupingAndSorting.h
@@ -4,42 +4,44 @@
 
 #include "CoreMinimal.h"
 
-// Insights
+// TraceInsightsCore
+#include "InsightsCore/Table/ViewModels/TableColumn.h"
+#include "InsightsCore/Table/ViewModels/TableCellValueSorter.h"
+#include "InsightsCore/Table/ViewModels/TreeNodeGrouping.h"
+
+// TraceInsights
 #include "Insights/ViewModels/StatsNode.h"
-#include "Insights/Table/ViewModels/TableColumn.h"
-#include "Insights/Table/ViewModels/TableCellValueSorter.h"
-#include "Insights/Table/ViewModels/TreeNodeGrouping.h"
 
 ////////////////////////////////////////////////////////////////////////////////////////////////////
 // Sorters
 ////////////////////////////////////////////////////////////////////////////////////////////////////
 
-class FStatsNodeSortingByStatsType : public Insights::FTableCellValueSorter
+class FStatsNodeSortingByStatsType : public UE::Insights::FTableCellValueSorter
 {
 public:
-	FStatsNodeSortingByStatsType(TSharedRef<Insights::FTableColumn> InColumnRef);
+	FStatsNodeSortingByStatsType(TSharedRef<UE::Insights::FTableColumn> InColumnRef);
 
-	virtual void Sort(TArray<Insights::FBaseTreeNodePtr>& NodesToSort, Insights::ESortMode SortMode) const override;
+	virtual void Sort(TArray<UE::Insights::FBaseTreeNodePtr>& NodesToSort, UE::Insights::ESortMode SortMode) const override;
 };
 
 ////////////////////////////////////////////////////////////////////////////////////////////////////
 
-class FStatsNodeSortingByDataType : public Insights::FTableCellValueSorter
+class FStatsNodeSortingByDataType : public UE::Insights::FTableCellValueSorter
 {
 public:
-	FStatsNodeSortingByDataType(TSharedRef<Insights::FTableColumn> InColumnRef);
+	FStatsNodeSortingByDataType(TSharedRef<UE::Insights::FTableColumn> InColumnRef);
 
-	virtual void Sort(TArray<Insights::FBaseTreeNodePtr>& NodesToSort, Insights::ESortMode SortMode) const override;
+	virtual void Sort(TArray<UE::Insights::FBaseTreeNodePtr>& NodesToSort, UE::Insights::ESortMode SortMode) const override;
 };
 
 ////////////////////////////////////////////////////////////////////////////////////////////////////
 
-class FStatsNodeSortingByCount : public Insights::FTableCellValueSorter
+class FStatsNodeSortingByCount : public UE::Insights::FTableCellValueSorter
 {
 public:
-	FStatsNodeSortingByCount(TSharedRef<Insights::FTableColumn> InColumnRef);
+	FStatsNodeSortingByCount(TSharedRef<UE::Insights::FTableColumn> InColumnRef);
 
-	virtual void Sort(TArray<Insights::FBaseTreeNodePtr>& NodesToSort, Insights::ESortMode SortMode) const override;
+	virtual void Sort(TArray<UE::Insights::FBaseTreeNodePtr>& NodesToSort, UE::Insights::ESortMode SortMode) const override;
 };
 
 ////////////////////////////////////////////////////////////////////////////////////////////////////
diff --git a/Engine/Source/Developer/TraceInsights/Private/Insights/ViewModels/StatsNode.cpp b/Engine/Source/Developer/TraceInsights/Private/Insights/ViewModels/StatsNode.cpp
--- a/Engine/Source/Developer/TraceInsights/Private/Insights/ViewModels/StatsNode.cpp
+++ b/Engine/Source/Developer/TraceInsights/Private/Insights/ViewModels/StatsNode.cpp
@@ -2,8 +2,8 @@
 
 #include "StatsNode.h"
 
-// Insights
-#include "Insights/Common/TimeUtils.h"
+// TraceInsightsCore
+#include "InsightsCore/Common/TimeUtils.h"
 
 #define LOCTEXT_NAMESPACE "FStatsNode"
 
@@ -99,13 +99,14 @@ const FText FStatsNode::FormatAggregatedStatsValue(double ValueDbl, int64 ValueI
 			//TODO: if (GetDisplayHint() == FStatsNodeDisplayHint::Seconds)
 			if (GetMetaGroupName() == FStatsNodeMetaGroupName::Time)
 			{
+				using namespace UE::Insights;
 				if (bForTooltip)
 				{
-					return FText::FromString(TimeUtils::FormatTimeAuto(ValueDbl, 2));
+					return FText::FromString(FormatTimeAuto(ValueDbl, 2));
 				}
 				else
 				{
-					return FText::FromString(TimeUtils::FormatTimeAuto(ValueDbl, 1));
+					return FText::FromString(FormatTimeAuto(ValueDbl, 1));
 				}
 			}
 			else
diff --git a/Engine/Source/Developer/TraceInsights/Private/Insights/ViewModels/StatsNode.h b/Engine/Source/Developer/TraceInsights/Private/Insights/ViewModels/StatsNode.h
--- a/Engine/Source/Developer/TraceInsights/Private/Insights/ViewModels/StatsNode.h
+++ b/Engine/Source/Developer/TraceInsights/Private/Insights/ViewModels/StatsNode.h
@@ -4,8 +4,8 @@
 
 #include "CoreMinimal.h"
 
-// Insights
-#include "Insights/Table/ViewModels/BaseTreeNode.h"
+// TraceInsightsCore
+#include "InsightsCore/Table/ViewModels/BaseTreeNode.h"
 
 ////////////////////////////////////////////////////////////////////////////////////////////////////
 
@@ -129,9 +129,9 @@ typedef TWeakPtr<class FStatsNode> FStatsNodeWeak;
 /**
  * Class used to store information about a stats counter node (used in SStatsView).
  */
-class FStatsNode : public Insights::FBaseTreeNode
+class FStatsNode : public UE::Insights::FBaseTreeNode
 {
-	INSIGHTS_DECLARE_RTTI(FStatsNode, FBaseTreeNode)
+	INSIGHTS_DECLARE_RTTI(FStatsNode, UE::Insights::FBaseTreeNode)
 
 public:
 	static constexpr uint32 InvalidCounterId = uint32(-1);
@@ -139,7 +139,7 @@ public:
 public:
 	/** Initialization constructor for the stats node. */
 	FStatsNode(uint32 InCounterId, const FName InName, const FName InMetaGroupName, EStatsNodeType InType, EStatsNodeDataType InDataType)
-		: FBaseTreeNode(InName, InType == EStatsNodeType::Group)
+		: UE::Insights::FBaseTreeNode(InName, InType == EStatsNodeType::Group)
 		, CounterId(InCounterId)
 		, MetaGroupName(InMetaGroupName)
 		, Type(InType)
@@ -157,7 +157,7 @@ public:
 
 	/** Initialization constructor for the group node. */
 	explicit FStatsNode(const FName InGroupName)
-		: FBaseTreeNode(InGroupName, true)
+		: UE::Insights::FBaseTreeNode(InGroupName, true)
 		, CounterId(InvalidCounterId)
 		, Type(EStatsNodeType::Group)
 		, DataType(EStatsNodeDataType::InvalidOrMax)
diff --git a/Engine/Source/Developer/TraceInsights/Private/Insights/ViewModels/StatsNodeHelper.h b/Engine/Source/Developer/TraceInsights/Private/Insights/ViewModels/StatsNodeHelper.h
--- a/Engine/Source/Developer/TraceInsights/Private/Insights/ViewModels/StatsNodeHelper.h
+++ b/Engine/Source/Developer/TraceInsights/Private/Insights/ViewModels/StatsNodeHelper.h
@@ -3,9 +3,10 @@
 #pragma once
 
 #include "CoreMinimal.h"
+
 #include "Styling/SlateBrush.h"
 
-// Insights
+// TraceInsights
 #include "Insights/ViewModels/StatsNode.h"
 #include "Insights/ViewModels/StatsGroupingAndSorting.h"
 
diff --git a/Engine/Source/Developer/TraceInsights/Private/Insights/ViewModels/StatsViewColumnFactory.cpp b/Engine/Source/Developer/TraceInsights/Private/Insights/ViewModels/StatsViewColumnFactory.cpp
--- a/Engine/Source/Developer/TraceInsights/Private/Insights/ViewModels/StatsViewColumnFactory.cpp
+++ b/Engine/Source/Developer/TraceInsights/Private/Insights/ViewModels/StatsViewColumnFactory.cpp
@@ -2,11 +2,13 @@
 
 #include "StatsViewColumnFactory.h"
 
-// Insights
-#include "Insights/Common/TimeUtils.h"
-#include "Insights/Table/ViewModels/TableCellValueFormatter.h"
-#include "Insights/Table/ViewModels/TableCellValueGetter.h"
-#include "Insights/Table/ViewModels/TableCellValueSorter.h"
+// TraceInsightsCore
+#include "InsightsCore/Common/TimeUtils.h"
+#include "InsightsCore/Table/ViewModels/TableCellValueFormatter.h"
+#include "InsightsCore/Table/ViewModels/TableCellValueGetter.h"
+#include "InsightsCore/Table/ViewModels/TableCellValueSorter.h"
+
+// TraceInsights
 #include "Insights/ViewModels/StatsGroupingAndSorting.h"
 #include "Insights/ViewModels/StatsNodeHelper.h"
 
@@ -31,7 +33,7 @@ const FName FStatsViewColumns::DiffColumnID(TEXT("Diff"));
 
 ////////////////////////////////////////////////////////////////////////////////////////////////////
 
-void FStatsViewColumnFactory::CreateStatsViewColumns(TArray<TSharedRef<Insights::FTableColumn>>& Columns)
+void FStatsViewColumnFactory::CreateStatsViewColumns(TArray<TSharedRef<UE::Insights::FTableColumn>>& Columns)
 {
 	Columns.Reset();
 
@@ -52,9 +54,9 @@ void FStatsViewColumnFactory::CreateStatsViewColumns(TArray<TSharedRef<Insights:
 
 ////////////////////////////////////////////////////////////////////////////////////////////////////
 
-TSharedRef<Insights::FTableColumn> FStatsViewColumnFactory::CreateNameColumn()
+TSharedRef<UE::Insights::FTableColumn> FStatsViewColumnFactory::CreateNameColumn()
 {
-	using namespace Insights;
+	using namespace UE::Insights;
 
 	TSharedRef<FTableColumn> ColumnRef = MakeShared<FTableColumn>(FStatsViewColumns::NameColumnID);
 	FTableColumn& Column = *ColumnRef;
@@ -87,9 +89,9 @@ TSharedRef<Insights::FTableColumn> FStatsViewColumnFactory::CreateNameColumn()
 
 ////////////////////////////////////////////////////////////////////////////////////////////////////
 
-TSharedRef<Insights::FTableColumn> FStatsViewColumnFactory::CreateMetaGroupNameColumn()
+TSharedRef<UE::Insights::FTableColumn> FStatsViewColumnFactory::CreateMetaGroupNameColumn()
 {
-	using namespace Insights;
+	using namespace UE::Insights;
 
 	TSharedRef<FTableColumn> ColumnRef = MakeShared<FTableColumn>(FStatsViewColumns::MetaGroupNameColumnID);
 	FTableColumn& Column = *ColumnRef;
@@ -131,9 +133,9 @@ TSharedRef<Insights::FTableColumn> FStatsViewColumnFactory::CreateMetaGroupNameC
 
 ////////////////////////////////////////////////////////////////////////////////////////////////////
 
-TSharedRef<Insights::FTableColumn> FStatsViewColumnFactory::CreateTypeColumn()
+TSharedRef<UE::Insights::FTableColumn> FStatsViewColumnFactory::CreateTypeColumn()
 {
-	using namespace Insights;
+	using namespace UE::Insights;
 
 	TSharedRef<FTableColumn> ColumnRef = MakeShared<FTableColumn>(FStatsViewColumns::TypeColumnID);
 	FTableColumn& Column = *ColumnRef;
@@ -176,9 +178,9 @@ TSharedRef<Insights::FTableColumn> FStatsViewColumnFactory::CreateTypeColumn()
 
 ////////////////////////////////////////////////////////////////////////////////////////////////////
 
-TSharedRef<Insights::FTableColumn> FStatsViewColumnFactory::CreateDataTypeColumn()
+TSharedRef<UE::Insights::FTableColumn> FStatsViewColumnFactory::CreateDataTypeColumn()
 {
-	using namespace Insights;
+	using namespace UE::Insights;
 
 	TSharedRef<FTableColumn> ColumnRef = MakeShared<FTableColumn>(FStatsViewColumns::DataTypeColumnID);
 	FTableColumn& Column = *ColumnRef;
@@ -221,9 +223,9 @@ TSharedRef<Insights::FTableColumn> FStatsViewColumnFactory::CreateDataTypeColumn
 
 ////////////////////////////////////////////////////////////////////////////////////////////////////
 
-TSharedRef<Insights::FTableColumn> FStatsViewColumnFactory::CreateCountColumn()
+TSharedRef<UE::Insights::FTableColumn> FStatsViewColumnFactory::CreateCountColumn()
 {
-	using namespace Insights;
+	using namespace UE::Insights;
 
 	TSharedRef<FTableColumn> ColumnRef = MakeShared<FTableColumn>(FStatsViewColumns::CountColumnID);
 	FTableColumn& Column = *ColumnRef;
@@ -268,9 +270,9 @@ TSharedRef<Insights::FTableColumn> FStatsViewColumnFactory::CreateCountColumn()
 
 ////////////////////////////////////////////////////////////////////////////////////////////////////
 
-TSharedRef<Insights::FTableColumn> FStatsViewColumnFactory::CreateSumColumn()
+TSharedRef<UE::Insights::FTableColumn> FStatsViewColumnFactory::CreateSumColumn()
 {
-	using namespace Insights;
+	using namespace UE::Insights;
 
 	TSharedRef<FTableColumn> ColumnRef = MakeShared<FTableColumn>(FStatsViewColumns::SumColumnID);
 	FTableColumn& Column = *ColumnRef;
@@ -326,9 +328,9 @@ TSharedRef<Insights::FTableColumn> FStatsViewColumnFactory::CreateSumColumn()
 
 ////////////////////////////////////////////////////////////////////////////////////////////////////
 
-TSharedRef<Insights::FTableColumn> FStatsViewColumnFactory::CreateMaxColumn()
+TSharedRef<UE::Insights::FTableColumn> FStatsViewColumnFactory::CreateMaxColumn()
 {
-	using namespace Insights;
+	using namespace UE::Insights;
 
 	TSharedRef<FTableColumn> ColumnRef = MakeShared<FTableColumn>(FStatsViewColumns::MaxColumnID);
 	FTableColumn& Column = *ColumnRef;
@@ -384,9 +386,9 @@ TSharedRef<Insights::FTableColumn> FStatsViewColumnFactory::CreateMaxColumn()
 
 ////////////////////////////////////////////////////////////////////////////////////////////////////
 
-TSharedRef<Insights::FTableColumn> FStatsViewColumnFactory::CreateUpperQuartileColumn()
+TSharedRef<UE::Insights::FTableColumn> FStatsViewColumnFactory::CreateUpperQuartileColumn()
 {
-	using namespace Insights;
+	using namespace UE::Insights;
 
 	TSharedRef<FTableColumn> ColumnRef = MakeShared<FTableColumn>(FStatsViewColumns::UpperQuartileColumnID);
 	FTableColumn& Column = *ColumnRef;
@@ -442,9 +444,9 @@ TSharedRef<Insights::FTableColumn> FStatsViewColumnFactory::CreateUpperQuartileC
 
 ////////////////////////////////////////////////////////////////////////////////////////////////////
 
-TSharedRef<Insights::FTableColumn> FStatsViewColumnFactory::CreateAverageColumn()
+TSharedRef<UE::Insights::FTableColumn> FStatsViewColumnFactory::CreateAverageColumn()
 {
-	using namespace Insights;
+	using namespace UE::Insights;
 
 	TSharedRef<FTableColumn> ColumnRef = MakeShared<FTableColumn>(FStatsViewColumns::AverageColumnID);
 	FTableColumn& Column = *ColumnRef;
@@ -499,9 +501,9 @@ TSharedRef<Insights::FTableColumn> FStatsViewColumnFactory::CreateAverageColumn(
 
 ////////////////////////////////////////////////////////////////////////////////////////////////////
 
-TSharedRef<Insights::FTableColumn> FStatsViewColumnFactory::CreateMedianColumn()
+TSharedRef<UE::Insights::FTableColumn> FStatsViewColumnFactory::CreateMedianColumn()
 {
-	using namespace Insights;
+	using namespace UE::Insights;
 
 	TSharedRef<FTableColumn> ColumnRef = MakeShared<FTableColumn>(FStatsViewColumns::MedianColumnID);
 	FTableColumn& Column = *ColumnRef;
@@ -556,9 +558,9 @@ TSharedRef<Insights::FTableColumn> FStatsViewColumnFactory::CreateMedianColumn()
 
 ////////////////////////////////////////////////////////////////////////////////////////////////////
 
-TSharedRef<Insights::FTableColumn> FStatsViewColumnFactory::CreateLowerQuartileColumn()
+TSharedRef<UE::Insights::FTableColumn> FStatsViewColumnFactory::CreateLowerQuartileColumn()
 {
-	using namespace Insights;
+	using namespace UE::Insights;
 
 	TSharedRef<FTableColumn> ColumnRef = MakeShared<FTableColumn>(FStatsViewColumns::LowerQuartileColumnID);
 	FTableColumn& Column = *ColumnRef;
@@ -614,9 +616,9 @@ TSharedRef<Insights::FTableColumn> FStatsViewColumnFactory::CreateLowerQuartileC
 
 ////////////////////////////////////////////////////////////////////////////////////////////////////
 
-TSharedRef<Insights::FTableColumn> FStatsViewColumnFactory::CreateMinColumn()
+TSharedRef<UE::Insights::FTableColumn> FStatsViewColumnFactory::CreateMinColumn()
 {
-	using namespace Insights;
+	using namespace UE::Insights;
 
 	TSharedRef<FTableColumn> ColumnRef = MakeShared<FTableColumn>(FStatsViewColumns::MinColumnID);
 	FTableColumn& Column = *ColumnRef;
@@ -672,9 +674,9 @@ TSharedRef<Insights::FTableColumn> FStatsViewColumnFactory::CreateMinColumn()
 
 ////////////////////////////////////////////////////////////////////////////////////////////////////
 
-TSharedRef<Insights::FTableColumn> FStatsViewColumnFactory::CreateDiffColumn()
+TSharedRef<UE::Insights::FTableColumn> FStatsViewColumnFactory::CreateDiffColumn()
 {
-	using namespace Insights;
+	using namespace UE::Insights;
 
 	TSharedRef<FTableColumn> ColumnRef = MakeShared<FTableColumn>(FStatsViewColumns::DiffColumnID);
 	FTableColumn& Column = *ColumnRef;
diff --git a/Engine/Source/Developer/TraceInsights/Private/Insights/ViewModels/StatsViewColumnFactory.h b/Engine/Source/Developer/TraceInsights/Private/Insights/ViewModels/StatsViewColumnFactory.h
--- a/Engine/Source/Developer/TraceInsights/Private/Insights/ViewModels/StatsViewColumnFactory.h
+++ b/Engine/Source/Developer/TraceInsights/Private/Insights/ViewModels/StatsViewColumnFactory.h
@@ -2,13 +2,17 @@
 
 #pragma once
 
-#include "CoreMinimal.h"
+#include "CoreTypes.h"
 
-////////////////////////////////////////////////////////////////////////////////////////////////////
+#include "Containers/Array.h"
+#include "Templates/SharedPointer.h"
+#include "UObject/NameTypes.h"
 
+namespace UE::Insights { class FTableColumn; }
+
+// Column identifiers
 struct FStatsViewColumns
 {
-	// Column identifiers
 	static const FName NameColumnID;
 	static const FName MetaGroupNameColumnID;
 	static const FName TypeColumnID;
@@ -24,36 +28,25 @@ struct FStatsViewColumns
 	static const FName DiffColumnID;
 };
 
-////////////////////////////////////////////////////////////////////////////////////////////////////
-
-namespace Insights
-{
-	class FTableColumn;
-}
-
-////////////////////////////////////////////////////////////////////////////////////////////////////
-
 struct FStatsViewColumnFactory
 {
 public:
-	static void CreateStatsViewColumns(TArray<TSharedRef<Insights::FTableColumn>>& Columns);
-
-	static TSharedRef<Insights::FTableColumn> CreateNameColumn();
-	static TSharedRef<Insights::FTableColumn> CreateMetaGroupNameColumn();
-	static TSharedRef<Insights::FTableColumn> CreateTypeColumn();
-	static TSharedRef<Insights::FTableColumn> CreateDataTypeColumn();
-	static TSharedRef<Insights::FTableColumn> CreateCountColumn();
-	static TSharedRef<Insights::FTableColumn> CreateSumColumn();
-	static TSharedRef<Insights::FTableColumn> CreateMaxColumn();
-	static TSharedRef<Insights::FTableColumn> CreateUpperQuartileColumn();
-	static TSharedRef<Insights::FTableColumn> CreateAverageColumn();
-	static TSharedRef<Insights::FTableColumn> CreateMedianColumn();
-	static TSharedRef<Insights::FTableColumn> CreateLowerQuartileColumn();
-	static TSharedRef<Insights::FTableColumn> CreateMinColumn();
-	static TSharedRef<Insights::FTableColumn> CreateDiffColumn();
+	static void CreateStatsViewColumns(TArray<TSharedRef<UE::Insights::FTableColumn>>& Columns);
+
+	static TSharedRef<UE::Insights::FTableColumn> CreateNameColumn();
+	static TSharedRef<UE::Insights::FTableColumn> CreateMetaGroupNameColumn();
+	static TSharedRef<UE::Insights::FTableColumn> CreateTypeColumn();
+	static TSharedRef<UE::Insights::FTableColumn> CreateDataTypeColumn();
+	static TSharedRef<UE::Insights::FTableColumn> CreateCountColumn();
+	static TSharedRef<UE::Insights::FTableColumn> CreateSumColumn();
+	static TSharedRef<UE::Insights::FTableColumn> CreateMaxColumn();
+	static TSharedRef<UE::Insights::FTableColumn> CreateUpperQuartileColumn();
+	static TSharedRef<UE::Insights::FTableColumn> CreateAverageColumn();
+	static TSharedRef<UE::Insights::FTableColumn> CreateMedianColumn();
+	static TSharedRef<UE::Insights::FTableColumn> CreateLowerQuartileColumn();
+	static TSharedRef<UE::Insights::FTableColumn> CreateMinColumn();
+	static TSharedRef<UE::Insights::FTableColumn> CreateDiffColumn();
 
 private:
 	static constexpr float AggregatedStatsColumnInitialWidth = 80.0f;
 };
-
-////////////////////////////////////////////////////////////////////////////////////////////////////
diff --git a/Engine/Source/Developer/TraceInsights/Private/Insights/ViewModels/ThreadTimingTrack.cpp b/Engine/Source/Developer/TraceInsights/Private/Insights/ViewModels/ThreadTimingTrack.cpp
--- a/Engine/Source/Developer/TraceInsights/Private/Insights/ViewModels/ThreadTimingTrack.cpp
+++ b/Engine/Source/Developer/TraceInsights/Private/Insights/ViewModels/ThreadTimingTrack.cpp
@@ -2,6 +2,7 @@
 
 #include "ThreadTimingTrack.h"
 
+#include "Async/TaskGraphInterfaces.h"
 #include "CborReader.h"
 #include "Framework/Commands/Commands.h"
 #include "Framework/Commands/UICommandList.h"
@@ -9,24 +10,27 @@
 #include "HAL/PlatformApplicationMisc.h"
 #include "Serialization/MemoryReader.h"
 #include "Styling/SlateBrush.h"
+
+// TraceServices
 #include "TraceServices/Model/LoadTimeProfiler.h"
 #include "TraceServices/Model/TasksProfiler.h"
 #include "TraceServices/Model/Threads.h"
 #include "TraceServices/Model/TimingProfiler.h"
-#include "Async/TaskGraphInterfaces.h"
 
-// Insights
-#include "Insights/Common/PaintUtils.h"
-#include "Insights/Common/TimeUtils.h"
+// TraceInsightsCore
+#include "InsightsCore/Common/PaintUtils.h"
+#include "InsightsCore/Common/TimeUtils.h"
+#include "InsightsCore/Filter/ViewModels/FilterConfigurator.h"
+#include "InsightsCore/Filter/ViewModels/Filters.h"
+
+// TraceInsights
 #include "Insights/InsightsManager.h"
 #include "Insights/InsightsStyle.h"
 #include "Insights/ITimingViewSession.h"
 #include "Insights/Log.h"
 #include "Insights/TimingProfilerManager.h"
-#include "Insights/ViewModels/Filters.h"
-#include "Insights/ViewModels/FilterConfigurator.h"
-#include "Insights/ViewModels/TimerNode.h"
 #include "Insights/ViewModels/ThreadTrackEvent.h"
+#include "Insights/ViewModels/TimerNode.h"
 #include "Insights/ViewModels/TimingEventSearch.h"
 #include "Insights/ViewModels/TimingTrackViewport.h"
 #include "Insights/ViewModels/TimingViewDrawHelper.h"
@@ -35,8 +39,6 @@
 
 #define LOCTEXT_NAMESPACE "ThreadTimingTrack"
 
-using namespace Insights;
-
 ////////////////////////////////////////////////////////////////////////////////////////////////////
 // FThreadTimingViewCommands
 ////////////////////////////////////////////////////////////////////////////////////////////////////
@@ -289,6 +291,7 @@ static void AddTimingEventToBuilder(ITimingEventsTrackDrawStateBuilder& Builder,
 									double EventStartTime, double EventEndTime, uint32 EventDepth,
 									uint32 TimerIndex, const TraceServices::FTimingProfilerTimer* Timer)
 {
+	using namespace UE::Insights::TimingProfiler;
 	if (EventDepth >= FTimingProfilerManager::Get()->GetEventDepthLimit())
 	{
 		return;
@@ -297,16 +300,16 @@ static void AddTimingEventToBuilder(ITimingEventsTrackDrawStateBuilder& Builder,
 	uint32 EventColor;
 	switch (FTimingProfilerManager::Get()->GetColoringMode())
 	{
-		case Insights::ETimingEventsColoringMode::ByTimerName:
+		case UE::Insights::ETimingEventsColoringMode::ByTimerName:
 			EventColor = FTimingEvent::ComputeEventColor(Timer->Name);
 			break;
-		case Insights::ETimingEventsColoringMode::ByTimerId:
+		case UE::Insights::ETimingEventsColoringMode::ByTimerId:
 			EventColor = FTimingEvent::ComputeEventColor(Timer->Id);
 			break;
-		case Insights::ETimingEventsColoringMode::BySourceFile:
+		case UE::Insights::ETimingEventsColoringMode::BySourceFile:
 			EventColor = FTimingEvent::ComputeEventColor(Timer->File);
 			break;
-		case Insights::ETimingEventsColoringMode::ByDuration:
+		case UE::Insights::ETimingEventsColoringMode::ByDuration:
 		{
 			const double EventDuration = EventEndTime - EventStartTime;
 			EventColor = (EventDuration >= 0.01)     ? 0xFF883333 : // red:    >= 10ms
@@ -364,7 +367,7 @@ static void AddTimingEventToBuilder(ITimingEventsTrackDrawStateBuilder& Builder,
 // FThreadTimingSharedState
 ////////////////////////////////////////////////////////////////////////////////////////////////////
 
-FThreadTimingSharedState::FThreadTimingSharedState(STimingView* InTimingView)
+FThreadTimingSharedState::FThreadTimingSharedState(UE::Insights::TimingProfiler::STimingView* InTimingView)
 	: TimingView(InTimingView)
 	, bShowHideAllGpuTracks(false)
 	, bShowHideAllCpuTracks(false)
@@ -443,7 +446,7 @@ void FThreadTimingSharedState::GetVisibleTimelineIndexes(TSet<uint32>& OutSet) c
 
 ////////////////////////////////////////////////////////////////////////////////////////////////////
 
-void FThreadTimingSharedState::OnBeginSession(Insights::ITimingViewSession& InSession)
+void FThreadTimingSharedState::OnBeginSession(UE::Insights::Timing::ITimingViewSession& InSession)
 {
 	if (&InSession != TimingView)
 	{
@@ -472,7 +475,7 @@ void FThreadTimingSharedState::OnBeginSession(Insights::ITimingViewSession& InSe
 
 ////////////////////////////////////////////////////////////////////////////////////////////////////
 
-void FThreadTimingSharedState::OnEndSession(Insights::ITimingViewSession& InSession)
+void FThreadTimingSharedState::OnEndSession(UE::Insights::Timing::ITimingViewSession& InSession)
 {
 	if (&InSession != TimingView)
 	{
@@ -493,7 +496,7 @@ void FThreadTimingSharedState::OnEndSession(Insights::ITimingViewSession& InSess
 
 ////////////////////////////////////////////////////////////////////////////////////////////////////
 
-void FThreadTimingSharedState::Tick(Insights::ITimingViewSession& InSession, const TraceServices::IAnalysisSession& InAnalysisSession)
+void FThreadTimingSharedState::Tick(UE::Insights::Timing::ITimingViewSession& InSession, const TraceServices::IAnalysisSession& InAnalysisSession)
 {
 	if (&InSession != TimingView)
 	{
@@ -636,7 +639,7 @@ void FThreadTimingSharedState::Tick(Insights::ITimingViewSession& InSession, con
 
 ////////////////////////////////////////////////////////////////////////////////////////////////////
 
-void FThreadTimingSharedState::ExtendGpuTracksFilterMenu(Insights::ITimingViewSession& InSession, FMenuBuilder& InOutMenuBuilder)
+void FThreadTimingSharedState::ExtendGpuTracksFilterMenu(UE::Insights::Timing::ITimingViewSession& InSession, FMenuBuilder& InOutMenuBuilder)
 {
 	if (&InSession != TimingView)
 	{
@@ -652,7 +655,7 @@ void FThreadTimingSharedState::ExtendGpuTracksFilterMenu(Insights::ITimingViewSe
 
 ////////////////////////////////////////////////////////////////////////////////////////////////////
 
-void FThreadTimingSharedState::ExtendCpuTracksFilterMenu(Insights::ITimingViewSession& InSession, FMenuBuilder& InOutMenuBuilder)
+void FThreadTimingSharedState::ExtendCpuTracksFilterMenu(UE::Insights::Timing::ITimingViewSession& InSession, FMenuBuilder& InOutMenuBuilder)
 {
 	if (&InSession != TimingView)
 	{
@@ -1105,8 +1108,10 @@ void FThreadTimingTrack::BuildFilteredDrawState(ITimingEventsTrackDrawStateBuild
 			TimingProfilerProvider.ReadTimeline(TimelineIndex,
 				[&Viewport, this, &Builder, TimerReader](const TraceServices::ITimingProfilerProvider::Timeline& Timeline)
 				{
+					using EFilterField = UE::Insights::EFilterField;
+
 					TArray<TArray<FPendingEventInfo>> FilteredEvents;
-					TArray<FFilterContext> FilterContexts;
+					TArray<UE::Insights::FFilterContext> FilterContexts;
 
 					TraceServices::ITimeline<TraceServices::FTimingProfilerEvent>::EnumerateAsyncParams Params;
 					constexpr uint32 LargeTimelineThreshold = 50 * 1000 * 1000;
@@ -1126,7 +1131,7 @@ void FThreadTimingTrack::BuildFilteredDrawState(ITimingEventsTrackDrawStateBuild
 					{
 						FilteredEvents.AddDefaulted(NumTasks);
 						FilterContexts.AddDefaulted(NumTasks);
-						for (FFilterContext& Context : FilterContexts)
+						for (UE::Insights::FFilterContext& Context : FilterContexts)
 						{
 							Context.SetReturnValueForUnsetFilters(false);
 							Context.AddFilterData<double>(static_cast<int32>(EFilterField::StartTime), 0.0f);
@@ -1143,7 +1148,7 @@ void FThreadTimingTrack::BuildFilteredDrawState(ITimingEventsTrackDrawStateBuild
 						const TraceServices::FTimingProfilerTimer* Timer = TimerReader->GetTimer(Event.TimerIndex);
 						if (ensure(Timer != nullptr))
 						{
-							FFilterContext& Context = FilterContexts[TaskIndex];
+							UE::Insights::FFilterContext& Context = FilterContexts[TaskIndex];
 							Context.SetFilterData<double>(static_cast<int32>(EFilterField::StartTime), StartTime);
 							Context.SetFilterData<double>(static_cast<int32>(EFilterField::EndTime), EndTime);
 							Context.SetFilterData<double>(static_cast<int32>(EFilterField::Duration), EndTime - StartTime);
@@ -1214,10 +1219,10 @@ void FThreadTimingTrack::PostDraw(const ITimingTrackDrawContext& Context) const
 
 			const double SelectedEventDuration = SelectedEvent.GetDuration();
 			TStringBuilder<1024> StringBuilder;
-			StringBuilder.Appendf(TEXT("  Incl.: %s"), *TimeUtils::FormatTimeAuto(SelectedEventDuration, 2));
+			StringBuilder.Appendf(TEXT("  Incl.: %s"), *UE::Insights::FormatTimeAuto(SelectedEventDuration, 2));
 			if (SelectedEventDuration != std::numeric_limits<double>::infinity())
 			{
-				StringBuilder.Appendf(TEXT("   Excl.: %s"), *TimeUtils::FormatTimeAuto(SelectedEvent.GetExclusiveTime(), 2));
+				StringBuilder.Appendf(TEXT("   Excl.: %s"), *UE::Insights::FormatTimeAuto(SelectedEvent.GetExclusiveTime(), 2));
 			}
 			FString StatsText(StringBuilder.ToView());
 
@@ -1291,7 +1296,7 @@ void FThreadTimingTrack::InitTooltip(FTooltipDrawState& InOutTooltip, const ITim
 			InOutTooltip.AddNameValueTextLine(TEXTVIEW("% of Root:"), ValueStr);
 		}
 
-		InOutTooltip.AddNameValueTextLine(TEXTVIEW("Inclusive Time:"), TimeUtils::FormatTimeAuto(TooltipEventDuration, 2));
+		InOutTooltip.AddNameValueTextLine(TEXTVIEW("Inclusive Time:"), UE::Insights::FormatTimeAuto(TooltipEventDuration, 2));
 
 		if (TooltipEventDuration > 0.0 &&
 			TooltipEventDuration != std::numeric_limits<double>::infinity())
@@ -1299,7 +1304,7 @@ void FThreadTimingTrack::InitTooltip(FTooltipDrawState& InOutTooltip, const ITim
 			const double ExclusiveTimePercent = TooltipEvent.GetExclusiveTime() / TooltipEventDuration;
 			FNumberFormattingOptions FormattingOptions;
 			FormattingOptions.MaximumFractionalDigits = 2;
-			const FString ExclStr = FString::Printf(TEXT("%s (%s)"), *TimeUtils::FormatTimeAuto(TooltipEvent.GetExclusiveTime(), 2), *FText::AsPercent(ExclusiveTimePercent, &FormattingOptions).ToString());
+			const FString ExclStr = FString::Printf(TEXT("%s (%s)"), *UE::Insights::FormatTimeAuto(TooltipEvent.GetExclusiveTime(), 2), *FText::AsPercent(ExclusiveTimePercent, &FormattingOptions).ToString());
 			InOutTooltip.AddNameValueTextLine(TEXTVIEW("Exclusive Time:"), ExclStr);
 		}
 
@@ -1334,7 +1339,7 @@ void FThreadTimingTrack::InitTooltip(FTooltipDrawState& InOutTooltip, const ITim
 
 				auto FormatTaskTime = [](double Time) -> FString
 				{
-					return TimeUtils::FormatTimeAuto(Time, 2);
+					return UE::Insights::FormatTimeAuto(Time, 2);
 				};
 
 				auto GetTrackName = [this](uint32 InThreadId) -> FString
@@ -1464,7 +1469,7 @@ void FThreadTimingTrack::InitTooltip(FTooltipDrawState& InOutTooltip, const ITim
 					InOutTooltip.AddNameValueTextLine(TEXT("Finished waiting:"),
 						FString::Printf(TEXT("%s (+%s)"),
 							*FString::SanitizeFloat(Waiting->FinishedTimestamp),
-							*TimeUtils::FormatTimeAuto(Waiting->FinishedTimestamp - Waiting->StartedTimestamp, 2)));
+							*UE::Insights::FormatTimeAuto(Waiting->FinishedTimestamp - Waiting->StartedTimestamp, 2)));
 				}
 				else
 				{
@@ -1668,6 +1673,8 @@ void FThreadTimingTrack::OnEventSelected(const ITimingEvent& InSelectedEvent) co
 	{
 		const FThreadTrackEvent& TrackEvent = InSelectedEvent.As<FThreadTrackEvent>();
 
+		using namespace UE::Insights::TimingProfiler;
+
 		// Select the timer node corresponding to timing event type of selected timing event.
 		FTimingProfilerManager::Get()->SetSelectedTimer(TrackEvent.GetTimerId());
 	}
@@ -1771,7 +1778,8 @@ bool FThreadTimingTrack::FindTimingProfilerEvent(const FThreadTrackEvent& InTimi
 
 bool FThreadTimingTrack::FindTimingProfilerEvent(const FTimingEventSearchParameters& InParameters, TFunctionRef<void(double, double, uint32, const TraceServices::FTimingProfilerEvent&)> InFoundPredicate) const
 {
-	FFilterContext FilterConfiguratorContext;
+	using EFilterField = UE::Insights::EFilterField;
+	UE::Insights::FFilterContext FilterConfiguratorContext;
 	FilterConfiguratorContext.SetReturnValueForUnsetFilters(false);
 	FilterConfiguratorContext.AddFilterData<double>(static_cast<int32>(EFilterField::StartTime), 0.0f);
 	FilterConfiguratorContext.AddFilterData<double>(static_cast<int32>(EFilterField::EndTime), 0.0f);
@@ -1934,7 +1942,7 @@ int32 FThreadTimingTrack::GetDepthAt(double Time) const
 
 ////////////////////////////////////////////////////////////////////////////////////////////////////
 
-void FThreadTimingTrack::SetFilterConfigurator(TSharedPtr<Insights::FFilterConfigurator> InFilterConfigurator)
+void FThreadTimingTrack::SetFilterConfigurator(TSharedPtr<UE::Insights::FFilterConfigurator> InFilterConfigurator)
 {
 	if (FilterConfigurator != InFilterConfigurator)
 	{
diff --git a/Engine/Source/Developer/TraceInsights/Private/Insights/ViewModels/ThreadTimingTrack.h b/Engine/Source/Developer/TraceInsights/Private/Insights/ViewModels/ThreadTimingTrack.h
--- a/Engine/Source/Developer/TraceInsights/Private/Insights/ViewModels/ThreadTimingTrack.h
+++ b/Engine/Source/Developer/TraceInsights/Private/Insights/ViewModels/ThreadTimingTrack.h
@@ -4,9 +4,11 @@
 
 #include "CoreMinimal.h"
 #include "Framework/Commands/Commands.h"
+
+// TraceServices
 #include "TraceServices/Model/TimingProfiler.h"
 
-// Insights
+// TraceInsights
 #include "Insights/ITimingViewExtender.h"
 #include "Insights/ViewModels/TimingEventSearch.h" // for TTimingEventSearchCache
 #include "Insights/ViewModels/TimingEventsTrack.h"
@@ -15,12 +17,9 @@ class FThreadTrackEvent;
 class FTimingEventSearchParameters;
 class FGpuTimingTrack;
 class FCpuTimingTrack;
-class STimingView;
+namespace UE::Insights::TimingProfiler { class STimingView; }
 
-namespace Insights
-{
-	class FFilterConfigurator;
-}
+namespace UE::Insights { class FFilterConfigurator; }
 
 ////////////////////////////////////////////////////////////////////////////////////////////////////
 
@@ -41,7 +40,7 @@ public:
 
 ////////////////////////////////////////////////////////////////////////////////////////////////////
 
-class FThreadTimingSharedState : public Insights::ITimingViewExtender, public TSharedFromThis<FThreadTimingSharedState>
+class FThreadTimingSharedState : public UE::Insights::Timing::ITimingViewExtender, public TSharedFromThis<FThreadTimingSharedState>
 {
 private:
 	struct FThreadGroup
@@ -55,7 +54,7 @@ private:
 	};
 
 public:
-	explicit FThreadTimingSharedState(STimingView* InTimingView);
+	explicit FThreadTimingSharedState(UE::Insights::TimingProfiler::STimingView* InTimingView);
 	virtual ~FThreadTimingSharedState() = default;
 
 	TSharedPtr<FGpuTimingTrack> GetGpuTrack() { return GpuTrack; }
@@ -72,11 +71,11 @@ public:
 	//////////////////////////////////////////////////
 	// ITimingViewExtender interface
 
-	virtual void OnBeginSession(Insights::ITimingViewSession& InSession) override;
-	virtual void OnEndSession(Insights::ITimingViewSession& InSession) override;
-	virtual void Tick(Insights::ITimingViewSession& InSession, const TraceServices::IAnalysisSession& InAnalysisSession) override;
-	virtual void ExtendGpuTracksFilterMenu(Insights::ITimingViewSession& InSession, FMenuBuilder& InMenuBuilder) override;
-	virtual void ExtendCpuTracksFilterMenu(Insights::ITimingViewSession& InSession, FMenuBuilder& InMenuBuilder) override;
+	virtual void OnBeginSession(UE::Insights::Timing::ITimingViewSession& InSession) override;
+	virtual void OnEndSession(UE::Insights::Timing::ITimingViewSession& InSession) override;
+	virtual void Tick(UE::Insights::Timing::ITimingViewSession& InSession, const TraceServices::IAnalysisSession& InAnalysisSession) override;
+	virtual void ExtendGpuTracksFilterMenu(UE::Insights::Timing::ITimingViewSession& InSession, FMenuBuilder& InMenuBuilder) override;
+	virtual void ExtendCpuTracksFilterMenu(UE::Insights::Timing::ITimingViewSession& InSession, FMenuBuilder& InMenuBuilder) override;
 
 	//////////////////////////////////////////////////
 
@@ -104,7 +103,7 @@ private:
 	void ToggleTrackVisibilityByGroup_Execute(const TCHAR* InGroupName);
 
 private:
-	STimingView* TimingView;
+	UE::Insights::TimingProfiler::STimingView* TimingView;
 
 	bool bShowHideAllGpuTracks;
 	bool bShowHideAllCpuTracks;
@@ -175,7 +174,7 @@ public:
 
 	int32 GetDepthAt(double Time) const;
 
-	virtual void SetFilterConfigurator(TSharedPtr<Insights::FFilterConfigurator> InFilterConfigurator) override;
+	virtual void SetFilterConfigurator(TSharedPtr<UE::Insights::FFilterConfigurator> InFilterConfigurator) override;
 
 	TSharedPtr<const ITimingEvent> FindMaxEventInstance(uint32 TimerId, double StartTime, double EndTime) const;
 	TSharedPtr<const ITimingEvent> FindMinEventInstance(uint32 TimerId, double StartTime, double EndTime) const;
@@ -201,7 +200,7 @@ private:
 
 	FThreadTimingSharedState& SharedState;
 
-	TSharedPtr<Insights::FFilterConfigurator> FilterConfigurator;
+	TSharedPtr<UE::Insights::FFilterConfigurator> FilterConfigurator;
 
 	// Search cache
 	mutable TTimingEventSearchCache<TraceServices::FTimingProfilerEvent> SearchCache;
