[Insights] Massive refactoring of TraceInsights module:

diff --git a/Engine/Source/Developer/TraceInsights/Private/Insights/ViewModels/TimeRulerTrack.cpp b/Engine/Source/Developer/TraceInsights/Private/Insights/ViewModels/TimeRulerTrack.cpp
--- a/Engine/Source/Developer/TraceInsights/Private/Insights/ViewModels/TimeRulerTrack.cpp
+++ b/Engine/Source/Developer/TraceInsights/Private/Insights/ViewModels/TimeRulerTrack.cpp
@@ -9,9 +9,11 @@
 #include "Rendering/DrawElements.h"
 #include "Styling/AppStyle.h"
 
-// Insights
-#include "Insights/Common/PaintUtils.h"
-#include "Insights/Common/TimeUtils.h"
+// TraceInsightsCore
+#include "InsightsCore/Common/PaintUtils.h"
+#include "InsightsCore/Common/TimeUtils.h"
+
+// TraceInsights
 #include "Insights/InsightsStyle.h"
 #include "Insights/TimingProfilerManager.h"
 #include "Insights/ViewModels/DrawHelpers.h"
@@ -21,7 +23,10 @@
 
 #include <limits>
 
-#define LOCTEXT_NAMESPACE "TimeRulerTrack"
+#define LOCTEXT_NAMESPACE "UE::Insights::TimingProfiler::FTimeRulerTrack"
+
+namespace UE::Insights::TimingProfiler
+{
 
 ////////////////////////////////////////////////////////////////////////////////////////////////////
 
@@ -74,14 +79,14 @@ void FTimeRulerTrack::SetSelection(const bool bInIsSelecting, const double InSel
 
 ////////////////////////////////////////////////////////////////////////////////////////////////////
 
-void FTimeRulerTrack::AddTimeMarker(TSharedRef<Insights::FTimeMarker> InTimeMarker)
+void FTimeRulerTrack::AddTimeMarker(TSharedRef<FTimeMarker> InTimeMarker)
 {
 	TimeMarkers.Add(InTimeMarker);
 }
 
 ////////////////////////////////////////////////////////////////////////////////////////////////////
 
-void FTimeRulerTrack::RemoveTimeMarker(TSharedRef<Insights::FTimeMarker> InTimeMarker)
+void FTimeRulerTrack::RemoveTimeMarker(TSharedRef<FTimeMarker> InTimeMarker)
 {
 	TimeMarkers.Remove(InTimeMarker);
 }
@@ -95,9 +100,9 @@ void FTimeRulerTrack::RemoveAllTimeMarkers()
 
 ////////////////////////////////////////////////////////////////////////////////////////////////////
 
-TSharedPtr<Insights::FTimeMarker> FTimeRulerTrack::GetTimeMarkerByName(const FString& InTimeMarkerName)
+TSharedPtr<FTimeMarker> FTimeRulerTrack::GetTimeMarkerByName(const FString& InTimeMarkerName)
 {
-	for (TSharedRef<Insights::FTimeMarker>& TimeMarker : TimeMarkers)
+	for (TSharedRef<FTimeMarker>& TimeMarker : TimeMarkers)
 	{
 		if (TimeMarker->GetName().Equals(InTimeMarkerName))
 		{
@@ -109,9 +114,9 @@ TSharedPtr<Insights::FTimeMarker> FTimeRulerTrack::GetTimeMarkerByName(const FSt
 
 ////////////////////////////////////////////////////////////////////////////////////////////////////
 
-TSharedPtr<Insights::FTimeMarker> FTimeRulerTrack::GetTimeMarkerAtPos(const FVector2D& InPosition, const FTimingTrackViewport& InViewport)
+TSharedPtr<FTimeMarker> FTimeRulerTrack::GetTimeMarkerAtPos(const FVector2D& InPosition, const FTimingTrackViewport& InViewport)
 {
-	TSharedPtr<Insights::FTimeMarker> ClosestTimeMarker = nullptr;
+	TSharedPtr<FTimeMarker> ClosestTimeMarker = nullptr;
 
 	constexpr float TimeMarkerBoxHeight = 12.0f;
 	const float InPositionY = static_cast<float>(InPosition.Y);
@@ -119,7 +124,7 @@ TSharedPtr<Insights::FTimeMarker> FTimeRulerTrack::GetTimeMarkerAtPos(const FVec
 	{
 		const float InPositionX = static_cast<float>(InPosition.X);
 		float MinDX = 42.0f;
-		for (TSharedRef<Insights::FTimeMarker>& TimeMarker : TimeMarkers)
+		for (TSharedRef<FTimeMarker>& TimeMarker : TimeMarkers)
 		{
 			if (TimeMarker->IsVisible())
 			{
@@ -139,7 +144,7 @@ TSharedPtr<Insights::FTimeMarker> FTimeRulerTrack::GetTimeMarkerAtPos(const FVec
 
 ////////////////////////////////////////////////////////////////////////////////////////////////////
 
-void FTimeRulerTrack::StartScrubbing(TSharedRef<Insights::FTimeMarker> InTimeMarker)
+void FTimeRulerTrack::StartScrubbing(TSharedRef<FTimeMarker> InTimeMarker)
 {
 	// Move the scrubbing time marker at the end of sorting list to be draw on top of other markers.
 	TimeMarkers.Remove(InTimeMarker);
@@ -220,8 +225,9 @@ void FTimeRulerTrack::Draw(const ITimingTrackDrawContext& Context) const
 	}
 	DrawContext.LayerId++;
 
+	using namespace UE::Insights;
 	const double DT = static_cast<double>(MajorTickMark) / Viewport.GetScaleX();
-	const double Precision = FMath::Max(DT / 10.0, TimeUtils::Nanosecond);
+	const double Precision = FMath::Max(DT / 10.0, FTimeValue::Nanosecond);
 
 	const TSharedRef<FSlateFontMeasure> FontMeasureService = FSlateApplication::Get().GetRenderer()->GetFontMeasureService();
 	const float FontScale = DrawContext.Geometry.Scale;
@@ -230,7 +236,7 @@ void FTimeRulerTrack::Draw(const ITimingTrackDrawContext& Context) const
 	for (float X = MajorOX; X < Viewport.GetWidth() + MajorTickMark; X += MajorTickMark)
 	{
 		const double T = Viewport.SlateUnitsToTime(X);
-		FString Text = TimeUtils::FormatTime(T, Precision);
+		FString Text = FormatTime(T, Precision);
 		const float TextWidth = static_cast<float>(FontMeasureService->Measure(Text, Font, FontScale).X / FontScale);
 		DrawContext.DrawText(X - TextWidth / 2, TextY, Text, Font,
 			(T < Viewport.GetMinValidTime() || T >= Viewport.GetMaxValidTime()) ? FLinearColor(0.7f, 0.5f, 0.5f, 1.0f) : FLinearColor(0.8f, 0.8f, 0.8f, 1.0f));
@@ -260,17 +266,18 @@ void FTimeRulerTrack::PostDraw(const ITimingTrackDrawContext& Context) const
 		const double MousePosTime = Viewport.SlateUnitsToTime(static_cast<float>(MousePosition.X));
 		CrtMousePosTime = MousePosTime;
 
+		using namespace UE::Insights;
 		const double DT = 100.0 / Viewport.GetScaleX();
-		const double MousePosPrecision = FMath::Max(DT / 100.0, TimeUtils::Nanosecond);
+		const double MousePosPrecision = FMath::Max(DT / 100.0, FTimeValue::Nanosecond);
 		if (bIsMouseOver)
 		{
 			// If mouse is hovering the time ruler, format time with a better precision (split seconds in ms, us, ns and ps).
-			MousePosText = TimeUtils::FormatTimeSplit(MousePosTime, MousePosPrecision);
+			MousePosText = FormatTimeSplit(MousePosTime, MousePosPrecision);
 		}
 		else
 		{
 			// Format current time with one more digit than the time at major tick marks.
-			MousePosText = TimeUtils::FormatTime(MousePosTime, MousePosPrecision);
+			MousePosText = FormatTime(MousePosTime, MousePosPrecision);
 		}
 
 		const TSharedRef<FSlateFontMeasure> FontMeasureService = FSlateApplication::Get().GetRenderer()->GetFontMeasureService();
@@ -325,7 +332,7 @@ void FTimeRulerTrack::PostDraw(const ITimingTrackDrawContext& Context) const
 	}
 
 	// Draw the time markers.
-	for (const TSharedRef<Insights::FTimeMarker>& TimeMarker : TimeMarkers)
+	for (const TSharedRef<FTimeMarker>& TimeMarker : TimeMarkers)
 	{
 		DrawTimeMarker(Context, *TimeMarker);
 	}
@@ -333,7 +340,7 @@ void FTimeRulerTrack::PostDraw(const ITimingTrackDrawContext& Context) const
 
 ////////////////////////////////////////////////////////////////////////////////////////////////////
 
-void FTimeRulerTrack::DrawTimeMarker(const ITimingTrackDrawContext& Context, const Insights::FTimeMarker& TimeMarker) const
+void FTimeRulerTrack::DrawTimeMarker(const ITimingTrackDrawContext& Context, const FTimeMarker& TimeMarker) const
 {
 	if (!TimeMarker.IsVisible())
 	{
@@ -382,17 +389,18 @@ void FTimeRulerTrack::DrawTimeMarker(const ITimingTrackDrawContext& Context, con
 		}
 
 		// Format time value with one more digit than the time at major tick marks.
+		using namespace UE::Insights;
 		const double DT = 100.0 / Viewport.GetScaleX();
-		const double Precision = FMath::Max(DT / 100.0, TimeUtils::Nanosecond);
+		const double Precision = FMath::Max(DT / 100.0, FTimeValue::Nanosecond);
 
 		if (bIsMouseOverMarker)
 		{
 			// If mouse is hovering the time marker, format time with a better precision (split seconds in ms, us, ns and ps).
-			TimeMarkerText += TimeUtils::FormatTimeSplit(TimeMarker.GetTime(), Precision);
+			TimeMarkerText += FormatTimeSplit(TimeMarker.GetTime(), Precision);
 		}
 		else
 		{
-			TimeMarkerText += TimeUtils::FormatTime(TimeMarker.GetTime(), Precision);
+			TimeMarkerText += FormatTime(TimeMarker.GetTime(), Precision);
 		}
 	}
 	else
@@ -435,8 +443,8 @@ void FTimeRulerTrack::DrawTimeMarker(const ITimingTrackDrawContext& Context, con
 
 void FTimeRulerTrack::BuildContextMenu(FMenuBuilder& MenuBuilder)
 {
-	TArray<TSharedRef<Insights::FTimeMarker>> VisibleTimeMarkers;
-	for (TSharedRef<Insights::FTimeMarker>& TimeMarker : TimeMarkers)
+	TArray<TSharedRef<FTimeMarker>> VisibleTimeMarkers;
+	for (TSharedRef<FTimeMarker>& TimeMarker : TimeMarkers)
 	{
 		if (TimeMarker->IsVisible())
 		{
@@ -444,16 +452,16 @@ void FTimeRulerTrack::BuildContextMenu(FMenuBuilder& MenuBuilder)
 		}
 	}
 
-	FText CrtMousePosTimeText = FText::FromString(TimeUtils::FormatTimeAuto(CrtMousePosTime, 2));
+	FText CrtMousePosTimeText = FText::FromString(FormatTimeAuto(CrtMousePosTime, 2));
 
 	if (VisibleTimeMarkers.Num() > 0)
 	{
 		// Sort TimeMarkers by name.
-		VisibleTimeMarkers.Sort([](const TSharedRef<Insights::FTimeMarker>& A, const TSharedRef<Insights::FTimeMarker>& B) -> bool { return A->GetName().Compare(B->GetName()) <= 0; });
+		VisibleTimeMarkers.Sort([](const TSharedRef<FTimeMarker>& A, const TSharedRef<FTimeMarker>& B) -> bool { return A->GetName().Compare(B->GetName()) <= 0; });
 
 		MenuBuilder.BeginSection("TimeMarkers", LOCTEXT("ContextMenu_Section_TimeMarkers", "Time Markers"));
 
-		for (TSharedRef<Insights::FTimeMarker>& TimeMarker : VisibleTimeMarkers)
+		for (TSharedRef<FTimeMarker>& TimeMarker : VisibleTimeMarkers)
 		{
 			FUIAction Action_MoveTimeMarker
 			(
@@ -499,7 +507,7 @@ void FTimeRulerTrack::BuildContextMenu(FMenuBuilder& MenuBuilder)
 
 ////////////////////////////////////////////////////////////////////////////////////////////////////
 
-void FTimeRulerTrack::ContextMenu_MoveTimeMarker_Execute(TSharedRef<Insights::FTimeMarker> InTimeMarker)
+void FTimeRulerTrack::ContextMenu_MoveTimeMarker_Execute(TSharedRef<FTimeMarker> InTimeMarker)
 {
 	InTimeMarker->SetTime(CrtMousePosTime);
 }
@@ -540,4 +548,6 @@ void FTimeRulerTrack::ContextMenu_ScrollLogView_Execute()
 
 ////////////////////////////////////////////////////////////////////////////////////////////////////
 
+} // namespace UE::Insights::TimingProfiler
+
 #undef LOCTEXT_NAMESPACE
diff --git a/Engine/Source/Developer/TraceInsights/Private/Insights/ViewModels/TimeRulerTrack.h b/Engine/Source/Developer/TraceInsights/Private/Insights/ViewModels/TimeRulerTrack.h
--- a/Engine/Source/Developer/TraceInsights/Private/Insights/ViewModels/TimeRulerTrack.h
+++ b/Engine/Source/Developer/TraceInsights/Private/Insights/ViewModels/TimeRulerTrack.h
@@ -3,23 +3,26 @@
 #pragma once
 
 #include "CoreMinimal.h"
+
 #include "Fonts/SlateFontInfo.h"
 
-// Insights
+// TraceInsights
 #include "Insights/ViewModels/BaseTimingTrack.h"
 #include "Insights/ITimingViewSession.h"
 
-struct FDrawContext;
 class FMenuBuilder;
 struct FSlateBrush;
+
 class FTimingTrackViewport;
 
-namespace Insights
-{
+namespace Insights { class ITimingViewExtender; }
 
-class ITimingViewExtender;
+namespace UE::Insights { class FDrawContext; }
 
-class FTimeMarker : public  ITimeMarker
+namespace UE::Insights::TimingProfiler
+{
+
+class FTimeMarker : public UE::Insights::Timing::ITimeMarker
 {
 public:
 	FTimeMarker()
@@ -62,8 +65,6 @@ private:
 	mutable float CrtTextWidth;
 };
 
-} // namespace Insights
-
 ////////////////////////////////////////////////////////////////////////////////////////////////////
 
 class FTimeRulerTrack : public FBaseTimingTrack
@@ -78,18 +79,18 @@ public:
 
 	void SetSelection(const bool bInIsSelecting, const double InSelectionStartTime, const double InSelectionEndTime);
 
-	TArray<TSharedRef<Insights::FTimeMarker>>& GetTimeMarkers() { return TimeMarkers; }
-	const TArray<TSharedRef<Insights::FTimeMarker>>& GetTimeMarkers() const { return TimeMarkers; }
-	void AddTimeMarker(TSharedRef<Insights::FTimeMarker> InTimeMarker);
-	void RemoveTimeMarker(TSharedRef<Insights::FTimeMarker> InTimeMarker);
+	TArray<TSharedRef<FTimeMarker>>& GetTimeMarkers() { return TimeMarkers; }
+	const TArray<TSharedRef<FTimeMarker>>& GetTimeMarkers() const { return TimeMarkers; }
+	void AddTimeMarker(TSharedRef<FTimeMarker> InTimeMarker);
+	void RemoveTimeMarker(TSharedRef<FTimeMarker> InTimeMarker);
 	void RemoveAllTimeMarkers();
 
-	TSharedPtr<Insights::FTimeMarker> GetTimeMarkerByName(const FString& InTimeMarkerName);
-	TSharedPtr<Insights::FTimeMarker> GetTimeMarkerAtPos(const FVector2D& InPosition, const FTimingTrackViewport& InViewport);
+	TSharedPtr<FTimeMarker> GetTimeMarkerByName(const FString& InTimeMarkerName);
+	TSharedPtr<FTimeMarker> GetTimeMarkerAtPos(const FVector2D& InPosition, const FTimingTrackViewport& InViewport);
 
 	bool IsScrubbing() const { return bIsScrubbing; }
-	TSharedRef<Insights::FTimeMarker> GetScrubbingTimeMarker() { return TimeMarkers.Last(); }
-	void StartScrubbing(TSharedRef<Insights::FTimeMarker> InTimeMarker);
+	TSharedRef<FTimeMarker> GetScrubbingTimeMarker() { return TimeMarkers.Last(); }
+	void StartScrubbing(TSharedRef<FTimeMarker> InTimeMarker);
 	void StopScrubbing();
 
 	virtual void PostUpdate(const ITimingTrackUpdateContext& Context) override;
@@ -99,8 +100,8 @@ public:
 	virtual void BuildContextMenu(FMenuBuilder& MenuBuilder) override;
 
 private:
-	void DrawTimeMarker(const ITimingTrackDrawContext& Context, const Insights::FTimeMarker& TimeMarker) const;
-	void ContextMenu_MoveTimeMarker_Execute(TSharedRef<Insights::FTimeMarker> InTimeMarker);
+	void DrawTimeMarker(const ITimingTrackDrawContext& Context, const FTimeMarker& TimeMarker) const;
+	void ContextMenu_MoveTimeMarker_Execute(TSharedRef<FTimeMarker> InTimeMarker);
 	bool ContextMenu_ScrollLogView_CanExecute();
 	void ContextMenu_ScrollLogView_Execute();
 
@@ -113,7 +114,7 @@ private:
 	double SelectionStartTime;
 	double SelectionEndTime;
 
-	// The last time value at mouse postion. Updated in PostDraw.
+	// The last time value at mouse position. Updated in PostDraw.
 	mutable double CrtMousePosTime;
 
 	// The smoothed width of "the text at mouse position" to avoid flickering. Updated in PostDraw.
@@ -123,10 +124,12 @@ private:
 	 * The sorted list of all registered time markers. It defines the draw order of time markers.
 	 * The time marker currently scrubbing will be moved at the end of the list in order to be displayed on top of other markers.
 	 */
-	TArray<TSharedRef<Insights::FTimeMarker>> TimeMarkers;
+	TArray<TSharedRef<FTimeMarker>> TimeMarkers;
 
 	/** True if the user is currently dragging a time marker. */
 	bool bIsScrubbing;
 };
 
 ////////////////////////////////////////////////////////////////////////////////////////////////////
+
+} // namespace UE::Insights::TimingProfiler
diff --git a/Engine/Source/Developer/TraceInsights/Private/Insights/ViewModels/TimerAggregation.cpp b/Engine/Source/Developer/TraceInsights/Private/Insights/ViewModels/TimerAggregation.cpp
--- a/Engine/Source/Developer/TraceInsights/Private/Insights/ViewModels/TimerAggregation.cpp
+++ b/Engine/Source/Developer/TraceInsights/Private/Insights/ViewModels/TimerAggregation.cpp
@@ -83,6 +83,8 @@ IStatsAggregationWorker* FTimerAggregator::CreateWorker(TSharedPtr<const TraceSe
 	bool bIsGpuTrackVisible = false;
 	TSet<uint32> CpuThreads;
 
+	using namespace UE::Insights::TimingProfiler;
+
 	TSharedPtr<STimingProfilerWindow> Wnd = FTimingProfilerManager::Get()->GetProfilerWindow();
 	if (Wnd.IsValid())
 	{
diff --git a/Engine/Source/Developer/TraceInsights/Private/Insights/ViewModels/TimerAggregation.h b/Engine/Source/Developer/TraceInsights/Private/Insights/ViewModels/TimerAggregation.h
--- a/Engine/Source/Developer/TraceInsights/Private/Insights/ViewModels/TimerAggregation.h
+++ b/Engine/Source/Developer/TraceInsights/Private/Insights/ViewModels/TimerAggregation.h
@@ -4,10 +4,13 @@
 
 #include "CoreMinimal.h"
 
-#include "Insights/ViewModels/StatsAggregator.h"
+// TraceServices
 #include "TraceServices/Containers/Tables.h"
 #include "TraceServices/Model/TimingProfiler.h"
 
+// TraceInsights
+#include "Insights/ViewModels/StatsAggregator.h"
+
 ////////////////////////////////////////////////////////////////////////////////////////////////////
 
 namespace Insights
diff --git a/Engine/Source/Developer/TraceInsights/Private/Insights/ViewModels/TimerButterflyAggregation.cpp b/Engine/Source/Developer/TraceInsights/Private/Insights/ViewModels/TimerButterflyAggregation.cpp
--- a/Engine/Source/Developer/TraceInsights/Private/Insights/ViewModels/TimerButterflyAggregation.cpp
+++ b/Engine/Source/Developer/TraceInsights/Private/Insights/ViewModels/TimerButterflyAggregation.cpp
@@ -2,19 +2,20 @@
 
 #include "TimerButterflyAggregation.h"
 
+// TraceInsights
 #include "Insights/TimingProfilerManager.h"
 #include "Insights/ViewModels/ThreadTimingTrack.h"
 #include "Insights/Widgets/STimingProfilerWindow.h"
 #include "Insights/Widgets/STimingView.h"
 
-namespace Insights
+namespace UE::Insights::TimingProfiler
 {
 
 ////////////////////////////////////////////////////////////////////////////////////////////////////
 // FTimerButterflyAggregationWorker
 ////////////////////////////////////////////////////////////////////////////////////////////////////
 
-class FTimerButterflyAggregationWorker : public IStatsAggregationWorker
+class FTimerButterflyAggregationWorker : public ::Insights::IStatsAggregationWorker
 {
 public:
 	FTimerButterflyAggregationWorker(TSharedPtr<const TraceServices::IAnalysisSession> InSession, double InStartTime, double InEndTime, const TSet<uint32>& InCpuThreads, bool bInIncludeGpuThread)
@@ -82,7 +83,7 @@ void FTimerButterflyAggregationWorker::DoWork(TSharedPtr<TraceServices::FCancell
 // FTimerButterflyAggregator
 ////////////////////////////////////////////////////////////////////////////////////////////////////
 
-IStatsAggregationWorker* FTimerButterflyAggregator::CreateWorker(TSharedPtr<const TraceServices::IAnalysisSession> InSession)
+::Insights::IStatsAggregationWorker* FTimerButterflyAggregator::CreateWorker(TSharedPtr<const TraceServices::IAnalysisSession> InSession)
 {
 	bool bIsGpuTrackVisible = false;
 	TSet<uint32> CpuThreads;
diff --git a/Engine/Source/Developer/TraceInsights/Private/Insights/ViewModels/TimerButterflyAggregation.h b/Engine/Source/Developer/TraceInsights/Private/Insights/ViewModels/TimerButterflyAggregation.h
--- a/Engine/Source/Developer/TraceInsights/Private/Insights/ViewModels/TimerButterflyAggregation.h
+++ b/Engine/Source/Developer/TraceInsights/Private/Insights/ViewModels/TimerButterflyAggregation.h
@@ -4,29 +4,26 @@
 
 #include "CoreMinimal.h"
 
-#include "Insights/ViewModels/StatsAggregator.h"
+// TraceServices
 #include "TraceServices/Model/TimingProfiler.h"
 
-////////////////////////////////////////////////////////////////////////////////////////////////////
+// TraceInsights
+#include "Insights/ViewModels/StatsAggregator.h"
 
-namespace Insights
+namespace UE::Insights::TimingProfiler
 {
 
-////////////////////////////////////////////////////////////////////////////////////////////////////
-
-class FTimerButterflyAggregator : public FStatsAggregator
+class FTimerButterflyAggregator : public ::Insights::FStatsAggregator
 {
 public:
-	FTimerButterflyAggregator() : FStatsAggregator(TEXT("Butterfly")) {}
+	FTimerButterflyAggregator() : ::Insights::FStatsAggregator(TEXT("Butterfly")) {}
 	virtual ~FTimerButterflyAggregator() {}
 
 	TraceServices::ITimingProfilerButterfly* GetResultButterfly() const;
 	void ResetResults();
 
 protected:
-	virtual IStatsAggregationWorker* CreateWorker(TSharedPtr<const TraceServices::IAnalysisSession> InSession) override;
+	virtual ::Insights::IStatsAggregationWorker* CreateWorker(TSharedPtr<const TraceServices::IAnalysisSession> InSession) override;
 };
 
-////////////////////////////////////////////////////////////////////////////////////////////////////
-
-} // namespace Insights
+} // namespace UE::Insights::TimingProfiler
diff --git a/Engine/Source/Developer/TraceInsights/Private/Insights/ViewModels/TimerFilters.cpp b/Engine/Source/Developer/TraceInsights/Private/Insights/ViewModels/TimerFilters.cpp
--- a/Engine/Source/Developer/TraceInsights/Private/Insights/ViewModels/TimerFilters.cpp
+++ b/Engine/Source/Developer/TraceInsights/Private/Insights/ViewModels/TimerFilters.cpp
@@ -45,14 +45,14 @@ void FTimerNameFilterState::Update()
 			const TraceServices::FTimingProfilerTimer* Timer = TimerReader->GetTimer(TimerIndex);
 			if (Timer && Timer->Name)
 			{
-				if (SelectedOperator->GetKey() == EFilterOperator::Eq)
+				if (SelectedOperator->GetKey() == UE::Insights::EFilterOperator::Eq)
 				{
 					if (FCString::Stricmp(Timer->Name, *FilterValue) == 0)
 					{
 						TimerIds.Add(Timer->Id);
 					}
 				}
-				else if (SelectedOperator->GetKey() == EFilterOperator::Contains)
+				else if (SelectedOperator->GetKey() == UE::Insights::EFilterOperator::Contains)
 				{
 					if (FCString::Stristr(Timer->Name, *FilterValue))
 					{
@@ -66,7 +66,7 @@ void FTimerNameFilterState::Update()
 
 ////////////////////////////////////////////////////////////////////////////////////////////////////
 
-bool FTimerNameFilterState::ApplyFilter(const FFilterContext& Context) const
+bool FTimerNameFilterState::ApplyFilter(const UE::Insights::FFilterContext& Context) const
 {
 	if (!Context.HasFilterData(static_cast<int32>(Filter->GetKey())))
 	{
@@ -94,7 +94,7 @@ bool FTimerNameFilterState::Equals(const FFilterState& Other) const
 
 ////////////////////////////////////////////////////////////////////////////////////////////////////
 
-TSharedRef<FFilterState> FTimerNameFilterState::DeepCopy() const
+TSharedRef<UE::Insights::FFilterState> FTimerNameFilterState::DeepCopy() const
 {
 	TSharedRef<FTimerNameFilterState> Copy = MakeShared<FTimerNameFilterState>(*this);
 
@@ -106,16 +106,18 @@ TSharedRef<FFilterState> FTimerNameFilterState::DeepCopy() const
 ////////////////////////////////////////////////////////////////////////////////////////////////////
 
 FTimerNameFilter::FTimerNameFilter()
-	: FCustomFilter(static_cast<int32>(EFilterField::TimerName),
+	: FCustomFilter(static_cast<int32>(UE::Insights::EFilterField::TimerName),
 		LOCTEXT("TimerName", "Timer Name"),
 		LOCTEXT("TimerName", "Timer Name"),
-		EFilterDataType::Custom,
+		UE::Insights::EFilterDataType::Custom,
 		nullptr,
 		nullptr)
 {
-	SupportedOperators = MakeShared<TArray<TSharedPtr<IFilterOperator>>>();
-	SupportedOperators->Add(StaticCastSharedRef<IFilterOperator>(MakeShared<FFilterOperator<int64>>(EFilterOperator::Eq, TEXT("Is"), [](int64 lhs, int64 rhs) { return lhs == rhs; })));
-	SupportedOperators->Add(StaticCastSharedRef<IFilterOperator>(MakeShared<FFilterOperator<int64>>(EFilterOperator::Contains, TEXT("Contains"), [](int64 lhs, int64 rhs) { return lhs == rhs; })));
+	SupportedOperators = MakeShared<TArray<TSharedPtr<UE::Insights::IFilterOperator>>>();
+	SupportedOperators->Add(StaticCastSharedRef<UE::Insights::IFilterOperator>(MakeShared<UE::Insights::FFilterOperator<int64>>(
+		UE::Insights::EFilterOperator::Eq, TEXT("Is"), [](int64 lhs, int64 rhs) { return lhs == rhs; })));
+	SupportedOperators->Add(StaticCastSharedRef<UE::Insights::IFilterOperator>(MakeShared<UE::Insights::FFilterOperator<int64>>(
+		UE::Insights::EFilterOperator::Contains, TEXT("Contains"), [](int64 lhs, int64 rhs) { return lhs == rhs; })));
 
 	SetCallback([this](const FString& Text, TArray<FString>& OutSuggestions)
 		{
@@ -164,7 +166,7 @@ void FTimerNameFilter::PopulateTimerNameSuggestionList(const FString& Text, TArr
 
 INSIGHTS_IMPLEMENT_RTTI(FMetadataFilterState)
 
-FMetadataFilterState::FMetadataFilterState(TSharedRef<FFilter> InFilter)
+FMetadataFilterState::FMetadataFilterState(TSharedRef<UE::Insights::FFilter> InFilter)
 	: FFilterState(InFilter)
 {
 	AvailableDataTypes.Add(MakeShared<FMetadataFilterDataTypeEntry>(EMetadataFilterDataType::Int, LOCTEXT("IntDataType", "Int")));
@@ -175,7 +177,7 @@ FMetadataFilterState::FMetadataFilterState(TSharedRef<FFilter> InFilter)
 	SelectedDataType = AvailableDataTypes[0];
 	DataType_OnSelectionChanged(SelectedDataType, ESelectInfo::Type::Direct);
 
-	BoolOperators.Add((MakeShared<FFilterOperator<bool>>(EFilterOperator::Eq, TEXT("IS"), [](bool lhs, bool rhs) { return lhs == rhs; })));
+	BoolOperators.Add((MakeShared<UE::Insights::FFilterOperator<bool>>(UE::Insights::EFilterOperator::Eq, TEXT("IS"), [](bool lhs, bool rhs) { return lhs == rhs; })));
 }
 
 ////////////////////////////////////////////////////////////////////////////////////////////////////
@@ -219,7 +221,7 @@ void FMetadataFilterState::AddCustomUI(TSharedRef<SHorizontalBox> Box)
 			.Padding(FMargin(4.0f, 0.0f, 0.0f, 0.0f))
 			.AutoWidth()
 			[
-				SAssignNew(OperatorComboBox, SComboBox<TSharedPtr<IFilterOperator>>)
+				SAssignNew(OperatorComboBox, SComboBox<TSharedPtr<UE::Insights::IFilterOperator>>)
 				.OptionsSource(&AvailableOperators)
 				.OnSelectionChanged(this, &FMetadataFilterState::AvailableOperators_OnSelectionChanged)
 				.OnGenerateWidget(this, &FMetadataFilterState::AvailableOperators_OnGenerateWidget)
@@ -292,7 +294,7 @@ void FMetadataFilterState::Update()
 
 ////////////////////////////////////////////////////////////////////////////////////////////////////
 
-bool FMetadataFilterState::ApplyFilter(const FFilterContext& Context) const
+bool FMetadataFilterState::ApplyFilter(const UE::Insights::FFilterContext& Context) const
 {
 	if (!Context.HasFilterData(static_cast<int32>(Filter->GetKey())))
 	{
@@ -370,7 +372,8 @@ bool FMetadataFilterState::ApplyFilterToMetadata(TArrayView<const uint8>& Metada
 			int64 MetadataValue = static_cast<int64>(Context.AsUInt());
 			int64 InputValue = ConvertedData.Get<int64>();
 
-			FFilterOperator<int64>* Operator = (FFilterOperator<int64>*) SelectedOperator.Get();
+			using FFilterOperatorInt64 = UE::Insights::FFilterOperator<int64>;
+			FFilterOperatorInt64* Operator = (FFilterOperatorInt64*)SelectedOperator.Get();
 			if (Operator->Apply(MetadataValue, InputValue))
 			{
 				return true;
@@ -387,7 +390,9 @@ bool FMetadataFilterState::ApplyFilterToMetadata(TArrayView<const uint8>& Metada
 			}
 
 			FString Value = Context.AsString();
-			FFilterOperator<FString>* Operator = (FFilterOperator<FString>*) SelectedOperator.Get();
+
+			using FFilterOperatorString = UE::Insights::FFilterOperator<FString>;
+			FFilterOperatorString* Operator = (FFilterOperatorString*)SelectedOperator.Get();
 			if (Operator->Apply(Value, Term))
 			{
 				return true;
@@ -406,7 +411,8 @@ bool FMetadataFilterState::ApplyFilterToMetadata(TArrayView<const uint8>& Metada
 			FAnsiStringView Value(Context.AsCString(), static_cast<int32>(Context.AsLength()));
 			FString ValueStr(Value);
 
-			FFilterOperator<FString>* Operator = (FFilterOperator<FString>*) SelectedOperator.Get();
+			using FFilterOperatorString = UE::Insights::FFilterOperator<FString>;
+			FFilterOperatorString* Operator = (FFilterOperatorString*)SelectedOperator.Get();
 			if (Operator->Apply(ValueStr, Term))
 			{
 				return true;
@@ -427,7 +433,8 @@ bool FMetadataFilterState::ApplyFilterToMetadata(TArrayView<const uint8>& Metada
 			double Value = static_cast<double>(Context.AsFloat());
 			double InputValue = ConvertedData.Get<double>();
 
-			FFilterOperator<double>* Operator = (FFilterOperator<double>*) SelectedOperator.Get();
+			using FFilterOperatorDouble = UE::Insights::FFilterOperator<double>;
+			FFilterOperatorDouble* Operator = (FFilterOperatorDouble*)SelectedOperator.Get();
 			if (Operator->Apply(Value, InputValue))
 			{
 				return true;
@@ -447,7 +454,8 @@ bool FMetadataFilterState::ApplyFilterToMetadata(TArrayView<const uint8>& Metada
 			double Value = Context.AsDouble();
 			double InputValue = ConvertedData.Get<double>();
 
-			FFilterOperator<double>* Operator = (FFilterOperator<double>*) SelectedOperator.Get();
+			using FFilterOperatorDouble = UE::Insights::FFilterOperator<double>;
+			FFilterOperatorDouble* Operator = (FFilterOperatorDouble*) SelectedOperator.Get();
 			if (Operator->Apply(Value, InputValue))
 			{
 				return true;
@@ -464,7 +472,8 @@ bool FMetadataFilterState::ApplyFilterToMetadata(TArrayView<const uint8>& Metada
 			}
 			ensure(ConvertedData.IsType<bool>());
 
-			FFilterOperator<bool>* Operator = (FFilterOperator<bool>*) SelectedOperator.Get();
+			using FFilterOperatorBool = UE::Insights::FFilterOperator<bool>;
+			FFilterOperatorBool* Operator = (FFilterOperatorBool*)SelectedOperator.Get();
 			if (Operator->Apply(ConvertedData.Get<bool>(), false))
 			{
 				return true;
@@ -481,7 +490,8 @@ bool FMetadataFilterState::ApplyFilterToMetadata(TArrayView<const uint8>& Metada
 			}
 			ensure(ConvertedData.IsType<bool>());
 
-			FFilterOperator<bool>* Operator = (FFilterOperator<bool>*) SelectedOperator.Get();
+			using FFilterOperatorBool = UE::Insights::FFilterOperator<bool>;
+			FFilterOperatorBool* Operator = (FFilterOperatorBool*)SelectedOperator.Get();
 			if (Operator->Apply(ConvertedData.Get<bool>(), true))
 			{
 				return true;
@@ -542,17 +552,17 @@ void FMetadataFilterState::DataType_OnSelectionChanged(TSharedPtr<FMetadataFilte
 	}
 	case EMetadataFilterDataType::Int:
 	{
-		AvailableOperators.Insert(*FFilterService::Get()->GetIntegerOperators(), 0);
+		AvailableOperators.Insert(*UE::Insights::FFilterService::Get()->GetIntegerOperators(), 0);
 		break;
 	}
 	case EMetadataFilterDataType::Double:
 	{
-		AvailableOperators.Insert(*FFilterService::Get()->GetDoubleOperators(), 0);
+		AvailableOperators.Insert(*UE::Insights::FFilterService::Get()->GetDoubleOperators(), 0);
 		break;
 	}
 	case EMetadataFilterDataType::String:
 	{
-		AvailableOperators.Insert(*FFilterService::Get()->GetStringOperators(), 0);
+		AvailableOperators.Insert(*UE::Insights::FFilterService::Get()->GetStringOperators(), 0);
 		break;
 	}
 	}
@@ -578,7 +588,7 @@ FText FMetadataFilterState::DataType_GetSelectionText() const
 
 ////////////////////////////////////////////////////////////////////////////////////////////////////
 
-TSharedRef<SWidget> FMetadataFilterState::AvailableOperators_OnGenerateWidget(TSharedPtr<IFilterOperator> InOperator)
+TSharedRef<SWidget> FMetadataFilterState::AvailableOperators_OnGenerateWidget(TSharedPtr<UE::Insights::IFilterOperator> InOperator)
 {
 	TSharedRef<SHorizontalBox> Widget = SNew(SHorizontalBox);
 	Widget->AddSlot()
@@ -594,7 +604,7 @@ TSharedRef<SWidget> FMetadataFilterState::AvailableOperators_OnGenerateWidget(TS
 
 ////////////////////////////////////////////////////////////////////////////////////////////////////
 
-void FMetadataFilterState::AvailableOperators_OnSelectionChanged(TSharedPtr<IFilterOperator> InOperator, ESelectInfo::Type SelectInfo)
+void FMetadataFilterState::AvailableOperators_OnSelectionChanged(TSharedPtr<UE::Insights::IFilterOperator> InOperator, ESelectInfo::Type SelectInfo)
 {
 	if (InOperator)
 	{
@@ -645,7 +655,7 @@ bool FMetadataFilterState::Equals(const FFilterState& Other) const
 
 ////////////////////////////////////////////////////////////////////////////////////////////////////
 
-TSharedRef<FFilterState> FMetadataFilterState::DeepCopy() const
+TSharedRef<UE::Insights::FFilterState> FMetadataFilterState::DeepCopy() const
 {
 	TSharedRef<FMetadataFilterState> Copy = MakeShared<FMetadataFilterState>(*this);
 	
@@ -659,14 +669,14 @@ TSharedRef<FFilterState> FMetadataFilterState::DeepCopy() const
 INSIGHTS_IMPLEMENT_RTTI(FMetadataFilter)
 
 FMetadataFilter::FMetadataFilter()
-	: FFilter(static_cast<int32>(EFilterField::Metadata),
+	: FFilter(static_cast<int32>(UE::Insights::EFilterField::Metadata),
 		LOCTEXT("MetadataFilterName", "Metadata"),
 		LOCTEXT("MetadataFilterDesc", "A filter for timing event metadata."),
-		EFilterDataType::Custom,
+		UE::Insights::EFilterDataType::Custom,
 		nullptr,
 		nullptr)
 {
-	SupportedOperators = MakeShared<TArray<TSharedPtr<IFilterOperator>>>();
+	SupportedOperators = MakeShared<TArray<TSharedPtr<UE::Insights::IFilterOperator>>>();
 }
 
 ////////////////////////////////////////////////////////////////////////////////////////////////////
diff --git a/Engine/Source/Developer/TraceInsights/Private/Insights/ViewModels/TimerFilters.h b/Engine/Source/Developer/TraceInsights/Private/Insights/ViewModels/TimerFilters.h
--- a/Engine/Source/Developer/TraceInsights/Private/Insights/ViewModels/TimerFilters.h
+++ b/Engine/Source/Developer/TraceInsights/Private/Insights/ViewModels/TimerFilters.h
@@ -2,11 +2,14 @@
 
 #pragma once
 
+#include "CoreTypes.h"
+
 #include "Types/SlateEnums.h"
 #include "Widgets/Input/SComboBox.h"
 
-#include "Insights/Common/SimpleRtti.h"
-#include "Insights/ViewModels/Filters.h"
+// TraceInsightsCore
+#include "InsightsCore/Common/SimpleRtti.h"
+#include "InsightsCore/Filter/ViewModels/Filters.h"
 
 namespace TraceServices
 {
@@ -16,12 +19,12 @@ namespace TraceServices
 namespace Insights
 {
 
-class FTimerNameFilterState : public FFilterState
+class FTimerNameFilterState : public UE::Insights::FFilterState
 {
-	INSIGHTS_DECLARE_RTTI(FTimerNameFilterState, FFilterState)
+	INSIGHTS_DECLARE_RTTI(FTimerNameFilterState, UE::Insights::FFilterState)
 
 public:
-	FTimerNameFilterState(TSharedRef<FFilter> InFilter)
+	FTimerNameFilterState(TSharedRef<UE::Insights::FFilter> InFilter)
 		: FFilterState(InFilter)
 	{}
 
@@ -35,7 +38,7 @@ public:
 
 	virtual void Update() override;
 
-	virtual bool ApplyFilter(const FFilterContext& Context) const override;
+	virtual bool ApplyFilter(const UE::Insights::FFilterContext& Context) const override;
 
 	virtual void SetFilterValue(FString InFilterValue) override { FilterValue = InFilterValue; }
 
@@ -47,21 +50,21 @@ private:
 	TSet<uint32> TimerIds;
 };
 
-class FTimerNameFilter : public FCustomFilter
+class FTimerNameFilter : public UE::Insights::FCustomFilter
 {
-	INSIGHTS_DECLARE_RTTI(FTimerNameFilter, FCustomFilter)
+	INSIGHTS_DECLARE_RTTI(FTimerNameFilter, UE::Insights::FCustomFilter)
 
 public:
 	FTimerNameFilter();
 
 	virtual ~FTimerNameFilter() {}
 
-	virtual TSharedRef<FFilterState> BuildFilterState() 
+	virtual TSharedRef<UE::Insights::FFilterState> BuildFilterState()
 	{ 
 		return MakeShared<FTimerNameFilterState>(SharedThis(this)); 
 	}
 
-	virtual TSharedRef<FFilterState> BuildFilterState(const FFilterState& Other)
+	virtual TSharedRef<UE::Insights::FFilterState> BuildFilterState(const UE::Insights::FFilterState& Other)
 	{
 		return MakeShared<FTimerNameFilterState>(static_cast<const FTimerNameFilterState&>(Other));
 	}
@@ -89,18 +92,18 @@ struct FMetadataFilterDataTypeEntry
 	FText Name;
 };
 
-class FMetadataFilterState : public FFilterState, public TSharedFromThis<FMetadataFilterState>
+class FMetadataFilterState : public UE::Insights::FFilterState, public TSharedFromThis<FMetadataFilterState>
 {
-	INSIGHTS_DECLARE_RTTI(FMetadataFilterState, FFilterState)
+	INSIGHTS_DECLARE_RTTI(FMetadataFilterState, UE::Insights::FFilterState)
 
 public:
-	FMetadataFilterState(TSharedRef<FFilter> InFilter);
+	FMetadataFilterState(TSharedRef<UE::Insights::FFilter> InFilter);
 
 	virtual ~FMetadataFilterState() {}
 
 	virtual void Update() override;
 
-	virtual bool ApplyFilter(const FFilterContext& Context) const override;
+	virtual bool ApplyFilter(const UE::Insights::FFilterContext& Context) const override;
 
 	virtual void SetFilterValue(FString InFilterValue) override {}
 
@@ -118,8 +121,8 @@ private:
 	void DataType_OnSelectionChanged(TSharedPtr<FMetadataFilterDataTypeEntry> InDataType, ESelectInfo::Type SelectInfo);
 	FText DataType_GetSelectionText() const;
 
-	TSharedRef<SWidget> AvailableOperators_OnGenerateWidget(TSharedPtr<IFilterOperator> InOperator);
-	void AvailableOperators_OnSelectionChanged(TSharedPtr<IFilterOperator> InOperator, ESelectInfo::Type SelectInfo);
+	TSharedRef<SWidget> AvailableOperators_OnGenerateWidget(TSharedPtr<UE::Insights::IFilterOperator> InOperator);
+	void AvailableOperators_OnSelectionChanged(TSharedPtr<UE::Insights::IFilterOperator> InOperator, ESelectInfo::Type SelectInfo);
 	FText AvailableOperators_GetSelectionText() const;
 
 	FText GetTermTextBoxValue() const;
@@ -139,29 +142,29 @@ private:
 	TArray<TSharedPtr<FMetadataFilterDataTypeEntry>> AvailableDataTypes;
 	TSharedPtr<FMetadataFilterDataTypeEntry> SelectedDataType;
 	
-	TArray<TSharedPtr<IFilterOperator>> AvailableOperators;
-	TArray<TSharedPtr<IFilterOperator>> BoolOperators;
+	TArray<TSharedPtr<UE::Insights::IFilterOperator>> AvailableOperators;
+	TArray<TSharedPtr<UE::Insights::IFilterOperator>> BoolOperators;
 
-	TSharedPtr<SComboBox<TSharedPtr<IFilterOperator>>> OperatorComboBox;
+	TSharedPtr<SComboBox<TSharedPtr<UE::Insights::IFilterOperator>>> OperatorComboBox;
 
 	const TraceServices::ITimingProfilerTimerReader* TimerReader;
 };
 
-class FMetadataFilter : public FFilter
+class FMetadataFilter : public UE::Insights::FFilter
 {
-	INSIGHTS_DECLARE_RTTI(FMetadataFilter, FFilter)
+	INSIGHTS_DECLARE_RTTI(FMetadataFilter, UE::Insights::FFilter)
 
 public:
 	FMetadataFilter();
 
 	virtual ~FMetadataFilter() {}
 
-	virtual TSharedRef<FFilterState> BuildFilterState()
+	virtual TSharedRef<UE::Insights::FFilterState> BuildFilterState()
 	{
 		return MakeShared<FMetadataFilterState>(SharedThis(this));
 	}
 
-	virtual TSharedRef<FFilterState> BuildFilterState(const FFilterState& Other)
+	virtual TSharedRef<UE::Insights::FFilterState> BuildFilterState(const UE::Insights::FFilterState& Other)
 	{
 		return MakeShared<FMetadataFilterState>(static_cast<const FMetadataFilterState&>(Other));
 	}
diff --git a/Engine/Source/Developer/TraceInsights/Private/Insights/ViewModels/TimerGroupingAndSorting.cpp b/Engine/Source/Developer/TraceInsights/Private/Insights/ViewModels/TimerGroupingAndSorting.cpp
--- a/Engine/Source/Developer/TraceInsights/Private/Insights/ViewModels/TimerGroupingAndSorting.cpp
+++ b/Engine/Source/Developer/TraceInsights/Private/Insights/ViewModels/TimerGroupingAndSorting.cpp
@@ -2,7 +2,10 @@
 
 #include "TimerGroupingAndSorting.h"
 
-#include "Insights/Table/ViewModels/TableColumn.h"
+// TraceInsightsCore
+#include "InsightsCore/Table/ViewModels/TableColumn.h"
+
+// TraceInsights
 #include "Insights/ViewModels/TimerNodeHelper.h"
 
 #define LOCTEXT_NAMESPACE "TimerNode"
@@ -18,8 +21,8 @@
 // Sorting by Timer Type
 ////////////////////////////////////////////////////////////////////////////////////////////////////
 
-FTimerNodeSortingByTimerType::FTimerNodeSortingByTimerType(TSharedRef<Insights::FTableColumn> InColumnRef)
-	: Insights::FTableCellValueSorter(
+FTimerNodeSortingByTimerType::FTimerNodeSortingByTimerType(TSharedRef<UE::Insights::FTableColumn> InColumnRef)
+	: UE::Insights::FTableCellValueSorter(
 		FName(TEXT("ByTimerType")),
 		LOCTEXT("Sorting_ByTimerType_Name", "By Type"),
 		LOCTEXT("Sorting_ByTimerType_Title", "Sort By Type"),
@@ -30,11 +33,11 @@ FTimerNodeSortingByTimerType::FTimerNodeSortingByTimerType(TSharedRef<Insights::
 
 ////////////////////////////////////////////////////////////////////////////////////////////////////
 
-void FTimerNodeSortingByTimerType::Sort(TArray<Insights::FBaseTreeNodePtr>& NodesToSort, Insights::ESortMode SortMode) const
+void FTimerNodeSortingByTimerType::Sort(TArray<UE::Insights::FBaseTreeNodePtr>& NodesToSort, UE::Insights::ESortMode SortMode) const
 {
-	if (SortMode == Insights::ESortMode::Ascending)
+	if (SortMode == UE::Insights::ESortMode::Ascending)
 	{
-		NodesToSort.Sort([](const Insights::FBaseTreeNodePtr& A, const Insights::FBaseTreeNodePtr& B) -> bool
+		NodesToSort.Sort([](const UE::Insights::FBaseTreeNodePtr& A, const UE::Insights::FBaseTreeNodePtr& B) -> bool
 		{
 			INSIGHTS_ENSURE(A.IsValid() && A->GetTypeName() == FTimerNode::TypeName);
 			const ETimerNodeType ValueA = reinterpret_cast<FTimerNode*>(A.Get())->GetType();
@@ -55,7 +58,7 @@ void FTimerNodeSortingByTimerType::Sort(TArray<Insights::FBaseTreeNodePtr>& Node
 	}
 	else // if (SortMode == ESortMode::Descending)
 	{
-		NodesToSort.Sort([](const Insights::FBaseTreeNodePtr& A, const Insights::FBaseTreeNodePtr& B) -> bool
+		NodesToSort.Sort([](const UE::Insights::FBaseTreeNodePtr& A, const UE::Insights::FBaseTreeNodePtr& B) -> bool
 		{
 			INSIGHTS_ENSURE(A.IsValid() && A->GetTypeName() == FTimerNode::TypeName);
 			const ETimerNodeType ValueA = reinterpret_cast<FTimerNode*>(A.Get())->GetType();
@@ -80,8 +83,8 @@ void FTimerNodeSortingByTimerType::Sort(TArray<Insights::FBaseTreeNodePtr>& Node
 // Sort by Instance Count
 ////////////////////////////////////////////////////////////////////////////////////////////////////
 
-FTimerNodeSortingByInstanceCount::FTimerNodeSortingByInstanceCount(TSharedRef<Insights::FTableColumn> InColumnRef)
-	: Insights::FTableCellValueSorter(
+FTimerNodeSortingByInstanceCount::FTimerNodeSortingByInstanceCount(TSharedRef<UE::Insights::FTableColumn> InColumnRef)
+	: UE::Insights::FTableCellValueSorter(
 		FName(TEXT("ByInstanceCount")),
 		LOCTEXT("Sorting_ByInstanceCount_Name", "By Instance Count"),
 		LOCTEXT("Sorting_ByInstanceCount_Title", "Sort By Instance Count"),
@@ -92,11 +95,11 @@ FTimerNodeSortingByInstanceCount::FTimerNodeSortingByInstanceCount(TSharedRef<In
 
 ////////////////////////////////////////////////////////////////////////////////////////////////////
 
-void FTimerNodeSortingByInstanceCount::Sort(TArray<Insights::FBaseTreeNodePtr>& NodesToSort, Insights::ESortMode SortMode) const
+void FTimerNodeSortingByInstanceCount::Sort(TArray<UE::Insights::FBaseTreeNodePtr>& NodesToSort, UE::Insights::ESortMode SortMode) const
 {
-	if (SortMode == Insights::ESortMode::Ascending)
+	if (SortMode == UE::Insights::ESortMode::Ascending)
 	{
-		NodesToSort.Sort([](const Insights::FBaseTreeNodePtr& A, const Insights::FBaseTreeNodePtr& B) -> bool
+		NodesToSort.Sort([](const UE::Insights::FBaseTreeNodePtr& A, const UE::Insights::FBaseTreeNodePtr& B) -> bool
 		{
 			INSIGHTS_ENSURE(A.IsValid() && A->GetTypeName() == FTimerNode::TypeName);
 			const uint64 ValueA = reinterpret_cast<FTimerNode*>(A.Get())->GetAggregatedStats().InstanceCount;
@@ -117,7 +120,7 @@ void FTimerNodeSortingByInstanceCount::Sort(TArray<Insights::FBaseTreeNodePtr>&
 	}
 	else // if (SortMode == ESortMode::Descending)
 	{
-		NodesToSort.Sort([](const Insights::FBaseTreeNodePtr& A, const Insights::FBaseTreeNodePtr& B) -> bool
+		NodesToSort.Sort([](const UE::Insights::FBaseTreeNodePtr& A, const UE::Insights::FBaseTreeNodePtr& B) -> bool
 		{
 			INSIGHTS_ENSURE(A.IsValid() && A->GetTypeName() == FTimerNode::TypeName);
 			const uint64 ValueA = reinterpret_cast<FTimerNode*>(A.Get())->GetAggregatedStats().InstanceCount;
@@ -142,8 +145,8 @@ void FTimerNodeSortingByInstanceCount::Sort(TArray<Insights::FBaseTreeNodePtr>&
 // Sort by Total Inclusive Time
 ////////////////////////////////////////////////////////////////////////////////////////////////////
 
-FTimerNodeSortingByTotalInclusiveTime::FTimerNodeSortingByTotalInclusiveTime(TSharedRef<Insights::FTableColumn> InColumnRef)
-	: Insights::FTableCellValueSorter(
+FTimerNodeSortingByTotalInclusiveTime::FTimerNodeSortingByTotalInclusiveTime(TSharedRef<UE::Insights::FTableColumn> InColumnRef)
+	: UE::Insights::FTableCellValueSorter(
 		FName(TEXT("ByTotalInclusiveTime")),
 		LOCTEXT("Sorting_ByTotalInclusiveTime_Name", "By Total Inclusive Time"),
 		LOCTEXT("Sorting_ByTotalInclusiveTime_Title", "Sort By Total Inclusive Time"),
@@ -154,11 +157,11 @@ FTimerNodeSortingByTotalInclusiveTime::FTimerNodeSortingByTotalInclusiveTime(TSh
 
 ////////////////////////////////////////////////////////////////////////////////////////////////////
 
-void FTimerNodeSortingByTotalInclusiveTime::Sort(TArray<Insights::FBaseTreeNodePtr>& NodesToSort, Insights::ESortMode SortMode) const
+void FTimerNodeSortingByTotalInclusiveTime::Sort(TArray<UE::Insights::FBaseTreeNodePtr>& NodesToSort, UE::Insights::ESortMode SortMode) const
 {
-	if (SortMode == Insights::ESortMode::Ascending)
+	if (SortMode == UE::Insights::ESortMode::Ascending)
 	{
-		NodesToSort.Sort([](const Insights::FBaseTreeNodePtr& A, const Insights::FBaseTreeNodePtr& B) -> bool
+		NodesToSort.Sort([](const UE::Insights::FBaseTreeNodePtr& A, const UE::Insights::FBaseTreeNodePtr& B) -> bool
 		{
 			INSIGHTS_ENSURE(A.IsValid() && A->GetTypeName() == FTimerNode::TypeName);
 			const double ValueA = reinterpret_cast<FTimerNode*>(A.Get())->GetAggregatedStats().TotalInclusiveTime;
@@ -179,7 +182,7 @@ void FTimerNodeSortingByTotalInclusiveTime::Sort(TArray<Insights::FBaseTreeNodeP
 	}
 	else // if (SortMode == ESortMode::Descending)
 	{
-		NodesToSort.Sort([](const Insights::FBaseTreeNodePtr& A, const Insights::FBaseTreeNodePtr& B) -> bool
+		NodesToSort.Sort([](const UE::Insights::FBaseTreeNodePtr& A, const UE::Insights::FBaseTreeNodePtr& B) -> bool
 		{
 			INSIGHTS_ENSURE(A.IsValid() && A->GetTypeName() == FTimerNode::TypeName);
 			const double ValueA = reinterpret_cast<FTimerNode*>(A.Get())->GetAggregatedStats().TotalInclusiveTime;
@@ -204,8 +207,8 @@ void FTimerNodeSortingByTotalInclusiveTime::Sort(TArray<Insights::FBaseTreeNodeP
 // Sort by Total Exclusive Time
 ////////////////////////////////////////////////////////////////////////////////////////////////////
 
-FTimerNodeSortingByTotalExclusiveTime::FTimerNodeSortingByTotalExclusiveTime(TSharedRef<Insights::FTableColumn> InColumnRef)
-	: Insights::FTableCellValueSorter(
+FTimerNodeSortingByTotalExclusiveTime::FTimerNodeSortingByTotalExclusiveTime(TSharedRef<UE::Insights::FTableColumn> InColumnRef)
+	: UE::Insights::FTableCellValueSorter(
 		FName(TEXT("ByTotalExclusiveTime")),
 		LOCTEXT("Sorting_ByTotalExclusiveTime_Name", "By Total Exclusive Time"),
 		LOCTEXT("Sorting_ByTotalExclusiveTime_Title", "Sort By Total Exclusive Time"),
@@ -216,11 +219,11 @@ FTimerNodeSortingByTotalExclusiveTime::FTimerNodeSortingByTotalExclusiveTime(TSh
 
 ////////////////////////////////////////////////////////////////////////////////////////////////////
 
-void FTimerNodeSortingByTotalExclusiveTime::Sort(TArray<Insights::FBaseTreeNodePtr>& NodesToSort, Insights::ESortMode SortMode) const
+void FTimerNodeSortingByTotalExclusiveTime::Sort(TArray<UE::Insights::FBaseTreeNodePtr>& NodesToSort, UE::Insights::ESortMode SortMode) const
 {
-	if (SortMode == Insights::ESortMode::Ascending)
+	if (SortMode == UE::Insights::ESortMode::Ascending)
 	{
-		NodesToSort.Sort([](const Insights::FBaseTreeNodePtr& A, const Insights::FBaseTreeNodePtr& B) -> bool
+		NodesToSort.Sort([](const UE::Insights::FBaseTreeNodePtr& A, const UE::Insights::FBaseTreeNodePtr& B) -> bool
 		{
 			INSIGHTS_ENSURE(A.IsValid() && A->GetTypeName() == FTimerNode::TypeName);
 			const double ValueA = reinterpret_cast<FTimerNode*>(A.Get())->GetAggregatedStats().TotalExclusiveTime;
@@ -241,7 +244,7 @@ void FTimerNodeSortingByTotalExclusiveTime::Sort(TArray<Insights::FBaseTreeNodeP
 	}
 	else // if (SortMode == ESortMode::Descending)
 	{
-		NodesToSort.Sort([](const Insights::FBaseTreeNodePtr& A, const Insights::FBaseTreeNodePtr& B) -> bool
+		NodesToSort.Sort([](const UE::Insights::FBaseTreeNodePtr& A, const UE::Insights::FBaseTreeNodePtr& B) -> bool
 		{
 			INSIGHTS_ENSURE(A.IsValid() && A->GetTypeName() == FTimerNode::TypeName);
 			const double ValueA = reinterpret_cast<FTimerNode*>(A.Get())->GetAggregatedStats().TotalExclusiveTime;
diff --git a/Engine/Source/Developer/TraceInsights/Private/Insights/ViewModels/TimerGroupingAndSorting.h b/Engine/Source/Developer/TraceInsights/Private/Insights/ViewModels/TimerGroupingAndSorting.h
--- a/Engine/Source/Developer/TraceInsights/Private/Insights/ViewModels/TimerGroupingAndSorting.h
+++ b/Engine/Source/Developer/TraceInsights/Private/Insights/ViewModels/TimerGroupingAndSorting.h
@@ -4,52 +4,54 @@
 
 #include "CoreMinimal.h"
 
-// Insights
-#include "Insights/Table/ViewModels/TableCellValueSorter.h"
-#include "Insights/Table/ViewModels/TableColumn.h"
-#include "Insights/Table/ViewModels/TreeNodeGrouping.h"
+// TraceInsightsCore
+#include "InsightsCore/Table/ViewModels/TableCellValueSorter.h"
+#include "InsightsCore/Table/ViewModels/TableColumn.h"
+#include "InsightsCore/Table/ViewModels/TreeNodeGrouping.h"
+
+// TraceInsights
 #include "Insights/ViewModels/TimerNode.h"
 
 ////////////////////////////////////////////////////////////////////////////////////////////////////
 // Sorters
 ////////////////////////////////////////////////////////////////////////////////////////////////////
 
-class FTimerNodeSortingByTimerType: public Insights::FTableCellValueSorter
+class FTimerNodeSortingByTimerType: public UE::Insights::FTableCellValueSorter
 {
 public:
-	FTimerNodeSortingByTimerType(TSharedRef<Insights::FTableColumn> InColumnRef);
+	FTimerNodeSortingByTimerType(TSharedRef<UE::Insights::FTableColumn> InColumnRef);
 
-	virtual void Sort(TArray<Insights::FBaseTreeNodePtr>& NodesToSort, Insights::ESortMode SortMode) const override;
+	virtual void Sort(TArray<UE::Insights::FBaseTreeNodePtr>& NodesToSort, UE::Insights::ESortMode SortMode) const override;
 };
 
 ////////////////////////////////////////////////////////////////////////////////////////////////////
 
-class FTimerNodeSortingByInstanceCount : public Insights::FTableCellValueSorter
+class FTimerNodeSortingByInstanceCount : public UE::Insights::FTableCellValueSorter
 {
 public:
-	FTimerNodeSortingByInstanceCount(TSharedRef<Insights::FTableColumn> InColumnRef);
+	FTimerNodeSortingByInstanceCount(TSharedRef<UE::Insights::FTableColumn> InColumnRef);
 
-	virtual void Sort(TArray<Insights::FBaseTreeNodePtr>& NodesToSort, Insights::ESortMode SortMode) const override;
+	virtual void Sort(TArray<UE::Insights::FBaseTreeNodePtr>& NodesToSort, UE::Insights::ESortMode SortMode) const override;
 };
 
 ////////////////////////////////////////////////////////////////////////////////////////////////////
 
-class FTimerNodeSortingByTotalInclusiveTime : public Insights::FTableCellValueSorter
+class FTimerNodeSortingByTotalInclusiveTime : public UE::Insights::FTableCellValueSorter
 {
 public:
-	FTimerNodeSortingByTotalInclusiveTime(TSharedRef<Insights::FTableColumn> InColumnRef);
+	FTimerNodeSortingByTotalInclusiveTime(TSharedRef<UE::Insights::FTableColumn> InColumnRef);
 
-	virtual void Sort(TArray<Insights::FBaseTreeNodePtr>& NodesToSort, Insights::ESortMode SortMode) const override;
+	virtual void Sort(TArray<UE::Insights::FBaseTreeNodePtr>& NodesToSort, UE::Insights::ESortMode SortMode) const override;
 };
 
 ////////////////////////////////////////////////////////////////////////////////////////////////////
 
-class FTimerNodeSortingByTotalExclusiveTime : public Insights::FTableCellValueSorter
+class FTimerNodeSortingByTotalExclusiveTime : public UE::Insights::FTableCellValueSorter
 {
 public:
-	FTimerNodeSortingByTotalExclusiveTime(TSharedRef<Insights::FTableColumn> InColumnRef);
+	FTimerNodeSortingByTotalExclusiveTime(TSharedRef<UE::Insights::FTableColumn> InColumnRef);
 
-	virtual void Sort(TArray<Insights::FBaseTreeNodePtr>& NodesToSort, Insights::ESortMode SortMode) const override;
+	virtual void Sort(TArray<UE::Insights::FBaseTreeNodePtr>& NodesToSort, UE::Insights::ESortMode SortMode) const override;
 };
 
 ////////////////////////////////////////////////////////////////////////////////////////////////////
diff --git a/Engine/Source/Developer/TraceInsights/Private/Insights/ViewModels/TimerNode.h b/Engine/Source/Developer/TraceInsights/Private/Insights/ViewModels/TimerNode.h
--- a/Engine/Source/Developer/TraceInsights/Private/Insights/ViewModels/TimerNode.h
+++ b/Engine/Source/Developer/TraceInsights/Private/Insights/ViewModels/TimerNode.h
@@ -3,11 +3,13 @@
 #pragma once
 
 #include "CoreMinimal.h"
+
+// TraceServices
 #include "TraceServices/Containers/Timelines.h"
 #include "TraceServices/Model/TimingProfiler.h"
 
-// Insights
-#include "Insights/Table/ViewModels/BaseTreeNode.h"
+// TraceInsightsCore
+#include "InsightsCore/Table/ViewModels/BaseTreeNode.h"
 
 ////////////////////////////////////////////////////////////////////////////////////////////////////
 
@@ -49,7 +51,7 @@ typedef TWeakPtr<class FTimerNode> FTimerNodeWeak;
 /**
  * Class used to store information about a timer node (used in the STimersView).
  */
-class FTimerNode : public Insights::FBaseTreeNode
+class FTimerNode : public UE::Insights::FBaseTreeNode
 {
 	INSIGHTS_DECLARE_RTTI(FTimerNode, FBaseTreeNode)
 
diff --git a/Engine/Source/Developer/TraceInsights/Private/Insights/ViewModels/TimersViewColumnFactory.cpp b/Engine/Source/Developer/TraceInsights/Private/Insights/ViewModels/TimersViewColumnFactory.cpp
--- a/Engine/Source/Developer/TraceInsights/Private/Insights/ViewModels/TimersViewColumnFactory.cpp
+++ b/Engine/Source/Developer/TraceInsights/Private/Insights/ViewModels/TimersViewColumnFactory.cpp
@@ -2,11 +2,13 @@
 
 #include "TimersViewColumnFactory.h"
 
-// Insights
-#include "Insights/Common/TimeUtils.h"
-#include "Insights/Table/ViewModels/TableCellValueFormatter.h"
-#include "Insights/Table/ViewModels/TableCellValueGetter.h"
-#include "Insights/Table/ViewModels/TableCellValueSorter.h"
+// TraceInsightsCore
+#include "InsightsCore/Common/TimeUtils.h"
+#include "InsightsCore/Table/ViewModels/TableCellValueFormatter.h"
+#include "InsightsCore/Table/ViewModels/TableCellValueGetter.h"
+#include "InsightsCore/Table/ViewModels/TableCellValueSorter.h"
+
+// TraceInsights
 #include "Insights/ViewModels/TimerGroupingAndSorting.h"
 #include "Insights/ViewModels/TimerNodeHelper.h"
 
@@ -41,7 +43,7 @@ const FName FTimersViewColumns::MinExclusiveTimeColumnID(TEXT("MinExclTime"));
 
 ////////////////////////////////////////////////////////////////////////////////////////////////////
 
-void FTimersViewColumnFactory::CreateTimersViewColumns(TArray<TSharedRef<Insights::FTableColumn>>& Columns)
+void FTimersViewColumnFactory::CreateTimersViewColumns(TArray<TSharedRef<UE::Insights::FTableColumn>>& Columns)
 {
 	Columns.Reset();
 
@@ -65,7 +67,7 @@ void FTimersViewColumnFactory::CreateTimersViewColumns(TArray<TSharedRef<Insight
 
 ////////////////////////////////////////////////////////////////////////////////////////////////////
 
-void FTimersViewColumnFactory::CreateTimerTreeViewColumns(TArray<TSharedRef<Insights::FTableColumn>>& Columns)
+void FTimersViewColumnFactory::CreateTimerTreeViewColumns(TArray<TSharedRef<UE::Insights::FTableColumn>>& Columns)
 {
 	Columns.Reset();
 
@@ -81,9 +83,9 @@ void FTimersViewColumnFactory::CreateTimerTreeViewColumns(TArray<TSharedRef<Insi
 
 ////////////////////////////////////////////////////////////////////////////////////////////////////
 
-TSharedRef<Insights::FTableColumn> FTimersViewColumnFactory::CreateNameColumn()
+TSharedRef<UE::Insights::FTableColumn> FTimersViewColumnFactory::CreateNameColumn()
 {
-	using namespace Insights;
+	using namespace UE::Insights;
 
 	TSharedRef<FTimersTableColumn> ColumnRef = MakeShared<FTimersTableColumn>(FTimersViewColumns::NameColumnID);
 	FTimersTableColumn& Column = *ColumnRef;
@@ -120,9 +122,9 @@ TSharedRef<Insights::FTableColumn> FTimersViewColumnFactory::CreateNameColumn()
 
 ////////////////////////////////////////////////////////////////////////////////////////////////////
 
-TSharedRef<Insights::FTableColumn> FTimersViewColumnFactory::CreateMetaGroupNameColumn()
+TSharedRef<UE::Insights::FTableColumn> FTimersViewColumnFactory::CreateMetaGroupNameColumn()
 {
-	using namespace Insights;
+	using namespace UE::Insights;
 
 	TSharedRef<FTimersTableColumn> ColumnRef = MakeShared<FTimersTableColumn>(FTimersViewColumns::MetaGroupNameColumnID);
 	FTimersTableColumn& Column = *ColumnRef;
@@ -169,9 +171,9 @@ TSharedRef<Insights::FTableColumn> FTimersViewColumnFactory::CreateMetaGroupName
 
 ////////////////////////////////////////////////////////////////////////////////////////////////////
 
-TSharedRef<Insights::FTableColumn> FTimersViewColumnFactory::CreateTypeColumn()
+TSharedRef<UE::Insights::FTableColumn> FTimersViewColumnFactory::CreateTypeColumn()
 {
-	using namespace Insights;
+	using namespace UE::Insights;
 
 	TSharedRef<FTimersTableColumn> ColumnRef = MakeShared<FTimersTableColumn>(FTimersViewColumns::TypeColumnID);
 	FTimersTableColumn& Column = *ColumnRef;
@@ -219,9 +221,9 @@ TSharedRef<Insights::FTableColumn> FTimersViewColumnFactory::CreateTypeColumn()
 
 ////////////////////////////////////////////////////////////////////////////////////////////////////
 
-TSharedRef<Insights::FTableColumn> FTimersViewColumnFactory::CreateInstanceCountColumn()
+TSharedRef<UE::Insights::FTableColumn> FTimersViewColumnFactory::CreateInstanceCountColumn()
 {
-	using namespace Insights;
+	using namespace UE::Insights;
 
 	TSharedRef<FTimersTableColumn> ColumnRef = MakeShared<FTimersTableColumn>(FTimersViewColumns::InstanceCountColumnID);
 	FTimersTableColumn& Column = *ColumnRef;
@@ -270,9 +272,9 @@ TSharedRef<Insights::FTableColumn> FTimersViewColumnFactory::CreateInstanceCount
 
 ////////////////////////////////////////////////////////////////////////////////////////////////////
 
-TSharedRef<Insights::FTableColumn> FTimersViewColumnFactory::CreateChildInstanceCountColumn()
+TSharedRef<UE::Insights::FTableColumn> FTimersViewColumnFactory::CreateChildInstanceCountColumn()
 {
-	using namespace Insights;
+	using namespace UE::Insights;
 
 	TSharedRef<FTimersTableColumn> ColumnRef = MakeShared<FTimersTableColumn>(FTimersViewColumns::ChildInstanceCountColumnID);
 	FTimersTableColumn& Column = *ColumnRef;
@@ -326,9 +328,9 @@ TSharedRef<Insights::FTableColumn> FTimersViewColumnFactory::CreateChildInstance
 // Inclusive Time Columns
 ////////////////////////////////////////////////////////////////////////////////////////////////////
 
-TSharedRef<Insights::FTableColumn> FTimersViewColumnFactory::CreateTotalInclusiveTimeColumn()
+TSharedRef<UE::Insights::FTableColumn> FTimersViewColumnFactory::CreateTotalInclusiveTimeColumn()
 {
-	using namespace Insights;
+	using namespace UE::Insights;
 
 	TSharedRef<FTimersTableColumn> ColumnRef = MakeShared<FTimersTableColumn>(FTimersViewColumns::TotalInclusiveTimeColumnID);
 	FTimersTableColumn& Column = *ColumnRef;
@@ -378,9 +380,9 @@ TSharedRef<Insights::FTableColumn> FTimersViewColumnFactory::CreateTotalInclusiv
 
 ////////////////////////////////////////////////////////////////////////////////////////////////////
 
-TSharedRef<Insights::FTableColumn> FTimersViewColumnFactory::CreateMaxInclusiveTimeColumn()
+TSharedRef<UE::Insights::FTableColumn> FTimersViewColumnFactory::CreateMaxInclusiveTimeColumn()
 {
-	using namespace Insights;
+	using namespace UE::Insights;
 
 	TSharedRef<FTimersTableColumn> ColumnRef = MakeShared<FTimersTableColumn>(FTimersViewColumns::MaxInclusiveTimeColumnID);
 	FTimersTableColumn& Column = *ColumnRef;
@@ -427,9 +429,9 @@ TSharedRef<Insights::FTableColumn> FTimersViewColumnFactory::CreateMaxInclusiveT
 
 ////////////////////////////////////////////////////////////////////////////////////////////////////
 
-TSharedRef<Insights::FTableColumn> FTimersViewColumnFactory::CreateAverageInclusiveTimeColumn()
+TSharedRef<UE::Insights::FTableColumn> FTimersViewColumnFactory::CreateAverageInclusiveTimeColumn()
 {
-	using namespace Insights;
+	using namespace UE::Insights;
 
 	TSharedRef<FTimersTableColumn> ColumnRef = MakeShared<FTimersTableColumn>(FTimersViewColumns::AverageInclusiveTimeColumnID);
 	FTimersTableColumn& Column = *ColumnRef;
@@ -476,9 +478,9 @@ TSharedRef<Insights::FTableColumn> FTimersViewColumnFactory::CreateAverageInclus
 
 ////////////////////////////////////////////////////////////////////////////////////////////////////
 
-TSharedRef<Insights::FTableColumn> FTimersViewColumnFactory::CreateMedianInclusiveTimeColumn()
+TSharedRef<UE::Insights::FTableColumn> FTimersViewColumnFactory::CreateMedianInclusiveTimeColumn()
 {
-	using namespace Insights;
+	using namespace UE::Insights;
 
 	TSharedRef<FTimersTableColumn> ColumnRef = MakeShared<FTimersTableColumn>(FTimersViewColumns::MedianInclusiveTimeColumnID);
 	FTimersTableColumn& Column = *ColumnRef;
@@ -525,9 +527,9 @@ TSharedRef<Insights::FTableColumn> FTimersViewColumnFactory::CreateMedianInclusi
 
 ////////////////////////////////////////////////////////////////////////////////////////////////////
 
-TSharedRef<Insights::FTableColumn> FTimersViewColumnFactory::CreateMinInclusiveTimeColumn()
+TSharedRef<UE::Insights::FTableColumn> FTimersViewColumnFactory::CreateMinInclusiveTimeColumn()
 {
-	using namespace Insights;
+	using namespace UE::Insights;
 
 	TSharedRef<FTimersTableColumn> ColumnRef = MakeShared<FTimersTableColumn>(FTimersViewColumns::MinInclusiveTimeColumnID);
 	FTimersTableColumn& Column = *ColumnRef;
@@ -576,9 +578,9 @@ TSharedRef<Insights::FTableColumn> FTimersViewColumnFactory::CreateMinInclusiveT
 // Exclusive Time Columns
 ////////////////////////////////////////////////////////////////////////////////////////////////////
 
-TSharedRef<Insights::FTableColumn> FTimersViewColumnFactory::CreateTotalExclusiveTimeColumn()
+TSharedRef<UE::Insights::FTableColumn> FTimersViewColumnFactory::CreateTotalExclusiveTimeColumn()
 {
-	using namespace Insights;
+	using namespace UE::Insights;
 
 	TSharedRef<FTimersTableColumn> ColumnRef = MakeShared<FTimersTableColumn>(FTimersViewColumns::TotalExclusiveTimeColumnID);
 	FTimersTableColumn& Column = *ColumnRef;
@@ -628,9 +630,9 @@ TSharedRef<Insights::FTableColumn> FTimersViewColumnFactory::CreateTotalExclusiv
 
 ////////////////////////////////////////////////////////////////////////////////////////////////////
 
-TSharedRef<Insights::FTableColumn> FTimersViewColumnFactory::CreateMaxExclusiveTimeColumn()
+TSharedRef<UE::Insights::FTableColumn> FTimersViewColumnFactory::CreateMaxExclusiveTimeColumn()
 {
-	using namespace Insights;
+	using namespace UE::Insights;
 
 	TSharedRef<FTimersTableColumn> ColumnRef = MakeShared<FTimersTableColumn>(FTimersViewColumns::MaxExclusiveTimeColumnID);
 	FTimersTableColumn& Column = *ColumnRef;
@@ -677,9 +679,9 @@ TSharedRef<Insights::FTableColumn> FTimersViewColumnFactory::CreateMaxExclusiveT
 
 ////////////////////////////////////////////////////////////////////////////////////////////////////
 
-TSharedRef<Insights::FTableColumn> FTimersViewColumnFactory::CreateAverageExclusiveTimeColumn()
+TSharedRef<UE::Insights::FTableColumn> FTimersViewColumnFactory::CreateAverageExclusiveTimeColumn()
 {
-	using namespace Insights;
+	using namespace UE::Insights;
 
 	TSharedRef<FTimersTableColumn> ColumnRef = MakeShared<FTimersTableColumn>(FTimersViewColumns::AverageExclusiveTimeColumnID);
 	FTimersTableColumn& Column = *ColumnRef;
@@ -726,9 +728,9 @@ TSharedRef<Insights::FTableColumn> FTimersViewColumnFactory::CreateAverageExclus
 
 ////////////////////////////////////////////////////////////////////////////////////////////////////
 
-TSharedRef<Insights::FTableColumn> FTimersViewColumnFactory::CreateMedianExclusiveTimeColumn()
+TSharedRef<UE::Insights::FTableColumn> FTimersViewColumnFactory::CreateMedianExclusiveTimeColumn()
 {
-	using namespace Insights;
+	using namespace UE::Insights;
 
 	TSharedRef<FTimersTableColumn> ColumnRef = MakeShared<FTimersTableColumn>(FTimersViewColumns::MedianExclusiveTimeColumnID);
 	FTimersTableColumn& Column = *ColumnRef;
@@ -775,9 +777,9 @@ TSharedRef<Insights::FTableColumn> FTimersViewColumnFactory::CreateMedianExclusi
 
 ////////////////////////////////////////////////////////////////////////////////////////////////////
 
-TSharedRef<Insights::FTableColumn> FTimersViewColumnFactory::CreateMinExclusiveTimeColumn()
+TSharedRef<UE::Insights::FTableColumn> FTimersViewColumnFactory::CreateMinExclusiveTimeColumn()
 {
-	using namespace Insights;
+	using namespace UE::Insights;
 
 	TSharedRef<FTimersTableColumn> ColumnRef = MakeShared<FTimersTableColumn>(FTimersViewColumns::MinExclusiveTimeColumnID);
 	FTimersTableColumn& Column = *ColumnRef;
diff --git a/Engine/Source/Developer/TraceInsights/Private/Insights/ViewModels/TimersViewColumnFactory.h b/Engine/Source/Developer/TraceInsights/Private/Insights/ViewModels/TimersViewColumnFactory.h
--- a/Engine/Source/Developer/TraceInsights/Private/Insights/ViewModels/TimersViewColumnFactory.h
+++ b/Engine/Source/Developer/TraceInsights/Private/Insights/ViewModels/TimersViewColumnFactory.h
@@ -2,17 +2,20 @@
 
 #pragma once
 
-#include "CoreMinimal.h"
+#include "CoreTypes.h"
 
-#include "Insights/Table/ViewModels/TableColumn.h"
+#include "Containers/Array.h"
+#include "Templates/SharedPointer.h"
+#include "UObject/NameTypes.h"
+
+// TraceInsightsCore
+#include "InsightsCore/Table/ViewModels/TableColumn.h"
 
 ////////////////////////////////////////////////////////////////////////////////////////////////////
 
+// Column identifiers
 struct FTimersViewColumns
 {
-	//////////////////////////////////////////////////
-	// Column identifiers
-
 	static const FName NameColumnID;
 	static const FName MetaGroupNameColumnID;
 	static const FName TypeColumnID;
@@ -40,16 +43,16 @@ struct FTimersViewColumns
 
 ////////////////////////////////////////////////////////////////////////////////////////////////////
 
-class FTimersTableColumn : public Insights::FTableColumn
+class FTimersTableColumn : public UE::Insights::FTableColumn
 {
 public:
 	FTimersTableColumn(const FName InId)
-		: Insights::FTableColumn(InId)
+		: UE::Insights::FTableColumn(InId)
 	{}
 
-	FText GetDescription(ETraceFrameType InAggreagationMode) const
+	FText GetDescription(ETraceFrameType InAggregationMode) const
 	{
-		switch (InAggreagationMode)
+		switch (InAggregationMode)
 		{
 		case TraceFrameType_Game:
 			return GameFrame_Description;
@@ -62,9 +65,9 @@ public:
 		}
 	}
 
-	void SetDescription(ETraceFrameType InAggreagationMode, FText InDescription)
+	void SetDescription(ETraceFrameType InAggregationMode, FText InDescription)
 	{
-		switch (InAggreagationMode)
+		switch (InAggregationMode)
 		{
 		case TraceFrameType_Game:
 			GameFrame_Description = InDescription;
@@ -73,7 +76,7 @@ public:
 			RenderingFrame_Description = InDescription;
 			break;
 		case TraceFrameType_Count:
-			Insights::FTableColumn::SetDescription(InDescription);
+			UE::Insights::FTableColumn::SetDescription(InDescription);
 			break;
 		default:
 			ensure(0);
@@ -81,7 +84,6 @@ public:
 	}
 
 private:
-
 	FText GameFrame_Description;
 	FText RenderingFrame_Description;
 };
@@ -91,26 +93,26 @@ private:
 struct FTimersViewColumnFactory
 {
 public:
-	static void CreateTimersViewColumns(TArray<TSharedRef<Insights::FTableColumn>>& Columns);
-	static void CreateTimerTreeViewColumns(TArray<TSharedRef<Insights::FTableColumn>>& Columns);
-
-	static TSharedRef<Insights::FTableColumn> CreateNameColumn();
-	static TSharedRef<Insights::FTableColumn> CreateMetaGroupNameColumn();
-	static TSharedRef<Insights::FTableColumn> CreateTypeColumn();
-	static TSharedRef<Insights::FTableColumn> CreateInstanceCountColumn();
-	static TSharedRef<Insights::FTableColumn> CreateChildInstanceCountColumn();
-
-	static TSharedRef<Insights::FTableColumn> CreateTotalInclusiveTimeColumn();
-	static TSharedRef<Insights::FTableColumn> CreateMaxInclusiveTimeColumn();
-	static TSharedRef<Insights::FTableColumn> CreateAverageInclusiveTimeColumn();
-	static TSharedRef<Insights::FTableColumn> CreateMedianInclusiveTimeColumn();
-	static TSharedRef<Insights::FTableColumn> CreateMinInclusiveTimeColumn();
-
-	static TSharedRef<Insights::FTableColumn> CreateTotalExclusiveTimeColumn();
-	static TSharedRef<Insights::FTableColumn> CreateMaxExclusiveTimeColumn();
-	static TSharedRef<Insights::FTableColumn> CreateAverageExclusiveTimeColumn();
-	static TSharedRef<Insights::FTableColumn> CreateMedianExclusiveTimeColumn();
-	static TSharedRef<Insights::FTableColumn> CreateMinExclusiveTimeColumn();
+	static void CreateTimersViewColumns(TArray<TSharedRef<UE::Insights::FTableColumn>>& Columns);
+	static void CreateTimerTreeViewColumns(TArray<TSharedRef<UE::Insights::FTableColumn>>& Columns);
+
+	static TSharedRef<UE::Insights::FTableColumn> CreateNameColumn();
+	static TSharedRef<UE::Insights::FTableColumn> CreateMetaGroupNameColumn();
+	static TSharedRef<UE::Insights::FTableColumn> CreateTypeColumn();
+	static TSharedRef<UE::Insights::FTableColumn> CreateInstanceCountColumn();
+	static TSharedRef<UE::Insights::FTableColumn> CreateChildInstanceCountColumn();
+
+	static TSharedRef<UE::Insights::FTableColumn> CreateTotalInclusiveTimeColumn();
+	static TSharedRef<UE::Insights::FTableColumn> CreateMaxInclusiveTimeColumn();
+	static TSharedRef<UE::Insights::FTableColumn> CreateAverageInclusiveTimeColumn();
+	static TSharedRef<UE::Insights::FTableColumn> CreateMedianInclusiveTimeColumn();
+	static TSharedRef<UE::Insights::FTableColumn> CreateMinInclusiveTimeColumn();
+
+	static TSharedRef<UE::Insights::FTableColumn> CreateTotalExclusiveTimeColumn();
+	static TSharedRef<UE::Insights::FTableColumn> CreateMaxExclusiveTimeColumn();
+	static TSharedRef<UE::Insights::FTableColumn> CreateAverageExclusiveTimeColumn();
+	static TSharedRef<UE::Insights::FTableColumn> CreateMedianExclusiveTimeColumn();
+	static TSharedRef<UE::Insights::FTableColumn> CreateMinExclusiveTimeColumn();
 
 private:
 	static constexpr float TotalTimeColumnInitialWidth = 60.0f;
diff --git a/Engine/Source/Developer/TraceInsights/Private/Insights/ViewModels/TimingEventsTrack.cpp b/Engine/Source/Developer/TraceInsights/Private/Insights/ViewModels/TimingEventsTrack.cpp
--- a/Engine/Source/Developer/TraceInsights/Private/Insights/ViewModels/TimingEventsTrack.cpp
+++ b/Engine/Source/Developer/TraceInsights/Private/Insights/ViewModels/TimingEventsTrack.cpp
@@ -4,9 +4,11 @@
 
 #include "Fonts/FontMeasure.h"
 
-// Insights
-#include "Insights/Common/PaintUtils.h"
-#include "Insights/Common/Stopwatch.h"
+// TraceInsightsCore
+#include "InsightsCore/Common/PaintUtils.h"
+#include "InsightsCore/Common/Stopwatch.h"
+
+// TraceInsights
 #include "Insights/TimingProfilerCommon.h"
 #include "Insights/ViewModels/TimingEvent.h"
 #include "Insights/ViewModels/TimingTrackViewport.h"
@@ -122,7 +124,7 @@ void FTimingEventsTrack::PreUpdate(const ITimingTrackUpdateContext& Context)
 
 			if (FilteredDrawStateInfo.Counter == 0)
 			{
-				FStopwatch Stopwatch;
+				UE::Insights::FStopwatch Stopwatch;
 				Stopwatch.Start();
 				{
 					FTimingEventsTrackDrawStateBuilder Builder(*FilteredDrawState, Context.GetViewport(), Context.GetGeometry().Scale);
@@ -458,7 +460,7 @@ TSharedPtr<ITimingEventFilter> FTimingEventsTrack::GetFilterByEvent(const TShare
 
 ////////////////////////////////////////////////////////////////////////////////////////////////////
 
-void FTimingEventsTrack::DrawSelectedEventInfo(const FString& InText, const FTimingTrackViewport& Viewport, const FDrawContext& DrawContext, const FSlateBrush* WhiteBrush, const FSlateFontInfo& Font) const
+void FTimingEventsTrack::DrawSelectedEventInfo(const FString& InText, const FTimingTrackViewport& Viewport, const UE::Insights::FDrawContext& DrawContext, const FSlateBrush* WhiteBrush, const FSlateFontInfo& Font) const
 {
 	const TSharedRef<FSlateFontMeasure> FontMeasureService = FSlateApplication::Get().GetRenderer()->GetFontMeasureService();
 	const float FontScale = DrawContext.Geometry.Scale;
@@ -480,7 +482,7 @@ void FTimingEventsTrack::DrawSelectedEventInfo(const FString& InText, const FTim
 
 ////////////////////////////////////////////////////////////////////////////////////////////////////
 
-void FTimingEventsTrack::DrawSelectedEventInfoEx(const FString& InText, const FString& InLeftText, const FString& InTopText, const FTimingTrackViewport& Viewport, const FDrawContext& DrawContext, const FSlateBrush* WhiteBrush, const FSlateFontInfo& Font) const
+void FTimingEventsTrack::DrawSelectedEventInfoEx(const FString& InText, const FString& InLeftText, const FString& InTopText, const FTimingTrackViewport& Viewport, const UE::Insights::FDrawContext& DrawContext, const FSlateBrush* WhiteBrush, const FSlateFontInfo& Font) const
 {
 	const TSharedRef<FSlateFontMeasure> FontMeasureService = FSlateApplication::Get().GetRenderer()->GetFontMeasureService();
 	const float FontScale = DrawContext.Geometry.Scale;
diff --git a/Engine/Source/Developer/TraceInsights/Private/Insights/ViewModels/TimingExporter.cpp b/Engine/Source/Developer/TraceInsights/Private/Insights/ViewModels/TimingExporter.cpp
--- a/Engine/Source/Developer/TraceInsights/Private/Insights/ViewModels/TimingExporter.cpp
+++ b/Engine/Source/Developer/TraceInsights/Private/Insights/ViewModels/TimingExporter.cpp
@@ -13,8 +13,10 @@
 #include "TraceServices/Model/Threads.h"
 #include "TraceServices/Model/TimingProfiler.h"
 
+// TraceInsightsCore
+#include "InsightsCore/Common/Stopwatch.h"
+
 // TraceInsights
-#include "Insights/Common/Stopwatch.h"
 #include "Insights/Log.h"
 #include "Insights/TimingProfilerManager.h"
 #include "Insights/ViewModels/ThreadTimingTrack.h"
@@ -66,6 +68,7 @@ IFileHandle* FTimingExporter::OpenExportFile(const TCHAR* InFilename) const
 
 void FTimingExporter::Error(const FText& InMessage) const
 {
+	using namespace UE::Insights::TimingProfiler;
 	FName LogListingName = FTimingProfilerManager::Get()->GetLogListingName();
 	FMessageLog ReportMessageLog((LogListingName != NAME_None) ? LogListingName : TEXT("Other"));
 	ReportMessageLog.Error(InMessage);
@@ -122,7 +125,7 @@ int32 FTimingExporter::ExportThreadsAsText(const FString& Filename, FExportThrea
 		return -1;
 	}
 
-	FStopwatch Stopwatch;
+	UE::Insights::FStopwatch Stopwatch;
 	Stopwatch.Start();
 
 	IFileHandle* ExportFileHandle = OpenExportFile(*Filename);
@@ -206,7 +209,7 @@ int32 FTimingExporter::ExportTimersAsText(const FString& Filename, FExportTimers
 		return -1;
 	}
 
-	FStopwatch Stopwatch;
+	UE::Insights::FStopwatch Stopwatch;
 	Stopwatch.Start();
 
 	IFileHandle* ExportFileHandle = OpenExportFile(*Filename);
@@ -516,7 +519,7 @@ int32 FTimingExporter::ExportTimingEventsAsTextByRegions(const FString& Filename
 		return -1;
 	}
 
-	FStopwatch Stopwatch;
+	UE::Insights::FStopwatch Stopwatch;
 	Stopwatch.Start();
 
 	// Export timing statistics for each region.
@@ -549,7 +552,7 @@ int32 FTimingExporter::ExportTimingEventsAsText(const FString& Filename, FExport
 	ExportTimingEvents_InitColumns();
 	const TArray<FName>& Columns = Params.Columns ? *Params.Columns : ExportTimingEventsDefaultColumns;
 
-	FStopwatch Stopwatch;
+	UE::Insights::FStopwatch Stopwatch;
 	Stopwatch.Start();
 
 	IFileHandle* ExportFileHandle = OpenExportFile(*Filename);
@@ -612,7 +615,7 @@ void FTimingExporter::GetRegions(const FString& InRegionNamePattern, TMap<FStrin
 
 	// Detect regions
 	int32 RegionCount = 0;
-	FStopwatch DetectRegionsStopwatch;
+	UE::Insights::FStopwatch DetectRegionsStopwatch;
 	DetectRegionsStopwatch.Start();
 	{
 		const TraceServices::IRegionProvider& RegionProvider = TraceServices::ReadRegionProvider(Session);
@@ -722,7 +725,7 @@ int32 FTimingExporter::ExportTimerStatisticsAsTextByRegions(const FString& Filen
 		return -1;
 	}
 
-	FStopwatch Stopwatch;
+	UE::Insights::FStopwatch Stopwatch;
 	Stopwatch.Start();
 
 	// Export timing statistics for each region.
@@ -773,7 +776,7 @@ int32 FTimingExporter::ExportTimerStatisticsAsText(const FString& Filename, FExp
 		StatsTable = TimingProfilerProvider->CreateAggregation(CreateAggreagationParams);
 	}
 
-	FStopwatch Stopwatch;
+	UE::Insights::FStopwatch Stopwatch;
 	Stopwatch.Start();
 
 	bool bSuccess = TraceServices::Table2Csv(*StatsTable, *Filename);
@@ -1016,7 +1019,7 @@ int32 FTimingExporter::ExportCountersAsText(const FString& Filename, FExportCoun
 		return -1;
 	}
 
-	FStopwatch Stopwatch;
+	UE::Insights::FStopwatch Stopwatch;
 	Stopwatch.Start();
 
 	IFileHandle* ExportFileHandle = OpenExportFile(*Filename);
@@ -1096,7 +1099,7 @@ int32 FTimingExporter::ExportCounterAsTextByRegions(const FString& FilenamePatte
 		return -1;
 	}
 
-	FStopwatch Stopwatch;
+	UE::Insights::FStopwatch Stopwatch;
 	Stopwatch.Start();
 
 	// Export counter for each region.
@@ -1132,7 +1135,7 @@ int32 FTimingExporter::ExportCounterAsText(const FString& FilenamePattern, uint3
 		return ExportCounterAsTextByRegions(FilenamePattern, CounterId, Params);
 	}
 
-	FStopwatch Stopwatch;
+	UE::Insights::FStopwatch Stopwatch;
 	Stopwatch.Start();
 
 	FString CounterName;
diff --git a/Engine/Source/Developer/TraceInsights/Private/Insights/ViewModels/TimingGraphTrack.cpp b/Engine/Source/Developer/TraceInsights/Private/Insights/ViewModels/TimingGraphTrack.cpp
--- a/Engine/Source/Developer/TraceInsights/Private/Insights/ViewModels/TimingGraphTrack.cpp
+++ b/Engine/Source/Developer/TraceInsights/Private/Insights/ViewModels/TimingGraphTrack.cpp
@@ -2,27 +2,30 @@
 
 #include "TimingGraphTrack.h"
 
-#include <limits>
-
+// TraceServices
+#include "TraceServices/Model/Counters.h"
 #include "TraceServices/Model/Frames.h"
 #include "TraceServices/Model/TimingProfiler.h"
-#include "TraceServices/Model/Counters.h"
 
-// Insights
-#include "Insights/Common/PaintUtils.h"
-#include "Insights/Common/TimeUtils.h"
+// TraceInsightsCore
+#include "InsightsCore/Common/PaintUtils.h"
+#include "InsightsCore/Common/TimeUtils.h"
+
+// TraceInsights
 #include "Insights/InsightsManager.h"
 #include "Insights/TimingProfilerManager.h"
 #include "Insights/ViewModels/AxisViewportDouble.h"
 #include "Insights/ViewModels/FrameStatsHelper.h"
 #include "Insights/ViewModels/GraphTrackBuilder.h"
 #include "Insights/ViewModels/ITimingViewDrawHelper.h"
-#include "Insights/ViewModels/TimingTrackViewport.h"
 #include "Insights/ViewModels/ThreadTimingTrack.h"
-#include "Insights/Widgets/STimingProfilerWindow.h"
+#include "Insights/ViewModels/TimingTrackViewport.h"
 #include "Insights/Widgets/STimersView.h"
+#include "Insights/Widgets/STimingProfilerWindow.h"
 #include "Insights/Widgets/STimingView.h"
 
+#include <limits>
+
 #define LOCTEXT_NAMESPACE "GraphTrack"
 
 ////////////////////////////////////////////////////////////////////////////////////////////////////
@@ -51,19 +54,21 @@ FTimingGraphSeries::~FTimingGraphSeries()
 
 FString FTimingGraphSeries::FormatValue(double Value) const
 {
+	using namespace UE::Insights;
+
 	switch (Type)
 	{
 	case FTimingGraphSeries::ESeriesType::Frame:
-		return FString::Printf(TEXT("%s (%g fps)"), *TimeUtils::FormatTimeAuto(Value), 1.0 / Value);
+		return FString::Printf(TEXT("%s (%g fps)"), *FormatTimeAuto(Value), 1.0 / Value);
 
 	case FTimingGraphSeries::ESeriesType::Timer:
 	case FTimingGraphSeries::ESeriesType::FrameStatsTimer:
-		return TimeUtils::FormatTimeAuto(Value);
+		return FormatTimeAuto(Value);
 
 	case FTimingGraphSeries::ESeriesType::StatsCounter:
 		if (bIsTime)
 		{
-			return TimeUtils::FormatTimeAuto(Value);
+			return FormatTimeAuto(Value);
 		}
 		else if (bIsMemory)
 		{
@@ -130,7 +135,7 @@ INSIGHTS_IMPLEMENT_RTTI(FTimingGraphTrack)
 
 ////////////////////////////////////////////////////////////////////////////////////////////////////
 
-FTimingGraphTrack::FTimingGraphTrack(TSharedPtr<STimingView> InTimingView)
+FTimingGraphTrack::FTimingGraphTrack(TSharedPtr<UE::Insights::TimingProfiler::STimingView> InTimingView)
 	: FGraphTrack()
 	, TimingView(InTimingView)
 	//, SharedValueViewport()
@@ -149,7 +154,7 @@ FTimingGraphTrack::~FTimingGraphTrack()
 {
 	if (OnTrackVisibilityChangedHandle.IsValid())
 	{
-		TSharedPtr<STimingView> TimingViewPtr = TimingView.Pin();
+		TSharedPtr<UE::Insights::TimingProfiler::STimingView> TimingViewPtr = TimingView.Pin();
 		if (TimingView.IsValid())
 		{
 			TimingViewPtr->OnTrackVisibilityChanged().Remove(OnTrackVisibilityChangedHandle);
@@ -179,6 +184,7 @@ FTimingGraphTrack::~FTimingGraphTrack()
 	TSharedPtr<STimersView> TimersView;
 	if (bNotifyTimersOnDestruction)
 	{
+		using namespace UE::Insights::TimingProfiler;
 		TSharedPtr<STimingProfilerWindow> ProfilerWindow = FTimingProfilerManager::Get()->GetProfilerWindow();
 		if (ProfilerWindow.IsValid())
 		{
@@ -212,7 +218,7 @@ void FTimingGraphTrack::Update(const ITimingTrackUpdateContext& Context)
 
 	if (!OnTrackVisibilityChangedHandle.IsValid())
 	{
-		TSharedPtr<STimingView> TimingViewPtr = TimingView.Pin();
+		TSharedPtr<UE::Insights::TimingProfiler::STimingView> TimingViewPtr = TimingView.Pin();
 		if (TimingViewPtr.IsValid())
 		{
 			auto OnTrackAddedRemovedLamda = [this](const TSharedPtr<const FBaseTimingTrack> Track)
@@ -301,7 +307,7 @@ void FTimingGraphTrack::Update(const ITimingTrackUpdateContext& Context)
 void FTimingGraphTrack::AddDefaultFrameSeries()
 {
 	const FInsightsSettings& Settings = FInsightsManager::Get()->GetSettings();
-	TSharedPtr<STimingView> TimingViewPtr = TimingView.Pin();
+	TSharedPtr<UE::Insights::TimingProfiler::STimingView> TimingViewPtr = TimingView.Pin();
 
 	TSharedRef<FTimingGraphSeries> GameFramesSeries = MakeShared<FTimingGraphSeries>(FTimingGraphSeries::ESeriesType::Frame);
 	GameFramesSeries->SetName(TEXT("Game Frames"));
@@ -375,7 +381,7 @@ void FTimingGraphTrack::UpdateFrameSeries(FTimingGraphSeries& Series, const FTim
 
 		FramesProvider.EnumerateFrames(Series.FrameType, StartIndex, EndIndex, [&Builder](const TraceServices::FFrame& Frame)
 		{
-			//TODO: add a "frame converter" (i.e. to fps, miliseconds or seconds)
+			//TODO: add a "frame converter" (i.e. to fps, milliseconds or seconds)
 			const double Duration = Frame.EndTime - Frame.StartTime;
 			Builder.AddEvent(Frame.StartTime, Duration, Duration);
 		});
@@ -537,7 +543,7 @@ void FTimingGraphTrack::UpdateTimerSeries(FTimingGraphSeries& Series, const FTim
 }
 
 ////////////////////////////////////////////////////////////////////////////////////////////////////
-// Frams Stats Timer Series
+// Frame Stats Timer Series
 ////////////////////////////////////////////////////////////////////////////////////////////////////
 
 TSharedPtr<FTimingGraphSeries> FTimingGraphTrack::GetFrameStatsTimerSeries(uint32 TimerId, ETraceFrameType FrameType)
@@ -795,7 +801,7 @@ void FTimingGraphTrack::UpdateStatsCounterSeries(FTimingGraphSeries& Series, con
 
 void FTimingGraphTrack::GetVisibleTimelineIndexes(TSet<uint32>& TimelineIndexes)
 {
-	TSharedPtr<STimingView> TimingViewPtr = TimingView.Pin();
+	TSharedPtr<UE::Insights::TimingProfiler::STimingView> TimingViewPtr = TimingView.Pin();
 	if (!TimingViewPtr.IsValid())
 	{
 		return;
@@ -811,7 +817,7 @@ void FTimingGraphTrack::ContextMenu_ToggleOption_Execute(EGraphOptions Option)
 {
 	FGraphTrack::ContextMenu_ToggleOption_Execute(Option);
 
-	TSharedPtr<STimingView> TimingViewPtr = TimingView.Pin();
+	TSharedPtr<UE::Insights::TimingProfiler::STimingView> TimingViewPtr = TimingView.Pin();
 	if (!TimingViewPtr.IsValid())
 	{
 		return;
@@ -852,7 +858,7 @@ void FTimingGraphTrack::ContextMenu_ToggleOption_Execute(EGraphOptions Option)
 
 void FTimingGraphTrack::LoadDefaultSettings()
 {
-	TSharedPtr<STimingView> TimingViewPtr = TimingView.Pin();
+	TSharedPtr<UE::Insights::TimingProfiler::STimingView> TimingViewPtr = TimingView.Pin();
 	if (TimingViewPtr.IsValid() && TimingViewPtr->GetName() == FInsightsManagerTabs::TimingProfilerTabId)
 	{
 		const FInsightsSettings& Settings = FInsightsManager::Get()->GetSettings();
@@ -920,7 +926,7 @@ void FTimingGraphTrack::DrawVerticalAxisGrid(const ITimingTrackDrawContext& Cont
 	constexpr float MinDY = 32.0f; // min vertical distance between horizontal grid lines
 	constexpr float TextH = 14.0f; // label height
 
-	FDrawContext& DrawContext = Context.GetDrawContext();
+	UE::Insights::FDrawContext& DrawContext = Context.GetDrawContext();
 	const FSlateBrush* Brush = Context.GetHelper().GetWhiteBrush();
 	//const FSlateFontInfo& Font = Context.GetHelper().GetEventFont();
 	const TSharedRef<FSlateFontMeasure> FontMeasureService = FSlateApplication::Get().GetRenderer()->GetFontMeasureService();
@@ -984,7 +990,7 @@ void FTimingGraphTrack::DrawVerticalAxisGrid(const ITimingTrackDrawContext& Cont
 		{
 			const float Y = Y0 + RoundedViewHeight - FMath::RoundToFloat(ViewportY.GetOffsetForValue(Value));
 
-			const FString LabelText = TimeUtils::FormatTimeAuto(Value);
+			const FString LabelText = UE::Insights::FormatTimeAuto(Value);
 
 			// Draw horizontal grid line.
 			DrawContext.DrawBox(0, Y, ViewWidth, 1, Brush, GridColor);
diff --git a/Engine/Source/Developer/TraceInsights/Private/Insights/ViewModels/TimingGraphTrack.h b/Engine/Source/Developer/TraceInsights/Private/Insights/ViewModels/TimingGraphTrack.h
--- a/Engine/Source/Developer/TraceInsights/Private/Insights/ViewModels/TimingGraphTrack.h
+++ b/Engine/Source/Developer/TraceInsights/Private/Insights/ViewModels/TimingGraphTrack.h
@@ -4,7 +4,7 @@
 
 #include "CoreMinimal.h"
 
-// Insights
+// TraceInsights
 #include "Insights/ViewModels/GraphSeries.h"
 #include "Insights/ViewModels/GraphTrack.h"
 
@@ -13,7 +13,7 @@ namespace Insights
 	struct FFrameStatsCachedEvent;
 }
 
-class STimingView;
+namespace UE::Insights::TimingProfiler { class STimingView; }
 
 ////////////////////////////////////////////////////////////////////////////////////////////////////
 
@@ -78,7 +78,7 @@ class FTimingGraphTrack : public FGraphTrack
 	INSIGHTS_DECLARE_RTTI(FTimingGraphTrack, FGraphTrack)
 
 public:
-	FTimingGraphTrack(TSharedPtr<STimingView> InTimingView);
+	FTimingGraphTrack(TSharedPtr<UE::Insights::TimingProfiler::STimingView> InTimingView);
 	virtual ~FTimingGraphTrack();
 
 	virtual void Update(const ITimingTrackUpdateContext& Context) override;
@@ -124,7 +124,7 @@ private:
 	FDelegateHandle GameFrameSeriesVisibilityHandle;
 	FDelegateHandle RenderingFrameSeriesVisibilityHandle;
 
-	TWeakPtr<STimingView> TimingView;
+	TWeakPtr<UE::Insights::TimingProfiler::STimingView> TimingView;
 	bool bNotifyTimersOnDestruction;
 };
 
diff --git a/Engine/Source/Developer/TraceInsights/Private/Insights/ViewModels/TimingViewDrawHelper.cpp b/Engine/Source/Developer/TraceInsights/Private/Insights/ViewModels/TimingViewDrawHelper.cpp
--- a/Engine/Source/Developer/TraceInsights/Private/Insights/ViewModels/TimingViewDrawHelper.cpp
+++ b/Engine/Source/Developer/TraceInsights/Private/Insights/ViewModels/TimingViewDrawHelper.cpp
@@ -7,9 +7,11 @@
 #include "Rendering/DrawElements.h"
 #include "Styling/AppStyle.h"
 
-// Insights
-#include "Insights/Common/PaintUtils.h"
-#include "Insights/Common/TimeUtils.h"
+// TraceInsightsCore
+#include "InsightsCore/Common/PaintUtils.h"
+#include "InsightsCore/Common/TimeUtils.h"
+
+// TraceInsights
 #include "Insights/InsightsStyle.h"
 #include "Insights/ViewModels/DrawHelpers.h"
 #include "Insights/ViewModels/MarkersTimingTrack.h"
@@ -40,7 +42,7 @@ FTimingEventsTrackDrawStateBuilder::FTimingEventsTrackDrawStateBuilder(FTimingEv
 void FTimingEventsTrackDrawStateBuilder::AppendDurationToEventName(FString& InOutEventName, const double InDuration)
 {
 	InOutEventName += TEXT(" (");
-	InOutEventName += TimeUtils::FormatTimeAuto(InDuration);
+	InOutEventName += UE::Insights::FormatTimeAuto(InDuration);
 	InOutEventName += TEXT(")");
 }
 
@@ -266,7 +268,7 @@ void FTimingEventsTrackDrawStateBuilder::Flush()
 // FTimingViewDrawHelper
 ////////////////////////////////////////////////////////////////////////////////////////////////////
 
-FTimingViewDrawHelper::FTimingViewDrawHelper(const FDrawContext& InDrawContext, const FTimingTrackViewport& InViewport)
+FTimingViewDrawHelper::FTimingViewDrawHelper(const UE::Insights::FDrawContext& InDrawContext, const FTimingTrackViewport& InViewport)
 	: DrawContext(InDrawContext)
 	, Viewport(InViewport)
 	, WhiteBrush(FInsightsStyle::Get().GetBrush("WhiteBrush"))
@@ -290,6 +292,15 @@ FTimingViewDrawHelper::FTimingViewDrawHelper(const FDrawContext& InDrawContext,
 
 ////////////////////////////////////////////////////////////////////////////////////////////////////
 
+PRAGMA_DISABLE_DEPRECATION_WARNINGS
+FTimingViewDrawHelper::FTimingViewDrawHelper(const ::FDrawContext& InDrawContext, const FTimingTrackViewport& InViewport)
+	: FTimingViewDrawHelper((const UE::Insights::FDrawContext&)InDrawContext, InViewport)
+{
+}
+PRAGMA_ENABLE_DEPRECATION_WARNINGS
+
+////////////////////////////////////////////////////////////////////////////////////////////////////
+
 FTimingViewDrawHelper::~FTimingViewDrawHelper()
 {
 }
diff --git a/Engine/Source/Developer/TraceInsights/Private/Insights/ViewModels/TimingViewDrawHelper.h b/Engine/Source/Developer/TraceInsights/Private/Insights/ViewModels/TimingViewDrawHelper.h
--- a/Engine/Source/Developer/TraceInsights/Private/Insights/ViewModels/TimingViewDrawHelper.h
+++ b/Engine/Source/Developer/TraceInsights/Private/Insights/ViewModels/TimingViewDrawHelper.h
@@ -3,10 +3,11 @@
 #pragma once
 
 #include "CoreMinimal.h"
+
 #include "Fonts/SlateFontInfo.h"
 #include "Math/Color.h"
 
-// Insights
+// TraceInsights
 #include "Insights/ViewModels/ITimingEvent.h"
 #include "Insights/ViewModels/ITimingViewDrawHelper.h"
 #include "Insights/ViewModels/TimingEventsTrack.h" // for ITimingEventsTrackDrawStateBuilder
@@ -14,7 +15,9 @@
 ////////////////////////////////////////////////////////////////////////////////////////////////////
 
 struct FSlateBrush;
-struct FDrawContext;
+
+struct FDrawContext; // for backward compatibility
+namespace UE::Insights { class FDrawContext; }
 
 class ITimingTrackDrawContext;
 class FTimingEventsTrack;
@@ -159,7 +162,10 @@ private:
 	static int32 ToInt32(EDrawLayer Layer) { return static_cast<int32>(Layer); }
 
 public:
-	explicit FTimingViewDrawHelper(const FDrawContext& InDrawContext, const FTimingTrackViewport& InViewport);
+	explicit FTimingViewDrawHelper(const UE::Insights::FDrawContext& InDrawContext, const FTimingTrackViewport& InViewport);
+	PRAGMA_DISABLE_DEPRECATION_WARNINGS
+	explicit FTimingViewDrawHelper(const ::FDrawContext& InDrawContext, const FTimingTrackViewport& InViewport);
+	PRAGMA_ENABLE_DEPRECATION_WARNINGS
 	~FTimingViewDrawHelper();
 
 	/**
@@ -184,7 +190,7 @@ public:
 	virtual int32 GetFirstLayerId() const override { return ReservedLayerId; }
 	virtual int32 GetNumLayerIds() const override { return ToInt32(EDrawLayer::Count); }
 
-	const FDrawContext& GetDrawContext() const { return DrawContext; }
+	const UE::Insights::FDrawContext& GetDrawContext() const { return DrawContext; }
 	const FTimingTrackViewport& GetViewport() const { return Viewport; }
 
 	void DrawBackground() const;
@@ -221,7 +227,7 @@ public:
 	int32 GetNumDrawTexts() const { return NumDrawTexts; }
 
 private:
-	const FDrawContext& DrawContext;
+	const UE::Insights::FDrawContext& DrawContext;
 	const FTimingTrackViewport& Viewport;
 
 	const FSlateBrush* WhiteBrush;
diff --git a/Engine/Source/Developer/TraceInsights/Private/Insights/ViewModels/TooltipDrawState.cpp b/Engine/Source/Developer/TraceInsights/Private/Insights/ViewModels/TooltipDrawState.cpp
--- a/Engine/Source/Developer/TraceInsights/Private/Insights/ViewModels/TooltipDrawState.cpp
+++ b/Engine/Source/Developer/TraceInsights/Private/Insights/ViewModels/TooltipDrawState.cpp
@@ -6,8 +6,10 @@
 #include "Fonts/SlateFontInfo.h"
 #include "Styling/AppStyle.h"
 
-// Insights
-#include "Insights/Common/PaintUtils.h"
+// TraceInsightsCore
+#include "InsightsCore/Common/PaintUtils.h"
+
+// TraceInsights
 #include "Insights/InsightsStyle.h"
 
 ////////////////////////////////////////////////////////////////////////////////////////////////////
@@ -244,7 +246,7 @@ void FTooltipDrawState::SetPosition(const FVector2D& MousePosition, const float
 
 ////////////////////////////////////////////////////////////////////////////////////////////////////
 
-void FTooltipDrawState::Draw(const FDrawContext& DrawContext) const
+void FTooltipDrawState::Draw(const UE::Insights::FDrawContext& DrawContext) const
 {
 	if (Opacity > 0.0f && Size.X > 0.0 && Size.Y > 0.0)
 	{
diff --git a/Engine/Source/Developer/TraceInsights/Private/Insights/ViewModels/TrackHeader.cpp b/Engine/Source/Developer/TraceInsights/Private/Insights/ViewModels/TrackHeader.cpp
--- a/Engine/Source/Developer/TraceInsights/Private/Insights/ViewModels/TrackHeader.cpp
+++ b/Engine/Source/Developer/TraceInsights/Private/Insights/ViewModels/TrackHeader.cpp
@@ -6,8 +6,10 @@
 #include "Framework/Application/SlateApplication.h"
 #include "Styling/AppStyle.h"
 
-// Insights
-#include "Insights/Common/PaintUtils.h"
+// TraceInsightsCore
+#include "InsightsCore/Common/PaintUtils.h"
+
+// TraceInsights
 #include "Insights/InsightsStyle.h"
 #include "Insights/ViewModels/BaseTimingTrack.h"
 #include "Insights/ViewModels/TimingTrackViewport.h"
@@ -116,7 +118,7 @@ void FTrackHeader::Draw(const ITimingTrackDrawContext& Context) const
 	}
 	else
 	{
-		FDrawContext& DrawContext = Context.GetDrawContext();
+		UE::Insights::FDrawContext& DrawContext = Context.GetDrawContext();
 
 		// Draw a horizontal line between tracks (top line of a track).
 		DrawContext.DrawBox(0.0f, ParentTrack.GetPosY(), Context.GetViewport().GetWidth(), 1.0f, WhiteBrush, FLinearColor(0.05f, 0.05f, 0.05f, 1.0f));
@@ -143,7 +145,7 @@ void FTrackHeader::PostDraw(const ITimingTrackDrawContext& Context) const
 
 void FTrackHeader::DrawInternal(const ITimingTrackDrawContext& Context, bool bDrawBackgroundLayer) const
 {
-	FDrawContext& DrawContext = Context.GetDrawContext();
+	UE::Insights::FDrawContext& DrawContext = Context.GetDrawContext();
 
 	const float Opacity = bIsInBackground ? CurrentHoveredAnimPercent : 1.0f;
 
