[Insights] Massive refactoring of TraceInsights module:

diff --git a/Engine/Source/Developer/TraceInsights/Private/Insights/Widgets/SFrameTrack.cpp b/Engine/Source/Developer/TraceInsights/Private/Insights/Widgets/SFrameTrack.cpp
--- a/Engine/Source/Developer/TraceInsights/Private/Insights/Widgets/SFrameTrack.cpp
+++ b/Engine/Source/Developer/TraceInsights/Private/Insights/Widgets/SFrameTrack.cpp
@@ -10,16 +10,20 @@
 #include "Misc/StringBuilder.h"
 #include "Rendering/DrawElements.h"
 #include "Styling/AppStyle.h"
-#include "TraceServices/Model/Frames.h"
-#include "TraceServices/Model/TimingProfiler.h"
 #include "Widgets/Input/SButton.h"
 #include "Widgets/Input/SEditableTextBox.h"
 #include "Widgets/Layout/SScrollBar.h"
 
-// Insights
-#include "Insights/Common/PaintUtils.h"
-#include "Insights/Common/Stopwatch.h"
-#include "Insights/Common/TimeUtils.h"
+// TraceServices
+#include "TraceServices/Model/Frames.h"
+#include "TraceServices/Model/TimingProfiler.h"
+
+// TraceInsightsCore
+#include "InsightsCore/Common/PaintUtils.h"
+#include "InsightsCore/Common/Stopwatch.h"
+#include "InsightsCore/Common/TimeUtils.h"
+
+// TraceInsights
 #include "Insights/InsightsManager.h"
 #include "Insights/InsightsStyle.h"
 #include "Insights/Log.h"
@@ -50,12 +54,14 @@ SFrameTrack::SFrameTrack()
 
 SFrameTrack::~SFrameTrack()
 {
+	using namespace UE::Insights::TimingProfiler;
+	TSharedPtr<STimingProfilerWindow> ProfilerWindow = FTimingProfilerManager::Get()->GetProfilerWindow();
+
 	if (OnTrackVisibilityChangedHandle.IsValid())
 	{
-		TSharedPtr<class STimingProfilerWindow> TimingWindow = FTimingProfilerManager::Get()->GetProfilerWindow();
-		if (TimingWindow.IsValid())
+		if (ProfilerWindow.IsValid())
 		{
-			TSharedPtr<STimingView> TimingView = TimingWindow->GetTimingView();
+			TSharedPtr<STimingView> TimingView = ProfilerWindow->GetTimingView();
 			if (TimingView.IsValid())
 			{
 				if (TimingView.Get() == RegisteredTimingView)
@@ -69,12 +75,10 @@ SFrameTrack::~SFrameTrack()
 	}
 
 	TSharedPtr<STimersView> TimersView;
-	TSharedPtr<STimingProfilerWindow> ProfilerWindow = FTimingProfilerManager::Get()->GetProfilerWindow();
 	if (ProfilerWindow.IsValid())
 	{
 		TimersView = ProfilerWindow->GetTimersView();
 	}
-
 	if (TimersView)
 	{
 		for (const TSharedPtr<FFrameTrackSeries>& Series : AllSeries)
@@ -191,6 +195,7 @@ void SFrameTrack::Construct(const FArguments& InArgs)
 
 void SFrameTrack::Tick(const FGeometry& AllottedGeometry, const double InCurrentTime, const float InDeltaTime)
 {
+	using namespace UE::Insights::TimingProfiler;
 	TSharedPtr<class STimingProfilerWindow> TimingWindow = FTimingProfilerManager::Get()->GetProfilerWindow();
 	if (TimingWindow.IsValid())
 	{
@@ -370,7 +375,7 @@ TSharedPtr<FFrameTrackSeries> SFrameTrack::FindFrameStatsSeries(ETraceFrameType
 
 void SFrameTrack::UpdateState()
 {
-	FStopwatch Stopwatch;
+	UE::Insights::FStopwatch Stopwatch;
 	Stopwatch.Start();
 
 	// Reset stats.
@@ -430,6 +435,8 @@ void SFrameTrack::UpdateState()
 
 			bool bTimingViewExists = false;
 			TSet<uint32> Timelines;
+
+			using namespace UE::Insights::TimingProfiler;
 			TSharedPtr<class STimingProfilerWindow> TimingWindow = FTimingProfilerManager::Get()->GetProfilerWindow();
 
 			// Attempt to compute only from visible timelines.
@@ -552,6 +559,7 @@ void SFrameTrack::SelectFrameAtMousePosition(double X, double Y, bool JoinCurren
 
 	if (SampleRef.IsValid())
 	{
+		using namespace UE::Insights::TimingProfiler;
 		TSharedPtr<STimingProfilerWindow> Window = FTimingProfilerManager::Get()->GetProfilerWindow();
 		if (Window.IsValid())
 		{
@@ -595,7 +603,7 @@ int32 SFrameTrack::OnPaint(const FPaintArgs& Args, const FGeometry& AllottedGeom
 {
 	const bool bEnabled = ShouldBeEnabled(bParentEnabled);
 	const ESlateDrawEffect DrawEffects = bEnabled ? ESlateDrawEffect::None : ESlateDrawEffect::DisabledEffect;
-	FDrawContext DrawContext(AllottedGeometry, MyCullingRect, InWidgetStyle, DrawEffects, OutDrawElements, LayerId);
+	UE::Insights::FDrawContext DrawContext(AllottedGeometry, MyCullingRect, InWidgetStyle, DrawEffects, OutDrawElements, LayerId);
 
 	const TSharedRef<FSlateFontMeasure> FontMeasureService = FSlateApplication::Get().GetRenderer()->GetFontMeasureService();
 	FSlateFontInfo SummaryFont = FAppStyle::Get().GetFontStyle("SmallFont");
@@ -609,7 +617,7 @@ int32 SFrameTrack::OnPaint(const FPaintArgs& Args, const FGeometry& AllottedGeom
 
 	//////////////////////////////////////////////////
 	{
-		FStopwatch Stopwatch;
+		UE::Insights::FStopwatch Stopwatch;
 		Stopwatch.Start();
 
 		FFrameTrackDrawHelper Helper(DrawContext, Viewport);
@@ -640,6 +648,7 @@ int32 SFrameTrack::OnPaint(const FPaintArgs& Args, const FGeometry& AllottedGeom
 		TSharedPtr<FFrameTrackSeries> GameFrameSeries = FindSeries(ETraceFrameType::TraceFrameType_Game);
 		if (GameFrameSeries.IsValid())
 		{
+			using namespace UE::Insights::TimingProfiler;
 			TSharedPtr<STimingProfilerWindow> Window = FTimingProfilerManager::Get()->GetProfilerWindow();
 			if (Window)
 			{
@@ -709,10 +718,11 @@ int32 SFrameTrack::OnPaint(const FPaintArgs& Args, const FGeometry& AllottedGeom
 			const FString Text1(StringBuilder);
 
 			// Second line: "1m 2.34s + 16.67ms (60 fps)"
+			using namespace UE::Insights;
 			StringBuilder.Reset();
-			StringBuilder.Append(TimeUtils::FormatTimeAuto(HoveredSample.Sample->LargestFrameStartTime, HoveredSample.Sample->LargestFrameStartTime > 60.0 ? 3 : 2));
+			StringBuilder.Append(FormatTimeAuto(HoveredSample.Sample->LargestFrameStartTime, HoveredSample.Sample->LargestFrameStartTime > 60.0 ? 3 : 2));
 			StringBuilder.Append(TEXT(" + "));
-			StringBuilder.Append(TimeUtils::FormatTimeAuto(HoveredSample.Sample->LargestFrameDuration, 2));
+			StringBuilder.Append(FormatTimeAuto(HoveredSample.Sample->LargestFrameDuration, 2));
 			StringBuilder.Appendf(TEXT(" (%.1f fps)"), 1.0 / HoveredSample.Sample->LargestFrameDuration);
 			const FString Text2(StringBuilder);
 
@@ -810,11 +820,11 @@ int32 SFrameTrack::OnPaint(const FPaintArgs& Args, const FGeometry& AllottedGeom
 		LastOnPaintTime = CurrentTime;
 		OnPaintDurationHistory.AddValue(OnPaintDuration); // saved for last 32 OnPaint calls
 		const uint64 AvgOnPaintDuration = OnPaintDurationHistory.ComputeAverage();
-		const uint64 AvgOnPaintDurationMs = FStopwatch::Cycles64ToMilliseconds(AvgOnPaintDuration);
-		const double AvgOnPaintFps = AvgOnPaintDurationMs != 0 ? 1.0 / FStopwatch::Cycles64ToSeconds(AvgOnPaintDuration) : 0.0;
+		const uint64 AvgOnPaintDurationMs = UE::Insights::FStopwatch::Cycles64ToMilliseconds(AvgOnPaintDuration);
+		const double AvgOnPaintFps = AvgOnPaintDurationMs != 0 ? 1.0 / UE::Insights::FStopwatch::Cycles64ToSeconds(AvgOnPaintDuration) : 0.0;
 
-		const uint64 AvgUpdateDurationMs = FStopwatch::Cycles64ToMilliseconds(UpdateDurationHistory.ComputeAverage());
-		const uint64 AvgDrawDurationMs = FStopwatch::Cycles64ToMilliseconds(DrawDurationHistory.ComputeAverage());
+		const uint64 AvgUpdateDurationMs = UE::Insights::FStopwatch::Cycles64ToMilliseconds(UpdateDurationHistory.ComputeAverage());
+		const uint64 AvgDrawDurationMs = UE::Insights::FStopwatch::Cycles64ToMilliseconds(DrawDurationHistory.ComputeAverage());
 
 		// Draw performance info.
 		DrawContext.DrawText
@@ -865,7 +875,7 @@ int32 SFrameTrack::OnPaint(const FPaintArgs& Args, const FGeometry& AllottedGeom
 
 ////////////////////////////////////////////////////////////////////////////////////////////////////
 
-void SFrameTrack::DrawVerticalAxisGrid(FDrawContext& DrawContext, const FSlateBrush* Brush, const FSlateFontInfo& Font) const
+void SFrameTrack::DrawVerticalAxisGrid(UE::Insights::FDrawContext& DrawContext, const FSlateBrush* Brush, const FSlateFontInfo& Font) const
 {
 	const float ViewWidth = Viewport.GetWidth();
 
@@ -987,9 +997,10 @@ void SFrameTrack::DrawVerticalAxisGrid(FDrawContext& DrawContext, const FSlateBr
 		// Draw horizontal grid line.
 		DrawContext.DrawBox(0, Axis.Y, ViewWidth, 1.0f, Brush, GridColor);
 
+		using namespace UE::Insights;
 		const FString LabelText = (Axis.Value == 0.0) ? TEXT("0") :
-								  (Axis.Value <= 1.0) ? FString::Printf(TEXT("%s (%.0f fps)"), *TimeUtils::FormatTimeAuto(Axis.Value), 1.0 / Axis.Value) :
-														TimeUtils::FormatTimeAuto(Axis.Value);
+								  (Axis.Value <= 1.0) ? FString::Printf(TEXT("%s (%.0f fps)"), *FormatTimeAuto(Axis.Value), 1.0 / Axis.Value) :
+														FormatTimeAuto(Axis.Value);
 
 		const float LabelTextWidth = static_cast<float>(FontMeasureService->Measure(LabelText, Font, FontScale).X / FontScale);
 		const float LabelX = bDrawVerticalAxisLabelsOnLeftSide ? 0.0f : ViewWidth - LabelTextWidth - 4.0f;
@@ -1005,7 +1016,7 @@ void SFrameTrack::DrawVerticalAxisGrid(FDrawContext& DrawContext, const FSlateBr
 
 ////////////////////////////////////////////////////////////////////////////////////////////////////
 
-void SFrameTrack::DrawHorizontalAxisGrid(FDrawContext& DrawContext, const FSlateBrush* Brush, const FSlateFontInfo& Font, bool bDrawBackgroundLayer) const
+void SFrameTrack::DrawHorizontalAxisGrid(UE::Insights::FDrawContext& DrawContext, const FSlateBrush* Brush, const FSlateFontInfo& Font, bool bDrawBackgroundLayer) const
 {
 	const FAxisViewportInt32& ViewportX = Viewport.GetHorizontalAxisViewport();
 
@@ -1654,7 +1665,7 @@ TSharedRef<SWidget> SFrameTrack::CreateUpperThresholdWidget()
 		SNew(STextBlock)
 		.Text_Lambda([this]
 		{
-			FString ThresholdTimeStr = TimeUtils::FormatTimeAuto(UpperThresholdTime, 2);
+			FString ThresholdTimeStr = UE::Insights::FormatTimeAuto(UpperThresholdTime, 2);
 			if (bShowUpperThresholdAsFps)
 			{
 				return FText::FromString(FString::Printf(TEXT("%s"), *ThresholdTimeStr));
@@ -1737,7 +1748,7 @@ TSharedRef<SWidget> SFrameTrack::CreateLowerThresholdWidget()
 		SNew(STextBlock)
 		.Text_Lambda([this]
 		{
-			FString ThresholdTimeStr = TimeUtils::FormatTimeAuto(LowerThresholdTime, 2);
+			FString ThresholdTimeStr = UE::Insights::FormatTimeAuto(LowerThresholdTime, 2);
 			if (bShowLowerThresholdAsFps)
 			{
 				return FText::FromString(FString::Printf(TEXT("%s"), *ThresholdTimeStr));
@@ -1862,7 +1873,7 @@ void SFrameTrack::CreateSelectedFrameMenu(FMenuBuilder& MenuBuilder)
 	FText Label;
 	if (SelectedSample.IsValid())
 	{
-		FText StartTimeText = FText::FromString(TimeUtils::FormatTimeAuto(SelectedSample.Sample->LargestFrameStartTime, 2));
+		FText StartTimeText = FText::FromString(UE::Insights::FormatTimeAuto(SelectedSample.Sample->LargestFrameStartTime, 2));
 		Label = FText::Format(LOCTEXT("ContextMenu_ScrollLogView_Fmt", "Scroll Log View (\u2192 {0})"), StartTimeText);
 	}
 	else
@@ -2068,16 +2079,20 @@ void SFrameTrack::ContextMenu_ScrollLogView_Execute()
 	{
 		return;
 	}
+
+	using namespace UE::Insights::TimingProfiler;
 	TSharedPtr<STimingProfilerWindow> TimingWindow = FTimingProfilerManager::Get()->GetProfilerWindow();
 	if (!TimingWindow.IsValid())
 	{
 		return;
 	}
+
 	TSharedPtr<SLogView> LogView = TimingWindow->GetLogView();
 	if (!LogView.IsValid())
 	{
 		return;
 	}
+
 	LogView->SelectLogMessageByClosestTime(SelectedSample.Sample->LargestFrameStartTime);
 }
 
@@ -2089,16 +2104,20 @@ bool SFrameTrack::ContextMenu_ScrollLogView_CanExecute()
 	{
 		return false;
 	}
+
+	using namespace UE::Insights::TimingProfiler;
 	TSharedPtr<STimingProfilerWindow> TimingWindow = FTimingProfilerManager::Get()->GetProfilerWindow();
 	if (!TimingWindow.IsValid())
 	{
 		return false;
 	}
+
 	TSharedPtr<SLogView> LogView = TimingWindow->GetLogView();
 	if (!LogView.IsValid())
 	{
 		return false;
 	}
+
 	return true;
 }
 
diff --git a/Engine/Source/Developer/TraceInsights/Private/Insights/Widgets/SFrameTrack.h b/Engine/Source/Developer/TraceInsights/Private/Insights/Widgets/SFrameTrack.h
--- a/Engine/Source/Developer/TraceInsights/Private/Insights/Widgets/SFrameTrack.h
+++ b/Engine/Source/Developer/TraceInsights/Private/Insights/Widgets/SFrameTrack.h
@@ -3,6 +3,7 @@
 #pragma once
 
 #include "CoreMinimal.h"
+
 #include "Input/CursorReply.h"
 #include "Input/Reply.h"
 #include "Layout/Geometry.h"
@@ -10,14 +11,17 @@
 #include "Widgets/DeclarativeSyntaxSupport.h"
 #include "Widgets/SCompoundWidget.h"
 
-// Insights
-#include "Insights/Common/FixedCircularBuffer.h"
+// TraceInsightsCore
+#include "InsightsCore/Common/FixedCircularBuffer.h"
+
+// TraceInsights
 #include "Insights/ViewModels/FrameTrackHelper.h"
 #include "Insights/ViewModels/FrameTrackViewport.h"
 
 class FMenuBuilder;
 class SScrollBar;
-class STimingView;
+
+namespace UE::Insights::TimingProfiler { class STimingView; }
 
 ////////////////////////////////////////////////////////////////////////////////////////////////////
 
@@ -142,8 +146,8 @@ protected:
 	TSharedPtr<FFrameTrackSeries> FindFrameStatsSeries(ETraceFrameType FrameType, uint32 TimerId) const;
 	void UpdateState();
 
-	void DrawHorizontalAxisGrid(FDrawContext& DrawContext, const FSlateBrush* Brush, const FSlateFontInfo& Font, bool bDrawBackgroundLayer) const;
-	void DrawVerticalAxisGrid(FDrawContext& DrawContext, const FSlateBrush* Brush, const FSlateFontInfo& Font) const;
+	void DrawHorizontalAxisGrid(UE::Insights::FDrawContext& DrawContext, const FSlateBrush* Brush, const FSlateFontInfo& Font, bool bDrawBackgroundLayer) const;
+	void DrawVerticalAxisGrid(UE::Insights::FDrawContext& DrawContext, const FSlateBrush* Brush, const FSlateFontInfo& Font) const;
 
 	FFrameTrackSampleRef GetSampleAtMousePosition(double X, double Y);
 	void SelectFrameAtMousePosition(double X, double Y, bool JoinCurrentSelection);
@@ -270,15 +274,15 @@ protected:
 	/** Cursor type. */
 	ECursorType CursorType = ECursorType::Default;
 
-	STimingView* RegisteredTimingView = nullptr; // for pointer comparison only; do not dereferentiate!
+	UE::Insights::TimingProfiler::STimingView* RegisteredTimingView = nullptr; // for pointer comparison only; do not dereferentiate!
 	FDelegateHandle OnTrackVisibilityChangedHandle;
 	FDelegateHandle OnTrackAddedHandle;
 	FDelegateHandle OnTrackRemovedHandle;
 
 	// Debug stats
 	int32 NumUpdatedFrames = 0;
-	TFixedCircularBuffer<uint64, 32> UpdateDurationHistory;
-	mutable TFixedCircularBuffer<uint64, 32> DrawDurationHistory;
-	mutable TFixedCircularBuffer<uint64, 32> OnPaintDurationHistory;
+	UE::Insights::TFixedCircularBuffer<uint64, 32> UpdateDurationHistory;
+	mutable UE::Insights::TFixedCircularBuffer<uint64, 32> DrawDurationHistory;
+	mutable UE::Insights::TFixedCircularBuffer<uint64, 32> OnPaintDurationHistory;
 	mutable uint64 LastOnPaintTime = 0;
 };
diff --git a/Engine/Source/Developer/TraceInsights/Private/Insights/Widgets/SLogView.cpp b/Engine/Source/Developer/TraceInsights/Private/Insights/Widgets/SLogView.cpp
--- a/Engine/Source/Developer/TraceInsights/Private/Insights/Widgets/SLogView.cpp
+++ b/Engine/Source/Developer/TraceInsights/Private/Insights/Widgets/SLogView.cpp
@@ -18,14 +18,18 @@
 #include "SlateOptMacros.h"
 #include "Styling/AppStyle.h"
 #include "Styling/StyleColors.h"
-#include "TraceServices/Model/Log.h"
 #include "Widgets/Images/SImage.h"
 #include "Widgets/Input/SComboButton.h"
 #include "Widgets/Layout/SScrollBox.h"
 #include "Widgets/Text/STextBlock.h"
 
-// Insights
-#include "Insights/Common/TimeUtils.h"
+// TraceServices
+#include "TraceServices/Model/Log.h"
+
+// TraceInsightsCore
+#include "InsightsCore/Common/TimeUtils.h"
+
+// TraceInsights
 #include "Insights/InsightsManager.h"
 #include "Insights/InsightsStyle.h"
 #include "Insights/Log.h"
@@ -258,17 +262,10 @@ public:
 				FLogMessageRecord& CacheEntry = ParentWidgetPin->GetCache().Get(LogMessagePin->GetIndex());
 				const double Time = CacheEntry.GetTime();
 
-				TSharedPtr<STimingProfilerWindow> Window = FTimingProfilerManager::Get()->GetProfilerWindow();
-				if (Window)
+				TSharedPtr<UE::Insights::TimingProfiler::STimingView> TimingView = ParentWidgetPin->GetTimingView();
+				if (TimingView && TimingView->IsTimeSelectedInclusive(Time))
 				{
-					TSharedPtr<STimingView> TimingView = Window->GetTimingView();
-					if (TimingView)
-					{
-						if (TimingView->IsTimeSelectedInclusive(Time))
-						{
-							return FSlateColor(FLinearColor(0.25f, 0.5f, 1.0f, 0.25f));
-						}
-					}
+					return FSlateColor(FLinearColor(0.25f, 0.5f, 1.0f, 0.25f));
 				}
 			}
 		}
@@ -303,23 +300,16 @@ public:
 				}
 			}
 
-			TSharedPtr<STimingProfilerWindow> Window = FTimingProfilerManager::Get()->GetProfilerWindow();
-			if (Window)
+			TSharedPtr<UE::Insights::TimingProfiler::STimingView> TimingView = ParentWidgetPin->GetTimingView();
+			if (TimingView && TimingView->IsTimeSelectedInclusive(Time))
 			{
-				TSharedPtr<STimingView> TimingView = Window->GetTimingView();
-				if (TimingView)
+				if (IsSelected)
 				{
-					if (TimingView->IsTimeSelectedInclusive(Time))
-					{
-						if (IsSelected)
-						{
-							return FSlateColor(FLinearColor(0.0f, 0.05f, 0.2f, 1.0f));
-						}
-						else
-						{
-							return FSlateColor(FLinearColor(0.4f, 0.8f, 1.6f, 1.0f));
-						}
-					}
+					return FSlateColor(FLinearColor(0.0f, 0.05f, 0.2f, 1.0f));
+				}
+				else
+				{
+					return FSlateColor(FLinearColor(0.4f, 0.8f, 1.6f, 1.0f));
 				}
 			}
 		}
@@ -529,6 +519,7 @@ SLogView::SLogView()
 SLogView::~SLogView()
 {
 	// Remove ourselves from the profiler manager.
+	//using namespace UE::Insights::TimingProfiler;
 	//if (FTimingProfilerManager::Get().IsValid())
 	//{
 	//	//TODO: FTimimgProfilerManager::Get()->OnRequestLogViewUpdate().RemoveAll(this);
@@ -1074,23 +1065,19 @@ void SLogView::OnSelectedLogMessageChanged(TSharedPtr<FLogMessage> LogMessage)
 {
 	if (LogMessage.IsValid())
 	{
-		TSharedPtr<STimingProfilerWindow> Window = FTimingProfilerManager::Get()->GetProfilerWindow();
-		if (Window)
+		TSharedPtr<UE::Insights::TimingProfiler::STimingView> TimingView = GetTimingView();
+		if (TimingView)
 		{
-			TSharedPtr<STimingView> TimingView = Window->GetTimingView();
-			if (TimingView)
-			{
-				const double Time = Cache.Get(LogMessage->GetIndex()).GetTime();
+			const double Time = Cache.Get(LogMessage->GetIndex()).GetTime();
 
-				if (FSlateApplication::Get().GetModifierKeys().IsShiftDown())
-				{
-					TimingView->SelectToTimeMarker(Time);
-				}
-				else
-				{
-					TimingView->SetAutoScroll(false);
-					TimingView->SetAndCenterOnTimeMarker(Time);
-				}
+			if (FSlateApplication::Get().GetModifierKeys().IsShiftDown())
+			{
+				TimingView->SelectToTimeMarker(Time);
+			}
+			else
+			{
+				TimingView->SetAutoScroll(false);
+				TimingView->SetAndCenterOnTimeMarker(Time);
 			}
 		}
 	}
@@ -1670,9 +1657,10 @@ FReply SLogView::OnKeyDown(const FGeometry& MyGeometry, const FKeyEvent& InKeyEv
 
 void SLogView::AppendFormatMessageDetailed(const FLogMessageRecord& Log, TStringBuilderBase<TCHAR>& InOutStringBuilder) const
 {
+	using namespace UE::Insights;
 	InOutStringBuilder.Appendf(TEXT("Index=%d"), Log.GetIndex());
 	InOutStringBuilder.Append(TEXT("\nTime="));
-	InOutStringBuilder.Append(TimeUtils::FormatTimeHMS(Log.GetTime(), TimeUtils::Microsecond));
+	InOutStringBuilder.Append(FormatTimeHMS(Log.GetTime(), FTimeValue::Microsecond));
 	InOutStringBuilder.Append(TEXT("\nVerbosity="));
 	InOutStringBuilder.Append(::ToString(Log.GetVerbosity()));
 	InOutStringBuilder.Append(TEXT("\nCategory="));
@@ -1709,7 +1697,7 @@ void SLogView::AppendFormatMessageDelimited(const FLogMessageRecord& Log, TStrin
 {
 	InOutStringBuilder.Appendf(TEXT("%d"), Log.GetIndex());
 	InOutStringBuilder.AppendChar(Separator);
-	InOutStringBuilder.Append(TimeUtils::FormatTimeHMS(Log.GetTime(), TimeUtils::Microsecond));
+	InOutStringBuilder.Append(UE::Insights::FormatTimeHMS(Log.GetTime(), UE::Insights::FTimeValue::Microsecond));
 	InOutStringBuilder.AppendChar(Separator);
 	InOutStringBuilder.Append(::ToString(Log.GetVerbosity()));
 	InOutStringBuilder.AppendChar(Separator);
@@ -1787,8 +1775,7 @@ bool SLogView::CanCopyRange() const
 		return false;
 	}
 
-	TSharedPtr<STimingProfilerWindow> Window = FTimingProfilerManager::Get()->GetProfilerWindow();
-	TSharedPtr<STimingView> TimingView = Window ? Window->GetTimingView() : nullptr;
+	TSharedPtr<UE::Insights::TimingProfiler::STimingView> TimingView = GetTimingView();
 	if (!TimingView)
 	{
 		return false;
@@ -1808,8 +1795,7 @@ void SLogView::CopyRange() const
 		return;
 	}
 
-	TSharedPtr<STimingProfilerWindow> Window = FTimingProfilerManager::Get()->GetProfilerWindow();
-	TSharedPtr<STimingView> TimingView = Window ? Window->GetTimingView() : nullptr;
+	TSharedPtr<UE::Insights::TimingProfiler::STimingView> TimingView = GetTimingView();
 	if (!TimingView)
 	{
 		return;
@@ -1940,8 +1926,7 @@ void SLogView::SaveLogsToFile(bool bSaveLogsInSelectedRangeOnly) const
 	double SelectionEndTime = 0.0;
 	if (bSaveLogsInSelectedRangeOnly)
 	{
-		TSharedPtr<STimingProfilerWindow> Window = FTimingProfilerManager::Get()->GetProfilerWindow();
-		TSharedPtr<STimingView> TimingView = Window ? Window->GetTimingView() : nullptr;
+		TSharedPtr<UE::Insights::TimingProfiler::STimingView> TimingView = GetTimingView();
 		if (!TimingView)
 		{
 			return;
@@ -1993,7 +1978,7 @@ void SLogView::SaveLogsToFile(bool bSaveLogsInSelectedRangeOnly) const
 		return;
 	}
 
-	FStopwatch Stopwatch;
+	UE::Insights::FStopwatch Stopwatch;
 	Stopwatch.Start();
 
 	UTF16CHAR BOM = UNICODE_BOM;
@@ -2076,4 +2061,13 @@ void SLogView::OpenSource() const
 
 ////////////////////////////////////////////////////////////////////////////////////////////////////
 
+TSharedPtr<UE::Insights::TimingProfiler::STimingView> SLogView::GetTimingView() const
+{
+	using namespace UE::Insights::TimingProfiler;
+	TSharedPtr<STimingProfilerWindow> ProfilerWindow = FTimingProfilerManager::Get()->GetProfilerWindow();
+	return ProfilerWindow.IsValid() ? ProfilerWindow->GetTimingView() : nullptr;
+}
+
+////////////////////////////////////////////////////////////////////////////////////////////////////
+
 #undef LOCTEXT_NAMESPACE
diff --git a/Engine/Source/Developer/TraceInsights/Private/Insights/Widgets/SLogView.h b/Engine/Source/Developer/TraceInsights/Private/Insights/Widgets/SLogView.h
--- a/Engine/Source/Developer/TraceInsights/Private/Insights/Widgets/SLogView.h
+++ b/Engine/Source/Developer/TraceInsights/Private/Insights/Widgets/SLogView.h
@@ -12,13 +12,17 @@
 #include "Widgets/SCompoundWidget.h"
 #include "Widgets/Views/SListView.h"
 
-// Insights
-#include "Insights/Common/Stopwatch.h"
+// TraceInsightsCore
+#include "InsightsCore/Common/Stopwatch.h"
+
+// TraceInsights
 #include "Insights/ViewModels/LogFilter.h"
 #include "Insights/ViewModels/LogMessage.h"
 
 class FMenuBuilder;
 
+namespace UE::Insights::TimingProfiler { class STimingView; }
+
 ////////////////////////////////////////////////////////////////////////////////////////////////////
 
 namespace LogViewColumns
@@ -79,6 +83,8 @@ public:
 
 	bool IsFilteringAsyncTaskCancelRequested() const { return bIsFilteringAsyncTaskCancelRequested; }
 
+	TSharedPtr<UE::Insights::TimingProfiler::STimingView> GetTimingView() const;
+
 protected:
 	/** Generate a new list view row. */
 	TSharedRef<ITableRow> OnGenerateRow(TSharedPtr<FLogMessage> InLogMessage, const TSharedRef<STableViewBase>& OwnerTable);
@@ -171,7 +177,7 @@ protected:
 	mutable volatile bool bIsFilteringAsyncTaskCancelRequested; // true if we want the async task to finish asap
 
 	/** Stopwatch used to measure how long it takes to filter the message list. */
-	mutable FStopwatch FilteringStopwatch;
+	mutable UE::Insights::FStopwatch FilteringStopwatch;
 
 	/**
 	 * Total number of log categories (including duplicates; last value read from LogProvider).
@@ -195,7 +201,7 @@ protected:
 	bool bIsDirty;
 
 	/** Stopwatch used to measure the time since the list of messages has become dirty. */
-	mutable FStopwatch DirtyStopwatch;
+	mutable UE::Insights::FStopwatch DirtyStopwatch;
 
 	/** Stats */
 	FText StatsText;
diff --git a/Engine/Source/Developer/TraceInsights/Private/Insights/Widgets/SMajorTabWindow.cpp b/Engine/Source/Developer/TraceInsights/Private/Insights/Widgets/SMajorTabWindow.cpp
--- a/Engine/Source/Developer/TraceInsights/Private/Insights/Widgets/SMajorTabWindow.cpp
+++ b/Engine/Source/Developer/TraceInsights/Private/Insights/Widgets/SMajorTabWindow.cpp
@@ -16,14 +16,13 @@
 #include "Widgets/SBoxPanel.h"
 #include "Widgets/Text/STextBlock.h"
 
-
 #if WITH_EDITOR
 	#include "EngineAnalytics.h"
 	#include "AnalyticsEventAttribute.h"
 	#include "Interfaces/IAnalyticsProvider.h"
 #endif // WITH_EDITOR
 
-// Insights
+// TraceInsights
 #include "Insights/Common/InsightsMenuBuilder.h"
 #include "Insights/InsightsManager.h"
 #include "Insights/InsightsStyle.h"
diff --git a/Engine/Source/Developer/TraceInsights/Private/Insights/Widgets/SModulesView.cpp b/Engine/Source/Developer/TraceInsights/Private/Insights/Widgets/SModulesView.cpp
--- a/Engine/Source/Developer/TraceInsights/Private/Insights/Widgets/SModulesView.cpp
+++ b/Engine/Source/Developer/TraceInsights/Private/Insights/Widgets/SModulesView.cpp
@@ -6,20 +6,22 @@
 #include "Framework/MultiBox/MultiBoxBuilder.h"
 #include "SlateOptMacros.h"
 #include "Textures/SlateIcon.h"
-#include "TraceServices/Model/Modules.h"
 #include "Widgets/Layout/SScrollBox.h"
 #include "Widgets/Text/STextBlock.h"
 
-// Insights
+// TraceServices
+#include "TraceServices/Model/Modules.h"
+
+// TraceInsights
 #include "DesktopPlatformModule.h"
 #include "Insights/InsightsManager.h"
 #include "Insights/InsightsStyle.h"
 
 ////////////////////////////////////////////////////////////////////////////////////////////////////
 
-#define LOCTEXT_NAMESPACE "SModulesView"
+#define LOCTEXT_NAMESPACE "UE::Insights::SModulesView"
 
-namespace Insights
+namespace UE::Insights
 {
 
 ////////////////////////////////////////////////////////////////////////////////////////////////////
@@ -1160,6 +1162,6 @@ void SModulesView::UpdateSorting()
 
 ////////////////////////////////////////////////////////////////////////////////////////////////////
 
-} //namespace Insights
+} //namespace UE::Insights
 
 #undef LOCTEXT_NAMESPACE
diff --git a/Engine/Source/Developer/TraceInsights/Private/Insights/Widgets/SModulesView.h b/Engine/Source/Developer/TraceInsights/Private/Insights/Widgets/SModulesView.h
--- a/Engine/Source/Developer/TraceInsights/Private/Insights/Widgets/SModulesView.h
+++ b/Engine/Source/Developer/TraceInsights/Private/Insights/Widgets/SModulesView.h
@@ -2,14 +2,17 @@
 
 #pragma once
 
+#include "CoreTypes.h"
+
 #include "SlateFwd.h"
+#include "UObject/NameTypes.h"
 #include "Widgets/DeclarativeSyntaxSupport.h"
 #include "Widgets/SCompoundWidget.h"
 #include "Widgets/Views/SListView.h"
 
 class FMenuBuilder;
 
-namespace Insights
+namespace UE::Insights
 {
 
 ////////////////////////////////////////////////////////////////////////////////////////////////////
@@ -119,4 +122,4 @@ protected:
 
 ////////////////////////////////////////////////////////////////////////////////////////////////////
 
-} // namespace Insights
+} // namespace UE::Insights
diff --git a/Engine/Source/Developer/TraceInsights/Private/Insights/Widgets/SQuickFind.cpp b/Engine/Source/Developer/TraceInsights/Private/Insights/Widgets/SQuickFind.cpp
--- a/Engine/Source/Developer/TraceInsights/Private/Insights/Widgets/SQuickFind.cpp
+++ b/Engine/Source/Developer/TraceInsights/Private/Insights/Widgets/SQuickFind.cpp
@@ -5,18 +5,20 @@
 #include "SlateOptMacros.h"
 #include "Widgets/DeclarativeSyntaxSupport.h"
 #include "Widgets/Docking/SDockTab.h"
+#include "Widgets/Images/SImage.h"
 #include "Widgets/Input/SButton.h"
 #include "Widgets/Layout/SScrollBox.h"
-#include "Widgets/Images/SImage.h"
 #include "Widgets/SBoxPanel.h"
 #include "Widgets/SOverlay.h"
 #include "Widgets/Text/STextBlock.h"
 #include "Widgets/Views/STableViewBase.h"
 
-// Insights
+// TraceInsightsCore
+#include "InsightsCore/Filter/ViewModels/FilterConfigurator.h"
+#include "InsightsCore/Filter/Widgets/SFilterConfigurator.h"
+
+// TraceInsights
 #include "Insights/InsightsStyle.h"
-#include "Insights/Widgets/SFilterConfigurator.h"
-#include "Insights/ViewModels/FilterConfigurator.h"
 #include "Insights/ViewModels/QuickFind.h"
 
 #define LOCTEXT_NAMESPACE "SQuickFind"
@@ -59,7 +61,7 @@ void SQuickFind::Construct(const FArguments& InArgs, TSharedPtr<FQuickFind> InQu
 		.FillHeight(1.0f)
 		.Padding(0.0f, 0.0f, 0.0f, 0.0f)
 		[
-			SAssignNew(FilterConfigurator, SFilterConfigurator, InQuickFindViewModel->GetFilterConfigurator())
+			SAssignNew(FilterConfigurator, UE::Insights::SFilterConfigurator, InQuickFindViewModel->GetFilterConfigurator())
 		]
 
 		+ SVerticalBox::Slot()
diff --git a/Engine/Source/Developer/TraceInsights/Private/Insights/Widgets/SQuickFind.h b/Engine/Source/Developer/TraceInsights/Private/Insights/Widgets/SQuickFind.h
--- a/Engine/Source/Developer/TraceInsights/Private/Insights/Widgets/SQuickFind.h
+++ b/Engine/Source/Developer/TraceInsights/Private/Insights/Widgets/SQuickFind.h
@@ -3,21 +3,24 @@
 #pragma once
 
 #include "CoreMinimal.h"
+
 #include "Widgets/SCompoundWidget.h"
 #include "Widgets/SWidget.h"
 #include "Widgets/Views/STreeView.h"
 
-// Insights
-#include "Insights/ViewModels/Filters.h"
-#include "Insights/ViewModels/FilterConfiguratorNode.h"
+// TraceInsightsCore
+#include "InsightsCore/Filter/ViewModels/FilterConfiguratorNode.h"
+#include "InsightsCore/Filter/ViewModels/Filters.h"
 
 class SDockTab;
 
+namespace UE::Insights { class SFilterConfigurator; }
+
 namespace Insights
 {
 
-class SFilterConfigurator;
 class FQuickFind;
+
 ////////////////////////////////////////////////////////////////////////////////////////////////////
 
 /**
@@ -58,7 +61,7 @@ private:
 	FReply ClearFilters_OnClicked();
 
 private:
-	TSharedPtr<SFilterConfigurator> FilterConfigurator;
+	TSharedPtr<UE::Insights::SFilterConfigurator> FilterConfigurator;
 
 	TSharedPtr<FQuickFind> QuickFindViewModel;
 
diff --git a/Engine/Source/Developer/TraceInsights/Private/Insights/Widgets/SSessionInfoWindow.cpp b/Engine/Source/Developer/TraceInsights/Private/Insights/Widgets/SSessionInfoWindow.cpp
--- a/Engine/Source/Developer/TraceInsights/Private/Insights/Widgets/SSessionInfoWindow.cpp
+++ b/Engine/Source/Developer/TraceInsights/Private/Insights/Widgets/SSessionInfoWindow.cpp
@@ -3,11 +3,9 @@
 #include "SSessionInfoWindow.h"
 
 #include "Framework/MultiBox/MultiBoxBuilder.h"
+#include "Misc/Paths.h"
 #include "SlateOptMacros.h"
 #include "Styling/AppStyle.h"
-#include "TraceServices/Model/Diagnostics.h"
-#include "TraceServices/ModuleService.h"
-#include "Misc/Paths.h"
 #include "Widgets/Input/SEditableTextBox.h"
 #include "Widgets/Input/SMultiLineEditableTextBox.h"
 #include "Widgets/Layout/SScrollBox.h"
@@ -20,9 +18,15 @@
 	#include "Interfaces/IAnalyticsProvider.h"
 #endif // WITH_EDITOR
 
-// Insights
+// TraceServices
+#include "TraceServices/Model/Diagnostics.h"
+#include "TraceServices/ModuleService.h"
+
+// TraceInsightsCore
+#include "InsightsCore/Common/TimeUtils.h"
+
+// TraceInsights
 #include "Insights/Common/InsightsMenuBuilder.h"
-#include "Insights/Common/TimeUtils.h"
 #include "Insights/InsightsManager.h"
 #include "Insights/InsightsStyle.h"
 #include "Insights/Version.h"
@@ -565,10 +569,12 @@ FText SSessionInfoWindow::GetStatusText() const
 	FormattingOptions.MaximumFractionalDigits = 1;
 
 	const int32 NumDigits = InsightsManager->IsAnalysisComplete() ? 2 : 0;
+
+	using namespace UE::Insights;
 	FText Status = FText::Format(LOCTEXT("StatusFmt", "{0}\nSession Duration: {1}\nAnalyzed in {2} at {3}X speed."),
 		InsightsManager->IsAnalysisComplete() ? FText::FromString(FString(TEXT("ANALYSIS COMPLETED."))) : FText::FromString(FString(TEXT("ANALYZING..."))),
-		FText::FromString(TimeUtils::FormatTimeAuto(InsightsManager->GetSessionDuration(), NumDigits)),
-		FText::FromString(TimeUtils::FormatTimeAuto(InsightsManager->GetAnalysisDuration(), NumDigits)),
+		FText::FromString(FormatTimeAuto(InsightsManager->GetSessionDuration(), NumDigits)),
+		FText::FromString(FormatTimeAuto(InsightsManager->GetAnalysisDuration(), NumDigits)),
 		FText::AsNumber(InsightsManager->GetAnalysisSpeedFactor(), &FormattingOptions));
 
 	return Status;
diff --git a/Engine/Source/Developer/TraceInsights/Private/Insights/Widgets/SSessionInfoWindow.h b/Engine/Source/Developer/TraceInsights/Private/Insights/Widgets/SSessionInfoWindow.h
--- a/Engine/Source/Developer/TraceInsights/Private/Insights/Widgets/SSessionInfoWindow.h
+++ b/Engine/Source/Developer/TraceInsights/Private/Insights/Widgets/SSessionInfoWindow.h
@@ -9,7 +9,6 @@
 #include "Layout/Visibility.h"
 #include "Misc/Guid.h"
 #include "SlateFwd.h"
-#include "TraceServices/ModuleService.h"
 #include "Widgets/DeclarativeSyntaxSupport.h"
 #include "Widgets/Docking/SDockTab.h"
 #include "Widgets/Layout/SSplitter.h"
@@ -17,7 +16,10 @@
 #include "Widgets/SOverlay.h"
 #include "Widgets/Views/SListView.h"
 
-// Insights
+// TraceServices
+#include "TraceServices/ModuleService.h"
+
+// TraceInsights
 #include "Insights/InsightsManager.h"
 
 ////////////////////////////////////////////////////////////////////////////////////////////////////
diff --git a/Engine/Source/Developer/TraceInsights/Private/Insights/Widgets/SStartPageWindow.cpp b/Engine/Source/Developer/TraceInsights/Private/Insights/Widgets/SStartPageWindow.cpp
--- a/Engine/Source/Developer/TraceInsights/Private/Insights/Widgets/SStartPageWindow.cpp
+++ b/Engine/Source/Developer/TraceInsights/Private/Insights/Widgets/SStartPageWindow.cpp
@@ -18,8 +18,6 @@
 #include "SocketSubsystem.h"
 #include "Styling/AppStyle.h"
 #include "Styling/StyleColors.h"
-#include "Trace/ControlClient.h"
-#include "Trace/StoreClient.h"
 #include "Widgets/Colors/SColorBlock.h"
 #include "Widgets/Images/SImage.h"
 #include "Widgets/Input/SButton.h"
@@ -45,9 +43,15 @@
 	#include "Interfaces/IAnalyticsProvider.h"
 #endif // WITH_EDITOR
 
-// Insights
+// TraceAnalysis
+#include "Trace/ControlClient.h"
+#include "Trace/StoreClient.h"
+
+// TraceInsightsCore
+#include "InsightsCore/Common/Stopwatch.h"
+
+// TraceInsights
 #include "Insights/Common/MessageDialogUtils.h"
-#include "Insights/Common/Stopwatch.h"
 #include "Insights/ImportTool/TableImportTool.h"
 #include "Insights/InsightsManager.h"
 #include "Insights/InsightsStyle.h"
@@ -2520,7 +2524,7 @@ FSlateColor STraceStoreWindow::GetColorByPath(const FString& Uri)
 
 void STraceStoreWindow::RefreshTraceList()
 {
-	FStopwatch StopwatchTotal;
+	UE::Insights::FStopwatch StopwatchTotal;
 	StopwatchTotal.Start();
 
 	int32 AddedTraces = 0;
diff --git a/Engine/Source/Developer/TraceInsights/Private/Insights/Widgets/SStatsTableCell.cpp b/Engine/Source/Developer/TraceInsights/Private/Insights/Widgets/SStatsTableCell.cpp
--- a/Engine/Source/Developer/TraceInsights/Private/Insights/Widgets/SStatsTableCell.cpp
+++ b/Engine/Source/Developer/TraceInsights/Private/Insights/Widgets/SStatsTableCell.cpp
@@ -11,10 +11,12 @@
 #include "Widgets/Text/STextBlock.h"
 #include "Widgets/Views/SExpanderArrow.h"
 
-// Insights
+// TraceInsightsCore
+#include "InsightsCore/Table/ViewModels/Table.h"
+#include "InsightsCore/Table/ViewModels/TableColumn.h"
+
+// TraceInsights
 #include "Insights/InsightsStyle.h"
-#include "Insights/Table/ViewModels/Table.h"
-#include "Insights/Table/ViewModels/TableColumn.h"
 #include "Insights/Widgets/SStatsTableRow.h"
 
 #define LOCTEXT_NAMESPACE "SStatsView"
diff --git a/Engine/Source/Developer/TraceInsights/Private/Insights/Widgets/SStatsTableCell.h b/Engine/Source/Developer/TraceInsights/Private/Insights/Widgets/SStatsTableCell.h
--- a/Engine/Source/Developer/TraceInsights/Private/Insights/Widgets/SStatsTableCell.h
+++ b/Engine/Source/Developer/TraceInsights/Private/Insights/Widgets/SStatsTableCell.h
@@ -3,21 +3,22 @@
 #pragma once
 
 #include "CoreMinimal.h"
+
 #include "Widgets/DeclarativeSyntaxSupport.h"
 #include "Widgets/SCompoundWidget.h"
 
-// Insights
+// TraceInsights
 #include "Insights/ViewModels/StatsNode.h"
 
 class ITableRow;
 
-namespace Insights
+namespace UE::Insights
 {
 	class FTable;
 	class FTableColumn;
 }
 
-DECLARE_DELEGATE_ThreeParams(FSetHoveredStatsTableCell, TSharedPtr<Insights::FTable> /*TablePtr*/, TSharedPtr<Insights::FTableColumn> /*ColumnPtr*/, FStatsNodePtr /*StatsNodePtr*/);
+DECLARE_DELEGATE_ThreeParams(FSetHoveredStatsTableCell, TSharedPtr<UE::Insights::FTable> /*TablePtr*/, TSharedPtr<UE::Insights::FTableColumn> /*ColumnPtr*/, FStatsNodePtr /*StatsNodePtr*/);
 
 class SStatsTableCell : public SCompoundWidget
 {
@@ -25,8 +26,8 @@ public:
 	SLATE_BEGIN_ARGS(SStatsTableCell) {}
 		SLATE_EVENT(FSetHoveredStatsTableCell, OnSetHoveredCell)
 		SLATE_ATTRIBUTE(FText, HighlightText)
-		SLATE_ARGUMENT(TSharedPtr<Insights::FTable>, TablePtr)
-		SLATE_ARGUMENT(TSharedPtr<Insights::FTableColumn>, ColumnPtr)
+		SLATE_ARGUMENT(TSharedPtr<UE::Insights::FTable>, TablePtr)
+		SLATE_ARGUMENT(TSharedPtr<UE::Insights::FTableColumn>, ColumnPtr)
 		SLATE_ARGUMENT(FStatsNodePtr, StatsNodePtr)
 		SLATE_ARGUMENT(bool, IsNameColumn)
 	SLATE_END_ARGS()
@@ -169,10 +170,10 @@ protected:
 
 protected:
 	/** A shared pointer to the table view model. */
-	TSharedPtr<Insights::FTable> TablePtr;
+	TSharedPtr<UE::Insights::FTable> TablePtr;
 
 	/** A shared pointer to the table column view model. */
-	TSharedPtr<Insights::FTableColumn> ColumnPtr;
+	TSharedPtr<UE::Insights::FTableColumn> ColumnPtr;
 
 	/** A shared pointer to the stats node. */
 	FStatsNodePtr StatsNodePtr;
diff --git a/Engine/Source/Developer/TraceInsights/Private/Insights/Widgets/SStatsTableRow.cpp b/Engine/Source/Developer/TraceInsights/Private/Insights/Widgets/SStatsTableRow.cpp
--- a/Engine/Source/Developer/TraceInsights/Private/Insights/Widgets/SStatsTableRow.cpp
+++ b/Engine/Source/Developer/TraceInsights/Private/Insights/Widgets/SStatsTableRow.cpp
@@ -3,15 +3,17 @@
 #include "SStatsTableRow.h"
 
 #include "SlateOptMacros.h"
+#include "Widgets/Images/SImage.h"
 #include "Widgets/SOverlay.h"
 #include "Widgets/SToolTip.h"
-#include "Widgets/Images/SImage.h"
 
-// Insights
-#include "Insights/Common/TimeUtils.h"
+// TraceInsightsCore
+#include "InsightsCore/Common/TimeUtils.h"
+#include "InsightsCore/Table/ViewModels/Table.h"
+#include "InsightsCore/Table/ViewModels/TableColumn.h"
+
+// TraceInsights
 #include "Insights/InsightsStyle.h"
-#include "Insights/Table/ViewModels/Table.h"
-#include "Insights/Table/ViewModels/TableColumn.h"
 #include "Insights/Widgets/SStatsTableCell.h"
 #include "Insights/Widgets/SStatsViewTooltip.h"
 
@@ -45,7 +47,7 @@ void SStatsTableRow::Construct(const FArguments& InArgs, const TSharedRef<STable
 
 TSharedRef<SWidget> SStatsTableRow::GenerateWidgetForColumn(const FName& ColumnId)
 {
-	TSharedPtr<Insights::FTableColumn> ColumnPtr = TablePtr->FindColumnChecked(ColumnId);
+	TSharedPtr<UE::Insights::FTableColumn> ColumnPtr = TablePtr->FindColumnChecked(ColumnId);
 
 	return
 		SNew(SOverlay)
@@ -126,10 +128,11 @@ FSlateColor SStatsTableRow::GetBackgroundColorAndOpacity() const
 
 FSlateColor SStatsTableRow::GetBackgroundColorAndOpacity(double Time) const
 {
-	const FLinearColor Color =	Time > TimeUtils::Second      ? FLinearColor(0.3f, 0.0f, 0.0f, 1.0f) :
-								Time > TimeUtils::Milisecond  ? FLinearColor(0.3f, 0.1f, 0.0f, 1.0f) :
-								Time > TimeUtils::Microsecond ? FLinearColor(0.0f, 0.1f, 0.0f, 1.0f) :
-								                                FLinearColor(0.0f, 0.0f, 0.0f, 1.0f);
+	using namespace UE::Insights;
+	const FLinearColor Color =	Time > FTimeValue::Second      ? FLinearColor(0.3f, 0.0f, 0.0f, 1.0f) :
+								Time > FTimeValue::Millisecond ? FLinearColor(0.3f, 0.1f, 0.0f, 1.0f) :
+								Time > FTimeValue::Microsecond ? FLinearColor(0.0f, 0.1f, 0.0f, 1.0f) :
+								                                 FLinearColor(0.0f, 0.0f, 0.0f, 1.0f);
 	return Color;
 }
 
@@ -192,7 +195,7 @@ EVisibility SStatsTableRow::IsColumnVisible(const FName ColumnId) const
 
 ////////////////////////////////////////////////////////////////////////////////////////////////////
 
-void SStatsTableRow::OnSetHoveredCell(TSharedPtr<Insights::FTable> InTablePtr, TSharedPtr<Insights::FTableColumn> InColumnPtr, FStatsNodePtr InStatsNodePtr)
+void SStatsTableRow::OnSetHoveredCell(TSharedPtr<UE::Insights::FTable> InTablePtr, TSharedPtr<UE::Insights::FTableColumn> InColumnPtr, FStatsNodePtr InStatsNodePtr)
 {
 	SetHoveredCellDelegate.ExecuteIfBound(InTablePtr, InColumnPtr, InStatsNodePtr);
 }
diff --git a/Engine/Source/Developer/TraceInsights/Private/Insights/Widgets/SStatsTableRow.h b/Engine/Source/Developer/TraceInsights/Private/Insights/Widgets/SStatsTableRow.h
--- a/Engine/Source/Developer/TraceInsights/Private/Insights/Widgets/SStatsTableRow.h
+++ b/Engine/Source/Developer/TraceInsights/Private/Insights/Widgets/SStatsTableRow.h
@@ -3,19 +3,20 @@
 #pragma once
 
 #include "CoreMinimal.h"
+
 #include "Widgets/DeclarativeSyntaxSupport.h"
 #include "Widgets/Views/SHeaderRow.h"
 #include "Widgets/Views/STableViewBase.h"
 #include "Widgets/Views/STableRow.h"
 #include "Widgets/Views/STreeView.h"
 
-// Insights
+// TraceInsights
 #include "Insights/ViewModels/StatsNodeHelper.h"
 
 class IToolTip;
 class SStatsCounterTableRowToolTip;
 
-namespace Insights
+namespace UE::Insights
 {
 	class FTable;
 	class FTableColumn;
@@ -24,7 +25,7 @@ namespace Insights
 DECLARE_DELEGATE_RetVal_OneParam(bool, FStatsNodeShouldBeEnabledDelegate, FStatsNodePtr /*NodePtr*/);
 DECLARE_DELEGATE_RetVal_OneParam(bool, FIsColumnVisibleDelegate, const FName /*ColumnId*/);
 DECLARE_DELEGATE_RetVal_OneParam(EHorizontalAlignment, FGetColumnOutlineHAlignmentDelegate, const FName /*ColumnId*/);
-DECLARE_DELEGATE_ThreeParams(FSetHoveredStatsTableCell, TSharedPtr<Insights::FTable> /*TablePtr*/, TSharedPtr<Insights::FTableColumn> /*ColumnPtr*/, FStatsNodePtr /*StatsNodePtr*/);
+DECLARE_DELEGATE_ThreeParams(FSetHoveredStatsTableCell, TSharedPtr<UE::Insights::FTable> /*TablePtr*/, TSharedPtr<UE::Insights::FTableColumn> /*ColumnPtr*/, FStatsNodePtr /*StatsNodePtr*/);
 
 /** Widget that represents a table row in the Stats Counters' tree control. Generates widgets for each column on demand. */
 class SStatsTableRow : public SMultiColumnTableRow<FStatsNodePtr>
@@ -37,7 +38,7 @@ public:
 		SLATE_EVENT(FSetHoveredStatsTableCell, OnSetHoveredCell)
 		SLATE_ATTRIBUTE(FText, HighlightText)
 		SLATE_ATTRIBUTE(FName, HighlightedNodeName)
-		SLATE_ARGUMENT(TSharedPtr<Insights::FTable>, TablePtr)
+		SLATE_ARGUMENT(TSharedPtr<UE::Insights::FTable>, TablePtr)
 		SLATE_ARGUMENT(FStatsNodePtr, StatsNodePtr)
 	SLATE_END_ARGS()
 
@@ -72,11 +73,11 @@ protected:
 	const FSlateBrush* GetOutlineBrush(const FName ColumnId) const;
 	bool HandleShouldBeEnabled() const;
 	EVisibility IsColumnVisible(const FName ColumnId) const;
-	void OnSetHoveredCell(TSharedPtr<Insights::FTable> InTablePtr, TSharedPtr<Insights::FTableColumn> InColumnPtr, FStatsNodePtr InStatsNodePtr);
+	void OnSetHoveredCell(TSharedPtr<UE::Insights::FTable> InTablePtr, TSharedPtr<UE::Insights::FTableColumn> InColumnPtr, FStatsNodePtr InStatsNodePtr);
 
 protected:
 	/** A shared pointer to the table view model. */
-	TSharedPtr<Insights::FTable> TablePtr;
+	TSharedPtr<UE::Insights::FTable> TablePtr;
 
 	/** Data context for this table row. */
 	FStatsNodePtr StatsNodePtr;
diff --git a/Engine/Source/Developer/TraceInsights/Private/Insights/Widgets/SStatsView.cpp b/Engine/Source/Developer/TraceInsights/Private/Insights/Widgets/SStatsView.cpp
--- a/Engine/Source/Developer/TraceInsights/Private/Insights/Widgets/SStatsView.cpp
+++ b/Engine/Source/Developer/TraceInsights/Private/Insights/Widgets/SStatsView.cpp
@@ -11,23 +11,28 @@
 #include "Logging/MessageLog.h"
 #include "SlateOptMacros.h"
 #include "Styling/AppStyle.h"
-#include "TraceServices/Model/Counters.h"
+#include "Widgets/Images/SImage.h"
 #include "Widgets/Input/SCheckBox.h"
-#include "Widgets/Layout/SScrollBox.h"
 #include "Widgets/Input/SSearchBox.h"
-#include "Widgets/Images/SImage.h"
 #include "Widgets/Layout/SGridPanel.h"
+#include "Widgets/Layout/SScrollBox.h"
 #include "Widgets/Layout/SSeparator.h"
 #include "Widgets/SToolTip.h"
 #include "Widgets/Views/STableViewBase.h"
 
-// Insights
-#include "Insights/Common/Stopwatch.h"
-#include "Insights/Common/TimeUtils.h"
+// TraceServices
+#include "TraceServices/Model/Counters.h"
+
+// TraceInsightsCore
+#include "InsightsCore/Common/Stopwatch.h"
+#include "InsightsCore/Common/TimeUtils.h"
+#include "InsightsCore/Table/ViewModels/Table.h"
+#include "InsightsCore/Table/ViewModels/TableColumn.h"
+#include "InsightsCore/Table/Widgets/SAsyncOperationStatus.h"
+
+// TraceInsights
 #include "Insights/InsightsStyle.h"
 #include "Insights/Log.h"
-#include "Insights/Table/ViewModels/Table.h"
-#include "Insights/Table/ViewModels/TableColumn.h"
 #include "Insights/TimingProfilerCommon.h"
 #include "Insights/TimingProfilerManager.h"
 #include "Insights/ViewModels/CounterAggregation.h"
@@ -35,9 +40,8 @@
 #include "Insights/ViewModels/StatsViewColumnFactory.h"
 #include "Insights/ViewModels/TimingExporter.h"
 #include "Insights/ViewModels/TimingGraphTrack.h"
-#include "Insights/Widgets/SAsyncOperationStatus.h"
-#include "Insights/Widgets/SStatsViewTooltip.h"
 #include "Insights/Widgets/SStatsTableRow.h"
+#include "Insights/Widgets/SStatsViewTooltip.h"
 #include "Insights/Widgets/STimingProfilerWindow.h"
 #include "Insights/Widgets/STimingView.h"
 
@@ -113,7 +117,7 @@ public:
 ////////////////////////////////////////////////////////////////////////////////////////////////////
 
 SStatsView::SStatsView()
-	: Table(MakeShared<Insights::FTable>())
+	: Table(MakeShared<UE::Insights::FTable>())
 	, bExpansionSaved(false)
 	, bFilterOutZeroCountStats(false)
 	, GroupingMode(EStatsGroupingMode::Flat)
@@ -309,7 +313,7 @@ void SStatsView::Construct(const FArguments& InArgs)
 				.VAlign(VAlign_Bottom)
 				.Padding(16.0f)
 				[
-					SAssignNew(AsyncOperationStatus, Insights::SAsyncOperationStatus, Aggregator)
+					SAssignNew(AsyncOperationStatus, UE::Insights::SAsyncOperationStatus, Aggregator)
 				]
 			]
 
@@ -427,6 +431,7 @@ TSharedPtr<SWidget> SStatsView::TreeView_GetMenuContent()
 	{
 		auto CanExecute = [NumSelectedNodes, SelectedNode]()
 		{
+			using namespace UE::Insights::TimingProfiler;
 			TSharedPtr<STimingProfilerWindow> Wnd = FTimingProfilerManager::Get()->GetProfilerWindow();
 			TSharedPtr<STimingView> TimingView = Wnd.IsValid() ? Wnd->GetTimingView() : nullptr;
 			return TimingView.IsValid() && NumSelectedNodes == 1 && SelectedNode.IsValid() && SelectedNode->GetType() != EStatsNodeType::Group;
@@ -526,9 +531,9 @@ void SStatsView::TreeView_BuildSortByMenu(FMenuBuilder& MenuBuilder)
 {
 	MenuBuilder.BeginSection("SortColumn", LOCTEXT("ContextMenu_Section_SortColumn", "Sort Column"));
 
-	for (const TSharedRef<Insights::FTableColumn>& ColumnRef : Table->GetColumns())
+	for (const TSharedRef<UE::Insights::FTableColumn>& ColumnRef : Table->GetColumns())
 	{
-		const Insights::FTableColumn& Column = *ColumnRef;
+		const UE::Insights::FTableColumn& Column = *ColumnRef;
 
 		if (Column.IsVisible() && Column.CanBeSorted())
 		{
@@ -595,9 +600,9 @@ void SStatsView::TreeView_BuildViewColumnMenu(FMenuBuilder& MenuBuilder)
 {
 	MenuBuilder.BeginSection("Columns", LOCTEXT("ContextMenu_Section_Columns", "Columns"));
 
-	for (const TSharedRef<Insights::FTableColumn>& ColumnRef : Table->GetColumns())
+	for (const TSharedRef<UE::Insights::FTableColumn>& ColumnRef : Table->GetColumns())
 	{
-		const Insights::FTableColumn& Column = *ColumnRef;
+		const UE::Insights::FTableColumn& Column = *ColumnRef;
 
 		FUIAction Action_ToggleColumn
 		(
@@ -624,12 +629,12 @@ void SStatsView::TreeView_BuildViewColumnMenu(FMenuBuilder& MenuBuilder)
 void SStatsView::InitializeAndShowHeaderColumns()
 {
 	// Create columns.
-	TArray<TSharedRef<Insights::FTableColumn>> Columns;
+	TArray<TSharedRef<UE::Insights::FTableColumn>> Columns;
 	FStatsViewColumnFactory::CreateStatsViewColumns(Columns);
 	Table->SetColumns(Columns);
 
 	// Show columns.
-	for (const TSharedRef<Insights::FTableColumn>& ColumnRef : Table->GetColumns())
+	for (const TSharedRef<UE::Insights::FTableColumn>& ColumnRef : Table->GetColumns())
 	{
 		if (ColumnRef->ShouldBeVisible())
 		{
@@ -642,13 +647,13 @@ void SStatsView::InitializeAndShowHeaderColumns()
 
 FText SStatsView::GetColumnHeaderText(const FName ColumnId) const
 {
-	const Insights::FTableColumn& Column = *Table->FindColumnChecked(ColumnId);
+	const UE::Insights::FTableColumn& Column = *Table->FindColumnChecked(ColumnId);
 	return Column.GetShortName();
 }
 
 ////////////////////////////////////////////////////////////////////////////////////////////////////
 
-TSharedRef<SWidget> SStatsView::TreeViewHeaderRow_GenerateColumnMenu(const Insights::FTableColumn& Column)
+TSharedRef<SWidget> SStatsView::TreeViewHeaderRow_GenerateColumnMenu(const UE::Insights::FTableColumn& Column)
 {
 	const bool bShouldCloseWindowAfterMenuSelection = true;
 	FMenuBuilder MenuBuilder(bShouldCloseWindowAfterMenuSelection, NULL);
@@ -838,7 +843,7 @@ void SStatsView::UpdateTree()
 
 void SStatsView::ApplyFiltering()
 {
-	FStopwatch Stopwatch;
+	UE::Insights::FStopwatch Stopwatch;
 	Stopwatch.Start();
 
 	FilteredGroupNodes.Reset();
@@ -851,11 +856,11 @@ void SStatsView::ApplyFiltering()
 		GroupPtr->ClearFilteredChildren();
 		const bool bIsGroupVisible = Filters->PassesAllFilters(GroupPtr);
 
-		const TArray<Insights::FBaseTreeNodePtr>& GroupChildren = GroupPtr->GetChildren();
+		const TArray<UE::Insights::FBaseTreeNodePtr>& GroupChildren = GroupPtr->GetChildren();
 		int32 NumVisibleChildren = 0;
-		for (const Insights::FBaseTreeNodePtr& ChildPtr : GroupChildren)
+		for (const UE::Insights::FBaseTreeNodePtr& ChildPtr : GroupChildren)
 		{
-			const FStatsNodePtr& NodePtr = StaticCastSharedPtr<FStatsNode, Insights::FBaseTreeNode>(ChildPtr);
+			const FStatsNodePtr& NodePtr = StaticCastSharedPtr<FStatsNode, UE::Insights::FBaseTreeNode>(ChildPtr);
 
 			const bool bIsChildVisible = (!bFilterOutZeroCountStats || NodePtr->GetAggregatedStats().Count > 0)
 									  && FilterByNodeType[static_cast<int>(NodePtr->GetType())]
@@ -936,10 +941,10 @@ void SStatsView::ApplyFiltering()
 
 		EStatsNodeDataType GroupDataType = EStatsNodeDataType::InvalidOrMax;
 
-		const TArray<Insights::FBaseTreeNodePtr>& GroupChildren = GroupPtr->GetFilteredChildren();
-		for (const Insights::FBaseTreeNodePtr& ChildPtr : GroupChildren)
+		const TArray<UE::Insights::FBaseTreeNodePtr>& GroupChildren = GroupPtr->GetFilteredChildren();
+		for (const UE::Insights::FBaseTreeNodePtr& ChildPtr : GroupChildren)
 		{
-			const FStatsNodePtr& NodePtr = StaticCastSharedPtr<FStatsNode, Insights::FBaseTreeNode>(ChildPtr);
+			const FStatsNodePtr& NodePtr = StaticCastSharedPtr<FStatsNode, UE::Insights::FBaseTreeNode>(ChildPtr);
 			const FAggregatedStats& NodeAggregatedStats = NodePtr->GetAggregatedStats();
 
 			if (NodeAggregatedStats.Count > 0)
@@ -989,7 +994,7 @@ void SStatsView::ApplyFiltering()
 	const double TotalTime = Stopwatch.GetAccumulatedTime();
 	if (TotalTime > 0.1)
 	{
-		UE_LOG(TimingProfiler, Log, TEXT("[Counters] Tree view filtered in %.3fs (%d counters)"),
+		UE_LOG(LogTimingProfiler, Log, TEXT("[Counters] Tree view filtered in %.3fs (%d counters)"),
 			TotalTime, StatsNodes.Num());
 	}
 }
@@ -1140,11 +1145,11 @@ void SStatsView::TreeView_OnSelectionChanged(FStatsNodePtr SelectedItem, ESelect
 
 void SStatsView::TreeView_OnGetChildren(FStatsNodePtr InParent, TArray<FStatsNodePtr>& OutChildren)
 {
-	const TArray<Insights::FBaseTreeNodePtr>& Children = InParent->GetFilteredChildren();
+	const TArray<UE::Insights::FBaseTreeNodePtr>& Children = InParent->GetFilteredChildren();
 	OutChildren.Reset(Children.Num());
-	for (const Insights::FBaseTreeNodePtr& Child : Children)
+	for (const UE::Insights::FBaseTreeNodePtr& Child : Children)
 	{
-		OutChildren.Add(StaticCastSharedPtr<FStatsNode, Insights::FBaseTreeNode>(Child));
+		OutChildren.Add(StaticCastSharedPtr<FStatsNode, UE::Insights::FBaseTreeNode>(Child));
 	}
 }
 
@@ -1192,7 +1197,7 @@ bool SStatsView::TableRow_ShouldBeEnabled(FStatsNodePtr NodePtr) const
 
 ////////////////////////////////////////////////////////////////////////////////////////////////////
 
-void SStatsView::TableRow_SetHoveredCell(TSharedPtr<Insights::FTable> InTablePtr, TSharedPtr<Insights::FTableColumn> InColumnPtr, FStatsNodePtr InNodePtr)
+void SStatsView::TableRow_SetHoveredCell(TSharedPtr<UE::Insights::FTable> InTablePtr, TSharedPtr<UE::Insights::FTableColumn> InColumnPtr, FStatsNodePtr InNodePtr)
 {
 	HoveredColumnId = InColumnPtr ? InColumnPtr->GetId() : FName();
 
@@ -1264,7 +1269,7 @@ bool SStatsView::SearchBox_IsEnabled() const
 
 void SStatsView::CreateGroups()
 {
-	FStopwatch Stopwatch;
+	UE::Insights::FStopwatch Stopwatch;
 	Stopwatch.Start();
 
 	if (GroupingMode == EStatsGroupingMode::Flat)
@@ -1397,7 +1402,7 @@ void SStatsView::CreateGroups()
 	const double TotalTime = Stopwatch.GetAccumulatedTime();
 	if (TotalTime > 0.1)
 	{
-		UE_LOG(TimingProfiler, Log, TEXT("[Counters] Tree view grouping updated in %.3fs (%d counters)"),
+		UE_LOG(LogTimingProfiler, Log, TEXT("[Counters] Tree view grouping updated in %.3fs (%d counters)"),
 			TotalTime, StatsNodes.Num());
 	}
 }
@@ -1485,11 +1490,11 @@ void SStatsView::CreateSortings()
 	AvailableSorters.Reset();
 	CurrentSorter = nullptr;
 
-	for (const TSharedRef<Insights::FTableColumn>& ColumnRef : Table->GetColumns())
+	for (const TSharedRef<UE::Insights::FTableColumn>& ColumnRef : Table->GetColumns())
 	{
 		if (ColumnRef->CanBeSorted())
 		{
-			TSharedPtr<Insights::ITableCellValueSorter> SorterPtr = ColumnRef->GetValueSorter();
+			TSharedPtr<UE::Insights::ITableCellValueSorter> SorterPtr = ColumnRef->GetValueSorter();
 			if (ensure(SorterPtr.IsValid()))
 			{
 				AvailableSorters.Add(SorterPtr);
@@ -1504,7 +1509,7 @@ void SStatsView::CreateSortings()
 
 void SStatsView::UpdateCurrentSortingByColumn()
 {
-	TSharedPtr<Insights::FTableColumn> ColumnPtr = Table->FindColumn(ColumnBeingSorted);
+	TSharedPtr<UE::Insights::FTableColumn> ColumnPtr = Table->FindColumn(ColumnBeingSorted);
 	CurrentSorter = ColumnPtr.IsValid() ? ColumnPtr->GetValueSorter() : nullptr;
 }
 
@@ -1512,7 +1517,7 @@ void SStatsView::UpdateCurrentSortingByColumn()
 
 void SStatsView::SortTreeNodes()
 {
-	FStopwatch Stopwatch;
+	UE::Insights::FStopwatch Stopwatch;
 	Stopwatch.Start();
 
 	if (CurrentSorter.IsValid())
@@ -1527,7 +1532,7 @@ void SStatsView::SortTreeNodes()
 	const double TotalTime = Stopwatch.GetAccumulatedTime();
 	if (TotalTime > 0.1)
 	{
-		UE_LOG(TimingProfiler, Log, TEXT("[Counters] Tree view sorted (%s, %c) in %.3fs (%d counters)"),
+		UE_LOG(LogTimingProfiler, Log, TEXT("[Counters] Tree view sorted (%s, %c) in %.3fs (%d counters)"),
 			CurrentSorter.IsValid() ? *CurrentSorter->GetShortName().ToString() : TEXT("N/A"),
 			(ColumnSortMode == EColumnSortMode::Type::Descending) ? TEXT('D') : TEXT('A'),
 			TotalTime, StatsNodes.Num());
@@ -1536,13 +1541,13 @@ void SStatsView::SortTreeNodes()
 
 ////////////////////////////////////////////////////////////////////////////////////////////////////
 
-void SStatsView::SortTreeNodesRec(FStatsNode& Node, const Insights::ITableCellValueSorter& Sorter)
+void SStatsView::SortTreeNodesRec(FStatsNode& Node, const UE::Insights::ITableCellValueSorter& Sorter)
 {
-	Insights::ESortMode SortMode = (ColumnSortMode == EColumnSortMode::Type::Descending) ? Insights::ESortMode::Descending : Insights::ESortMode::Ascending;
+	UE::Insights::ESortMode SortMode = (ColumnSortMode == EColumnSortMode::Type::Descending) ? UE::Insights::ESortMode::Descending : UE::Insights::ESortMode::Ascending;
 	Node.SortChildren(Sorter, SortMode);
 
 #if 0 // Current groupings creates only one level.
-	for (Insights::FBaseTreeNodePtr ChildPtr : Node.GetChildren())
+	for (UE::Insights::FBaseTreeNodePtr ChildPtr : Node.GetChildren())
 	{
 		if (ChildPtr->GetChildrenCount() > 0)
 		{
@@ -1597,7 +1602,7 @@ bool SStatsView::HeaderMenu_SortMode_IsChecked(const FName ColumnId, const EColu
 
 bool SStatsView::HeaderMenu_SortMode_CanExecute(const FName ColumnId, const EColumnSortMode::Type InSortMode) const
 {
-	const Insights::FTableColumn& Column = *Table->FindColumnChecked(ColumnId);
+	const UE::Insights::FTableColumn& Column = *Table->FindColumnChecked(ColumnId);
 	return Column.CanBeSorted();
 }
 
@@ -1670,7 +1675,7 @@ bool SStatsView::CanShowColumn(const FName ColumnId) const
 
 void SStatsView::ShowColumn(const FName ColumnId)
 {
-	Insights::FTableColumn& Column = *Table->FindColumnChecked(ColumnId);
+	UE::Insights::FTableColumn& Column = *Table->FindColumnChecked(ColumnId);
 	Column.Show();
 
 	SHeaderRow::FColumn::FArguments ColumnArgs;
@@ -1725,7 +1730,7 @@ void SStatsView::ShowColumn(const FName ColumnId)
 
 bool SStatsView::CanHideColumn(const FName ColumnId) const
 {
-	const Insights::FTableColumn& Column = *Table->FindColumnChecked(ColumnId);
+	const UE::Insights::FTableColumn& Column = *Table->FindColumnChecked(ColumnId);
 	return Column.CanBeHidden();
 }
 
@@ -1733,7 +1738,7 @@ bool SStatsView::CanHideColumn(const FName ColumnId) const
 
 void SStatsView::HideColumn(const FName ColumnId)
 {
-	Insights::FTableColumn& Column = *Table->FindColumnChecked(ColumnId);
+	UE::Insights::FTableColumn& Column = *Table->FindColumnChecked(ColumnId);
 	Column.Hide();
 
 	TreeViewHeaderRow->RemoveColumn(ColumnId);
@@ -1745,7 +1750,7 @@ void SStatsView::HideColumn(const FName ColumnId)
 
 bool SStatsView::IsColumnVisible(const FName ColumnId) const
 {
-	const Insights::FTableColumn& Column = *Table->FindColumnChecked(ColumnId);
+	const UE::Insights::FTableColumn& Column = *Table->FindColumnChecked(ColumnId);
 	return Column.IsVisible();
 }
 
@@ -1753,7 +1758,7 @@ bool SStatsView::IsColumnVisible(const FName ColumnId) const
 
 bool SStatsView::CanToggleColumnVisibility(const FName ColumnId) const
 {
-	const Insights::FTableColumn& Column = *Table->FindColumnChecked(ColumnId);
+	const UE::Insights::FTableColumn& Column = *Table->FindColumnChecked(ColumnId);
 	return !Column.IsVisible() || Column.CanBeHidden();
 }
 
@@ -1761,7 +1766,7 @@ bool SStatsView::CanToggleColumnVisibility(const FName ColumnId) const
 
 void SStatsView::ToggleColumnVisibility(const FName ColumnId)
 {
-	const Insights::FTableColumn& Column = *Table->FindColumnChecked(ColumnId);
+	const UE::Insights::FTableColumn& Column = *Table->FindColumnChecked(ColumnId);
 	if (Column.IsVisible())
 	{
 		HideColumn(ColumnId);
@@ -1789,9 +1794,9 @@ void SStatsView::ContextMenu_ShowAllColumns_Execute()
 	ColumnSortMode = GetDefaultColumnSortMode();
 	UpdateCurrentSortingByColumn();
 
-	for (const TSharedRef<Insights::FTableColumn>& ColumnRef : Table->GetColumns())
+	for (const TSharedRef<UE::Insights::FTableColumn>& ColumnRef : Table->GetColumns())
 	{
-		const Insights::FTableColumn& Column = *ColumnRef;
+		const UE::Insights::FTableColumn& Column = *ColumnRef;
 
 		if (!Column.IsVisible())
 		{
@@ -1833,9 +1838,9 @@ void SStatsView::ContextMenu_ShowMinMaxMedColumns_Execute()
 	ColumnSortMode = EColumnSortMode::Descending;
 	UpdateCurrentSortingByColumn();
 
-	for (const TSharedRef<Insights::FTableColumn>& ColumnRef : Table->GetColumns())
+	for (const TSharedRef<UE::Insights::FTableColumn>& ColumnRef : Table->GetColumns())
 	{
-		const Insights::FTableColumn& Column = *ColumnRef;
+		const UE::Insights::FTableColumn& Column = *ColumnRef;
 
 		const bool bShouldBeVisible = Preset.Contains(Column.GetId());
 
@@ -1867,9 +1872,9 @@ void SStatsView::ContextMenu_ResetColumns_Execute()
 	ColumnSortMode = GetDefaultColumnSortMode();
 	UpdateCurrentSortingByColumn();
 
-	for (const TSharedRef<Insights::FTableColumn>& ColumnRef : Table->GetColumns())
+	for (const TSharedRef<UE::Insights::FTableColumn>& ColumnRef : Table->GetColumns())
 	{
-		const Insights::FTableColumn& Column = *ColumnRef;
+		const UE::Insights::FTableColumn& Column = *ColumnRef;
 
 		if (Column.ShouldBeVisible() && !Column.IsVisible())
 		{
@@ -1918,10 +1923,10 @@ void SStatsView::Tick(const FGeometry& AllottedGeometry, const double InCurrentT
 
 void SStatsView::RebuildTree(bool bResync)
 {
-	FStopwatch Stopwatch;
+	UE::Insights::FStopwatch Stopwatch;
 	Stopwatch.Start();
 
-	FStopwatch SyncStopwatch;
+	UE::Insights::FStopwatch SyncStopwatch;
 	SyncStopwatch.Start();
 
 	if (bResync)
@@ -2018,7 +2023,7 @@ void SStatsView::RebuildTree(bool bResync)
 	if (TotalTime > 0.01)
 	{
 		const double SyncTime = SyncStopwatch.GetAccumulatedTime();
-		UE_LOG(TimingProfiler, Log, TEXT("[Counters] Tree view rebuilt in %.4fs (sync: %.4fs + update: %.4fs) --> %d counters (%d added)"),
+		UE_LOG(LogTimingProfiler, Log, TEXT("[Counters] Tree view rebuilt in %.4fs (sync: %.4fs + update: %.4fs) --> %d counters (%d added)"),
 			TotalTime, SyncTime, TotalTime - SyncTime, StatsNodes.Num(), StatsNodes.Num() - PreviousNodeCount);
 	}
 }
@@ -2031,6 +2036,7 @@ void SStatsView::UpdateNode(FStatsNodePtr NodePtr)
 
 	if (!NodePtr->IsGroup())
 	{
+		using namespace UE::Insights::TimingProfiler;
 		TSharedPtr<STimingProfilerWindow> Wnd = FTimingProfilerManager::Get()->GetProfilerWindow();
 		if (Wnd.IsValid())
 		{
@@ -2123,6 +2129,7 @@ void SStatsView::SelectCounterNode(uint32 CounterId)
 
 TSharedPtr<FTimingGraphTrack> SStatsView::GetTimingViewMainGraphTrack() const
 {
+	using namespace UE::Insights::TimingProfiler;
 	TSharedPtr<STimingProfilerWindow> Wnd = FTimingProfilerManager::Get()->GetProfilerWindow();
 	TSharedPtr<STimingView> TimingView = Wnd.IsValid() ? Wnd->GetTimingView() : nullptr;
 
@@ -2196,7 +2203,7 @@ void SStatsView::ContextMenu_CopySelectedToClipboard_Execute()
 		return;
 	}
 
-	TArray<Insights::FBaseTreeNodePtr> SelectedNodes;
+	TArray<UE::Insights::FBaseTreeNodePtr> SelectedNodes;
 	for (FStatsNodePtr CounterPtr : TreeView->GetSelectedItems())
 	{
 		SelectedNodes.Add(CounterPtr);
@@ -2211,9 +2218,10 @@ void SStatsView::ContextMenu_CopySelectedToClipboard_Execute()
 
 	if (CurrentSorter.IsValid())
 	{
-		CurrentSorter->Sort(SelectedNodes, ColumnSortMode == EColumnSortMode::Ascending ? Insights::ESortMode::Ascending : Insights::ESortMode::Descending);
+		CurrentSorter->Sort(SelectedNodes, ColumnSortMode == EColumnSortMode::Ascending ? UE::Insights::ESortMode::Ascending : UE::Insights::ESortMode::Descending);
 	}
 
+	using namespace UE::Insights::TimingProfiler;
 	Table->GetVisibleColumnsData(SelectedNodes, FTimingProfilerManager::Get()->GetLogListingName(), TEXT('\t'), true, ClipboardText);
 
 	if (ClipboardText.Len() > 0)
@@ -2239,7 +2247,7 @@ void SStatsView::ContextMenu_Export_Execute()
 		return;
 	}
 
-	TArray<Insights::FBaseTreeNodePtr> SelectedNodes;
+	TArray<UE::Insights::FBaseTreeNodePtr> SelectedNodes;
 	for (FStatsNodePtr Item : TreeView->GetSelectedItems())
 	{
 		SelectedNodes.Add(Item);
@@ -2264,7 +2272,7 @@ void SStatsView::ContextMenu_Export_Execute()
 		return;
 	}
 
-	FStopwatch Stopwatch;
+	UE::Insights::FStopwatch Stopwatch;
 	Stopwatch.Start();
 
 	UTF16CHAR BOM = UNICODE_BOM;
@@ -2281,13 +2289,13 @@ void SStatsView::ContextMenu_Export_Execute()
 
 	TStringBuilder<1024> StringBuilder;
 
-	TArray<TSharedRef<Insights::FTableColumn>> VisibleColumns;
+	TArray<TSharedRef<UE::Insights::FTableColumn>> VisibleColumns;
 	Table->GetVisibleColumns(VisibleColumns);
 
 	// Write header.
 	{
 		bool bIsFirstColumn = true;
-		for (const TSharedRef<Insights::FTableColumn>& ColumnRef : VisibleColumns)
+		for (const TSharedRef<UE::Insights::FTableColumn>& ColumnRef : VisibleColumns)
 		{
 			if (bIsFirstColumn)
 			{
@@ -2316,18 +2324,18 @@ void SStatsView::ContextMenu_Export_Execute()
 
 	if (CurrentSorter.IsValid())
 	{
-		CurrentSorter->Sort(SelectedNodes, ColumnSortMode == EColumnSortMode::Ascending ? Insights::ESortMode::Ascending : Insights::ESortMode::Descending);
+		CurrentSorter->Sort(SelectedNodes, ColumnSortMode == EColumnSortMode::Ascending ? UE::Insights::ESortMode::Ascending : UE::Insights::ESortMode::Descending);
 	}
 
 	const int32 NodeCount = SelectedNodes.Num();
 	for (int32 Index = 0; Index < NodeCount; Index++)
 	{
-		const Insights::FBaseTreeNodePtr& Node = SelectedNodes[Index];
+		const UE::Insights::FBaseTreeNodePtr& Node = SelectedNodes[Index];
 
 		StringBuilder.Reset();
 
 		bool bIsFirstColumn = true;
-		for (const TSharedRef<Insights::FTableColumn>& ColumnRef : VisibleColumns)
+		for (const TSharedRef<UE::Insights::FTableColumn>& ColumnRef : VisibleColumns)
 		{
 			if (bIsFirstColumn)
 			{
@@ -2416,6 +2424,7 @@ void SStatsView::ContextMenu_ExportValues_Execute() const
 	Params.IntervalStartTime = -std::numeric_limits<double>::infinity();
 	Params.IntervalEndTime = +std::numeric_limits<double>::infinity();
 
+	using namespace UE::Insights::TimingProfiler;
 	TSharedPtr<STimingProfilerWindow> Wnd = FTimingProfilerManager::Get()->GetProfilerWindow();
 	TSharedPtr<STimingView> TimingView = Wnd.IsValid() ? Wnd->GetTimingView() : nullptr;
 	if (TimingView.IsValid())
@@ -2488,6 +2497,7 @@ void SStatsView::ContextMenu_ExportOps_Execute() const
 	Params.IntervalStartTime = -std::numeric_limits<double>::infinity();
 	Params.IntervalEndTime = +std::numeric_limits<double>::infinity();
 
+	using namespace UE::Insights::TimingProfiler;
 	TSharedPtr<STimingProfilerWindow> Wnd = FTimingProfilerManager::Get()->GetProfilerWindow();
 	TSharedPtr<STimingView> TimingView = Wnd.IsValid() ? Wnd->GetTimingView() : nullptr;
 	if (TimingView.IsValid())
@@ -2576,6 +2586,7 @@ IFileHandle* SStatsView::OpenExportFile(const TCHAR* InFilename) const
 
 	if (ExportFileHandle == nullptr)
 	{
+		using namespace UE::Insights::TimingProfiler;
 		FName LogListingName = FTimingProfilerManager::Get()->GetLogListingName();
 		FMessageLog ReportMessageLog((LogListingName != NAME_None) ? LogListingName : TEXT("Other"));
 		ReportMessageLog.Error(LOCTEXT("FailedToOpenFile", "Export failed. Failed to open file for write."));
diff --git a/Engine/Source/Developer/TraceInsights/Private/Insights/Widgets/SStatsView.h b/Engine/Source/Developer/TraceInsights/Private/Insights/Widgets/SStatsView.h
--- a/Engine/Source/Developer/TraceInsights/Private/Insights/Widgets/SStatsView.h
+++ b/Engine/Source/Developer/TraceInsights/Private/Insights/Widgets/SStatsView.h
@@ -3,6 +3,7 @@
 #pragma once
 
 #include "CoreMinimal.h"
+
 #include "Misc/FilterCollection.h"
 #include "Misc/TextFilter.h"
 #include "SlateFwd.h"
@@ -15,7 +16,7 @@
 #include "Widgets/Views/STableViewBase.h"
 #include "Widgets/Views/STreeView.h"
 
-// Insights
+// TraceInsights
 #include "Insights/ViewModels/StatsGroupingAndSorting.h"
 #include "Insights/ViewModels/StatsNode.h"
 
@@ -30,14 +31,17 @@ namespace TraceServices
 	class IAnalysisSession;
 }
 
-namespace Insights
+namespace UE::Insights
 {
 	class FTable;
 	class FTableColumn;
 	class ITableCellValueSorter;
+	class SAsyncOperationStatus;
+}
 
+namespace Insights
+{
 	class FCounterAggregator;
-	class SAsyncOperationStatus;
 }
 
 ////////////////////////////////////////////////////////////////////////////////////////////////////
@@ -70,7 +74,7 @@ public:
 	 */
 	void Construct(const FArguments& InArgs);
 
-	TSharedPtr<Insights::FTable> GetTable() const { return Table; }
+	TSharedPtr<UE::Insights::FTable> GetTable() const { return Table; }
 
 	void Reset();
 
@@ -144,7 +148,7 @@ private:
 
 	FText GetColumnHeaderText(const FName ColumnId) const;
 
-	TSharedRef<SWidget> TreeViewHeaderRow_GenerateColumnMenu(const Insights::FTableColumn& Column);
+	TSharedRef<SWidget> TreeViewHeaderRow_GenerateColumnMenu(const UE::Insights::FTableColumn& Column);
 
 	////////////////////////////////////////////////////////////////////////////////////////////////////
 	// Tree View - Misc
@@ -172,7 +176,7 @@ private:
 
 	bool TableRow_ShouldBeEnabled(FStatsNodePtr NodePtr) const;
 
-	void TableRow_SetHoveredCell(TSharedPtr<Insights::FTable> TablePtr, TSharedPtr<Insights::FTableColumn> ColumnPtr, FStatsNodePtr NodePtr);
+	void TableRow_SetHoveredCell(TSharedPtr<UE::Insights::FTable> TablePtr, TSharedPtr<UE::Insights::FTableColumn> ColumnPtr, FStatsNodePtr NodePtr);
 	EHorizontalAlignment TableRow_GetColumnOutlineHAlignment(const FName ColumnId) const;
 
 	FText TableRow_GetHighlightText() const;
@@ -222,7 +226,7 @@ private:
 
 	void UpdateCurrentSortingByColumn();
 	void SortTreeNodes();
-	void SortTreeNodesRec(FStatsNode& Node, const Insights::ITableCellValueSorter& Sorter);
+	void SortTreeNodesRec(FStatsNode& Node, const UE::Insights::ITableCellValueSorter& Sorter);
 
 	EColumnSortMode::Type GetSortModeForColumn(const FName ColumnId) const;
 	void SetSortModeForColumn(const FName& ColumnId, EColumnSortMode::Type SortMode);
@@ -294,7 +298,7 @@ private:
 
 private:
 	/** Table view model. */
-	TSharedPtr<Insights::FTable> Table;
+	TSharedPtr<UE::Insights::FTable> Table;
 
 	/** The analysis session used to populate this widget. */
 	TSharedPtr<const TraceServices::IAnalysisSession> Session;
@@ -381,10 +385,10 @@ private:
 	// Sorting
 
 	/** All available sorters. */
-	TArray<TSharedPtr<Insights::ITableCellValueSorter>> AvailableSorters;
+	TArray<TSharedPtr<UE::Insights::ITableCellValueSorter>> AvailableSorters;
 
 	/** Current sorter. It is nullptr if sorting is disabled. */
-	TSharedPtr<Insights::ITableCellValueSorter> CurrentSorter;
+	TSharedPtr<UE::Insights::ITableCellValueSorter> CurrentSorter;
 
 	/** Name of the column currently being sorted. Can be NAME_None if sorting is disabled (CurrentSorting == nullptr) or if a complex sorting is used (CurrentSorting != nullptr). */
 	FName ColumnBeingSorted;
@@ -395,7 +399,7 @@ private:
 	//////////////////////////////////////////////////
 
 	TSharedRef<Insights::FCounterAggregator> Aggregator;
-	TSharedPtr<Insights::SAsyncOperationStatus> AsyncOperationStatus;
+	TSharedPtr<UE::Insights::SAsyncOperationStatus> AsyncOperationStatus;
 };
 
 ////////////////////////////////////////////////////////////////////////////////////////////////////
diff --git a/Engine/Source/Developer/TraceInsights/Private/Insights/Widgets/SStatsViewTooltip.cpp b/Engine/Source/Developer/TraceInsights/Private/Insights/Widgets/SStatsViewTooltip.cpp
--- a/Engine/Source/Developer/TraceInsights/Private/Insights/Widgets/SStatsViewTooltip.cpp
+++ b/Engine/Source/Developer/TraceInsights/Private/Insights/Widgets/SStatsViewTooltip.cpp
@@ -3,18 +3,22 @@
 #include "SStatsViewTooltip.h"
 
 #include "SlateOptMacros.h"
-#include "TraceServices/Model/AnalysisSession.h"
 #include "Widgets/Layout/SGridPanel.h"
 #include "Widgets/Layout/SSeparator.h"
 #include "Widgets/SBoxPanel.h"
 #include "Widgets/SToolTip.h"
 #include "Widgets/Text/STextBlock.h"
 
-// Insights
-#include "Insights/Common/TimeUtils.h"
+// TraceServices
+#include "TraceServices/Model/AnalysisSession.h"
+
+// TraceInsightsCore
+#include "InsightsCore/Common/TimeUtils.h"
+#include "InsightsCore/Table/ViewModels/Table.h"
+#include "InsightsCore/Table/ViewModels/TableColumn.h"
+
+// TraceInsights
 #include "Insights/InsightsStyle.h"
-#include "Insights/Table/ViewModels/Table.h"
-#include "Insights/Table/ViewModels/TableColumn.h"
 #include "Insights/ViewModels/StatsNode.h"
 #include "Insights/ViewModels/StatsNodeHelper.h"
 
@@ -26,7 +30,7 @@
 
 BEGIN_SLATE_FUNCTION_BUILD_OPTIMIZATION
 
-TSharedPtr<SToolTip> SStatsViewTooltip::GetTableTooltip(const Insights::FTable& Table)
+TSharedPtr<SToolTip> SStatsViewTooltip::GetTableTooltip(const UE::Insights::FTable& Table)
 {
 	TSharedPtr<SToolTip> ColumnTooltip =
 		SNew(SToolTip)
@@ -57,7 +61,7 @@ TSharedPtr<SToolTip> SStatsViewTooltip::GetTableTooltip(const Insights::FTable&
 
 ////////////////////////////////////////////////////////////////////////////////////////////////////
 
-TSharedPtr<SToolTip> SStatsViewTooltip::GetColumnTooltip(const Insights::FTableColumn& Column)
+TSharedPtr<SToolTip> SStatsViewTooltip::GetColumnTooltip(const UE::Insights::FTableColumn& Column)
 {
 	TSharedPtr<SToolTip> ColumnTooltip =
 		SNew(SToolTip)
diff --git a/Engine/Source/Developer/TraceInsights/Private/Insights/Widgets/SStatsViewTooltip.h b/Engine/Source/Developer/TraceInsights/Private/Insights/Widgets/SStatsViewTooltip.h
--- a/Engine/Source/Developer/TraceInsights/Private/Insights/Widgets/SStatsViewTooltip.h
+++ b/Engine/Source/Developer/TraceInsights/Private/Insights/Widgets/SStatsViewTooltip.h
@@ -3,12 +3,13 @@
 #pragma once
 
 #include "CoreMinimal.h"
+
 #include "Widgets/IToolTip.h"
 #include "Widgets/SToolTip.h"
 
 class SGridPanel;
 
-namespace Insights
+namespace UE::Insights
 {
 	class FTable;
 	class FTableColumn;
@@ -24,8 +25,8 @@ class SStatsViewTooltip
 public:
 	SStatsViewTooltip() = delete;
 
-	static TSharedPtr<SToolTip> GetTableTooltip(const Insights::FTable& Table);
-	static TSharedPtr<SToolTip> GetColumnTooltip(const Insights::FTableColumn& Column);
+	static TSharedPtr<SToolTip> GetTableTooltip(const UE::Insights::FTable& Table);
+	static TSharedPtr<SToolTip> GetColumnTooltip(const UE::Insights::FTableColumn& Column);
 	static TSharedPtr<SToolTip> GetRowTooltip(const TSharedPtr<FStatsNode> TreeNodePtr);
 
 private:
diff --git a/Engine/Source/Developer/TraceInsights/Private/Insights/Widgets/STimerTableCell.cpp b/Engine/Source/Developer/TraceInsights/Private/Insights/Widgets/STimerTableCell.cpp
--- a/Engine/Source/Developer/TraceInsights/Private/Insights/Widgets/STimerTableCell.cpp
+++ b/Engine/Source/Developer/TraceInsights/Private/Insights/Widgets/STimerTableCell.cpp
@@ -12,10 +12,12 @@
 #include "Widgets/Text/STextBlock.h"
 #include "Widgets/Views/SExpanderArrow.h"
 
-// Insights
+// TraceInsightsCore
+#include "InsightsCore/Table/ViewModels/Table.h"
+#include "InsightsCore/Table/ViewModels/TableColumn.h"
+
+// TraceInsights
 #include "Insights/InsightsStyle.h"
-#include "Insights/Table/ViewModels/Table.h"
-#include "Insights/Table/ViewModels/TableColumn.h"
 #include "Insights/Widgets/STimerTableRow.h"
 
 #define LOCTEXT_NAMESPACE "STimersView"
diff --git a/Engine/Source/Developer/TraceInsights/Private/Insights/Widgets/STimerTableCell.h b/Engine/Source/Developer/TraceInsights/Private/Insights/Widgets/STimerTableCell.h
--- a/Engine/Source/Developer/TraceInsights/Private/Insights/Widgets/STimerTableCell.h
+++ b/Engine/Source/Developer/TraceInsights/Private/Insights/Widgets/STimerTableCell.h
@@ -3,21 +3,22 @@
 #pragma once
 
 #include "CoreMinimal.h"
+
 #include "Widgets/DeclarativeSyntaxSupport.h"
 #include "Widgets/SCompoundWidget.h"
 
-// Insights
+// TraceInsights
 #include "Insights/ViewModels/TimerNode.h"
 
 class ITableRow;
 
-namespace Insights
+namespace UE::Insights
 {
 	class FTable;
 	class FTableColumn;
 }
 
-DECLARE_DELEGATE_ThreeParams(FSetHoveredTimerTableCell, TSharedPtr<Insights::FTable> /*TablePtr*/, TSharedPtr<Insights::FTableColumn> /*ColumnPtr*/, FTimerNodePtr /*TimerNodePtr*/);
+DECLARE_DELEGATE_ThreeParams(FSetHoveredTimerTableCell, TSharedPtr<UE::Insights::FTable> /*TablePtr*/, TSharedPtr<UE::Insights::FTableColumn> /*ColumnPtr*/, FTimerNodePtr /*TimerNodePtr*/);
 
 class STimerTableCell : public SCompoundWidget
 {
@@ -25,8 +26,8 @@ public:
 	SLATE_BEGIN_ARGS(STimerTableCell) {}
 		SLATE_EVENT(FSetHoveredTimerTableCell, OnSetHoveredCell)
 		SLATE_ATTRIBUTE(FText, HighlightText)
-		SLATE_ARGUMENT(TSharedPtr<Insights::FTable>, TablePtr)
-		SLATE_ARGUMENT(TSharedPtr<Insights::FTableColumn>, ColumnPtr)
+		SLATE_ARGUMENT(TSharedPtr<UE::Insights::FTable>, TablePtr)
+		SLATE_ARGUMENT(TSharedPtr<UE::Insights::FTableColumn>, ColumnPtr)
 		SLATE_ARGUMENT(FTimerNodePtr, TimerNodePtr)
 		SLATE_ARGUMENT(bool, IsNameColumn)
 	SLATE_END_ARGS()
@@ -174,10 +175,10 @@ protected:
 
 protected:
 	/** A shared pointer to the table view model. */
-	TSharedPtr<Insights::FTable> TablePtr; // TODO: TSharedRef
+	TSharedPtr<UE::Insights::FTable> TablePtr; // TODO: TSharedRef
 
 	/** A shared pointer to the table column view model. */
-	TSharedPtr<Insights::FTableColumn> ColumnPtr; // TODO: TSharedRef
+	TSharedPtr<UE::Insights::FTableColumn> ColumnPtr; // TODO: TSharedRef
 
 	/** A shared pointer to the timer node. */
 	FTimerNodePtr TimerNodePtr; // TODO: TSharedRef
diff --git a/Engine/Source/Developer/TraceInsights/Private/Insights/Widgets/STimerTableRow.cpp b/Engine/Source/Developer/TraceInsights/Private/Insights/Widgets/STimerTableRow.cpp
--- a/Engine/Source/Developer/TraceInsights/Private/Insights/Widgets/STimerTableRow.cpp
+++ b/Engine/Source/Developer/TraceInsights/Private/Insights/Widgets/STimerTableRow.cpp
@@ -7,11 +7,13 @@
 #include "Widgets/SToolTip.h"
 #include "Widgets/Images/SImage.h"
 
-// Insights
-#include "Insights/Common/TimeUtils.h"
+// TraceInsightsCore
+#include "InsightsCore/Common/TimeUtils.h"
+#include "InsightsCore/Table/ViewModels/Table.h"
+#include "InsightsCore/Table/ViewModels/TableColumn.h"
+
+// TraceInsights
 #include "Insights/InsightsStyle.h"
-#include "Insights/Table/ViewModels/Table.h"
-#include "Insights/Table/ViewModels/TableColumn.h"
 #include "Insights/Widgets/STimersViewTooltip.h"
 #include "Insights/Widgets/STimerTableCell.h"
 
@@ -45,7 +47,7 @@ void STimerTableRow::Construct(const FArguments& InArgs, const TSharedRef<STable
 
 TSharedRef<SWidget> STimerTableRow::GenerateWidgetForColumn(const FName& ColumnId)
 {
-	TSharedPtr<Insights::FTableColumn> ColumnPtr = TablePtr->FindColumnChecked(ColumnId);
+	TSharedPtr<UE::Insights::FTableColumn> ColumnPtr = TablePtr->FindColumnChecked(ColumnId);
 
 	return
 		SNew(SOverlay)
@@ -145,10 +147,11 @@ FSlateColor STimerTableRow::GetBackgroundColorAndOpacity() const
 
 FSlateColor STimerTableRow::GetBackgroundColorAndOpacity(double Time) const
 {
-	const FLinearColor Color =	Time > TimeUtils::Second      ? FLinearColor(0.3f, 0.0f, 0.0f, 1.0f) :
-								Time > TimeUtils::Milisecond  ? FLinearColor(0.3f, 0.1f, 0.0f, 1.0f) :
-								Time > TimeUtils::Microsecond ? FLinearColor(0.0f, 0.1f, 0.0f, 1.0f) :
-								                                FLinearColor(0.0f, 0.0f, 0.0f, 1.0f);
+	using namespace UE::Insights;
+	const FLinearColor Color =	Time > FTimeValue::Second      ? FLinearColor(0.3f, 0.0f, 0.0f, 1.0f) :
+								Time > FTimeValue::Millisecond ? FLinearColor(0.3f, 0.1f, 0.0f, 1.0f) :
+								Time > FTimeValue::Microsecond ? FLinearColor(0.0f, 0.1f, 0.0f, 1.0f) :
+								                                 FLinearColor(0.0f, 0.0f, 0.0f, 1.0f);
 	return Color;
 }
 
@@ -211,7 +214,7 @@ EVisibility STimerTableRow::IsColumnVisible(const FName ColumnId) const
 
 ////////////////////////////////////////////////////////////////////////////////////////////////////
 
-void STimerTableRow::OnSetHoveredCell(TSharedPtr<Insights::FTable> InTablePtr, TSharedPtr<Insights::FTableColumn> InColumnPtr, FTimerNodePtr InTimerNodePtr)
+void STimerTableRow::OnSetHoveredCell(TSharedPtr<UE::Insights::FTable> InTablePtr, TSharedPtr<UE::Insights::FTableColumn> InColumnPtr, FTimerNodePtr InTimerNodePtr)
 {
 	SetHoveredCellDelegate.ExecuteIfBound(InTablePtr, InColumnPtr, InTimerNodePtr);
 }
diff --git a/Engine/Source/Developer/TraceInsights/Private/Insights/Widgets/STimerTableRow.h b/Engine/Source/Developer/TraceInsights/Private/Insights/Widgets/STimerTableRow.h
--- a/Engine/Source/Developer/TraceInsights/Private/Insights/Widgets/STimerTableRow.h
+++ b/Engine/Source/Developer/TraceInsights/Private/Insights/Widgets/STimerTableRow.h
@@ -3,19 +3,20 @@
 #pragma once
 
 #include "CoreMinimal.h"
+
 #include "Widgets/DeclarativeSyntaxSupport.h"
 #include "Widgets/Views/SHeaderRow.h"
 #include "Widgets/Views/STableViewBase.h"
 #include "Widgets/Views/STableRow.h"
 #include "Widgets/Views/STreeView.h"
 
-// Insights
+// TraceInsights
 #include "Insights/ViewModels/TimerNodeHelper.h"
 
 class IToolTip;
 class STimerTableRowToolTip;
 
-namespace Insights
+namespace UE::Insights
 {
 	class FTable;
 	class FTableColumn;
@@ -24,7 +25,7 @@ namespace Insights
 DECLARE_DELEGATE_RetVal_OneParam(bool, FTimerNodeShouldBeEnabledDelegate, FTimerNodePtr /*NodePtr*/);
 DECLARE_DELEGATE_RetVal_OneParam(bool, FIsColumnVisibleDelegate, const FName /*ColumnId*/);
 DECLARE_DELEGATE_RetVal_OneParam(EHorizontalAlignment, FGetColumnOutlineHAlignmentDelegate, const FName /*ColumnId*/);
-DECLARE_DELEGATE_ThreeParams(FSetHoveredTimerTableCell, TSharedPtr<Insights::FTable> /*TablePtr*/, TSharedPtr<Insights::FTableColumn> /*ColumnPtr*/, FTimerNodePtr /*TimerNodePtr*/);
+DECLARE_DELEGATE_ThreeParams(FSetHoveredTimerTableCell, TSharedPtr<UE::Insights::FTable> /*TablePtr*/, TSharedPtr<UE::Insights::FTableColumn> /*ColumnPtr*/, FTimerNodePtr /*TimerNodePtr*/);
 
 /** Widget that represents a table row in the tree control. Generates widgets for each column on demand. */
 class STimerTableRow : public SMultiColumnTableRow<FTimerNodePtr>
@@ -37,7 +38,7 @@ public:
 		SLATE_EVENT(FSetHoveredTimerTableCell, OnSetHoveredCell)
 		SLATE_ATTRIBUTE(FText, HighlightText)
 		SLATE_ATTRIBUTE(FName, HighlightedNodeName)
-		SLATE_ARGUMENT(TSharedPtr<Insights::FTable>, TablePtr)
+		SLATE_ARGUMENT(TSharedPtr<UE::Insights::FTable>, TablePtr)
 		SLATE_ARGUMENT(FTimerNodePtr, TimerNodePtr)
 	SLATE_END_ARGS()
 
@@ -72,11 +73,11 @@ protected:
 	const FSlateBrush* GetOutlineBrush(const FName ColumnId) const;
 	bool HandleShouldBeEnabled() const;
 	EVisibility IsColumnVisible(const FName ColumnId) const;
-	void OnSetHoveredCell(TSharedPtr<Insights::FTable> InTablePtr, TSharedPtr<Insights::FTableColumn> InColumnPtr, FTimerNodePtr InTimerNodePtr);
+	void OnSetHoveredCell(TSharedPtr<UE::Insights::FTable> InTablePtr, TSharedPtr<UE::Insights::FTableColumn> InColumnPtr, FTimerNodePtr InTimerNodePtr);
 
 protected:
 	/** A shared pointer to the table view model. */
-	TSharedPtr<Insights::FTable> TablePtr;
+	TSharedPtr<UE::Insights::FTable> TablePtr;
 
 	/** Data context for this table row. */
 	FTimerNodePtr TimerNodePtr;
diff --git a/Engine/Source/Developer/TraceInsights/Private/Insights/Widgets/STimerTreeView.cpp b/Engine/Source/Developer/TraceInsights/Private/Insights/Widgets/STimerTreeView.cpp
--- a/Engine/Source/Developer/TraceInsights/Private/Insights/Widgets/STimerTreeView.cpp
+++ b/Engine/Source/Developer/TraceInsights/Private/Insights/Widgets/STimerTreeView.cpp
@@ -2,9 +2,9 @@
 
 #include "STimerTreeView.h"
 
+#include "Framework/Application/SlateApplication.h"
 #include "Framework/Commands/Commands.h"
 #include "Framework/Commands/UICommandList.h"
-#include "Framework/Application/SlateApplication.h"
 #include "Framework/MultiBox/MultiBoxBuilder.h"
 #include "HAL/PlatformApplicationMisc.h"
 #include "SlateOptMacros.h"
@@ -12,15 +12,16 @@
 #include "Widgets/Layout/SScrollBox.h"
 #include "Widgets/SToolTip.h"
 
-// Insights
+// TraceInsightsCore
+#include "InsightsCore/Table/ViewModels/Table.h"
+#include "InsightsCore/Table/ViewModels/TableColumn.h"
+#include "InsightsCore/Table/Widgets/SAsyncOperationStatus.h"
+
+// TraceInsights
 #include "Insights/InsightsStyle.h"
-#include "Insights/Table/ViewModels/Table.h"
-#include "Insights/Table/ViewModels/TableColumn.h"
-#include "Insights/Table/ViewModels/TreeNodeSorting.h"
 #include "Insights/TimingProfilerManager.h"
 #include "Insights/ViewModels/TimerButterflyAggregation.h"
 #include "Insights/ViewModels/TimersViewColumnFactory.h"
-#include "Insights/Widgets/SAsyncOperationStatus.h"
 #include "Insights/Widgets/STimersViewTooltip.h"
 #include "Insights/Widgets/STimerTableRow.h"
 
@@ -66,7 +67,7 @@ public:
 ////////////////////////////////////////////////////////////////////////////////////////////////////
 
 STimerTreeView::STimerTreeView()
-	: Table(MakeShared<Insights::FTable>())
+	: Table(MakeShared<UE::Insights::FTable>())
 	, ViewName()
 	, TreeView(nullptr)
 	, TreeViewHeaderRow(nullptr)
@@ -104,7 +105,8 @@ void STimerTreeView::Construct(const FArguments& InArgs, const FText& InViewName
 {
 	ViewName = InViewName;
 
-	TSharedRef<Insights::FTimerButterflyAggregator> TimerButterflyAggregator = FTimingProfilerManager::Get()->GetTimerButterflyAggregator();
+	using namespace UE::Insights::TimingProfiler;
+	TSharedRef<FTimerButterflyAggregator> TimerButterflyAggregator = FTimingProfilerManager::Get()->GetTimerButterflyAggregator();
 
 	SAssignNew(ExternalScrollbar, SScrollBar)
 	.AlwaysShowScrollbar(true);
@@ -144,7 +146,7 @@ void STimerTreeView::Construct(const FArguments& InArgs, const FText& InViewName
 			.VAlign(VAlign_Bottom)
 			.Padding(16.0f)
 			[
-				SAssignNew(AsyncOperationStatus, Insights::SAsyncOperationStatus, TimerButterflyAggregator)
+				SAssignNew(AsyncOperationStatus, UE::Insights::SAsyncOperationStatus, TimerButterflyAggregator)
 			]
 		]
 
@@ -246,9 +248,9 @@ void STimerTreeView::TreeView_BuildSortByMenu(FMenuBuilder& MenuBuilder)
 {
 	MenuBuilder.BeginSection("SortColumn", LOCTEXT("ContextMenu_Section_SortColumn", "Sort Column"));
 
-	for (const TSharedRef<Insights::FTableColumn>& ColumnRef : Table->GetColumns())
+	for (const TSharedRef<UE::Insights::FTableColumn>& ColumnRef : Table->GetColumns())
 	{
-		const Insights::FTableColumn& Column = *ColumnRef;
+		const UE::Insights::FTableColumn& Column = *ColumnRef;
 
 		if (Column.IsVisible() && Column.CanBeSorted())
 		{
@@ -315,9 +317,9 @@ void STimerTreeView::TreeView_BuildViewColumnMenu(FMenuBuilder& MenuBuilder)
 {
 	MenuBuilder.BeginSection("Columns", LOCTEXT("ContextMenu_Section_Columns", "Columns"));
 
-	for (const TSharedRef<Insights::FTableColumn>& ColumnRef : Table->GetColumns())
+	for (const TSharedRef<UE::Insights::FTableColumn>& ColumnRef : Table->GetColumns())
 	{
-		const Insights::FTableColumn& Column = *ColumnRef;
+		const UE::Insights::FTableColumn& Column = *ColumnRef;
 
 		FUIAction Action_ToggleColumn
 		(
@@ -344,7 +346,7 @@ void STimerTreeView::TreeView_BuildViewColumnMenu(FMenuBuilder& MenuBuilder)
 void STimerTreeView::InitializeAndShowHeaderColumns()
 {
 	// Create columns.
-	TArray<TSharedRef<Insights::FTableColumn>> Columns;
+	TArray<TSharedRef<UE::Insights::FTableColumn>> Columns;
 	FTimersViewColumnFactory::CreateTimerTreeViewColumns(Columns);
 	if (ensure(Columns.Num() > 0 && Columns[0]->IsHierarchy()))
 	{
@@ -354,7 +356,7 @@ void STimerTreeView::InitializeAndShowHeaderColumns()
 	Table->SetColumns(Columns);
 
 	// Show columns.
-	for (const TSharedRef<Insights::FTableColumn>& ColumnRef : Table->GetColumns())
+	for (const TSharedRef<UE::Insights::FTableColumn>& ColumnRef : Table->GetColumns())
 	{
 		if (ColumnRef->ShouldBeVisible())
 		{
@@ -367,13 +369,13 @@ void STimerTreeView::InitializeAndShowHeaderColumns()
 
 FText STimerTreeView::GetColumnHeaderText(const FName ColumnId) const
 {
-	const Insights::FTableColumn& Column = *Table->FindColumnChecked(ColumnId);
+	const UE::Insights::FTableColumn& Column = *Table->FindColumnChecked(ColumnId);
 	return Column.GetShortName();
 }
 
 ////////////////////////////////////////////////////////////////////////////////////////////////////
 
-TSharedRef<SWidget> STimerTreeView::TreeViewHeaderRow_GenerateColumnMenu(const Insights::FTableColumn& Column)
+TSharedRef<SWidget> STimerTreeView::TreeViewHeaderRow_GenerateColumnMenu(const UE::Insights::FTableColumn& Column)
 {
 	const bool bShouldCloseWindowAfterMenuSelection = true;
 	FMenuBuilder MenuBuilder(bShouldCloseWindowAfterMenuSelection, NULL);
@@ -515,20 +517,20 @@ void STimerTreeView::TreeView_OnGetChildren(FTimerNodePtr InParent, TArray<FTime
 	constexpr bool bUseFiltering = false;
 	if (bUseFiltering)
 	{
-		const TArray<Insights::FBaseTreeNodePtr>& Children = InParent->GetFilteredChildren();
+		const TArray<UE::Insights::FBaseTreeNodePtr>& Children = InParent->GetFilteredChildren();
 		OutChildren.Reset(Children.Num());
-		for (const Insights::FBaseTreeNodePtr& Child : Children)
+		for (const UE::Insights::FBaseTreeNodePtr& Child : Children)
 		{
-			OutChildren.Add(StaticCastSharedPtr<FTimerNode, Insights::FBaseTreeNode>(Child));
+			OutChildren.Add(StaticCastSharedPtr<FTimerNode, UE::Insights::FBaseTreeNode>(Child));
 		}
 	}
 	else
 	{
-		const TArray<Insights::FBaseTreeNodePtr>& Children = InParent->GetChildren();
+		const TArray<UE::Insights::FBaseTreeNodePtr>& Children = InParent->GetChildren();
 		OutChildren.Reset(Children.Num());
-		for (const Insights::FBaseTreeNodePtr& Child : Children)
+		for (const UE::Insights::FBaseTreeNodePtr& Child : Children)
 		{
-			OutChildren.Add(StaticCastSharedPtr<FTimerNode, Insights::FBaseTreeNode>(Child));
+			OutChildren.Add(StaticCastSharedPtr<FTimerNode, UE::Insights::FBaseTreeNode>(Child));
 		}
 	}
 }
@@ -574,7 +576,7 @@ bool STimerTreeView::TableRow_ShouldBeEnabled(FTimerNodePtr NodePtr) const
 
 ////////////////////////////////////////////////////////////////////////////////////////////////////
 
-void STimerTreeView::TableRow_SetHoveredCell(TSharedPtr<Insights::FTable> InTablePtr, TSharedPtr<Insights::FTableColumn> InColumnPtr, FTimerNodePtr InNodePtr)
+void STimerTreeView::TableRow_SetHoveredCell(TSharedPtr<UE::Insights::FTable> InTablePtr, TSharedPtr<UE::Insights::FTableColumn> InColumnPtr, FTimerNodePtr InNodePtr)
 {
 	HoveredColumnId = InColumnPtr ? InColumnPtr->GetId() : FName();
 
@@ -645,11 +647,11 @@ void STimerTreeView::CreateSortings()
 	AvailableSorters.Reset();
 	CurrentSorter = nullptr;
 
-	for (const TSharedRef<Insights::FTableColumn>& ColumnRef : Table->GetColumns())
+	for (const TSharedRef<UE::Insights::FTableColumn>& ColumnRef : Table->GetColumns())
 	{
 		if (ColumnRef->CanBeSorted())
 		{
-			TSharedPtr<Insights::ITableCellValueSorter> SorterPtr = ColumnRef->GetValueSorter();
+			TSharedPtr<UE::Insights::ITableCellValueSorter> SorterPtr = ColumnRef->GetValueSorter();
 			if (ensure(SorterPtr.IsValid()))
 			{
 				AvailableSorters.Add(SorterPtr);
@@ -664,7 +666,7 @@ void STimerTreeView::CreateSortings()
 
 void STimerTreeView::UpdateCurrentSortingByColumn()
 {
-	TSharedPtr<Insights::FTableColumn> ColumnPtr = Table->FindColumn(ColumnBeingSorted);
+	TSharedPtr<UE::Insights::FTableColumn> ColumnPtr = Table->FindColumn(ColumnBeingSorted);
 	CurrentSorter = ColumnPtr.IsValid() ? ColumnPtr->GetValueSorter() : nullptr;
 }
 
@@ -683,12 +685,12 @@ void STimerTreeView::SortTreeNodes()
 
 ////////////////////////////////////////////////////////////////////////////////////////////////////
 
-void STimerTreeView::SortTreeNodesRec(FTimerNode& Node, const Insights::ITableCellValueSorter& Sorter)
+void STimerTreeView::SortTreeNodesRec(FTimerNode& Node, const UE::Insights::ITableCellValueSorter& Sorter)
 {
-	Insights::ESortMode SortMode = (ColumnSortMode == EColumnSortMode::Type::Descending) ? Insights::ESortMode::Descending : Insights::ESortMode::Ascending;
+	UE::Insights::ESortMode SortMode = (ColumnSortMode == EColumnSortMode::Type::Descending) ? UE::Insights::ESortMode::Descending : UE::Insights::ESortMode::Ascending;
 	Node.SortChildren(Sorter, SortMode);
 
-	for (Insights::FBaseTreeNodePtr ChildPtr : Node.GetChildren())
+	for (UE::Insights::FBaseTreeNodePtr ChildPtr : Node.GetChildren())
 	{
 		//if (ChildPtr->IsGroup())
 		if (ChildPtr->GetChildrenCount() > 0)
@@ -743,7 +745,7 @@ bool STimerTreeView::HeaderMenu_SortMode_IsChecked(const FName ColumnId, const E
 
 bool STimerTreeView::HeaderMenu_SortMode_CanExecute(const FName ColumnId, const EColumnSortMode::Type InSortMode) const
 {
-	const Insights::FTableColumn& Column = *Table->FindColumnChecked(ColumnId);
+	const UE::Insights::FTableColumn& Column = *Table->FindColumnChecked(ColumnId);
 	return Column.CanBeSorted();
 }
 
@@ -816,7 +818,7 @@ bool STimerTreeView::CanShowColumn(const FName ColumnId) const
 
 void STimerTreeView::ShowColumn(const FName ColumnId)
 {
-	Insights::FTableColumn& Column = *Table->FindColumnChecked(ColumnId);
+	UE::Insights::FTableColumn& Column = *Table->FindColumnChecked(ColumnId);
 	Column.Show();
 
 	SHeaderRow::FColumn::FArguments ColumnArgs;
@@ -871,7 +873,7 @@ void STimerTreeView::ShowColumn(const FName ColumnId)
 
 bool STimerTreeView::CanHideColumn(const FName ColumnId) const
 {
-	const Insights::FTableColumn& Column = *Table->FindColumnChecked(ColumnId);
+	const UE::Insights::FTableColumn& Column = *Table->FindColumnChecked(ColumnId);
 	return Column.CanBeHidden();
 }
 
@@ -879,7 +881,7 @@ bool STimerTreeView::CanHideColumn(const FName ColumnId) const
 
 void STimerTreeView::HideColumn(const FName ColumnId)
 {
-	Insights::FTableColumn& Column = *Table->FindColumnChecked(ColumnId);
+	UE::Insights::FTableColumn& Column = *Table->FindColumnChecked(ColumnId);
 	Column.Hide();
 
 	TreeViewHeaderRow->RemoveColumn(ColumnId);
@@ -891,7 +893,7 @@ void STimerTreeView::HideColumn(const FName ColumnId)
 
 bool STimerTreeView::IsColumnVisible(const FName ColumnId)
 {
-	const Insights::FTableColumn& Column = *Table->FindColumnChecked(ColumnId);
+	const UE::Insights::FTableColumn& Column = *Table->FindColumnChecked(ColumnId);
 	return Column.IsVisible();
 }
 
@@ -899,7 +901,7 @@ bool STimerTreeView::IsColumnVisible(const FName ColumnId)
 
 bool STimerTreeView::CanToggleColumnVisibility(const FName ColumnId) const
 {
-	const Insights::FTableColumn& Column = *Table->FindColumnChecked(ColumnId);
+	const UE::Insights::FTableColumn& Column = *Table->FindColumnChecked(ColumnId);
 	return !Column.IsVisible() || Column.CanBeHidden();
 }
 
@@ -907,7 +909,7 @@ bool STimerTreeView::CanToggleColumnVisibility(const FName ColumnId) const
 
 void STimerTreeView::ToggleColumnVisibility(const FName ColumnId)
 {
-	const Insights::FTableColumn& Column = *Table->FindColumnChecked(ColumnId);
+	const UE::Insights::FTableColumn& Column = *Table->FindColumnChecked(ColumnId);
 	if (Column.IsVisible())
 	{
 		HideColumn(ColumnId);
@@ -935,9 +937,9 @@ void STimerTreeView::ContextMenu_ShowAllColumns_Execute()
 	ColumnSortMode = GetDefaultColumnSortMode();
 	UpdateCurrentSortingByColumn();
 
-	for (const TSharedRef<Insights::FTableColumn>& ColumnRef : Table->GetColumns())
+	for (const TSharedRef<UE::Insights::FTableColumn>& ColumnRef : Table->GetColumns())
 	{
-		const Insights::FTableColumn& Column = *ColumnRef;
+		const UE::Insights::FTableColumn& Column = *ColumnRef;
 
 		if (!Column.IsVisible())
 		{
@@ -963,9 +965,9 @@ void STimerTreeView::ContextMenu_ResetColumns_Execute()
 	ColumnSortMode = GetDefaultColumnSortMode();
 	UpdateCurrentSortingByColumn();
 
-	for (const TSharedRef<Insights::FTableColumn>& ColumnRef : Table->GetColumns())
+	for (const TSharedRef<UE::Insights::FTableColumn>& ColumnRef : Table->GetColumns())
 	{
-		const Insights::FTableColumn& Column = *ColumnRef;
+		const UE::Insights::FTableColumn& Column = *ColumnRef;
 
 		if (Column.ShouldBeVisible() && !Column.IsVisible())
 		{
@@ -1000,8 +1002,8 @@ void STimerTreeView::SetTree(const TraceServices::FTimingProfilerButterflyNode&
 		while (TimerNodePtr.IsValid())
 		{
 			TimerNodePtr->SetIsHotPath(true);
-			const TArray<Insights::FBaseTreeNodePtr>& Children = TimerNodePtr->GetChildren();
-			TimerNodePtr = Children.Num() > 0 ? StaticCastSharedPtr<FTimerNode, Insights::FBaseTreeNode>(Children[0]) : nullptr;
+			const TArray<UE::Insights::FBaseTreeNodePtr>& Children = TimerNodePtr->GetChildren();
+			TimerNodePtr = Children.Num() > 0 ? StaticCastSharedPtr<FTimerNode, UE::Insights::FBaseTreeNode>(Children[0]) : nullptr;
 		}
 
 		TreeNodes.Add(RootTimerNodePtr);
@@ -1058,7 +1060,7 @@ FTimerNodePtr STimerTreeView::CreateTimerNodeRec(const TraceServices::FTimingPro
 	}
 
 	// Sort children by InclTime (descending).
-	TimerNodePtr->SortChildren([](const Insights::FBaseTreeNodePtr& A, const Insights::FBaseTreeNodePtr& B) -> bool
+	TimerNodePtr->SortChildren([](const UE::Insights::FBaseTreeNodePtr& A, const UE::Insights::FBaseTreeNodePtr& B) -> bool
 	{
 		const double InclTimeA = StaticCastSharedPtr<FTimerNode>(A)->GetAggregatedStats().TotalInclusiveTime;
 		const double InclTimeB = StaticCastSharedPtr<FTimerNode>(B)->GetAggregatedStats().TotalInclusiveTime;
@@ -1078,7 +1080,7 @@ void STimerTreeView::ExpandNodesRec(FTimerNodePtr NodePtr, int32 Depth)
 
 	//if (Depth < MaxDepth)
 	{
-		for (const Insights::FBaseTreeNodePtr& ChildPtr : NodePtr->GetChildren())
+		for (const UE::Insights::FBaseTreeNodePtr& ChildPtr : NodePtr->GetChildren())
 		{
 			ExpandNodesRec(StaticCastSharedPtr<FTimerNode>(ChildPtr), Depth + 1);
 		}
@@ -1103,7 +1105,7 @@ void STimerTreeView::ContextMenu_CopySelectedToClipboard_Execute()
 		return;
 	}
 
-	TArray<Insights::FBaseTreeNodePtr> SelectedNodes;
+	TArray<UE::Insights::FBaseTreeNodePtr> SelectedNodes;
 	for (FTimerNodePtr TimerPtr : TreeView->GetSelectedItems())
 	{
 		SelectedNodes.Add(TimerPtr);
@@ -1118,9 +1120,10 @@ void STimerTreeView::ContextMenu_CopySelectedToClipboard_Execute()
 
 	if (CurrentSorter.IsValid())
 	{
-		CurrentSorter->Sort(SelectedNodes, ColumnSortMode == EColumnSortMode::Ascending ? Insights::ESortMode::Ascending : Insights::ESortMode::Descending);
+		CurrentSorter->Sort(SelectedNodes, ColumnSortMode == EColumnSortMode::Ascending ? UE::Insights::ESortMode::Ascending : UE::Insights::ESortMode::Descending);
 	}
 
+	using namespace UE::Insights::TimingProfiler;
 	Table->GetVisibleColumnsData(SelectedNodes, FTimingProfilerManager::Get()->GetLogListingName(), TEXT('\t'), true, ClipboardText);
 
 	if (ClipboardText.Len() > 0)
