[Insights] Massive refactoring of TraceInsights module:

diff --git a/Engine/Source/Developer/TraceInsights/Private/Insights/Widgets/STimerTreeView.h b/Engine/Source/Developer/TraceInsights/Private/Insights/Widgets/STimerTreeView.h
--- a/Engine/Source/Developer/TraceInsights/Private/Insights/Widgets/STimerTreeView.h
+++ b/Engine/Source/Developer/TraceInsights/Private/Insights/Widgets/STimerTreeView.h
@@ -3,8 +3,8 @@
 #pragma once
 
 #include "CoreMinimal.h"
+
 #include "SlateFwd.h"
-#include "TraceServices/Model/TimingProfiler.h"
 #include "Widgets/DeclarativeSyntaxSupport.h"
 #include "Widgets/SCompoundWidget.h"
 #include "Widgets/SWidget.h"
@@ -13,7 +13,10 @@
 #include "Widgets/Views/STableViewBase.h"
 #include "Widgets/Views/STreeView.h"
 
-// Insights
+// TraceServices
+#include "TraceServices/Model/TimingProfiler.h"
+
+// TraceInsights
 #include "Insights/ViewModels/TimerGroupingAndSorting.h"
 #include "Insights/ViewModels/TimerNode.h"
 
@@ -22,7 +25,7 @@
 class FMenuBuilder;
 class FUICommandList;
 
-namespace Insights
+namespace UE::Insights
 {
 	class FTable;
 	class FTableColumn;
@@ -80,7 +83,7 @@ private:
 
 	FText GetColumnHeaderText(const FName ColumnId) const;
 
-	TSharedRef<SWidget> TreeViewHeaderRow_GenerateColumnMenu(const Insights::FTableColumn& Column);
+	TSharedRef<SWidget> TreeViewHeaderRow_GenerateColumnMenu(const UE::Insights::FTableColumn& Column);
 
 	////////////////////////////////////////////////////////////////////////////////////////////////////
 	// Tree View - Misc
@@ -108,7 +111,7 @@ private:
 
 	bool TableRow_ShouldBeEnabled(FTimerNodePtr NodePtr) const;
 
-	void TableRow_SetHoveredCell(TSharedPtr<Insights::FTable> TablePtr, TSharedPtr<Insights::FTableColumn> ColumnPtr, FTimerNodePtr NodePtr);
+	void TableRow_SetHoveredCell(TSharedPtr<UE::Insights::FTable> TablePtr, TSharedPtr<UE::Insights::FTableColumn> ColumnPtr, FTimerNodePtr NodePtr);
 	EHorizontalAlignment TableRow_GetColumnOutlineHAlignment(const FName ColumnId) const;
 
 	FText TableRow_GetHighlightText() const;
@@ -124,7 +127,7 @@ private:
 
 	void UpdateCurrentSortingByColumn();
 	void SortTreeNodes();
-	void SortTreeNodesRec(FTimerNode& Node, const Insights::ITableCellValueSorter& Sorter);
+	void SortTreeNodesRec(FTimerNode& Node, const UE::Insights::ITableCellValueSorter& Sorter);
 
 	EColumnSortMode::Type GetSortModeForColumn(const FName ColumnId) const;
 	void SetSortModeForColumn(const FName& ColumnId, EColumnSortMode::Type SortMode);
@@ -176,7 +179,7 @@ private:
 
 private:
 	/** Table view model. */
-	TSharedPtr<Insights::FTable> Table;
+	TSharedPtr<UE::Insights::FTable> Table;
 
 	/** The view name (ex.: "Callers" or "Callees"). */
 	FText ViewName;
@@ -217,10 +220,10 @@ private:
 	// Sorting
 
 	/** All available sorters. */
-	TArray<TSharedPtr<Insights::ITableCellValueSorter>> AvailableSorters;
+	TArray<TSharedPtr<UE::Insights::ITableCellValueSorter>> AvailableSorters;
 
 	/** Current sorter. It is nullptr if sorting is disabled. */
-	TSharedPtr<Insights::ITableCellValueSorter> CurrentSorter;
+	TSharedPtr<UE::Insights::ITableCellValueSorter> CurrentSorter;
 
 	/** Name of the column currently being sorted. Can be NAME_None if sorting is disabled (CurrentSorting == nullptr) or if a complex sorting is used (CurrentSorting != nullptr). */
 	FName ColumnBeingSorted;
@@ -228,7 +231,7 @@ private:
 	/** How we sort the nodes? Ascending or Descending. */
 	EColumnSortMode::Type ColumnSortMode;
 
-	TSharedPtr<Insights::SAsyncOperationStatus> AsyncOperationStatus;
+	TSharedPtr<UE::Insights::SAsyncOperationStatus> AsyncOperationStatus;
 };
 
 ////////////////////////////////////////////////////////////////////////////////////////////////////
diff --git a/Engine/Source/Developer/TraceInsights/Private/Insights/Widgets/STimersView.cpp b/Engine/Source/Developer/TraceInsights/Private/Insights/Widgets/STimersView.cpp
--- a/Engine/Source/Developer/TraceInsights/Private/Insights/Widgets/STimersView.cpp
+++ b/Engine/Source/Developer/TraceInsights/Private/Insights/Widgets/STimersView.cpp
@@ -14,24 +14,29 @@
 #include "Modules/ModuleManager.h"
 #include "SlateOptMacros.h"
 #include "Styling/AppStyle.h"
-#include "TraceServices/Model/Threads.h"
-#include "TraceServices/Model/TimingProfiler.h"
+#include "Widgets/Images/SImage.h"
 #include "Widgets/Input/SCheckBox.h"
-#include "Widgets/Layout/SScrollBox.h"
 #include "Widgets/Input/SSearchBox.h"
-#include "Widgets/Images/SImage.h"
 #include "Widgets/Layout/SGridPanel.h"
+#include "Widgets/Layout/SScrollBox.h"
 #include "Widgets/Layout/SSeparator.h"
 #include "Widgets/SToolTip.h"
 #include "Widgets/Views/STableViewBase.h"
 
-// Insights
-#include "Insights/Common/Stopwatch.h"
-#include "Insights/Common/TimeUtils.h"
+// TraceServices
+#include "TraceServices/Model/Threads.h"
+#include "TraceServices/Model/TimingProfiler.h"
+
+// TraceInsightsCore
+#include "InsightsCore/Common/Stopwatch.h"
+#include "InsightsCore/Common/TimeUtils.h"
+#include "InsightsCore/Table/ViewModels/Table.h"
+#include "InsightsCore/Table/ViewModels/TableColumn.h"
+#include "InsightsCore/Table/Widgets/SAsyncOperationStatus.h"
+
+// TraceInsights
 #include "Insights/InsightsStyle.h"
 #include "Insights/Log.h"
-#include "Insights/Table/ViewModels/Table.h"
-#include "Insights/Table/ViewModels/TableColumn.h"
 #include "Insights/TimingProfilerCommon.h"
 #include "Insights/TimingProfilerManager.h"
 #include "Insights/ViewModels/ThreadTimingTrack.h"
@@ -40,7 +45,6 @@
 #include "Insights/ViewModels/TimersViewColumnFactory.h"
 #include "Insights/ViewModels/TimingExporter.h"
 #include "Insights/ViewModels/TimingGraphTrack.h"
-#include "Insights/Widgets/SAsyncOperationStatus.h"
 #include "Insights/Widgets/SFrameTrack.h"
 #include "Insights/Widgets/STimersViewTooltip.h"
 #include "Insights/Widgets/STimerTableRow.h"
@@ -162,7 +166,7 @@ public:
 ////////////////////////////////////////////////////////////////////////////////////////////////////
 
 STimersView::STimersView()
-	: Table(MakeShared<Insights::FTable>())
+	: Table(MakeShared<UE::Insights::FTable>())
 	, bExpansionSaved(false)
 	, bFilterOutZeroCountTimers(false)
 	, GroupingMode(ETimerGroupingMode::ByType)
@@ -412,7 +416,7 @@ void STimersView::Construct(const FArguments& InArgs)
 				.VAlign(VAlign_Bottom)
 				.Padding(16.0f)
 				[
-					SAssignNew(AsyncOperationStatus, Insights::SAsyncOperationStatus, Aggregator)
+					SAssignNew(AsyncOperationStatus, UE::Insights::SAsyncOperationStatus, Aggregator)
 				]
 			]
 
@@ -533,10 +537,10 @@ TSharedPtr<SWidget> STimersView::TreeView_GetMenuContent()
 	// Timer options section
 	MenuBuilder.BeginSection("TimerOptions", LOCTEXT("ContextMenu_Section_TimerOptions", "Timer Options"));
 	{
-		auto CanExecute = [NumSelectedNodes, SelectedNode]()
+		TSharedPtr<UE::Insights::TimingProfiler::STimingView> TimingView = GetTimingView();
+
+		auto CanExecute = [TimingView, NumSelectedNodes, SelectedNode]()
 		{
-			TSharedPtr<STimingProfilerWindow> Wnd = FTimingProfilerManager::Get()->GetProfilerWindow();
-			TSharedPtr<STimingView> TimingView = Wnd.IsValid() ? Wnd->GetTimingView() : nullptr;
 			return TimingView.IsValid() && NumSelectedNodes == 1 && SelectedNode.IsValid() && SelectedNode->GetType() != ETimerNodeType::Group;
 		};
 
@@ -546,9 +550,6 @@ TSharedPtr<SWidget> STimersView::TreeView_GetMenuContent()
 			Action_ToggleHighlight.CanExecuteAction = FCanExecuteAction::CreateLambda(CanExecute);
 			Action_ToggleHighlight.ExecuteAction = FExecuteAction::CreateSP(this, &STimersView::ToggleTimingViewEventFilter, SelectedNode);
 
-			TSharedPtr<STimingProfilerWindow> Wnd = FTimingProfilerManager::Get()->GetProfilerWindow();
-			TSharedPtr<STimingView> TimingView = Wnd.IsValid() ? Wnd->GetTimingView() : nullptr;
-
 			if (SelectedNode.IsValid() &&
 				SelectedNode->GetType() != ETimerNodeType::Group &&
 				TimingView.IsValid() &&
@@ -702,9 +703,9 @@ void STimersView::TreeView_BuildSortByMenu(FMenuBuilder& MenuBuilder)
 {
 	MenuBuilder.BeginSection("SortColumn", LOCTEXT("ContextMenu_Section_SortColumn", "Sort Column"));
 
-	for (const TSharedRef<Insights::FTableColumn>& ColumnRef : Table->GetColumns())
+	for (const TSharedRef<UE::Insights::FTableColumn>& ColumnRef : Table->GetColumns())
 	{
-		const Insights::FTableColumn& Column = *ColumnRef;
+		const UE::Insights::FTableColumn& Column = *ColumnRef;
 
 		if (Column.IsVisible() && Column.CanBeSorted())
 		{
@@ -771,9 +772,9 @@ void STimersView::TreeView_BuildViewColumnMenu(FMenuBuilder& MenuBuilder)
 {
 	MenuBuilder.BeginSection("Columns", LOCTEXT("ContextMenu_Section_Columns", "Columns"));
 
-	for (const TSharedRef<Insights::FTableColumn>& ColumnRef : Table->GetColumns())
+	for (const TSharedRef<UE::Insights::FTableColumn>& ColumnRef : Table->GetColumns())
 	{
-		const Insights::FTableColumn& Column = *ColumnRef;
+		const UE::Insights::FTableColumn& Column = *ColumnRef;
 
 		FUIAction Action_ToggleColumn
 		(
@@ -841,6 +842,7 @@ void STimersView::TreeView_BuildPlotTimerMenu(FMenuBuilder& MenuBuilder)
 
 	auto CanExecuteAddToGraphTrack = [NumSelectedNodes, SelectedNode]()
 	{
+		using namespace UE::Insights::TimingProfiler;
 		TSharedPtr<STimingProfilerWindow> Wnd = FTimingProfilerManager::Get()->GetProfilerWindow();
 		TSharedPtr<STimingView> TimingView = Wnd.IsValid() ? Wnd->GetTimingView() : nullptr;
 		return TimingView.IsValid() && NumSelectedNodes == 1 && SelectedNode.IsValid() && SelectedNode->GetType() != ETimerNodeType::Group;
@@ -848,6 +850,7 @@ void STimersView::TreeView_BuildPlotTimerMenu(FMenuBuilder& MenuBuilder)
 
 	auto CanExecuteAddToFramesTrack = [NumSelectedNodes, SelectedNode]()
 	{
+		using namespace UE::Insights::TimingProfiler;
 		TSharedPtr<STimingProfilerWindow> Wnd = FTimingProfilerManager::Get()->GetProfilerWindow();
 		TSharedPtr<SFrameTrack> FrameTrack = Wnd.IsValid() ? Wnd->GetFrameView() : nullptr;
 		return FrameTrack.IsValid() && NumSelectedNodes == 1 && SelectedNode.IsValid() && SelectedNode->GetType() != ETimerNodeType::Group;
@@ -1084,12 +1087,12 @@ void STimersView::TreeView_FindMenu(FMenuBuilder& MenuBuilder)
 void STimersView::InitializeAndShowHeaderColumns()
 {
 	// Create columns.
-	TArray<TSharedRef<Insights::FTableColumn>> Columns;
+	TArray<TSharedRef<UE::Insights::FTableColumn>> Columns;
 	FTimersViewColumnFactory::CreateTimersViewColumns(Columns);
 	Table->SetColumns(Columns);
 
 	// Show columns.
-	for (const TSharedRef<Insights::FTableColumn>& ColumnRef : Table->GetColumns())
+	for (const TSharedRef<UE::Insights::FTableColumn>& ColumnRef : Table->GetColumns())
 	{
 		if (ColumnRef->ShouldBeVisible())
 		{
@@ -1102,13 +1105,13 @@ void STimersView::InitializeAndShowHeaderColumns()
 
 FText STimersView::GetColumnHeaderText(const FName ColumnId) const
 {
-	const Insights::FTableColumn& Column = *Table->FindColumnChecked(ColumnId);
+	const UE::Insights::FTableColumn& Column = *Table->FindColumnChecked(ColumnId);
 	return Column.GetShortName();
 }
 
 ////////////////////////////////////////////////////////////////////////////////////////////////////
 
-TSharedRef<SWidget> STimersView::TreeViewHeaderRow_GenerateColumnMenu(const Insights::FTableColumn& Column)
+TSharedRef<SWidget> STimersView::TreeViewHeaderRow_GenerateColumnMenu(const UE::Insights::FTableColumn& Column)
 {
 	const bool bShouldCloseWindowAfterMenuSelection = true;
 	FMenuBuilder MenuBuilder(bShouldCloseWindowAfterMenuSelection, NULL);
@@ -1299,7 +1302,7 @@ void STimersView::UpdateTree()
 
 void STimersView::ApplyFiltering()
 {
-	FStopwatch Stopwatch;
+	UE::Insights::FStopwatch Stopwatch;
 	Stopwatch.Start();
 
 	FilteredGroupNodes.Reset();
@@ -1312,11 +1315,11 @@ void STimersView::ApplyFiltering()
 		GroupPtr->ClearFilteredChildren();
 		const bool bIsGroupVisible = Filters->PassesAllFilters(GroupPtr);
 
-		const TArray<Insights::FBaseTreeNodePtr>& GroupChildren = GroupPtr->GetChildren();
+		const TArray<UE::Insights::FBaseTreeNodePtr>& GroupChildren = GroupPtr->GetChildren();
 		int32 NumVisibleChildren = 0;
-		for (const Insights::FBaseTreeNodePtr& ChildPtr : GroupChildren)
+		for (const UE::Insights::FBaseTreeNodePtr& ChildPtr : GroupChildren)
 		{
-			const FTimerNodePtr& NodePtr = StaticCastSharedPtr<FTimerNode, Insights::FBaseTreeNode>(ChildPtr);
+			const FTimerNodePtr& NodePtr = StaticCastSharedPtr<FTimerNode, UE::Insights::FBaseTreeNode>(ChildPtr);
 
 			const bool bIsChildVisible = (!bFilterOutZeroCountTimers || NodePtr->GetAggregatedStats().InstanceCount > 0)
 									  && bTimerTypeIsVisible[static_cast<int>(NodePtr->GetType())]
@@ -1384,10 +1387,10 @@ void STimersView::ApplyFiltering()
 		AggregatedStats.AverageExclusiveTime = NanTimeValue;
 		AggregatedStats.MedianExclusiveTime = NanTimeValue;
 
-		const TArray<Insights::FBaseTreeNodePtr>& GroupChildren = GroupPtr->GetFilteredChildren();
-		for (const Insights::FBaseTreeNodePtr& ChildPtr : GroupChildren)
+		const TArray<UE::Insights::FBaseTreeNodePtr>& GroupChildren = GroupPtr->GetFilteredChildren();
+		for (const UE::Insights::FBaseTreeNodePtr& ChildPtr : GroupChildren)
 		{
-			const FTimerNodePtr& NodePtr = StaticCastSharedPtr<FTimerNode, Insights::FBaseTreeNode>(ChildPtr);
+			const FTimerNodePtr& NodePtr = StaticCastSharedPtr<FTimerNode, UE::Insights::FBaseTreeNode>(ChildPtr);
 			const TraceServices::FTimingProfilerAggregatedStats& NodeAggregatedStats = NodePtr->GetAggregatedStats();
 
 			if (NodeAggregatedStats.InstanceCount > 0)
@@ -1410,7 +1413,7 @@ void STimersView::ApplyFiltering()
 	const double TotalTime = Stopwatch.GetAccumulatedTime();
 	if (TotalTime > 0.1)
 	{
-		UE_LOG(TimingProfiler, Log, TEXT("[Timers] Tree view filtered in %.3fs (%d timers)"),
+		UE_LOG(LogTimingProfiler, Log, TEXT("[Timers] Tree view filtered in %.3fs (%d timers)"),
 			TotalTime, TimerNodes.Num());
 	}
 }
@@ -1505,6 +1508,7 @@ void STimersView::TreeView_OnSelectionChanged(FTimerNodePtr SelectedItem, ESelec
 		FTimerNodePtr SelectedNode = GetSingleSelectedTimerNode();
 		if (SelectedNode.IsValid() && SelectedNode->GetType() != ETimerNodeType::Group)
 		{
+			using namespace UE::Insights::TimingProfiler;
 			FTimingProfilerManager::Get()->SetSelectedTimer(SelectedNode->GetTimerId());
 		}
 	}
@@ -1531,11 +1535,11 @@ FTimerNodePtr STimersView::GetSingleSelectedTimerNode() const
 
 void STimersView::TreeView_OnGetChildren(FTimerNodePtr InParent, TArray<FTimerNodePtr>& OutChildren)
 {
-	const TArray<Insights::FBaseTreeNodePtr>& Children = InParent->GetFilteredChildren();
+	const TArray<UE::Insights::FBaseTreeNodePtr>& Children = InParent->GetFilteredChildren();
 	OutChildren.Reset(Children.Num());
-	for (const Insights::FBaseTreeNodePtr& Child : Children)
+	for (const UE::Insights::FBaseTreeNodePtr& Child : Children)
 	{
-		OutChildren.Add(StaticCastSharedPtr<FTimerNode, Insights::FBaseTreeNode>(Child));
+		OutChildren.Add(StaticCastSharedPtr<FTimerNode, UE::Insights::FBaseTreeNode>(Child));
 	}
 }
 
@@ -1590,7 +1594,7 @@ bool STimersView::TableRow_ShouldBeEnabled(FTimerNodePtr NodePtr) const
 
 ////////////////////////////////////////////////////////////////////////////////////////////////////
 
-void STimersView::TableRow_SetHoveredCell(TSharedPtr<Insights::FTable> InTablePtr, TSharedPtr<Insights::FTableColumn> InColumnPtr, FTimerNodePtr InNodePtr)
+void STimersView::TableRow_SetHoveredCell(TSharedPtr<UE::Insights::FTable> InTablePtr, TSharedPtr<UE::Insights::FTableColumn> InColumnPtr, FTimerNodePtr InNodePtr)
 {
 	HoveredColumnId = InColumnPtr ? InColumnPtr->GetId() : FName();
 
@@ -1662,7 +1666,7 @@ bool STimersView::SearchBox_IsEnabled() const
 
 void STimersView::CreateGroups()
 {
-	FStopwatch Stopwatch;
+	UE::Insights::FStopwatch Stopwatch;
 	Stopwatch.Start();
 
 	if (GroupingMode == ETimerGroupingMode::Flat)
@@ -1774,19 +1778,23 @@ void STimersView::CreateGroups()
 	// Creates one group for each logarithmic range ie. 0.001 - 0.01, 0.01 - 0.1, 0.1 - 1.0, 1.0 - 10.0, etc.
 	else if (GroupingMode == ETimerGroupingMode::ByTotalInclusiveTime)
 	{
-		//im:TODO:
+		unimplemented();
 	}
 	// Creates one group for each logarithmic range ie. 0.001 - 0.01, 0.01 - 0.1, 0.1 - 1.0, 1.0 - 10.0, etc.
 	else if (GroupingMode == ETimerGroupingMode::ByTotalExclusiveTime)
 	{
-		//im:TODO:
+		unimplemented();
+	}
+	else
+	{
+		unimplemented();
 	}
 
 	Stopwatch.Stop();
 	const double TotalTime = Stopwatch.GetAccumulatedTime();
 	if (TotalTime > 0.1)
 	{
-		UE_LOG(TimingProfiler, Log, TEXT("[Timers] Tree view grouping updated in %.3fs (%d timers)"),
+		UE_LOG(LogTimingProfiler, Log, TEXT("[Timers] Tree view grouping updated in %.3fs (%d timers)"),
 			TotalTime, TimerNodes.Num());
 	}
 }
@@ -1980,11 +1988,11 @@ void STimersView::CreateSortings()
 	AvailableSorters.Reset();
 	CurrentSorter = nullptr;
 
-	for (const TSharedRef<Insights::FTableColumn>& ColumnRef : Table->GetColumns())
+	for (const TSharedRef<UE::Insights::FTableColumn>& ColumnRef : Table->GetColumns())
 	{
 		if (ColumnRef->CanBeSorted())
 		{
-			TSharedPtr<Insights::ITableCellValueSorter> SorterPtr = ColumnRef->GetValueSorter();
+			TSharedPtr<UE::Insights::ITableCellValueSorter> SorterPtr = ColumnRef->GetValueSorter();
 			if (ensure(SorterPtr.IsValid()))
 			{
 				AvailableSorters.Add(SorterPtr);
@@ -1999,7 +2007,7 @@ void STimersView::CreateSortings()
 
 void STimersView::UpdateCurrentSortingByColumn()
 {
-	TSharedPtr<Insights::FTableColumn> ColumnPtr = Table->FindColumn(ColumnBeingSorted);
+	TSharedPtr<UE::Insights::FTableColumn> ColumnPtr = Table->FindColumn(ColumnBeingSorted);
 	CurrentSorter = ColumnPtr.IsValid() ? ColumnPtr->GetValueSorter() : nullptr;
 }
 
@@ -2007,7 +2015,7 @@ void STimersView::UpdateCurrentSortingByColumn()
 
 void STimersView::SortTreeNodes()
 {
-	FStopwatch Stopwatch;
+	UE::Insights::FStopwatch Stopwatch;
 	Stopwatch.Start();
 
 	if (CurrentSorter.IsValid())
@@ -2022,7 +2030,7 @@ void STimersView::SortTreeNodes()
 	const double TotalTime = Stopwatch.GetAccumulatedTime();
 	if (TotalTime > 0.1)
 	{
-		UE_LOG(TimingProfiler, Log, TEXT("[Timers] Tree view sorted (%s, %c) in %.3fs (%d timers)"),
+		UE_LOG(LogTimingProfiler, Log, TEXT("[Timers] Tree view sorted (%s, %c) in %.3fs (%d timers)"),
 			CurrentSorter.IsValid() ? *CurrentSorter->GetShortName().ToString() : TEXT("N/A"),
 			(ColumnSortMode == EColumnSortMode::Type::Descending) ? TEXT('D') : TEXT('A'),
 			TotalTime, TimerNodes.Num());
@@ -2031,13 +2039,13 @@ void STimersView::SortTreeNodes()
 
 ////////////////////////////////////////////////////////////////////////////////////////////////////
 
-void STimersView::SortTreeNodesRec(FTimerNode& Node, const Insights::ITableCellValueSorter& Sorter)
+void STimersView::SortTreeNodesRec(FTimerNode& Node, const UE::Insights::ITableCellValueSorter& Sorter)
 {
-	Insights::ESortMode SortMode = (ColumnSortMode == EColumnSortMode::Type::Descending) ? Insights::ESortMode::Descending : Insights::ESortMode::Ascending;
+	UE::Insights::ESortMode SortMode = (ColumnSortMode == EColumnSortMode::Type::Descending) ? UE::Insights::ESortMode::Descending : UE::Insights::ESortMode::Ascending;
 	Node.SortChildren(Sorter, SortMode);
 
 #if 0 // Current groupings creates only one level.
-	for (Insights::FBaseTreeNodePtr ChildPtr : Node.GetChildren())
+	for (UE::Insights::FBaseTreeNodePtr ChildPtr : Node.GetChildren())
 	{
 		if (ChildPtr->GetChildrenCount() > 0)
 		{
@@ -2092,7 +2100,7 @@ bool STimersView::HeaderMenu_SortMode_IsChecked(const FName ColumnId, const ECol
 
 bool STimersView::HeaderMenu_SortMode_CanExecute(const FName ColumnId, const EColumnSortMode::Type InSortMode) const
 {
-	const Insights::FTableColumn& Column = *Table->FindColumnChecked(ColumnId);
+	const UE::Insights::FTableColumn& Column = *Table->FindColumnChecked(ColumnId);
 	return Column.CanBeSorted();
 }
 
@@ -2165,7 +2173,7 @@ bool STimersView::CanShowColumn(const FName ColumnId) const
 
 void STimersView::ShowColumn(const FName ColumnId)
 {
-	Insights::FTableColumn& Column = *Table->FindColumnChecked(ColumnId);
+	UE::Insights::FTableColumn& Column = *Table->FindColumnChecked(ColumnId);
 	Column.Show();
 
 	SHeaderRow::FColumn::FArguments ColumnArgs;
@@ -2220,7 +2228,7 @@ void STimersView::ShowColumn(const FName ColumnId)
 
 bool STimersView::CanHideColumn(const FName ColumnId) const
 {
-	const Insights::FTableColumn& Column = *Table->FindColumnChecked(ColumnId);
+	const UE::Insights::FTableColumn& Column = *Table->FindColumnChecked(ColumnId);
 	return Column.CanBeHidden();
 }
 
@@ -2228,7 +2236,7 @@ bool STimersView::CanHideColumn(const FName ColumnId) const
 
 void STimersView::HideColumn(const FName ColumnId)
 {
-	Insights::FTableColumn& Column = *Table->FindColumnChecked(ColumnId);
+	UE::Insights::FTableColumn& Column = *Table->FindColumnChecked(ColumnId);
 	Column.Hide();
 
 	TreeViewHeaderRow->RemoveColumn(ColumnId);
@@ -2240,7 +2248,7 @@ void STimersView::HideColumn(const FName ColumnId)
 
 bool STimersView::IsColumnVisible(const FName ColumnId) const
 {
-	const Insights::FTableColumn& Column = *Table->FindColumnChecked(ColumnId);
+	const UE::Insights::FTableColumn& Column = *Table->FindColumnChecked(ColumnId);
 	return Column.IsVisible();
 }
 
@@ -2248,7 +2256,7 @@ bool STimersView::IsColumnVisible(const FName ColumnId) const
 
 bool STimersView::CanToggleColumnVisibility(const FName ColumnId) const
 {
-	const Insights::FTableColumn& Column = *Table->FindColumnChecked(ColumnId);
+	const UE::Insights::FTableColumn& Column = *Table->FindColumnChecked(ColumnId);
 	return !Column.IsVisible() || Column.CanBeHidden();
 }
 
@@ -2256,7 +2264,7 @@ bool STimersView::CanToggleColumnVisibility(const FName ColumnId) const
 
 void STimersView::ToggleColumnVisibility(const FName ColumnId)
 {
-	const Insights::FTableColumn& Column = *Table->FindColumnChecked(ColumnId);
+	const UE::Insights::FTableColumn& Column = *Table->FindColumnChecked(ColumnId);
 	if (Column.IsVisible())
 	{
 		HideColumn(ColumnId);
@@ -2284,9 +2292,9 @@ void STimersView::ContextMenu_ShowAllColumns_Execute()
 	ColumnSortMode = GetDefaultColumnSortMode();
 	UpdateCurrentSortingByColumn();
 
-	for (const TSharedRef<Insights::FTableColumn>& ColumnRef : Table->GetColumns())
+	for (const TSharedRef<UE::Insights::FTableColumn>& ColumnRef : Table->GetColumns())
 	{
-		const Insights::FTableColumn& Column = *ColumnRef;
+		const UE::Insights::FTableColumn& Column = *ColumnRef;
 
 		if (!Column.IsVisible())
 		{
@@ -2330,9 +2338,9 @@ void STimersView::ContextMenu_ShowMinMaxMedColumns_Execute()
 	ColumnSortMode = EColumnSortMode::Descending;
 	UpdateCurrentSortingByColumn();
 
-	for (const TSharedRef<Insights::FTableColumn>& ColumnRef : Table->GetColumns())
+	for (const TSharedRef<UE::Insights::FTableColumn>& ColumnRef : Table->GetColumns())
 	{
-		const Insights::FTableColumn& Column = *ColumnRef;
+		const UE::Insights::FTableColumn& Column = *ColumnRef;
 
 		const bool bShouldBeVisible = Preset.Contains(Column.GetId());
 
@@ -2364,9 +2372,9 @@ void STimersView::ContextMenu_ResetColumns_Execute()
 	ColumnSortMode = GetDefaultColumnSortMode();
 	UpdateCurrentSortingByColumn();
 
-	for (const TSharedRef<Insights::FTableColumn>& ColumnRef : Table->GetColumns())
+	for (const TSharedRef<UE::Insights::FTableColumn>& ColumnRef : Table->GetColumns())
 	{
-		const Insights::FTableColumn& Column = *ColumnRef;
+		const UE::Insights::FTableColumn& Column = *ColumnRef;
 
 		if (Column.ShouldBeVisible() && !Column.IsVisible())
 		{
@@ -2415,10 +2423,10 @@ void STimersView::Tick(const FGeometry& AllottedGeometry, const double InCurrent
 
 void STimersView::RebuildTree(bool bResync)
 {
-	FStopwatch Stopwatch;
+	UE::Insights::FStopwatch Stopwatch;
 	Stopwatch.Start();
 
-	FStopwatch SyncStopwatch;
+	UE::Insights::FStopwatch SyncStopwatch;
 	SyncStopwatch.Start();
 
 	if (bResync)
@@ -2525,7 +2533,7 @@ void STimersView::RebuildTree(bool bResync)
 	if (TotalTime > 0.01)
 	{
 		const double SyncTime = SyncStopwatch.GetAccumulatedTime();
-		UE_LOG(TimingProfiler, Log, TEXT("[Timers] Tree view rebuilt in %.4fs (sync: %.4fs + update: %.4fs) --> %d timers (%d added)"),
+		UE_LOG(LogTimingProfiler, Log, TEXT("[Timers] Tree view rebuilt in %.4fs (sync: %.4fs + update: %.4fs) --> %d timers (%d added)"),
 			TotalTime, SyncTime, TotalTime - SyncTime, TimerNodes.Num(), TimerNodes.Num() - PreviousNodeCount);
 	}
 }
@@ -2626,9 +2634,7 @@ void STimersView::SelectTimerNode(uint32 TimerId)
 
 void STimersView::ToggleTimingViewEventFilter(FTimerNodePtr TimerNode) const
 {
-	TSharedPtr<STimingProfilerWindow> Wnd = FTimingProfilerManager::Get()->GetProfilerWindow();
-	TSharedPtr<STimingView> TimingView = Wnd.IsValid() ? Wnd->GetTimingView() : nullptr;
-
+	TSharedPtr<UE::Insights::TimingProfiler::STimingView> TimingView = GetTimingView();
 	if (TimingView.IsValid())
 	{
 		const uint64 EventType = static_cast<uint64>(TimerNode->GetTimerId());
@@ -2640,9 +2646,7 @@ void STimersView::ToggleTimingViewEventFilter(FTimerNodePtr TimerNode) const
 
 TSharedPtr<FTimingGraphTrack> STimersView::GetTimingViewMainGraphTrack() const
 {
-	TSharedPtr<STimingProfilerWindow> Wnd = FTimingProfilerManager::Get()->GetProfilerWindow();
-	TSharedPtr<STimingView> TimingView = Wnd.IsValid() ? Wnd->GetTimingView() : nullptr;
-
+	TSharedPtr<UE::Insights::TimingProfiler::STimingView> TimingView = GetTimingView();
 	return TimingView.IsValid() ? TimingView->GetMainTimingGraphTrack() : nullptr;
 }
 
@@ -2650,6 +2654,7 @@ TSharedPtr<FTimingGraphTrack> STimersView::GetTimingViewMainGraphTrack() const
 
 TSharedPtr<SFrameTrack> STimersView::GetFrameTrack() const
 {
+	using namespace UE::Insights::TimingProfiler;
 	TSharedPtr<STimingProfilerWindow> Wnd = FTimingProfilerManager::Get()->GetProfilerWindow();
 	return Wnd.IsValid() ? Wnd->GetFrameView() : nullptr;
 }
@@ -2846,7 +2851,7 @@ void STimersView::ContextMenu_CopyToClipboard_Execute()
 		return;
 	}
 
-	TArray<Insights::FBaseTreeNodePtr> SelectedNodes;
+	TArray<UE::Insights::FBaseTreeNodePtr> SelectedNodes;
 	for (FTimerNodePtr TimerPtr : SelectedTimerNodes)
 	{
 		SelectedNodes.Add(TimerPtr);
@@ -2856,9 +2861,10 @@ void STimersView::ContextMenu_CopyToClipboard_Execute()
 
 	if (CurrentSorter.IsValid())
 	{
-		CurrentSorter->Sort(SelectedNodes, ColumnSortMode == EColumnSortMode::Ascending ? Insights::ESortMode::Ascending : Insights::ESortMode::Descending);
+		CurrentSorter->Sort(SelectedNodes, ColumnSortMode == EColumnSortMode::Ascending ? UE::Insights::ESortMode::Ascending : UE::Insights::ESortMode::Descending);
 	}
 
+	using namespace UE::Insights::TimingProfiler;
 	Table->GetVisibleColumnsData(SelectedNodes, FTimingProfilerManager::Get()->GetLogListingName(), TEXT('\t'), true, ClipboardText);
 
 	if (ClipboardText.Len() > 0)
@@ -2890,7 +2896,7 @@ void STimersView::ContextMenu_Export_Execute()
 		return;
 	}
 
-	TArray<Insights::FBaseTreeNodePtr> SelectedNodes;
+	TArray<UE::Insights::FBaseTreeNodePtr> SelectedNodes;
 	for (FTimerNodePtr TimerPtr : SelectedTimerNodes)
 	{
 		SelectedNodes.Add(TimerPtr);
@@ -2910,7 +2916,7 @@ void STimersView::ContextMenu_Export_Execute()
 		return;
 	}
 
-	FStopwatch Stopwatch;
+	UE::Insights::FStopwatch Stopwatch;
 	Stopwatch.Start();
 
 	UTF16CHAR BOM = UNICODE_BOM;
@@ -2927,13 +2933,13 @@ void STimersView::ContextMenu_Export_Execute()
 
 	TStringBuilder<1024> StringBuilder;
 
-	TArray<TSharedRef<Insights::FTableColumn>> VisibleColumns;
+	TArray<TSharedRef<UE::Insights::FTableColumn>> VisibleColumns;
 	Table->GetVisibleColumns(VisibleColumns);
 
 	// Write header.
 	{
 		bool bIsFirstColumn = true;
-		for (const TSharedRef<Insights::FTableColumn>& ColumnRef : VisibleColumns)
+		for (const TSharedRef<UE::Insights::FTableColumn>& ColumnRef : VisibleColumns)
 		{
 			if (bIsFirstColumn)
 			{
@@ -2962,18 +2968,18 @@ void STimersView::ContextMenu_Export_Execute()
 
 	if (CurrentSorter.IsValid())
 	{
-		CurrentSorter->Sort(SelectedNodes, ColumnSortMode == EColumnSortMode::Ascending ? Insights::ESortMode::Ascending : Insights::ESortMode::Descending);
+		CurrentSorter->Sort(SelectedNodes, ColumnSortMode == EColumnSortMode::Ascending ? UE::Insights::ESortMode::Ascending : UE::Insights::ESortMode::Descending);
 	}
 
 	const int32 NodeCount = SelectedNodes.Num();
 	for (int32 Index = 0; Index < NodeCount; Index++)
 	{
-		const Insights::FBaseTreeNodePtr& Node = SelectedNodes[Index];
+		const UE::Insights::FBaseTreeNodePtr& Node = SelectedNodes[Index];
 
 		StringBuilder.Reset();
 
 		bool bIsFirstColumn = true;
-		for (const TSharedRef<Insights::FTableColumn>& ColumnRef : VisibleColumns)
+		for (const TSharedRef<UE::Insights::FTableColumn>& ColumnRef : VisibleColumns)
 		{
 			if (bIsFirstColumn)
 			{
@@ -3025,7 +3031,7 @@ void STimersView::AddTimerNodeRecursive(FTimerNodePtr InNode, TSet<uint32>& InOu
 {
 	if (InNode->GetType() == ETimerNodeType::Group)
 	{
-		for (Insights::FBaseTreeNodePtr ChildNode : InNode->GetFilteredChildren())
+		for (UE::Insights::FBaseTreeNodePtr ChildNode : InNode->GetFilteredChildren())
 		{
 			AddTimerNodeRecursive(StaticCastSharedPtr<FTimerNode>(ChildNode), InOutIncludedTimers);
 		}
@@ -3058,8 +3064,7 @@ void STimersView::ContextMenu_ExportTimingEventsSelection_Execute() const
 
 	////////////////////////////////////////////////////////////////////////////////////////////////////
 
-	TSharedPtr<STimingProfilerWindow> Wnd = FTimingProfilerManager::Get()->GetProfilerWindow();
-	TSharedPtr<STimingView> TimingView = Wnd.IsValid() ? Wnd->GetTimingView() : nullptr;
+	TSharedPtr<UE::Insights::TimingProfiler::STimingView> TimingView = GetTimingView();
 
 	////////////////////////////////////////////////////////////////////////////////////////////////////
 	// Filter by thread (visible Gpu/Cpu tracks in TimingView).
@@ -3277,6 +3282,7 @@ IFileHandle* STimersView::OpenExportFile(const TCHAR* InFilename) const
 
 	if (ExportFileHandle == nullptr)
 	{
+		using namespace UE::Insights::TimingProfiler;
 		FName LogListingName = FTimingProfilerManager::Get()->GetLogListingName();
 		FMessageLog ReportMessageLog((LogListingName != NAME_None) ? LogListingName : TEXT("Other"));
 		ReportMessageLog.Error(LOCTEXT("FailedToOpenFile", "Export failed. Failed to open file for write."));
@@ -3339,8 +3345,9 @@ void STimersView::ContextMenu_FindInstance_Execute(bool bFindMax) const
 		return;
 	}
 
-	TSharedPtr<STimingProfilerWindow> Wnd = FTimingProfilerManager::Get()->GetProfilerWindow();
-	TSharedPtr<STimingView> TimingView = Wnd.IsValid() ? Wnd->GetTimingView() : nullptr;
+	using namespace UE::Insights::TimingProfiler;
+
+	TSharedPtr<STimingView> TimingView = GetTimingView();
 	if (!TimingView.IsValid())
 	{
 		return;
@@ -3360,8 +3367,7 @@ bool STimersView::ContextMenu_FindInstanceInSelection_CanExecute() const
 		return false;
 	}
 
-	TSharedPtr<STimingProfilerWindow> Wnd = FTimingProfilerManager::Get()->GetProfilerWindow();
-	TSharedPtr<STimingView> TimingView = Wnd.IsValid() ? Wnd->GetTimingView() : nullptr;
+	TSharedPtr<UE::Insights::TimingProfiler::STimingView> TimingView = GetTimingView();
 	if (TimingView.IsValid())
 	{
 		return TimingView->GetSelectionEndTime() > TimingView->GetSelectionStartTime();
@@ -3380,8 +3386,9 @@ void STimersView::ContextMenu_FindInstanceInSelection_Execute(bool bFindMax) con
 		return;
 	}
 
-	TSharedPtr<STimingProfilerWindow> Wnd = FTimingProfilerManager::Get()->GetProfilerWindow();
-	TSharedPtr<STimingView> TimingView = Wnd.IsValid() ? Wnd->GetTimingView() : nullptr;
+	using namespace UE::Insights::TimingProfiler;
+
+	TSharedPtr<STimingView> TimingView = GetTimingView();
 	if (!TimingView.IsValid())
 	{
 		return;
@@ -3430,14 +3437,10 @@ FReply STimersView::OnKeyDown(const FGeometry& MyGeometry, const FKeyEvent& InKe
 
 void STimersView::SetTimingViewFrameType()
 {
-	TSharedPtr<STimingProfilerWindow> Window = FTimingProfilerManager::Get()->GetProfilerWindow();
-	if (Window.IsValid())
+	TSharedPtr<UE::Insights::TimingProfiler::STimingView> TimingView = GetTimingView();
+	if (TimingView.IsValid())
 	{
-		TSharedPtr<STimingView> TimingView = Window->GetTimingView();
-		if (TimingView.IsValid())
-		{
-			TimingView->SelectTimeInterval(TimingView->GetSelectionStartTime(), TimingView->GetSelectionEndTime() - TimingView->GetSelectionStartTime());
-		}
+		TimingView->SelectTimeInterval(TimingView->GetSelectionStartTime(), TimingView->GetSelectionEndTime() - TimingView->GetSelectionStartTime());
 	}
 }
 
@@ -3448,7 +3451,7 @@ void STimersView::SaveVisibleColumnsSettings()
 	FInsightsSettings& Settings = FInsightsManager::Get()->GetSettings();
 
 	TArray<FString> Columns;
-	for (const TSharedRef<Insights::FTableColumn>& ColumnRef : Table->GetColumns())
+	for (const TSharedRef<UE::Insights::FTableColumn>& ColumnRef : Table->GetColumns())
 	{
 		if (ColumnRef->IsVisible() && ColumnRef->CanBeHidden())
 		{
@@ -3519,7 +3522,7 @@ void STimersView::LoadVisibleColumnsSettings()
 		}
 	}
 
-	for (const TSharedRef<Insights::FTableColumn>& ColumnRef : Table->GetColumns())
+	for (const TSharedRef<UE::Insights::FTableColumn>& ColumnRef : Table->GetColumns())
 	{
 		if (ColumnRef->CanBeHidden())
 		{
@@ -3562,4 +3565,13 @@ void STimersView::OnTimingViewTrackListChanged()
 
 ////////////////////////////////////////////////////////////////////////////////////////////////////
 
+TSharedPtr<UE::Insights::TimingProfiler::STimingView> STimersView::GetTimingView() const
+{
+	using namespace UE::Insights::TimingProfiler;
+	TSharedPtr<STimingProfilerWindow> ProfilerWindow = FTimingProfilerManager::Get()->GetProfilerWindow();
+	return ProfilerWindow.IsValid() ? ProfilerWindow->GetTimingView() : nullptr;
+}
+
+////////////////////////////////////////////////////////////////////////////////////////////////////
+
 #undef LOCTEXT_NAMESPACE
diff --git a/Engine/Source/Developer/TraceInsights/Private/Insights/Widgets/STimersView.h b/Engine/Source/Developer/TraceInsights/Private/Insights/Widgets/STimersView.h
--- a/Engine/Source/Developer/TraceInsights/Private/Insights/Widgets/STimersView.h
+++ b/Engine/Source/Developer/TraceInsights/Private/Insights/Widgets/STimersView.h
@@ -3,6 +3,7 @@
 #pragma once
 
 #include "CoreMinimal.h"
+
 #include "Misc/FilterCollection.h"
 #include "Misc/TextFilter.h"
 #include "SlateFwd.h"
@@ -15,7 +16,7 @@
 #include "Widgets/Views/STableViewBase.h"
 #include "Widgets/Views/STreeView.h"
 
-// Insights
+// TraceInsights
 #include "Insights/ViewModels/TimerGroupingAndSorting.h"
 #include "Insights/ViewModels/TimerNode.h"
 
@@ -26,21 +27,20 @@ class FMenuBuilder;
 class FTimingGraphTrack;
 class FUICommandList;
 
-namespace TraceServices
-{
-	class IAnalysisSession;
-}
+namespace TraceServices { class IAnalysisSession; }
 
-namespace Insights
+namespace UE::Insights
 {
 	class FTable;
 	class FTableColumn;
 	class ITableCellValueSorter;
-
-	class FTimerAggregator;
 	class SAsyncOperationStatus;
 }
 
+namespace Insights { class FTimerAggregator; }
+
+namespace UE::Insights::TimingProfiler { class STimingView; }
+
 ////////////////////////////////////////////////////////////////////////////////////////////////////
 
 /** The filter collection - used for updating the list of timer nodes. */
@@ -71,7 +71,7 @@ public:
 	 */
 	void Construct(const FArguments& InArgs);
 
-	TSharedPtr<Insights::FTable> GetTable() const { return Table; }
+	TSharedPtr<UE::Insights::FTable> GetTable() const { return Table; }
 
 	void Reset();
 
@@ -165,7 +165,7 @@ private:
 
 	FText GetColumnHeaderText(const FName ColumnId) const;
 
-	TSharedRef<SWidget> TreeViewHeaderRow_GenerateColumnMenu(const Insights::FTableColumn& Column);
+	TSharedRef<SWidget> TreeViewHeaderRow_GenerateColumnMenu(const UE::Insights::FTableColumn& Column);
 
 	////////////////////////////////////////////////////////////////////////////////////////////////////
 	// Tree View - Misc
@@ -195,7 +195,7 @@ private:
 
 	bool TableRow_ShouldBeEnabled(FTimerNodePtr NodePtr) const;
 
-	void TableRow_SetHoveredCell(TSharedPtr<Insights::FTable> TablePtr, TSharedPtr<Insights::FTableColumn> ColumnPtr, FTimerNodePtr NodePtr);
+	void TableRow_SetHoveredCell(TSharedPtr<UE::Insights::FTable> TablePtr, TSharedPtr<UE::Insights::FTableColumn> ColumnPtr, FTimerNodePtr NodePtr);
 	EHorizontalAlignment TableRow_GetColumnOutlineHAlignment(const FName ColumnId) const;
 
 	FText TableRow_GetHighlightText() const;
@@ -258,7 +258,7 @@ private:
 
 	void UpdateCurrentSortingByColumn();
 	void SortTreeNodes();
-	void SortTreeNodesRec(FTimerNode& Node, const Insights::ITableCellValueSorter& Sorter);
+	void SortTreeNodesRec(FTimerNode& Node, const UE::Insights::ITableCellValueSorter& Sorter);
 
 	EColumnSortMode::Type GetSortModeForColumn(const FName ColumnId) const;
 	void SetSortModeForColumn(const FName& ColumnId, EColumnSortMode::Type SortMode);
@@ -352,9 +352,11 @@ private:
 
 	void SetTimingViewFrameType();
 
+	TSharedPtr<UE::Insights::TimingProfiler::STimingView> GetTimingView() const;
+
 private:
 	/** Table view model. */
-	TSharedPtr<Insights::FTable> Table;
+	TSharedPtr<UE::Insights::FTable> Table;
 
 	/** The analysis session used to populate this widget. */
 	TSharedPtr<const TraceServices::IAnalysisSession> Session;
@@ -450,10 +452,10 @@ private:
 	//bool bUseSorting;
 
 	/** All available sorters. */
-	TArray<TSharedPtr<Insights::ITableCellValueSorter>> AvailableSorters;
+	TArray<TSharedPtr<UE::Insights::ITableCellValueSorter>> AvailableSorters;
 
 	/** Current sorter. It is nullptr if sorting is disabled. */
-	TSharedPtr<Insights::ITableCellValueSorter> CurrentSorter;
+	TSharedPtr<UE::Insights::ITableCellValueSorter> CurrentSorter;
 
 	/** Name of the column currently being sorted. Can be NAME_None if sorting is disabled (CurrentSorting == nullptr) or if a complex sorting is used (CurrentSorting != nullptr). */
 	FName ColumnBeingSorted;
@@ -464,7 +466,7 @@ private:
 	//////////////////////////////////////////////////
 
 	TSharedRef<Insights::FTimerAggregator> Aggregator;
-	TSharedPtr<Insights::SAsyncOperationStatus> AsyncOperationStatus;
+	TSharedPtr<UE::Insights::SAsyncOperationStatus> AsyncOperationStatus;
 };
 
 ////////////////////////////////////////////////////////////////////////////////////////////////////
diff --git a/Engine/Source/Developer/TraceInsights/Private/Insights/Widgets/STimersViewTooltip.cpp b/Engine/Source/Developer/TraceInsights/Private/Insights/Widgets/STimersViewTooltip.cpp
--- a/Engine/Source/Developer/TraceInsights/Private/Insights/Widgets/STimersViewTooltip.cpp
+++ b/Engine/Source/Developer/TraceInsights/Private/Insights/Widgets/STimersViewTooltip.cpp
@@ -3,18 +3,22 @@
 #include "STimersViewTooltip.h"
 
 #include "SlateOptMacros.h"
-#include "TraceServices/Model/AnalysisSession.h"
 #include "Widgets/Layout/SGridPanel.h"
 #include "Widgets/Layout/SSeparator.h"
 #include "Widgets/SBoxPanel.h"
 #include "Widgets/SToolTip.h"
 #include "Widgets/Text/STextBlock.h"
 
-// Insights
-#include "Insights/Common/TimeUtils.h"
+// TraceServices
+#include "TraceServices/Model/AnalysisSession.h"
+
+// TraceInsightsCore
+#include "InsightsCore/Common/TimeUtils.h"
+#include "InsightsCore/Table/ViewModels/Table.h"
+#include "InsightsCore/Table/ViewModels/TableColumn.h"
+
+// TraceInsights
 #include "Insights/InsightsStyle.h"
-#include "Insights/Table/ViewModels/Table.h"
-#include "Insights/Table/ViewModels/TableColumn.h"
 #include "Insights/TimingProfilerManager.h"
 #include "Insights/ViewModels/TimerNode.h"
 #include "Insights/ViewModels/TimerNodeHelper.h"
@@ -32,7 +36,7 @@
 
 BEGIN_SLATE_FUNCTION_BUILD_OPTIMIZATION
 
-TSharedPtr<SToolTip> STimersViewTooltip::GetTableTooltip(const Insights::FTable& Table)
+TSharedPtr<SToolTip> STimersViewTooltip::GetTableTooltip(const UE::Insights::FTable& Table)
 {
 	TSharedPtr<SToolTip> ColumnTooltip =
 		SNew(SToolTip)
@@ -63,7 +67,7 @@ TSharedPtr<SToolTip> STimersViewTooltip::GetTableTooltip(const Insights::FTable&
 
 ////////////////////////////////////////////////////////////////////////////////////////////////////
 
-TSharedPtr<SToolTip> STimersViewTooltip::GetColumnTooltip(const Insights::FTableColumn& Column)
+TSharedPtr<SToolTip> STimersViewTooltip::GetColumnTooltip(const UE::Insights::FTableColumn& Column)
 {
 	const FTimersTableColumn& TimersColumn = static_cast<const FTimersTableColumn&>(Column);
 	FText InstanceDescription = TimersColumn.GetDescription(ETraceFrameType::TraceFrameType_Count);
@@ -73,6 +77,9 @@ TSharedPtr<SToolTip> STimersViewTooltip::GetColumnTooltip(const Insights::FTable
 	auto GetDescriptionLamda = [InstanceDescription, GameFrameDescription, RenderingDescription]()
 	{
 		ETraceFrameType FrameType = ETraceFrameType::TraceFrameType_Count;
+
+		using namespace UE::Insights::TimingProfiler;
+
 		TSharedPtr<STimingProfilerWindow> Window = FTimingProfilerManager::Get()->GetProfilerWindow();
 		if (Window.IsValid())
 		{
@@ -129,7 +136,7 @@ TSharedPtr<SToolTip> STimersViewTooltip::GetColumnTooltip(const Insights::FTable
 
 ////////////////////////////////////////////////////////////////////////////////////////////////////
 
-TSharedPtr<SToolTip> STimersViewTooltip::GetColumnTooltipForMode(const Insights::FTableColumn& Column, ETraceFrameType InAggregationMode)
+TSharedPtr<SToolTip> STimersViewTooltip::GetColumnTooltipForMode(const UE::Insights::FTableColumn& Column, ETraceFrameType InAggregationMode)
 {
 	const FTimersTableColumn& TimersColumn = static_cast<const FTimersTableColumn&>(Column);
 	FText Description = TimersColumn.GetDescription(InAggregationMode);
@@ -366,36 +373,38 @@ TSharedPtr<SToolTip> STimersViewTooltip::GetRowTooltip(const TSharedPtr<FTimerNo
 	const int32 NumDigits = 5;
 	int32 Row = 1;
 
+	using namespace UE::Insights;
+
 	if (true)
 	{
 		TCHAR FormatString[32];
 		FCString::Snprintf(FormatString, sizeof(FormatString), TEXT("%%.%dfs (%%s)"), NumDigits);
-		const FText TotalInclusiveTimeText = FText::FromString(FString::Printf(FormatString, Stats.TotalInclusiveTime, *TimeUtils::FormatTimeAuto(Stats.TotalInclusiveTime, 2)));
-		const FText TotalExclusiveTimeText = FText::FromString(FString::Printf(FormatString, Stats.TotalExclusiveTime, *TimeUtils::FormatTimeAuto(Stats.TotalExclusiveTime, 2)));
+		const FText TotalInclusiveTimeText = FText::FromString(FString::Printf(FormatString, Stats.TotalInclusiveTime, *FormatTimeAuto(Stats.TotalInclusiveTime, 2)));
+		const FText TotalExclusiveTimeText = FText::FromString(FString::Printf(FormatString, Stats.TotalExclusiveTime, *FormatTimeAuto(Stats.TotalExclusiveTime, 2)));
 		AddStatsRow(GridPanel, Row, LOCTEXT("TT_TotalTime", "Total Time:"), TotalInclusiveTimeText, TotalExclusiveTimeText);
 	}
 	if (!std::isnan(Stats.MaxInclusiveTime) || !std::isnan(Stats.MaxExclusiveTime))
 	{
-		const FText MaxInclusiveTimeText = FText::FromString(TimeUtils::FormatTimeMs(Stats.MaxInclusiveTime, NumDigits, true));
-		const FText MaxExclusiveTimeText = FText::FromString(TimeUtils::FormatTimeMs(Stats.MaxExclusiveTime, NumDigits, true));
+		const FText MaxInclusiveTimeText = FText::FromString(FormatTimeMs(Stats.MaxInclusiveTime, NumDigits, true));
+		const FText MaxExclusiveTimeText = FText::FromString(FormatTimeMs(Stats.MaxExclusiveTime, NumDigits, true));
 		AddStatsRow(GridPanel, Row, LOCTEXT("TT_MaxTime", "Max Time:"), MaxInclusiveTimeText, MaxExclusiveTimeText);
 	}
 	if (true)
 	{
-		const FText AvgInclusiveTimeText = FText::FromString(TimeUtils::FormatTimeMs(Stats.AverageInclusiveTime, NumDigits, true));
-		const FText AvgExclusiveTimeText = FText::FromString(TimeUtils::FormatTimeMs(Stats.AverageExclusiveTime, NumDigits, true));
-		AddStatsRow(GridPanel, Row, LOCTEXT("TT_AverageTime", "Average Time:"), AvgInclusiveTimeText,   AvgExclusiveTimeText);
+		const FText AvgInclusiveTimeText = FText::FromString(FormatTimeMs(Stats.AverageInclusiveTime, NumDigits, true));
+		const FText AvgExclusiveTimeText = FText::FromString(FormatTimeMs(Stats.AverageExclusiveTime, NumDigits, true));
+		AddStatsRow(GridPanel, Row, LOCTEXT("TT_AverageTime", "Average Time:"), AvgInclusiveTimeText, AvgExclusiveTimeText);
 	}
 	if (!std::isnan(Stats.MedianInclusiveTime) || !std::isnan(Stats.MedianExclusiveTime))
 	{
-		const FText MedInclusiveTimeText = FText::FromString(TimeUtils::FormatTimeMs(Stats.MedianInclusiveTime, NumDigits, true));
-		const FText MedExclusiveTimeText = FText::FromString(TimeUtils::FormatTimeMs(Stats.MedianExclusiveTime, NumDigits, true));
-		AddStatsRow(GridPanel, Row, LOCTEXT("TT_MedianTime",  "Median Time:"),  MedInclusiveTimeText,   MedExclusiveTimeText);
+		const FText MedInclusiveTimeText = FText::FromString(FormatTimeMs(Stats.MedianInclusiveTime, NumDigits, true));
+		const FText MedExclusiveTimeText = FText::FromString(FormatTimeMs(Stats.MedianExclusiveTime, NumDigits, true));
+		AddStatsRow(GridPanel, Row, LOCTEXT("TT_MedianTime", "Median Time:"), MedInclusiveTimeText, MedExclusiveTimeText);
 	}
 	if (!std::isnan(Stats.MinInclusiveTime) || !std::isnan(Stats.MinExclusiveTime))
 	{
-		const FText MinInclusiveTimeText = FText::FromString(TimeUtils::FormatTimeMs(Stats.MinInclusiveTime, NumDigits, true));
-		const FText MinExclusiveTimeText = FText::FromString(TimeUtils::FormatTimeMs(Stats.MinExclusiveTime, NumDigits, true));
+		const FText MinInclusiveTimeText = FText::FromString(FormatTimeMs(Stats.MinInclusiveTime, NumDigits, true));
+		const FText MinExclusiveTimeText = FText::FromString(FormatTimeMs(Stats.MinExclusiveTime, NumDigits, true));
 		AddStatsRow(GridPanel, Row, LOCTEXT("TT_MinTime", "Min Time:"), MinInclusiveTimeText, MinExclusiveTimeText);
 	}
 
diff --git a/Engine/Source/Developer/TraceInsights/Private/Insights/Widgets/STimersViewTooltip.h b/Engine/Source/Developer/TraceInsights/Private/Insights/Widgets/STimersViewTooltip.h
--- a/Engine/Source/Developer/TraceInsights/Private/Insights/Widgets/STimersViewTooltip.h
+++ b/Engine/Source/Developer/TraceInsights/Private/Insights/Widgets/STimersViewTooltip.h
@@ -8,7 +8,7 @@
 
 class SGridPanel;
 
-namespace Insights
+namespace UE::Insights
 {
 	class FTable;
 	class FTableColumn;
@@ -24,9 +24,9 @@ class STimersViewTooltip
 public:
 	STimersViewTooltip() = delete;
 
-	static TSharedPtr<SToolTip> GetTableTooltip(const Insights::FTable& Table);
-	static TSharedPtr<SToolTip> GetColumnTooltip(const Insights::FTableColumn& Column);
-	static TSharedPtr<SToolTip> GetColumnTooltipForMode(const Insights::FTableColumn& Column, ETraceFrameType InAggregationMode);
+	static TSharedPtr<SToolTip> GetTableTooltip(const UE::Insights::FTable& Table);
+	static TSharedPtr<SToolTip> GetColumnTooltip(const UE::Insights::FTableColumn& Column);
+	static TSharedPtr<SToolTip> GetColumnTooltipForMode(const UE::Insights::FTableColumn& Column, ETraceFrameType InAggregationMode);
 	static TSharedPtr<SToolTip> GetRowTooltip(const TSharedPtr<FTimerNode> TreeNodePtr);
 
 	static bool GetSource(const TSharedPtr<FTimerNode> TreeNodePtr, FText& OutSourcePrefix, FText& OutSourceSuffix);
diff --git a/Engine/Source/Developer/TraceInsights/Private/Insights/Widgets/STimingProfilerToolbar.cpp b/Engine/Source/Developer/TraceInsights/Private/Insights/Widgets/STimingProfilerToolbar.cpp
--- a/Engine/Source/Developer/TraceInsights/Private/Insights/Widgets/STimingProfilerToolbar.cpp
+++ b/Engine/Source/Developer/TraceInsights/Private/Insights/Widgets/STimingProfilerToolbar.cpp
@@ -6,7 +6,7 @@
 #include "Framework/MultiBox/MultiBoxBuilder.h"
 #include "Widgets/SBoxPanel.h"
 
-// Insights
+// TraceInsights
 #include "Insights/InsightsCommands.h"
 #include "Insights/InsightsManager.h"
 #include "Insights/InsightsStyle.h"
@@ -14,7 +14,10 @@
 
 ////////////////////////////////////////////////////////////////////////////////////////////////////
 
-#define LOCTEXT_NAMESPACE "STimingProfilerToolbar"
+#define LOCTEXT_NAMESPACE "UE::Insights::TimingProfiler"
+
+namespace UE::Insights::TimingProfiler
+{
 
 ////////////////////////////////////////////////////////////////////////////////////////////////////
 
@@ -121,4 +124,6 @@ void STimingProfilerToolbar::Construct(const FArguments& InArgs)
 
 ////////////////////////////////////////////////////////////////////////////////////////////////////
 
+} // namespace UE::Insights::TimingProfiler
+
 #undef LOCTEXT_NAMESPACE
diff --git a/Engine/Source/Developer/TraceInsights/Private/Insights/Widgets/STimingProfilerToolbar.h b/Engine/Source/Developer/TraceInsights/Private/Insights/Widgets/STimingProfilerToolbar.h
--- a/Engine/Source/Developer/TraceInsights/Private/Insights/Widgets/STimingProfilerToolbar.h
+++ b/Engine/Source/Developer/TraceInsights/Private/Insights/Widgets/STimingProfilerToolbar.h
@@ -7,6 +7,9 @@
 
 class FExtender;
 
+namespace UE::Insights::TimingProfiler
+{
+
 /** Ribbon based toolbar used as a main menu in the Profiler window. */
 class STimingProfilerToolbar : public SCompoundWidget
 {
@@ -28,3 +31,5 @@ public:
 	 */
 	void Construct(const FArguments& InArgs);
 };
+
+} // namespace UE::Insights::TimingProfiler
diff --git a/Engine/Source/Developer/TraceInsights/Private/Insights/Widgets/STimingProfilerWindow.cpp b/Engine/Source/Developer/TraceInsights/Private/Insights/Widgets/STimingProfilerWindow.cpp
--- a/Engine/Source/Developer/TraceInsights/Private/Insights/Widgets/STimingProfilerWindow.cpp
+++ b/Engine/Source/Developer/TraceInsights/Private/Insights/Widgets/STimingProfilerWindow.cpp
@@ -8,7 +8,7 @@
 #include "SlateOptMacros.h"
 #include "Widgets/Docking/SDockTab.h"
 
-// Insights
+// TraceInsights
 #include "Insights/InsightsStyle.h"
 #include "Insights/TimingProfilerManager.h"
 #include "Insights/Widgets/SFrameTrack.h"
@@ -21,10 +21,6 @@
 
 ////////////////////////////////////////////////////////////////////////////////////////////////////
 
-#define LOCTEXT_NAMESPACE "STimingProfilerWindow"
-
-////////////////////////////////////////////////////////////////////////////////////////////////////
-
 const FName FTimingProfilerTabs::ToolbarID(TEXT("Toolbar")); // DEPRECATED
 const FName FTimingProfilerTabs::FramesTrackID(TEXT("Frames"));
 const FName FTimingProfilerTabs::TimingViewID(TEXT("TimingView"));
@@ -34,6 +30,13 @@ const FName FTimingProfilerTabs::CalleesID(TEXT("Callees"));
 const FName FTimingProfilerTabs::StatsCountersID(TEXT("StasCounters"));
 const FName FTimingProfilerTabs::LogViewID(TEXT("LogView"));
 
+////////////////////////////////////////////////////////////////////////////////////////////////////
+
+#define LOCTEXT_NAMESPACE "UE::Insights::TimingProfiler"
+
+namespace UE::Insights::TimingProfiler
+{
+
 ////////////////////////////////////////////////////////////////////////////////////////////////////
 // STimingProfilerWindow
 ////////////////////////////////////////////////////////////////////////////////////////////////////
@@ -475,9 +478,9 @@ TSharedRef<SWidget> STimingProfilerWindow::CreateToolbar(TSharedPtr<FExtender> E
 
 ////////////////////////////////////////////////////////////////////////////////////////////////////
 
-void STimingProfilerWindow::OnTimeSelectionChanged(Insights::ETimeChangedFlags InFlags, double InStartTime, double InEndTime)
+void STimingProfilerWindow::OnTimeSelectionChanged(UE::Insights::Timing::ETimeChangedFlags InFlags, double InStartTime, double InEndTime)
 {
-	if (InFlags != Insights::ETimeChangedFlags::Interactive)
+	if (InFlags != UE::Insights::Timing::ETimeChangedFlags::Interactive)
 	{
 		FTimingProfilerManager::Get()->SetSelectedTimeRange(InStartTime, InEndTime);
 	}
@@ -485,4 +488,6 @@ void STimingProfilerWindow::OnTimeSelectionChanged(Insights::ETimeChangedFlags I
 
 ////////////////////////////////////////////////////////////////////////////////////////////////////
 
+} // namespace UE::Insights::TimingProfiler
+
 #undef LOCTEXT_NAMESPACE
diff --git a/Engine/Source/Developer/TraceInsights/Private/Insights/Widgets/STimingProfilerWindow.h b/Engine/Source/Developer/TraceInsights/Private/Insights/Widgets/STimingProfilerWindow.h
--- a/Engine/Source/Developer/TraceInsights/Private/Insights/Widgets/STimingProfilerWindow.h
+++ b/Engine/Source/Developer/TraceInsights/Private/Insights/Widgets/STimingProfilerWindow.h
@@ -4,23 +4,23 @@
 
 #include "CoreMinimal.h"
 
-// Insights
-#include "Insights/ITimingViewSession.h" // for Insights::ETimeChangedFlags
+// TraceInsights
+#include "Insights/ITimingViewSession.h" // for ETimeChangedFlags
 #include "Insights/Widgets/SMajorTabWindow.h"
 
-////////////////////////////////////////////////////////////////////////////////////////////////////
-
 class SFrameTrack;
 class SLogView;
 class SStatsView;
 class STimersView;
 class STimerTreeView;
-class STimingView;
 
-////////////////////////////////////////////////////////////////////////////////////////////////////
+namespace UE::Insights::TimingProfiler
+{
+
+class STimingView;
 
 /** Implements the Timing Insights major tab window. */
-class STimingProfilerWindow : public Insights::SMajorTabWindow
+class STimingProfilerWindow : public ::Insights::SMajorTabWindow
 {
 public:
 	/** Default constructor. */
@@ -37,7 +37,7 @@ public:
 	/** Constructs this widget. */
 	void Construct(const FArguments& InArgs, const TSharedRef<SDockTab>& ConstructUnderMajorTab, const TSharedPtr<SWindow>& ConstructUnderWindow);
 
-	TSharedPtr<STimingView> GetTimingView() const { return TimingView; }
+	TSharedPtr<UE::Insights::TimingProfiler::STimingView> GetTimingView() const { return TimingView; }
 	TSharedPtr<STimersView> GetTimersView() const { return TimersView; }
 	TSharedPtr<STimerTreeView> GetCallersTreeView() const { return CallersTreeView; }
 	TSharedPtr<STimerTreeView> GetCalleesTreeView() const { return CalleesTreeView; }
@@ -74,14 +74,14 @@ private:
 	TSharedRef<SDockTab> SpawnTab_LogView(const FSpawnTabArgs& Args);
 	void OnLogViewTabClosed(TSharedRef<SDockTab> TabBeingClosed);
 
-	void OnTimeSelectionChanged(Insights::ETimeChangedFlags InFlags, double InStartTime, double InEndTime);
+	void OnTimeSelectionChanged(UE::Insights::Timing::ETimeChangedFlags InFlags, double InStartTime, double InEndTime);
 
 private:
 	/** The Frame track widget */
 	TSharedPtr<SFrameTrack> FrameTrack;
 
 	/** The Timing view (multi-track) widget */
-	TSharedPtr<STimingView> TimingView;
+	TSharedPtr<UE::Insights::TimingProfiler::STimingView> TimingView;
 
 	/** The Timers view widget */
 	TSharedPtr<STimersView> TimersView;
@@ -98,3 +98,5 @@ private:
 	/** The Log view widget */
 	TSharedPtr<SLogView> LogView;
 };
+
+} // namespace UE::Insights::TimingProfiler
diff --git a/Engine/Source/Developer/TraceInsights/Private/Insights/Widgets/STimingView.cpp b/Engine/Source/Developer/TraceInsights/Private/Insights/Widgets/STimingView.cpp
--- a/Engine/Source/Developer/TraceInsights/Private/Insights/Widgets/STimingView.cpp
+++ b/Engine/Source/Developer/TraceInsights/Private/Insights/Widgets/STimingView.cpp
@@ -34,18 +34,23 @@
 #include "Widgets/SOverlay.h"
 #include "Widgets/Text/STextBlock.h"
 
-// Insights
+// TraceInsightsCore
+#include "InsightsCore/Common/PaintUtils.h"
+#include "InsightsCore/Common/Stopwatch.h"
+#include "InsightsCore/Common/TimeUtils.h"
+#include "InsightsCore/Filter/ViewModels/FilterConfigurator.h"
+#include "InsightsCore/Filter/ViewModels/TimeFilterValueConverter.h"
+#include "InsightsCore/Table/Widgets/STableTreeView.h"
+
+// TraceInsights
 #include "Insights/Common/InsightsMenuBuilder.h"
-#include "Insights/Common/PaintUtils.h"
-#include "Insights/Common/Stopwatch.h"
-#include "Insights/Common/TimeUtils.h"
+#include "Insights/Config.h"
 #include "Insights/InsightsManager.h"
 #include "Insights/InsightsStyle.h"
 #include "Insights/ITimingViewExtender.h"
 #include "Insights/LoadingProfiler/LoadingProfilerManager.h"
 #include "Insights/LoadingProfiler/Widgets/SLoadingProfilerWindow.h"
 #include "Insights/Log.h"
-#include "Insights/Table/Widgets/STableTreeView.h"
 #include "Insights/TaskGraphProfiler/TaskGraphProfilerManager.h"
 #include "Insights/Tests/TimingProfilerTests.h"
 #include "Insights/TimingProfilerCommon.h"
@@ -54,7 +59,6 @@
 #include "Insights/ViewModels/DrawHelpers.h"
 #include "Insights/ViewModels/EventNameFilterValueConverter.h"
 #include "Insights/ViewModels/FileActivityTimingTrack.h"
-#include "Insights/ViewModels/FilterConfigurator.h"
 #include "Insights/ViewModels/FrameTimingTrack.h"
 #include "Insights/ViewModels/GraphSeries.h"
 #include "Insights/ViewModels/GraphTrack.h"
@@ -64,7 +68,6 @@
 #include "Insights/ViewModels/RegionsTimingTrack.h"
 #include "Insights/ViewModels/ThreadTimingTrack.h"
 #include "Insights/ViewModels/ThreadTrackEvent.h"
-#include "Insights/ViewModels/TimeFilterValueConverter.h"
 #include "Insights/ViewModels/TimerFilters.h"
 #include "Insights/ViewModels/TimeRulerTrack.h"
 #include "Insights/ViewModels/TimingEventSearch.h"
@@ -78,21 +81,34 @@
 
 #include <limits>
 
-////////////////////////////////////////////////////////////////////////////////////////////////////
-
-#define LOCTEXT_NAMESPACE "STimingView"
+#define LOCTEXT_NAMESPACE "UE::Insights::TimingProfiler::STimingView"
 
 #define INSIGHTS_ACTIVATE_BENCHMARK 0
 
-// start auto generated ids from a big number (MSB set to 1) to avoid collisions with ids for GPU/CPU tracks based on 32bit timeline index
-uint64 FBaseTimingTrack::IdGenerator = (1ULL << 63);
-
-uint32 STimingView::TimingViewId = 0;
-
 const TCHAR* GetFileActivityTypeName(TraceServices::EFileActivityType Type);
 uint32 GetFileActivityTypeColor(TraceServices::EFileActivityType Type);
 
-namespace Insights { const FName TimingViewExtenderFeatureName(TEXT("TimingViewExtender")); }
+namespace UE::Insights::Timing
+{
+	const FName TimingViewExtenderFeatureName(TEXT("TimingViewExtender"));
+}
+#if UE_INSIGHTS_BACKWARD_COMPATIBILITY_UE54
+namespace Insights
+{
+	const FName TimingViewExtenderFeatureName(TEXT("TimingViewExtenderOld"));
+}
+namespace UE::Insights
+{
+	const FName TimingViewExtenderFeatureName(TEXT("TimingViewExtenderOld"));
+}
+#endif // UE_INSIGHTS_BACKWARD_COMPATIBILITY_UE54
+
+namespace UE::Insights::TimingProfiler
+{
+
+////////////////////////////////////////////////////////////////////////////////////////////////////
+
+uint32 STimingView::TimingViewId = 0;
 
 ////////////////////////////////////////////////////////////////////////////////////////////////////
 
@@ -102,9 +118,9 @@ STimingView::STimingView()
 	, ThreadTimingSharedState(MakeShared<FThreadTimingSharedState>(this))
 	, LoadingSharedState(MakeShared<FLoadingSharedState>(this))
 	, FileActivitySharedState(MakeShared<FFileActivitySharedState>(this))
-	, TimingRegionsSharedState(MakeShared<Insights::FTimingRegionsSharedState>(this))
+	, TimingRegionsSharedState(MakeShared<::Insights::FTimingRegionsSharedState>(this))
 	, TimeRulerTrack(MakeShared<FTimeRulerTrack>())
-	, DefaultTimeMarker(MakeShared<Insights::FTimeMarker>())
+	, DefaultTimeMarker(MakeShared<FTimeMarker>())
 	, MarkersTrack(MakeShared<FMarkersTimingTrack>())
 	, bAllowPanningOnScreenEdges(false)
 	, DPIScaleFactor(1.0f)
@@ -117,12 +133,11 @@ STimingView::STimingView()
 	DefaultTimeMarker->SetName(TEXT(""));
 	DefaultTimeMarker->SetColor(FLinearColor(0.85f, 0.5f, 0.03f, 0.5f));
 
-
-	IModularFeatures::Get().RegisterModularFeature(Insights::TimingViewExtenderFeatureName, FrameSharedState.Get());
-	IModularFeatures::Get().RegisterModularFeature(Insights::TimingViewExtenderFeatureName, ThreadTimingSharedState.Get());
-	IModularFeatures::Get().RegisterModularFeature(Insights::TimingViewExtenderFeatureName, LoadingSharedState.Get());
-	IModularFeatures::Get().RegisterModularFeature(Insights::TimingViewExtenderFeatureName, FileActivitySharedState.Get());
-	IModularFeatures::Get().RegisterModularFeature(Insights::TimingViewExtenderFeatureName, TimingRegionsSharedState.Get());
+	IModularFeatures::Get().RegisterModularFeature(Timing::TimingViewExtenderFeatureName, FrameSharedState.Get());
+	IModularFeatures::Get().RegisterModularFeature(Timing::TimingViewExtenderFeatureName, ThreadTimingSharedState.Get());
+	IModularFeatures::Get().RegisterModularFeature(Timing::TimingViewExtenderFeatureName, LoadingSharedState.Get());
+	IModularFeatures::Get().RegisterModularFeature(Timing::TimingViewExtenderFeatureName, FileActivitySharedState.Get());
+	IModularFeatures::Get().RegisterModularFeature(Timing::TimingViewExtenderFeatureName, TimingRegionsSharedState.Get());
 
 	ExtensionOverlay = SNew(SOverlay).Visibility(EVisibility::SelfHitTestInvisible);
 }
@@ -138,16 +153,26 @@ STimingView::~STimingView()
 	ForegroundTracks.Reset();
 
 	SelectedEvent.Reset();
-	for (Insights::ITimingViewExtender* Extender : GetExtenders())
+
+	for (Timing::ITimingViewExtender* Extender : GetExtenders())
 	{
 		Extender->OnEndSession(*this);
 	}
+#if UE_INSIGHTS_BACKWARD_COMPATIBILITY_UE54
+PRAGMA_DISABLE_DEPRECATION_WARNINGS
+	::Insights::ITimingViewSession* CurrentTimingViewSession = (::Insights::ITimingViewSession*)(Timing::ITimingViewSession*)this;
+	for (::Insights::ITimingViewExtender* Extender : GetOldExtenders())
+	{
+		Extender->OnEndSession(*CurrentTimingViewSession);
+	}
+PRAGMA_ENABLE_DEPRECATION_WARNINGS
+#endif // UE_INSIGHTS_BACKWARD_COMPATIBILITY_UE54
 
-	IModularFeatures::Get().UnregisterModularFeature(Insights::TimingViewExtenderFeatureName, TimingRegionsSharedState.Get());
-	IModularFeatures::Get().UnregisterModularFeature(Insights::TimingViewExtenderFeatureName, FileActivitySharedState.Get());
-	IModularFeatures::Get().UnregisterModularFeature(Insights::TimingViewExtenderFeatureName, LoadingSharedState.Get());
-	IModularFeatures::Get().UnregisterModularFeature(Insights::TimingViewExtenderFeatureName, ThreadTimingSharedState.Get());
-	IModularFeatures::Get().UnregisterModularFeature(Insights::TimingViewExtenderFeatureName, FrameSharedState.Get());
+	IModularFeatures::Get().UnregisterModularFeature(Timing::TimingViewExtenderFeatureName, TimingRegionsSharedState.Get());
+	IModularFeatures::Get().UnregisterModularFeature(Timing::TimingViewExtenderFeatureName, FileActivitySharedState.Get());
+	IModularFeatures::Get().UnregisterModularFeature(Timing::TimingViewExtenderFeatureName, LoadingSharedState.Get());
+	IModularFeatures::Get().UnregisterModularFeature(Timing::TimingViewExtenderFeatureName, ThreadTimingSharedState.Get());
+	IModularFeatures::Get().UnregisterModularFeature(Timing::TimingViewExtenderFeatureName, FrameSharedState.Get());
 
 	FGlobalTabmanager::Get()->UnregisterNomadTabSpawner(QuickFindTabId);
 }
@@ -321,10 +346,19 @@ void STimingView::Reset(bool bIsFirstReset)
 
 	if (!bIsFirstReset)
 	{
-		for (Insights::ITimingViewExtender* Extender : GetExtenders())
+		for (Timing::ITimingViewExtender* Extender : GetExtenders())
 		{
 			Extender->OnEndSession(*this);
 		}
+#if UE_INSIGHTS_BACKWARD_COMPATIBILITY_UE54
+PRAGMA_DISABLE_DEPRECATION_WARNINGS
+		::Insights::ITimingViewSession* CurrentTimingViewSession = (::Insights::ITimingViewSession*)(Timing::ITimingViewSession*)this;
+		for (::Insights::ITimingViewExtender* Extender : GetOldExtenders())
+		{
+			Extender->OnEndSession(*CurrentTimingViewSession);
+		}
+PRAGMA_ENABLE_DEPRECATION_WARNINGS
+#endif // UE_INSIGHTS_BACKWARD_COMPATIBILITY_UE54
 	}
 
 	//////////////////////////////////////////////////
@@ -363,16 +397,16 @@ void STimingView::Reset(bool bIsFirstReset)
 	SetTimeMarker(std::numeric_limits<double>::infinity());
 
 #if 0 // test for multiple time markers
-	TSharedRef<Insights::FTimeMarker> TimeMarkerA = DefaultTimeMarker;
+	TSharedRef<FTimeMarker> TimeMarkerA = DefaultTimeMarker;
 	TimeMarkerA->SetName(TEXT("A"));
 	TimeMarkerA->SetColor(FLinearColor(0.85f, 0.5f, 0.03f, 0.5f));
 
-	TSharedRef<Insights::FTimeMarker> TimeMarkerB = MakeShared<Insights::FTimeMarker>();
+	TSharedRef<FTimeMarker> TimeMarkerB = MakeShared<FTimeMarker>();
 	TimeRulerTrack->AddTimeMarker(TimeMarkerB);
 	TimeMarkerB->SetName(TEXT("B"));
 	TimeMarkerB->SetColor(FLinearColor(0.03f, 0.85f, 0.5f, 0.5f));
 
-	TSharedRef<Insights::FTimeMarker> TimeMarkerC = MakeShared<Insights::FTimeMarker>();
+	TSharedRef<FTimeMarker> TimeMarkerC = MakeShared<FTimeMarker>();
 	TimeRulerTrack->AddTimeMarker(TimeMarkerC);
 	TimeMarkerC->SetName(TEXT("C"));
 	TimeMarkerC->SetColor(FLinearColor(0.03f, 0.5f, 0.85f, 0.5f));
@@ -501,10 +535,19 @@ void STimingView::Reset(bool bIsFirstReset)
 
 	//////////////////////////////////////////////////
 
-	for (Insights::ITimingViewExtender* Extender : GetExtenders())
+	for (Timing::ITimingViewExtender* Extender : GetExtenders())
 	{
 		Extender->OnBeginSession(*this);
 	}
+#if UE_INSIGHTS_BACKWARD_COMPATIBILITY_UE54
+PRAGMA_DISABLE_DEPRECATION_WARNINGS
+	::Insights::ITimingViewSession* CurrentTimingViewSession = (::Insights::ITimingViewSession*)(Timing::ITimingViewSession*)this;
+	for (::Insights::ITimingViewExtender* Extender : GetOldExtenders())
+	{
+		Extender->OnBeginSession(*CurrentTimingViewSession);
+	}
+PRAGMA_ENABLE_DEPRECATION_WARNINGS
+#endif // UE_INSIGHTS_BACKWARD_COMPATIBILITY_UE54
 }
 
 ////////////////////////////////////////////////////////////////////////////////////////////////////
@@ -600,7 +643,7 @@ void STimingView::Tick(const FGeometry& AllottedGeometry, const double InCurrent
 				Viewport.AddDirtyFlags(ETimingTrackViewportDirtyFlags::HClippedSessionTimeChanged);
 			}
 
-			//UE_LOG(TimingProfiler, Log, TEXT("Session Duration: %g"), DT);
+			//UE_LOG(LogTimingProfiler, Log, TEXT("Session Duration: %g"), DT);
 			Viewport.SetMaxValidTime(SessionTime);
 			UpdateHorizontalScrollBar();
 		}
@@ -676,10 +719,19 @@ void STimingView::Tick(const FGeometry& AllottedGeometry, const double InCurrent
 	{
 		// Tick plugin extenders.
 		// Each extender can add/remove tracks and/or change order of tracks.
-		for (Insights::ITimingViewExtender* Extender : GetExtenders())
+		for (Timing::ITimingViewExtender* Extender : GetExtenders())
 		{
 			Extender->Tick(*this, *Session.Get());
 		}
+#if UE_INSIGHTS_BACKWARD_COMPATIBILITY_UE54
+PRAGMA_DISABLE_DEPRECATION_WARNINGS
+		::Insights::ITimingViewSession* CurrentTimingViewSession = (::Insights::ITimingViewSession*)(Timing::ITimingViewSession*)this;
+		for (::Insights::ITimingViewExtender* Extender : GetOldExtenders())
+		{
+			Extender->Tick(*CurrentTimingViewSession, *Session.Get());
+		}
+PRAGMA_ENABLE_DEPRECATION_WARNINGS
+#endif // UE_INSIGHTS_BACKWARD_COMPATIBILITY_UE54
 
 		// Re-sort now (if we need to).
 		UpdateScrollableTracksOrder();
@@ -1687,7 +1739,7 @@ int32 STimingView::OnPaint(const FPaintArgs& Args, const FGeometry& AllottedGeom
 			FString::Printf(TEXT("SX: %g, ST: %g, ET: %s"),
 				Viewport.GetScaleX(),
 				Viewport.GetStartTime(),
-				*TimeUtils::FormatTimeAuto(Viewport.GetMaxValidTime())),
+				*FormatTimeAuto(Viewport.GetMaxValidTime())),
 			SummaryFont, DbgTextColor
 		);
 		DbgY += DbgDY;
@@ -1760,7 +1812,8 @@ void STimingView::AddTrack(TSharedPtr<FBaseTimingTrack> Track, ETimingTrackLocat
 	const int32 MaxNumTracks = 1000;
 	if (TrackList.Num() >= MaxNumTracks)
 	{
-		UE_LOG(TimingProfiler, Warning, TEXT("Too many tracks already created (%d tracks)! Ignoring %s track : %s (\"%s\")"),
+		using namespace UE::Insights::TimingProfiler;
+		UE_LOG(LogTimingProfiler, Warning, TEXT("Too many tracks already created (%d tracks)! Ignoring %s track : %s (\"%s\")"),
 			TrackList.Num(),
 			LocationName,
 			*Track->GetTypeName().ToString(),
@@ -1769,7 +1822,7 @@ void STimingView::AddTrack(TSharedPtr<FBaseTimingTrack> Track, ETimingTrackLocat
 	}
 
 #if 0
-	UE_LOG(TimingProfiler, Log, TEXT("New %s Track (%d) : %s (\"%s\")"),
+	UE_LOG(LogTimingProfiler, Log, TEXT("New %s Track (%d) : %s (\"%s\")"),
 		LocationName,
 		TrackList.Num() + 1,
 		*Track->GetTypeName().ToString(),
@@ -1822,7 +1875,7 @@ bool STimingView::RemoveTrack(TSharedPtr<FBaseTimingTrack> Track)
 		OnTrackRemovedDelegate.Broadcast(Track);
 
 #if 0
-		UE_LOG(TimingProfiler, Log, TEXT("Removed %s Track (%d) : %s (\"%s\")"),
+		UE_LOG(LogTimingProfiler, Log, TEXT("Removed %s Track (%d) : %s (\"%s\")"),
 			LocationName,
 			TrackList.Num(),
 			*Track->GetTypeName().ToString(),
@@ -1978,7 +2031,7 @@ FReply STimingView::OnMouseButtonDown(const FGeometry& MyGeometry, const FPointe
 		}
 	}
 
-	TSharedPtr<Insights::FTimeMarker> ScrubbingTimeMarker = nullptr;
+	TSharedPtr<FTimeMarker> ScrubbingTimeMarker = nullptr;
 
 	if (bStartPanningSelectingOrScrubbing)
 	{
@@ -2449,7 +2502,7 @@ FReply STimingView::OnMouseMove(const FGeometry& MyGeometry, const FPointerEvent
 					Time = MarkersTrack->Snap(Time, SnapTolerance);
 				}
 
-				TSharedRef<Insights::FTimeMarker> ScrubbingTimeMarker = TimeRulerTrack->GetScrubbingTimeMarker();
+				TSharedRef<FTimeMarker> ScrubbingTimeMarker = TimeRulerTrack->GetScrubbingTimeMarker();
 				ScrubbingTimeMarker->SetTime(Time);
 				RaiseTimeMarkerChanging(ScrubbingTimeMarker);
 			}
@@ -2904,6 +2957,7 @@ FReply STimingView::OnKeyUp(const FGeometry& MyGeometry, const FKeyEvent& InKeyE
 
 void STimingView::ShowContextMenu(const FPointerEvent& MouseEvent)
 {
+	using namespace UE::Insights::TimingProfiler;
 	const FTimingViewCommands& Commands = FTimingViewCommands::Get();
 
 	const bool bShouldCloseWindowAfterMenuSelection = true;
@@ -2964,10 +3018,19 @@ void STimingView::ShowContextMenu(const FPointerEvent& MouseEvent)
 	}
 	MenuBuilder.EndSection();
 
-	for (Insights::ITimingViewExtender* Extender : GetExtenders())
+	for (Timing::ITimingViewExtender* Extender : GetExtenders())
 	{
 		bHasAnyActions |= Extender->ExtendGlobalContextMenu(*this, MenuBuilder);
 	}
+#if UE_INSIGHTS_BACKWARD_COMPATIBILITY_UE54
+PRAGMA_DISABLE_DEPRECATION_WARNINGS
+	::Insights::ITimingViewSession* CurrentTimingViewSession = (::Insights::ITimingViewSession*)(Timing::ITimingViewSession*)this;
+	for (::Insights::ITimingViewExtender* Extender : GetOldExtenders())
+	{
+		bHasAnyActions |= Extender->ExtendGlobalContextMenu(*CurrentTimingViewSession, MenuBuilder);
+	}
+PRAGMA_ENABLE_DEPRECATION_WARNINGS
+#endif // UE_INSIGHTS_BACKWARD_COMPATIBILITY_UE54
 
 	if (!bHasAnyActions)
 	{
@@ -3158,6 +3221,7 @@ bool STimingView::CheckTrackLocation(TSharedRef<FBaseTimingTrack> Track, ETiming
 
 void STimingView::BindCommands()
 {
+	using namespace UE::Insights::TimingProfiler;
 	FTimingViewCommands::Register();
 
 	const FTimingViewCommands& Commands = FTimingViewCommands::Get();
@@ -3536,38 +3600,38 @@ void STimingView::SnapToFrameBound(double& StartTime, double& EndTime)
 
 void STimingView::RaiseSelectionChanging()
 {
-	OnSelectionChangedDelegate.Broadcast(Insights::ETimeChangedFlags::Interactive, SelectionStartTime, SelectionEndTime);
+	OnSelectionChangedDelegate.Broadcast(Timing::ETimeChangedFlags::Interactive, SelectionStartTime, SelectionEndTime);
 }
 
 ////////////////////////////////////////////////////////////////////////////////////////////////////
 
 void STimingView::RaiseSelectionChanged()
 {
-	OnSelectionChangedDelegate.Broadcast(Insights::ETimeChangedFlags::None, SelectionStartTime, SelectionEndTime);
+	OnSelectionChangedDelegate.Broadcast(Timing::ETimeChangedFlags::None, SelectionStartTime, SelectionEndTime);
 }
 
 ////////////////////////////////////////////////////////////////////////////////////////////////////
 
-void STimingView::RaiseTimeMarkerChanging(TSharedRef<Insights::FTimeMarker> InTimeMarker)
+void STimingView::RaiseTimeMarkerChanging(TSharedRef<FTimeMarker> InTimeMarker)
 {
 	if (InTimeMarker == DefaultTimeMarker)
 	{
 		const double Time = DefaultTimeMarker->GetTime();
-		OnTimeMarkerChangedDelegate.Broadcast(Insights::ETimeChangedFlags::Interactive, Time);
+		OnTimeMarkerChangedDelegate.Broadcast(Timing::ETimeChangedFlags::Interactive, Time);
 	}
-	OnCustomTimeMarkerChangedDelegate.Broadcast(Insights::ETimeChangedFlags::Interactive, InTimeMarker);
+	OnCustomTimeMarkerChangedDelegate.Broadcast(Timing::ETimeChangedFlags::Interactive, InTimeMarker);
 }
 
 ////////////////////////////////////////////////////////////////////////////////////////////////////
 
-void STimingView::RaiseTimeMarkerChanged(TSharedRef<Insights::FTimeMarker> InTimeMarker)
+void STimingView::RaiseTimeMarkerChanged(TSharedRef<FTimeMarker> InTimeMarker)
 {
 	if (InTimeMarker == DefaultTimeMarker)
 	{
 		const double Time = DefaultTimeMarker->GetTime();
-		OnTimeMarkerChangedDelegate.Broadcast(Insights::ETimeChangedFlags::None, Time);
+		OnTimeMarkerChangedDelegate.Broadcast(Timing::ETimeChangedFlags::None, Time);
 	}
-	OnCustomTimeMarkerChangedDelegate.Broadcast(Insights::ETimeChangedFlags::None, InTimeMarker);
+	OnCustomTimeMarkerChangedDelegate.Broadcast(Timing::ETimeChangedFlags::None, InTimeMarker);
 }
 
 ////////////////////////////////////////////////////////////////////////////////////////////////////
@@ -3644,7 +3708,7 @@ void STimingView::SetTimeMarkersVisible(bool bIsMarkersTrackVisible)
 		{
 			if (Viewport.GetScrollPosY() != 0.0f)
 			{
-				UE_LOG(TimingProfiler, Log, TEXT("SetTimeMarkersVisible!!!"));
+				UE_LOG(LogTimingProfiler, Log, TEXT("SetTimeMarkersVisible!!!"));
 				Viewport.SetScrollPosY(Viewport.GetScrollPosY() + MarkersTrack->GetHeight());
 			}
 
@@ -3652,7 +3716,7 @@ void STimingView::SetTimeMarkersVisible(bool bIsMarkersTrackVisible)
 		}
 		else
 		{
-			UE_LOG(TimingProfiler, Log, TEXT("SetTimeMarkersVisible!!!"));
+			UE_LOG(LogTimingProfiler, Log, TEXT("SetTimeMarkersVisible!!!"));
 			Viewport.SetScrollPosY(Viewport.GetScrollPosY() - MarkersTrack->GetHeight());
 		}
 	}
@@ -3671,7 +3735,7 @@ void STimingView::SetDrawOnlyBookmarks(bool bIsBookmarksTrack)
 		{
 			if (Viewport.GetScrollPosY() != 0.0f)
 			{
-				UE_LOG(TimingProfiler, Log, TEXT("SetDrawOnlyBookmarks!!!"));
+				UE_LOG(LogTimingProfiler, Log, TEXT("SetDrawOnlyBookmarks!!!"));
 				Viewport.SetScrollPosY(Viewport.GetScrollPosY() + MarkersTrack->GetHeight() - PrevHeight);
 			}
 
@@ -3769,7 +3833,7 @@ void STimingView::UpdateHoveredTimingEvent(float InMousePosX, float InMousePosY)
 		const double DT = Stopwatch.GetAccumulatedTime();
 		if (DT > 0.001)
 		{
-			UE_LOG(TimingProfiler, Log, TEXT("HoveredTrack [%g, %g] GetEvent: %.1f ms"), InMousePosX, InMousePosY, DT * 1000.0);
+			UE_LOG(LogTimingProfiler, Log, TEXT("HoveredTrack [%g, %g] GetEvent: %.1f ms"), InMousePosX, InMousePosY, DT * 1000.0);
 		}
 	}
 
@@ -3798,7 +3862,7 @@ void STimingView::UpdateHoveredTimingEvent(float InMousePosX, float InMousePosY)
 			const double T3 = Stopwatch.GetAccumulatedTime();
 			if (T3 > 0.001)
 			{
-				UE_LOG(TimingProfiler, Log, TEXT("HoveredTrack [%g, %g] Tooltip: %.1f ms (%.1f + %.1f + %.1f)"),
+				UE_LOG(LogTimingProfiler, Log, TEXT("HoveredTrack [%g, %g] Tooltip: %.1f ms (%.1f + %.1f + %.1f)"),
 					InMousePosX, InMousePosY, T3 * 1000.0, T1 * 1000.0, (T2 - T1) * 1000.0, (T3 - T2) * 1000.0);
 			}
 		}
@@ -3902,10 +3966,11 @@ void STimingView::SelectTimingEvent(const TSharedPtr<const ITimingEvent> InEvent
 
 void STimingView::ToggleGraphSeries(const TSharedPtr<const ITimingEvent> InEvent)
 {
-	if(InEvent.Get() && InEvent.Get()->Is<FThreadTrackEvent>() && IsInTimingProfiler())
+	if (InEvent.Get() && InEvent.Get()->Is<FThreadTrackEvent>() && IsInTimingProfiler())
 	{
 		const FThreadTrackEvent& TrackEvent = InEvent.Get()->As<FThreadTrackEvent>();
 
+		using namespace UE::Insights::TimingProfiler;
 		FTimingProfilerManager::Get()->ToggleTimingViewMainGraphEventSeries(TrackEvent.GetTimerId());
 	}
 }
@@ -4594,16 +4659,33 @@ TSharedRef<SWidget> STimingView::MakeCpuGpuTracksFilterMenu()
 	FMenuBuilder MenuBuilder(/*bInShouldCloseWindowAfterMenuSelection=*/true, CommandList);
 
 	// Let any plugin extend the GPU Tracks Filter menu.
-	for (Insights::ITimingViewExtender* Extender : GetExtenders())
+	for (Timing::ITimingViewExtender* Extender : GetExtenders())
 	{
 		Extender->ExtendGpuTracksFilterMenu(*this, MenuBuilder);
 	}
+#if UE_INSIGHTS_BACKWARD_COMPATIBILITY_UE54
+PRAGMA_DISABLE_DEPRECATION_WARNINGS
+	::Insights::ITimingViewSession* CurrentTimingViewSession = (::Insights::ITimingViewSession*)(Timing::ITimingViewSession*)this;
+	for (::Insights::ITimingViewExtender* Extender : GetOldExtenders())
+	{
+		Extender->ExtendGpuTracksFilterMenu(*CurrentTimingViewSession, MenuBuilder);
+	}
+PRAGMA_ENABLE_DEPRECATION_WARNINGS
+#endif // UE_INSIGHTS_BACKWARD_COMPATIBILITY_UE54
 
 	// Let any plugin extend the CPU Tracks Filter menu.
-	for (Insights::ITimingViewExtender* Extender : GetExtenders())
+	for (Timing::ITimingViewExtender* Extender : GetExtenders())
 	{
 		Extender->ExtendCpuTracksFilterMenu(*this, MenuBuilder);
 	}
+#if UE_INSIGHTS_BACKWARD_COMPATIBILITY_UE54
+PRAGMA_DISABLE_DEPRECATION_WARNINGS
+	for (::Insights::ITimingViewExtender* Extender : GetOldExtenders())
+	{
+		Extender->ExtendCpuTracksFilterMenu(*CurrentTimingViewSession, MenuBuilder);
+	}
+PRAGMA_ENABLE_DEPRECATION_WARNINGS
+#endif // UE_INSIGHTS_BACKWARD_COMPATIBILITY_UE54
 
 	return MenuBuilder.MakeWidget();
 }
@@ -4623,10 +4705,19 @@ TSharedRef<SWidget> STimingView::MakeOtherTracksFilterMenu()
 	MenuBuilder.EndSection();
 
 	// Let any plugin extend the Other Tracks Filter menu.
-	for (Insights::ITimingViewExtender* Extender : GetExtenders())
+	for (Timing::ITimingViewExtender* Extender : GetExtenders())
 	{
 		Extender->ExtendOtherTracksFilterMenu(*this, MenuBuilder);
 	}
+#if UE_INSIGHTS_BACKWARD_COMPATIBILITY_UE54
+PRAGMA_DISABLE_DEPRECATION_WARNINGS
+	::Insights::ITimingViewSession* CurrentTimingViewSession = (::Insights::ITimingViewSession*)(Timing::ITimingViewSession*)this;
+	for (::Insights::ITimingViewExtender* Extender : GetOldExtenders())
+	{
+		Extender->ExtendOtherTracksFilterMenu(*CurrentTimingViewSession, MenuBuilder);
+	}
+PRAGMA_ENABLE_DEPRECATION_WARNINGS
+#endif // UE_INSIGHTS_BACKWARD_COMPATIBILITY_UE54
 
 	return MenuBuilder.MakeWidget();
 }
@@ -4652,10 +4743,19 @@ TSharedRef<SWidget> STimingView::MakePluginTracksFilterMenu()
 	FMenuBuilder MenuBuilder(/*bInShouldCloseWindowAfterMenuSelection=*/true, CommandList);
 
 	// Let any plugin extend the filter menu.
-	for (Insights::ITimingViewExtender* Extender : GetExtenders())
+	for (Timing::ITimingViewExtender* Extender : GetExtenders())
 	{
 		Extender->ExtendFilterMenu(*this, MenuBuilder);
 	}
+#if UE_INSIGHTS_BACKWARD_COMPATIBILITY_UE54
+PRAGMA_DISABLE_DEPRECATION_WARNINGS
+	::Insights::ITimingViewSession* CurrentTimingViewSession = (::Insights::ITimingViewSession*)(Timing::ITimingViewSession*)this;
+	for (::Insights::ITimingViewExtender* Extender : GetOldExtenders())
+	{
+		Extender->ExtendFilterMenu(*CurrentTimingViewSession, MenuBuilder);
+	}
+PRAGMA_ENABLE_DEPRECATION_WARNINGS
+#endif // UE_INSIGHTS_BACKWARD_COMPATIBILITY_UE54
 
 	return MenuBuilder.MakeWidget();
 }
@@ -4851,9 +4951,9 @@ void STimingView::CreateCpuThreadTrackColoringModeMenu(FMenuBuilder& MenuBuilder
 			LOCTEXT("CpuThreadTrackColoringMode_ByTimerName_Desc", "Assign a color to CPU/GPU timing events based on their timer name."),
 			FSlateIcon(),
 			FUIAction(
-				FExecuteAction::CreateSP(this, &STimingView::SetCpuThreadTrackColoringMode, Insights::ETimingEventsColoringMode::ByTimerName),
+				FExecuteAction::CreateSP(this, &STimingView::SetCpuThreadTrackColoringMode, ETimingEventsColoringMode::ByTimerName),
 				FCanExecuteAction(),
-				FIsActionChecked::CreateSP(this, &STimingView::CheckCpuThreadTrackColoringMode, Insights::ETimingEventsColoringMode::ByTimerName)),
+				FIsActionChecked::CreateSP(this, &STimingView::CheckCpuThreadTrackColoringMode, ETimingEventsColoringMode::ByTimerName)),
 			NAME_None,
 			EUserInterfaceActionType::RadioButton
 		);
@@ -4862,9 +4962,9 @@ void STimingView::CreateCpuThreadTrackColoringModeMenu(FMenuBuilder& MenuBuilder
 			LOCTEXT("CpuThreadTrackColoringMode_ByTimerId_Desc", "Assign a color to CPU/GPU timing events based on their timer id."),
 			FSlateIcon(),
 			FUIAction(
-				FExecuteAction::CreateSP(this, &STimingView::SetCpuThreadTrackColoringMode, Insights::ETimingEventsColoringMode::ByTimerId),
+				FExecuteAction::CreateSP(this, &STimingView::SetCpuThreadTrackColoringMode, ETimingEventsColoringMode::ByTimerId),
 				FCanExecuteAction(),
-				FIsActionChecked::CreateSP(this, &STimingView::CheckCpuThreadTrackColoringMode, Insights::ETimingEventsColoringMode::ByTimerId)),
+				FIsActionChecked::CreateSP(this, &STimingView::CheckCpuThreadTrackColoringMode, ETimingEventsColoringMode::ByTimerId)),
 			NAME_None,
 			EUserInterfaceActionType::RadioButton
 		);
@@ -4873,9 +4973,9 @@ void STimingView::CreateCpuThreadTrackColoringModeMenu(FMenuBuilder& MenuBuilder
 			LOCTEXT("CpuThreadTrackColoringMode_BySourceFile_Desc", "Assign a color to CPU/GPU timing events based on their source file."),
 			FSlateIcon(),
 			FUIAction(
-				FExecuteAction::CreateSP(this, &STimingView::SetCpuThreadTrackColoringMode, Insights::ETimingEventsColoringMode::BySourceFile),
+				FExecuteAction::CreateSP(this, &STimingView::SetCpuThreadTrackColoringMode, ETimingEventsColoringMode::BySourceFile),
 				FCanExecuteAction(),
-				FIsActionChecked::CreateSP(this, &STimingView::CheckCpuThreadTrackColoringMode, Insights::ETimingEventsColoringMode::BySourceFile)),
+				FIsActionChecked::CreateSP(this, &STimingView::CheckCpuThreadTrackColoringMode, ETimingEventsColoringMode::BySourceFile)),
 			NAME_None,
 			EUserInterfaceActionType::RadioButton
 		);
@@ -4884,9 +4984,9 @@ void STimingView::CreateCpuThreadTrackColoringModeMenu(FMenuBuilder& MenuBuilder
 			LOCTEXT("CpuThreadTrackColoringMode_ByDuration_Desc", "Assign a color to CPU/GPU timing events based on their duration (inclusive time).\n\t≥ 10ms : red\n\t≥ 1ms : yellow\n\t≥ 100μs : green\n\t≥ 10μs : cyan\n\t≥ 1μs : blue\n\t< 1μs : gray"),
 			FSlateIcon(),
 			FUIAction(
-				FExecuteAction::CreateSP(this, &STimingView::SetCpuThreadTrackColoringMode, Insights::ETimingEventsColoringMode::ByDuration),
+				FExecuteAction::CreateSP(this, &STimingView::SetCpuThreadTrackColoringMode, ETimingEventsColoringMode::ByDuration),
 				FCanExecuteAction(),
-				FIsActionChecked::CreateSP(this, &STimingView::CheckCpuThreadTrackColoringMode, Insights::ETimingEventsColoringMode::ByDuration)),
+				FIsActionChecked::CreateSP(this, &STimingView::CheckCpuThreadTrackColoringMode, ETimingEventsColoringMode::ByDuration)),
 			NAME_None,
 			EUserInterfaceActionType::RadioButton
 		);
@@ -4899,14 +4999,14 @@ void STimingView::CreateCpuThreadTrackColoringModeMenu(FMenuBuilder& MenuBuilder
 void STimingView::ChooseNextCpuThreadTrackColoringMode()
 {
 	uint32 Mode = (uint32)FTimingProfilerManager::Get()->GetColoringMode();
-	Mode = (Mode + 1) % (uint32)Insights::ETimingEventsColoringMode::Count;
-	FTimingProfilerManager::Get()->SetColoringMode((Insights::ETimingEventsColoringMode)Mode);
+	Mode = (Mode + 1) % (uint32)ETimingEventsColoringMode::Count;
+	FTimingProfilerManager::Get()->SetColoringMode((ETimingEventsColoringMode)Mode);
 	Viewport.AddDirtyFlags(ETimingTrackViewportDirtyFlags::HInvalidated);
 }
 
 ////////////////////////////////////////////////////////////////////////////////////////////////////
 
-void STimingView::SetCpuThreadTrackColoringMode(Insights::ETimingEventsColoringMode Mode)
+void STimingView::SetCpuThreadTrackColoringMode(UE::Insights::ETimingEventsColoringMode Mode)
 {
 	FTimingProfilerManager::Get()->SetColoringMode(Mode);
 	Viewport.AddDirtyFlags(ETimingTrackViewportDirtyFlags::HInvalidated);
@@ -4914,7 +5014,7 @@ void STimingView::SetCpuThreadTrackColoringMode(Insights::ETimingEventsColoringM
 
 ////////////////////////////////////////////////////////////////////////////////////////////////////
 
-bool STimingView::CheckCpuThreadTrackColoringMode(Insights::ETimingEventsColoringMode Mode)
+bool STimingView::CheckCpuThreadTrackColoringMode(UE::Insights::ETimingEventsColoringMode Mode)
 {
 	return Mode == FTimingProfilerManager::Get()->GetColoringMode();
 }
@@ -4954,7 +5054,7 @@ bool STimingView::QuickFind_CanExecute() const
 
 void STimingView::QuickFind_Execute()
 {
-	using namespace Insights;
+	using namespace UE::Insights;
 
 	LLM_SCOPE_BYTAG(Insights);
 
@@ -5007,15 +5107,23 @@ void STimingView::QuickFind_Execute()
 			nullptr,
 			FFilterService::Get()->GetIntegerOperators()));
 
-		NewFilterConfigurator->Add(MakeShared<FTimerNameFilter>());
-		NewFilterConfigurator->Add(MakeShared<FMetadataFilter>());
+		NewFilterConfigurator->Add(MakeShared<::Insights::FTimerNameFilter>());
+		NewFilterConfigurator->Add(MakeShared<::Insights::FMetadataFilter>());
 
-		for (Insights::ITimingViewExtender* Extender : GetExtenders())
+		for (Timing::ITimingViewExtender* Extender : GetExtenders())
 		{
 			Extender->AddQuickFindFilters(NewFilterConfigurator);
 		}
+#if UE_INSIGHTS_BACKWARD_COMPATIBILITY_UE54
+PRAGMA_DISABLE_DEPRECATION_WARNINGS
+		for (::Insights::ITimingViewExtender* Extender : GetOldExtenders())
+		{
+			Extender->AddQuickFindFilters(NewFilterConfigurator);
+		}
+PRAGMA_ENABLE_DEPRECATION_WARNINGS
+#endif // UE_INSIGHTS_BACKWARD_COMPATIBILITY_UE54
 
-		QuickFindVm = MakeShared<FQuickFind>(NewFilterConfigurator);
+		QuickFindVm = MakeShared<::Insights::FQuickFind>(NewFilterConfigurator);
 		QuickFindVm->GetOnFindFirstEvent().AddSP(this, &STimingView::FindFirstEvent);
 		QuickFindVm->GetOnFindPreviousEvent().AddSP(this, &STimingView::FindPrevEvent);
 		QuickFindVm->GetOnFindNextEvent().AddSP(this, &STimingView::FindNextEvent);
@@ -5024,7 +5132,7 @@ void STimingView::QuickFind_Execute()
 		QuickFindVm->GetOnClearFiltersEvent().AddSP(this, &STimingView::ClearFilters);
 	}
 
-	SAssignNew(QuickFindWidgetSharedPtr, SQuickFind, QuickFindVm);
+	SAssignNew(QuickFindWidgetSharedPtr, ::Insights::SQuickFind, QuickFindVm);
 
 	if (FGlobalTabmanager::Get()->HasTabSpawner(QuickFindTabId))
 	{
@@ -5064,7 +5172,7 @@ TSharedRef<SDockTab> STimingView::SpawnQuickFindTab(const FSpawnTabArgs& Args)
 
 void STimingView::CloseQuickFindTab()
 {
-	TSharedPtr<Insights::SQuickFind> QuickFindWidget = QuickFindWidgetWeakPtr.Pin();
+	TSharedPtr<::Insights::SQuickFind> QuickFindWidget = QuickFindWidgetWeakPtr.Pin();
 	if (QuickFindWidget)
 	{
 		TSharedPtr<SDockTab> QuickFindTab = QuickFindWidget->GetParentTab().Pin();
@@ -5122,10 +5230,21 @@ TSharedPtr<FBaseTimingTrack> STimingView::FindTrack(uint64 InTrackId)
 
 ////////////////////////////////////////////////////////////////////////////////////////////////////
 
-TArray<Insights::ITimingViewExtender*> STimingView::GetExtenders() const
+TArray<Timing::ITimingViewExtender*> STimingView::GetExtenders() const
+{
+	return IModularFeatures::Get().GetModularFeatureImplementations<Timing::ITimingViewExtender>(Timing::TimingViewExtenderFeatureName);
+}
+
+////////////////////////////////////////////////////////////////////////////////////////////////////
+
+#if UE_INSIGHTS_BACKWARD_COMPATIBILITY_UE54
+PRAGMA_DISABLE_DEPRECATION_WARNINGS
+TArray<::Insights::ITimingViewExtender*> STimingView::GetOldExtenders() const
 {
-	return IModularFeatures::Get().GetModularFeatureImplementations<Insights::ITimingViewExtender>(Insights::TimingViewExtenderFeatureName);
+	return IModularFeatures::Get().GetModularFeatureImplementations<::Insights::ITimingViewExtender>(::Insights::TimingViewExtenderFeatureName);
 }
+PRAGMA_ENABLE_DEPRECATION_WARNINGS
+#endif // UE_INSIGHTS_BACKWARD_COMPATIBILITY_UE54
 
 ////////////////////////////////////////////////////////////////////////////////////////////////////
 
@@ -5271,7 +5390,8 @@ void STimingView::FindLastEvent()
 void STimingView::FilterAllTracks()
 {
 	LLM_SCOPE_BYTAG(Insights);
-	FilterConfigurator = MakeShared<Insights::FFilterConfigurator>(*QuickFindVm->GetFilterConfigurator());
+
+	FilterConfigurator = MakeShared<UE::Insights::FFilterConfigurator>(*QuickFindVm->GetFilterConfigurator());
 
 	for (auto& Entry : AllTracks)
 	{
@@ -5284,7 +5404,9 @@ void STimingView::FilterAllTracks()
 void STimingView::ClearFilters()
 {
 	LLM_SCOPE_BYTAG(Insights);
+
 	FilterConfigurator.Reset();
+
 	for (auto& Entry : AllTracks)
 	{
 		Entry.Value->SetFilterConfigurator(nullptr);
@@ -5341,11 +5463,10 @@ void STimingView::PopulateTimerNameSuggestionList(const FString& Text, TArray<FS
 
 ////////////////////////////////////////////////////////////////////////////////////////////////////
 
-void STimingView::EnumerateFilteredTracks(TSharedPtr<Insights::FFilterConfigurator> InFilterConfigurator, TSharedPtr<const FBaseTimingTrack> PriorityTrack, EnumerateFilteredTracksCallback Callback)
+void STimingView::EnumerateFilteredTracks(TSharedPtr<UE::Insights::FFilterConfigurator> InFilterConfigurator, TSharedPtr<const FBaseTimingTrack> PriorityTrack, EnumerateFilteredTracksCallback Callback)
 {
-	using namespace Insights;
-
-	Insights::FFilterContext FilterContext;
+	using EFilterField = UE::Insights::EFilterField;
+	UE::Insights::FFilterContext FilterContext;
 	FilterContext.AddFilterData(static_cast<int32>(EFilterField::TrackName), FString());
 
 	// Call the callback for the PriorityTrack first if it passes the filters.
@@ -5463,5 +5584,7 @@ bool STimingView::IsInTimingProfiler()
 
 ////////////////////////////////////////////////////////////////////////////////////////////////////
 
+} // namespace UE::Insights::TimingProfiler
+
 #undef INSIGHTS_ACTIVATE_BENCHMARK
 #undef LOCTEXT_NAMESPACE
