[Insights] Massive refactoring of TraceInsights module:

diff --git a/Engine/Source/Developer/TraceInsights/Private/Insights/Widgets/STimingView.h b/Engine/Source/Developer/TraceInsights/Private/Insights/Widgets/STimingView.h
--- a/Engine/Source/Developer/TraceInsights/Private/Insights/Widgets/STimingView.h
+++ b/Engine/Source/Developer/TraceInsights/Private/Insights/Widgets/STimingView.h
@@ -13,8 +13,11 @@
 #include "Widgets/SCompoundWidget.h"
 #include "Widgets/SWidget.h"
 
-// Insights
-#include "Insights/Common/FixedCircularBuffer.h"
+// TraceInsightsCore
+#include "InsightsCore/Common/FixedCircularBuffer.h"
+
+// TraceInsights
+#include "Insights/Config.h"
 #include "Insights/ITimingViewSession.h"
 #include "Insights/ViewModels/BaseTimingTrack.h"
 #include "Insights/ViewModels/TimerNode.h"
@@ -24,31 +27,49 @@
 #include "Insights/ViewModels/TimingViewDrawHelper.h"
 #include "Insights/ViewModels/TooltipDrawState.h"
 
+#if UE_INSIGHTS_BACKWARD_COMPATIBILITY_UE54
+#include "Insights/ITimingViewExtender.h"
+#endif
+
+class FMenuBuilder;
+class FUICommandList;
+class SDockTab;
+class SOverlay;
+class SScrollBar;
+class FSpawnTabArgs;
+
 class FFileActivitySharedState;
 class FFrameSharedState;
 class FLoadingSharedState;
 class FMarkersTimingTrack;
-class FMenuBuilder;
-class FSpawnTabArgs;
 class FThreadTimingSharedState;
-class FTimeRulerTrack;
 class FTimingGraphTrack;
 class FTimingViewDrawHelper;
-class FUICommandList;
-class SDockTab;
-class SOverlay;
-class SScrollBar;
 
 namespace Insights
 {
-	class ITimingViewExtender;
-	class FTimeMarker;
 	class FTimingRegionsSharedState;
 	class FQuickFind;
 	class SQuickFind;
+}
+
+namespace UE::Insights
+{
+	class FFilterConfigurator;
 	enum class ETimingEventsColoringMode : uint32;
 }
 
+namespace UE::Insights::Timing
+{
+	class ITimingViewExtender;
+}
+
+namespace UE::Insights::TimingProfiler
+{
+
+class FTimeMarker;
+class FTimeRulerTrack;
+
 enum class ESelectEventType : uint32
 {
 	Min = 0,
@@ -56,7 +77,7 @@ enum class ESelectEventType : uint32
 };
 
 /** A custom widget used to display timing events. */
-class STimingView : public SCompoundWidget, public Insights::ITimingViewSession
+class STimingView : public SCompoundWidget, public UE::Insights::Timing::ITimingViewSession
 {
 public:
 	/** Default constructor. */
@@ -108,8 +129,8 @@ public:
 	const TSharedRef<const FTimeRulerTrack> GetTimeRulerTrack() const { return TimeRulerTrack; }
 
 	/** Gets the default (custom) time marker (for backward compatibility). */
-	TSharedRef<Insights::FTimeMarker> GetDefaultTimeMarker() { return DefaultTimeMarker; }
-	const TSharedRef<const Insights::FTimeMarker> GetDefaultTimeMarker() const { return DefaultTimeMarker; }
+	TSharedRef<FTimeMarker> GetDefaultTimeMarker() { return DefaultTimeMarker; }
+	const TSharedRef<const FTimeMarker> GetDefaultTimeMarker() const { return DefaultTimeMarker; }
 
 	/** Resets internal widget's data to the default one. */
 	void Reset(bool bIsFirstReset = false);
@@ -226,7 +247,7 @@ public:
 	virtual void OnDragLeave(const FDragDropEvent& DragDropEvent) override;
 
 	/**
-	 * Called during drag and drop when the the mouse is being dragged over a widget.
+	 * Called during drag and drop when the mouse is being dragged over a widget.
 	 *
 	 * @param MyGeometry      The geometry of the widget receiving the event.
 	 * @param DragDropEvent   The drag and drop event.
@@ -271,16 +292,16 @@ public:
 	virtual void SetTimeMarker(double InTimeMarker) override;
 	virtual void SetAndCenterOnTimeMarker(double InTimeMarker) override;
 
-	virtual Insights::FSelectionChangedDelegate& OnSelectionChanged() override { return OnSelectionChangedDelegate; }
-	virtual Insights::FTimeMarkerChangedDelegate& OnTimeMarkerChanged() override { return OnTimeMarkerChangedDelegate; }
-	virtual Insights::FCustomTimeMarkerChangedDelegate& OnCustomTimeMarkerChanged() override { return OnCustomTimeMarkerChangedDelegate; }
-	virtual Insights::FHoveredTrackChangedDelegate& OnHoveredTrackChanged() override { return OnHoveredTrackChangedDelegate; }
-	virtual Insights::FHoveredEventChangedDelegate& OnHoveredEventChanged() override { return OnHoveredEventChangedDelegate; }
-	virtual Insights::FSelectedTrackChangedDelegate& OnSelectedTrackChanged() override { return OnSelectedTrackChangedDelegate; }
-	virtual Insights::FSelectedEventChangedDelegate& OnSelectedEventChanged() override { return OnSelectedEventChangedDelegate; }
-	virtual Insights::FTrackVisibilityChangedDelegate& OnTrackVisibilityChanged() override { return OnTrackVisibilityChangedDelegate; }
-	virtual Insights::FTrackAddedDelegate& OnTrackAdded() override { return OnTrackAddedDelegate; }
-	virtual Insights::FTrackRemovedDelegate& OnTrackRemoved() override { return OnTrackRemovedDelegate; }
+	virtual Timing::FSelectionChangedDelegate& OnSelectionChanged() override { return OnSelectionChangedDelegate; }
+	virtual Timing::FTimeMarkerChangedDelegate& OnTimeMarkerChanged() override { return OnTimeMarkerChangedDelegate; }
+	virtual Timing::FCustomTimeMarkerChangedDelegate& OnCustomTimeMarkerChanged() override { return OnCustomTimeMarkerChangedDelegate; }
+	virtual Timing::FHoveredTrackChangedDelegate& OnHoveredTrackChanged() override { return OnHoveredTrackChangedDelegate; }
+	virtual Timing::FHoveredEventChangedDelegate& OnHoveredEventChanged() override { return OnHoveredEventChangedDelegate; }
+	virtual Timing::FSelectedTrackChangedDelegate& OnSelectedTrackChanged() override { return OnSelectedTrackChangedDelegate; }
+	virtual Timing::FSelectedEventChangedDelegate& OnSelectedEventChanged() override { return OnSelectedEventChangedDelegate; }
+	virtual Timing::FTrackVisibilityChangedDelegate& OnTrackVisibilityChanged() override { return OnTrackVisibilityChangedDelegate; }
+	virtual Timing::FTrackAddedDelegate& OnTrackAdded() override { return OnTrackAddedDelegate; }
+	virtual Timing::FTrackRemovedDelegate& OnTrackRemoved() override { return OnTrackRemovedDelegate; }
 
 	virtual void ResetSelectedEvent() override
 	{
@@ -395,7 +416,7 @@ public:
 
 	void CloseQuickFindTab();
 
-	TSharedPtr<Insights::FFilterConfigurator> GetFilterConfigurator() { return FilterConfigurator; }
+	TSharedPtr<UE::Insights::FFilterConfigurator> GetFilterConfigurator() { return FilterConfigurator; }
 
 	TMap<uint64, TSharedPtr<FBaseTimingTrack>>& GetAllTracks() { return AllTracks; }
 
@@ -436,8 +457,8 @@ protected:
 
 	void CreateCpuThreadTrackColoringModeMenu(FMenuBuilder& MenuBuilder);
 	void ChooseNextCpuThreadTrackColoringMode();
-	void SetCpuThreadTrackColoringMode(Insights::ETimingEventsColoringMode Mode);
-	bool CheckCpuThreadTrackColoringMode(Insights::ETimingEventsColoringMode Mode);
+	void SetCpuThreadTrackColoringMode(UE::Insights::ETimingEventsColoringMode Mode);
+	bool CheckCpuThreadTrackColoringMode(UE::Insights::ETimingEventsColoringMode Mode);
 
 	void ShowContextMenu(const FPointerEvent& MouseEvent);
 	void CreateTrackLocationMenu(FMenuBuilder& MenuBuilder, TSharedRef<FBaseTimingTrack> Track);
@@ -485,10 +506,8 @@ protected:
 	void RaiseSelectionChanging();
 	void RaiseSelectionChanged();
 
-	void RaiseTimeMarkerChanging(TSharedRef<Insights::FTimeMarker> InTimeMarker);
-	void RaiseTimeMarkerChanged(TSharedRef<Insights::FTimeMarker> InTimeMarker);
-
-	void UpdateAggregatedStats();
+	void RaiseTimeMarkerChanging(TSharedRef<FTimeMarker> InTimeMarker);
+	void RaiseTimeMarkerChanged(TSharedRef<FTimeMarker> InTimeMarker);
 
 	void UpdateHoveredTimingEvent(float InMousePosX, float InMousePosY);
 
@@ -505,7 +524,12 @@ protected:
 	void FrameSelection();
 
 	// Get all the plugin extenders we care about
-	TArray<Insights::ITimingViewExtender*> GetExtenders() const;
+	TArray<UE::Insights::Timing::ITimingViewExtender*> GetExtenders() const;
+#if UE_INSIGHTS_BACKWARD_COMPATIBILITY_UE54
+PRAGMA_DISABLE_DEPRECATION_WARNINGS
+	TArray<::Insights::ITimingViewExtender*> GetOldExtenders() const;
+PRAGMA_ENABLE_DEPRECATION_WARNINGS
+#endif // UE_INSIGHTS_BACKWARD_COMPATIBILITY_UE54
 
 	FReply AllowTracksToProcessOnMouseButtonDown(const FGeometry& MyGeometry, const FPointerEvent& MouseEvent);
 	FReply AllowTracksToProcessOnMouseButtonUp(const FGeometry& MyGeometry, const FPointerEvent& MouseEvent);
@@ -525,7 +549,7 @@ protected:
 	void PopulateTimerNameSuggestionList(const FString& Text, TArray<FString>& OutSuggestions);
 
 	typedef TFunctionRef<void(TSharedPtr<const FBaseTimingTrack> Track)> EnumerateFilteredTracksCallback;
-	void EnumerateFilteredTracks(TSharedPtr<Insights::FFilterConfigurator> FilterConfigurator, TSharedPtr<const FBaseTimingTrack> PriorityTrack, EnumerateFilteredTracksCallback Callback);
+	void EnumerateFilteredTracks(TSharedPtr<UE::Insights::FFilterConfigurator> FilterConfigurator, TSharedPtr<const FBaseTimingTrack> PriorityTrack, EnumerateFilteredTracksCallback Callback);
 
 	ETraceFrameType GetFrameTypeToSnapTo();
 	
@@ -563,7 +587,7 @@ protected:
 	// Shared state for Frame Thread tracks
 	TSharedPtr<FFrameSharedState> FrameSharedState;
 
-	// Shared state for Cpu/Gpu Thread tracks
+	// Shared state for CPU/GPU Thread tracks
 	TSharedPtr<FThreadTimingSharedState> ThreadTimingSharedState;
 
 	// Shared state for Asset Loading tracks
@@ -573,7 +597,7 @@ protected:
 	TSharedPtr<FFileActivitySharedState> FileActivitySharedState;
 
 	// Shared state for Regions tracks
-	TSharedPtr<Insights::FTimingRegionsSharedState> TimingRegionsSharedState;
+	TSharedPtr<::Insights::FTimingRegionsSharedState> TimingRegionsSharedState;
 	
 	////////////////////////////////////////////////////////////
 
@@ -581,12 +605,12 @@ protected:
 	TSharedRef<FTimeRulerTrack> TimeRulerTrack;
 
 	/** The default time marker (for backward compatibility). */
-	TSharedRef<Insights::FTimeMarker> DefaultTimeMarker;
+	TSharedRef<FTimeMarker> DefaultTimeMarker;
 
-	/** The time markers track. It displayes fixed time markers based on bookmarks and log messages. */
+	/** The time markers track. It displays fixed time markers based on bookmarks and log messages. */
 	TSharedRef<FMarkersTimingTrack> MarkersTrack;
 
-	/** A graph track for frame times and cpu/gpu timing graphs. */
+	/** A graph track for frame times and CPU/GPU timing graphs. */
 	TSharedPtr<FTimingGraphTrack> GraphTrack;
 
 	////////////////////////////////////////////////////////////
@@ -734,28 +758,30 @@ protected:
 	////////////////////////////////////////////////////////////
 	// Delegates
 
-	Insights::FSelectionChangedDelegate OnSelectionChangedDelegate;
-	Insights::FTimeMarkerChangedDelegate OnTimeMarkerChangedDelegate;
-	Insights::FCustomTimeMarkerChangedDelegate OnCustomTimeMarkerChangedDelegate;
-	Insights::FHoveredTrackChangedDelegate OnHoveredTrackChangedDelegate;
-	Insights::FHoveredEventChangedDelegate OnHoveredEventChangedDelegate;
-	Insights::FSelectedTrackChangedDelegate OnSelectedTrackChangedDelegate;
-	Insights::FSelectedEventChangedDelegate OnSelectedEventChangedDelegate;
-	Insights::FTrackVisibilityChangedDelegate OnTrackVisibilityChangedDelegate;
-	Insights::FTrackAddedDelegate OnTrackAddedDelegate;
-	Insights::FTrackRemovedDelegate OnTrackRemovedDelegate;
+	Timing::FSelectionChangedDelegate OnSelectionChangedDelegate;
+	Timing::FTimeMarkerChangedDelegate OnTimeMarkerChangedDelegate;
+	Timing::FCustomTimeMarkerChangedDelegate OnCustomTimeMarkerChangedDelegate;
+	Timing::FHoveredTrackChangedDelegate OnHoveredTrackChangedDelegate;
+	Timing::FHoveredEventChangedDelegate OnHoveredEventChangedDelegate;
+	Timing::FSelectedTrackChangedDelegate OnSelectedTrackChangedDelegate;
+	Timing::FSelectedEventChangedDelegate OnSelectedEventChangedDelegate;
+	Timing::FTrackVisibilityChangedDelegate OnTrackVisibilityChangedDelegate;
+	Timing::FTrackAddedDelegate OnTrackAddedDelegate;
+	Timing::FTrackRemovedDelegate OnTrackRemovedDelegate;
 
 	TSharedPtr<FUICommandList> CommandList;
 
 	TArray<TUniquePtr<ITimingEventRelation>> CurrentRelations;
 
-	TSharedPtr<Insights::FQuickFind> QuickFindVm;
-	TSharedPtr<Insights::FFilterConfigurator> FilterConfigurator;
+	TSharedPtr<::Insights::FQuickFind> QuickFindVm;
+	TSharedPtr<FFilterConfigurator> FilterConfigurator;
 	static uint32 TimingViewId;
 	const FName QuickFindTabId;
 	bool bUpdateFilters = true;
 	
 	// Used only between the creation of the widget and the spawning of the owning tab. When the tab is spawned, we relinquish ownership.
-	TSharedPtr<Insights::SQuickFind> QuickFindWidgetSharedPtr;
-	TWeakPtr<Insights::SQuickFind> QuickFindWidgetWeakPtr;
-};
\ No newline at end of file
+	TSharedPtr<::Insights::SQuickFind> QuickFindWidgetSharedPtr;
+	TWeakPtr<::Insights::SQuickFind> QuickFindWidgetWeakPtr;
+};
+
+} // namespace UE::Insights::TimingProfiler
diff --git a/Engine/Source/Developer/TraceInsights/Private/Insights/Widgets/STimingViewTrackList.cpp b/Engine/Source/Developer/TraceInsights/Private/Insights/Widgets/STimingViewTrackList.cpp
--- a/Engine/Source/Developer/TraceInsights/Private/Insights/Widgets/STimingViewTrackList.cpp
+++ b/Engine/Source/Developer/TraceInsights/Private/Insights/Widgets/STimingViewTrackList.cpp
@@ -11,7 +11,7 @@
 #include "Widgets/SBoxPanel.h"
 #include "Widgets/Views/SHeaderRow.h"
 
-// Insights
+// TraceInsights
 #include "Insights/ViewModels/BaseTimingTrack.h"
 #include "Insights/Widgets/STimingView.h"
 
@@ -27,7 +27,7 @@ class STimingViewTrackListEntry : public STableRow<TSharedPtr<FBaseTimingTrack>>
 public:
 	SLATE_BEGIN_ARGS(STimingViewTrackListEntry) {}
 
-	SLATE_ARGUMENT(TSharedPtr<STimingView>, TimingView)
+	SLATE_ARGUMENT(TSharedPtr<UE::Insights::TimingProfiler::STimingView>, TimingView)
 	SLATE_ARGUMENT(ETimingTrackLocation, TrackLocation)
 	SLATE_ARGUMENT(TSharedPtr<FBaseTimingTrack>, Track)
 
@@ -92,7 +92,7 @@ public:
 	}
 
 	// The widget containing the track we represent
-	TWeakPtr<STimingView> TimingView;
+	TWeakPtr<UE::Insights::TimingProfiler::STimingView> TimingView;
 
 	// The location of the track we represent
 	ETimingTrackLocation TrackLocation;
@@ -108,7 +108,7 @@ public:
 // STimingViewTrackList
 ////////////////////////////////////////////////////////////////////////////////////////////////////
 
-void STimingViewTrackList::Construct(const FArguments& InArgs, const TSharedRef<STimingView>& InTimingView, ETimingTrackLocation InTrackLocation)
+void STimingViewTrackList::Construct(const FArguments& InArgs, const TSharedRef<UE::Insights::TimingProfiler::STimingView>& InTimingView, ETimingTrackLocation InTrackLocation)
 {
 	TimingView = InTimingView;
 	TrackLocation = InTrackLocation;
diff --git a/Engine/Source/Developer/TraceInsights/Private/Insights/Widgets/STimingViewTrackList.h b/Engine/Source/Developer/TraceInsights/Private/Insights/Widgets/STimingViewTrackList.h
--- a/Engine/Source/Developer/TraceInsights/Private/Insights/Widgets/STimingViewTrackList.h
+++ b/Engine/Source/Developer/TraceInsights/Private/Insights/Widgets/STimingViewTrackList.h
@@ -3,13 +3,14 @@
 #pragma once
 
 #include "CoreMinimal.h"
+
 #include "Widgets/SCompoundWidget.h"
 #include "Widgets/Views/SListView.h"
 
 #include "Insights/ViewModels/BaseTimingTrack.h"
 
 class SSearchBox;
-class STimingView;
+namespace UE::Insights::TimingProfiler { class STimingView; }
 
 // A widget representing all the series in a graph track, allowing management of their visibility
 class STimingViewTrackList : public SCompoundWidget
@@ -19,7 +20,7 @@ public:
 
 	SLATE_END_ARGS()
 
-	void Construct(const FArguments& InArgs, const TSharedRef<STimingView>& InTimingView, ETimingTrackLocation InTrackLocation);
+	void Construct(const FArguments& InArgs, const TSharedRef<UE::Insights::TimingProfiler::STimingView>& InTimingView, ETimingTrackLocation InTrackLocation);
 
 private:
 	// Generate a row for the list
@@ -30,7 +31,7 @@ private:
 
 private:
 	// The timing view widget we are operating on
-	TWeakPtr<STimingView> TimingView;
+	TWeakPtr<UE::Insights::TimingProfiler::STimingView> TimingView;
 
 	// The location of tracks we are operating on (scrollable tracks, top docked tracks, etc.)
 	ETimingTrackLocation TrackLocation;
diff --git a/Engine/Source/Developer/TraceInsights/Private/Insights/Widgets/STraceServerControl.cpp b/Engine/Source/Developer/TraceInsights/Private/Insights/Widgets/STraceServerControl.cpp
--- a/Engine/Source/Developer/TraceInsights/Private/Insights/Widgets/STraceServerControl.cpp
+++ b/Engine/Source/Developer/TraceInsights/Private/Insights/Widgets/STraceServerControl.cpp
@@ -2,10 +2,12 @@
 
 #include "Insights/Widgets/STraceServerControl.h"
 
-#include "SlateOptMacros.h"
 #include "Framework/MultiBox/MultiBoxBuilder.h"
 #include "Insights/Log.h"
 #include "ProfilingDebugging/TraceAuxiliary.h"
+#include "SlateOptMacros.h"
+
+// TraceAnalysis
 #include "Trace/StoreClient.h"
 
 /////////////////////////////////////////////////////////////////////////////////////////
diff --git a/Engine/Source/Developer/TraceInsights/Public/Insights/Common/AvailabilityCheck.h b/Engine/Source/Developer/TraceInsights/Public/Insights/Common/AvailabilityCheck.h
--- a/Engine/Source/Developer/TraceInsights/Public/Insights/Common/AvailabilityCheck.h
+++ b/Engine/Source/Developer/TraceInsights/Public/Insights/Common/AvailabilityCheck.h
@@ -2,33 +2,12 @@
 
 #pragma once
 
-#include "HAL/Platform.h"
+// TraceInsightsCore
+#include "InsightsCore/Common/AvailabilityCheck.h"
 
-namespace Insights
-{
+// TraceInsights
+#include "Insights/Config.h"
 
-////////////////////////////////////////////////////////////////////////////////////////////////////
-
-/**
- * Utility class used by profiler managers to limit how often they check for availability conditions.
- */
-class TRACEINSIGHTS_API FAvailabilityCheck
-{
-public:
-	/** Returns true if managers are allowed to do (slow) availability check during this tick. */
-	bool Tick();
-
-	/** Disables the "availability check" (i.e. Tick() calls will return false when disabled). */
-	void Disable();
-
-	/** Enables the "availability check" with a specified initial delay. */
-	void Enable(double InWaitTime);
-
-private:
-	double WaitTime = 0.0;
-	uint64 NextTimestamp = (uint64)-1;
-};
-
-////////////////////////////////////////////////////////////////////////////////////////////////////
-
-} // namespace Insights
+#if UE_INSIGHTS_BACKWARD_COMPATIBILITY_UE54
+namespace Insights { using FAvailabilityCheck = UE::Insights::FAvailabilityCheck; }
+#endif
diff --git a/Engine/Source/Developer/TraceInsights/Public/Insights/Common/PaintUtils.h b/Engine/Source/Developer/TraceInsights/Public/Insights/Common/PaintUtils.h
--- a/Engine/Source/Developer/TraceInsights/Public/Insights/Common/PaintUtils.h
+++ b/Engine/Source/Developer/TraceInsights/Public/Insights/Common/PaintUtils.h
@@ -2,118 +2,5 @@
 
 #pragma once
 
-#include "CoreMinimal.h"
-#include "Fonts/FontMeasure.h"
-#include "Framework/Application/SlateApplication.h"
-#include "Rendering/DrawElements.h"
-
-#define MAKE_PAINT_GEOMETRY_PT(Geometry, X, Y)       Geometry.ToPaintGeometry(FSlateLayoutTransform(1.0f, FVector2D(X, Y)))
-#define MAKE_PAINT_GEOMETRY_RC(Geometry, X, Y, W, H) Geometry.ToPaintGeometry(FVector2D(W, H), FSlateLayoutTransform(1.0f, FVector2D(X, Y)))
-
-struct FGeometry;
-class FSlateRect;
-class WidgetStyle;
-enum class ESlateDrawEffect : uint8;
-class FSlateWindowElementList;
-
-/**
- * Holds current state provided by OnPaint function, used to simplify drawing.
- */
-struct TRACEINSIGHTS_API FDrawContext
-{
-	FDrawContext(const FGeometry& InGeometry,
-				 const FSlateRect& InCullingRect,
-				 const FWidgetStyle& InWidgetStyle,
-				 const ESlateDrawEffect InDrawEffects,
-				 FSlateWindowElementList& InOutElementList,
-				 int32& InOutLayerId)
-		: Geometry(InGeometry)
-		, CullingRect(InCullingRect)
-		, WidgetStyle(InWidgetStyle)
-		, DrawEffects(InDrawEffects)
-		, ElementList(InOutElementList)
-		, LayerId(InOutLayerId)
-	{}
-
-	/**
-	 * Non-copyable
-	 */
-	FDrawContext(const FDrawContext&) = delete;
-	FDrawContext& operator=(const FDrawContext&) = delete;
-
-	inline void DrawBox(const float X, const float Y, const float W, const float H, const FSlateBrush* Brush, const FLinearColor& Color) const
-	{
-		FSlateDrawElement::MakeBox(ElementList, LayerId, MAKE_PAINT_GEOMETRY_RC(Geometry, X, Y, W, H), Brush, DrawEffects, Color);
-	}
-
-	inline void DrawBox(const int32 InLayer, const float X, const float Y, const float W, const float H, const FSlateBrush* Brush, const FLinearColor& Color) const
-	{
-		FSlateDrawElement::MakeBox(ElementList, InLayer, MAKE_PAINT_GEOMETRY_RC(Geometry, X, Y, W, H), Brush, DrawEffects, Color);
-	}
-
-	inline void DrawRotatedBox(const float X, const float Y, const float W, const float H, const FSlateBrush* Brush, const FLinearColor& Color, float Angle, TOptional<FVector2D> RotationPoint) const
-	{
-		FSlateDrawElement::MakeRotatedBox(ElementList, LayerId, MAKE_PAINT_GEOMETRY_RC(Geometry, X, Y, W, H), Brush, DrawEffects, Angle, RotationPoint, FSlateDrawElement::RelativeToElement, Color);
-	}
-
-	inline void DrawText(const float X, const float Y, const FString& Text, const FSlateFontInfo& Font, const FLinearColor& Color) const
-	{
-		FSlateDrawElement::MakeText(ElementList, LayerId, MAKE_PAINT_GEOMETRY_PT(Geometry, X, Y), Text, Font, DrawEffects, Color);
-	}
-
-	inline void DrawText(const int32 InLayer, const float X, const float Y, const FString& Text, const FSlateFontInfo& Font, const FLinearColor& Color) const
-	{
-		FSlateDrawElement::MakeText(ElementList, InLayer, MAKE_PAINT_GEOMETRY_PT(Geometry, X, Y), Text, Font, DrawEffects, Color);
-	}
-
-	inline void DrawText(const float X, const float Y, const FString& Text, const int32 StartIndex, const int32 EndIndex, const FSlateFontInfo& Font, const FLinearColor& Color) const
-	{
-		FSlateDrawElement::MakeText(ElementList, LayerId, MAKE_PAINT_GEOMETRY_PT(Geometry, X, Y), Text, StartIndex, EndIndex, Font, DrawEffects, Color);
-	}
-
-	inline void DrawText(const int32 InLayer, const float X, const float Y, const FString& Text, const int32 StartIndex, const int32 EndIndex, const FSlateFontInfo& Font, const FLinearColor& Color) const
-	{
-		FSlateDrawElement::MakeText(ElementList, InLayer, MAKE_PAINT_GEOMETRY_PT(Geometry, X, Y), Text, StartIndex, EndIndex, Font, DrawEffects, Color);
-	}
-
-	inline void DrawTextAligned(EHorizontalAlignment HAlign, const float X, const float Y, const FString& Text, const FSlateFontInfo& Font, const FLinearColor& Color) const
-	{
-		float TextX = X;
-		if (HAlign == HAlign_Right || HAlign == HAlign_Center)
-		{
-			const TSharedRef<FSlateFontMeasure> FontMeasureService = FSlateApplication::Get().GetRenderer()->GetFontMeasureService();
-			const float FontScale = Geometry.Scale;
-			const float TextWidth = static_cast<float>(FontMeasureService->Measure(Text, Font, FontScale).X / FontScale);
-			if (HAlign == HAlign_Right)
-			{
-				TextX -= TextWidth;
-			}
-			else
-			{
-				TextX -= TextWidth / 2;
-			}
-		}
-		FSlateDrawElement::MakeText(ElementList, LayerId, MAKE_PAINT_GEOMETRY_PT(Geometry, TextX, Y), Text, Font, DrawEffects, Color);
-	}
-
-	inline void DrawSpline(uint32 InLayer, const float X, const float Y, const FVector2D& InStart, const FVector2D& InStartDir, const FVector2D& InEnd, const FVector2D& InEndDir, float InThickness = 0.0f, const FLinearColor& InTint=FLinearColor::White) const
-	{
-		FSlateDrawElement::MakeSpline(ElementList, InLayer, MAKE_PAINT_GEOMETRY_PT(Geometry, X, Y), InStart, InStartDir, InEnd, InEndDir, InThickness, DrawEffects, InTint);
-	}
-
-	inline void DrawLines(uint32 InLayer, const float X, const float Y, const TArray<FVector2D>& Points, ESlateDrawEffect InDrawEffects = ESlateDrawEffect::None, const FLinearColor& InTint = FLinearColor::White, bool bAntialias = true, float Thickness = 1.0f) const
-	{
-		FSlateDrawElement::MakeLines(ElementList, InLayer, MAKE_PAINT_GEOMETRY_PT(Geometry, X, Y), Points, InDrawEffects, InTint, bAntialias, Thickness);
-	}
-
-	//inline void IncrementLayer() const { ++LayerId; }
-
-	// Accessors
-	const FGeometry& Geometry;
-	const FSlateRect& CullingRect;
-	const FWidgetStyle& WidgetStyle;
-	const ESlateDrawEffect DrawEffects;
-
-	FSlateWindowElementList& ElementList;
-	int32& LayerId;
-};
+// TraceInsightsCore
+#include "InsightsCore/Common/PaintUtils.h"
diff --git a/Engine/Source/Developer/TraceInsights/Public/Insights/Common/SimpleRtti.h b/Engine/Source/Developer/TraceInsights/Public/Insights/Common/SimpleRtti.h
--- a/Engine/Source/Developer/TraceInsights/Public/Insights/Common/SimpleRtti.h
+++ b/Engine/Source/Developer/TraceInsights/Public/Insights/Common/SimpleRtti.h
@@ -2,32 +2,5 @@
 
 #pragma once
 
-// Simple RTTI implementation.
-
-// Note: The SimpleRTTI class hierarchy may be different than the actual class hierarchy.
-//       Ex.: Some classes in the actual hierarchy can be hidden (do not use INSIGHTS_DECLARE_RTTI, so just inherits type info from parent class).
-
-// Macro to be used in declaration of the base class of the SimpleRTTI class hierarchy.
-#define INSIGHTS_DECLARE_RTTI_BASE(BaseClassName) \
-	public: \
-		static const FName& GetStaticTypeName() { return BaseClassName::TypeName; } \
-		virtual const FName& GetTypeName() const { return BaseClassName::GetStaticTypeName(); } \
-		virtual bool IsKindOf(const FName& InTypeName) const { return InTypeName == BaseClassName::GetStaticTypeName(); } \
-		template<typename Type> bool Is() const { return IsKindOf(Type::GetStaticTypeName()); } \
-		template<typename Type> const Type& As() const { return *static_cast<const Type*>(this); } \
-		template<typename Type> Type& As() { return *static_cast<Type*>(this); } \
-	private: \
-		static const FName TypeName;
-
-// Macro to be used in declaration of a derived class in a SimpleRTTI hierarchy.
-#define INSIGHTS_DECLARE_RTTI(ClassName, BaseClassName) \
-	public: \
-		static const FName& GetStaticTypeName() { return ClassName::TypeName; } \
-		virtual const FName& GetTypeName() const override { return ClassName::GetStaticTypeName(); } \
-		virtual bool IsKindOf(const FName& InTypeName) const override { return InTypeName == ClassName::GetStaticTypeName() || BaseClassName::IsKindOf(InTypeName); } \
-	private: \
-		static const FName TypeName;
-
-// Macro to be used in implementation of a class in a SimpleRTTI hierarchy.
-#define INSIGHTS_IMPLEMENT_RTTI(ClassName) \
-	const FName ClassName::TypeName = TEXT(#ClassName);
+// TraceInsightsCore
+#include "InsightsCore/Common/SimpleRtti.h"
diff --git a/Engine/Source/Developer/TraceInsights/Public/Insights/Common/TimeUtils.h b/Engine/Source/Developer/TraceInsights/Public/Insights/Common/TimeUtils.h
--- a/Engine/Source/Developer/TraceInsights/Public/Insights/Common/TimeUtils.h
+++ b/Engine/Source/Developer/TraceInsights/Public/Insights/Common/TimeUtils.h
@@ -2,70 +2,54 @@
 
 #pragma once
 
-#include "Containers/UnrealString.h"
-#include "CoreMinimal.h"
-#include "HAL/Platform.h"
-
-/*
-double is a 64 bit IEEE 754 double precision Floating Point Number
-1 bit for the sign, 11 bits for the exponent, and 52* bits for the value
-15.5 decimal digits of precision (max value*: 4`503`599`627`370`495)
-|------------------------------------------------------------------------------
-| 900`000.000`000`001     --> up to 10 days   with    1 ns precision
-|  90`000.000`000`000`1   --> up to 25 hours  with  0.1 ns precision (100 ps)
-|   9`000.000`000`000`01  --> up to 2.5 hours with 0.01 ns precision (10 ps)
-|     900.000`000`000`001 --> up to 15 min    with    1 ps precision
-|------------------------------------------------------------------------------
-|   3`600.000`000`000`001 --> 1 ps precision at 1 hour : NOT OK (16 digits)!!!
-|  86`400.000`000`000`001 --> 1 ps precision at 1 day  : NOT OK (17 digits)!!!
-|------------------------------------------------------------------------------
-*/
+// TraceInsightsCore
+#include "InsightsCore/Common/TimeUtils.h"
+
+// TraceInsights
+#include "Insights/Config.h"
+
+#if UE_INSIGHTS_BACKWARD_COMPATIBILITY_UE54
 
 namespace TimeUtils
 {
-	static constexpr double Picosecond = 0.000000000001;
-	static constexpr double Nanosecond = 0.000000001;
-	static constexpr double Microsecond = 0.000001;
-	static constexpr double Milisecond = 0.001;
-	static constexpr double Second = 1.0;
-	static constexpr double Minute = 60.0;
-	static constexpr double Hour = 3600.0;
-	static constexpr double Day = 86400.0;
-	static constexpr double Week = 604800.0;
-
-	struct TRACEINSIGHTS_API FTimeSplit
-	{
-		union
-		{
-			int32 Units[8]; // Units[0] = Days, Units[1] = Hours, ..., Units[7] = Picoseconds
-			struct
-			{
-				int32 Days;
-				int32 Hours;
-				int32 Minutes;
-				int32 Seconds;
-				int32 Miliseconds;
-				int32 Microseconds;
-				int32 Nanoseconds;
-				int32 Picoseconds;
-			};
-		};
-		bool bIsZero;
-		bool bIsNegative;
-		bool bIsInfinite;
-		bool bIsNaN;
-	};
-
-	TRACEINSIGHTS_API FString FormatTimeValue(const double Duration, const int32 NumDigits = 1);
-	TRACEINSIGHTS_API FString FormatTimeAuto(const double Duration, const int32 NumDigits = 1);
-	TRACEINSIGHTS_API FString FormatTimeMs(const double Duration, const int32 NumDigits = 2, bool bAddTimeUnit = false);
-	TRACEINSIGHTS_API FString FormatTime(const double Time, const double Precision = 0.0);
-	TRACEINSIGHTS_API FString FormatTimeHMS(const double Time, const double Precision = 0.0);
-	TRACEINSIGHTS_API void SplitTime(const double Time, FTimeSplit& OutTimeSplit);
-	TRACEINSIGHTS_API FString FormatTimeSplit(const FTimeSplit& TimeSplit, const double Precision = 0.0);
-	TRACEINSIGHTS_API FString FormatTimeSplit(const double Time, const double Precision = 0.0);
-
-	void TestTimeFormatting();
-	void TestTimeAutoFormatting();
-	void TestOptimizationIssue();
-}
+
+static constexpr double /*UE_DEPRECATED(5.5, "TimeUtils namespace was merged into UE::Insights namespace")*/ Picosecond = 0.000000000001;
+static constexpr double /*UE_DEPRECATED(5.5, "TimeUtils namespace was merged into UE::Insights namespace")*/ Nanosecond = 0.000000001;
+static constexpr double /*UE_DEPRECATED(5.5, "TimeUtils namespace was merged into UE::Insights namespace")*/ Microsecond = 0.000001;
+static constexpr double /*UE_DEPRECATED(5.5, "TimeUtils namespace was merged into UE::Insights namespace")*/ Milisecond = 0.001; // !!!
+static constexpr double /*UE_DEPRECATED(5.5, "TimeUtils namespace was merged into UE::Insights namespace")*/ Second = 1.0;
+static constexpr double /*UE_DEPRECATED(5.5, "TimeUtils namespace was merged into UE::Insights namespace")*/ Minute = 60.0;
+static constexpr double /*UE_DEPRECATED(5.5, "TimeUtils namespace was merged into UE::Insights namespace")*/ Hour = 3600.0;
+static constexpr double /*UE_DEPRECATED(5.5, "TimeUtils namespace was merged into UE::Insights namespace")*/ Day = 86400.0;
+static constexpr double /*UE_DEPRECATED(5.5, "TimeUtils namespace was merged into UE::Insights namespace")*/ Week = 604800.0;
+
+//struct UE_DEPRECATED(5.5, "FTimeSplit was moved into UE::Insights namespace") FTimeSplit;
+struct FTimeSplit : public UE::Insights::FTimeSplit {};
+
+//UE_DEPRECATED(5.5, "TimeUtils namespace was merged into UE::Insights namespace")
+inline FString FormatTimeValue(const double Duration, const int32 NumDigits = 1) { return UE::Insights::FormatTimeValue(Duration, NumDigits); }
+
+//UE_DEPRECATED(5.5, "TimeUtils namespace was merged into UE::Insights namespace")
+inline FString FormatTimeAuto(const double Duration, const int32 NumDigits = 1) { return UE::Insights::FormatTimeAuto(Duration, NumDigits); }
+
+//UE_DEPRECATED(5.5, "TimeUtils namespace was merged into UE::Insights namespace")
+inline FString FormatTimeMs(const double Duration, const int32 NumDigits = 2, bool bAddTimeUnit = false) { return UE::Insights::FormatTimeMs(Duration, NumDigits, bAddTimeUnit); }
+
+//UE_DEPRECATED(5.5, "TimeUtils namespace was merged into UE::Insights namespace")
+inline FString FormatTime(const double Time, const double Precision = 0.0) { return UE::Insights::FormatTime(Time, Precision); }
+
+//UE_DEPRECATED(5.5, "TimeUtils namespace was merged into UE::Insights namespace")
+inline FString FormatTimeHMS(const double Time, const double Precision = 0.0) { return UE::Insights::FormatTimeHMS(Time, Precision); }
+
+//UE_DEPRECATED(5.5, "TimeUtils namespace was merged into UE::Insights namespace")
+inline void SplitTime(const double Time, UE::Insights::FTimeSplit& OutTimeSplit) { UE::Insights::SplitTime(Time, OutTimeSplit); }
+
+//UE_DEPRECATED(5.5, "TimeUtils namespace was merged into UE::Insights namespace")
+inline FString FormatTimeSplit(const UE::Insights::FTimeSplit& TimeSplit, const double Precision = 0.0) { return UE::Insights::FormatTimeSplit(TimeSplit, Precision); }
+
+//UE_DEPRECATED(5.5, "TimeUtils namespace was merged into UE::Insights namespace")
+inline FString FormatTimeSplit(const double Time, const double Precision = 0.0) { return UE::Insights::FormatTimeSplit(Time, Precision); }
+
+} // namespace TimeUtils
+
+#endif // UE_INSIGHTS_BACKWARD_COMPATIBILITY_UE54
diff --git a/Engine/Source/Developer/TraceInsights/Public/Insights/ITimingViewExtender.h b/Engine/Source/Developer/TraceInsights/Public/Insights/ITimingViewExtender.h
--- a/Engine/Source/Developer/TraceInsights/Public/Insights/ITimingViewExtender.h
+++ b/Engine/Source/Developer/TraceInsights/Public/Insights/ITimingViewExtender.h
@@ -2,19 +2,29 @@
 
 #pragma once
 
-#include "CoreMinimal.h"
+#include "CoreTypes.h"
 #include "Features/IModularFeature.h"
+#include "Templates/SharedPointer.h"
+
+#include "Insights/Config.h"
+
+#if UE_INSIGHTS_BACKWARD_COMPATIBILITY_UE54
+#include "Insights/ITimingViewSession.h"
+#endif
 
-namespace TraceServices { class IAnalysisSession; }
 class FMenuBuilder;
 
-namespace Insights
+namespace TraceServices { class IAnalysisSession; }
+namespace UE::Insights { class FFilterConfigurator; }
+
+namespace UE::Insights::Timing
 {
 
 class ITimingViewSession;
+
 extern TRACEINSIGHTS_API const FName TimingViewExtenderFeatureName;
 
-class TRACEINSIGHTS_API ITimingViewExtender : public IModularFeature
+class ITimingViewExtender : public IModularFeature
 {
 public:
 	virtual ~ITimingViewExtender() = default;
@@ -38,11 +48,65 @@ public:
 	virtual void ExtendFilterMenu(ITimingViewSession& InSession, FMenuBuilder& InMenuBuilder) {}
 
 	/** Extension hook for the context menu for all tracks
-	@return True if any menu option was added and False if no option was added */ 
+	@return True if any menu option was added and False if no option was added */
 	virtual bool ExtendGlobalContextMenu(ITimingViewSession& InSession, FMenuBuilder& InMenuBuilder) { return false; }
 
 	/** Allows extender to add filters to the Quick Find widget. */
-	virtual void AddQuickFindFilters(TSharedPtr<class FFilterConfigurator> FilterConfigurator) {}
+	virtual void AddQuickFindFilters(TSharedPtr<FFilterConfigurator> FilterConfigurator) {}
+};
+
+} // namespace UE::Insights::Timing
+
+#if UE_INSIGHTS_BACKWARD_COMPATIBILITY_UE54
+
+namespace Insights
+{
+
+extern TRACEINSIGHTS_API const FName TimingViewExtenderFeatureName;
+
+//class UE_DEPRECATED(5.5, "ITimingViewExtender class was moved inside UE::Insights::Timing namespace") ITimingViewExtender;
+class ITimingViewExtender : public IModularFeature
+{
+public:
+	virtual ~ITimingViewExtender() = default;
+	PRAGMA_DISABLE_DEPRECATION_WARNINGS
+	virtual void OnBeginSession(ITimingViewSession& InSession) = 0;
+	virtual void OnEndSession(ITimingViewSession& InSession) = 0;
+	virtual void Tick(ITimingViewSession& InSession, const TraceServices::IAnalysisSession& InAnalysisSession) = 0;
+	virtual void ExtendCpuTracksFilterMenu(ITimingViewSession& InSession, FMenuBuilder& InMenuBuilder) {}
+	virtual void ExtendGpuTracksFilterMenu(ITimingViewSession& InSession, FMenuBuilder& InMenuBuilder) {}
+	virtual void ExtendOtherTracksFilterMenu(ITimingViewSession& InSession, FMenuBuilder& InMenuBuilder) {}
+	virtual void ExtendFilterMenu(ITimingViewSession& InSession, FMenuBuilder& InMenuBuilder) {}
+	virtual bool ExtendGlobalContextMenu(ITimingViewSession& InSession, FMenuBuilder& InMenuBuilder) { return false; }
+	virtual void AddQuickFindFilters(TSharedPtr<UE::Insights::FFilterConfigurator> FilterConfigurator) {}
+	PRAGMA_ENABLE_DEPRECATION_WARNINGS
 };
 
 } // namespace Insights
+
+namespace UE::Insights
+{
+
+extern TRACEINSIGHTS_API const FName TimingViewExtenderFeatureName;
+
+//class UE_DEPRECATED(5.5, "ITimingViewExtender class was moved inside UE::Insights::Timing namespace") ITimingViewExtender;
+class ITimingViewExtender : public IModularFeature
+{
+public:
+	virtual ~ITimingViewExtender() = default;
+	PRAGMA_DISABLE_DEPRECATION_WARNINGS
+	virtual void OnBeginSession(ITimingViewSession& InSession) = 0;
+	virtual void OnEndSession(ITimingViewSession& InSession) = 0;
+	virtual void Tick(ITimingViewSession& InSession, const TraceServices::IAnalysisSession& InAnalysisSession) = 0;
+	virtual void ExtendCpuTracksFilterMenu(ITimingViewSession& InSession, FMenuBuilder& InMenuBuilder) {}
+	virtual void ExtendGpuTracksFilterMenu(ITimingViewSession& InSession, FMenuBuilder& InMenuBuilder) {}
+	virtual void ExtendOtherTracksFilterMenu(ITimingViewSession& InSession, FMenuBuilder& InMenuBuilder) {}
+	virtual void ExtendFilterMenu(ITimingViewSession& InSession, FMenuBuilder& InMenuBuilder) {}
+	virtual bool ExtendGlobalContextMenu(ITimingViewSession& InSession, FMenuBuilder& InMenuBuilder) { return false; }
+	virtual void AddQuickFindFilters(TSharedPtr<UE::Insights::FFilterConfigurator> FilterConfigurator) {}
+	PRAGMA_ENABLE_DEPRECATION_WARNINGS
+};
+
+} // namespace UE::Insights
+
+#endif // UE_INSIGHTS_BACKWARD_COMPATIBILITY_UE54
diff --git a/Engine/Source/Developer/TraceInsights/Public/Insights/ITimingViewSession.h b/Engine/Source/Developer/TraceInsights/Public/Insights/ITimingViewSession.h
--- a/Engine/Source/Developer/TraceInsights/Public/Insights/ITimingViewSession.h
+++ b/Engine/Source/Developer/TraceInsights/Public/Insights/ITimingViewSession.h
@@ -2,15 +2,33 @@
 
 #pragma once
 
-#include "CoreMinimal.h"
+#include "CoreTypes.h"
+
+#include "Delegates/Delegate.h"
+#include "Delegates/DelegateCombinations.h"
+#include "Misc/EnumClassFlags.h"
+#include "Templates/SharedPointer.h"
+
+// TraceInsights
+#include "Insights/Config.h"
+
+class SWidget;
 
 class FBaseTimingTrack;
 class FTimingEventsTrack;
 class ITimingEvent;
-class SWidget;
 enum class ETimingTrackLocation : uint32; // Insights/ViewModels/BaseTimingTrack.h
 
-namespace Insights
+#if UE_INSIGHTS_BACKWARD_COMPATIBILITY_UE54
+namespace Insights { enum class ETimeChangedFlags : int32; }
+namespace Insights { class ITimeMarker; }
+namespace Insights { class ITimingViewSession; }
+namespace UE::Insights { using ETimeChangedFlags = ::Insights::ETimeChangedFlags; }
+namespace UE::Insights { using ITimeMarker = ::Insights::ITimeMarker; }
+namespace UE::Insights { using ITimingViewSession = ::Insights::ITimingViewSession; }
+#endif // UE_INSIGHTS_BACKWARD_COMPATIBILITY_UE54
+
+namespace UE::Insights::Timing
 {
 
 ////////////////////////////////////////////////////////////////////////////////////////////////////
@@ -32,7 +50,7 @@ public:
 	virtual ~ITimeMarker() = default;
 
 	virtual double GetTime() const = 0;
-	virtual void SetTime(const double  InTime) = 0;
+	virtual void SetTime(const double InTime) = 0;
 };
 
 ////////////////////////////////////////////////////////////////////////////////////////////////////
@@ -114,7 +132,7 @@ public:
 	virtual double GetTimeMarker() const = 0;
 	/** Sets the current marker time. */
 	virtual void SetTimeMarker(double InTimeMarker) = 0;
-	/** Sets the current marker time and center the view on it */
+	/** Sets the current marker time and center the view on it. */
 	virtual void SetAndCenterOnTimeMarker(double InTimeMarker) = 0;
 
 	//////////////////////////////////////////////////
@@ -140,10 +158,10 @@ public:
 	virtual FTrackVisibilityChangedDelegate& OnTrackVisibilityChanged() = 0;
 
 	/** Gets the delegate to be invoked when a new track is added. */
-	virtual Insights::FTrackAddedDelegate& OnTrackAdded() = 0;
+	virtual FTrackAddedDelegate& OnTrackAdded() = 0;
 
 	/** Gets the delegate to be invoked when a track is removed. */
-	virtual Insights::FTrackRemovedDelegate& OnTrackRemoved() = 0;
+	virtual FTrackRemovedDelegate& OnTrackRemoved() = 0;
 
 	//////////////////////////////////////////////////
 
@@ -164,4 +182,77 @@ public:
 
 ////////////////////////////////////////////////////////////////////////////////////////////////////
 
+} // namespace UE::Insights::Timing
+
+////////////////////////////////////////////////////////////////////////////////////////////////////
+
+#if UE_INSIGHTS_BACKWARD_COMPATIBILITY_UE54
+namespace Insights
+{
+
+enum class ETimeChangedFlags : int32
+{
+	None,
+	Interactive = (1 << 0)
+};
+ENUM_CLASS_FLAGS(ETimeChangedFlags);
+
+class ITimeMarker
+{
+public:
+	virtual ~ITimeMarker() = default;
+
+	virtual double GetTime() const = 0;
+	virtual void SetTime(const double  InTime) = 0;
+};
+
+DECLARE_MULTICAST_DELEGATE_ThreeParams(FSelectionChangedDelegate, ETimeChangedFlags /*InFlags*/, double /*StartTime*/, double /*EndTime*/);
+DECLARE_MULTICAST_DELEGATE_TwoParams(FTimeMarkerChangedDelegate, ETimeChangedFlags /*InFlags*/, double /*TimeMarker*/);
+DECLARE_MULTICAST_DELEGATE_TwoParams(FCustomTimeMarkerChangedDelegate, ETimeChangedFlags /*InFlags*/, TSharedRef<ITimeMarker> /*TimeMarker*/);
+DECLARE_MULTICAST_DELEGATE_OneParam(FHoveredTrackChangedDelegate, const TSharedPtr<FBaseTimingTrack> /*InTrack*/);
+DECLARE_MULTICAST_DELEGATE_OneParam(FHoveredEventChangedDelegate, const TSharedPtr<const ITimingEvent> /*InEvent*/);
+DECLARE_MULTICAST_DELEGATE_OneParam(FSelectedTrackChangedDelegate, const TSharedPtr<FBaseTimingTrack> /*InTrack*/);
+DECLARE_MULTICAST_DELEGATE_OneParam(FSelectedEventChangedDelegate, const TSharedPtr<const ITimingEvent> /*InEvent*/);
+DECLARE_MULTICAST_DELEGATE(FTrackVisibilityChangedDelegate);
+DECLARE_MULTICAST_DELEGATE_OneParam(FTrackAddedDelegate, const TSharedPtr<const FBaseTimingTrack> /*Track*/);
+DECLARE_MULTICAST_DELEGATE_OneParam(FTrackRemovedDelegate, const TSharedPtr<const FBaseTimingTrack> /*Track*/);
+
+//class UE_DEPRECATED(5.5, "ITimingViewSession class was moved inside UE::Insights::Timing namespace") ITimingViewSession;
+class ITimingViewSession
+{
+public:
+	virtual ~ITimingViewSession() = default;
+	virtual const FName& GetName() const = 0;
+	virtual void AddTopDockedTrack(TSharedPtr<FBaseTimingTrack> Track) = 0;
+	virtual bool RemoveTopDockedTrack(TSharedPtr<FBaseTimingTrack> Track) = 0;
+	virtual void AddBottomDockedTrack(TSharedPtr<FBaseTimingTrack> Track) = 0;
+	virtual bool RemoveBottomDockedTrack(TSharedPtr<FBaseTimingTrack> Track) = 0;
+	virtual void AddScrollableTrack(TSharedPtr<FBaseTimingTrack> Track) = 0;
+	virtual bool RemoveScrollableTrack(TSharedPtr<FBaseTimingTrack> Track) = 0;
+	virtual void InvalidateScrollableTracksOrder() = 0;
+	virtual void AddForegroundTrack(TSharedPtr<FBaseTimingTrack> Track) = 0;
+	virtual bool RemoveForegroundTrack(TSharedPtr<FBaseTimingTrack> Track) = 0;
+	virtual void AddTrack(TSharedPtr<FBaseTimingTrack> Track, ETimingTrackLocation Location) = 0;
+	virtual bool RemoveTrack(TSharedPtr<FBaseTimingTrack> Track) = 0;
+	virtual TSharedPtr<FBaseTimingTrack> FindTrack(uint64 InTrackId) = 0;
+	virtual double GetTimeMarker() const = 0;
+	virtual void SetTimeMarker(double InTimeMarker) = 0;
+	virtual void SetAndCenterOnTimeMarker(double InTimeMarker) = 0;
+	virtual FSelectionChangedDelegate& OnSelectionChanged() = 0;
+	virtual FTimeMarkerChangedDelegate& OnTimeMarkerChanged() = 0;
+	virtual FCustomTimeMarkerChangedDelegate& OnCustomTimeMarkerChanged() = 0;
+	virtual FHoveredTrackChangedDelegate& OnHoveredTrackChanged() = 0;
+	virtual FHoveredEventChangedDelegate& OnHoveredEventChanged() = 0;
+	virtual FSelectedTrackChangedDelegate& OnSelectedTrackChanged() = 0;
+	virtual FSelectedEventChangedDelegate& OnSelectedEventChanged() = 0;
+	virtual FTrackVisibilityChangedDelegate& OnTrackVisibilityChanged() = 0;
+	virtual FTrackAddedDelegate& OnTrackAdded() = 0;
+	virtual FTrackRemovedDelegate& OnTrackRemoved() = 0;
+	virtual void ResetSelectedEvent() = 0;
+	virtual void ResetEventFilter() = 0;
+	virtual void PreventThrottling() = 0;
+	virtual void AddOverlayWidget(const TSharedRef<SWidget>& InWidget) = 0;
+};
+
 } // namespace Insights
+#endif // UE_INSIGHTS_BACKWARD_COMPATIBILITY_UE54
diff --git a/Engine/Source/Developer/TraceInsights/Public/Insights/IUnrealInsightsModule.h b/Engine/Source/Developer/TraceInsights/Public/Insights/IUnrealInsightsModule.h
--- a/Engine/Source/Developer/TraceInsights/Public/Insights/IUnrealInsightsModule.h
+++ b/Engine/Source/Developer/TraceInsights/Public/Insights/IUnrealInsightsModule.h
@@ -2,32 +2,23 @@
 
 #pragma once
 
-#include "CoreMinimal.h"
-#include "Modules/ModuleInterface.h"
+#include "CoreTypes.h"
+
+#include "Containers/Array.h"
 #include "Framework/Docking/LayoutExtender.h"
 #include "Framework/Docking/TabManager.h"
 #include "Framework/MultiBox/MultiBoxExtender.h"
+#include "Internationalization/Text.h"
+#include "Modules/ModuleInterface.h"
+#include "Templates/SharedPointer.h"
+#include "UObject/NameTypes.h"
 
 class FExtender;
 class FWorkspaceItem;
 
-namespace UE
-{
-namespace Trace
-{
-	class FStoreClient;
-}
-}
-
-namespace Insights
-{
-	class IInsightsManager;
-}
-
-namespace TraceServices
-{
-	class IAnalysisSession;
-}
+namespace UE::Trace { class FStoreClient; }
+namespace TraceServices { class IAnalysisSession; }
+namespace Insights { class IInsightsManager; }
 
 ////////////////////////////////////////////////////////////////////////////////////////////////////
 
@@ -168,6 +159,10 @@ class TRACEINSIGHTS_API IInsightsComponent;
 class TRACEINSIGHTS_API IUnrealInsightsModule : public IModuleInterface
 {
 public:
+	virtual ~IUnrealInsightsModule()
+	{
+	}
+
 	/**
 	 * Registers an IInsightsComponent. The component will Initialize().
 	 */
@@ -200,7 +195,7 @@ public:
 	 *
 	 * @param InStoreHost The host of the store to connect to.
 	 * @param InStorePort The port of the store to connect to.
-	 * @return If connected succesfully or not.
+	 * @return If connected successfully or not.
 	 */
 	virtual bool ConnectToStore(const TCHAR* InStoreHost, uint32 InStorePort=0) = 0;
 
@@ -209,7 +204,7 @@ public:
 	/**
 	 * Gets the Insights manager.
 	*/
-	virtual TSharedPtr<Insights::IInsightsManager> GetInsightsManager() = 0;
+	virtual TSharedPtr<::Insights::IInsightsManager> GetInsightsManager() = 0;
 
 	//////////////////////////////////////////////////
 
@@ -228,9 +223,9 @@ public:
 
 	/**
 	 * Starts analysis of the last live session. Called when the application starts in "Viewer" mode.
-	 * On failure, if InRetryTime is > 0, retry connecting every frame for RetryTime seconds 
-	 * 
-	 * @param InRetryTime How many seconds to retry connecting asyncronously
+	 * On failure, if InRetryTime is > 0, retry connecting every frame for RetryTime seconds.
+	 *
+	 * @param InRetryTime How many seconds to retry connecting asynchronously
 	 */
 	virtual void StartAnalysisForLastLiveSession(float InRetryTime = 1.0f) = 0;
 
@@ -286,7 +281,7 @@ public:
 	virtual void CreateSessionViewer(bool bAllowDebugTools = false) = 0;
 
 	/**
-	 * Called when the application shutsdown.
+	 * Called when the application shuts-down.
 	 */
 	virtual void ShutdownUserInterface() = 0;
 
@@ -319,7 +314,7 @@ public:
 	/** Initializes this component. Called by TraceInsights module when this component is registered. */
 	virtual void Initialize(IUnrealInsightsModule& Module) = 0;
 
-	/** Shutsdown this component. Called by TraceInsights module when this component is unregistered. */
+	/** Shuts-down this component. Called by TraceInsights module when this component is unregistered. */
 	virtual void Shutdown() = 0;
 
 	/* Allows this component to register major tabs. */
diff --git a/Engine/Source/Developer/TraceInsights/Public/Insights/Log.h b/Engine/Source/Developer/TraceInsights/Public/Insights/Log.h
--- a/Engine/Source/Developer/TraceInsights/Public/Insights/Log.h
+++ b/Engine/Source/Developer/TraceInsights/Public/Insights/Log.h
@@ -2,10 +2,7 @@
 
 #pragma once
 
-#include "CoreMinimal.h"
-#include "HAL/LowLevelMemTracker.h"
-#include "Logging/LogMacros.h"
+// TraceInsightsCore
+#include "InsightsCore/Common/Log.h"
 
 DECLARE_LOG_CATEGORY_EXTERN(TraceInsights, Verbose, All);
-
-LLM_DECLARE_TAG(Insights);
diff --git a/Engine/Source/Developer/TraceInsights/Public/Insights/Tests/InsightsTestUtils.h b/Engine/Source/Developer/TraceInsights/Public/Insights/Tests/InsightsTestUtils.h
--- a/Engine/Source/Developer/TraceInsights/Public/Insights/Tests/InsightsTestUtils.h
+++ b/Engine/Source/Developer/TraceInsights/Public/Insights/Tests/InsightsTestUtils.h
@@ -2,8 +2,10 @@
 
 #pragma once
 
-#include "CoreMinimal.h"
-#include "ProfilingDebugging/TraceAuxiliary.h"
+#include "CoreTypes.h"
+
+#include "Containers/UnrealString.h"
+#include "ProfilingDebugging/TraceAuxiliary.h" // for FTraceAuxiliary::EConnectionType
 
 class FAutomationTestBase;
 
diff --git a/Engine/Source/Developer/TraceInsights/Public/Insights/Tests/TimingProfilerTests.h b/Engine/Source/Developer/TraceInsights/Public/Insights/Tests/TimingProfilerTests.h
--- a/Engine/Source/Developer/TraceInsights/Public/Insights/Tests/TimingProfilerTests.h
+++ b/Engine/Source/Developer/TraceInsights/Public/Insights/Tests/TimingProfilerTests.h
@@ -2,10 +2,12 @@
 
 #pragma once
 
-#include "CoreMinimal.h"
+#include "CoreTypes.h"
+
 #include "Logging/LogMacros.h"
 
-#include "TraceServices/Model/TimingProfiler.h"
+// TraceServices
+#include "TraceServices/Model/TimingProfiler.h" // for TraceServices::EEventSortOrder
 
 DECLARE_LOG_CATEGORY_EXTERN(TimingProfilerTests, Log, All);
 
@@ -25,9 +27,8 @@ public:
 		uint64 EventCount = 0;
 		uint32 SumDepth = 0;
 		uint32 SumTimerIndex = 0;
-		double SessionDuration = 0;
-
-		double EnumerationDuration = 0;
+		double SessionDuration = 0.0;
+		double EnumerationDuration = 0.0;
 	};
 
 	struct FEnumerateTestParams
diff --git a/Engine/Source/Developer/TraceInsights/Public/Insights/TimingProfilerCommon.h b/Engine/Source/Developer/TraceInsights/Public/Insights/TimingProfilerCommon.h
--- a/Engine/Source/Developer/TraceInsights/Public/Insights/TimingProfilerCommon.h
+++ b/Engine/Source/Developer/TraceInsights/Public/Insights/TimingProfilerCommon.h
@@ -7,7 +7,7 @@
 #include "Stats/Stats.h"
 #include "Stats/StatsMisc.h"
 
-DECLARE_LOG_CATEGORY_EXTERN(TimingProfiler, Log, All);
+DECLARE_LOG_CATEGORY_EXTERN(LogTimingProfiler, Log, All);
 
 DECLARE_STATS_GROUP(TEXT("TimingProfiler"), STATGROUP_TimingProfiler, STATCAT_Advanced);
 
diff --git a/Engine/Source/Developer/TraceInsights/Public/Insights/ViewModels/BaseTimingTrack.h b/Engine/Source/Developer/TraceInsights/Public/Insights/ViewModels/BaseTimingTrack.h
--- a/Engine/Source/Developer/TraceInsights/Public/Insights/ViewModels/BaseTimingTrack.h
+++ b/Engine/Source/Developer/TraceInsights/Public/Insights/ViewModels/BaseTimingTrack.h
@@ -2,18 +2,31 @@
 
 #pragma once
 
-#include "CoreMinimal.h"
+#include "CoreTypes.h"
+
+#include "Containers/Array.h"
 #include "Input/Events.h"
 #include "Input/Reply.h"
 #include "Misc/EnumClassFlags.h"
+#include "Templates/SharedPointer.h"
+
+// TraceInsightsCore
+#include "InsightsCore/Common/SimpleRtti.h"
 
-#include "Insights/Common/SimpleRtti.h"
+// TraceInsights
+#include "Insights/Config.h"
+
+#if UE_INSIGHTS_BACKWARD_COMPATIBILITY_UE54
+#include "Insights/ITimingViewSession.h"
+#endif
 
 ////////////////////////////////////////////////////////////////////////////////////////////////////
 
 class FMenuBuilder;
-namespace TraceServices { class IAnalysisSession; };
-namespace Insights { class FFilterConfigurator; }
+struct FGeometry;
+
+namespace UE::Insights { class FDrawContext; }
+namespace UE::Insights { class FFilterConfigurator; }
 
 class FTimingEventSearchParameters;
 class FTimingTrackViewport;
@@ -23,9 +36,6 @@ class ITimingEvent;
 class ITimingEventRelation;
 class ITimingEventFilter;
 
-struct FDrawContext;
-struct FGeometry;
-
 ////////////////////////////////////////////////////////////////////////////////////////////////////
 
 enum class ETimingTrackLocation : uint32
@@ -113,7 +123,7 @@ public:
 	virtual const TSharedPtr<const ITimingEvent> GetHoveredEvent() const = 0;
 	virtual const TSharedPtr<const ITimingEvent> GetSelectedEvent() const = 0;
 	virtual const TSharedPtr<ITimingEventFilter> GetEventFilter() const = 0;
-	virtual FDrawContext& GetDrawContext() const = 0;
+	virtual UE::Insights::FDrawContext& GetDrawContext() const = 0;
 	virtual const ITimingViewDrawHelper& GetHelper() const = 0;
 };
 
@@ -194,7 +204,7 @@ public:
 	//////////////////////////////////////////////////
 
 	// PreUpdate callback called each frame, but only if the track is visible.
-	// In this update, neither the postion nor the size of the track is yet computed.
+	// In this update, neither the position nor the size of the track is yet computed.
 	// Track should update here its height.
 	virtual void PreUpdate(const ITimingTrackUpdateContext& Context) {}
 
@@ -233,7 +243,7 @@ public:
 	 * @param X The horizontal coordinate of the point tested; in Slate pixels (viewport coordinates).
 	 * @param Y The vertical coordinate of the point tested; in Slate pixels (viewport coordinates).
 	 * @param Viewport The timing viewport used to transform time in viewport coordinates.
-	 * @return The event located at (PosX, PosY) coordinates, if any; nullptr ottherwise.
+	 * @return The event located at (PosX, PosY) coordinates, if any; nullptr otherwise.
 	 */
 	virtual const TSharedPtr<const ITimingEvent> GetEvent(float InPosX, float InPosY, const FTimingTrackViewport& Viewport) const { return nullptr; }
 
@@ -265,7 +275,7 @@ public:
 	TWeakPtr<FBaseTimingTrack> GetParentTrack() const { return ParentTrack; }
 	bool IsChildTrack() const { return ParentTrack.IsValid(); }
 
-	virtual void SetFilterConfigurator(TSharedPtr<Insights::FFilterConfigurator> InFilterConfigurator) {}
+	virtual void SetFilterConfigurator(TSharedPtr<UE::Insights::FFilterConfigurator> InFilterConfigurator) {}
 
 	// Returns number of text lines needed to display the debug string.
 	//TODO: virtual int GetDebugStringLineCount() const { return 0; }
diff --git a/Engine/Source/Developer/TraceInsights/Public/Insights/ViewModels/GraphTrack.h b/Engine/Source/Developer/TraceInsights/Public/Insights/ViewModels/GraphTrack.h
--- a/Engine/Source/Developer/TraceInsights/Public/Insights/ViewModels/GraphTrack.h
+++ b/Engine/Source/Developer/TraceInsights/Public/Insights/ViewModels/GraphTrack.h
@@ -2,10 +2,12 @@
 
 #pragma once
 
-#include "CoreMinimal.h"
+#include "CoreTypes.h"
+
 #include "Fonts/SlateFontInfo.h"
+#include "Templates/SharedPointer.h"
 
-// Insights
+// TraceInsights
 #include "Insights/ViewModels/BaseTimingTrack.h"
 #include "Insights/ViewModels/GraphSeries.h"
 
@@ -115,7 +117,7 @@ public:
 protected:
 	void UpdateStats();
 
-	void DrawSeries(const FGraphSeries& Series, FDrawContext& DrawContext, const FTimingTrackViewport& Viewport) const;
+	void DrawSeries(const FGraphSeries& Series, UE::Insights::FDrawContext& DrawContext, const FTimingTrackViewport& Viewport) const;
 
 	virtual void DrawVerticalAxisGrid(const ITimingTrackDrawContext& Context) const;
 	virtual void DrawHeader(const ITimingTrackDrawContext& Context) const;
diff --git a/Engine/Source/Developer/TraceInsights/Public/Insights/ViewModels/GraphTrackBuilder.h b/Engine/Source/Developer/TraceInsights/Public/Insights/ViewModels/GraphTrackBuilder.h
--- a/Engine/Source/Developer/TraceInsights/Public/Insights/ViewModels/GraphTrackBuilder.h
+++ b/Engine/Source/Developer/TraceInsights/Public/Insights/ViewModels/GraphTrackBuilder.h
@@ -4,7 +4,6 @@
 
 #include "CoreMinimal.h"
 
-// Insights
 class FGraphTrack;
 class FGraphSeries;
 class FTimingTrackViewport;
diff --git a/Engine/Source/Developer/TraceInsights/Public/Insights/ViewModels/IFilterExecutor.h b/Engine/Source/Developer/TraceInsights/Public/Insights/ViewModels/IFilterExecutor.h
--- a/Engine/Source/Developer/TraceInsights/Public/Insights/ViewModels/IFilterExecutor.h
+++ b/Engine/Source/Developer/TraceInsights/Public/Insights/ViewModels/IFilterExecutor.h
@@ -2,19 +2,12 @@
 
 #pragma once
 
-#include "CoreMinimal.h"
+// TraceInsightsCore
+#include "InsightsCore/Filter/ViewModels/IFilterExecutor.h"
 
-namespace Insights
-{
- 
-////////////////////////////////////////////////////////////////////////////////////////////////////
+// TraceInsights
+#include "Insights/Config.h"
 
-class IFilterExecutor
-{
-public:
-	virtual bool ApplyFilters(const class FFilterContext& Context) const = 0;
-};
-
-////////////////////////////////////////////////////////////////////////////////////////////////////
-
-} // namespace Insights
\ No newline at end of file
+#if UE_INSIGHTS_BACKWARD_COMPATIBILITY_UE54
+namespace Insights { using IFilterExecutor = UE::Insights::IFilterExecutor; }
+#endif
diff --git a/Engine/Source/Developer/TraceInsights/Public/Insights/ViewModels/ITimingEvent.h b/Engine/Source/Developer/TraceInsights/Public/Insights/ViewModels/ITimingEvent.h
--- a/Engine/Source/Developer/TraceInsights/Public/Insights/ViewModels/ITimingEvent.h
+++ b/Engine/Source/Developer/TraceInsights/Public/Insights/ViewModels/ITimingEvent.h
@@ -2,15 +2,26 @@
 
 #pragma once
 
-#include "CoreMinimal.h"
+#include "CoreTypes.h"
 
-#include "Insights/Common/SimpleRtti.h"
+#include "Templates/SharedPointer.h"
+
+// TraceInsightsCore
+#include "InsightsCore/Common/SimpleRtti.h"
+
+// TraceInsights
+#include "Insights/Config.h"
+
+#if UE_INSIGHTS_BACKWARD_COMPATIBILITY_UE54
+#include "Insights/ITimingViewSession.h"
+#endif
 
 class ITimingViewDrawHelper;
 class FBaseTimingTrack;
-struct FDrawContext;
 class FTimingTrackViewport;
 
+namespace UE::Insights { class FDrawContext; }
+
 ////////////////////////////////////////////////////////////////////////////////////////////////////
 
 class TRACEINSIGHTS_API ITimingEvent
@@ -67,7 +78,7 @@ public:
 
 ////////////////////////////////////////////////////////////////////////////////////////////////////
 
-class ITimingEventRelation
+class TRACEINSIGHTS_API ITimingEventRelation
 {
 	INSIGHTS_DECLARE_RTTI_BASE(ITimingEventRelation)
 
@@ -80,7 +91,8 @@ public:
 
 	ITimingEventRelation() {}
 	virtual ~ITimingEventRelation() {}
-	virtual void Draw(const FDrawContext& DrawContext, const FTimingTrackViewport& Viewport, const ITimingViewDrawHelper& Helper, const EDrawFilter Filter) = 0;
+
+	virtual void Draw(const UE::Insights::FDrawContext& DrawContext, const FTimingTrackViewport& Viewport, const ITimingViewDrawHelper& Helper, const EDrawFilter Filter) = 0;
 };
 
 ////////////////////////////////////////////////////////////////////////////////////////////////////
diff --git a/Engine/Source/Developer/TraceInsights/Public/Insights/ViewModels/TimingEvent.h b/Engine/Source/Developer/TraceInsights/Public/Insights/ViewModels/TimingEvent.h
--- a/Engine/Source/Developer/TraceInsights/Public/Insights/ViewModels/TimingEvent.h
+++ b/Engine/Source/Developer/TraceInsights/Public/Insights/ViewModels/TimingEvent.h
@@ -2,9 +2,15 @@
 
 #pragma once
 
-#include "CoreMinimal.h"
+#include "CoreTypes.h"
 
-// Insights
+#include "Containers/Array.h"
+#include "Templates/SharedPointer.h"
+
+// TraceInsightsCore
+#include "InsightsCore/Common/SimpleRtti.h"
+
+// TraceInsights
 #include "Insights/ViewModels/ITimingEvent.h"
 #include "Insights/ViewModels/TimingEventSearch.h"
 
diff --git a/Engine/Source/Developer/TraceInsights/Public/Insights/ViewModels/TimingEventSearch.h b/Engine/Source/Developer/TraceInsights/Public/Insights/ViewModels/TimingEventSearch.h
--- a/Engine/Source/Developer/TraceInsights/Public/Insights/ViewModels/TimingEventSearch.h
+++ b/Engine/Source/Developer/TraceInsights/Public/Insights/ViewModels/TimingEventSearch.h
@@ -2,10 +2,13 @@
 
 #pragma once
 
-#include "CoreMinimal.h"
+#include "CoreTypes.h"
+
 #include "Containers/Map.h"
+#include "Misc/EnumClassFlags.h"
+#include "Templates/SharedPointer.h"
 
-namespace Insights
+namespace UE::Insights
 {
 	class IFilterExecutor;
 }
@@ -100,7 +103,7 @@ public:
 	// Search behavior flags
 	ETimingEventSearchFlags Flags;
 
-	TSharedPtr<Insights::IFilterExecutor> FilterExecutor;
+	TSharedPtr<UE::Insights::IFilterExecutor> FilterExecutor;
 
 	ESearchDirection SearchDirection = ESearchDirection::Forward;
 };
@@ -194,21 +197,21 @@ private:
 //
 //		TTimingEventSearch<FMyStruct>::Search(
 //		Parameters,
-// 		[](const TTimingEventSearch<FMyStruct>::FContext& InContext)
-// 		{
-// 			for (FMyStruct& Payload : Payloads)
+//		[](const TTimingEventSearch<FMyStruct>::FContext& InContext)
+//		{
+//			for (FMyStruct& Payload : Payloads)
 //			{
-// 				InContext.Check(InEventStartTime, InEventEndTime, 0, Payload);
-// 			}
-// 		},
-// 		[&MatchedPayload](double InStartTime, double InEndTime, uint32 InDepth, const FMyStruct& InEvent)
-// 		{
-// 			MatchedPayload = InEvent;
-// 		},
-// 		[&MatchedPayload](double InStartTime, double InEndTime, uint32 InDepth)
-// 		{
-// 			// Do something with MatchedPayload, e.g. call a captured lambda.
-// 		});
+//				InContext.Check(InEventStartTime, InEventEndTime, 0, Payload);
+//			}
+//		},
+//		[&MatchedPayload](double InStartTime, double InEndTime, uint32 InDepth, const FMyStruct& InEvent)
+//		{
+//			MatchedPayload = InEvent;
+//		},
+//		[&MatchedPayload](double InStartTime, double InEndTime, uint32 InDepth)
+//		{
+//			// Do something with MatchedPayload, e.g. call a captured lambda.
+//		});
 //
 template <typename PayloadType>
 struct TTimingEventSearch
diff --git a/Engine/Source/Developer/TraceInsights/Public/Insights/ViewModels/TimingEventsTrack.h b/Engine/Source/Developer/TraceInsights/Public/Insights/ViewModels/TimingEventsTrack.h
--- a/Engine/Source/Developer/TraceInsights/Public/Insights/ViewModels/TimingEventsTrack.h
+++ b/Engine/Source/Developer/TraceInsights/Public/Insights/ViewModels/TimingEventsTrack.h
@@ -2,10 +2,12 @@
 
 #pragma once
 
-#include "CoreMinimal.h"
+#include "CoreTypes.h"
+
 #include "Fonts/SlateFontInfo.h"
+#include "Templates/SharedPointer.h"
 
-// Insights
+// TraceInsights
 #include "Insights/ViewModels/BaseTimingTrack.h"
 
 ////////////////////////////////////////////////////////////////////////////////////////////////////
@@ -84,8 +86,8 @@ protected:
 
 	void DrawMarkers(const ITimingTrackDrawContext& Context, float LineY, float LineH) const;
 
-	void DrawSelectedEventInfo(const FString& InText, const FTimingTrackViewport& Viewport, const FDrawContext& DrawContext, const FSlateBrush* WhiteBrush, const FSlateFontInfo& Font) const;
-	void DrawSelectedEventInfoEx(const FString& InText, const FString& InLeftText, const FString& InTopText, const FTimingTrackViewport& Viewport, const FDrawContext& DrawContext, const FSlateBrush* WhiteBrush, const FSlateFontInfo& Font) const;
+	void DrawSelectedEventInfo(const FString& InText, const FTimingTrackViewport& Viewport, const UE::Insights::FDrawContext& DrawContext, const FSlateBrush* WhiteBrush, const FSlateFontInfo& Font) const;
+	void DrawSelectedEventInfoEx(const FString& InText, const FString& InLeftText, const FString& InTopText, const FTimingTrackViewport& Viewport, const UE::Insights::FDrawContext& DrawContext, const FSlateBrush* WhiteBrush, const FSlateFontInfo& Font) const;
 
 	int32 GetHeaderBackgroundLayerId(const ITimingTrackDrawContext& Context) const;
 	int32 GetHeaderTextLayerId(const ITimingTrackDrawContext& Context) const;
diff --git a/Engine/Source/Developer/TraceInsights/Public/Insights/ViewModels/TooltipDrawState.h b/Engine/Source/Developer/TraceInsights/Public/Insights/ViewModels/TooltipDrawState.h
--- a/Engine/Source/Developer/TraceInsights/Public/Insights/ViewModels/TooltipDrawState.h
+++ b/Engine/Source/Developer/TraceInsights/Public/Insights/ViewModels/TooltipDrawState.h
@@ -2,15 +2,19 @@
 
 #pragma once
 
-#include "CoreMinimal.h"
+#include "CoreTypes.h"
+
+#include "Containers/Array.h"
 #include "Fonts/SlateFontInfo.h"
 #include "Math/Color.h"
+#include "Templates/SharedPointer.h"
 
 ////////////////////////////////////////////////////////////////////////////////////////////////////
 
-struct FDrawContext;
 struct FSlateBrush;
 
+namespace UE::Insights { class FDrawContext; }
+
 ////////////////////////////////////////////////////////////////////////////////////////////////////
 
 class TRACEINSIGHTS_API FTooltipDrawState
@@ -80,7 +84,7 @@ public:
 	void SetPosition(const FVector2D& MousePosition, const float MinX, const float MaxX, const float MinY, const float MaxY);
 	void SetPosition(const float PosX, const float PosY) { Position.X = PosX; Position.Y = PosY; }
 
-	void Draw(const FDrawContext& DrawContext) const;
+	void Draw(const UE::Insights::FDrawContext& DrawContext) const;
 
 	void SetFontScale(float InFontScale) { FontScale = InFontScale; }
 	float GetFontScale() const { return FontScale; }
diff --git a/Engine/Source/Developer/TraceInsights/Public/Insights/Widgets/STraceServerControl.h b/Engine/Source/Developer/TraceInsights/Public/Insights/Widgets/STraceServerControl.h
--- a/Engine/Source/Developer/TraceInsights/Public/Insights/Widgets/STraceServerControl.h
+++ b/Engine/Source/Developer/TraceInsights/Public/Insights/Widgets/STraceServerControl.h
@@ -2,14 +2,17 @@
 
 #pragma once
 
-#include <atomic>
 #include "CoreFwd.h"
 #include "HAL/CriticalSection.h"
 #include "Templates/UniquePtr.h"
-#include "Trace/StoreClient.h"
 #include "UObject/NameTypes.h"
 #include "UObject/UnrealNames.h"
 
+// TraceAnalysis
+#include "Trace/StoreClient.h"
+
+#include <atomic>
+
 class FMenuBuilder;
 
 class TRACEINSIGHTS_API STraceServerControl
diff --git a/Engine/Source/Developer/TraceInsightsCore/Public/InsightsCore/Filter/ViewModels/IFilterExecutor.h b/Engine/Source/Developer/TraceInsightsCore/Public/InsightsCore/Filter/ViewModels/IFilterExecutor.h
new file mode 100644
--- /dev/null
+++ b/Engine/Source/Developer/TraceInsightsCore/Public/InsightsCore/Filter/ViewModels/IFilterExecutor.h
@@ -0,0 +1,18 @@
+// Copyright Epic Games, Inc. All Rights Reserved.
+
+#pragma once
+
+#include "CoreTypes.h"
+
+namespace UE::Insights
+{
+
+class FFilterContext;
+
+class TRACEINSIGHTSCORE_API IFilterExecutor
+{
+public:
+	virtual bool ApplyFilters(const class FFilterContext& Context) const = 0;
+};
+
+} // namespace UE::Insights
diff --git a/Engine/Source/Developer/TraceInsightsCore/Public/InsightsCore/Filter/ViewModels/TimeFilterValueConverter.h b/Engine/Source/Developer/TraceInsightsCore/Public/InsightsCore/Filter/ViewModels/TimeFilterValueConverter.h
new file mode 100644
--- /dev/null
+++ b/Engine/Source/Developer/TraceInsightsCore/Public/InsightsCore/Filter/ViewModels/TimeFilterValueConverter.h
@@ -0,0 +1,23 @@
+// Copyright Epic Games, Inc. All Rights Reserved.
+
+#pragma once
+
+#include "CoreTypes.h"
+
+#include "Containers/UnrealString.h"
+#include "Internationalization/Text.h"
+
+#include "InsightsCore/Filter/ViewModels/Filters.h"
+
+namespace UE::Insights
+{
+
+class TRACEINSIGHTSCORE_API FTimeFilterValueConverter : public IFilterValueConverter
+{
+public:
+	virtual bool Convert(const FString& Input, double& Output, FText& OutError) const override;
+	virtual FText GetTooltipText() const override;
+	virtual FText GetHintText() const override;
+};
+
+} // namespace UE::Insights
diff --git a/Engine/Source/Developer/TraceInsightsCore/Public/InsightsCore/ITraceInsightsCoreModule.h b/Engine/Source/Developer/TraceInsightsCore/Public/InsightsCore/ITraceInsightsCoreModule.h
new file mode 100644
--- /dev/null
+++ b/Engine/Source/Developer/TraceInsightsCore/Public/InsightsCore/ITraceInsightsCoreModule.h
@@ -0,0 +1,16 @@
+// Copyright Epic Games, Inc. All Rights Reserved.
+
+#pragma once
+
+#include "CoreTypes.h"
+
+#include "Modules/ModuleInterface.h"
+
+/** Interface for TraceInsightsCore module. */
+class TRACEINSIGHTSCORE_API ITraceInsightsCoreModule : public IModuleInterface
+{
+public:
+	virtual ~ITraceInsightsCoreModule()
+	{
+	}
+};
diff --git a/Engine/Source/Developer/TraceInsightsCore/Public/InsightsCore/Table/ViewModels/TableCellValueFormatter.h b/Engine/Source/Developer/TraceInsightsCore/Public/InsightsCore/Table/ViewModels/TableCellValueFormatter.h
new file mode 100644
--- /dev/null
+++ b/Engine/Source/Developer/TraceInsightsCore/Public/InsightsCore/Table/ViewModels/TableCellValueFormatter.h
@@ -0,0 +1,210 @@
+// Copyright Epic Games, Inc. All Rights Reserved.
+
+#pragma once
+
+#include "CoreTypes.h"
+
+#include "Layout/Visibility.h"
+
+#include "InsightsCore/Table/ViewModels/TableCellValue.h"
+
+class IToolTip;
+
+namespace UE::Insights
+{
+
+class FBaseTreeNode;
+class FTableColumn;
+
+////////////////////////////////////////////////////////////////////////////////////////////////////
+
+class TRACEINSIGHTSCORE_API ITableCellValueFormatter
+{
+public:
+	virtual FText FormatValue(const TOptional<FTableCellValue>& InValue) const = 0;
+	virtual FText FormatValueForTooltip(const TOptional<FTableCellValue>& InValue) const = 0;
+
+	virtual FText FormatValue(const FTableColumn& Column, const FBaseTreeNode& Node) const = 0;
+	virtual FText FormatValueForTooltip(const FTableColumn& Column, const FBaseTreeNode& Node) const = 0;
+	virtual FText FormatValueForGrouping(const FTableColumn& Column, const FBaseTreeNode& Node) const = 0;
+
+	virtual FText CopyValue(const FTableColumn& Column, const FBaseTreeNode& Node) const = 0;
+	virtual FText CopyTooltip(const FTableColumn& Column, const FBaseTreeNode& Node) const = 0;
+
+	virtual TSharedPtr<IToolTip> GetCustomTooltip(const FTableColumn& Column, const FBaseTreeNode& Node) const = 0;
+};
+
+////////////////////////////////////////////////////////////////////////////////////////////////////
+
+class TRACEINSIGHTSCORE_API FTableCellValueFormatter : public ITableCellValueFormatter
+{
+public:
+	FTableCellValueFormatter() {}
+	virtual ~FTableCellValueFormatter() {}
+
+	virtual FText FormatValue(const TOptional<FTableCellValue>& InValue) const override { return FText::GetEmpty(); }
+	virtual FText FormatValueForTooltip(const TOptional<FTableCellValue>& InValue) const override { return FormatValue(InValue); }
+
+	virtual FText FormatValue(const FTableColumn& Column, const FBaseTreeNode& Node) const override; // { return FormatValue(Column.GetValue(Node)); }
+	virtual FText FormatValueForTooltip(const FTableColumn& Column, const FBaseTreeNode& Node) const override; // { return FormatValueForTooltip(Column.GetValue(Node)); }
+	virtual FText FormatValueForGrouping(const FTableColumn& Column, const FBaseTreeNode& Node) const override; // { return FormatValueForTooltip(Column.GetValue(Node)); }
+
+	virtual FText CopyValue(const FTableColumn& Column, const FBaseTreeNode& Node) const override; // { return FormatValue(Column, Node); }
+	virtual FText CopyTooltip(const FTableColumn& Column, const FBaseTreeNode& Node) const override; // { return FormatValueForTooltip(Column, Node); }
+
+	virtual TSharedPtr<IToolTip> GetCustomTooltip(const FTableColumn& Column, const FBaseTreeNode& Node) const override;
+
+	static EVisibility GetTooltipVisibility();
+};
+
+////////////////////////////////////////////////////////////////////////////////////////////////////
+
+class TRACEINSIGHTSCORE_API FTextValueFormatter : public FTableCellValueFormatter
+{
+public:
+	virtual FText FormatValue(const TOptional<FTableCellValue>& InValue) const override
+	{
+		if (InValue.IsSet())
+		{
+			return InValue.GetValue().GetText();
+		}
+		return FText::GetEmpty();
+	}
+};
+
+////////////////////////////////////////////////////////////////////////////////////////////////////
+
+class TRACEINSIGHTSCORE_API FAsTextValueFormatter : public FTableCellValueFormatter
+{
+public:
+	virtual FText FormatValue(const TOptional<FTableCellValue>& InValue) const override
+	{
+		if (InValue.IsSet())
+		{
+			return InValue.GetValue().AsText();
+		}
+		return FText::GetEmpty();
+	}
+};
+
+////////////////////////////////////////////////////////////////////////////////////////////////////
+
+class TRACEINSIGHTSCORE_API FBoolValueFormatterAsTrueFalse : public FTableCellValueFormatter
+{
+public:
+	virtual FText FormatValue(const TOptional<FTableCellValue>& InValue) const override;
+};
+
+////////////////////////////////////////////////////////////////////////////////////////////////////
+
+class TRACEINSIGHTSCORE_API FBoolValueFormatterAsOnOff : public FTableCellValueFormatter
+{
+public:
+	virtual FText FormatValue(const TOptional<FTableCellValue>& InValue) const override;
+};
+
+////////////////////////////////////////////////////////////////////////////////////////////////////
+
+class TRACEINSIGHTSCORE_API FInt64ValueFormatterAsNumber : public FTableCellValueFormatter
+{
+public:
+	virtual FText FormatValue(const TOptional<FTableCellValue>& InValue) const override
+	{
+		return InValue.IsSet() ? FText::AsNumber(InValue.GetValue().Int64) : FText::GetEmpty();
+	}
+};
+
+////////////////////////////////////////////////////////////////////////////////////////////////////
+
+class TRACEINSIGHTSCORE_API FInt64ValueFormatterAsUInt32InfinteNumber : public FTableCellValueFormatter
+{
+public:
+	virtual FText FormatValue(const TOptional<FTableCellValue>& InValue) const override;
+	virtual FText FormatValueForTooltip(const TOptional<FTableCellValue>& InValue) const override;
+};
+
+////////////////////////////////////////////////////////////////////////////////////////////////////
+
+class TRACEINSIGHTSCORE_API FInt64ValueFormatterAsHex32 : public FTableCellValueFormatter
+{
+public:
+	virtual FText FormatValue(const TOptional<FTableCellValue>& InValue) const override
+	{
+		return InValue.IsSet() ? FText::FromString(FString::Printf(TEXT("0x%08X"), static_cast<uint32>(InValue.GetValue().Int64))) : FText::GetEmpty();
+	}
+};
+
+////////////////////////////////////////////////////////////////////////////////////////////////////
+
+class TRACEINSIGHTSCORE_API FInt64ValueFormatterAsHex64 : public FTableCellValueFormatter
+{
+public:
+	virtual FText FormatValue(const TOptional<FTableCellValue>& InValue) const override
+	{
+		return InValue.IsSet() ? FText::FromString(FString::Printf(TEXT("0x%016llX"), static_cast<uint64>(InValue.GetValue().Int64))) : FText::GetEmpty();
+	}
+};
+
+////////////////////////////////////////////////////////////////////////////////////////////////////
+
+class TRACEINSIGHTSCORE_API FInt64ValueFormatterAsMemory : public FTableCellValueFormatter
+{
+public:
+	virtual FText FormatValue(const TOptional<FTableCellValue>& InValue) const override;
+	virtual FText FormatValueForTooltip(const TOptional<FTableCellValue>& InValue) const override;
+};
+
+////////////////////////////////////////////////////////////////////////////////////////////////////
+
+class TRACEINSIGHTSCORE_API FFloatValueFormatterAsNumber : public FTableCellValueFormatter
+{
+public:
+	virtual FText FormatValue(const TOptional<FTableCellValue>& InValue) const override;
+};
+
+////////////////////////////////////////////////////////////////////////////////////////////////////
+
+class TRACEINSIGHTSCORE_API FFloatValueFormatterAsTimeAuto : public FTableCellValueFormatter
+{
+public:
+	virtual FText FormatValue(const TOptional<FTableCellValue>& InValue) const override;
+	virtual FText FormatValueForTooltip(const TOptional<FTableCellValue>& InValue) const override;
+};
+
+////////////////////////////////////////////////////////////////////////////////////////////////////
+
+class TRACEINSIGHTSCORE_API FDoubleValueFormatterAsNumber : public FTableCellValueFormatter
+{
+public:
+	virtual FText FormatValue(const TOptional<FTableCellValue>& InValue) const override;
+};
+
+////////////////////////////////////////////////////////////////////////////////////////////////////
+
+class TRACEINSIGHTSCORE_API FDoubleValueFormatterAsTimeAuto : public FTableCellValueFormatter
+{
+public:
+	virtual FText FormatValue(const TOptional<FTableCellValue>& InValue) const override;
+	virtual FText FormatValueForTooltip(const TOptional<FTableCellValue>& InValue) const override;
+};
+
+////////////////////////////////////////////////////////////////////////////////////////////////////
+
+class TRACEINSIGHTSCORE_API FDoubleValueFormatterAsTimeMs : public FTableCellValueFormatter
+{
+public:
+	virtual FText FormatValue(const TOptional<FTableCellValue>& InValue) const override;
+	virtual FText FormatValueForTooltip(const TOptional<FTableCellValue>& InValue) const override;
+};
+
+////////////////////////////////////////////////////////////////////////////////////////////////////
+
+class TRACEINSIGHTSCORE_API FCStringValueFormatterAsText : public FTableCellValueFormatter
+{
+public:
+	virtual FText FormatValue(const TOptional<FTableCellValue>& InValue) const override;
+};
+
+////////////////////////////////////////////////////////////////////////////////////////////////////
+
+} // namespace UE::Insights
diff --git a/Engine/Source/Developer/TraceInsightsCore/Public/InsightsCore/Table/ViewModels/TableCellValueGetter.h b/Engine/Source/Developer/TraceInsightsCore/Public/InsightsCore/Table/ViewModels/TableCellValueGetter.h
new file mode 100644
--- /dev/null
+++ b/Engine/Source/Developer/TraceInsightsCore/Public/InsightsCore/Table/ViewModels/TableCellValueGetter.h
@@ -0,0 +1,54 @@
+// Copyright Epic Games, Inc. All Rights Reserved.
+
+#pragma once
+
+#include "CoreTypes.h"
+
+#include "InsightsCore/Table/ViewModels/TableCellValue.h"
+
+namespace UE::Insights
+{
+
+class FBaseTreeNode;
+class FTableColumn;
+
+////////////////////////////////////////////////////////////////////////////////////////////////////
+
+class TRACEINSIGHTSCORE_API ITableCellValueGetter
+{
+public:
+	virtual const TOptional<FTableCellValue> GetValue(const FTableColumn& Column, const FBaseTreeNode& Node) const = 0;
+	virtual uint64 GetValueId(const FTableColumn& Column, const FBaseTreeNode& Node) const = 0;
+};
+
+////////////////////////////////////////////////////////////////////////////////////////////////////
+
+class TRACEINSIGHTSCORE_API FTableCellValueGetter : public ITableCellValueGetter
+{
+public:
+	FTableCellValueGetter() {}
+	virtual ~FTableCellValueGetter() {}
+
+	virtual const TOptional<FTableCellValue> GetValue(const FTableColumn& Column, const FBaseTreeNode& Node) const override { return TOptional<FTableCellValue>(); }
+	virtual uint64 GetValueId(const FTableColumn& Column, const FBaseTreeNode& Node) const override { return 0; }
+};
+
+////////////////////////////////////////////////////////////////////////////////////////////////////
+
+class TRACEINSIGHTSCORE_API FNameValueGetter : public FTableCellValueGetter
+{
+public:
+	virtual const TOptional<FTableCellValue> GetValue(const FTableColumn& Column, const FBaseTreeNode& Node) const override;
+};
+
+////////////////////////////////////////////////////////////////////////////////////////////////////
+
+class TRACEINSIGHTSCORE_API FDisplayNameValueGetter : public FTableCellValueGetter
+{
+public:
+	virtual const TOptional<FTableCellValue> GetValue(const FTableColumn& Column, const FBaseTreeNode& Node) const override;
+};
+
+////////////////////////////////////////////////////////////////////////////////////////////////////
+
+} // namespace UE::Insights
diff --git a/Engine/Source/Developer/TraceInsightsCore/Public/InsightsCore/Table/ViewModels/TableCellValueSorter.h b/Engine/Source/Developer/TraceInsightsCore/Public/InsightsCore/Table/ViewModels/TableCellValueSorter.h
new file mode 100644
--- /dev/null
+++ b/Engine/Source/Developer/TraceInsightsCore/Public/InsightsCore/Table/ViewModels/TableCellValueSorter.h
@@ -0,0 +1,184 @@
+// Copyright Epic Games, Inc. All Rights Reserved.
+
+#pragma once
+
+#include "CoreTypes.h"
+
+#include "InsightsCore/Common/AsyncOperationProgress.h"
+#include "InsightsCore/Table/ViewModels/TableCellValue.h"
+#include "InsightsCore/Table/ViewModels/TableColumn.h"
+
+struct FSlateBrush;
+
+namespace UE::Insights
+{
+
+class FBaseTreeNode;
+typedef TSharedPtr<class FBaseTreeNode> FBaseTreeNodePtr;
+
+////////////////////////////////////////////////////////////////////////////////////////////////////
+
+enum class TRACEINSIGHTSCORE_API ESortMode
+{
+	Ascending,
+	Descending
+};
+
+////////////////////////////////////////////////////////////////////////////////////////////////////
+
+typedef TFunction<bool(const FBaseTreeNodePtr& A, const FBaseTreeNodePtr& B)> FTreeNodeCompareFunc;
+
+class TRACEINSIGHTSCORE_API ITableCellValueSorter
+{
+public:
+	virtual FName GetName() const = 0;
+	virtual FText GetShortName() const = 0;
+	virtual FText GetTitleName() const = 0;
+	virtual FText GetDescription() const = 0;
+	virtual FName GetColumnId() const = 0;
+
+	virtual FSlateBrush* GetIcon(ESortMode SortMode) const = 0;
+
+	virtual FTreeNodeCompareFunc GetTreeNodeCompareDelegate(ESortMode SortMode) const = 0;
+
+	virtual void Sort(TArray<FBaseTreeNodePtr>& NodesToSort, ESortMode SortMode) const = 0;
+
+	virtual void SetAsyncOperationProgress(IAsyncOperationProgress* AsyncOperationProgress) = 0;
+};
+
+////////////////////////////////////////////////////////////////////////////////////////////////////
+
+class TRACEINSIGHTSCORE_API FTableCellValueSorter : public ITableCellValueSorter
+{
+public:
+	FTableCellValueSorter(const FName InName, const FText& InShortName, const FText& InTitleName, const FText& InDescription, TSharedRef<FTableColumn> InColumnRef);
+	virtual ~FTableCellValueSorter() {}
+
+	virtual FName GetName() const override { return Name; }
+	virtual FText GetShortName() const override { return ShortName; }
+	virtual FText GetTitleName() const override { return TitleName; }
+	virtual FText GetDescription() const override { return Description; }
+	virtual FName GetColumnId() const override { return ColumnRef->GetId(); }
+
+	virtual FSlateBrush* GetIcon(ESortMode SortMode) const override { return SortMode == ESortMode::Ascending ? AscendingIcon : DescendingIcon; }
+	virtual FTreeNodeCompareFunc GetTreeNodeCompareDelegate(ESortMode SortMode) const override { return SortMode == ESortMode::Ascending ? AscendingCompareDelegate : DescendingCompareDelegate; }
+
+	virtual void Sort(TArray<FBaseTreeNodePtr>& NodesToSort, ESortMode SortMode) const override;
+
+	virtual void SetAsyncOperationProgress(IAsyncOperationProgress* InAsyncOperationProgress) override { AsyncOperationProgress = InAsyncOperationProgress; }
+
+protected:
+	// Attempts to cancel the sort by throwing an exception. If exceptions are not available the return value is meant to be returned from sort predicates to speed up the sort.
+	bool CancelSort() const;
+	bool ShouldCancelSort() const;
+
+protected:
+	FName Name;
+	FText ShortName;
+	FText TitleName;
+	FText Description;
+
+	TSharedRef<FTableColumn> ColumnRef;
+
+	FSlateBrush* AscendingIcon;
+	FSlateBrush* DescendingIcon;
+
+	FTreeNodeCompareFunc AscendingCompareDelegate;
+	FTreeNodeCompareFunc DescendingCompareDelegate;
+
+	IAsyncOperationProgress* AsyncOperationProgress = nullptr;
+};
+
+////////////////////////////////////////////////////////////////////////////////////////////////////
+
+class TRACEINSIGHTSCORE_API FBaseTableColumnSorter : public FTableCellValueSorter
+{
+public:
+	FBaseTableColumnSorter(TSharedRef<FTableColumn> InColumnRef);
+};
+
+////////////////////////////////////////////////////////////////////////////////////////////////////
+
+class TRACEINSIGHTSCORE_API FSorterByName : public FTableCellValueSorter
+{
+public:
+	FSorterByName(TSharedRef<FTableColumn> InColumnRef);
+};
+
+////////////////////////////////////////////////////////////////////////////////////////////////////
+
+class TRACEINSIGHTSCORE_API FSorterByTypeName : public FTableCellValueSorter
+{
+public:
+	FSorterByTypeName(TSharedRef<FTableColumn> InColumnRef);
+};
+
+////////////////////////////////////////////////////////////////////////////////////////////////////
+
+class TRACEINSIGHTSCORE_API FSorterByBoolValue : public FBaseTableColumnSorter
+{
+public:
+	FSorterByBoolValue(TSharedRef<FTableColumn> InColumnRef);
+};
+
+////////////////////////////////////////////////////////////////////////////////////////////////////
+
+class TRACEINSIGHTSCORE_API FSorterByInt64Value : public FBaseTableColumnSorter
+{
+public:
+	FSorterByInt64Value(TSharedRef<FTableColumn> InColumnRef);
+
+	virtual void Sort(TArray<FBaseTreeNodePtr>& NodesToSort, ESortMode SortMode) const override;
+};
+
+////////////////////////////////////////////////////////////////////////////////////////////////////
+
+class TRACEINSIGHTSCORE_API FSorterByFloatValue : public FBaseTableColumnSorter
+{
+public:
+	FSorterByFloatValue(TSharedRef<FTableColumn> InColumnRef);
+
+	virtual void Sort(TArray<FBaseTreeNodePtr>& NodesToSort, ESortMode SortMode) const override;
+};
+
+////////////////////////////////////////////////////////////////////////////////////////////////////
+
+class TRACEINSIGHTSCORE_API FSorterByDoubleValue : public FBaseTableColumnSorter
+{
+public:
+	FSorterByDoubleValue(TSharedRef<FTableColumn> InColumnRef);
+
+	virtual void Sort(TArray<FBaseTreeNodePtr>& NodesToSort, ESortMode SortMode) const override;
+};
+
+////////////////////////////////////////////////////////////////////////////////////////////////////
+
+class TRACEINSIGHTSCORE_API FSorterByCStringValue : public FBaseTableColumnSorter
+{
+public:
+	FSorterByCStringValue(TSharedRef<FTableColumn> InColumnRef);
+};
+
+////////////////////////////////////////////////////////////////////////////////////////////////////
+
+class TRACEINSIGHTSCORE_API FSorterByTextValue : public FBaseTableColumnSorter
+{
+public:
+	FSorterByTextValue(TSharedRef<FTableColumn> InColumnRef);
+
+	virtual void Sort(TArray<FBaseTreeNodePtr>& NodesToSort, ESortMode SortMode) const override;
+};
+
+////////////////////////////////////////////////////////////////////////////////////////////////////
+
+class TRACEINSIGHTSCORE_API FSorterByTextValueWithId : public FBaseTableColumnSorter
+{
+public:
+	FSorterByTextValueWithId(TSharedRef<FTableColumn> InColumnRef);
+
+	virtual void Sort(TArray<FBaseTreeNodePtr>& NodesToSort, ESortMode SortMode) const override;
+};
+
+////////////////////////////////////////////////////////////////////////////////////////////////////
+
+} // namespace UE::Insights
diff --git a/Engine/Source/Developer/TraceInsightsCore/Public/InsightsCore/Table/ViewModels/TreeNodeGrouping.h b/Engine/Source/Developer/TraceInsightsCore/Public/InsightsCore/Table/ViewModels/TreeNodeGrouping.h
new file mode 100644
--- /dev/null
+++ b/Engine/Source/Developer/TraceInsightsCore/Public/InsightsCore/Table/ViewModels/TreeNodeGrouping.h
@@ -0,0 +1,283 @@
+// Copyright Epic Games, Inc. All Rights Reserved.
+
+#pragma once
+
+#include "CoreTypes.h"
+
+#include "Containers/Array.h"
+#include "Containers/Map.h"
+#include "Containers/StringFwd.h"
+#include "Delegates/DelegateCombinations.h"
+#include "Internationalization/Text.h"
+#include "Templates/SharedPointer.h"
+#include "UObject/NameTypes.h"
+
+#include "InsightsCore/Common/AsyncOperationProgress.h"
+#include "InsightsCore/Common/SimpleRtti.h"
+#include "InsightsCore/Table/ViewModels/BaseTreeNode.h"
+#include "InsightsCore/Table/ViewModels/TableColumn.h"
+#include "InsightsCore/Table/ViewModels/TableTreeNode.h"
+
+#include <atomic>
+
+struct FSlateBrush;
+
+namespace UE::Insights
+{
+
+class FTable;
+
+////////////////////////////////////////////////////////////////////////////////////////////////////
+
+class TRACEINSIGHTSCORE_API ITreeNodeGrouping
+{
+	INSIGHTS_DECLARE_RTTI_BASE(ITreeNodeGrouping)
+
+public:
+	virtual FText GetShortName() const = 0;
+	virtual FText GetTitleName() const = 0;
+	virtual FText GetDescription() const = 0;
+
+	virtual FName GetBrushName() const = 0;
+	virtual const FSlateBrush* GetIcon() const = 0;
+
+	virtual FName GetColumnId() const = 0;
+};
+
+////////////////////////////////////////////////////////////////////////////////////////////////////
+
+struct TRACEINSIGHTSCORE_API FTreeNodeGroupInfo
+{
+	FName Name;
+	bool IsExpanded;
+};
+
+////////////////////////////////////////////////////////////////////////////////////////////////////
+
+class TRACEINSIGHTSCORE_API FTreeNodeGrouping : public ITreeNodeGrouping
+{
+	INSIGHTS_DECLARE_RTTI(FTreeNodeGrouping, ITreeNodeGrouping)
+
+public:
+	FTreeNodeGrouping(const FText& InShortName, const FText& InTitleName, const FText& InDescription, const FName InBrushName, const FSlateBrush* InIcon);
+	virtual ~FTreeNodeGrouping() {}
+
+	virtual FText GetShortName() const override { return ShortName; }
+	virtual FText GetTitleName() const override { return TitleName; }
+	virtual FText GetDescription() const override { return Description; }
+
+	virtual FName GetBrushName() const override { return BrushName; }
+	virtual const FSlateBrush* GetIcon() const override { return Icon; }
+
+	virtual FName GetColumnId() const override { return NAME_None; }
+
+	virtual FTreeNodeGroupInfo GetGroupForNode(const FBaseTreeNodePtr InNode) const { return { FName(), false }; }
+	virtual void GroupNodes(const TArray<FTableTreeNodePtr>& Nodes, FTableTreeNode& ParentGroup, TWeakPtr<FTable> InParentTable, IAsyncOperationProgress& InAsyncOperationProgress) const;
+
+protected:
+	FText ShortName;
+	FText TitleName;
+	FText Description;
+	FName BrushName;
+	const FSlateBrush* Icon;
+};
+
+////////////////////////////////////////////////////////////////////////////////////////////////////
+
+/** Creates a single group for all nodes. */
+class TRACEINSIGHTSCORE_API FTreeNodeGroupingFlat : public FTreeNodeGrouping
+{
+	INSIGHTS_DECLARE_RTTI(FTreeNodeGroupingFlat, FTreeNodeGrouping)
+
+public:
+	FTreeNodeGroupingFlat();
+	virtual ~FTreeNodeGroupingFlat() {}
+
+	virtual void GroupNodes(const TArray<FTableTreeNodePtr>& Nodes, FTableTreeNode& ParentGroup, TWeakPtr<FTable> InParentTable, IAsyncOperationProgress& InAsyncOperationProgress) const override;
+};
+
+////////////////////////////////////////////////////////////////////////////////////////////////////
+
+/** Creates a group for each unique value. */
+class TRACEINSIGHTSCORE_API FTreeNodeGroupingByUniqueValue : public FTreeNodeGrouping
+{
+	INSIGHTS_DECLARE_RTTI(FTreeNodeGroupingByUniqueValue, FTreeNodeGrouping);
+
+public:
+	FTreeNodeGroupingByUniqueValue(TSharedRef<FTableColumn> InColumnRef);
+	virtual ~FTreeNodeGroupingByUniqueValue() {}
+
+	virtual FTreeNodeGroupInfo GetGroupForNode(const FBaseTreeNodePtr InNode) const override;
+
+	virtual FName GetColumnId() const override { return ColumnRef->GetId(); }
+	TSharedRef<FTableColumn> GetColumn() const { return ColumnRef; }
+
+private:
+	TSharedRef<FTableColumn> ColumnRef;
+};
+
+////////////////////////////////////////////////////////////////////////////////////////////////////
+
+/** Creates a group for each unique value (assumes data type of cell values is a simple type). */
+template<typename Type>
+class TTreeNodeGroupingByUniqueValue : public FTreeNodeGroupingByUniqueValue
+{
+public:
+	TTreeNodeGroupingByUniqueValue(TSharedRef<FTableColumn> InColumnRef) : FTreeNodeGroupingByUniqueValue(InColumnRef) {}
+	virtual ~TTreeNodeGroupingByUniqueValue() {}
+
+	virtual void GroupNodes(const TArray<FTableTreeNodePtr>& Nodes, FTableTreeNode& ParentGroup, TWeakPtr<FTable> InParentTable, IAsyncOperationProgress& InAsyncOperationProgress) const override;
+
+private:
+	static Type GetValue(const FTableCellValue& CellValue);
+	static FName GetGroupName(const FTableColumn& Column, const FTableTreeNode& Node);
+};
+
+template<typename Type>
+FName TTreeNodeGroupingByUniqueValue<Type>::GetGroupName(const FTableColumn& Column, const FTableTreeNode& Node)
+{
+	FText ValueAsText = Column.GetValueAsGroupingText(Node);
+
+	if (ValueAsText.IsEmpty())
+	{
+		static FName EmptyGroupName(TEXT("N/A"));
+		return EmptyGroupName;
+	}
+
+	FStringView StringView(ValueAsText.ToString());
+	if (StringView.Len() >= NAME_SIZE)
+	{
+		StringView = FStringView(StringView.GetData(), NAME_SIZE - 1);
+	}
+	return FName(StringView, 0);
+}
+
+template<typename Type>
+void TTreeNodeGroupingByUniqueValue<Type>::GroupNodes(const TArray<FTableTreeNodePtr>& Nodes, FTableTreeNode& ParentGroup, TWeakPtr<FTable> InParentTable, IAsyncOperationProgress& InAsyncOperationProgress) const
+{
+	TMap<Type, FTableTreeNodePtr> GroupMap;
+	FTableTreeNodePtr UnsetGroupPtr = nullptr;
+
+	ParentGroup.ClearChildren();
+
+	for (FTableTreeNodePtr NodePtr : Nodes)
+	{
+		if (InAsyncOperationProgress.ShouldCancelAsyncOp())
+		{
+			return;
+		}
+
+		if (NodePtr->IsGroup())
+		{
+			ParentGroup.AddChildAndSetParent(NodePtr);
+			continue;
+		}
+
+		FTableTreeNodePtr GroupPtr = nullptr;
+
+		const FTableColumn& Column = *GetColumn();
+		const TOptional<FTableCellValue> CellValue = Column.GetValue(*NodePtr);
+		if (CellValue.IsSet())
+		{
+			const Type Value = GetValue(CellValue.GetValue());
+
+			FTableTreeNodePtr* GroupPtrPtr = GroupMap.Find(Value);
+			if (!GroupPtrPtr)
+			{
+				const FName GroupName = GetGroupName(Column, *NodePtr);
+				GroupPtr = MakeShared<FTableTreeNode>(GroupName, InParentTable);
+				GroupPtr->SetExpansion(false);
+				ParentGroup.AddChildAndSetParent(GroupPtr);
+				GroupMap.Add(Value, GroupPtr);
+			}
+			else
+			{
+				GroupPtr = *GroupPtrPtr;
+			}
+		}
+		else
+		{
+			if (!UnsetGroupPtr)
+			{
+				UnsetGroupPtr = MakeShared<FTableTreeNode>(FName(TEXT("<unset>")), InParentTable);
+				UnsetGroupPtr->SetExpansion(false);
+				ParentGroup.AddChildAndSetParent(UnsetGroupPtr);
+			}
+			GroupPtr = UnsetGroupPtr;
+		}
+
+		GroupPtr->AddChildAndSetParent(NodePtr);
+	}
+}
+
+typedef TTreeNodeGroupingByUniqueValue<bool> FTreeNodeGroupingByUniqueValueBool;
+typedef TTreeNodeGroupingByUniqueValue<int64> FTreeNodeGroupingByUniqueValueInt64;
+typedef TTreeNodeGroupingByUniqueValue<float> FTreeNodeGroupingByUniqueValueFloat;
+typedef TTreeNodeGroupingByUniqueValue<double> FTreeNodeGroupingByUniqueValueDouble;
+
+////////////////////////////////////////////////////////////////////////////////////////////////////
+
+/** Creates a group for each unique value (assumes data type of cell values is const TCHAR*). */
+class TRACEINSIGHTSCORE_API FTreeNodeGroupingByUniqueValueCString : public FTreeNodeGroupingByUniqueValue
+{
+public:
+	FTreeNodeGroupingByUniqueValueCString(TSharedRef<FTableColumn> InColumnRef) : FTreeNodeGroupingByUniqueValue(InColumnRef) {}
+	virtual ~FTreeNodeGroupingByUniqueValueCString() {}
+
+	virtual void GroupNodes(const TArray<FTableTreeNodePtr>& Nodes, FTableTreeNode& ParentGroup, TWeakPtr<FTable> InParentTable, IAsyncOperationProgress& InAsyncOperationProgress) const override;
+
+private:
+	static FName GetGroupName(const TCHAR* Value);
+};
+
+////////////////////////////////////////////////////////////////////////////////////////////////////
+
+/** Creates a group for each first letter of node names. */
+class TRACEINSIGHTSCORE_API FTreeNodeGroupingByNameFirstLetter : public FTreeNodeGrouping
+{
+	INSIGHTS_DECLARE_RTTI(FTreeNodeGroupingByNameFirstLetter, FTreeNodeGrouping);
+
+public:
+	FTreeNodeGroupingByNameFirstLetter();
+	virtual ~FTreeNodeGroupingByNameFirstLetter() {}
+
+	virtual FTreeNodeGroupInfo GetGroupForNode(const FBaseTreeNodePtr InNode) const override;
+};
+
+////////////////////////////////////////////////////////////////////////////////////////////////////
+
+/** Creates a group for each type. */
+class TRACEINSIGHTSCORE_API FTreeNodeGroupingByType : public FTreeNodeGrouping
+{
+	INSIGHTS_DECLARE_RTTI(FTreeNodeGroupingByType, FTreeNodeGrouping);
+
+public:
+	FTreeNodeGroupingByType();
+	virtual ~FTreeNodeGroupingByType() {}
+
+	virtual FTreeNodeGroupInfo GetGroupForNode(const FBaseTreeNodePtr InNode) const override;
+};
+
+////////////////////////////////////////////////////////////////////////////////////////////////////
+
+/** Creates a tree hierarchy out of the path structure of string values. */
+class TRACEINSIGHTSCORE_API FTreeNodeGroupingByPathBreakdown : public FTreeNodeGrouping
+{
+	INSIGHTS_DECLARE_RTTI(FTreeNodeGroupingByPathBreakdown, FTreeNodeGrouping);
+
+public:
+	FTreeNodeGroupingByPathBreakdown(TSharedRef<FTableColumn> InColumnRef);
+	virtual ~FTreeNodeGroupingByPathBreakdown() {}
+
+	virtual void GroupNodes(const TArray<FTableTreeNodePtr>& Nodes, FTableTreeNode& ParentGroup, TWeakPtr<FTable> InParentTable, IAsyncOperationProgress& InAsyncOperationProgress) const override;
+
+	virtual FName GetColumnId() const override { return ColumnRef->GetId(); }
+	TSharedRef<FTableColumn> GetColumn() const { return ColumnRef; }
+
+private:
+	TSharedRef<FTableColumn> ColumnRef;
+};
+
+////////////////////////////////////////////////////////////////////////////////////////////////////
+
+} // namespace UE::Insights
diff --git a/Engine/Source/Developer/TraceInsightsCore/Public/InsightsCore/Table/Widgets/SAsyncOperationStatus.h b/Engine/Source/Developer/TraceInsightsCore/Public/InsightsCore/Table/Widgets/SAsyncOperationStatus.h
new file mode 100644
--- /dev/null
+++ b/Engine/Source/Developer/TraceInsightsCore/Public/InsightsCore/Table/Widgets/SAsyncOperationStatus.h
@@ -0,0 +1,40 @@
+// Copyright Epic Games, Inc. All Rights Reserved.
+
+#pragma once
+
+#include "CoreMinimal.h"
+
+#include "Widgets/DeclarativeSyntaxSupport.h"
+#include "Widgets/SCompoundWidget.h"
+
+namespace UE::Insights
+{
+
+class IAsyncOperationStatusProvider;
+
+// A widget representing the status message used in table / tree views for background operations.
+class TRACEINSIGHTSCORE_API SAsyncOperationStatus : public SCompoundWidget
+{
+public:
+	SLATE_BEGIN_ARGS(SAsyncOperationStatus) {}
+
+	SLATE_END_ARGS()
+
+	void Construct(const FArguments& InArgs, TSharedRef<IAsyncOperationStatusProvider> InStatusProvider);
+
+private:
+	EVisibility GetContentVisibility() const;
+
+	FSlateColor GetBackgroundColorAndOpacity() const;
+	FSlateColor GetTextColorAndOpacity() const;
+	float ComputeOpacity() const;
+
+	FText GetText() const;
+	FText GetAnimatedText() const;
+	FText GetTooltipText() const;
+
+private:
+	TWeakPtr<IAsyncOperationStatusProvider> StatusProvider;
+};
+
+} // namespace UE::Insights
