[Insights] Fixed tootips for the table columns of per-frame aggregated stats in Timers panel.

diff --git a/Engine/Source/Developer/TraceInsights/Private/Insights/ViewModels/TimersViewColumnFactory.cpp b/Engine/Source/Developer/TraceInsights/Private/Insights/ViewModels/TimersViewColumnFactory.cpp
--- a/Engine/Source/Developer/TraceInsights/Private/Insights/ViewModels/TimersViewColumnFactory.cpp
+++ b/Engine/Source/Developer/TraceInsights/Private/Insights/ViewModels/TimersViewColumnFactory.cpp
@@ -390,9 +390,21 @@ TSharedRef<UE::Insights::FTableColumn> FTimersViewColumnFactory::CreateMaxInclus
 	Column.SetShortName(LOCTEXT("MaxInclusiveTime_ColumnName", "I.Max"));
 	Column.SetTitleName(LOCTEXT("MaxInclusiveTime_ColumnTitle", "Max Inclusive Time (ms)"));
 
-	Column.SetDescription(ETraceFrameType::TraceFrameType_Count, LOCTEXT("MaxInclusiveTime_ColumnDesc", "Maximum inclusive duration of selected timer's instances, in milliseconds"));
-	Column.SetDescription(ETraceFrameType::TraceFrameType_Game, LOCTEXT("MaxInclusiveTime_GameFrameColumnDesc", "Game Frame Maximum Inclusive Duration.\nInclusive duration is computed for a single frame as the sum of inclusive duration of all instances of the timer in the respective frame.\nThe maximum is selected from these per-frame inclusive durations. Unit is miliseconds."));
-	Column.SetDescription(ETraceFrameType::TraceFrameType_Rendering, LOCTEXT("MaxInclusiveTime_RenderingFrameColumnDesc", "Rendering Frame Maximum Inclusive Duration.\nInclusive duration is computed for a single frame as the sum of inclusive duration of all instances of the timer in the respective frame.\nThe maximum is selected from these per-frame inclusive durations. Unit is miliseconds."));
+	Column.SetDescription(ETraceFrameType::TraceFrameType_Count,
+		LOCTEXT("MaxInclusiveTime_ColumnDesc",
+			"Maximum inclusive duration of selected timer's instances, in milliseconds"));
+	Column.SetDescription(ETraceFrameType::TraceFrameType_Game,
+		LOCTEXT("MaxInclusiveTime_GameFrameColumnDesc",
+			"Game Frame Maximum Inclusive Duration\n"
+			"Inclusive duration is computed for a single frame as the sum of inclusive duration of all instances of the timer in the respective frame.\n"
+			"The maximum is selected from these per-frame inclusive durations.\n"
+			"Unit is milliseconds."));
+	Column.SetDescription(ETraceFrameType::TraceFrameType_Rendering,
+		LOCTEXT("MaxInclusiveTime_RenderingFrameColumnDesc",
+			"Rendering Frame Maximum Inclusive Duration\n"
+			"Inclusive duration is computed for a single frame as the sum of inclusive duration of all instances of the timer in the respective frame.\n"
+			"The maximum is selected from these per-frame inclusive durations.\n"
+			"Unit is milliseconds."));
 
 	Column.SetFlags(ETableColumnFlags::CanBeHidden |
 					//ETableColumnFlags::ShouldBeVisible |
@@ -439,9 +451,21 @@ TSharedRef<UE::Insights::FTableColumn> FTimersViewColumnFactory::CreateAverageIn
 	Column.SetShortName(LOCTEXT("AvgInclusiveTime_ColumnName", "I.Avg"));
 	Column.SetTitleName(LOCTEXT("AvgInclusiveTime_ColumnTitle", "Average Inclusive Time (ms)"));
 
-	Column.SetDescription(ETraceFrameType::TraceFrameType_Count, LOCTEXT("AvgInclusiveTime_ColumnDesc", "Average inclusive duration of selected timer's instances, in milliseconds"));
-	Column.SetDescription(ETraceFrameType::TraceFrameType_Game, LOCTEXT("AvgInclusiveTime_GameFrameColumnDesc", "Game Frame Average Inclusive Duration.\nInclusive duration is computed for a single frame as the sum of inclusive duration of all instances of the timer in the respective frame.\nThe average is applied for these per-frame inclusive durations. Unit is miliseconds."));
-	Column.SetDescription(ETraceFrameType::TraceFrameType_Rendering, LOCTEXT("AvgInclusiveTime_RenderingFrameColumnDesc", "Rendering Frame Average Inclusive Duration.\nInclusive duration is computed for a single frame as the sum of inclusive duration of all instances of the timer in the respective frame.\nThe average is applied for these per-frame inclusive durations. Unit is miliseconds."));
+	Column.SetDescription(ETraceFrameType::TraceFrameType_Count,
+		LOCTEXT("AvgInclusiveTime_ColumnDesc",
+			"Average inclusive duration of selected timer's instances, in milliseconds"));
+	Column.SetDescription(ETraceFrameType::TraceFrameType_Game,
+		LOCTEXT("AvgInclusiveTime_GameFrameColumnDesc",
+			"Game Frame Average Inclusive Duration\n"
+			"Inclusive duration is computed for a single frame as the sum of inclusive duration of all instances of the timer in the respective frame.\n"
+			"The average is computed from these per-frame inclusive durations.\n"
+			"Unit is milliseconds."));
+	Column.SetDescription(ETraceFrameType::TraceFrameType_Rendering,
+		LOCTEXT("AvgInclusiveTime_RenderingFrameColumnDesc",
+			"Rendering Frame Average Inclusive Duration\n"
+			"Inclusive duration is computed for a single frame as the sum of inclusive duration of all instances of the timer in the respective frame.\n"
+			"The average is computed from these per-frame inclusive durations.\n"
+			"Unit is milliseconds."));
 
 	Column.SetFlags(ETableColumnFlags::CanBeHidden |
 					//ETableColumnFlags::ShouldBeVisible |
@@ -488,9 +512,21 @@ TSharedRef<UE::Insights::FTableColumn> FTimersViewColumnFactory::CreateMedianInc
 	Column.SetShortName(LOCTEXT("MedInclusiveTime_ColumnName", "I.Med"));
 	Column.SetTitleName(LOCTEXT("MedInclusiveTime_ColumnTitle", "Median Inclusive Time (ms)"));
 
-	Column.SetDescription(ETraceFrameType::TraceFrameType_Count, LOCTEXT("MedInclusiveTime_ColumnDesc", "Median inclusive duration of selected timer's instances, in milliseconds"));
-	Column.SetDescription(ETraceFrameType::TraceFrameType_Game, LOCTEXT("MedInclusiveTime_GameFrameColumnDesc", "Game Frame Median Inclusive Duration.\nInclusive duration is computed for a single frame as the sum of inclusive duration of all instances of the timer in the respective frame.\nThe median is aproximated these per-frame inclusive durations. Unit is miliseconds."));
-	Column.SetDescription(ETraceFrameType::TraceFrameType_Rendering, LOCTEXT("MedInclusiveTime_RenderingFrameColumnDesc", "Rendering Frame Median Inclusive Duration.\nInclusive duration is computed for a single frame as the sum of inclusive duration of all instances of the timer in the respective frame.The median is aproximated these per-frame inclusive durations.\nUnit is miliseconds."));
+	Column.SetDescription(ETraceFrameType::TraceFrameType_Count,
+		LOCTEXT("MedInclusiveTime_ColumnDesc",
+			"Median inclusive duration of selected timer's instances, in milliseconds"));
+	Column.SetDescription(ETraceFrameType::TraceFrameType_Game,
+		LOCTEXT("MedInclusiveTime_GameFrameColumnDesc",
+			"Game Frame Median Inclusive Duration\n"
+			"Inclusive duration is computed for a single frame as the sum of inclusive duration of all instances of the timer in the respective frame.\n"
+			"The median is computed from these per-frame inclusive durations.\n"
+			"Unit is milliseconds."));
+	Column.SetDescription(ETraceFrameType::TraceFrameType_Rendering,
+		LOCTEXT("MedInclusiveTime_RenderingFrameColumnDesc",
+			"Rendering Frame Median Inclusive Duration\n"
+			"Inclusive duration is computed for a single frame as the sum of inclusive duration of all instances of the timer in the respective frame.\n"
+			"The median is computed from these per-frame inclusive durations.\n"
+			"Unit is milliseconds."));
 
 	Column.SetFlags(ETableColumnFlags::CanBeHidden |
 					//ETableColumnFlags::ShouldBeVisible |
@@ -537,9 +573,21 @@ TSharedRef<UE::Insights::FTableColumn> FTimersViewColumnFactory::CreateMinInclus
 	Column.SetShortName(LOCTEXT("MinInclusiveTime_ColumnName", "I.Min"));
 	Column.SetTitleName(LOCTEXT("MinInclusiveTime_ColumnTitle", "Min Inclusive Time (ms)"));
 
-	Column.SetDescription(ETraceFrameType::TraceFrameType_Count, LOCTEXT("MinInclusiveTime_ColumnDesc", "Minimum inclusive duration of selected timer's instances, in milliseconds"));
-	Column.SetDescription(ETraceFrameType::TraceFrameType_Game, LOCTEXT("MinInclusiveTime_GameFrameColumnDesc", "Game Frame Minimum Inclusive Duration.\nInclusive duration is computed for a single frame as the sum of inclusive duration of all instances of the timer in the respective frame.\nThe minimum is selectedfrom these per-frame inclusive durations. Unit is miliseconds."));
-	Column.SetDescription(ETraceFrameType::TraceFrameType_Rendering, LOCTEXT("MinInclusiveTime_RenderingFrameColumnDesc", "Rendering Frame Minimum Inclusive Duration.\nInclusive duration is computed for a single frame as the sum of inclusive duration of all instances of the timer in the respective frame.The minimum is selectedfrom these per-frame inclusive durations.\nUnit is miliseconds."));
+	Column.SetDescription(ETraceFrameType::TraceFrameType_Count,
+		LOCTEXT("MinInclusiveTime_ColumnDesc",
+			"Minimum inclusive duration of selected timer's instances, in milliseconds"));
+	Column.SetDescription(ETraceFrameType::TraceFrameType_Game,
+		LOCTEXT("MinInclusiveTime_GameFrameColumnDesc",
+			"Game Frame Minimum Inclusive Duration\n"
+			"Inclusive duration is computed for a single frame as the sum of inclusive duration of all instances of the timer in the respective frame.\n"
+			"The minimum is selected from these per-frame inclusive durations.\n"
+			"Unit is milliseconds."));
+	Column.SetDescription(ETraceFrameType::TraceFrameType_Rendering,
+		LOCTEXT("MinInclusiveTime_RenderingFrameColumnDesc",
+			"Rendering Frame Minimum Inclusive Duration\n"
+			"Inclusive duration is computed for a single frame as the sum of inclusive duration of all instances of the timer in the respective frame.\n"
+			"The minimum is selected from these per-frame inclusive durations.\n"
+			"Unit is milliseconds."));
 
 	Column.SetFlags(ETableColumnFlags::CanBeHidden |
 					//ETableColumnFlags::ShouldBeVisible |
@@ -640,9 +688,21 @@ TSharedRef<UE::Insights::FTableColumn> FTimersViewColumnFactory::CreateMaxExclus
 	Column.SetShortName(LOCTEXT("MaxExclusiveTime_ColumnName", "E.Max"));
 	Column.SetTitleName(LOCTEXT("MaxExclusiveTime_ColumnTitle", "Max Exclusive Time (ms)"));
 
-	Column.SetDescription(ETraceFrameType::TraceFrameType_Count, LOCTEXT("MaxExclusiveTime_ColumnDesc", "Maximum exclusive duration of selected timer's instances, in milliseconds"));
-	Column.SetDescription(ETraceFrameType::TraceFrameType_Game, LOCTEXT("MaxExclusiveTime_GameFrameColumnDesc", "Game Frame Maximum Exclusive Duration.\nExclusive duration is computed for a single frame as the sum of exclusive duration of all instances of the timer in the respective frame.\nThe maximum is selected from these per-frame Exclusive durations. Unit is miliseconds."));
-	Column.SetDescription(ETraceFrameType::TraceFrameType_Rendering, LOCTEXT("MaxExclusiveTime_RenderingFrameColumnDesc", "Rendering Frame Maximum Exclusive Duration.\nExclusive duration is computed for a single frame as the sum of exclusive duration of all instances of the timer in the respective frame. The maximum is selected from these per-frame Exclusive durations.\nUnit is miliseconds."));
+	Column.SetDescription(ETraceFrameType::TraceFrameType_Count,
+		LOCTEXT("MaxExclusiveTime_ColumnDesc",
+			"Maximum exclusive duration of selected timer's instances, in milliseconds"));
+	Column.SetDescription(ETraceFrameType::TraceFrameType_Game,
+		LOCTEXT("MaxExclusiveTime_GameFrameColumnDesc",
+			"Game Frame Maximum Exclusive Duration\n"
+			"Exclusive duration is computed for a single frame as the sum of exclusive duration of all instances of the timer in the respective frame.\n"
+			"The maximum is selected from these per-frame exclusive durations.\n"
+			"Unit is milliseconds."));
+	Column.SetDescription(ETraceFrameType::TraceFrameType_Rendering,
+		LOCTEXT("MaxExclusiveTime_RenderingFrameColumnDesc",
+			"Rendering Frame Maximum Exclusive Duration\n"
+			"Exclusive duration is computed for a single frame as the sum of exclusive duration of all instances of the timer in the respective frame.\n"
+			"The maximum is selected from these per-frame exclusive durations.\n"
+			"Unit is milliseconds."));
 
 	Column.SetFlags(ETableColumnFlags::CanBeHidden |
 					//ETableColumnFlags::ShouldBeVisible |
@@ -689,9 +749,21 @@ TSharedRef<UE::Insights::FTableColumn> FTimersViewColumnFactory::CreateAverageEx
 	Column.SetShortName(LOCTEXT("AvgExclusiveTime_ColumnName", "E.Avg"));
 	Column.SetTitleName(LOCTEXT("AvgExclusiveTime_ColumnTitle", "Average Exclusive Time (ms)"));
 
-	Column.SetDescription(ETraceFrameType::TraceFrameType_Count, LOCTEXT("AvgExclusiveTime_ColumnDesc", "Average exclusive duration of selected timer's instances, in milliseconds"));
-	Column.SetDescription(ETraceFrameType::TraceFrameType_Game, LOCTEXT("AvgExclusiveTime_GameFrameColumnDesc", "Game Frame Average Exclusive Duration.\nExclusive duration is computed for a single frame as the sum of exclusive duration of all instances of the timer in the respective frame.\nThe average is applied for these per-frame Exclusive durations. Unit is miliseconds."));
-	Column.SetDescription(ETraceFrameType::TraceFrameType_Rendering, LOCTEXT("AvgExclusiveTime_RenderingFrameColumnDesc", "Rendering Frame Average Exclusive Duration.\nExclusive duration is computed for a single frame as the sum of exclusive duration of all instances of the timer in the respective frame. The average is applied for these per-frame Exclusive durations.\nUnit is miliseconds."));
+	Column.SetDescription(ETraceFrameType::TraceFrameType_Count,
+		LOCTEXT("AvgExclusiveTime_ColumnDesc",
+			"Average exclusive duration of selected timer's instances, in milliseconds"));
+	Column.SetDescription(ETraceFrameType::TraceFrameType_Game,
+		LOCTEXT("AvgExclusiveTime_GameFrameColumnDesc",
+			"Game Frame Average Exclusive Duration\n"
+			"Exclusive duration is computed for a single frame as the sum of exclusive duration of all instances of the timer in the respective frame.\n"
+			"The average is computed from these per-frame exclusive durations.\n"
+			"Unit is milliseconds."));
+	Column.SetDescription(ETraceFrameType::TraceFrameType_Rendering,
+		LOCTEXT("AvgExclusiveTime_RenderingFrameColumnDesc",
+			"Rendering Frame Average Exclusive Duration\n"
+			"Exclusive duration is computed for a single frame as the sum of exclusive duration of all instances of the timer in the respective frame.\n"
+			"The average is computed from these per-frame exclusive durations.\n"
+			"Unit is milliseconds."));
 
 	Column.SetFlags(ETableColumnFlags::CanBeHidden |
 					//ETableColumnFlags::ShouldBeVisible |
@@ -738,9 +810,21 @@ TSharedRef<UE::Insights::FTableColumn> FTimersViewColumnFactory::CreateMedianExc
 	Column.SetShortName(LOCTEXT("MedExclusiveTime_ColumnName", "E.Med"));
 	Column.SetTitleName(LOCTEXT("MedExclusiveTime_ColumnTitle", "Median Exclusive Time (ms)"));
 
-	Column.SetDescription(ETraceFrameType::TraceFrameType_Count, LOCTEXT("MedExclusiveTime_ColumnDesc", "Median exclusive duration of selected timer's instances, in milliseconds"));
-	Column.SetDescription(ETraceFrameType::TraceFrameType_Game, LOCTEXT("MedExclusiveTime_GameFrameColumnDesc", "Game Frame Median Exclusive Duration.\nExclusive duration is computed for a single frame as the sum of exclusive duration of all instances of the timer in the respective frame.\nThe median is aproximated these per-frame Exclusive durations. Unit is miliseconds."));
-	Column.SetDescription(ETraceFrameType::TraceFrameType_Rendering, LOCTEXT("MedExclusiveTime_RenderingFrameColumnDesc", "Rendering Frame Median Exclusive Duration.\nExclusive duration is computed for a single frame as the sum of exclusive duration of all instances of the timer in the respective frame.The median is aproximated these per-frame Exclusive durations.\nUnit is miliseconds."));
+	Column.SetDescription(ETraceFrameType::TraceFrameType_Count,
+		LOCTEXT("MedExclusiveTime_ColumnDesc",
+			"Median exclusive duration of selected timer's instances, in milliseconds"));
+	Column.SetDescription(ETraceFrameType::TraceFrameType_Game,
+		LOCTEXT("MedExclusiveTime_GameFrameColumnDesc",
+			"Game Frame Median Exclusive Duration\n"
+			"Exclusive duration is computed for a single frame as the sum of exclusive duration of all instances of the timer in the respective frame.\n"
+			"The median is computed from these per-frame exclusive durations.\n"
+			"Unit is milliseconds."));
+	Column.SetDescription(ETraceFrameType::TraceFrameType_Rendering,
+		LOCTEXT("MedExclusiveTime_RenderingFrameColumnDesc",
+			"Rendering Frame Median Exclusive Duration\n"
+			"Exclusive duration is computed for a single frame as the sum of exclusive duration of all instances of the timer in the respective frame.\n"
+			"The median is computed from these per-frame exclusive durations.\n"
+			"Unit is milliseconds."));
 
 	Column.SetFlags(ETableColumnFlags::CanBeHidden |
 					//ETableColumnFlags::ShouldBeVisible |
@@ -787,9 +871,21 @@ TSharedRef<UE::Insights::FTableColumn> FTimersViewColumnFactory::CreateMinExclus
 	Column.SetShortName(LOCTEXT("MinExclusiveTime_ColumnName", "E.Min"));
 	Column.SetTitleName(LOCTEXT("MinExclusiveTime_ColumnTitle", "Min Exclusive Time (ms)"));
 
-	Column.SetDescription(ETraceFrameType::TraceFrameType_Count, LOCTEXT("MinExclusiveTime_ColumnDesc", "Minimum exclusive duration of selected timer's instances, in milliseconds"));
-	Column.SetDescription(ETraceFrameType::TraceFrameType_Game, LOCTEXT("MinExclusiveTime_GameFrameColumnDesc", "Game Frame Minimum Exclusive Duration.\nExclusive duration is computed for a single frame as the sum of exclusive duration of all instances of the timer in the respective frame.\nThe minimum is selected from these per-frame Exclusive durations. Unit is miliseconds."));
-	Column.SetDescription(ETraceFrameType::TraceFrameType_Rendering, LOCTEXT("MinExclusiveTime_RenderingFrameColumnDesc", "Rendering Frame Minimum Exclusive Duration.\nExclusive duration is computed for a single frame as the sum of exclusive duration of all instances of the timer in the respective frame.The minimum is selectedfrom these per-frame Exclusive durations.\nUnit is miliseconds."));
+	Column.SetDescription(ETraceFrameType::TraceFrameType_Count,
+		LOCTEXT("MinExclusiveTime_ColumnDesc",
+			"Minimum exclusive duration of selected timer's instances, in milliseconds"));
+	Column.SetDescription(ETraceFrameType::TraceFrameType_Game,
+		LOCTEXT("MinExclusiveTime_GameFrameColumnDesc",
+			"Game Frame Minimum Exclusive Duration\n"
+			"Exclusive duration is computed for a single frame as the sum of exclusive duration of all instances of the timer in the respective frame.\n"
+			"The minimum is selected from these per-frame exclusive durations.\n"
+			"Unit is milliseconds."));
+	Column.SetDescription(ETraceFrameType::TraceFrameType_Rendering,
+		LOCTEXT("MinExclusiveTime_RenderingFrameColumnDesc",
+			"Rendering Frame Minimum Exclusive Duration\n"
+			"Exclusive duration is computed for a single frame as the sum of exclusive duration of all instances of the timer in the respective frame.\n"
+			"The minimum is selected from these per-frame exclusive durations.\n"
+			"Unit is milliseconds."));
 
 	Column.SetFlags(ETableColumnFlags::CanBeHidden |
 					//ETableColumnFlags::ShouldBeVisible |
