[Insights] Moved implementations of ContextSwitches, CookProfiler, LoadingProfiler and TaskGraphProfiler into corresponding UE::Insights::* namespaces.

#rnx

diff --git a/Engine/Source/Developer/TraceInsights/Private/Insights/ContextSwitches/ContextSwitchesProfilerManager.cpp b/Engine/Source/Developer/TraceInsights/Private/Insights/ContextSwitches/ContextSwitchesProfilerManager.cpp
--- a/Engine/Source/Developer/TraceInsights/Private/Insights/ContextSwitches/ContextSwitchesProfilerManager.cpp
+++ b/Engine/Source/Developer/TraceInsights/Private/Insights/ContextSwitches/ContextSwitchesProfilerManager.cpp
@@ -18,9 +18,7 @@
 
 ////////////////////////////////////////////////////////////////////////////////////////////////////
 
-#define LOCTEXT_NAMESPACE "ContextSwitchesProfilerManager"
-
-namespace Insights
+namespace UE::Insights::ContextSwitches
 {
 
 ////////////////////////////////////////////////////////////////////////////////////////////////////
@@ -108,7 +106,7 @@ FContextSwitchesProfilerManager::~FContextSwitchesProfilerManager()
 
 	if (ContextSwitchesSharedState.IsValid())
 	{
-		IModularFeatures::Get().UnregisterModularFeature(UE::Insights::Timing::TimingViewExtenderFeatureName, ContextSwitchesSharedState.Get());
+		IModularFeatures::Get().UnregisterModularFeature(Timing::TimingViewExtenderFeatureName, ContextSwitchesSharedState.Get());
 	}
 }
 
@@ -161,7 +159,7 @@ bool FContextSwitchesProfilerManager::Tick(float DeltaTime)
 				{
 					ContextSwitchesSharedState = MakeShared<FContextSwitchesSharedState>(TimingView.Get());
 					ContextSwitchesSharedState->AddCommands();
-					IModularFeatures::Get().RegisterModularFeature(UE::Insights::Timing::TimingViewExtenderFeatureName, ContextSwitchesSharedState.Get());
+					IModularFeatures::Get().RegisterModularFeature(Timing::TimingViewExtenderFeatureName, ContextSwitchesSharedState.Get());
 				}
 			}
 
@@ -204,6 +202,4 @@ void FContextSwitchesProfilerManager::RegisterTimingProfilerLayoutExtensions(FIn
 
 ////////////////////////////////////////////////////////////////////////////////////////////////////
 
-} // namespace Insights
-
-#undef LOCTEXT_NAMESPACE
+} // namespace UE::Insights::ContextSwitches
diff --git a/Engine/Source/Developer/TraceInsights/Private/Insights/ContextSwitches/ContextSwitchesProfilerManager.h b/Engine/Source/Developer/TraceInsights/Private/Insights/ContextSwitches/ContextSwitchesProfilerManager.h
--- a/Engine/Source/Developer/TraceInsights/Private/Insights/ContextSwitches/ContextSwitchesProfilerManager.h
+++ b/Engine/Source/Developer/TraceInsights/Private/Insights/ContextSwitches/ContextSwitchesProfilerManager.h
@@ -15,7 +15,7 @@
 class FThreadTrackEvent;
 class FUICommandList;
 
-namespace Insights
+namespace UE::Insights::ContextSwitches
 {
 
 class FContextSwitchesSharedState;
@@ -63,7 +63,7 @@ private:
 private:
 	bool bIsInitialized;
 	bool bIsAvailable;
-	UE::Insights::FAvailabilityCheck AvailabilityCheck;
+	FAvailabilityCheck AvailabilityCheck;
 
 	/** The delegate to be invoked when this manager ticks. */
 	FTickerDelegate OnTick;
@@ -80,4 +80,4 @@ private:
 
 ////////////////////////////////////////////////////////////////////////////////////////////////////
 
-} // namespace Insights
+} // namespace UE::Insights::ContextSwitches
diff --git a/Engine/Source/Developer/TraceInsights/Private/Insights/ContextSwitches/ViewModels/ContextSwitchTimingEvent.cpp b/Engine/Source/Developer/TraceInsights/Private/Insights/ContextSwitches/ViewModels/ContextSwitchTimingEvent.cpp
--- a/Engine/Source/Developer/TraceInsights/Private/Insights/ContextSwitches/ViewModels/ContextSwitchTimingEvent.cpp
+++ b/Engine/Source/Developer/TraceInsights/Private/Insights/ContextSwitches/ViewModels/ContextSwitchTimingEvent.cpp
@@ -2,6 +2,9 @@
 
 #include "ContextSwitchTimingEvent.h"
 
+namespace UE::Insights::ContextSwitches
+{
+
 ////////////////////////////////////////////////////////////////////////////////////////////////////
 // FContextSwitchTimingEvent
 ////////////////////////////////////////////////////////////////////////////////////////////////////
@@ -28,3 +31,5 @@ FCpuCoreTimingEvent::FCpuCoreTimingEvent(const TSharedRef<const FBaseTimingTrack
 
 ////////////////////////////////////////////////////////////////////////////////////////////////////
 ////////////////////////////////////////////////////////////////////////////////////////////////////
+
+} // namespace UE::Insights::ContextSwitches
diff --git a/Engine/Source/Developer/TraceInsights/Private/Insights/ContextSwitches/ViewModels/ContextSwitchTimingEvent.h b/Engine/Source/Developer/TraceInsights/Private/Insights/ContextSwitches/ViewModels/ContextSwitchTimingEvent.h
--- a/Engine/Source/Developer/TraceInsights/Private/Insights/ContextSwitches/ViewModels/ContextSwitchTimingEvent.h
+++ b/Engine/Source/Developer/TraceInsights/Private/Insights/ContextSwitches/ViewModels/ContextSwitchTimingEvent.h
@@ -2,11 +2,16 @@
 
 #pragma once
 
-#include "CoreMinimal.h"
+#include "CoreTypes.h"
+
+#include "Templates/SharedPointer.h"
 
 // TraceInsights
 #include "Insights/ViewModels/TimingEvent.h"
 
+namespace UE::Insights::ContextSwitches
+{
+
 ////////////////////////////////////////////////////////////////////////////////////////////////////
 
 class FContextSwitchTimingEvent : public FTimingEvent
@@ -36,3 +41,5 @@ public:
 };
 
 ////////////////////////////////////////////////////////////////////////////////////////////////////
+
+} // namespace UE::Insights::ContextSwitches
diff --git a/Engine/Source/Developer/TraceInsights/Private/Insights/ContextSwitches/ViewModels/ContextSwitchesFilterConverters.cpp b/Engine/Source/Developer/TraceInsights/Private/Insights/ContextSwitches/ViewModels/ContextSwitchesFilterConverters.cpp
--- a/Engine/Source/Developer/TraceInsights/Private/Insights/ContextSwitches/ViewModels/ContextSwitchesFilterConverters.cpp
+++ b/Engine/Source/Developer/TraceInsights/Private/Insights/ContextSwitches/ViewModels/ContextSwitchesFilterConverters.cpp
@@ -9,9 +9,9 @@
 // TraceInsights
 #include "Insights/InsightsManager.h"
 
-#define LOCTEXT_NAMESPACE "Insights::ContextSwitchesFilterValueConverter"
+#define LOCTEXT_NAMESPACE "UE::Insights::ContextSwitches"
 
-namespace Insights
+namespace UE::Insights::ContextSwitches
 {
 
 ////////////////////////////////////////////////////////////////////////////////////////////////////
@@ -68,6 +68,6 @@ FText FCoreEventNameFilterValueConverter::GetHintText() const
 
 ////////////////////////////////////////////////////////////////////////////////////////////////////
 
-} // namespace Insights
+} // namespace UE::Insights::ContextSwitches
 
-#undef LOCTEXT_NAMESPACE
\ No newline at end of file
+#undef LOCTEXT_NAMESPACE
diff --git a/Engine/Source/Developer/TraceInsights/Private/Insights/ContextSwitches/ViewModels/ContextSwitchesFilterConverters.h b/Engine/Source/Developer/TraceInsights/Private/Insights/ContextSwitches/ViewModels/ContextSwitchesFilterConverters.h
--- a/Engine/Source/Developer/TraceInsights/Private/Insights/ContextSwitches/ViewModels/ContextSwitchesFilterConverters.h
+++ b/Engine/Source/Developer/TraceInsights/Private/Insights/ContextSwitches/ViewModels/ContextSwitchesFilterConverters.h
@@ -2,15 +2,19 @@
 
 #pragma once
 
+#include "CoreTypes.h"
+
+#include "Internationalization/Text.h"
+
 // TraceInsightsCore
 #include "InsightsCore/Filter/ViewModels/Filters.h"
 
-namespace Insights
+namespace UE::Insights::ContextSwitches
 {
 
 ////////////////////////////////////////////////////////////////////////////////////////////////////
 
-class FCoreEventNameFilterValueConverter : public UE::Insights::IFilterValueConverter
+class FCoreEventNameFilterValueConverter : public IFilterValueConverter
 {
 public:
 	virtual bool Convert(const FString& Input, int64& Output, FText& OutError) const override;
@@ -20,4 +24,4 @@ public:
 
 ////////////////////////////////////////////////////////////////////////////////////////////////////
 
-} // namespace Insights
\ No newline at end of file
+} // namespace UE::Insights::ContextSwitches
diff --git a/Engine/Source/Developer/TraceInsights/Private/Insights/ContextSwitches/ViewModels/ContextSwitchesSharedState.cpp b/Engine/Source/Developer/TraceInsights/Private/Insights/ContextSwitches/ViewModels/ContextSwitchesSharedState.cpp
--- a/Engine/Source/Developer/TraceInsights/Private/Insights/ContextSwitches/ViewModels/ContextSwitchesSharedState.cpp
+++ b/Engine/Source/Developer/TraceInsights/Private/Insights/ContextSwitches/ViewModels/ContextSwitchesSharedState.cpp
@@ -27,9 +27,9 @@
 #include "Insights/Widgets/STimingProfilerWindow.h"
 #include "Insights/Widgets/STimingView.h"
 
-#define LOCTEXT_NAMESPACE "ContextSwitches"
+#define LOCTEXT_NAMESPACE "UE::Insights::ContextSwitches"
 
-namespace Insights
+namespace UE::Insights::ContextSwitches
 {
 
 ////////////////////////////////////////////////////////////////////////////////////////////////////
@@ -117,7 +117,7 @@ UE_ENABLE_OPTIMIZATION_SHIP
 // FContextSwitchesSharedState
 ////////////////////////////////////////////////////////////////////////////////////////////////////
 
-FContextSwitchesSharedState::FContextSwitchesSharedState(UE::Insights::TimingProfiler::STimingView* InTimingView)
+FContextSwitchesSharedState::FContextSwitchesSharedState(TimingProfiler::STimingView* InTimingView)
 	: TimingViewSession(InTimingView)
 	, ThreadsSerial(0)
 	, CpuCoresSerial(0)
@@ -132,7 +132,7 @@ FContextSwitchesSharedState::FContextSwitchesSharedState(UE::Insights::TimingPro
 
 ////////////////////////////////////////////////////////////////////////////////////////////////////
 
-void FContextSwitchesSharedState::OnBeginSession(UE::Insights::Timing::ITimingViewSession& InSession)
+void FContextSwitchesSharedState::OnBeginSession(Timing::ITimingViewSession& InSession)
 {
 	if (&InSession != TimingViewSession)
 	{
@@ -166,7 +166,7 @@ void FContextSwitchesSharedState::OnBeginSession(UE::Insights::Timing::ITimingVi
 
 ////////////////////////////////////////////////////////////////////////////////////////////////////
 
-void FContextSwitchesSharedState::OnEndSession(UE::Insights::Timing::ITimingViewSession& InSession)
+void FContextSwitchesSharedState::OnEndSession(Timing::ITimingViewSession& InSession)
 {
 	if (&InSession != TimingViewSession)
 	{
@@ -192,7 +192,7 @@ void FContextSwitchesSharedState::OnEndSession(UE::Insights::Timing::ITimingView
 
 ////////////////////////////////////////////////////////////////////////////////////////////////////
 
-void FContextSwitchesSharedState::Tick(UE::Insights::Timing::ITimingViewSession& InSession, const TraceServices::IAnalysisSession& InAnalysisSession)
+void FContextSwitchesSharedState::Tick(Timing::ITimingViewSession& InSession, const TraceServices::IAnalysisSession& InAnalysisSession)
 {
 	if (&InSession != TimingViewSession)
 	{
@@ -221,7 +221,7 @@ void FContextSwitchesSharedState::Tick(UE::Insights::Timing::ITimingViewSession&
 
 ////////////////////////////////////////////////////////////////////////////////////////////////////
 
-void FContextSwitchesSharedState::ExtendCpuTracksFilterMenu(UE::Insights::Timing::ITimingViewSession& InSession, FMenuBuilder& InMenuBuilder)
+void FContextSwitchesSharedState::ExtendCpuTracksFilterMenu(Timing::ITimingViewSession& InSession, FMenuBuilder& InMenuBuilder)
 {
 	if (&InSession != TimingViewSession)
 	{
@@ -260,7 +260,7 @@ void FContextSwitchesSharedState::BuildSubMenu(FMenuBuilder& InMenuBuilder)
 
 ////////////////////////////////////////////////////////////////////////////////////////////////////
 
-void FContextSwitchesSharedState::AddQuickFindFilters(TSharedPtr<UE::Insights::FFilterConfigurator> FilterConfigurator)
+void FContextSwitchesSharedState::AddQuickFindFilters(TSharedPtr<FFilterConfigurator> FilterConfigurator)
 {
 	using namespace UE::Insights;
 
@@ -308,7 +308,7 @@ void FContextSwitchesSharedState::AddCommands()
 {
 	FContextSwitchesStateCommands::Register();
 
-	TSharedPtr<UE::Insights::TimingProfiler::STimingView> TimingView = GetTimingView();
+	TSharedPtr<TimingProfiler::STimingView> TimingView = GetTimingView();
 	
 	if (!TimingView.IsValid())
 	{
@@ -379,7 +379,7 @@ bool FContextSwitchesSharedState::AreContextSwitchesAvailable() const
 
 void FContextSwitchesSharedState::AddCoreTracks()
 {
-	TSharedPtr<UE::Insights::TimingProfiler::STimingView> TimingView = GetTimingView();
+	TSharedPtr<TimingProfiler::STimingView> TimingView = GetTimingView();
 
 	if (!TimingView.IsValid())
 	{
@@ -429,7 +429,7 @@ void FContextSwitchesSharedState::AddCoreTracks()
 
 void FContextSwitchesSharedState::RemoveCoreTracks()
 {
-	TSharedPtr<UE::Insights::TimingProfiler::STimingView> TimingView = GetTimingView();
+	TSharedPtr<TimingProfiler::STimingView> TimingView = GetTimingView();
 	CpuCoresSerial = 0;
 
 	if (TimingView)
@@ -474,7 +474,7 @@ void FContextSwitchesSharedState::AddContextSwitchesChildTracks()
 	//TODO: Create "Cpu Thread" timing tracks also for threads without cpu timing events (i.e. only with context switch events).
 #endif
 
-	TSharedPtr<UE::Insights::TimingProfiler::STimingView> TimingView = GetTimingView();
+	TSharedPtr<TimingProfiler::STimingView> TimingView = GetTimingView();
 	if (!TimingView.IsValid())
 	{
 		return;
@@ -507,7 +507,7 @@ void FContextSwitchesSharedState::RemoveContextSwitchesChildTracks()
 {
 	ThreadsSerial = 0;
 	
-	TSharedPtr<UE::Insights::TimingProfiler::STimingView> TimingView = GetTimingView();
+	TSharedPtr<TimingProfiler::STimingView> TimingView = GetTimingView();
 	if (!TimingView.IsValid())
 	{
 		return;
@@ -589,7 +589,7 @@ void FContextSwitchesSharedState::SetNonTargetProcessEventsVisible(bool bOnOff)
 {
 	if (bAreNonTargetProcessEventsVisible != bOnOff)
 	{
-		TSharedPtr<UE::Insights::TimingProfiler::STimingView> TimingView = GetTimingView();
+		TSharedPtr<TimingProfiler::STimingView> TimingView = GetTimingView();
 		const TSharedPtr<ITimingEventFilter> EventFilter = TimingView->GetEventFilter();
 		for (auto Track : CpuCoreTimingTracks)
 		{
@@ -611,7 +611,7 @@ void FContextSwitchesSharedState::SetNonTargetProcessEventsVisible(bool bOnOff)
 
 bool FContextSwitchesSharedState::IsValidCpuCoreEventSelected() const
 {
-	TSharedPtr<UE::Insights::TimingProfiler::STimingView> TimingView = GetTimingView();
+	TSharedPtr<TimingProfiler::STimingView> TimingView = GetTimingView();
 	if (TimingView == nullptr)
 	{
 		return false;
@@ -627,7 +627,7 @@ bool FContextSwitchesSharedState::IsValidCpuCoreEventSelected() const
 
 bool FContextSwitchesSharedState::IsValidContextSwitchEventSelected() const
 {
-	TSharedPtr<UE::Insights::TimingProfiler::STimingView> TimingView = GetTimingView();
+	TSharedPtr<TimingProfiler::STimingView> TimingView = GetTimingView();
 	if (TimingView == nullptr)
 	{
 		return false;
@@ -643,7 +643,7 @@ bool FContextSwitchesSharedState::IsValidContextSwitchEventSelected() const
 
 void FContextSwitchesSharedState::Command_NavigateToCpuThreadEvent_Execute()
 {
-	TSharedPtr<UE::Insights::TimingProfiler::STimingView> TimingView = GetTimingView();
+	TSharedPtr<TimingProfiler::STimingView> TimingView = GetTimingView();
 
 	if (TimingView && ensure(IsValidCpuCoreEventSelected()))
 	{
@@ -695,7 +695,7 @@ bool FContextSwitchesSharedState::Command_NavigateToCpuThreadEvent_CanExecute()
 		return false;
 	}
 
-	TSharedPtr<UE::Insights::TimingProfiler::STimingView> TimingView = GetTimingView();
+	TSharedPtr<TimingProfiler::STimingView> TimingView = GetTimingView();
 
 	if (TimingView && IsValidCpuCoreEventSelected())
 	{
@@ -724,7 +724,7 @@ bool FContextSwitchesSharedState::Command_NavigateToCpuThreadEvent_CanExecute()
 
 void FContextSwitchesSharedState::Command_DockCpuThreadTrackToBottom_Execute()
 {
-	TSharedPtr<UE::Insights::TimingProfiler::STimingView> TimingView = GetTimingView();
+	TSharedPtr<TimingProfiler::STimingView> TimingView = GetTimingView();
 
 	if (TimingView && ensure(IsValidCpuCoreEventSelected()))
 	{
@@ -761,7 +761,7 @@ void FContextSwitchesSharedState::Command_DockCpuThreadTrackToBottom_Execute()
 
 void FContextSwitchesSharedState::Command_NavigateToCpuCoreEvent_Execute()
 {
-	TSharedPtr<UE::Insights::TimingProfiler::STimingView> TimingView = GetTimingView();
+	TSharedPtr<TimingProfiler::STimingView> TimingView = GetTimingView();
 
 	if (TimingView && ensure(IsValidContextSwitchEventSelected()))
 	{
@@ -797,7 +797,7 @@ void FContextSwitchesSharedState::Command_NavigateToCpuCoreEvent_Execute()
 
 void FContextSwitchesSharedState::Command_DockCpuCoreTrackToTop_Execute()
 {
-	TSharedPtr<UE::Insights::TimingProfiler::STimingView> TimingView = GetTimingView();
+	TSharedPtr<TimingProfiler::STimingView> TimingView = GetTimingView();
 
 	if (TimingView && ensure(IsValidContextSwitchEventSelected()))
 	{
@@ -858,7 +858,7 @@ void FContextSwitchesSharedState::GetThreadInfo(uint32 InSystemThreadId, uint32&
 TSharedPtr<FThreadTimingTrack> FContextSwitchesSharedState::GetThreadTimingTrack(uint32 ThreadId) const
 {
 	TSharedPtr<FThreadTimingTrack> FoundTrack;
-	TSharedPtr<UE::Insights::TimingProfiler::STimingView> TimingView = GetTimingView();
+	TSharedPtr<TimingProfiler::STimingView> TimingView = GetTimingView();
 
 	if (TimingView)
 	{
@@ -882,7 +882,7 @@ TSharedPtr<FThreadTimingTrack> FContextSwitchesSharedState::GetThreadTimingTrack
 TSharedPtr<FCpuCoreTimingTrack> FContextSwitchesSharedState::GetCpuCoreTimingTrack(uint32 CoreNumber) const
 {
 	TSharedPtr<FCpuCoreTimingTrack> FoundTrack;
-	TSharedPtr<UE::Insights::TimingProfiler::STimingView> TimingView = GetTimingView();
+	TSharedPtr<TimingProfiler::STimingView> TimingView = GetTimingView();
 
 	if (TimingView)
 	{
@@ -903,7 +903,7 @@ TSharedPtr<FCpuCoreTimingTrack> FContextSwitchesSharedState::GetCpuCoreTimingTra
 
 ////////////////////////////////////////////////////////////////////////////////////////////////////
 
-TSharedPtr<UE::Insights::TimingProfiler::STimingView> FContextSwitchesSharedState::GetTimingView()
+TSharedPtr<TimingProfiler::STimingView> FContextSwitchesSharedState::GetTimingView()
 {
 	using namespace UE::Insights::TimingProfiler;
 	TSharedPtr<STimingProfilerWindow> Window = FTimingProfilerManager::Get()->GetProfilerWindow();
@@ -912,6 +912,6 @@ TSharedPtr<UE::Insights::TimingProfiler::STimingView> FContextSwitchesSharedStat
 
 ////////////////////////////////////////////////////////////////////////////////////////////////////
 
-} // namespace Insights
+} // namespace UE::Insights::ContextSwitches
 
 #undef LOCTEXT_NAMESPACE
diff --git a/Engine/Source/Developer/TraceInsights/Private/Insights/ContextSwitches/ViewModels/ContextSwitchesSharedState.h b/Engine/Source/Developer/TraceInsights/Private/Insights/ContextSwitches/ViewModels/ContextSwitchesSharedState.h
--- a/Engine/Source/Developer/TraceInsights/Private/Insights/ContextSwitches/ViewModels/ContextSwitchesSharedState.h
+++ b/Engine/Source/Developer/TraceInsights/Private/Insights/ContextSwitches/ViewModels/ContextSwitchesSharedState.h
@@ -11,8 +11,9 @@
 
 ////////////////////////////////////////////////////////////////////////////////////////////////////
 
-class FBaseTimingTrack;
 class FMenuBuilder;
+
+class FBaseTimingTrack;
 class FThreadTimingTrack;
 class ITimingEvent;
 
@@ -20,7 +21,7 @@ namespace TraceServices { class IAnalysisSession; }
 namespace UE::Insights::Timing { class ITimingViewSession; }
 namespace UE::Insights::TimingProfiler { class STimingView; }
 
-namespace Insights
+namespace UE::Insights::ContextSwitches
 {
 
 class FCpuCoreTimingTrack;
@@ -42,35 +43,35 @@ public:
 	TSharedPtr<FUICommandInfo> Command_ShowExtendedLines;
 	TSharedPtr<FUICommandInfo> Command_ShowNonTargetProcessEvents;
 
-	// Commands for a Cpu Core track (context menu).
+	// Commands for a CPU Core track (context menu).
 	TSharedPtr<FUICommandInfo> Command_NavigateToCpuThreadEvent;
 	TSharedPtr<FUICommandInfo> Command_DockCpuThreadTrackToBottom;
 
-	// Commands for a Cpu Thread track (context menu).
+	// Commands for a CPU Thread track (context menu).
 	TSharedPtr<FUICommandInfo> Command_NavigateToCpuCoreEvent;
 	TSharedPtr<FUICommandInfo> Command_DockCpuCoreTrackToTop;
 };
 
 ////////////////////////////////////////////////////////////////////////////////////////////////////
 
-class FContextSwitchesSharedState : public UE::Insights::Timing::ITimingViewExtender, public TSharedFromThis<FContextSwitchesSharedState>
+class FContextSwitchesSharedState : public Timing::ITimingViewExtender, public TSharedFromThis<FContextSwitchesSharedState>
 {
 public:
-	FContextSwitchesSharedState(UE::Insights::TimingProfiler::STimingView* InTimingView);
+	FContextSwitchesSharedState(TimingProfiler::STimingView* InTimingView);
 	virtual ~FContextSwitchesSharedState() = default;
 
 	//////////////////////////////////////////////////
 	// ITimingViewExtender interface
 
-	virtual void OnBeginSession(UE::Insights::Timing::ITimingViewSession& InSession) override;
-	virtual void OnEndSession(UE::Insights::Timing::ITimingViewSession& InSession) override;
-	virtual void Tick(UE::Insights::Timing::ITimingViewSession& InSession, const TraceServices::IAnalysisSession& InAnalysisSession) override;
-	virtual void ExtendCpuTracksFilterMenu(UE::Insights::Timing::ITimingViewSession& InSession, FMenuBuilder& InMenuBuilder) override;
-	virtual void AddQuickFindFilters(TSharedPtr<UE::Insights::FFilterConfigurator> FilterConfigurator) override;
+	virtual void OnBeginSession(Timing::ITimingViewSession& InSession) override;
+	virtual void OnEndSession(Timing::ITimingViewSession& InSession) override;
+	virtual void Tick(Timing::ITimingViewSession& InSession, const TraceServices::IAnalysisSession& InAnalysisSession) override;
+	virtual void ExtendCpuTracksFilterMenu(Timing::ITimingViewSession& InSession, FMenuBuilder& InMenuBuilder) override;
+	virtual void AddQuickFindFilters(TSharedPtr<FFilterConfigurator> FilterConfigurator) override;
 
 	//////////////////////////////////////////////////
 
-	static TSharedPtr<UE::Insights::TimingProfiler::STimingView> GetTimingView();
+	static TSharedPtr<TimingProfiler::STimingView> GetTimingView();
 
 	bool AreContextSwitchesAvailable() const;
 
@@ -165,7 +166,7 @@ private:
 	TMap<uint32, TSharedPtr<FCpuCoreTimingTrack>> CpuCoreTimingTracks;
 	TMap<uint32, TSharedPtr<FContextSwitchesTimingTrack>> ContextSwitchesTimingTracks;
 
-	UE::Insights::Timing::ITimingViewSession* TimingViewSession;
+	Timing::ITimingViewSession* TimingViewSession;
 
 	uint64 ThreadsSerial;
 	uint64 CpuCoresSerial;
@@ -183,4 +184,4 @@ private:
 
 ////////////////////////////////////////////////////////////////////////////////////////////////////
 
-} // namespace Insights
+} // namespace UE::Insights::ContextSwitches
diff --git a/Engine/Source/Developer/TraceInsights/Private/Insights/ContextSwitches/ViewModels/ContextSwitchesTimingTrack.cpp b/Engine/Source/Developer/TraceInsights/Private/Insights/ContextSwitches/ViewModels/ContextSwitchesTimingTrack.cpp
--- a/Engine/Source/Developer/TraceInsights/Private/Insights/ContextSwitches/ViewModels/ContextSwitchesTimingTrack.cpp
+++ b/Engine/Source/Developer/TraceInsights/Private/Insights/ContextSwitches/ViewModels/ContextSwitchesTimingTrack.cpp
@@ -19,9 +19,7 @@
 #include "Insights/ViewModels/TooltipDrawState.h"
 #include "Insights/Widgets/STimingView.h"
 
-#define LOCTEXT_NAMESPACE "FContextSwitchesTimingTrack"
-
-namespace Insights
+namespace UE::Insights::ContextSwitches
 {
 
 ////////////////////////////////////////////////////////////////////////////////////////////////////
@@ -331,6 +329,4 @@ const TSharedPtr<const ITimingEvent> FContextSwitchesTimingTrack::SearchEvent(co
 
 ////////////////////////////////////////////////////////////////////////////////////////////////////
 
-} // namespace Insights
-
-#undef LOCTEXT_NAMESPACE
+} // namespace UE::Insights::ContextSwitches
diff --git a/Engine/Source/Developer/TraceInsights/Private/Insights/ContextSwitches/ViewModels/ContextSwitchesTimingTrack.h b/Engine/Source/Developer/TraceInsights/Private/Insights/ContextSwitches/ViewModels/ContextSwitchesTimingTrack.h
--- a/Engine/Source/Developer/TraceInsights/Private/Insights/ContextSwitches/ViewModels/ContextSwitchesTimingTrack.h
+++ b/Engine/Source/Developer/TraceInsights/Private/Insights/ContextSwitches/ViewModels/ContextSwitchesTimingTrack.h
@@ -2,7 +2,7 @@
 
 #pragma once
 
-#include "CoreMinimal.h"
+#include "CoreTypes.h"
 
 // TraceInsights
 #include "Insights/ViewModels/TimingEventsTrack.h"
@@ -11,7 +11,7 @@
 
 class FThreadTrackEvent;
 
-namespace Insights
+namespace UE::Insights::ContextSwitches
 {
 
 class FContextSwitchesSharedState;
@@ -62,4 +62,4 @@ private:
 
 ////////////////////////////////////////////////////////////////////////////////////////////////////
 
-} //namespace Insights
+} // namespace UE::Insights::ContextSwitches
diff --git a/Engine/Source/Developer/TraceInsights/Private/Insights/ContextSwitches/ViewModels/CpuCoreTimingTrack.cpp b/Engine/Source/Developer/TraceInsights/Private/Insights/ContextSwitches/ViewModels/CpuCoreTimingTrack.cpp
--- a/Engine/Source/Developer/TraceInsights/Private/Insights/ContextSwitches/ViewModels/CpuCoreTimingTrack.cpp
+++ b/Engine/Source/Developer/TraceInsights/Private/Insights/ContextSwitches/ViewModels/CpuCoreTimingTrack.cpp
@@ -22,9 +22,7 @@
 #include "Insights/ViewModels/TooltipDrawState.h"
 #include "Insights/Widgets/STimingView.h"
 
-#define LOCTEXT_NAMESPACE "FCpuCoreTimingTrack"
-
-namespace Insights
+namespace UE::Insights::ContextSwitches
 {
 
 ////////////////////////////////////////////////////////////////////////////////////////////////////
@@ -147,20 +145,19 @@ void FCpuCoreTimingTrack::BuildFilteredDrawState(ITimingEventsTrackDrawStateBuil
 
 	if (HasCustomFilter())
 	{
-		TSharedPtr<UE::Insights::TimingProfiler::STimingView> TimingView = SharedState.GetTimingView();
+		TSharedPtr<TimingProfiler::STimingView> TimingView = SharedState.GetTimingView();
 		if (!TimingView)
 		{
 			return;
 		}
 
-		TSharedPtr<UE::Insights::FFilterConfigurator> FilterConfigurator = TimingView->GetFilterConfigurator();
+		TSharedPtr<FFilterConfigurator> FilterConfigurator = TimingView->GetFilterConfigurator();
 		if (!FilterConfigurator.IsValid())
 		{
 			return;
 		}
 
-		using EFilterField = UE::Insights::EFilterField;
-		UE::Insights::FFilterContext FilterContext;
+		FFilterContext FilterContext;
 		FilterContext.SetReturnValueForUnsetFilters(false);
 
 		FilterContext.AddFilterData<double>(static_cast<int32>(EFilterField::StartTime), 0.0f);
@@ -256,7 +253,7 @@ void FCpuCoreTimingTrack::PostDraw(const ITimingTrackDrawContext& Context) const
 
 		FString Str = FString::Printf(TEXT("%s (Duration.: %s)"),
 			*GetThreadName(static_cast<uint32>(SelectedEvent.GetType())),
-			*UE::Insights::FormatTimeAuto(SelectedEvent.GetDuration()));
+			*FormatTimeAuto(SelectedEvent.GetDuration()));
 
 		DrawSelectedEventInfo(Str, Context.GetViewport(), Context.GetDrawContext(), Helper.GetWhiteBrush(), Helper.GetEventFont());
 	}
@@ -338,8 +335,7 @@ const TSharedPtr<const ITimingEvent> FCpuCoreTimingTrack::SearchEvent(const FTim
 {
 	TSharedPtr<const ITimingEvent> FoundEvent;
 
-	using EFilterField = UE::Insights::EFilterField;
-	UE::Insights::FFilterContext FilterConfiguratorContext;
+	FFilterContext FilterConfiguratorContext;
 	FilterConfiguratorContext.SetReturnValueForUnsetFilters(false);
 	FilterConfiguratorContext.AddFilterData<double>(static_cast<int32>(EFilterField::StartTime), 0.0f);
 	FilterConfiguratorContext.AddFilterData<double>(static_cast<int32>(EFilterField::EndTime), 0.0f);
@@ -489,13 +485,13 @@ void FCpuCoreTimingTrack::BuildContextMenu(FMenuBuilder& InOutMenuBuilder)
 
 bool FCpuCoreTimingTrack::HasCustomFilter() const
 {
-	TSharedPtr<UE::Insights::TimingProfiler::STimingView> TimingView = SharedState.GetTimingView();
+	TSharedPtr<TimingProfiler::STimingView> TimingView = SharedState.GetTimingView();
 	if (!TimingView)
 	{
 		return false;
 	}
 
-	TSharedPtr<UE::Insights::FFilterConfigurator> FilterConfigurator = TimingView->GetFilterConfigurator();
+	TSharedPtr<FFilterConfigurator> FilterConfigurator = TimingView->GetFilterConfigurator();
 	return FilterConfigurator.IsValid() && !FilterConfigurator->IsEmpty();
 }
 
@@ -533,6 +529,4 @@ FString FCpuCoreTimingTrack::GetThreadName(uint32 InSystemThreadId) const
 
 ////////////////////////////////////////////////////////////////////////////////////////////////////
 
-} // namespace Insights
-
-#undef LOCTEXT_NAMESPACE
+} // namespace UE::Insights::ContextSwitches
diff --git a/Engine/Source/Developer/TraceInsights/Private/Insights/ContextSwitches/ViewModels/CpuCoreTimingTrack.h b/Engine/Source/Developer/TraceInsights/Private/Insights/ContextSwitches/ViewModels/CpuCoreTimingTrack.h
--- a/Engine/Source/Developer/TraceInsights/Private/Insights/ContextSwitches/ViewModels/CpuCoreTimingTrack.h
+++ b/Engine/Source/Developer/TraceInsights/Private/Insights/ContextSwitches/ViewModels/CpuCoreTimingTrack.h
@@ -2,7 +2,7 @@
 
 #pragma once
 
-#include "CoreMinimal.h"
+#include "CoreTypes.h"
 
 // TraceInsights
 #include "Insights/ViewModels/TimingEventsTrack.h"
@@ -17,7 +17,7 @@ namespace TraceServices
 	struct FCpuCoreEvent;
 };
 
-namespace Insights
+namespace UE::Insights::ContextSwitches
 {
 
 class FContextSwitchesSharedState;
@@ -73,4 +73,4 @@ private:
 
 ////////////////////////////////////////////////////////////////////////////////////////////////////
 
-} // namespace Insights
+} // namespace UE::Insights::ContextSwitches
diff --git a/Engine/Source/Developer/TraceInsights/Private/Insights/CookProfiler/CookProfilerManager.cpp b/Engine/Source/Developer/TraceInsights/Private/Insights/CookProfiler/CookProfilerManager.cpp
--- a/Engine/Source/Developer/TraceInsights/Private/Insights/CookProfiler/CookProfilerManager.cpp
+++ b/Engine/Source/Developer/TraceInsights/Private/Insights/CookProfiler/CookProfilerManager.cpp
@@ -18,13 +18,13 @@
 #include "Insights/TimingProfilerManager.h"
 #include "Insights/Widgets/STimingProfilerWindow.h"
 
-////////////////////////////////////////////////////////////////////////////////////////////////////
-
-#define LOCTEXT_NAMESPACE "CookProfilerManager"
+#define LOCTEXT_NAMESPACE "UE::Insights::CookProfiler"
 
-namespace Insights
+namespace UE::Insights::CookProfiler
 {
 
+////////////////////////////////////////////////////////////////////////////////////////////////////
+
 const FName FCookProfilerTabs::PackageTableTreeViewTabID(TEXT("PackageTableTreeView"));
 
 ////////////////////////////////////////////////////////////////////////////////////////////////////
@@ -263,6 +263,6 @@ void FCookProfilerManager::OnWindowClosedEvent()
 
 ////////////////////////////////////////////////////////////////////////////////////////////////////
 
-} // namespace Insights
+} // namespace UE::Insights::CookProfiler
 
 #undef LOCTEXT_NAMESPACE
diff --git a/Engine/Source/Developer/TraceInsights/Private/Insights/CookProfiler/CookProfilerManager.h b/Engine/Source/Developer/TraceInsights/Private/Insights/CookProfiler/CookProfilerManager.h
--- a/Engine/Source/Developer/TraceInsights/Private/Insights/CookProfiler/CookProfilerManager.h
+++ b/Engine/Source/Developer/TraceInsights/Private/Insights/CookProfiler/CookProfilerManager.h
@@ -2,16 +2,17 @@
 
 #pragma once
 
-#include "CoreMinimal.h"
+#include "CoreTypes.h"
 
 #include "Containers/Ticker.h"
 #include "Framework/Commands/UICommandList.h"
+#include "Templates/SharedPointer.h"
 
 // TraceInsights
 #include "Insights/InsightsManager.h"
 #include "Insights/IUnrealInsightsModule.h"
 
-namespace Insights
+namespace UE::Insights::CookProfiler
 {
 
 class SPackageTableTreeView;
@@ -77,7 +78,7 @@ private:
 private:
 	bool bIsInitialized;
 	bool bIsAvailable;
-	UE::Insights::FAvailabilityCheck AvailabilityCheck;
+	FAvailabilityCheck AvailabilityCheck;
 
 	/** The delegate to be invoked when this manager ticks. */
 	FTickerDelegate OnTick;
@@ -93,5 +94,5 @@ private:
 	TSharedPtr<SPackageTableTreeView> PackageTableTreeView;
 };
 
-} // namespace Insights
+} // namespace UE::Insights::CookProfiler
 
diff --git a/Engine/Source/Developer/TraceInsights/Private/Insights/CookProfiler/ViewModels/PackageEntry.cpp b/Engine/Source/Developer/TraceInsights/Private/Insights/CookProfiler/ViewModels/PackageEntry.cpp
--- a/Engine/Source/Developer/TraceInsights/Private/Insights/CookProfiler/ViewModels/PackageEntry.cpp
+++ b/Engine/Source/Developer/TraceInsights/Private/Insights/CookProfiler/ViewModels/PackageEntry.cpp
@@ -2,7 +2,10 @@
 
 #include "PackageEntry.h"
 
-namespace Insights
+// TraceServices
+#include "TraceServices/Model/CookProfilerProvider.h"
+
+namespace UE::Insights::CookProfiler
 {
 
 ////////////////////////////////////////////////////////////////////////////////////////////////////
@@ -24,4 +27,4 @@ FPackageEntry::FPackageEntry(const TraceServices::FPackageData& PackageData)
 
 ////////////////////////////////////////////////////////////////////////////////////////////////////
 
-} // namespace Insights
+} // namespace UE::Insights::CookProfiler
diff --git a/Engine/Source/Developer/TraceInsights/Private/Insights/CookProfiler/ViewModels/PackageEntry.h b/Engine/Source/Developer/TraceInsights/Private/Insights/CookProfiler/ViewModels/PackageEntry.h
--- a/Engine/Source/Developer/TraceInsights/Private/Insights/CookProfiler/ViewModels/PackageEntry.h
+++ b/Engine/Source/Developer/TraceInsights/Private/Insights/CookProfiler/ViewModels/PackageEntry.h
@@ -2,12 +2,11 @@
 
 #pragma once
 
-#include "CoreMinimal.h"
+#include "CoreTypes.h"
 
-// TraceServices
-#include "TraceServices/Model/CookProfilerProvider.h"
+namespace TraceServices { struct FPackageData; }
 
-namespace Insights
+namespace UE::Insights::CookProfiler
 {
 
 ////////////////////////////////////////////////////////////////////////////////////////////////////
@@ -17,7 +16,7 @@ class FPackageEntry
 	friend class SPackageTableTreeView;
 
 public:
-	FPackageEntry(const TraceServices::FPackageData &PackageData);
+	FPackageEntry(const TraceServices::FPackageData& PackageData);
 	~FPackageEntry() {}
 
 	uint64 GetId() const { return Id; }
@@ -39,7 +38,7 @@ public:
 
 private:
 	uint64 Id;
-	const TCHAR* Name; 
+	const TCHAR* Name;
 	double LoadTimeIncl;
 	double LoadTimeExcl;
 
@@ -57,4 +56,4 @@ private:
 
 ////////////////////////////////////////////////////////////////////////////////////////////////////
 
-} // namespace Insights
+} // namespace UE::Insights::CookProfiler
diff --git a/Engine/Source/Developer/TraceInsights/Private/Insights/CookProfiler/ViewModels/PackageNode.cpp b/Engine/Source/Developer/TraceInsights/Private/Insights/CookProfiler/ViewModels/PackageNode.cpp
--- a/Engine/Source/Developer/TraceInsights/Private/Insights/CookProfiler/ViewModels/PackageNode.cpp
+++ b/Engine/Source/Developer/TraceInsights/Private/Insights/CookProfiler/ViewModels/PackageNode.cpp
@@ -2,13 +2,9 @@
 
 #include "PackageNode.h"
 
-#define LOCTEXT_NAMESPACE "Insights::FPackageNode"
-
-namespace Insights
+namespace UE::Insights::CookProfiler
 {
 
 INSIGHTS_IMPLEMENT_RTTI(FPackageNode)
 
-} // namespace Insights
-
-#undef LOCTEXT_NAMESPACE
+} // namespace UE::Insights::CookProfiler
diff --git a/Engine/Source/Developer/TraceInsights/Private/Insights/CookProfiler/ViewModels/PackageNode.h b/Engine/Source/Developer/TraceInsights/Private/Insights/CookProfiler/ViewModels/PackageNode.h
--- a/Engine/Source/Developer/TraceInsights/Private/Insights/CookProfiler/ViewModels/PackageNode.h
+++ b/Engine/Source/Developer/TraceInsights/Private/Insights/CookProfiler/ViewModels/PackageNode.h
@@ -2,7 +2,9 @@
 
 #pragma once
 
-#include "CoreMinimal.h"
+#include "CoreTypes.h"
+
+#include "Templates/SharedPointer.h"
 
 // TraceInsightsCore
 #include "InsightsCore/Table/ViewModels/TableTreeNode.h"
@@ -13,7 +15,7 @@
 
 ////////////////////////////////////////////////////////////////////////////////////////////////////
 
-namespace Insights
+namespace UE::Insights::CookProfiler
 {
 
 ////////////////////////////////////////////////////////////////////////////////////////////////////
@@ -36,26 +38,26 @@ typedef TWeakPtr<class FTaskNode> FPackageNodeWeak;
 /**
  * Class used to store information about a package node (used in the SPackageTableTreeView).
  */
-class FPackageNode : public UE::Insights::FTableTreeNode
+class FPackageNode : public FTableTreeNode
 {
-	INSIGHTS_DECLARE_RTTI(FPackageNode, UE::Insights::FTableTreeNode)
+	INSIGHTS_DECLARE_RTTI(FPackageNode, FTableTreeNode)
 
 public:
 	/** Initialization constructor for the Task node. */
 	explicit FPackageNode(const FName InName, TWeakPtr<FPackageTable> InParentTable, int32 InRowIndex)
-		: UE::Insights::FTableTreeNode(InName, InParentTable, InRowIndex)
+		: FTableTreeNode(InName, InParentTable, InRowIndex)
 	{
 	}
 
 	/** Initialization constructor for the group node. */
 	explicit FPackageNode(const FName InGroupName, TWeakPtr<FPackageTable> InParentTable)
-		: UE::Insights::FTableTreeNode(InGroupName, InParentTable)
+		: FTableTreeNode(InGroupName, InParentTable)
 	{
 	}
 
 	FPackageTable& GetPackageTableChecked() const
 	{
-		const TSharedPtr<UE::Insights::FTable>& TablePin = GetParentTable().Pin();
+		const TSharedPtr<FTable>& TablePin = GetParentTable().Pin();
 		check(TablePin.IsValid());
 		return *StaticCastSharedPtr<FPackageTable>(TablePin);
 	}
@@ -67,4 +69,4 @@ public:
 
 ////////////////////////////////////////////////////////////////////////////////////////////////////
 
-} // namespace Insights
+} // namespace UE::Insights::CookProfiler
diff --git a/Engine/Source/Developer/TraceInsights/Private/Insights/CookProfiler/ViewModels/PackageTable.cpp b/Engine/Source/Developer/TraceInsights/Private/Insights/CookProfiler/ViewModels/PackageTable.cpp
--- a/Engine/Source/Developer/TraceInsights/Private/Insights/CookProfiler/ViewModels/PackageTable.cpp
+++ b/Engine/Source/Developer/TraceInsights/Private/Insights/CookProfiler/ViewModels/PackageTable.cpp
@@ -12,9 +12,9 @@
 #include "Insights/CookProfiler/ViewModels/PackageNode.h"
 #include "Insights/CookProfiler/ViewModels/PackageTable.h"
 
-#define LOCTEXT_NAMESPACE "Insights::FPackageTable"
+#define LOCTEXT_NAMESPACE "UE::Insights::CookProfiler::FPackageTable"
 
-namespace Insights
+namespace UE::Insights::CookProfiler
 {
 
 ////////////////////////////////////////////////////////////////////////////////////////////////////
@@ -39,17 +39,17 @@ const FName FPackageTableColumns::PackageAssetClassColumnId(TEXT("AssetClass"));
 
 ////////////////////////////////////////////////////////////////////////////////////////////////////
 
-typedef UE::Insights::FTableCellValue (*PackageFieldGetter) (const UE::Insights::FTableColumn&, const FPackageEntry&);
+typedef FTableCellValue (*PackageFieldGetter) (const FTableColumn&, const FPackageEntry&);
 
 template<PackageFieldGetter Getter>
-class FPackageColumnValueGetter : public UE::Insights::FTableCellValueGetter
+class FPackageColumnValueGetter : public FTableCellValueGetter
 {
 public:
-	virtual const TOptional<UE::Insights::FTableCellValue> GetValue(const UE::Insights::FTableColumn& Column, const UE::Insights::FBaseTreeNode& Node) const override
+	virtual const TOptional<FTableCellValue> GetValue(const FTableColumn& Column, const FBaseTreeNode& Node) const override
 	{
 		if (Node.IsGroup())
 		{
-			const UE::Insights::FTableTreeNode& NodePtr = static_cast<const UE::Insights::FTableTreeNode&>(Node);
+			const FTableTreeNode& NodePtr = static_cast<const FTableTreeNode&>(Node);
 			if (NodePtr.HasAggregatedValue(Column.GetId()))
 			{
 				return NodePtr.GetAggregatedValue(Column.GetId());
@@ -65,7 +65,7 @@ public:
 			}
 		}
 
-		return TOptional<UE::Insights::FTableCellValue>();
+		return TOptional<FTableCellValue>();
 	}
 };
 
@@ -73,9 +73,6 @@ public:
 
 struct DefaultPackageFieldGetterFuncts
 {
-	using FTableCellValue = UE::Insights::FTableCellValue;
-	using FTableColumn = UE::Insights::FTableColumn;
-
 	static FTableCellValue GetId(const FTableColumn& Column, const FPackageEntry& Package) { return FTableCellValue((int64)Package.GetId());	}
 	static FTableCellValue GetName(const FTableColumn& Column, const FPackageEntry& Package) { return FTableCellValue((const TCHAR*)Package.GetName());	}
 
@@ -495,6 +492,6 @@ void FPackageTable::AddDefaultColumns()
 
 ////////////////////////////////////////////////////////////////////////////////////////////////////
 
-} // namespace Insights
+} // namespace UE::Insights::CookProfiler
 
 #undef LOCTEXT_NAMESPACE
diff --git a/Engine/Source/Developer/TraceInsights/Private/Insights/CookProfiler/ViewModels/PackageTable.h b/Engine/Source/Developer/TraceInsights/Private/Insights/CookProfiler/ViewModels/PackageTable.h
--- a/Engine/Source/Developer/TraceInsights/Private/Insights/CookProfiler/ViewModels/PackageTable.h
+++ b/Engine/Source/Developer/TraceInsights/Private/Insights/CookProfiler/ViewModels/PackageTable.h
@@ -15,7 +15,7 @@
 
 namespace UE::Insights { class FTableColumn; }
 
-namespace Insights
+namespace UE::Insights::CookProfiler
 {
 
 ////////////////////////////////////////////////////////////////////////////////////////////////////
@@ -43,7 +43,7 @@ struct FPackageTableColumns
 
 ////////////////////////////////////////////////////////////////////////////////////////////////////
 
-class FPackageTable : public UE::Insights::FTable
+class FPackageTable : public FTable
 {
 public:
 	FPackageTable();
@@ -67,4 +67,4 @@ private:
 
 ////////////////////////////////////////////////////////////////////////////////////////////////////
 
-} // namespace Insights
+} // namespace UE::Insights::CookProfiler
diff --git a/Engine/Source/Developer/TraceInsights/Private/Insights/CookProfiler/Widgets/SPackageTableTreeView.cpp b/Engine/Source/Developer/TraceInsights/Private/Insights/CookProfiler/Widgets/SPackageTableTreeView.cpp
--- a/Engine/Source/Developer/TraceInsights/Private/Insights/CookProfiler/Widgets/SPackageTableTreeView.cpp
+++ b/Engine/Source/Developer/TraceInsights/Private/Insights/CookProfiler/Widgets/SPackageTableTreeView.cpp
@@ -24,11 +24,9 @@
 
 #include <limits>
 
-#define LOCTEXT_NAMESPACE "SPackageTableTreeView"
+#define LOCTEXT_NAMESPACE "UE::Insights::CookProfiler::SPackageTableTreeView"
 
-using namespace TraceServices;
-
-namespace Insights
+namespace UE::Insights::CookProfiler
 {
 
 ////////////////////////////////////////////////////////////////////////////////////////////////////
@@ -56,6 +54,8 @@ public:
 	}
 };
 
+////////////////////////////////////////////////////////////////////////////////////////////////////
+// SPackageTableTreeView
 ////////////////////////////////////////////////////////////////////////////////////////////////////
 
 SPackageTableTreeView::SPackageTableTreeView()
@@ -137,13 +137,13 @@ void SPackageTableTreeView::RebuildTree(bool bResync)
 	{
 		TraceServices::FProviderReadScopeLock ProviderReadScope(*CookProvider);
 
-		TArray64<FPackageData> PackageAggreagation;
+		TArray64<TraceServices::FPackageData> PackageAggreagation;
 		CookProvider->CreateAggregation(PackageAggreagation);
 		const uint32 NumPackages = CookProvider->GetNumPackages();
 		Packages.Reserve(NumPackages);
 		TableRowNodes.Reserve(NumPackages);
 
-		TArray<UE::Insights::FTableTreeNodePtr>* Nodes = &TableRowNodes;
+		TArray<FTableTreeNodePtr>* Nodes = &TableRowNodes;
 		for(const TraceServices::FPackageData& Package : PackageAggreagation)
 		{
 			Packages.Emplace(Package);
@@ -204,7 +204,7 @@ TSharedPtr<SWidget> SPackageTableTreeView::ConstructToolbar()
 			SNew(SBox)
 			.MinDesiredWidth(150.0f)
 			[
-				SAssignNew(PresetComboBox, SComboBox<TSharedRef<UE::Insights::ITableTreeViewPreset>>)
+				SAssignNew(PresetComboBox, SComboBox<TSharedRef<ITableTreeViewPreset>>)
 				.ToolTipText(this, &SPackageTableTreeView::ViewPreset_GetSelectedToolTipText)
 				.OptionsSource(GetAvailableViewPresets())
 				.OnSelectionChanged(this, &SPackageTableTreeView::ViewPreset_OnSelectionChanged)
@@ -235,11 +235,11 @@ void SPackageTableTreeView::InternalCreateGroupings()
 	STableTreeView::InternalCreateGroupings();
 
 	AvailableGroupings.RemoveAll(
-		[](TSharedPtr<UE::Insights::FTreeNodeGrouping>& Grouping)
+		[](TSharedPtr<FTreeNodeGrouping>& Grouping)
 		{
-			if (Grouping->Is<UE::Insights::FTreeNodeGroupingByUniqueValue>())
+			if (Grouping->Is<FTreeNodeGroupingByUniqueValue>())
 			{
-				const FName ColumnId = Grouping->As<UE::Insights::FTreeNodeGroupingByUniqueValue>().GetColumnId();
+				const FName ColumnId = Grouping->As<FTreeNodeGroupingByUniqueValue>().GetColumnId();
 				if (ColumnId == FPackageTableColumns::BeginCacheForCookedPlatformDataTimeInclColumnId ||
 					ColumnId == FPackageTableColumns::BeginCacheForCookedPlatformDataTimeExclColumnId ||
 					ColumnId == FPackageTableColumns::GetIsCachedCookedPlatformDataLoadedInclColumnId ||
@@ -254,9 +254,9 @@ void SPackageTableTreeView::InternalCreateGroupings()
 					return true;
 				}
 			}
-			else if (Grouping->Is<UE::Insights::FTreeNodeGroupingByPathBreakdown>())
+			else if (Grouping->Is<FTreeNodeGroupingByPathBreakdown>())
 			{
-				const FName ColumnId = Grouping->As<UE::Insights::FTreeNodeGroupingByPathBreakdown>().GetColumnId();
+				const FName ColumnId = Grouping->As<FTreeNodeGroupingByPathBreakdown>().GetColumnId();
 				if (ColumnId == FPackageTableColumns::PackageAssetClassColumnId)
 				{
 					return true;
@@ -268,7 +268,7 @@ void SPackageTableTreeView::InternalCreateGroupings()
 
 ////////////////////////////////////////////////////////////////////////////////////////////////////
 
-void SPackageTableTreeView::TreeView_OnMouseButtonDoubleClick(UE::Insights::FTableTreeNodePtr TreeNode)
+void SPackageTableTreeView::TreeView_OnMouseButtonDoubleClick(FTableTreeNodePtr TreeNode)
 {
 	STableTreeView::TreeView_OnMouseButtonDoubleClick(TreeNode);
 }
@@ -277,8 +277,6 @@ void SPackageTableTreeView::TreeView_OnMouseButtonDoubleClick(UE::Insights::FTab
 
 void SPackageTableTreeView::InitAvailableViewPresets()
 {
-	using namespace UE::Insights;
-
 	//////////////////////////////////////////////////
 	// Default View
 
@@ -464,6 +462,6 @@ void SPackageTableTreeView::UpdateBannerText()
 
 ////////////////////////////////////////////////////////////////////////////////////////////////////
 
-} // namespace Insights
+} // namespace UE::Insights::CookProfiler
 
 #undef LOCTEXT_NAMESPACE
diff --git a/Engine/Source/Developer/TraceInsights/Private/Insights/CookProfiler/Widgets/SPackageTableTreeView.h b/Engine/Source/Developer/TraceInsights/Private/Insights/CookProfiler/Widgets/SPackageTableTreeView.h
--- a/Engine/Source/Developer/TraceInsights/Private/Insights/CookProfiler/Widgets/SPackageTableTreeView.h
+++ b/Engine/Source/Developer/TraceInsights/Private/Insights/CookProfiler/Widgets/SPackageTableTreeView.h
@@ -13,12 +13,12 @@
 
 class FMenuBuilder;
 
-namespace Insights
+namespace UE::Insights::CookProfiler
 {
 
 ////////////////////////////////////////////////////////////////////////////////////////////////////
 
-class SPackageTableTreeView : public UE::Insights::SSessionTableTreeView
+class SPackageTableTreeView : public SSessionTableTreeView
 {
 public:
 	/** Default constructor. */
@@ -68,7 +68,7 @@ public:
 
 	////////////////////////////////////////////////////////////////////////////////////////////////////
 
-	virtual void TreeView_OnMouseButtonDoubleClick(UE::Insights::FTableTreeNodePtr TreeNode) override;
+	virtual void TreeView_OnMouseButtonDoubleClick(FTableTreeNodePtr TreeNode) override;
 
 protected:
 	virtual void InternalCreateGroupings() override;
@@ -88,4 +88,4 @@ private:
 
 ////////////////////////////////////////////////////////////////////////////////////////////////////
 
-} // namespace Insights
+} // namespace UE::Insights::CookProfiler
diff --git a/Engine/Source/Developer/TraceInsights/Private/Insights/LoadingProfiler/LoadingProfilerCommands.cpp b/Engine/Source/Developer/TraceInsights/Private/Insights/LoadingProfiler/LoadingProfilerCommands.cpp
--- a/Engine/Source/Developer/TraceInsights/Private/Insights/LoadingProfiler/LoadingProfilerCommands.cpp
+++ b/Engine/Source/Developer/TraceInsights/Private/Insights/LoadingProfiler/LoadingProfilerCommands.cpp
@@ -11,9 +11,10 @@
 #include "Insights/InsightsStyle.h"
 #include "Insights/LoadingProfiler/LoadingProfilerManager.h"
 
-////////////////////////////////////////////////////////////////////////////////////////////////////
+#define LOCTEXT_NAMESPACE "UE::Insights::LoadingProfiler"
 
-#define LOCTEXT_NAMESPACE "FLoadingProfilerCommands"
+namespace UE::Insights::LoadingProfiler
+{
 
 ////////////////////////////////////////////////////////////////////////////////////////////////////
 // FLoadingProfilerMenuBuilder
@@ -137,4 +138,6 @@ IMPLEMENT_TOGGLE_COMMAND(ToggleRequestsTreeViewVisibility, IsRequestsTreeViewVis
 
 ////////////////////////////////////////////////////////////////////////////////////////////////////
 
+} // namespace UE::Insights::LoadingProfiler
+
 #undef LOCTEXT_NAMESPACE
diff --git a/Engine/Source/Developer/TraceInsights/Private/Insights/LoadingProfiler/LoadingProfilerCommands.h b/Engine/Source/Developer/TraceInsights/Private/Insights/LoadingProfiler/LoadingProfilerCommands.h
--- a/Engine/Source/Developer/TraceInsights/Private/Insights/LoadingProfiler/LoadingProfilerCommands.h
+++ b/Engine/Source/Developer/TraceInsights/Private/Insights/LoadingProfiler/LoadingProfilerCommands.h
@@ -2,12 +2,18 @@
 
 #pragma once
 
-#include "CoreMinimal.h"
-#include "Framework/Commands/UIAction.h"
+#include "CoreTypes.h"
+
 #include "Framework/Commands/Commands.h"
+#include "Framework/Commands/UIAction.h"
+#include "Styling/SlateTypes.h" // for ECheckBoxState
+#include "Templates/SharedPointer.h"
 
 class FMenuBuilder;
 
+namespace UE::Insights::LoadingProfiler
+{
+
 ////////////////////////////////////////////////////////////////////////////////////////////////////
 /**
  * Class that holds all profiler commands.
@@ -111,3 +117,7 @@ protected:
 	/** Reference to the global instance of the profiler manager. */
 	class FLoadingProfilerManager* This;
 };
+
+////////////////////////////////////////////////////////////////////////////////////////////////////
+
+} // namespace UE::Insights::LoadingProfiler
diff --git a/Engine/Source/Developer/TraceInsights/Private/Insights/LoadingProfiler/LoadingProfilerManager.cpp b/Engine/Source/Developer/TraceInsights/Private/Insights/LoadingProfiler/LoadingProfilerManager.cpp
--- a/Engine/Source/Developer/TraceInsights/Private/Insights/LoadingProfiler/LoadingProfilerManager.cpp
+++ b/Engine/Source/Developer/TraceInsights/Private/Insights/LoadingProfiler/LoadingProfilerManager.cpp
@@ -21,9 +21,12 @@
 #include "Insights/LoadingProfiler/Widgets/SLoadingProfilerWindow.h"
 #include "Insights/Widgets/STimingView.h"
 
-////////////////////////////////////////////////////////////////////////////////////////////////////
+#define LOCTEXT_NAMESPACE "UE::Insights::LoadingProfiler"
 
-#define LOCTEXT_NAMESPACE "LoadingProfilerManager"
+namespace UE::Insights::LoadingProfiler
+{
+
+////////////////////////////////////////////////////////////////////////////////////////////////////
 
 DEFINE_LOG_CATEGORY(LoadingProfiler);
 
@@ -388,7 +391,7 @@ void FLoadingProfilerManager::OnWindowClosedEvent()
 	TSharedPtr<SLoadingProfilerWindow> Wnd = GetProfilerWindow();
 	if (Wnd)
 	{
-		TSharedPtr<UE::Insights::TimingProfiler::STimingView> TimingView = Wnd->GetTimingView();
+		TSharedPtr<TimingProfiler::STimingView> TimingView = Wnd->GetTimingView();
 		if (TimingView.IsValid())
 		{
 			TimingView->CloseQuickFindTab();
@@ -398,4 +401,6 @@ void FLoadingProfilerManager::OnWindowClosedEvent()
 
 ////////////////////////////////////////////////////////////////////////////////////////////////////
 
+} // namespace UE::Insights::LoadingProfiler
+
 #undef LOCTEXT_NAMESPACE
diff --git a/Engine/Source/Developer/TraceInsights/Private/Insights/LoadingProfiler/LoadingProfilerManager.h b/Engine/Source/Developer/TraceInsights/Private/Insights/LoadingProfiler/LoadingProfilerManager.h
--- a/Engine/Source/Developer/TraceInsights/Private/Insights/LoadingProfiler/LoadingProfilerManager.h
+++ b/Engine/Source/Developer/TraceInsights/Private/Insights/LoadingProfiler/LoadingProfilerManager.h
@@ -2,16 +2,23 @@
 
 #pragma once
 
+#include "CoreTypes.h"
+
 #include "Containers/Ticker.h"
-#include "CoreMinimal.h"
 #include "Framework/Commands/UICommandList.h"
 #include "Logging/LogMacros.h"
+#include "Templates/SharedPointer.h"
+
+// TraceInsightsCore
+#include "InsightsCore/Common/AvailabilityCheck.h"
 
 // TraceInsights
-#include "Insights/InsightsManager.h"
 #include "Insights/IUnrealInsightsModule.h"
 #include "Insights/LoadingProfiler/LoadingProfilerCommands.h"
 
+namespace UE::Insights::LoadingProfiler
+{
+
 class SLoadingProfilerWindow;
 
 DECLARE_LOG_CATEGORY_EXTERN(LoadingProfiler, Log, All);
@@ -139,7 +146,7 @@ private:
 private:
 	bool bIsInitialized;
 	bool bIsAvailable;
-	UE::Insights::FAvailabilityCheck AvailabilityCheck;
+	FAvailabilityCheck AvailabilityCheck;
 
 	/** The delegate to be invoked when this manager ticks. */
 	FTickerDelegate OnTick;
@@ -180,3 +187,7 @@ private:
 	/** A shared pointer to the global instance of the Loading Profiler manager. */
 	static TSharedPtr<FLoadingProfilerManager> Instance;
 };
+
+////////////////////////////////////////////////////////////////////////////////////////////////////
+
+} // namespace UE::Insights::LoadingProfiler
diff --git a/Engine/Source/Developer/TraceInsights/Private/Insights/LoadingProfiler/Widgets/SLoadingProfilerToolbar.cpp b/Engine/Source/Developer/TraceInsights/Private/Insights/LoadingProfiler/Widgets/SLoadingProfilerToolbar.cpp
--- a/Engine/Source/Developer/TraceInsights/Private/Insights/LoadingProfiler/Widgets/SLoadingProfilerToolbar.cpp
+++ b/Engine/Source/Developer/TraceInsights/Private/Insights/LoadingProfiler/Widgets/SLoadingProfilerToolbar.cpp
@@ -12,9 +12,8 @@
 #include "Insights/InsightsStyle.h"
 #include "Insights/LoadingProfiler/LoadingProfilerCommands.h"
 
-////////////////////////////////////////////////////////////////////////////////////////////////////
-
-#define LOCTEXT_NAMESPACE "SLoadingProfilerToolbar"
+namespace UE::Insights::LoadingProfiler
+{
 
 ////////////////////////////////////////////////////////////////////////////////////////////////////
 
@@ -118,4 +117,4 @@ void SLoadingProfilerToolbar::Construct(const FArguments& InArgs)
 
 ////////////////////////////////////////////////////////////////////////////////////////////////////
 
-#undef LOCTEXT_NAMESPACE
+} // namespace UE::Insights::LoadingProfiler
diff --git a/Engine/Source/Developer/TraceInsights/Private/Insights/LoadingProfiler/Widgets/SLoadingProfilerToolbar.h b/Engine/Source/Developer/TraceInsights/Private/Insights/LoadingProfiler/Widgets/SLoadingProfilerToolbar.h
--- a/Engine/Source/Developer/TraceInsights/Private/Insights/LoadingProfiler/Widgets/SLoadingProfilerToolbar.h
+++ b/Engine/Source/Developer/TraceInsights/Private/Insights/LoadingProfiler/Widgets/SLoadingProfilerToolbar.h
@@ -2,11 +2,16 @@
 
 #pragma once
 
+#include "CoreTypes.h"
+
 #include "Widgets/DeclarativeSyntaxSupport.h"
 #include "Widgets/SCompoundWidget.h"
 
 class FExtender;
 
+namespace UE::Insights::LoadingProfiler
+{
+
 /** Ribbon based toolbar used as a main menu in the Profiler window. */
 class SLoadingProfilerToolbar : public SCompoundWidget
 {
@@ -28,3 +33,5 @@ public:
 	 */
 	void Construct(const FArguments& InArgs);
 };
+
+} // namespace UE::Insights::LoadingProfiler
diff --git a/Engine/Source/Developer/TraceInsights/Private/Insights/LoadingProfiler/Widgets/SLoadingProfilerWindow.cpp b/Engine/Source/Developer/TraceInsights/Private/Insights/LoadingProfiler/Widgets/SLoadingProfilerWindow.cpp
--- a/Engine/Source/Developer/TraceInsights/Private/Insights/LoadingProfiler/Widgets/SLoadingProfilerWindow.cpp
+++ b/Engine/Source/Developer/TraceInsights/Private/Insights/LoadingProfiler/Widgets/SLoadingProfilerWindow.cpp
@@ -12,15 +12,17 @@
 #include "TraceServices/Model/LoadTimeProfiler.h"
 
 // TraceInsights
+#include "Insights/InsightsManager.h"
 #include "Insights/InsightsStyle.h"
 #include "Insights/LoadingProfiler/LoadingProfilerManager.h"
 #include "Insights/LoadingProfiler/Widgets/SLoadingProfilerToolbar.h"
 #include "Insights/Table/Widgets/SUntypedTableTreeView.h"
 #include "Insights/Widgets/STimingView.h"
 
-////////////////////////////////////////////////////////////////////////////////////////////////////
+#define LOCTEXT_NAMESPACE "UE::Insights::LoadingProfiler"
 
-#define LOCTEXT_NAMESPACE "SLoadingProfilerWindow"
+namespace UE::Insights::LoadingProfiler
+{
 
 ////////////////////////////////////////////////////////////////////////////////////////////////////
 
@@ -230,7 +232,7 @@ TSharedRef<SDockTab> SLoadingProfilerWindow::SpawnTab_TimingView(const FSpawnTab
 		.ShouldAutosize(false)
 		.TabRole(ETabRole::PanelTab)
 		[
-			SAssignNew(TimingView, UE::Insights::TimingProfiler::STimingView, FInsightsManagerTabs::LoadingProfilerTabId)
+			SAssignNew(TimingView, TimingProfiler::STimingView, FInsightsManagerTabs::LoadingProfilerTabId)
 		];
 
 	TimingView->Reset(true);
@@ -265,14 +267,14 @@ TSharedRef<SDockTab> SLoadingProfilerWindow::SpawnTab_EventAggregationTreeView(c
 {
 	FLoadingProfilerManager::Get()->SetEventAggregationTreeViewVisible(true);
 
-	TSharedRef<Insights::FUntypedTable> Table = MakeShared<Insights::FUntypedTable>();
+	TSharedRef<::Insights::FUntypedTable> Table = MakeShared<::Insights::FUntypedTable>();
 	Table->SetDisplayName(LOCTEXT("EventAggregation_TableName", "Event Aggregation"));
 
 	const TSharedRef<SDockTab> DockTab = SNew(SDockTab)
 		.ShouldAutosize(false)
 		.TabRole(ETabRole::PanelTab)
 		[
-			SAssignNew(EventAggregationTreeView, Insights::SUntypedTableTreeView, Table)
+			SAssignNew(EventAggregationTreeView, ::Insights::SUntypedTableTreeView, Table)
 		];
 
 	EventAggregationTreeView->SetLogListingName(FLoadingProfilerManager::Get()->GetLogListingName());
@@ -302,14 +304,14 @@ TSharedRef<SDockTab> SLoadingProfilerWindow::SpawnTab_ObjectTypeAggregationTreeV
 {
 	FLoadingProfilerManager::Get()->SetObjectTypeAggregationTreeViewVisible(true);
 
-	TSharedRef<Insights::FUntypedTable> Table = MakeShared<Insights::FUntypedTable>();
+	TSharedRef<::Insights::FUntypedTable> Table = MakeShared<::Insights::FUntypedTable>();
 	Table->SetDisplayName(LOCTEXT("ObjectTypeAggregation_TableName", "Object Type Aggregation"));
 
 	const TSharedRef<SDockTab> DockTab = SNew(SDockTab)
 		.ShouldAutosize(false)
 		.TabRole(ETabRole::PanelTab)
 		[
-			SAssignNew(ObjectTypeAggregationTreeView, Insights::SUntypedTableTreeView, Table)
+			SAssignNew(ObjectTypeAggregationTreeView, ::Insights::SUntypedTableTreeView, Table)
 		];
 
 	ObjectTypeAggregationTreeView->SetLogListingName(FLoadingProfilerManager::Get()->GetLogListingName());
@@ -339,14 +341,14 @@ TSharedRef<SDockTab> SLoadingProfilerWindow::SpawnTab_PackageDetailsTreeView(con
 {
 	FLoadingProfilerManager::Get()->SetPackageDetailsTreeViewVisible(true);
 
-	TSharedRef<Insights::FUntypedTable> Table = MakeShared<Insights::FUntypedTable>();
+	TSharedRef<::Insights::FUntypedTable> Table = MakeShared<::Insights::FUntypedTable>();
 	Table->SetDisplayName(LOCTEXT("PackageDetails_TableName", "Package Details"));
 
 	const TSharedRef<SDockTab> DockTab = SNew(SDockTab)
 		.ShouldAutosize(false)
 		.TabRole(ETabRole::PanelTab)
 		[
-			SAssignNew(PackageDetailsTreeView, Insights::SUntypedTableTreeView, Table)
+			SAssignNew(PackageDetailsTreeView, ::Insights::SUntypedTableTreeView, Table)
 		];
 
 	PackageDetailsTreeView->SetLogListingName(FLoadingProfilerManager::Get()->GetLogListingName());
@@ -376,14 +378,14 @@ TSharedRef<SDockTab> SLoadingProfilerWindow::SpawnTab_ExportDetailsTreeView(cons
 {
 	FLoadingProfilerManager::Get()->SetExportDetailsTreeViewVisible(true);
 
-	TSharedRef<Insights::FUntypedTable> Table = MakeShared<Insights::FUntypedTable>();
+	TSharedRef<::Insights::FUntypedTable> Table = MakeShared<::Insights::FUntypedTable>();
 	Table->SetDisplayName(LOCTEXT("ExportDetails_TableName", "Export Details"));
 
 	const TSharedRef<SDockTab> DockTab = SNew(SDockTab)
 		.ShouldAutosize(false)
 		.TabRole(ETabRole::PanelTab)
 		[
-			SAssignNew(ExportDetailsTreeView, Insights::SUntypedTableTreeView, Table)
+			SAssignNew(ExportDetailsTreeView, ::Insights::SUntypedTableTreeView, Table)
 		];
 
 	ExportDetailsTreeView->SetLogListingName(FLoadingProfilerManager::Get()->GetLogListingName());
@@ -413,14 +415,14 @@ TSharedRef<SDockTab> SLoadingProfilerWindow::SpawnTab_RequestsTreeView(const FSp
 {
 	FLoadingProfilerManager::Get()->SetRequestsTreeViewVisible(true);
 
-	TSharedRef<Insights::FUntypedTable> Table = MakeShared<Insights::FUntypedTable>();
+	TSharedRef<::Insights::FUntypedTable> Table = MakeShared<::Insights::FUntypedTable>();
 	Table->SetDisplayName(LOCTEXT("Requests_TableName", "Requests"));
 
 	const TSharedRef<SDockTab> DockTab = SNew(SDockTab)
 		.ShouldAutosize(false)
 		.TabRole(ETabRole::PanelTab)
 		[
-			SAssignNew(RequestsTreeView, Insights::SUntypedTableTreeView, Table)
+			SAssignNew(RequestsTreeView, ::Insights::SUntypedTableTreeView, Table)
 		];
 
 	RequestsTreeView->SetLogListingName(FLoadingProfilerManager::Get()->GetLogListingName());
@@ -544,9 +546,9 @@ TSharedRef<SWidget> SLoadingProfilerWindow::CreateToolbar(TSharedPtr<FExtender>
 
 ////////////////////////////////////////////////////////////////////////////////////////////////////
 
-void SLoadingProfilerWindow::OnTimeSelectionChanged(UE::Insights::Timing::ETimeChangedFlags InFlags, double InStartTime, double InEndTime)
+void SLoadingProfilerWindow::OnTimeSelectionChanged(Timing::ETimeChangedFlags InFlags, double InStartTime, double InEndTime)
 {
-	if (InFlags != UE::Insights::Timing::ETimeChangedFlags::Interactive)
+	if (InFlags != Timing::ETimeChangedFlags::Interactive)
 	{
 		if (InStartTime < InEndTime)
 		{
@@ -559,4 +561,6 @@ void SLoadingProfilerWindow::OnTimeSelectionChanged(UE::Insights::Timing::ETimeC
 
 ////////////////////////////////////////////////////////////////////////////////////////////////////
 
+} // namespace UE::Insights::LoadingProfiler
+
 #undef LOCTEXT_NAMESPACE
diff --git a/Engine/Source/Developer/TraceInsights/Private/Insights/LoadingProfiler/Widgets/SLoadingProfilerWindow.h b/Engine/Source/Developer/TraceInsights/Private/Insights/LoadingProfiler/Widgets/SLoadingProfilerWindow.h
--- a/Engine/Source/Developer/TraceInsights/Private/Insights/LoadingProfiler/Widgets/SLoadingProfilerWindow.h
+++ b/Engine/Source/Developer/TraceInsights/Private/Insights/LoadingProfiler/Widgets/SLoadingProfilerWindow.h
@@ -8,14 +8,11 @@
 #include "Insights/ITimingViewSession.h" // for ETimeChangedFlags
 #include "Insights/Widgets/SMajorTabWindow.h"
 
-////////////////////////////////////////////////////////////////////////////////////////////////////
-
 namespace UE::Insights::TimingProfiler { class STimingView; }
+namespace Insights { class SUntypedTableTreeView; }
 
-namespace Insights
+namespace UE::Insights::LoadingProfiler
 {
-	class SUntypedTableTreeView;
-}
 
 ////////////////////////////////////////////////////////////////////////////////////////////////////
 
@@ -33,7 +30,7 @@ struct FLoadingProfilerTabs
 ////////////////////////////////////////////////////////////////////////////////////////////////////
 
 /** Implements the Asset Loading Insights window. */
-class SLoadingProfilerWindow : public Insights::SMajorTabWindow
+class SLoadingProfilerWindow : public ::Insights::SMajorTabWindow
 {
 public:
 	/** Default constructor. */
@@ -57,12 +54,12 @@ public:
 	/** Constructs this widget. */
 	void Construct(const FArguments& InArgs, const TSharedRef<SDockTab>& ConstructUnderMajorTab, const TSharedPtr<SWindow>& ConstructUnderWindow);
 
-	TSharedPtr<UE::Insights::TimingProfiler::STimingView> GetTimingView() const { return TimingView; }
-	TSharedPtr<Insights::SUntypedTableTreeView> GetEventAggregationTreeView() const { return EventAggregationTreeView; }
-	TSharedPtr<Insights::SUntypedTableTreeView> GetObjectTypeAggregationTreeView() const { return ObjectTypeAggregationTreeView; }
-	TSharedPtr<Insights::SUntypedTableTreeView> GetPackageDetailsTreeView() const { return PackageDetailsTreeView; }
-	TSharedPtr<Insights::SUntypedTableTreeView> GetExportDetailsTreeView() const { return ExportDetailsTreeView; }
-	TSharedPtr<Insights::SUntypedTableTreeView> GetRequestsTreeView() const { return RequestsTreeView; }
+	TSharedPtr<TimingProfiler::STimingView> GetTimingView() const { return TimingView; }
+	TSharedPtr<::Insights::SUntypedTableTreeView> GetEventAggregationTreeView() const { return EventAggregationTreeView; }
+	TSharedPtr<::Insights::SUntypedTableTreeView> GetObjectTypeAggregationTreeView() const { return ObjectTypeAggregationTreeView; }
+	TSharedPtr<::Insights::SUntypedTableTreeView> GetPackageDetailsTreeView() const { return PackageDetailsTreeView; }
+	TSharedPtr<::Insights::SUntypedTableTreeView> GetExportDetailsTreeView() const { return ExportDetailsTreeView; }
+	TSharedPtr<::Insights::SUntypedTableTreeView> GetRequestsTreeView() const { return RequestsTreeView; }
 
 protected:
 	virtual const TCHAR* GetAnalyticsEventName() const override;
@@ -90,27 +87,31 @@ private:
 	TSharedRef<SDockTab> SpawnTab_RequestsTreeView(const FSpawnTabArgs& Args);
 	void OnRequestsTreeViewTabClosed(TSharedRef<SDockTab> TabBeingClosed);
 
-	void OnTimeSelectionChanged(UE::Insights::Timing::ETimeChangedFlags InFlags, double InStartTime, double InEndTime);
+	void OnTimeSelectionChanged(Timing::ETimeChangedFlags InFlags, double InStartTime, double InEndTime);
 
 private:
 	/** The Timing view (multi-track) widget */
-	TSharedPtr<UE::Insights::TimingProfiler::STimingView> TimingView;
+	TSharedPtr<TimingProfiler::STimingView> TimingView;
 
 	/** The Event Aggregation tree view widget */
-	TSharedPtr<Insights::SUntypedTableTreeView> EventAggregationTreeView;
+	TSharedPtr<::Insights::SUntypedTableTreeView> EventAggregationTreeView;
 
 	/** The Object Type Aggregation tree view widget */
-	TSharedPtr<Insights::SUntypedTableTreeView> ObjectTypeAggregationTreeView;
+	TSharedPtr<::Insights::SUntypedTableTreeView> ObjectTypeAggregationTreeView;
 
 	/** The Package Details tree view widget */
-	TSharedPtr<Insights::SUntypedTableTreeView> PackageDetailsTreeView;
+	TSharedPtr<::Insights::SUntypedTableTreeView> PackageDetailsTreeView;
 
 	/** The Export Details tree view widget */
-	TSharedPtr<Insights::SUntypedTableTreeView> ExportDetailsTreeView;
+	TSharedPtr<::Insights::SUntypedTableTreeView> ExportDetailsTreeView;
 
 	/** The Requests tree view widget */
-	TSharedPtr<Insights::SUntypedTableTreeView> RequestsTreeView;
+	TSharedPtr<::Insights::SUntypedTableTreeView> RequestsTreeView;
 
 	double SelectionStartTime;
 	double SelectionEndTime;
 };
+
+////////////////////////////////////////////////////////////////////////////////////////////////////
+
+} // namespace UE::Insights::LoadingProfiler
diff --git a/Engine/Source/Developer/TraceInsights/Private/Insights/TaskGraphProfiler/TaskGraphProfilerManager.cpp b/Engine/Source/Developer/TraceInsights/Private/Insights/TaskGraphProfiler/TaskGraphProfilerManager.cpp
--- a/Engine/Source/Developer/TraceInsights/Private/Insights/TaskGraphProfiler/TaskGraphProfilerManager.cpp
+++ b/Engine/Source/Developer/TraceInsights/Private/Insights/TaskGraphProfiler/TaskGraphProfilerManager.cpp
@@ -24,13 +24,13 @@
 #include "Insights/Widgets/STimingProfilerWindow.h"
 #include "Insights/Widgets/STimingView.h"
 
-////////////////////////////////////////////////////////////////////////////////////////////////////
-
-#define LOCTEXT_NAMESPACE "TaskGraphProfilerManager"
+#define LOCTEXT_NAMESPACE "UE::Insights::TaskGraphProfiler"
 
-namespace Insights
+namespace UE::Insights::TaskGraphProfiler
 {
 
+////////////////////////////////////////////////////////////////////////////////////////////////////
+
 const FName FTaskGraphProfilerTabs::TaskTableTreeViewTabID(TEXT("TaskTableTreeView"));
 
 ////////////////////////////////////////////////////////////////////////////////////////////////////
@@ -120,7 +120,7 @@ FTaskGraphProfilerManager::~FTaskGraphProfilerManager()
 
 	if (TaskTimingSharedState.IsValid())
 	{
-		IModularFeatures::Get().UnregisterModularFeature(UE::Insights::Timing::TimingViewExtenderFeatureName, TaskTimingSharedState.Get());
+		IModularFeatures::Get().UnregisterModularFeature(Timing::TimingViewExtenderFeatureName, TaskTimingSharedState.Get());
 	}
 }
 
@@ -154,7 +154,7 @@ bool FTaskGraphProfilerManager::Tick(float DeltaTime)
 			TSharedPtr<FTabManager> TabManagerShared = TimingTabManager.Pin();
 			if (TasksProvider && TasksProvider->GetNumTasks() > 0 && TabManagerShared.IsValid())
 			{
-				TSharedPtr<UE::Insights::TimingProfiler::STimingView> TimingView = GetTimingView();
+				TSharedPtr<TimingProfiler::STimingView> TimingView = GetTimingView();
 				if (!TimingView.IsValid())
 				{
 					return true;
@@ -165,7 +165,7 @@ bool FTaskGraphProfilerManager::Tick(float DeltaTime)
 				if (!TaskTimingSharedState.IsValid())
 				{
 					TaskTimingSharedState = MakeShared<FTaskTimingSharedState>(TimingView.Get());
-					IModularFeatures::Get().RegisterModularFeature(UE::Insights::Timing::TimingViewExtenderFeatureName, TaskTimingSharedState.Get());
+					IModularFeatures::Get().RegisterModularFeature(Timing::TimingViewExtenderFeatureName, TaskTimingSharedState.Get());
 				}
 				TabManagerShared->TryInvokeTab(FTaskGraphProfilerTabs::TaskTableTreeViewTabID);
 			}
@@ -413,7 +413,7 @@ void FTaskGraphProfilerManager::GetRelationsOnCriticalPath(const TraceServices::
 			return;
 		}
 
-		TSharedPtr<UE::Insights::TimingProfiler::STimingView> TimingView = GetTimingView();
+		TSharedPtr<TimingProfiler::STimingView> TimingView = GetTimingView();
 		if (!TimingView.IsValid())
 		{
 			return;
@@ -565,7 +565,7 @@ void FTaskGraphProfilerManager::AddRelation(const FThreadTrackEvent* InSelectedE
 		return;
 	}
 
-	TSharedPtr<UE::Insights::TimingProfiler::STimingView> TimingView = GetTimingView();
+	TSharedPtr<TimingProfiler::STimingView> TimingView = GetTimingView();
 	if (!TimingView.IsValid())
 	{
 		return;
@@ -632,7 +632,7 @@ int32 FTaskGraphProfilerManager::GetRelationDisplayDepth(TSharedPtr<const FThrea
 
 void FTaskGraphProfilerManager::ClearTaskRelations()
 {
-	TSharedPtr<UE::Insights::TimingProfiler::STimingView> TimingView = GetTimingView();
+	TSharedPtr<TimingProfiler::STimingView> TimingView = GetTimingView();
 	if (!TimingView.IsValid())
 	{
 		return;
@@ -653,7 +653,7 @@ int32 FTaskGraphProfilerManager::GetDepthOfTaskExecution(double TaskStartedTime,
 {
 	int32 Depth = -1;
 
-	TSharedPtr<UE::Insights::TimingProfiler::STimingView> TimingView = GetTimingView();
+	TSharedPtr<TimingProfiler::STimingView> TimingView = GetTimingView();
 	if (!TimingView.IsValid())
 	{
 		return Depth;
@@ -880,7 +880,7 @@ void FTaskGraphProfilerManager::OutputWarnings()
 
 ////////////////////////////////////////////////////////////////////////////////////////////////////
 
-TSharedPtr<UE::Insights::TimingProfiler::STimingView> FTaskGraphProfilerManager::GetTimingView()
+TSharedPtr<TimingProfiler::STimingView> FTaskGraphProfilerManager::GetTimingView()
 {
 	using namespace UE::Insights::TimingProfiler;
 	TSharedPtr<STimingProfilerWindow> Window = FTimingProfilerManager::Get()->GetProfilerWindow();
@@ -889,6 +889,6 @@ TSharedPtr<UE::Insights::TimingProfiler::STimingView> FTaskGraphProfilerManager:
 
 ////////////////////////////////////////////////////////////////////////////////////////////////////
 
-} // namespace Insights
+} // namespace UE::Insights::TaskGraphProfiler
 
 #undef LOCTEXT_NAMESPACE
diff --git a/Engine/Source/Developer/TraceInsights/Private/Insights/TaskGraphProfiler/TaskGraphProfilerManager.h b/Engine/Source/Developer/TraceInsights/Private/Insights/TaskGraphProfiler/TaskGraphProfilerManager.h
--- a/Engine/Source/Developer/TraceInsights/Private/Insights/TaskGraphProfiler/TaskGraphProfilerManager.h
+++ b/Engine/Source/Developer/TraceInsights/Private/Insights/TaskGraphProfiler/TaskGraphProfilerManager.h
@@ -23,7 +23,7 @@ class FThreadTrackEvent;
 
 namespace UE::Insights::TimingProfiler { class STimingView; }
 
-namespace Insights
+namespace UE::Insights::TaskGraphProfiler
 {
 
 class FTaskGraphRelation;
@@ -111,7 +111,7 @@ public:
 	FLinearColor GetColorForTaskEvent(ETaskEventType InEvent);
 	uint32 GetColorForTaskEventAsPackedARGB(ETaskEventType InEvent);
 
-	TSharedPtr<::Insights::FTaskTimingSharedState> GetTaskTimingSharedState() { return TaskTimingSharedState;	}
+	TSharedPtr<FTaskTimingSharedState> GetTaskTimingSharedState() { return TaskTimingSharedState;	}
 
 	bool GetShowCriticalPath() const { return bShowCriticalPath; }
 	void SetShowCriticalPath(bool bInValue) { bShowCriticalPath = bInValue; }
@@ -166,7 +166,7 @@ private:
 private:
 	bool bIsInitialized;
 	bool bIsAvailable;
-	UE::Insights::FAvailabilityCheck AvailabilityCheck;
+	FAvailabilityCheck AvailabilityCheck;
 
 	/** The delegate to be invoked when this manager ticks. */
 	FTickerDelegate OnTick;
@@ -182,7 +182,7 @@ private:
 
 	TWeakPtr<FTabManager> TimingTabManager;
 
-	TSharedPtr<::Insights::STaskTableTreeView> TaskTableTreeView;
+	TSharedPtr<STaskTableTreeView> TaskTableTreeView;
 	FLinearColor ColorCode[static_cast<uint32>(ETaskEventType::NumTaskEventTypes)];
 	bool bShowCriticalPath = false;
 	bool bShowTransitions = true;
@@ -197,5 +197,4 @@ private:
 	TSet<FString> HiddenTrackNames;
 };
 
-} // namespace Insights
-
+} // namespace UE::Insights::TaskGraphProfiler
diff --git a/Engine/Source/Developer/TraceInsights/Private/Insights/TaskGraphProfiler/ViewModels/TaskEntry.cpp b/Engine/Source/Developer/TraceInsights/Private/Insights/TaskGraphProfiler/ViewModels/TaskEntry.cpp
--- a/Engine/Source/Developer/TraceInsights/Private/Insights/TaskGraphProfiler/ViewModels/TaskEntry.cpp
+++ b/Engine/Source/Developer/TraceInsights/Private/Insights/TaskGraphProfiler/ViewModels/TaskEntry.cpp
@@ -2,7 +2,10 @@
 
 #include "TaskEntry.h"
 
-namespace Insights
+// TraceServices
+#include "TraceServices/Model/TasksProfiler.h"
+
+namespace UE::Insights::TaskGraphProfiler
 {
 
 ////////////////////////////////////////////////////////////////////////////////////////////////////
@@ -35,4 +38,4 @@ FTaskEntry::FTaskEntry(const TraceServices::FTaskInfo& TaskInfo)
 
 ////////////////////////////////////////////////////////////////////////////////////////////////////
 
-} // namespace Insights
+} // namespace UE::Insights::TaskGraphProfiler
diff --git a/Engine/Source/Developer/TraceInsights/Private/Insights/TaskGraphProfiler/ViewModels/TaskEntry.h b/Engine/Source/Developer/TraceInsights/Private/Insights/TaskGraphProfiler/ViewModels/TaskEntry.h
--- a/Engine/Source/Developer/TraceInsights/Private/Insights/TaskGraphProfiler/ViewModels/TaskEntry.h
+++ b/Engine/Source/Developer/TraceInsights/Private/Insights/TaskGraphProfiler/ViewModels/TaskEntry.h
@@ -2,14 +2,13 @@
 
 #pragma once
 
-#include "CoreMinimal.h"
+#include "CoreTypes.h"
 
 #include "Async/TaskTrace.h"
 
-// TraceServices
-#include "TraceServices/Model/TasksProfiler.h"
+namespace TraceServices { struct FTaskInfo; }
 
-namespace Insights
+namespace UE::Insights::TaskGraphProfiler
 {
 
 ////////////////////////////////////////////////////////////////////////////////////////////////////
@@ -90,4 +89,4 @@ private:
 
 ////////////////////////////////////////////////////////////////////////////////////////////////////
 
-} // namespace Insights
+} // namespace UE::Insights::TaskGraphProfiler
diff --git a/Engine/Source/Developer/TraceInsights/Private/Insights/TaskGraphProfiler/ViewModels/TaskGraphRelation.cpp b/Engine/Source/Developer/TraceInsights/Private/Insights/TaskGraphProfiler/ViewModels/TaskGraphRelation.cpp
--- a/Engine/Source/Developer/TraceInsights/Private/Insights/TaskGraphProfiler/ViewModels/TaskGraphRelation.cpp
+++ b/Engine/Source/Developer/TraceInsights/Private/Insights/TaskGraphProfiler/ViewModels/TaskGraphRelation.cpp
@@ -11,11 +11,11 @@
 #include "Insights/ViewModels/TimingTrackViewport.h"
 #include "Insights/ViewModels/TimingViewDrawHelper.h"
 
-////////////////////////////////////////////////////////////////////////////////////////////////////
-
-namespace Insights
+namespace UE::Insights::TaskGraphProfiler
 {
 
+////////////////////////////////////////////////////////////////////////////////////////////////////
+
 INSIGHTS_IMPLEMENT_RTTI(FTaskGraphRelation)
 
 ////////////////////////////////////////////////////////////////////////////////////////////////////
@@ -31,7 +31,7 @@ FTaskGraphRelation::FTaskGraphRelation(double InSourceTime, int32 InSourceThread
 
 ////////////////////////////////////////////////////////////////////////////////////////////////////
 
-void FTaskGraphRelation::Draw(const UE::Insights::FDrawContext& DrawContext, const FTimingTrackViewport& Viewport, const ITimingViewDrawHelper& Helper, const ITimingEventRelation::EDrawFilter Filter)
+void FTaskGraphRelation::Draw(const FDrawContext& DrawContext, const FTimingTrackViewport& Viewport, const ITimingViewDrawHelper& Helper, const ITimingEventRelation::EDrawFilter Filter)
 {
 	int32 LayerId = Helper.GetRelationLayerId();
 
@@ -77,12 +77,13 @@ void FTaskGraphRelation::Draw(const UE::Insights::FDrawContext& DrawContext, con
 		return;
 	}
 
+	const int32 MaxEventDepth = (int32)TimingProfiler::FTimingProfilerManager::Get()->GetEventDepthLimit() - 1;
+
 	float Y1 = 0.0f;
 	float Y2 = 0.0f;
 	if (SourceTrackShared->IsVisible())
 	{
-		using namespace UE::Insights::TimingProfiler;
-		int32 ActualSourceDepth = FMath::Min(SourceDepth, (int32)FTimingProfilerManager::Get()->GetEventDepthLimit() - 1);
+		int32 ActualSourceDepth = FMath::Min(SourceDepth, MaxEventDepth);
 		Y1 = SourceTrackShared->GetPosY();
 		Y1 += Viewport.GetLayout().GetLaneY(ActualSourceDepth) + Viewport.GetLayout().EventH / 2.0f;
 		if (SourceTrackShared->GetChildTrack() && SourceTrackShared->GetChildTrack()->GetHeight() > 0.0f)
@@ -97,8 +98,7 @@ void FTaskGraphRelation::Draw(const UE::Insights::FDrawContext& DrawContext, con
 
 	if (TargetTrackShared->IsVisible())
 	{
-		using namespace UE::Insights::TimingProfiler;
-		int32 ActualTargetDepth = FMath::Min(TargetDepth, (int32)FTimingProfilerManager::Get()->GetEventDepthLimit() - 1);
+		int32 ActualTargetDepth = FMath::Min(TargetDepth, MaxEventDepth);
 		Y2 = TargetTrackShared->GetPosY();
 		Y2 += Viewport.GetLayout().GetLaneY(ActualTargetDepth) + Viewport.GetLayout().EventH / 2.0f;
 		if (TargetTrackShared->GetChildTrack() && TargetTrackShared->GetChildTrack()->GetHeight() > 0.0f)
@@ -188,4 +188,4 @@ void FTaskGraphRelation::Draw(const UE::Insights::FDrawContext& DrawContext, con
 
 ////////////////////////////////////////////////////////////////////////////////////////////////////
 
-} // namespace Insights
+} // namespace UE::Insights::TaskGraphProfiler
diff --git a/Engine/Source/Developer/TraceInsights/Private/Insights/TaskGraphProfiler/ViewModels/TaskGraphRelation.h b/Engine/Source/Developer/TraceInsights/Private/Insights/TaskGraphProfiler/ViewModels/TaskGraphRelation.h
--- a/Engine/Source/Developer/TraceInsights/Private/Insights/TaskGraphProfiler/ViewModels/TaskGraphRelation.h
+++ b/Engine/Source/Developer/TraceInsights/Private/Insights/TaskGraphProfiler/ViewModels/TaskGraphRelation.h
@@ -2,11 +2,11 @@
 
 #pragma once
 
-#include "CoreMinimal.h"
+#include "CoreTypes.h"
 
 #include "Insights/ViewModels/ITimingEvent.h"
 
-namespace Insights
+namespace UE::Insights::TaskGraphProfiler
 {
 
 enum class ETaskEventType : uint32;
@@ -18,7 +18,7 @@ class FTaskGraphRelation : public ITimingEventRelation
 public:
 	FTaskGraphRelation(double InSourceTime, int32 InSourceThreadId, double InTargetTime, int32 InTargetThreadId, ETaskEventType InType);
 
-	virtual void Draw(const UE::Insights::FDrawContext& DrawContext, const FTimingTrackViewport& Viewport, const ITimingViewDrawHelper& Helper, const ITimingEventRelation::EDrawFilter Filter) override;
+	virtual void Draw(const FDrawContext& DrawContext, const FTimingTrackViewport& Viewport, const ITimingViewDrawHelper& Helper, const ITimingEventRelation::EDrawFilter Filter) override;
 
 	void SetSourceTrack(TSharedPtr<const FBaseTimingTrack> InSourceTrack) { SourceTrack = InSourceTrack; }
 	TSharedPtr<const FBaseTimingTrack> GetSourceTrack() { return SourceTrack.Pin(); }
@@ -51,4 +51,4 @@ private:
 	TWeakPtr<const FBaseTimingTrack> TargetTrack;
 };
 
-} // namespace Insights
+} // namespace UE::Insights::TaskGraphProfiler
diff --git a/Engine/Source/Developer/TraceInsights/Private/Insights/TaskGraphProfiler/ViewModels/TaskNode.cpp b/Engine/Source/Developer/TraceInsights/Private/Insights/TaskGraphProfiler/ViewModels/TaskNode.cpp
--- a/Engine/Source/Developer/TraceInsights/Private/Insights/TaskGraphProfiler/ViewModels/TaskNode.cpp
+++ b/Engine/Source/Developer/TraceInsights/Private/Insights/TaskGraphProfiler/ViewModels/TaskNode.cpp
@@ -2,13 +2,9 @@
 
 #include "TaskNode.h"
 
-#define LOCTEXT_NAMESPACE "Insights::FTaskNode"
-
-namespace Insights
+namespace UE::Insights::TaskGraphProfiler
 {
 
 INSIGHTS_IMPLEMENT_RTTI(FTaskNode)
 
-} // namespace Insights
-
-#undef LOCTEXT_NAMESPACE
+} // namespace UE::Insights::TaskGraphProfiler
diff --git a/Engine/Source/Developer/TraceInsights/Private/Insights/TaskGraphProfiler/ViewModels/TaskNode.h b/Engine/Source/Developer/TraceInsights/Private/Insights/TaskGraphProfiler/ViewModels/TaskNode.h
--- a/Engine/Source/Developer/TraceInsights/Private/Insights/TaskGraphProfiler/ViewModels/TaskNode.h
+++ b/Engine/Source/Developer/TraceInsights/Private/Insights/TaskGraphProfiler/ViewModels/TaskNode.h
@@ -2,7 +2,7 @@
 
 #pragma once
 
-#include "CoreMinimal.h"
+#include "CoreTypes.h"
 
 // TraceInsightsCore
 #include "InsightsCore/Table/ViewModels/TableTreeNode.h"
@@ -13,7 +13,7 @@
 
 ////////////////////////////////////////////////////////////////////////////////////////////////////
 
-namespace Insights
+namespace UE::Insights::TaskGraphProfiler
 {
 
 ////////////////////////////////////////////////////////////////////////////////////////////////////
@@ -36,26 +36,26 @@ typedef TWeakPtr<class FTaskNode> FTaskNodeWeak;
 /**
  * Class used to store information about a task node (used in the STaskTreeView).
  */
-class FTaskNode : public UE::Insights::FTableTreeNode
+class FTaskNode : public FTableTreeNode
 {
-	INSIGHTS_DECLARE_RTTI(FTaskNode, UE::Insights::FTableTreeNode)
+	INSIGHTS_DECLARE_RTTI(FTaskNode, FTableTreeNode)
 
 public:
 	/** Initialization constructor for the Task node. */
 	explicit FTaskNode(const FName InName, TWeakPtr<FTaskTable> InParentTable, int32 InRowIndex)
-		: UE::Insights::FTableTreeNode(InName, InParentTable, InRowIndex)
+		: FTableTreeNode(InName, InParentTable, InRowIndex)
 	{
 	}
 
 	/** Initialization constructor for the group node. */
 	explicit FTaskNode(const FName InGroupName, TWeakPtr<FTaskTable> InParentTable)
-		: UE::Insights::FTableTreeNode(InGroupName, InParentTable)
+		: FTableTreeNode(InGroupName, InParentTable)
 	{
 	}
 
 	FTaskTable& GetTaskTableChecked() const
 	{
-		const TSharedPtr<UE::Insights::FTable>& TablePin = GetParentTable().Pin();
+		const TSharedPtr<FTable>& TablePin = GetParentTable().Pin();
 		check(TablePin.IsValid());
 		return *StaticCastSharedPtr<FTaskTable>(TablePin);
 	}
@@ -67,4 +67,4 @@ public:
 
 ////////////////////////////////////////////////////////////////////////////////////////////////////
 
-} // namespace Insights
+} // namespace UE::Insights::TaskGraphProfiler
diff --git a/Engine/Source/Developer/TraceInsights/Private/Insights/TaskGraphProfiler/ViewModels/TaskTable.cpp b/Engine/Source/Developer/TraceInsights/Private/Insights/TaskGraphProfiler/ViewModels/TaskTable.cpp
--- a/Engine/Source/Developer/TraceInsights/Private/Insights/TaskGraphProfiler/ViewModels/TaskTable.cpp
+++ b/Engine/Source/Developer/TraceInsights/Private/Insights/TaskGraphProfiler/ViewModels/TaskTable.cpp
@@ -2,6 +2,9 @@
 
 #include "TaskTable.h"
 
+// TraceServices
+#include "TraceServices/Model/TasksProfiler.h"
+
 // TraceInsightsCore
 #include "InsightsCore/Common/TimeUtils.h"
 #include "InsightsCore/Table/ViewModels/TableCellValue.h"
@@ -14,9 +17,9 @@
 #include "Insights/TaskGraphProfiler/ViewModels/TaskNode.h"
 #include "Insights/TaskGraphProfiler/ViewModels/TaskTable.h"
 
-#define LOCTEXT_NAMESPACE "Insights::FTaskTable"
+#define LOCTEXT_NAMESPACE "UE::Insights::TaskGraphProfiler::FTaskTable"
 
-namespace Insights
+namespace UE::Insights::TaskGraphProfiler
 {
 
 ////////////////////////////////////////////////////////////////////////////////////////////////////
@@ -44,17 +47,17 @@ const FName FTaskTableColumns::NumPrerequisitesColumnId(TEXT("NumPrerequisites")
 
 ////////////////////////////////////////////////////////////////////////////////////////////////////
 
-typedef UE::Insights::FTableCellValue (*TaskFieldGetter) (const UE::Insights::FTableColumn&, const FTaskEntry&);
+typedef FTableCellValue (*TaskFieldGetter) (const FTableColumn&, const FTaskEntry&);
 
 template<TaskFieldGetter Getter>
-class FTaskColumnValueGetter : public UE::Insights::FTableCellValueGetter
+class FTaskColumnValueGetter : public FTableCellValueGetter
 {
 public:
-	virtual const TOptional<UE::Insights::FTableCellValue> GetValue(const UE::Insights::FTableColumn& Column, const UE::Insights::FBaseTreeNode& Node) const override
+	virtual const TOptional<FTableCellValue> GetValue(const FTableColumn& Column, const FBaseTreeNode& Node) const override
 	{
 		if (Node.IsGroup())
 		{
-			const UE::Insights::FTableTreeNode& NodePtr = static_cast<const UE::Insights::FTableTreeNode&>(Node);
+			const FTableTreeNode& NodePtr = static_cast<const FTableTreeNode&>(Node);
 			if (NodePtr.HasAggregatedValue(Column.GetId()))
 			{
 				return NodePtr.GetAggregatedValue(Column.GetId());
@@ -70,7 +73,7 @@ public:
 			}
 		}
 
-		return TOptional<UE::Insights::FTableCellValue>();
+		return TOptional<FTableCellValue>();
 	}
 };
 
@@ -90,9 +93,6 @@ double GetRelativeValue(double ValueA, double ValueB)
 
 struct DefaultTaskFieldGetterFuncts
 {
-	using FTableCellValue = UE::Insights::FTableCellValue;
-	using FTableColumn = UE::Insights::FTableColumn;
-
 	static FTableCellValue GetDebugName(const FTableColumn& Column, const FTaskEntry& Task) { return FTableCellValue(Task.GetDebugName()); }
 	static FTableCellValue GetCreatedTimestamp(const FTableColumn& Column, const FTaskEntry& Task) { return FTableCellValue(Task.GetCreatedTimestamp());	}
 	static FTableCellValue GetCreatedThreadId(const FTableColumn& Column, const FTaskEntry& Task) { return FTableCellValue((int64)Task.GetCreatedThreadId());	}
@@ -126,9 +126,6 @@ struct DefaultTaskFieldGetterFuncts
 
 struct RelativeToPreviousTaskFieldGetterFuncts
 {
-	using FTableCellValue = UE::Insights::FTableCellValue;
-	using FTableColumn = UE::Insights::FTableColumn;
-
 	static FTableCellValue GetCreatedTimestamp(const FTableColumn& Column, const FTaskEntry& Task) { return FTableCellValue(Task.GetCreatedTimestamp()); }
 	static FTableCellValue GetLaunchedTimestamp(const FTableColumn& Column, const FTaskEntry& Task) { return FTableCellValue(GetRelativeValue(Task.GetLaunchedTimestamp(), Task.GetCreatedTimestamp())); }
 	static FTableCellValue GetScheduledTimestamp(const FTableColumn& Column, const FTaskEntry& Task) { return FTableCellValue(GetRelativeValue(Task.GetScheduledTimestamp(), Task.GetLaunchedTimestamp())); }
@@ -142,9 +139,6 @@ struct RelativeToPreviousTaskFieldGetterFuncts
 
 struct RelativeToCreatedTaskFieldGetterFuncts
 {
-	using FTableCellValue = UE::Insights::FTableCellValue;
-	using FTableColumn = UE::Insights::FTableColumn;
-
 	static FTableCellValue GetCreatedTimestamp(const FTableColumn& Column, const FTaskEntry& Task) { return FTableCellValue(Task.GetCreatedTimestamp()); }
 	static FTableCellValue GetLaunchedTimestamp(const FTableColumn& Column, const FTaskEntry& Task) { return FTableCellValue(GetRelativeValue(Task.GetLaunchedTimestamp(), Task.GetCreatedTimestamp())); }
 	static FTableCellValue GetScheduledTimestamp(const FTableColumn& Column, const FTaskEntry& Task) { return FTableCellValue(GetRelativeValue(Task.GetScheduledTimestamp(), Task.GetCreatedTimestamp())); }
@@ -156,24 +150,24 @@ struct RelativeToCreatedTaskFieldGetterFuncts
 
 ////////////////////////////////////////////////////////////////////////////////////////////////////
 
-class FTaskDoubleValueFormatterAsTimeAuto : public UE::Insights::FTableCellValueFormatter
+class FTaskDoubleValueFormatterAsTimeAuto : public FTableCellValueFormatter
 {
 public:
-	virtual FText FormatValue(const TOptional<UE::Insights::FTableCellValue>& InValue) const override
+	virtual FText FormatValue(const TOptional<FTableCellValue>& InValue) const override
 	{
 		if (InValue.IsSet())
 		{
 			const double Value = InValue.GetValue().Double;
 			if (Value != TraceServices::FTaskInfo::InvalidTimestamp)
 			{
-				return FText::FromString(UE::Insights::FormatTimeAuto(Value));
+				return FText::FromString(FormatTimeAuto(Value));
 			}
 		}
 
 		return FText::FromString(TEXT("N/A"));
 	}
 
-	virtual FText FormatValueForTooltip(const TOptional<UE::Insights::FTableCellValue>& InValue) const override
+	virtual FText FormatValueForTooltip(const TOptional<FTableCellValue>& InValue) const override
 	{
 		if (InValue.IsSet())
 		{
@@ -184,7 +178,7 @@ public:
 			}
 			else if (Value != TraceServices::FTaskInfo::InvalidTimestamp)
 			{
-				return FText::FromString(FString::Printf(TEXT("%f (%s)"), Value, *UE::Insights::FormatTimeAuto(Value)));
+				return FText::FromString(FString::Printf(TEXT("%f (%s)"), Value, *FormatTimeAuto(Value)));
 			}
 		}
 
@@ -866,6 +860,6 @@ void FTaskTable::SwitchToRelativeToCreatedTimestamps()
 
 ////////////////////////////////////////////////////////////////////////////////////////////////////
 
-} // namespace Insights
+} // namespace UE::Insights::TaskGraphProfiler
 
 #undef LOCTEXT_NAMESPACE
diff --git a/Engine/Source/Developer/TraceInsights/Private/Insights/TaskGraphProfiler/ViewModels/TaskTable.h b/Engine/Source/Developer/TraceInsights/Private/Insights/TaskGraphProfiler/ViewModels/TaskTable.h
--- a/Engine/Source/Developer/TraceInsights/Private/Insights/TaskGraphProfiler/ViewModels/TaskTable.h
+++ b/Engine/Source/Developer/TraceInsights/Private/Insights/TaskGraphProfiler/ViewModels/TaskTable.h
@@ -3,6 +3,7 @@
 #pragma once
 
 #include "CoreTypes.h"
+
 #include "Containers/Array.h"
 #include "UObject/NameTypes.h"
 
@@ -14,7 +15,7 @@
 
 namespace UE::Insights { class FTableColumn; }
 
-namespace Insights
+namespace UE::Insights::TaskGraphProfiler
 {
 
 ////////////////////////////////////////////////////////////////////////////////////////////////////
@@ -45,7 +46,7 @@ struct FTaskTableColumns
 
 ////////////////////////////////////////////////////////////////////////////////////////////////////
 
-class FTaskTable : public UE::Insights::FTable
+class FTaskTable : public FTable
 {
 public:
 	FTaskTable();
@@ -73,4 +74,4 @@ private:
 
 ////////////////////////////////////////////////////////////////////////////////////////////////////
 
-} // namespace Insights
+} // namespace UE::Insights::TaskGraphProfiler
diff --git a/Engine/Source/Developer/TraceInsights/Private/Insights/TaskGraphProfiler/ViewModels/TaskTimingTrack.cpp b/Engine/Source/Developer/TraceInsights/Private/Insights/TaskGraphProfiler/ViewModels/TaskTimingTrack.cpp
--- a/Engine/Source/Developer/TraceInsights/Private/Insights/TaskGraphProfiler/ViewModels/TaskTimingTrack.cpp
+++ b/Engine/Source/Developer/TraceInsights/Private/Insights/TaskGraphProfiler/ViewModels/TaskTimingTrack.cpp
@@ -28,10 +28,11 @@
 #include "Insights/Widgets/STimingProfilerWindow.h"
 #include "Insights/Widgets/STimingView.h"
 
-#define LOCTEXT_NAMESPACE "TaskTimingTrack"
+#define LOCTEXT_NAMESPACE "UE::Insights::TaskGraphProfiler"
 
-namespace Insights
+namespace UE::Insights::TaskGraphProfiler
 {
+
 ////////////////////////////////////////////////////////////////////////////////////////////////////
 // FTaskTimingStateCommands
 ////////////////////////////////////////////////////////////////////////////////////////////////////
@@ -57,57 +58,57 @@ public:
 	virtual void RegisterCommands() override
 	{
 		UI_COMMAND(Command_ShowTaskTransitions,
-				   "Show Task Transitions",
-				   "Show/hide transitions between the stages of the current task (for a selected cpu timing event.)",
-				    EUserInterfaceActionType::ToggleButton,
-				    FInputChord(EKeys::T));
+			"Show Task Transitions",
+			"Show/hide transitions between the stages of the current task (for a selected CPU timing event.)",
+			EUserInterfaceActionType::ToggleButton,
+			FInputChord(EKeys::T));
 
 		UI_COMMAND(Command_ShowTaskConnections,
-				   "Show Task Connections",
-				   "Show/hide conections between:\nThe current task's prerequisites completed time and the current task's started time.\nThe current task's completed time and the current task's subsequents started time.\nThe current task's nested tasks added time and their started time.",
-				   EUserInterfaceActionType::ToggleButton, FInputChord());
+			"Show Task Connections",
+			"Show/hide connections between:\nThe current task's prerequisites completed time and the current task's started time.\nThe current task's completed time and the current task's subsequents started time.\nThe current task's nested tasks added time and their started time.",
+			EUserInterfaceActionType::ToggleButton, FInputChord());
 
 		UI_COMMAND(Command_ShowTaskPrerequisites,
-				   "Show Transitions of Prerequisites",
-				   "Show/hide stage transitions for the current task's prerequisites.",
-				   EUserInterfaceActionType::ToggleButton,
-				   FInputChord(EKeys::P));
+			"Show Transitions of Prerequisites",
+			"Show/hide stage transitions for the current task's prerequisites.",
+			EUserInterfaceActionType::ToggleButton,
+			FInputChord(EKeys::P));
 
 		UI_COMMAND(Command_ShowTaskSubsequents,
-				   "Show Transitions of Subsequents",
-				   "Show/hide stage transitions for the current task's subsequents.",
-				   EUserInterfaceActionType::ToggleButton,
-				   FInputChord(EKeys::S));
+			"Show Transitions of Subsequents",
+			"Show/hide stage transitions for the current task's subsequents.",
+			EUserInterfaceActionType::ToggleButton,
+			FInputChord(EKeys::S));
 
 		UI_COMMAND(Command_ShowParentTasks,
-				   "Show Transitions of Parent Tasks",
-				   "Show/hide stage transitions for the current task's parent tasks.",
-				   EUserInterfaceActionType::ToggleButton,
-				   FInputChord(EKeys::R));
+			"Show Transitions of Parent Tasks",
+			"Show/hide stage transitions for the current task's parent tasks.",
+			EUserInterfaceActionType::ToggleButton,
+			FInputChord(EKeys::R));
 
 		UI_COMMAND(Command_ShowNestedTasks,
-				   "Show Transitions of Nested Tasks",
-				   "Show/hide stage transitions for the current task's nested tasks.",
-				   EUserInterfaceActionType::ToggleButton,
-				   FInputChord(EKeys::N));
+			"Show Transitions of Nested Tasks",
+			"Show/hide stage transitions for the current task's nested tasks.",
+			EUserInterfaceActionType::ToggleButton,
+			FInputChord(EKeys::N));
 
 		UI_COMMAND(Command_ShowCriticalPath,
-				   "Show Task Critical Path",
-				   "Show/hide relations representing the critical path containing the current task.",
-				   EUserInterfaceActionType::ToggleButton,
-				   FInputChord());
+			"Show Task Critical Path",
+			"Show/hide relations representing the critical path containing the current task.",
+			EUserInterfaceActionType::ToggleButton,
+			FInputChord());
 
 		UI_COMMAND(Command_ShowTaskTrack,
-				   "Show Task Overview Track",
-				   "Show/hide the Task Overview Track when a task is selected.",
-				   EUserInterfaceActionType::ToggleButton,
-				   FInputChord());
+			"Show Task Overview Track",
+			"Show/hide the Task Overview Track when a task is selected.",
+			EUserInterfaceActionType::ToggleButton,
+			FInputChord());
 
 		UI_COMMAND(Command_ShowDetailedTaskTrackInfo,
-				   "Show Detailed Info on the Task Overview Track",
-				   "Show the current task's prerequisites/nested tasks/subsequents in the Task Overview Track.",
-				   EUserInterfaceActionType::ToggleButton,
-				   FInputChord());
+			"Show Detailed Info on the Task Overview Track",
+			"Show the current task's prerequisites/nested tasks/subsequents in the Task Overview Track.",
+			EUserInterfaceActionType::ToggleButton,
+			FInputChord());
 	}
 	UE_ENABLE_OPTIMIZATION_SHIP
 
@@ -126,14 +127,14 @@ public:
 // FTaskTimingSharedState
 ////////////////////////////////////////////////////////////////////////////////////////////////////
 
-FTaskTimingSharedState::FTaskTimingSharedState(UE::Insights::TimingProfiler::STimingView* InTimingView)
+FTaskTimingSharedState::FTaskTimingSharedState(TimingProfiler::STimingView* InTimingView)
 	: TimingViewSession(InTimingView)
 {
 }
 
 ////////////////////////////////////////////////////////////////////////////////////////////////////
 
-void FTaskTimingSharedState::OnBeginSession(UE::Insights::Timing::ITimingViewSession& InSession)
+void FTaskTimingSharedState::OnBeginSession(Timing::ITimingViewSession& InSession)
 {
 	if (&InSession != TimingViewSession)
 	{
@@ -152,7 +153,7 @@ void FTaskTimingSharedState::OnBeginSession(UE::Insights::Timing::ITimingViewSes
 
 ////////////////////////////////////////////////////////////////////////////////////////////////////
 
-void FTaskTimingSharedState::OnEndSession(UE::Insights::Timing::ITimingViewSession& InSession)
+void FTaskTimingSharedState::OnEndSession(Timing::ITimingViewSession& InSession)
 {
 	if (&InSession != TimingViewSession)
 	{
@@ -166,7 +167,7 @@ void FTaskTimingSharedState::OnEndSession(UE::Insights::Timing::ITimingViewSessi
 
 ////////////////////////////////////////////////////////////////////////////////////////////////////
 
-void FTaskTimingSharedState::Tick(UE::Insights::Timing::ITimingViewSession& InSession, const TraceServices::IAnalysisSession& InAnalysisSession)
+void FTaskTimingSharedState::Tick(Timing::ITimingViewSession& InSession, const TraceServices::IAnalysisSession& InAnalysisSession)
 {
 	if (&InSession != TimingViewSession)
 	{
@@ -175,7 +176,7 @@ void FTaskTimingSharedState::Tick(UE::Insights::Timing::ITimingViewSession& InSe
 
 	if (!TaskTrack.IsValid() && FTaskGraphProfilerManager::Get()->GetIsAvailable())
 	{
-		TSharedPtr<UE::Insights::TimingProfiler::STimingView> TimingView = GetTimingView();
+		TSharedPtr<TimingProfiler::STimingView> TimingView = GetTimingView();
 
 		InitCommandList(TimingView);
 
@@ -194,7 +195,7 @@ void FTaskTimingSharedState::Tick(UE::Insights::Timing::ITimingViewSession& InSe
 	{
 		bResetOnNextTick = false;
 
-		TSharedPtr<UE::Insights::TimingProfiler::STimingView> TimingView = GetTimingView();
+		TSharedPtr<TimingProfiler::STimingView> TimingView = GetTimingView();
 		if (TimingView.IsValid() &&
 			!TimingView->GetSelectedEvent().IsValid() &&
 			(!TimingView->GetSelectedTrack().IsValid() || TimingView->GetSelectedTrack().Get() != TaskTrack.Get()))
@@ -207,7 +208,7 @@ void FTaskTimingSharedState::Tick(UE::Insights::Timing::ITimingViewSession& InSe
 
 ////////////////////////////////////////////////////////////////////////////////////////////////////
 
-TSharedPtr<UE::Insights::TimingProfiler::STimingView> FTaskTimingSharedState::GetTimingView()
+TSharedPtr<TimingProfiler::STimingView> FTaskTimingSharedState::GetTimingView()
 {
 	using namespace UE::Insights::TimingProfiler;
 	TSharedPtr<STimingProfilerWindow> Window = FTimingProfilerManager::Get()->GetProfilerWindow();
@@ -216,7 +217,7 @@ TSharedPtr<UE::Insights::TimingProfiler::STimingView> FTaskTimingSharedState::Ge
 
 ////////////////////////////////////////////////////////////////////////////////////////////////////
 
-void FTaskTimingSharedState::ExtendFilterMenu(UE::Insights::Timing::ITimingViewSession& InSession, FMenuBuilder& InOutMenuBuilder)
+void FTaskTimingSharedState::ExtendFilterMenu(Timing::ITimingViewSession& InSession, FMenuBuilder& InOutMenuBuilder)
 {
 	if (&InSession != TimingViewSession)
 	{
@@ -236,7 +237,7 @@ void FTaskTimingSharedState::SetTaskId(TaskTrace::FId InTaskId)
 
 ////////////////////////////////////////////////////////////////////////////////////////////////////
 
-void FTaskTimingSharedState::ExtendOtherTracksFilterMenu(UE::Insights::Timing::ITimingViewSession& InSession, FMenuBuilder& InMenuBuilder)
+void FTaskTimingSharedState::ExtendOtherTracksFilterMenu(Timing::ITimingViewSession& InSession, FMenuBuilder& InMenuBuilder)
 {
 	if (&InSession != TimingViewSession)
 	{
@@ -347,7 +348,7 @@ void FTaskTimingSharedState::BuildTasksSubMenu(FMenuBuilder& MenuBuilder)
 
 ////////////////////////////////////////////////////////////////////////////////////////////////////
 
-void FTaskTimingSharedState::InitCommandList(TSharedPtr<UE::Insights::TimingProfiler::STimingView> TimingView)
+void FTaskTimingSharedState::InitCommandList(TSharedPtr<TimingProfiler::STimingView> TimingView)
 {
 	if (!ensure(TimingView.IsValid()))
 	{
@@ -422,7 +423,7 @@ void FTaskTimingSharedState::InitCommandList(TSharedPtr<UE::Insights::TimingProf
 
 void FTaskTimingSharedState::ContextMenu_ShowTaskTransitions_Execute()
 {
-	TSharedPtr<Insights::FTaskGraphProfilerManager> TaskGraphManager = Insights::FTaskGraphProfilerManager::Get();
+	TSharedPtr<FTaskGraphProfilerManager> TaskGraphManager = FTaskGraphProfilerManager::Get();
 	if (TaskGraphManager.IsValid() && TaskGraphManager->GetIsAvailable())
 	{
 		TaskGraphManager->SetShowTransitions(!TaskGraphManager->GetShowTransitions());
@@ -434,14 +435,14 @@ void FTaskTimingSharedState::ContextMenu_ShowTaskTransitions_Execute()
 
 bool FTaskTimingSharedState::ContextMenu_ShowTaskTransitions_CanExecute()
 {
-	return Insights::FTaskGraphProfilerManager::Get().IsValid() && Insights::FTaskGraphProfilerManager::Get()->GetIsAvailable();
+	return FTaskGraphProfilerManager::Get().IsValid() && FTaskGraphProfilerManager::Get()->GetIsAvailable();
 }
 
 ////////////////////////////////////////////////////////////////////////////////////////////////////
 
 bool FTaskTimingSharedState::ContextMenu_ShowTaskTransitions_IsChecked()
 {
-	TSharedPtr<Insights::FTaskGraphProfilerManager> TaskGraphManager = Insights::FTaskGraphProfilerManager::Get();
+	TSharedPtr<FTaskGraphProfilerManager> TaskGraphManager = FTaskGraphProfilerManager::Get();
 	return TaskGraphManager.IsValid() && TaskGraphManager->GetIsAvailable() && TaskGraphManager->GetShowTransitions();
 }
 
@@ -450,7 +451,7 @@ bool FTaskTimingSharedState::ContextMenu_ShowTaskTransitions_IsChecked()
 
 void FTaskTimingSharedState::ContextMenu_ShowTaskConnections_Execute()
 {
-	TSharedPtr<Insights::FTaskGraphProfilerManager> TaskGraphManager = Insights::FTaskGraphProfilerManager::Get();
+	TSharedPtr<FTaskGraphProfilerManager> TaskGraphManager = FTaskGraphProfilerManager::Get();
 	if (TaskGraphManager.IsValid() && TaskGraphManager->GetIsAvailable())
 	{
 		TaskGraphManager->SetShowConnections(!TaskGraphManager->GetShowConnections());
@@ -462,14 +463,14 @@ void FTaskTimingSharedState::ContextMenu_ShowTaskConnections_Execute()
 
 bool FTaskTimingSharedState::ContextMenu_ShowTaskConnections_CanExecute()
 {
-	return Insights::FTaskGraphProfilerManager::Get().IsValid() && Insights::FTaskGraphProfilerManager::Get()->GetIsAvailable();
+	return FTaskGraphProfilerManager::Get().IsValid() && FTaskGraphProfilerManager::Get()->GetIsAvailable();
 }
 
 ////////////////////////////////////////////////////////////////////////////////////////////////////
 
 bool FTaskTimingSharedState::ContextMenu_ShowTaskConnections_IsChecked()
 {
-	TSharedPtr<Insights::FTaskGraphProfilerManager> TaskGraphManager = Insights::FTaskGraphProfilerManager::Get();
+	TSharedPtr<FTaskGraphProfilerManager> TaskGraphManager = FTaskGraphProfilerManager::Get();
 	return TaskGraphManager.IsValid() && TaskGraphManager->GetIsAvailable() && TaskGraphManager->GetShowConnections();
 }
 
@@ -478,14 +479,14 @@ bool FTaskTimingSharedState::ContextMenu_ShowTaskConnections_IsChecked()
 
 bool FTaskTimingSharedState::ContextMenu_ShowTaskPrerequisites_CanExecute()
 {
-	return Insights::FTaskGraphProfilerManager::Get().IsValid() && Insights::FTaskGraphProfilerManager::Get()->GetIsAvailable();
+	return FTaskGraphProfilerManager::Get().IsValid() && FTaskGraphProfilerManager::Get()->GetIsAvailable();
 }
 
 ////////////////////////////////////////////////////////////////////////////////////////////////////
 
 bool FTaskTimingSharedState::ContextMenu_ShowTaskPrerequisites_IsChecked()
 {
-	TSharedPtr<Insights::FTaskGraphProfilerManager> TaskGraphManager = Insights::FTaskGraphProfilerManager::Get();
+	TSharedPtr<FTaskGraphProfilerManager> TaskGraphManager = FTaskGraphProfilerManager::Get();
 	return TaskGraphManager.IsValid() && TaskGraphManager->GetIsAvailable() && TaskGraphManager->GetShowPrerequisites();
 }
 
@@ -493,7 +494,7 @@ bool FTaskTimingSharedState::ContextMenu_ShowTaskPrerequisites_IsChecked()
 
 void FTaskTimingSharedState::ContextMenu_ShowTaskPrerequisites_Execute()
 {
-	TSharedPtr<Insights::FTaskGraphProfilerManager> TaskGraphManager = Insights::FTaskGraphProfilerManager::Get();
+	TSharedPtr<FTaskGraphProfilerManager> TaskGraphManager = FTaskGraphProfilerManager::Get();
 	if (TaskGraphManager.IsValid() && TaskGraphManager->GetIsAvailable())
 	{
 		TaskGraphManager->SetShowPrerequisites(!TaskGraphManager->GetShowPrerequisites());
@@ -506,14 +507,14 @@ void FTaskTimingSharedState::ContextMenu_ShowTaskPrerequisites_Execute()
 
 bool FTaskTimingSharedState::ContextMenu_ShowTaskSubsequents_CanExecute()
 {
-	return Insights::FTaskGraphProfilerManager::Get().IsValid() && Insights::FTaskGraphProfilerManager::Get()->GetIsAvailable();
+	return FTaskGraphProfilerManager::Get().IsValid() && FTaskGraphProfilerManager::Get()->GetIsAvailable();
 }
 
 ////////////////////////////////////////////////////////////////////////////////////////////////////
 
 bool FTaskTimingSharedState::ContextMenu_ShowTaskSubsequents_IsChecked()
 {
-	TSharedPtr<Insights::FTaskGraphProfilerManager> TaskGraphManager = Insights::FTaskGraphProfilerManager::Get();
+	TSharedPtr<FTaskGraphProfilerManager> TaskGraphManager = FTaskGraphProfilerManager::Get();
 	return TaskGraphManager.IsValid() && TaskGraphManager->GetIsAvailable() && TaskGraphManager->GetShowSubsequents();
 }
 
@@ -521,7 +522,7 @@ bool FTaskTimingSharedState::ContextMenu_ShowTaskSubsequents_IsChecked()
 
 void FTaskTimingSharedState::ContextMenu_ShowTaskSubsequents_Execute()
 {
-	TSharedPtr<Insights::FTaskGraphProfilerManager> TaskGraphManager = Insights::FTaskGraphProfilerManager::Get();
+	TSharedPtr<FTaskGraphProfilerManager> TaskGraphManager = FTaskGraphProfilerManager::Get();
 	if (TaskGraphManager.IsValid() && TaskGraphManager->GetIsAvailable())
 	{
 		TaskGraphManager->SetShowSubsequents(!TaskGraphManager->GetShowSubsequents());
@@ -534,14 +535,14 @@ void FTaskTimingSharedState::ContextMenu_ShowTaskSubsequents_Execute()
 
 bool FTaskTimingSharedState::ContextMenu_ShowParentTasks_CanExecute()
 {
-	return Insights::FTaskGraphProfilerManager::Get().IsValid() && Insights::FTaskGraphProfilerManager::Get()->GetIsAvailable();
+	return FTaskGraphProfilerManager::Get().IsValid() && FTaskGraphProfilerManager::Get()->GetIsAvailable();
 }
 
 ////////////////////////////////////////////////////////////////////////////////////////////////////
 
 bool FTaskTimingSharedState::ContextMenu_ShowParentTasks_IsChecked()
 {
-	TSharedPtr<Insights::FTaskGraphProfilerManager> TaskGraphManager = Insights::FTaskGraphProfilerManager::Get();
+	TSharedPtr<FTaskGraphProfilerManager> TaskGraphManager = FTaskGraphProfilerManager::Get();
 	return TaskGraphManager.IsValid() && TaskGraphManager->GetIsAvailable() && TaskGraphManager->GetShowParentTasks();
 }
 
@@ -549,7 +550,7 @@ bool FTaskTimingSharedState::ContextMenu_ShowParentTasks_IsChecked()
 
 void FTaskTimingSharedState::ContextMenu_ShowParentTasks_Execute()
 {
-	TSharedPtr<Insights::FTaskGraphProfilerManager> TaskGraphManager = Insights::FTaskGraphProfilerManager::Get();
+	TSharedPtr<FTaskGraphProfilerManager> TaskGraphManager = FTaskGraphProfilerManager::Get();
 	if (TaskGraphManager.IsValid() && TaskGraphManager->GetIsAvailable())
 	{
 		TaskGraphManager->SetShowParentTasks(!TaskGraphManager->GetShowParentTasks());
@@ -562,14 +563,14 @@ void FTaskTimingSharedState::ContextMenu_ShowParentTasks_Execute()
 
 bool FTaskTimingSharedState::ContextMenu_ShowNestedTasks_CanExecute()
 {
-	return Insights::FTaskGraphProfilerManager::Get().IsValid() && Insights::FTaskGraphProfilerManager::Get()->GetIsAvailable();
+	return FTaskGraphProfilerManager::Get().IsValid() && FTaskGraphProfilerManager::Get()->GetIsAvailable();
 }
 
 ////////////////////////////////////////////////////////////////////////////////////////////////////
 
 bool FTaskTimingSharedState::ContextMenu_ShowNestedTasks_IsChecked()
 {
-	TSharedPtr<Insights::FTaskGraphProfilerManager> TaskGraphManager = Insights::FTaskGraphProfilerManager::Get();
+	TSharedPtr<FTaskGraphProfilerManager> TaskGraphManager = FTaskGraphProfilerManager::Get();
 	return TaskGraphManager.IsValid() && TaskGraphManager->GetIsAvailable() && TaskGraphManager->GetShowNestedTasks();
 }
 
@@ -577,7 +578,7 @@ bool FTaskTimingSharedState::ContextMenu_ShowNestedTasks_IsChecked()
 
 void FTaskTimingSharedState::ContextMenu_ShowNestedTasks_Execute()
 {
-	TSharedPtr<Insights::FTaskGraphProfilerManager> TaskGraphManager = Insights::FTaskGraphProfilerManager::Get();
+	TSharedPtr<FTaskGraphProfilerManager> TaskGraphManager = FTaskGraphProfilerManager::Get();
 	if (TaskGraphManager.IsValid() && TaskGraphManager->GetIsAvailable())
 	{
 		TaskGraphManager->SetShowNestedTasks(!TaskGraphManager->GetShowNestedTasks());
@@ -590,14 +591,14 @@ void FTaskTimingSharedState::ContextMenu_ShowNestedTasks_Execute()
 
 bool FTaskTimingSharedState::ContextMenu_ShowCriticalPath_CanExecute()
 {
-	return Insights::FTaskGraphProfilerManager::Get().IsValid() && Insights::FTaskGraphProfilerManager::Get()->GetIsAvailable();
+	return FTaskGraphProfilerManager::Get().IsValid() && FTaskGraphProfilerManager::Get()->GetIsAvailable();
 }
 
 ////////////////////////////////////////////////////////////////////////////////////////////////////
 
 bool FTaskTimingSharedState::ContextMenu_ShowCriticalPath_IsChecked()
 {
-	TSharedPtr<Insights::FTaskGraphProfilerManager> TaskGraphManager = Insights::FTaskGraphProfilerManager::Get();
+	TSharedPtr<FTaskGraphProfilerManager> TaskGraphManager = FTaskGraphProfilerManager::Get();
 	return TaskGraphManager.IsValid() && TaskGraphManager->GetIsAvailable() && TaskGraphManager->GetShowCriticalPath();
 }
 
@@ -605,7 +606,7 @@ bool FTaskTimingSharedState::ContextMenu_ShowCriticalPath_IsChecked()
 
 void FTaskTimingSharedState::ContextMenu_ShowCriticalPath_Execute()
 {
-	TSharedPtr<Insights::FTaskGraphProfilerManager> TaskGraphManager = Insights::FTaskGraphProfilerManager::Get();
+	TSharedPtr<FTaskGraphProfilerManager> TaskGraphManager = FTaskGraphProfilerManager::Get();
 	if (TaskGraphManager.IsValid() && TaskGraphManager->GetIsAvailable())
 	{
 		TaskGraphManager->SetShowCriticalPath(!TaskGraphManager->GetShowCriticalPath());
@@ -684,9 +685,7 @@ void FTaskTimingSharedState::OnTaskSettingsChanged()
 	FTaskGraphProfilerManager::Get()->ClearTaskRelations();
 	TaskTrack->SetTaskId(TaskTrace::InvalidId);
 
-	TSharedPtr<Insights::FTaskGraphProfilerManager> TaskGraphManager = Insights::FTaskGraphProfilerManager::Get();
-
-	TSharedPtr<UE::Insights::TimingProfiler::STimingView> TimingView = GetTimingView();
+	TSharedPtr<TimingProfiler::STimingView> TimingView = GetTimingView();
 	if (!TimingView)
 	{
 		return;
@@ -751,12 +750,12 @@ void FTaskTimingTrack::BuildDrawState(ITimingEventsTrackDrawStateBuilder& Builde
 
 	if (bShowDetailInfoOnTaskTrack)
 	{
-		using namespace UE::Insights::TimingProfiler;
+		const uint32 MaxEventDepth = TimingProfiler::FTimingProfilerManager::Get()->GetEventDepthLimit() - 1;
 
 		uint32 Depth = 0;
 		for (TraceServices::FTaskInfo::FRelationInfo Relation : Task->Prerequisites)
 		{
-			if (Depth >= FTimingProfilerManager::Get()->GetEventDepthLimit() - 1)
+			if (Depth >= MaxEventDepth)
 			{
 				break;
 			}
@@ -770,7 +769,7 @@ void FTaskTimingTrack::BuildDrawState(ITimingEventsTrackDrawStateBuilder& Builde
 		// parent tasks can overlap with prerequisites, hence they are shown under them (`Depth` is not reset)
 		for (TraceServices::FTaskInfo::FRelationInfo Relation : Task->ParentTasks)
 		{
-			if (Depth >= FTimingProfilerManager::Get()->GetEventDepthLimit() - 1)
+			if (Depth >= MaxEventDepth)
 			{
 				break;
 			}
@@ -784,7 +783,7 @@ void FTaskTimingTrack::BuildDrawState(ITimingEventsTrackDrawStateBuilder& Builde
 		Depth = 0;
 		for (TraceServices::FTaskInfo::FRelationInfo Relation : Task->NestedTasks)
 		{
-			if (Depth >= FTimingProfilerManager::Get()->GetEventDepthLimit() - 1)
+			if (Depth >= MaxEventDepth)
 			{
 				break;
 			}
@@ -798,7 +797,7 @@ void FTaskTimingTrack::BuildDrawState(ITimingEventsTrackDrawStateBuilder& Builde
 		Depth = 0;
 		for (TraceServices::FTaskInfo::FRelationInfo Relation : Task->Subsequents)
 		{
-			if (Depth >= FTimingProfilerManager::Get()->GetEventDepthLimit() - 1)
+			if (Depth >= MaxEventDepth)
 			{
 				break;
 			}
@@ -1112,7 +1111,7 @@ void FTaskTimingTrack::GetEventRelations(const FThreadTrackEvent& InSelectedEven
 	const TraceServices::ITasksProvider* TasksProvider = TraceServices::ReadTasksProvider(*Session.Get());
 	if (TasksProvider)
 	{
-		TSharedPtr<UE::Insights::TimingProfiler::STimingView> TimingView = SharedState.GetTimingView();
+		TSharedPtr<TimingProfiler::STimingView> TimingView = SharedState.GetTimingView();
 		if (!TimingView.IsValid())
 		{
 			return;
@@ -1191,6 +1190,6 @@ FReply FTaskTimingTrack::OnMouseButtonUp(const FGeometry& MyGeometry, const FPoi
 
 ////////////////////////////////////////////////////////////////////////////////////////////////////
 
-} // namespace Insights
+} // namespace UE::Insights::TaskGraphProfiler
 
 #undef LOCTEXT_NAMESPACE
diff --git a/Engine/Source/Developer/TraceInsights/Private/Insights/TaskGraphProfiler/ViewModels/TaskTimingTrack.h b/Engine/Source/Developer/TraceInsights/Private/Insights/TaskGraphProfiler/ViewModels/TaskTimingTrack.h
--- a/Engine/Source/Developer/TraceInsights/Private/Insights/TaskGraphProfiler/ViewModels/TaskTimingTrack.h
+++ b/Engine/Source/Developer/TraceInsights/Private/Insights/TaskGraphProfiler/ViewModels/TaskTimingTrack.h
@@ -2,7 +2,7 @@
 
 #pragma once
 
-#include "CoreMinimal.h"
+#include "CoreTypes.h"
 
 #include "Async/TaskTrace.h"
 
@@ -15,18 +15,18 @@ class FUICommandList;
 class FThreadTrackEvent;
 namespace UE::Insights::TimingProfiler { class STimingView; }
 
-////////////////////////////////////////////////////////////////////////////////////////////////////
-
-namespace Insights
+namespace UE::Insights::TaskGraphProfiler
 {
 
+////////////////////////////////////////////////////////////////////////////////////////////////////
+
 class FTaskTimingTrack;
 
-class FTaskTimingSharedState : public UE::Insights::Timing::ITimingViewExtender, public TSharedFromThis<FTaskTimingSharedState>
+class FTaskTimingSharedState : public Timing::ITimingViewExtender, public TSharedFromThis<FTaskTimingSharedState>
 {
 
 public:
-	FTaskTimingSharedState(UE::Insights::TimingProfiler::STimingView* InTimingView);
+	FTaskTimingSharedState(TimingProfiler::STimingView* InTimingView);
 	virtual ~FTaskTimingSharedState() = default;
 
 	TSharedPtr<FTaskTimingTrack> GetTaskTrack() { return TaskTrack; }
@@ -36,11 +36,11 @@ public:
 	//////////////////////////////////////////////////
 	// ITimingViewExtender interface
 
-	virtual void OnBeginSession(UE::Insights::Timing::ITimingViewSession& InSession) override;
-	virtual void OnEndSession(UE::Insights::Timing::ITimingViewSession& InSession) override;
-	virtual void Tick(UE::Insights::Timing::ITimingViewSession& InSession, const TraceServices::IAnalysisSession& InAnalysisSession) override;
-	virtual void ExtendFilterMenu(UE::Insights::Timing::ITimingViewSession& InSession, FMenuBuilder& InMenuBuilder) override;
-	virtual void ExtendOtherTracksFilterMenu(UE::Insights::Timing::ITimingViewSession& InSession, FMenuBuilder& InMenuBuilder) override;
+	virtual void OnBeginSession(Timing::ITimingViewSession& InSession) override;
+	virtual void OnEndSession(Timing::ITimingViewSession& InSession) override;
+	virtual void Tick(Timing::ITimingViewSession& InSession, const TraceServices::IAnalysisSession& InAnalysisSession) override;
+	virtual void ExtendFilterMenu(Timing::ITimingViewSession& InSession, FMenuBuilder& InMenuBuilder) override;
+	virtual void ExtendOtherTracksFilterMenu(Timing::ITimingViewSession& InSession, FMenuBuilder& InMenuBuilder) override;
 
 	//////////////////////////////////////////////////
 
@@ -53,10 +53,10 @@ public:
 
 	void SetResetOnNextTick(bool bInValue) { bResetOnNextTick = bInValue; }
 
-	static TSharedPtr<UE::Insights::TimingProfiler::STimingView> GetTimingView();
+	static TSharedPtr<TimingProfiler::STimingView> GetTimingView();
 
 private:
-	void InitCommandList(TSharedPtr<UE::Insights::TimingProfiler::STimingView> TimingView);
+	void InitCommandList(TSharedPtr<TimingProfiler::STimingView> TimingView);
 
 	void BuildTasksSubMenu(FMenuBuilder& MenuBuilder);
 
@@ -99,7 +99,7 @@ private:
 	void OnTaskSettingsChanged();
 
 private:
-	UE::Insights::Timing::ITimingViewSession* TimingViewSession;
+	Timing::ITimingViewSession* TimingViewSession;
 
 	bool bShowHideTaskTrack;
 	bool bResetOnNextTick = false;
@@ -159,4 +159,4 @@ private:
 
 ////////////////////////////////////////////////////////////////////////////////////////////////////
 
-} //namespace Insights
+} // namespace UE::Insights::TaskGraphProfiler
diff --git a/Engine/Source/Developer/TraceInsights/Private/Insights/TaskGraphProfiler/ViewModels/TaskTrackEvent.cpp b/Engine/Source/Developer/TraceInsights/Private/Insights/TaskGraphProfiler/ViewModels/TaskTrackEvent.cpp
--- a/Engine/Source/Developer/TraceInsights/Private/Insights/TaskGraphProfiler/ViewModels/TaskTrackEvent.cpp
+++ b/Engine/Source/Developer/TraceInsights/Private/Insights/TaskGraphProfiler/ViewModels/TaskTrackEvent.cpp
@@ -2,7 +2,7 @@
 
 #include "Insights/TaskGraphProfiler/ViewModels/TaskTrackEvent.h"
 
-namespace Insights
+namespace UE::Insights::TaskGraphProfiler
 {
 
 ////////////////////////////////////////////////////////////////////////////////////////////////////
@@ -114,4 +114,4 @@ FString FTaskTrackEvent::GetEventName() const
 ////////////////////////////////////////////////////////////////////////////////////////////////////
 ////////////////////////////////////////////////////////////////////////////////////////////////////
 
-} // namespace Insights
+} // namespace UE::Insights::TaskGraphProfiler
diff --git a/Engine/Source/Developer/TraceInsights/Private/Insights/TaskGraphProfiler/ViewModels/TaskTrackEvent.h b/Engine/Source/Developer/TraceInsights/Private/Insights/TaskGraphProfiler/ViewModels/TaskTrackEvent.h
--- a/Engine/Source/Developer/TraceInsights/Private/Insights/TaskGraphProfiler/ViewModels/TaskTrackEvent.h
+++ b/Engine/Source/Developer/TraceInsights/Private/Insights/TaskGraphProfiler/ViewModels/TaskTrackEvent.h
@@ -2,7 +2,7 @@
 
 #pragma once
 
-#include "CoreMinimal.h"
+#include "CoreTypes.h"
 
 #include "Async/TaskTrace.h"
 
@@ -10,7 +10,7 @@
 #include "Insights/TaskGraphProfiler/TaskGraphProfilerManager.h"
 #include "Insights/ViewModels/TimingEvent.h"
 
-namespace Insights
+namespace UE::Insights::TaskGraphProfiler
 {
 
 ////////////////////////////////////////////////////////////////////////////////////////////////////
@@ -40,4 +40,4 @@ private:
 
 ////////////////////////////////////////////////////////////////////////////////////////////////////
 
-} // namespace Insights
+} // namespace UE::Insights::TaskGraphProfiler
diff --git a/Engine/Source/Developer/TraceInsights/Private/Insights/TaskGraphProfiler/Widgets/STaskTableTreeView.cpp b/Engine/Source/Developer/TraceInsights/Private/Insights/TaskGraphProfiler/Widgets/STaskTableTreeView.cpp
--- a/Engine/Source/Developer/TraceInsights/Private/Insights/TaskGraphProfiler/Widgets/STaskTableTreeView.cpp
+++ b/Engine/Source/Developer/TraceInsights/Private/Insights/TaskGraphProfiler/Widgets/STaskTableTreeView.cpp
@@ -36,7 +36,7 @@
 
 using namespace TraceServices;
 
-namespace Insights
+namespace UE::Insights::TaskGraphProfiler
 {
 
 ////////////////////////////////////////////////////////////////////////////////////////////////////
@@ -96,7 +96,7 @@ STaskTableTreeView::~STaskTableTreeView()
 
 ////////////////////////////////////////////////////////////////////////////////////////////////////
 
-void STaskTableTreeView::Construct(const FArguments& InArgs, TSharedPtr<Insights::FTaskTable> InTablePtr)
+void STaskTableTreeView::Construct(const FArguments& InArgs, TSharedPtr<FTaskTable> InTablePtr)
 {
 	ConstructWidget(InTablePtr);
 
@@ -238,7 +238,7 @@ void STaskTableTreeView::RebuildTree(bool bResync)
 			{
 				FName BaseNodeName(TEXT("task"));
 
-				TArray<UE::Insights::FTableTreeNodePtr>* Nodes = &TableRowNodes;
+				TArray<FTableTreeNodePtr>* Nodes = &TableRowNodes;
 
 				TasksProvider->EnumerateTasks(QueryStartTime, QueryEndTime, SelectedTasksSelectionOption, [&Tasks, &TaskTable, &BaseNodeName, Nodes](const TraceServices::FTaskInfo& TaskInfo)
 				{
@@ -688,7 +688,7 @@ bool STaskTableTreeView::TasksSelectionOptions_IsEnabled() const
 
 bool STaskTableTreeView::ContextMenu_GoToTask_CanExecute() const
 {
-	TArray<UE::Insights::FTableTreeNodePtr> SelectedItems;
+	TArray<FTableTreeNodePtr> SelectedItems;
 	TreeView->GetSelectedItems(SelectedItems);
 
 	if (SelectedItems.Num() != 1)
@@ -710,7 +710,7 @@ bool STaskTableTreeView::ContextMenu_GoToTask_CanExecute() const
 
 void STaskTableTreeView::ContextMenu_GoToTask_Execute()
 {
-	TArray<UE::Insights::FTableTreeNodePtr> SelectedItems;
+	TArray<FTableTreeNodePtr> SelectedItems;
 	TreeView->GetSelectedItems(SelectedItems);
 
 	if (SelectedItems.Num() != 1)
@@ -808,7 +808,7 @@ bool STaskTableTreeView::ContextMenu_OpenInIDE_CanExecute() const
 
 bool STaskTableTreeView::GetSourceFileAndLineForSelectedTask(FString& OutFile, uint32& OutLine) const
 {
-	TArray<UE::Insights::FTableTreeNodePtr> SelectedItems;
+	TArray<FTableTreeNodePtr> SelectedItems;
 	TreeView->GetSelectedItems(SelectedItems);
 
 	if (SelectedItems.Num() != 1)
@@ -816,7 +816,7 @@ bool STaskTableTreeView::GetSourceFileAndLineForSelectedTask(FString& OutFile, u
 		return false;
 	}
 
-	UE::Insights::FTableTreeNodePtr SelectedTreeNode = SelectedItems[0];
+	FTableTreeNodePtr SelectedTreeNode = SelectedItems[0];
 	if (!SelectedTreeNode.IsValid() || !SelectedTreeNode->Is<FTaskNode>())
 	{
 		return false;
@@ -877,7 +877,7 @@ void STaskTableTreeView::ContextMenu_OpenInIDE_Execute()
 
 ////////////////////////////////////////////////////////////////////////////////////////////////////
 
-void STaskTableTreeView::TreeView_OnMouseButtonDoubleClick(UE::Insights::FTableTreeNodePtr TreeNode)
+void STaskTableTreeView::TreeView_OnMouseButtonDoubleClick(FTableTreeNodePtr TreeNode)
 {
 	if (!TreeNode->IsGroup())
 	{
@@ -892,12 +892,12 @@ void STaskTableTreeView::TreeView_OnMouseButtonDoubleClick(UE::Insights::FTableT
 void STaskTableTreeView::SelectTaskEntry(TaskTrace::FId InId)
 {
 	TaskIdToSelect = InId;
-	StartTableDataTask<UE::Insights::FSearchForItemToSelectTask>(SharedThis(this));
+	StartTableDataTask<FSearchForItemToSelectTask>(SharedThis(this));
 }
 
 ////////////////////////////////////////////////////////////////////////////////////////////////////
 
-void STaskTableTreeView::SearchForItem(TSharedPtr<UE::Insights::FTableTaskCancellationToken> CancellationToken)
+void STaskTableTreeView::SearchForItem(TSharedPtr<FTableTaskCancellationToken> CancellationToken)
 {
 	TSharedPtr<FTaskTable> TaskTable = GetTaskTable();
 	TArray<FTaskEntry>& Tasks = TaskTable->GetTaskEntries();
@@ -912,7 +912,7 @@ void STaskTableTreeView::SearchForItem(TSharedPtr<UE::Insights::FTableTaskCancel
 
 		if (Tasks[Index].Id == TaskIdToSelect)
 		{
-			TGraphTask<UE::Insights::FSelectNodeByTableRowIndexTask>::CreateTask().ConstructAndDispatchWhenReady(CancellationToken, SharedThis(this), Index);
+			TGraphTask<FSelectNodeByTableRowIndexTask>::CreateTask().ConstructAndDispatchWhenReady(CancellationToken, SharedThis(this), Index);
 			break;
 		}
 	}
@@ -920,6 +920,6 @@ void STaskTableTreeView::SearchForItem(TSharedPtr<UE::Insights::FTableTaskCancel
 
 ////////////////////////////////////////////////////////////////////////////////////////////////////
 
-} // namespace Insights
+} // namespace UE::Insights::TaskGraphProfiler
 
 #undef LOCTEXT_NAMESPACE
diff --git a/Engine/Source/Developer/TraceInsights/Private/Insights/TaskGraphProfiler/Widgets/STaskTableTreeView.h b/Engine/Source/Developer/TraceInsights/Private/Insights/TaskGraphProfiler/Widgets/STaskTableTreeView.h
--- a/Engine/Source/Developer/TraceInsights/Private/Insights/TaskGraphProfiler/Widgets/STaskTableTreeView.h
+++ b/Engine/Source/Developer/TraceInsights/Private/Insights/TaskGraphProfiler/Widgets/STaskTableTreeView.h
@@ -2,7 +2,7 @@
 
 #pragma once
 
-#include "CoreTypes.h"
+#include "CoreMinimal.h"
 
 #include "Async/TaskTrace.h"
 
@@ -15,12 +15,12 @@
 
 class FMenuBuilder;
 
-namespace Insights
+namespace UE::Insights::TaskGraphProfiler
 {
 
 ////////////////////////////////////////////////////////////////////////////////////////////////////
 
-class STaskTableTreeView : public UE::Insights::SSessionTableTreeView
+class STaskTableTreeView : public SSessionTableTreeView
 {
 private:
 	struct FColumnConfig
@@ -88,7 +88,7 @@ public:
 
 	////////////////////////////////////////////////////////////////////////////////////////////////////
 
-	virtual void TreeView_OnMouseButtonDoubleClick(UE::Insights::FTableTreeNodePtr TreeNode) override;
+	virtual void TreeView_OnMouseButtonDoubleClick(FTableTreeNodePtr TreeNode) override;
 
 	void SelectTaskEntry(TaskTrace::FId InId);
 
@@ -97,7 +97,7 @@ protected:
 
 	virtual void ExtendMenu(FMenuBuilder& MenuBuilder) override;
 
-	virtual void SearchForItem(TSharedPtr<UE::Insights::FTableTaskCancellationToken> CancellationToken) override;
+	virtual void SearchForItem(TSharedPtr<FTableTaskCancellationToken> CancellationToken) override;
 
 private:
 	void AddCommmands();
@@ -141,4 +141,4 @@ private:
 
 ////////////////////////////////////////////////////////////////////////////////////////////////////
 
-} // namespace Insights
+} // namespace UE::Insights::TaskGraphProfiler
diff --git a/Engine/Source/Developer/TraceInsights/Private/Insights/TraceInsightsModule.cpp b/Engine/Source/Developer/TraceInsights/Private/Insights/TraceInsightsModule.cpp
--- a/Engine/Source/Developer/TraceInsights/Private/Insights/TraceInsightsModule.cpp
+++ b/Engine/Source/Developer/TraceInsights/Private/Insights/TraceInsightsModule.cpp
@@ -71,12 +71,12 @@ void FTraceInsightsModule::StartupModule()
 
 	// Register other default components.
 	RegisterComponent(UE::Insights::TimingProfiler::FTimingProfilerManager::CreateInstance());
-	RegisterComponent(FLoadingProfilerManager::CreateInstance());
+	RegisterComponent(UE::Insights::LoadingProfiler::FLoadingProfilerManager::CreateInstance());
 	RegisterComponent(UE::Insights::NetworkingProfiler::FNetworkingProfilerManager::CreateInstance());
 	RegisterComponent(UE::Insights::MemoryProfiler::FMemoryProfilerManager::CreateInstance());
-	RegisterComponent(Insights::FTaskGraphProfilerManager::CreateInstance());
-	RegisterComponent(Insights::FContextSwitchesProfilerManager::CreateInstance());
-	RegisterComponent(Insights::FCookProfilerManager::CreateInstance());
+	RegisterComponent(UE::Insights::TaskGraphProfiler::FTaskGraphProfilerManager::CreateInstance());
+	RegisterComponent(UE::Insights::ContextSwitches::FContextSwitchesProfilerManager::CreateInstance());
+	RegisterComponent(UE::Insights::CookProfiler::FCookProfilerManager::CreateInstance());
 	RegisterComponent(Insights::FTableImportTool::CreateInstance());
 
 #if !WITH_EDITOR
@@ -114,7 +114,7 @@ void FTraceInsightsModule::ShutdownModule()
 #if INSIGHTS_CHECK_SHARED_REFERENCES
 	TSharedPtr<const TraceServices::IAnalysisSession> Session = FInsightsManager::Get()->GetSession();
 	auto TimingInsightsWindow = UE::Insights::TimingProfiler::FTimingProfilerManager::Get()->GetProfilerWindow();
-	auto AssetLoadingInsightsWindow = FLoadingProfilerManager::Get()->GetProfilerWindow();
+	auto AssetLoadingInsightsWindow = UE::Insights::LoadingProfiler::FLoadingProfilerManager::Get()->GetProfilerWindow();
 	auto NetworkingInsightsWindow0 = UE::Insights::NetworkingProfiler::FNetworkingProfilerManager::Get()->GetProfilerWindow(0);
 	auto MemoryInsightsWindow = UE::Insights::MemoryProfiler::FMemoryProfilerManager::Get()->GetProfilerWindow();
 #endif
