[Insights]
  - TraceInsightsCore: Updated the default color (for icon and name text) of items in a table/tree view. Also added FBaseTreeNode::GetDefaultIconColor().
  - TraceInsightsCore: Added icon color customization to FCustomTableTreeNode.
  - Memory Insights / Mem Allocs table: Updated coloring for Heap and Swap Page group nodes. Also updated name of "Swap Size" and "Compressed Swap Size" columns.

#offorgithub public

diff --git a/Engine/Source/Developer/TraceInsights/Private/Insights/MemoryProfiler/ViewModels/MemAllocGroupingByCallstack.cpp b/Engine/Source/Developer/TraceInsights/Private/Insights/MemoryProfiler/ViewModels/MemAllocGroupingByCallstack.cpp
--- a/Engine/Source/Developer/TraceInsights/Private/Insights/MemoryProfiler/ViewModels/MemAllocGroupingByCallstack.cpp
+++ b/Engine/Source/Developer/TraceInsights/Private/Insights/MemoryProfiler/ViewModels/MemAllocGroupingByCallstack.cpp
@@ -46,6 +46,13 @@ const FSlateBrush* FCallstackFrameGroupNode::GetIcon() const
 
 ////////////////////////////////////////////////////////////////////////////////////////////////////
 
+FLinearColor FCallstackFrameGroupNode::GetIconColor() const
+{
+	return FLinearColor(0.5f, 0.75f, 1.0f, 1.0f);
+}
+
+////////////////////////////////////////////////////////////////////////////////////////////////////
+
 FLinearColor FCallstackFrameGroupNode::GetColor() const
 {
 	return FLinearColor(0.5f, 0.75f, 1.0f, 1.0f);
@@ -63,14 +70,14 @@ FMemAllocGroupingByCallstack::FMemAllocGroupingByCallstack(bool bInIsAllocCallst
 	: FTreeNodeGrouping(
 		bInIsAllocCallstack
 		? (bInIsInverted ? LOCTEXT("Grouping_ByCallstack2_ShortName", "Inverted Alloc Callstack")
-						: LOCTEXT("Grouping_ByCallstack1_ShortName", "Alloc Callstack"))
+						 : LOCTEXT("Grouping_ByCallstack1_ShortName", "Alloc Callstack"))
 		: (bInIsInverted ? LOCTEXT("Grouping_ByCallstack4_ShortName", "Inverted Free Callstack")
-						: LOCTEXT("Grouping_ByCallstack3_ShortName", "Free Callstack")),
+						 : LOCTEXT("Grouping_ByCallstack3_ShortName", "Free Callstack")),
 		bInIsAllocCallstack
 		? (bInIsInverted ? LOCTEXT("Grouping_ByCallstack2_TitleName", "By Inverted Alloc Callstack")
-					  : LOCTEXT("Grouping_ByCallstack1_TitleName", "By Alloc Callstack"))
+						 : LOCTEXT("Grouping_ByCallstack1_TitleName", "By Alloc Callstack"))
 		: (bInIsInverted ? LOCTEXT("Grouping_ByCallstack4_TitleName", "By Inverted Free Callstack")
-					  : LOCTEXT("Grouping_ByCallstack3_TitleName", "By Free Callstack")),
+						 : LOCTEXT("Grouping_ByCallstack3_TitleName", "By Free Callstack")),
 		LOCTEXT("Grouping_Callstack_Desc", "Creates a tree based on callstack of each allocation."),
 		TEXT("Icons.Group.TreeItem"),
 		nullptr)
diff --git a/Engine/Source/Developer/TraceInsights/Private/Insights/MemoryProfiler/ViewModels/MemAllocGroupingByCallstack.h b/Engine/Source/Developer/TraceInsights/Private/Insights/MemoryProfiler/ViewModels/MemAllocGroupingByCallstack.h
--- a/Engine/Source/Developer/TraceInsights/Private/Insights/MemoryProfiler/ViewModels/MemAllocGroupingByCallstack.h
+++ b/Engine/Source/Developer/TraceInsights/Private/Insights/MemoryProfiler/ViewModels/MemAllocGroupingByCallstack.h
@@ -34,6 +34,7 @@ public:
 
 	virtual const FText GetTooltipText() const override;
 	virtual const FSlateBrush* GetIcon() const override;
+	virtual FLinearColor GetIconColor() const override;
 	virtual FLinearColor GetColor() const override;
 
 	/**
diff --git a/Engine/Source/Developer/TraceInsights/Private/Insights/MemoryProfiler/ViewModels/MemAllocGroupingByHeap.cpp b/Engine/Source/Developer/TraceInsights/Private/Insights/MemoryProfiler/ViewModels/MemAllocGroupingByHeap.cpp
--- a/Engine/Source/Developer/TraceInsights/Private/Insights/MemoryProfiler/ViewModels/MemAllocGroupingByHeap.cpp
+++ b/Engine/Source/Developer/TraceInsights/Private/Insights/MemoryProfiler/ViewModels/MemAllocGroupingByHeap.cpp
@@ -42,12 +42,39 @@ FMemAllocGroupingByHeap::~FMemAllocGroupingByHeap()
 
 ////////////////////////////////////////////////////////////////////////////////////////////////////
 
-FTableTreeNodePtr MakeGroupNodeHierarchy(const TraceServices::IAllocationsProvider::FHeapSpec& Spec, TWeakPtr<FTable>& InParentTable, TArray<FTableTreeNodePtr>& NodeTable)
+class FMemHeapTreeNode : public FTableTreeNode
 {
-	const FSlateBrush* IconBrush = FBaseTreeNode::GetDefaultIcon(true);
-	const FLinearColor Color(1.0f, 0.7f, 0.3f, 1.0f);
+	INSIGHTS_DECLARE_RTTI(FMemHeapTreeNode, FTableTreeNode)
+
+public:
+	/** Initialization constructor for the group node. */
+	explicit FMemHeapTreeNode(const FName InName, TWeakPtr<FTable> InParentTable)
+		: FTableTreeNode(InName, InParentTable)
+	{
+	}
+
+	virtual ~FMemHeapTreeNode()
+	{
+	}
+
+	virtual FLinearColor GetIconColor() const override final
+	{
+		return FLinearColor(1.0f, 0.45f, 0.6f, 1.0f);
+	}
+
+	virtual FLinearColor GetColor() const override final
+	{
+		return FLinearColor(1.0f, 0.45f, 0.6f, 1.0f);
+	}
+};
+
+INSIGHTS_IMPLEMENT_RTTI(FMemHeapTreeNode)
 
-	auto HeapGroup = MakeShared<FCustomTableTreeNode>(FName(Spec.Name), InParentTable, IconBrush, Color);
+////////////////////////////////////////////////////////////////////////////////////////////////////
+
+FTableTreeNodePtr MakeGroupNodeHierarchy(const TraceServices::IAllocationsProvider::FHeapSpec& Spec, TWeakPtr<FTable>& InParentTable, TArray<FTableTreeNodePtr>& NodeTable)
+{
+	auto HeapGroup = MakeShared<FMemHeapTreeNode>(FName(Spec.Name), InParentTable);
 	if (int32(Spec.Id) >= NodeTable.Num())
 	{
 		NodeTable.AddDefaulted(int32(Spec.Id) - NodeTable.Num() + 1);
diff --git a/Engine/Source/Developer/TraceInsights/Private/Insights/MemoryProfiler/ViewModels/MemAllocGroupingBySwapPage.cpp b/Engine/Source/Developer/TraceInsights/Private/Insights/MemoryProfiler/ViewModels/MemAllocGroupingBySwapPage.cpp
--- a/Engine/Source/Developer/TraceInsights/Private/Insights/MemoryProfiler/ViewModels/MemAllocGroupingBySwapPage.cpp
+++ b/Engine/Source/Developer/TraceInsights/Private/Insights/MemoryProfiler/ViewModels/MemAllocGroupingBySwapPage.cpp
@@ -61,6 +61,36 @@ struct TMemoryPageMapKeyFuncs : BaseKeyFuncs<TPair<uint64, FTableTreeNode*>, uin
 
 ////////////////////////////////////////////////////////////////////////////////////////////////////
 
+class FMemSwapPageTreeNode : public FTableTreeNode
+{
+	INSIGHTS_DECLARE_RTTI(FMemSwapPageTreeNode, FTableTreeNode)
+
+public:
+	/** Initialization constructor for the group node. */
+	explicit FMemSwapPageTreeNode(const FName InName, TWeakPtr<FTable> InParentTable)
+		: FTableTreeNode(InName, InParentTable)
+	{
+	}
+
+	virtual ~FMemSwapPageTreeNode()
+	{
+	}
+
+	virtual FLinearColor GetIconColor() const override final
+	{
+		return FLinearColor(0.3f, 0.8f, 0.4f, 1.0f);
+	}
+
+	virtual FLinearColor GetColor() const override final
+	{
+		return FLinearColor(0.2f, 0.8f, 0.4f, 1.0f);
+	}
+};
+
+INSIGHTS_IMPLEMENT_RTTI(FMemSwapPageTreeNode)
+
+////////////////////////////////////////////////////////////////////////////////////////////////////
+
 void FMemAllocGroupingBySwapPage::GroupNodes(
 	const TArray<FTableTreeNodePtr>& Nodes,
 	FTableTreeNode& ParentGroup,
@@ -69,11 +99,8 @@ void FMemAllocGroupingBySwapPage::GroupNodes(
 {
 	ParentGroup.ClearChildren();
 
-	const FSlateBrush* IconBrush = FBaseTreeNode::GetDefaultIcon(true);
-	const FLinearColor Color(1.0f, 0.7f, 0.3f, 1.0f);
-
-	auto InSwapGroup = MakeShared<FCustomTableTreeNode>(FName(TEXT("In Swap")), InParentTable, IconBrush, Color);
-	auto InRamGroup = MakeShared<FCustomTableTreeNode>(FName(TEXT("In RAM")), InParentTable, IconBrush, Color);
+	auto InSwapGroup = MakeShared<FTableTreeNode>(FName(TEXT("In Swap")), InParentTable);
+	auto InRamGroup = MakeShared<FTableTreeNode>(FName(TEXT("In RAM")), InParentTable);
 	ParentGroup.AddChildAndSetParent(InSwapGroup);
 	ParentGroup.AddChildAndSetParent(InRamGroup);
 
@@ -91,7 +118,7 @@ void FMemAllocGroupingBySwapPage::GroupNodes(
 
 		if (Alloc && Alloc->IsSwap())
 		{
-			auto SwapEntry = MakeShared<FCustomTableTreeNode>(FName(FString::Printf(TEXT("0x%016llx"), Alloc->GetAddress())), InParentTable, IconBrush, Color);
+			auto SwapEntry = MakeShared<FMemSwapPageTreeNode>(FName(FString::Printf(TEXT("0x%016llx"), Alloc->GetAddress())), InParentTable);
 			InSwapGroup->AddChildAndSetParent(SwapEntry);
 			SwapEntry->AddChildAndSetParent(NodePtr);
 
diff --git a/Engine/Source/Developer/TraceInsights/Private/Insights/MemoryProfiler/ViewModels/MemAllocGroupingByTag.h b/Engine/Source/Developer/TraceInsights/Private/Insights/MemoryProfiler/ViewModels/MemAllocGroupingByTag.h
--- a/Engine/Source/Developer/TraceInsights/Private/Insights/MemoryProfiler/ViewModels/MemAllocGroupingByTag.h
+++ b/Engine/Source/Developer/TraceInsights/Private/Insights/MemoryProfiler/ViewModels/MemAllocGroupingByTag.h
@@ -56,9 +56,11 @@ public:
 	{
 	}
 
-	/**
-	 * @return the color tint for icon and name text.
-	 */
+	virtual FLinearColor GetIconColor() const override
+	{
+		return FLinearColor(0.75f, 0.5f, 1.0f, 1.0f);
+	}
+
 	virtual FLinearColor GetColor() const override
 	{
 		return FLinearColor(0.75f, 0.5f, 1.0f, 1.0f);
diff --git a/Engine/Source/Developer/TraceInsights/Private/Insights/MemoryProfiler/ViewModels/MemAllocTable.cpp b/Engine/Source/Developer/TraceInsights/Private/Insights/MemoryProfiler/ViewModels/MemAllocTable.cpp
--- a/Engine/Source/Developer/TraceInsights/Private/Insights/MemoryProfiler/ViewModels/MemAllocTable.cpp
+++ b/Engine/Source/Developer/TraceInsights/Private/Insights/MemoryProfiler/ViewModels/MemAllocTable.cpp
@@ -44,8 +44,8 @@ const FName FMemAllocTableColumns::FreeThreadColumnId(TEXT("FreeThread"));
 const FName FMemAllocTableColumns::AddressColumnId(TEXT("Address"));
 const FName FMemAllocTableColumns::MemoryPageColumnId(TEXT("MemoryPage"));
 const FName FMemAllocTableColumns::CountColumnId(TEXT("Count"));
-const FName FMemAllocTableColumns::SwapCompressedSizeColumnId(TEXT("SwapCompressedSize"));
-const FName FMemAllocTableColumns::SizeInSwapColumnId(TEXT("SizeInSwap"));
+const FName FMemAllocTableColumns::CompressedSwapSizeColumnId(TEXT("CompressedSwapSize"));
+const FName FMemAllocTableColumns::SwapSizeColumnId(TEXT("SwapSize"));
 const FName FMemAllocTableColumns::SizeColumnId(TEXT("Size"));
 const FName FMemAllocTableColumns::LLMSizeColumnId(TEXT("LLMSize"));
 const FName FMemAllocTableColumns::LLMDeltaSizeColumnId(TEXT("LLMDeltaSize"));
@@ -747,20 +747,20 @@ void FMemAllocTable::AddDefaultColumns()
 		AddColumn(ColumnRef);
 	}
 	//////////////////////////////////////////////////
-	// Swap Compressed Size
+	// Compressed Swap Size
 	{
-		TSharedRef<FTableColumn> ColumnRef = MakeShared<FTableColumn>(FMemAllocTableColumns::SwapCompressedSizeColumnId);
+		TSharedRef<FTableColumn> ColumnRef = MakeShared<FTableColumn>(FMemAllocTableColumns::CompressedSwapSizeColumnId);
 		FTableColumn& Column = *ColumnRef;
 
 		Column.SetIndex(ColumnIndex++);
 
-		Column.SetShortName(LOCTEXT("SwapCompressedSizeColumnName", "Compressed Size In Swap"));
-		Column.SetTitleName(LOCTEXT("SwapCompressedSizeColumnTitle", "Compressed Size In Swap"));
-		Column.SetDescription(LOCTEXT("SwapCompressedSizeColumnDesc", "Shows compressed size of swap page."));
+		Column.SetShortName(LOCTEXT("CompressedSwapSizeColumnName", "Compressed Swap Size"));
+		Column.SetTitleName(LOCTEXT("CompressedSwapSizeColumnTitle", "Compressed Swap Size"));
+		Column.SetDescription(LOCTEXT("CompressedSwapSizeColumnDesc", "Shows compressed size of swap page."));
 
 		Column.SetFlags(ETableColumnFlags::CanBeHidden | ETableColumnFlags::CanBeFiltered);
 
-		Column.SetHorizontalAlignment(HAlign_Left);
+		Column.SetHorizontalAlignment(HAlign_Right);
 		Column.SetInitialWidth(100.0f);
 
 		Column.SetDataType(ETableCellDataType::Int64);
@@ -810,20 +810,20 @@ void FMemAllocTable::AddDefaultColumns()
 		AddColumn(ColumnRef);
 	}
 	//////////////////////////////////////////////////
-	// Page Size in Swap Column
+	// Swap Size Column
 	{
-		TSharedRef<FTableColumn> ColumnRef = MakeShared<FTableColumn>(FMemAllocTableColumns::SizeInSwapColumnId);
+		TSharedRef<FTableColumn> ColumnRef = MakeShared<FTableColumn>(FMemAllocTableColumns::SwapSizeColumnId);
 		FTableColumn& Column = *ColumnRef;
 
 		Column.SetIndex(ColumnIndex++);
 
-		Column.SetShortName(LOCTEXT("SizeInSwapColumnName", "Size in Swap"));
-		Column.SetTitleName(LOCTEXT("SizeInSwapColumnTitle", "Size in Swap"));
-		Column.SetDescription(LOCTEXT("SizeInSwapColumnDesc", "Shows (partial) size of allocation which is stored in the correponding swap page."));
+		Column.SetShortName(LOCTEXT("SwapSizeColumnName", "Swap Size"));
+		Column.SetTitleName(LOCTEXT("SwapSizeColumnTitle", "Swap Size"));
+		Column.SetDescription(LOCTEXT("SwapSizeColumnDesc", "Shows (partial) size of allocation which is stored in the corresponding swap page."));
 
 		Column.SetFlags(ETableColumnFlags::CanBeHidden | ETableColumnFlags::CanBeFiltered);
 
-		Column.SetHorizontalAlignment(HAlign_Left);
+		Column.SetHorizontalAlignment(HAlign_Right);
 		Column.SetInitialWidth(100.0f);
 
 		Column.SetDataType(ETableCellDataType::Int64);
diff --git a/Engine/Source/Developer/TraceInsights/Private/Insights/MemoryProfiler/ViewModels/MemAllocTable.h b/Engine/Source/Developer/TraceInsights/Private/Insights/MemoryProfiler/ViewModels/MemAllocTable.h
--- a/Engine/Source/Developer/TraceInsights/Private/Insights/MemoryProfiler/ViewModels/MemAllocTable.h
+++ b/Engine/Source/Developer/TraceInsights/Private/Insights/MemoryProfiler/ViewModels/MemAllocTable.h
@@ -34,8 +34,8 @@ struct FMemAllocTableColumns
 	static const FName AddressColumnId;
 	static const FName MemoryPageColumnId;
 	static const FName CountColumnId;
-	static const FName SwapCompressedSizeColumnId;
-	static const FName SizeInSwapColumnId;
+	static const FName CompressedSwapSizeColumnId;
+	static const FName SwapSizeColumnId;
 	static const FName SizeColumnId;
 	static const FName LLMSizeColumnId;
 	static const FName LLMDeltaSizeColumnId;
diff --git a/Engine/Source/Developer/TraceInsights/Private/Insights/MemoryProfiler/Widgets/SMemAllocTableTreeView.cpp b/Engine/Source/Developer/TraceInsights/Private/Insights/MemoryProfiler/Widgets/SMemAllocTableTreeView.cpp
--- a/Engine/Source/Developer/TraceInsights/Private/Insights/MemoryProfiler/Widgets/SMemAllocTableTreeView.cpp
+++ b/Engine/Source/Developer/TraceInsights/Private/Insights/MemoryProfiler/Widgets/SMemAllocTableTreeView.cpp
@@ -1192,8 +1192,8 @@ void SMemAllocTableTreeView::InitAvailableViewPresets()
 			InOutConfigSet.Add({ FMemAllocTableColumns::EndTimeColumnId,            true, 100.0f });
 			InOutConfigSet.Add({ FMemAllocTableColumns::AddressColumnId,            true, 120.0f });
 			InOutConfigSet.Add({ FMemAllocTableColumns::CountColumnId,              true, 100.0f });
-			InOutConfigSet.Add({ FMemAllocTableColumns::SwapCompressedSizeColumnId, true, 100.0f });
-			InOutConfigSet.Add({ FMemAllocTableColumns::SizeInSwapColumnId,         true, 100.0f });
+			InOutConfigSet.Add({ FMemAllocTableColumns::CompressedSwapSizeColumnId, true, 100.0f });
+			InOutConfigSet.Add({ FMemAllocTableColumns::SwapSizeColumnId,           true, 100.0f });
 			InOutConfigSet.Add({ FMemAllocTableColumns::SizeColumnId,               true, 100.0f });
 			InOutConfigSet.Add({ FMemAllocTableColumns::TagColumnId,                true, 120.0f });
 			InOutConfigSet.Add({ FMemAllocTableColumns::AllocFunctionColumnId,      true, 400.0f });
diff --git a/Engine/Source/Developer/TraceInsightsCore/Private/InsightsCore/Table/ViewModels/BaseTreeNode.cpp b/Engine/Source/Developer/TraceInsightsCore/Private/InsightsCore/Table/ViewModels/BaseTreeNode.cpp
--- a/Engine/Source/Developer/TraceInsightsCore/Private/InsightsCore/Table/ViewModels/BaseTreeNode.cpp
+++ b/Engine/Source/Developer/TraceInsightsCore/Private/InsightsCore/Table/ViewModels/BaseTreeNode.cpp
@@ -69,15 +69,33 @@ const FSlateBrush* FBaseTreeNode::GetDefaultIcon(bool bIsGroupNode)
 
 ////////////////////////////////////////////////////////////////////////////////////////////////////
 
+FLinearColor FBaseTreeNode::GetDefaultIconColor(bool bIsGroupNode)
+{
+	if (bIsGroupNode)
+	{
+		// USlateThemeManager::Get().GetColor(EStyleColor::AccentFolder), but lighter
+		return FLinearColor(1.0f, 0.7f, 0.3f, 1.0f);
+	}
+	else
+	{
+		// USlateThemeManager::Get().GetColor(EStyleColor::AccentWhite)
+		return FLinearColor(1.0f, 1.0f, 1.0f, 1.0f);
+	}
+}
+
+////////////////////////////////////////////////////////////////////////////////////////////////////
+
 FLinearColor FBaseTreeNode::GetDefaultColor(bool bIsGroupNode)
 {
 	if (bIsGroupNode)
 	{
-		return USlateThemeManager::Get().GetColor(EStyleColor::AccentYellow);
+		// USlateThemeManager::Get().GetColor(EStyleColor::AccentFolder), but lighter
+		return FLinearColor(1.0f, 0.7f, 0.3f, 1.0f);
 	}
 	else
 	{
-		return USlateThemeManager::Get().GetColor(EStyleColor::AccentWhite);
+		// USlateThemeManager::Get().GetColor(EStyleColor::AccentWhite)
+		return FLinearColor(1.0f, 1.0f, 1.0f, 1.0f);
 	}
 }
 
diff --git a/Engine/Source/Developer/TraceInsightsCore/Public/InsightsCore/Table/ViewModels/BaseTreeNode.h b/Engine/Source/Developer/TraceInsightsCore/Public/InsightsCore/Table/ViewModels/BaseTreeNode.h
--- a/Engine/Source/Developer/TraceInsightsCore/Public/InsightsCore/Table/ViewModels/BaseTreeNode.h
+++ b/Engine/Source/Developer/TraceInsightsCore/Public/InsightsCore/Table/ViewModels/BaseTreeNode.h
@@ -116,7 +116,7 @@ public:
 	static const FSlateBrush* GetDefaultIcon(bool bIsGroupNode);
 
 	/**
-	 * @returns a brush icon for this node.
+	 * @returns the brush icon for this node.
 	 */
 	virtual const FSlateBrush* GetIcon() const
 	{
@@ -124,15 +124,20 @@ public:
 	}
 
 	/**
-	 * @returns the color tint to be used for the icon used for this node.
+	 * @returns the default color tint for a group or leaf node.
+	 */
+	static FLinearColor GetDefaultIconColor(bool bIsGroupNode);
+
+	/**
+	 * @returns the color tint to be used for the icon of this node.
 	 */
 	virtual FLinearColor GetIconColor() const
 	{
-		return GetDefaultColor(IsGroup());
+		return GetDefaultIconColor(IsGroup());
 	}
 
 	/**
-	 * @returns the default color tint for a group/leaf node.
+	 * @returns the default color tint for a group or leaf node.
 	 */
 	static FLinearColor GetDefaultColor(bool bIsGroupNode);
 
@@ -141,7 +146,7 @@ public:
 	 */
 	virtual FLinearColor GetColor() const
 	{
-		return FLinearColor(1.0f, 1.0f, 1.0f, 1.0f);
+		return GetDefaultColor(IsGroup());
 	}
 
 	//////////////////////////////////////////////////
diff --git a/Engine/Source/Developer/TraceInsightsCore/Public/InsightsCore/Table/ViewModels/TableTreeNode.h b/Engine/Source/Developer/TraceInsightsCore/Public/InsightsCore/Table/ViewModels/TableTreeNode.h
--- a/Engine/Source/Developer/TraceInsightsCore/Public/InsightsCore/Table/ViewModels/TableTreeNode.h
+++ b/Engine/Source/Developer/TraceInsightsCore/Public/InsightsCore/Table/ViewModels/TableTreeNode.h
@@ -176,6 +176,16 @@ public:
 	explicit FCustomTableTreeNode(const FName InName, TWeakPtr<FTable> InParentTable, int32 InRowIndex, const FSlateBrush* InIconBrush, FLinearColor InColor, bool IsGroup)
 		: FTableTreeNode(InName, InParentTable, InRowIndex, IsGroup)
 		, IconBrush(InIconBrush)
+		, IconColor(InColor)
+		, Color(InColor)
+	{
+	}
+
+	/** Initialization constructor for a table record node. */
+	explicit FCustomTableTreeNode(const FName InName, TWeakPtr<FTable> InParentTable, int32 InRowIndex, const FSlateBrush* InIconBrush, FLinearColor InIconColor, FLinearColor InColor, bool IsGroup)
+		: FTableTreeNode(InName, InParentTable, InRowIndex, IsGroup)
+		, IconBrush(InIconBrush)
+		, IconColor(InIconColor)
 		, Color(InColor)
 	{
 	}
@@ -184,6 +194,16 @@ public:
 	explicit FCustomTableTreeNode(const FName InName, TWeakPtr<FTable> InParentTable, const FSlateBrush* InIconBrush, FLinearColor InColor)
 		: FTableTreeNode(InName, InParentTable)
 		, IconBrush(InIconBrush)
+		, IconColor(InColor)
+		, Color(InColor)
+	{
+	}
+
+	/** Initialization constructor for the group node. */
+	explicit FCustomTableTreeNode(const FName InName, TWeakPtr<FTable> InParentTable, const FSlateBrush* InIconBrush, FLinearColor InIconColor, FLinearColor InColor)
+		: FTableTreeNode(InName, InParentTable)
+		, IconBrush(InIconBrush)
+		, IconColor(InIconColor)
 		, Color(InColor)
 	{
 	}
@@ -192,9 +212,6 @@ public:
 	{
 	}
 
-	/**
-	 * @return a brush icon for this node.
-	 */
 	virtual const FSlateBrush* GetIcon() const override
 	{
 		return IconBrush;
@@ -208,9 +225,11 @@ public:
 		IconBrush = InIconBrush;
 	}
 
-	/**
-	 * @return the color tint for icon and name text.
-	 */
+	virtual FLinearColor GetIconColor() const override
+	{
+		return IconColor;
+	}
+
 	virtual FLinearColor GetColor() const override
 	{
 		return Color;
@@ -220,7 +239,10 @@ private:
 	/** The icon of this node. */
 	const FSlateBrush* IconBrush;
 
-	/** The color tint of this node. */
+	/** The color tint for the icon of this node. */
+	FLinearColor IconColor;
+
+	/** The color tint for the name text of this node. */
 	FLinearColor Color;
 };
 
