[Insights]
  - Enabled deprecation warnings for API moved in another namespace.
  - Fixed the deprecation warnings in the Insights plugins.

diff --git a/Engine/Plugins/Animation/GameplayInsights/Source/GameplayInsights/Private/AnimNodesTrack.cpp b/Engine/Plugins/Animation/GameplayInsights/Source/GameplayInsights/Private/AnimNodesTrack.cpp
--- a/Engine/Plugins/Animation/GameplayInsights/Source/GameplayInsights/Private/AnimNodesTrack.cpp
+++ b/Engine/Plugins/Animation/GameplayInsights/Source/GameplayInsights/Private/AnimNodesTrack.cpp
@@ -1,19 +1,26 @@
 // Copyright Epic Games, Inc. All Rights Reserved.
 
 #include "AnimNodesTrack.h"
-#include "Insights/ViewModels/ITimingViewDrawHelper.h"
-#include "GameplayProvider.h"
+
 #include "AnimationProvider.h"
-#include "Insights/ViewModels/TimingTrackViewport.h"
-#include "Insights/ViewModels/TimingEvent.h"
 #include "AnimationSharedData.h"
 #include "Framework/MultiBox/MultiBoxBuilder.h"
+#include "GameplayProvider.h"
+#include "GameplayTimingViewExtender.h"
+#include "SkeletalMeshPoseTrack.h"
+
+// TraceServices
+#include "TraceServices/Model/Frames.h"
+
+// TraceInsightsCore
+#include "InsightsCore/Common/TimeUtils.h"
+
+// TraceInsights
+#include "Insights/ViewModels/ITimingViewDrawHelper.h"
+#include "Insights/ViewModels/TimingEvent.h"
 #include "Insights/ViewModels/TimingEventSearch.h"
+#include "Insights/ViewModels/TimingTrackViewport.h"
 #include "Insights/ViewModels/TooltipDrawState.h"
-#include "Insights/Common/TimeUtils.h"
-#include "TraceServices/Model/Frames.h"
-#include "SkeletalMeshPoseTrack.h"
-#include "GameplayTimingViewExtender.h"
 
 #if WITH_ENGINE
 #include "Animation/AnimTrace.h"
@@ -137,8 +144,8 @@ void FAnimNodesTrack::InitTooltip(FTooltipDrawState& Tooltip, const ITimingEvent
 
 		Tooltip.AddTitle(FText::Format(LOCTEXT("GraphPhaseFormat", "{0} Anim Graph"), FText::FromString(GetPhaseName(InMessage.Phase))).ToString());
 
-		Tooltip.AddNameValueTextLine(LOCTEXT("EventDuration", "Duration").ToString(), TimeUtils::FormatTimeAuto(InFoundEndTime - InFoundStartTime));
-		Tooltip.AddNameValueTextLine(LOCTEXT("EventTime", "Time").ToString(), TimeUtils::FormatTimeAuto(InFoundStartTime));
+		Tooltip.AddNameValueTextLine(LOCTEXT("EventDuration", "Duration").ToString(), UE::Insights::FormatTimeAuto(InFoundEndTime - InFoundStartTime));
+		Tooltip.AddNameValueTextLine(LOCTEXT("EventTime", "Time").ToString(), UE::Insights::FormatTimeAuto(InFoundStartTime));
 		Tooltip.AddNameValueTextLine(LOCTEXT("EventWorld", "World").ToString(), GetGameplayTrack().GetWorldName(SharedData.GetAnalysisSession()).ToString());
 
 		Tooltip.UpdateLayout();
diff --git a/Engine/Plugins/Animation/GameplayInsights/Source/GameplayInsights/Private/AnimNotifiesTrack.cpp b/Engine/Plugins/Animation/GameplayInsights/Source/GameplayInsights/Private/AnimNotifiesTrack.cpp
--- a/Engine/Plugins/Animation/GameplayInsights/Source/GameplayInsights/Private/AnimNotifiesTrack.cpp
+++ b/Engine/Plugins/Animation/GameplayInsights/Source/GameplayInsights/Private/AnimNotifiesTrack.cpp
@@ -1,18 +1,23 @@
 // Copyright Epic Games, Inc. All Rights Reserved.
 
 #include "AnimNotifiesTrack.h"
-#include "Insights/ViewModels/ITimingViewDrawHelper.h"
-#include "GameplayProvider.h"
+
 #include "AnimationProvider.h"
-#include "Insights/ViewModels/TimingTrackViewport.h"
-#include "Insights/ViewModels/TimingEvent.h"
 #include "AnimationSharedData.h"
 #include "Framework/MultiBox/MultiBoxBuilder.h"
+#include "GameplayProvider.h"
+#include "TraceServices/Model/Frames.h"
+#include "VariantTreeNode.h"
+
+// TraceInsightsCore
+#include "InsightsCore/Common/TimeUtils.h"
+
+// TraceInsights
+#include "Insights/ViewModels/ITimingViewDrawHelper.h"
+#include "Insights/ViewModels/TimingEvent.h"
 #include "Insights/ViewModels/TimingEventSearch.h"
+#include "Insights/ViewModels/TimingTrackViewport.h"
 #include "Insights/ViewModels/TooltipDrawState.h"
-#include "Insights/Common/TimeUtils.h"
-#include "VariantTreeNode.h"
-#include "TraceServices/Model/Frames.h"
 
 #define LOCTEXT_NAMESPACE "AnimNotifiesTrack"
 
@@ -148,8 +153,8 @@ void FAnimNotifiesTrack::InitTooltip(FTooltipDrawState& Tooltip, const ITimingEv
 			}
 		}
 
-		Tooltip.AddNameValueTextLine(LOCTEXT("EventDuration", "Duration").ToString(), TimeUtils::FormatTimeAuto(InFoundEndTime - InFoundStartTime));
-		Tooltip.AddNameValueTextLine(LOCTEXT("EventTime", "Time").ToString(), TimeUtils::FormatTimeAuto(InFoundStartTime));
+		Tooltip.AddNameValueTextLine(LOCTEXT("EventDuration", "Duration").ToString(), UE::Insights::FormatTimeAuto(InFoundEndTime - InFoundStartTime));
+		Tooltip.AddNameValueTextLine(LOCTEXT("EventTime", "Time").ToString(), UE::Insights::FormatTimeAuto(InFoundStartTime));
 
 		if(InMessage.NotifyEventType != EAnimNotifyMessageType::SyncMarker)
 		{
diff --git a/Engine/Plugins/Animation/GameplayInsights/Source/GameplayInsights/Private/AnimationSharedData.cpp b/Engine/Plugins/Animation/GameplayInsights/Source/GameplayInsights/Private/AnimationSharedData.cpp
--- a/Engine/Plugins/Animation/GameplayInsights/Source/GameplayInsights/Private/AnimationSharedData.cpp
+++ b/Engine/Plugins/Animation/GameplayInsights/Source/GameplayInsights/Private/AnimationSharedData.cpp
@@ -1,26 +1,29 @@
 // Copyright Epic Games, Inc. All Rights Reserved.
 
 #include "AnimationSharedData.h"
-#include "TraceServices/Model/AnalysisSession.h"
-#include "GameplaySharedData.h"
-#include "ObjectEventsTrack.h"
-#include "SkeletalMeshPoseTrack.h"
-#include "SkeletalMeshCurvesTrack.h"
-#include "AnimationTickRecordsTrack.h"
-#include "Insights/ITimingViewSession.h"
-#include "Insights/ViewModels/TimingEvent.h"
+
 #include "AnimationProvider.h"
-#include "GameplayProvider.h"
-#include "Framework/MultiBox/MultiBoxBuilder.h"
+#include "AnimationTickRecordsTrack.h"
 #include "AnimNodesTrack.h"
-#include "GameplayTimingViewExtender.h"
-#include "SAnimGraphSchematicView.h"
+#include "AnimNotifiesTrack.h"
+#include "Framework/MultiBox/MultiBoxBuilder.h"
 #include "GameplayInsightsModule.h"
+#include "GameplayProvider.h"
+#include "GameplaySharedData.h"
+#include "GameplayTimingViewExtender.h"
 #include "Modules/ModuleManager.h"
-#include "Widgets/Docking/SDockTab.h"
-#include "AnimNotifiesTrack.h"
 #include "MontageTrack.h"
+#include "ObjectEventsTrack.h"
+#include "SAnimGraphSchematicView.h"
+#include "SkeletalMeshCurvesTrack.h"
+#include "SkeletalMeshPoseTrack.h"
 #include "STrackVariantValueView.h"
+#include "TraceServices/Model/AnalysisSession.h"
+#include "Widgets/Docking/SDockTab.h"
+
+// TraceInsights
+#include "Insights/ITimingViewSession.h"
+#include "Insights/ViewModels/TimingEvent.h"
 
 #if WITH_ENGINE
 #include "Animation/AnimTypes.h"
@@ -48,7 +51,7 @@ FAnimationSharedData::FAnimationSharedData(FGameplaySharedData& InGameplayShared
 {
 }
 
-void FAnimationSharedData::OnBeginSession(Insights::ITimingViewSession& InTimingViewSession)
+void FAnimationSharedData::OnBeginSession(UE::Insights::Timing::ITimingViewSession& InTimingViewSession)
 {
 	TimingViewSession = &InTimingViewSession;
 
@@ -58,7 +61,7 @@ void FAnimationSharedData::OnBeginSession(Insights::ITimingViewSession& InTiming
 	TimeMarkerChangedHandle = InTimingViewSession.OnTimeMarkerChanged().AddRaw(this, &FAnimationSharedData::OnTimeMarkerChanged);
 }
 
-void FAnimationSharedData::OnEndSession(Insights::ITimingViewSession& InTimingViewSession)
+void FAnimationSharedData::OnEndSession(UE::Insights::Timing::ITimingViewSession& InTimingViewSession)
 {
 	SkeletalMeshPoseTracks.Reset();
 	AnimationTickRecordsTracks.Reset();
@@ -68,7 +71,7 @@ void FAnimationSharedData::OnEndSession(Insights::ITimingViewSession& InTimingVi
 	TimingViewSession = nullptr;
 }
 
-void FAnimationSharedData::Tick(Insights::ITimingViewSession& InTimingViewSession, const TraceServices::IAnalysisSession& InAnalysisSession)
+void FAnimationSharedData::Tick(UE::Insights::Timing::ITimingViewSession& InTimingViewSession, const TraceServices::IAnalysisSession& InAnalysisSession)
 {
 	AnalysisSession = &InAnalysisSession;
 
@@ -501,7 +504,7 @@ void FAnimationSharedData::ToggleMontageTracks()
 	}
 }
 
-void FAnimationSharedData::OnTimeMarkerChanged(Insights::ETimeChangedFlags InFlags, double InTimeMarker)
+void FAnimationSharedData::OnTimeMarkerChanged(UE::Insights::Timing::ETimeChangedFlags InFlags, double InTimeMarker)
 {
 	bTimeMarkerValid = InTimeMarker != std::numeric_limits<double>::infinity();
 	MarkerTime = InTimeMarker;
@@ -672,7 +675,7 @@ void FAnimationSharedData::OpenAnimGraphTab(uint64 InAnimInstanceId) const
 		
 		TSharedPtr<SAnimGraphSchematicView> AnimGraphView = SNew(SAnimGraphSchematicView, InAnimInstanceId, TimingViewSession->GetTimeMarker(), *AnalysisSession);
 		TimingViewSession->OnTimeMarkerChanged().AddLambda(
-			[AnimGraphViewWeakPtr = TWeakPtr<SAnimGraphSchematicView>(AnimGraphView)](Insights::ETimeChangedFlags InFlags, double TimeMarker)
+			[AnimGraphViewWeakPtr = TWeakPtr<SAnimGraphSchematicView>(AnimGraphView)](UE::Insights::Timing::ETimeChangedFlags InFlags, double TimeMarker)
 			{
 				if (AnimGraphViewWeakPtr.IsValid())
 				{
diff --git a/Engine/Plugins/Animation/GameplayInsights/Source/GameplayInsights/Private/AnimationSharedData.h b/Engine/Plugins/Animation/GameplayInsights/Source/GameplayInsights/Private/AnimationSharedData.h
--- a/Engine/Plugins/Animation/GameplayInsights/Source/GameplayInsights/Private/AnimationSharedData.h
+++ b/Engine/Plugins/Animation/GameplayInsights/Source/GameplayInsights/Private/AnimationSharedData.h
@@ -7,8 +7,8 @@
 
 class FGameplaySharedData;
 namespace TraceServices { class IAnalysisSession; }
-namespace Insights { class ITimingViewSession; }
-namespace Insights { enum class ETimeChangedFlags : int32; }
+namespace UE::Insights::Timing { class ITimingViewSession; }
+namespace UE::Insights::Timing { enum class ETimeChangedFlags : int32; }
 class FSkeletalMeshPoseTrack;
 class FSkeletalMeshCurvesTrack;
 class FAnimationTickRecordsTrack;
@@ -27,9 +27,9 @@ class FAnimationSharedData
 public:
 	FAnimationSharedData(FGameplaySharedData& InGameplaySharedData);
 
-	void OnBeginSession(Insights::ITimingViewSession& InTimingViewSession);
-	void OnEndSession(Insights::ITimingViewSession& InTimingViewSession);
-	void Tick(Insights::ITimingViewSession& InTimingViewSession, const TraceServices::IAnalysisSession& InAnalysisSession);
+	void OnBeginSession(UE::Insights::Timing::ITimingViewSession& InTimingViewSession);
+	void OnEndSession(UE::Insights::Timing::ITimingViewSession& InTimingViewSession);
+	void Tick(UE::Insights::Timing::ITimingViewSession& InTimingViewSession, const TraceServices::IAnalysisSession& InAnalysisSession);
 	void ExtendFilterMenu(FMenuBuilder& InMenuBuilder);
 
 #if WITH_ENGINE
@@ -86,7 +86,7 @@ public:
 private:
 	// UI handlers
 	void ToggleAnimationTracks();
-	void OnTimeMarkerChanged(Insights::ETimeChangedFlags InFlags, double InTimeMarker);
+	void OnTimeMarkerChanged(UE::Insights::Timing::ETimeChangedFlags InFlags, double InTimeMarker);
 	void ToggleSkeletalMeshPoseTracks();
 	void ToggleSkeletalMeshCurveTracks();
 	void ToggleTickRecordTracks();
@@ -102,7 +102,7 @@ private:
 	const TraceServices::IAnalysisSession* AnalysisSession;
 
 	// Cached timing view session, set in OnBeginSession/OnEndSession
-	Insights::ITimingViewSession* TimingViewSession;
+	UE::Insights::Timing::ITimingViewSession* TimingViewSession;
 
 	// All the tracks we manage
 	TArray<TSharedRef<FSkeletalMeshPoseTrack>> SkeletalMeshPoseTracks;
diff --git a/Engine/Plugins/Animation/GameplayInsights/Source/GameplayInsights/Private/AnimationTickRecordsTrack.cpp b/Engine/Plugins/Animation/GameplayInsights/Source/GameplayInsights/Private/AnimationTickRecordsTrack.cpp
--- a/Engine/Plugins/Animation/GameplayInsights/Source/GameplayInsights/Private/AnimationTickRecordsTrack.cpp
+++ b/Engine/Plugins/Animation/GameplayInsights/Source/GameplayInsights/Private/AnimationTickRecordsTrack.cpp
@@ -1,19 +1,23 @@
 // Copyright Epic Games, Inc. All Rights Reserved.
 
 #include "AnimationTickRecordsTrack.h"
-#include "GameplayProvider.h"
+
 #include "AnimationProvider.h"
-#include "Insights/ViewModels/TimingTrackViewport.h"
-#include "Insights/ViewModels/TimingEvent.h"
-#include "Insights/ViewModels/TooltipDrawState.h"
 #include "AnimationSharedData.h"
-#include "Insights/ViewModels/TimingEventSearch.h"
-#include "Insights/ViewModels/TooltipDrawState.h"
+#include "Framework/MultiBox/MultiBoxBuilder.h"
+#include "GameplayProvider.h"
+#include "GameplaySharedData.h"
 #include "Templates/Invoke.h"
+
+// TraceInsightsCore
+#include "InsightsCore/Common/TimeUtils.h"
+
+// TraceInsights
 #include "Insights/ViewModels/GraphTrackBuilder.h"
-#include "Insights/Common/TimeUtils.h"
-#include "GameplaySharedData.h"
-#include "Framework/MultiBox/MultiBoxBuilder.h"
+#include "Insights/ViewModels/TimingEvent.h"
+#include "Insights/ViewModels/TimingEventSearch.h"
+#include "Insights/ViewModels/TimingTrackViewport.h"
+#include "Insights/ViewModels/TooltipDrawState.h"
 
 #if WITH_EDITOR
 #include "Animation/AnimBlueprintGeneratedClass.h"
@@ -34,7 +38,7 @@ FString FTickRecordSeries::FormatValue(double Value) const
 	switch (Type)
 	{
 	case ESeriesType::PlaybackTime:
-		return TimeUtils::FormatTimeAuto(Value);
+		return UE::Insights::FormatTimeAuto(Value);
 	case ESeriesType::BlendWeight:
 	case ESeriesType::RootMotionWeight:
 	case ESeriesType::PlayRate:
diff --git a/Engine/Plugins/Animation/GameplayInsights/Source/GameplayInsights/Private/GameplayGraphTrack.cpp b/Engine/Plugins/Animation/GameplayInsights/Source/GameplayInsights/Private/GameplayGraphTrack.cpp
--- a/Engine/Plugins/Animation/GameplayInsights/Source/GameplayInsights/Private/GameplayGraphTrack.cpp
+++ b/Engine/Plugins/Animation/GameplayInsights/Source/GameplayInsights/Private/GameplayGraphTrack.cpp
@@ -1,20 +1,27 @@
 // Copyright Epic Games, Inc. All Rights Reserved.
 
 #include "GameplayGraphTrack.h"
-#include "GameplayProvider.h"
+
 #include "AnimationProvider.h"
-#include "Insights/ViewModels/TimingTrackViewport.h"
-#include "Insights/ViewModels/TimingEvent.h"
-#include "GameplaySharedData.h"
+#include "Fonts/FontMeasure.h"
 #include "Framework/MultiBox/MultiBoxBuilder.h"
-#include "Insights/ViewModels/TimingEventSearch.h"
-#include "Insights/ViewModels/TooltipDrawState.h"
+#include "GameplayProvider.h"
+#include "GameplaySharedData.h"
 #include "Modules/ModuleManager.h"
+
+// TraceServices
+#include "TraceServices/Model/Frames.h"
+
+// TraceInsightsCore
+#include "InsightsCore/Common/PaintUtils.h"
+
+// TraceInsights
 #include "Insights/ViewModels/GraphTrackBuilder.h"
-#include "Insights/Common/PaintUtils.h"
 #include "Insights/ViewModels/ITimingViewDrawHelper.h"
-#include "Fonts/FontMeasure.h"
-#include "TraceServices/Model/Frames.h"
+#include "Insights/ViewModels/TimingEvent.h"
+#include "Insights/ViewModels/TimingEventSearch.h"
+#include "Insights/ViewModels/TimingTrackViewport.h"
+#include "Insights/ViewModels/TooltipDrawState.h"
 
 #define LOCTEXT_NAMESPACE "GameplayGraphTrack"
 
diff --git a/Engine/Plugins/Animation/GameplayInsights/Source/GameplayInsights/Private/GameplayInsightsModule.cpp b/Engine/Plugins/Animation/GameplayInsights/Source/GameplayInsights/Private/GameplayInsightsModule.cpp
--- a/Engine/Plugins/Animation/GameplayInsights/Source/GameplayInsights/Private/GameplayInsightsModule.cpp
+++ b/Engine/Plugins/Animation/GameplayInsights/Source/GameplayInsights/Private/GameplayInsightsModule.cpp
@@ -50,7 +50,7 @@ void FGameplayInsightsModule::StartupModule()
 	LLM_SCOPE_BYNAME(TEXT("Insights/GameplayInsights"));
 
 	IModularFeatures::Get().RegisterModularFeature(TraceServices::ModuleFeatureName, &GameplayTraceModule);
-	IModularFeatures::Get().RegisterModularFeature(Insights::TimingViewExtenderFeatureName, &GameplayTimingViewExtender);
+	IModularFeatures::Get().RegisterModularFeature(UE::Insights::Timing::TimingViewExtenderFeatureName, &GameplayTimingViewExtender);
 
 	TickerHandle = FTSTicker::GetCoreTicker().AddTicker(TEXT("GameplayInsights"), 0.0f, [this](float DeltaTime)
 	{
@@ -228,7 +228,7 @@ void FGameplayInsightsModule::ShutdownModule()
 	FTSTicker::GetCoreTicker().RemoveTicker(TickerHandle);
 
 	IModularFeatures::Get().UnregisterModularFeature(TraceServices::ModuleFeatureName, &GameplayTraceModule);
-	IModularFeatures::Get().UnregisterModularFeature(Insights::TimingViewExtenderFeatureName, &GameplayTimingViewExtender);
+	IModularFeatures::Get().UnregisterModularFeature(UE::Insights::Timing::TimingViewExtenderFeatureName, &GameplayTimingViewExtender);
 }
 
 TSharedRef<SDockTab> FGameplayInsightsModule::SpawnTimingProfilerDocumentTab(const FTabManager::FSearchPreference& InSearchPreference)
diff --git a/Engine/Plugins/Animation/GameplayInsights/Source/GameplayInsights/Private/GameplaySharedData.cpp b/Engine/Plugins/Animation/GameplayInsights/Source/GameplayInsights/Private/GameplaySharedData.cpp
--- a/Engine/Plugins/Animation/GameplayInsights/Source/GameplayInsights/Private/GameplaySharedData.cpp
+++ b/Engine/Plugins/Animation/GameplayInsights/Source/GameplayInsights/Private/GameplaySharedData.cpp
@@ -1,18 +1,24 @@
 // Copyright Epic Games, Inc. All Rights Reserved.
 
 #include "GameplaySharedData.h"
-#include "TraceServices/Model/AnalysisSession.h"
-#include "ObjectEventsTrack.h"
+
 #include "Algo/Sort.h"
-#include "GameplayProvider.h"
-#include "Insights/ITimingViewSession.h"
 #include "Framework/MultiBox/MultiBoxBuilder.h"
-#include "SGameplayTrackTree.h"
-#include "GameplayInsightsModule.h"
-#include "STrackVariantValueView.h"
 #include "Modules/ModuleManager.h"
 #include "Widgets/Docking/SDockTab.h"
+
+// TraceServices
+#include "TraceServices/Model/AnalysisSession.h"
+
+// TraceInsights
+#include "Insights/ITimingViewSession.h"
+
+#include "GameplayInsightsModule.h"
+#include "GameplayProvider.h"
+#include "ObjectEventsTrack.h"
 #include "ObjectPropertiesTrack.h"
+#include "SGameplayTrackTree.h"
+#include "STrackVariantValueView.h"
 
 #define LOCTEXT_NAMESPACE "GameplaySharedData"
 
@@ -24,21 +30,21 @@ FGameplaySharedData::FGameplaySharedData()
 {
 }
 
-void FGameplaySharedData::OnBeginSession(Insights::ITimingViewSession& InTimingViewSession)
+void FGameplaySharedData::OnBeginSession(UE::Insights::Timing::ITimingViewSession& InTimingViewSession)
 {
 	TimingViewSession = &InTimingViewSession;
 
 	ObjectTracks.Reset();
 }
 
-void FGameplaySharedData::OnEndSession(Insights::ITimingViewSession& InTimingViewSession)
+void FGameplaySharedData::OnEndSession(UE::Insights::Timing::ITimingViewSession& InTimingViewSession)
 {
 	ObjectTracks.Reset();
 
 	TimingViewSession = nullptr;
 }
 
-TSharedRef<FObjectEventsTrack> FGameplaySharedData::GetObjectEventsTrackForId(Insights::ITimingViewSession& InTimingViewSession, const TraceServices::IAnalysisSession& InAnalysisSession, const FObjectInfo& InObjectInfo)
+TSharedRef<FObjectEventsTrack> FGameplaySharedData::GetObjectEventsTrackForId(UE::Insights::Timing::ITimingViewSession& InTimingViewSession, const TraceServices::IAnalysisSession& InAnalysisSession, const FObjectInfo& InObjectInfo)
 {
 	const FGameplayProvider* GameplayProvider = InAnalysisSession.ReadProvider<FGameplayProvider>(FGameplayProvider::ProviderName);
 	check(GameplayProvider);
@@ -136,7 +142,7 @@ void FGameplaySharedData::MakeTrackAndAncestorsVisible(const TSharedRef<FObjectE
 	InvalidateObjectTracksOrder();
 }
 
-void FGameplaySharedData::Tick(Insights::ITimingViewSession& InTimingViewSession, const TraceServices::IAnalysisSession& InAnalysisSession)
+void FGameplaySharedData::Tick(UE::Insights::Timing::ITimingViewSession& InTimingViewSession, const TraceServices::IAnalysisSession& InAnalysisSession)
 {
 	AnalysisSession = &InAnalysisSession;
 
diff --git a/Engine/Plugins/Animation/GameplayInsights/Source/GameplayInsights/Private/GameplaySharedData.h b/Engine/Plugins/Animation/GameplayInsights/Source/GameplayInsights/Private/GameplaySharedData.h
--- a/Engine/Plugins/Animation/GameplayInsights/Source/GameplayInsights/Private/GameplaySharedData.h
+++ b/Engine/Plugins/Animation/GameplayInsights/Source/GameplayInsights/Private/GameplaySharedData.h
@@ -6,7 +6,7 @@
 #include "Framework/Docking/TabManager.h"
 
 namespace TraceServices { class IAnalysisSession; }
-namespace Insights { class ITimingViewSession; }
+namespace UE::Insights::Timing { class ITimingViewSession; }
 class FObjectEventsTrack;
 class FObjectPropertiesTrack;
 class FSkeletalMeshPoseTrack;
@@ -19,13 +19,13 @@ class FGameplaySharedData
 public:
 	FGameplaySharedData();
 
-	void OnBeginSession(Insights::ITimingViewSession& InTimingViewSession);
-	void OnEndSession(Insights::ITimingViewSession& InTimingViewSession);
-	void Tick(Insights::ITimingViewSession& InTimingViewSession, const TraceServices::IAnalysisSession& InAnalysisSession);
+	void OnBeginSession(UE::Insights::Timing::ITimingViewSession& InTimingViewSession);
+	void OnEndSession(UE::Insights::Timing::ITimingViewSession& InTimingViewSession);
+	void Tick(UE::Insights::Timing::ITimingViewSession& InTimingViewSession, const TraceServices::IAnalysisSession& InAnalysisSession);
 	void ExtendFilterMenu(FMenuBuilder& InMenuBuilder);
 
 	// Helper function. Builds object track hierarchy on-demand and returns a track for the supplied object info.
-	TSharedRef<FObjectEventsTrack> GetObjectEventsTrackForId(Insights::ITimingViewSession& InTimingViewSession, const TraceServices::IAnalysisSession& InAnalysisSession, const FObjectInfo& InObjectInfo);
+	TSharedRef<FObjectEventsTrack> GetObjectEventsTrackForId(UE::Insights::Timing::ITimingViewSession& InTimingViewSession, const TraceServices::IAnalysisSession& InAnalysisSession, const FObjectInfo& InObjectInfo);
 
 	// Helper function to make tracks visible when children are first added
 	void MakeTrackAndAncestorsVisible(const TSharedRef<FObjectEventsTrack>& InObjectEventsTrack, bool bInVisible);
@@ -40,7 +40,7 @@ public:
 	const TraceServices::IAnalysisSession& GetAnalysisSession() const { return *AnalysisSession; }
 
 	// Get the timing view session
-	Insights::ITimingViewSession& GetTimingViewSession() { return *TimingViewSession; }
+	UE::Insights::Timing::ITimingViewSession& GetTimingViewSession() { return *TimingViewSession; }
 
 	// Check the validity of the analysis session
 	bool IsAnalysisSessionValid() const { return AnalysisSession != nullptr; }
@@ -97,7 +97,7 @@ private:
 	const TraceServices::IAnalysisSession* AnalysisSession;
 
 	// Cached timing view session, set in OnBeginSession/OnEndSession
-	Insights::ITimingViewSession* TimingViewSession;
+	UE::Insights::Timing::ITimingViewSession* TimingViewSession;
 
 	// Delegate fired when tracks change
 	FSimpleMulticastDelegate OnTracksChangedDelegate;
diff --git a/Engine/Plugins/Animation/GameplayInsights/Source/GameplayInsights/Private/GameplayTimingViewExtender.cpp b/Engine/Plugins/Animation/GameplayInsights/Source/GameplayInsights/Private/GameplayTimingViewExtender.cpp
--- a/Engine/Plugins/Animation/GameplayInsights/Source/GameplayInsights/Private/GameplayTimingViewExtender.cpp
+++ b/Engine/Plugins/Animation/GameplayInsights/Source/GameplayInsights/Private/GameplayTimingViewExtender.cpp
@@ -19,7 +19,7 @@
 
 #define LOCTEXT_NAMESPACE "GameplayTimingViewExtender"
 
-void FGameplayTimingViewExtender::OnBeginSession(Insights::ITimingViewSession& InSession)
+void FGameplayTimingViewExtender::OnBeginSession(UE::Insights::Timing::ITimingViewSession& InSession)
 {
 	FPerSessionData* PerSessionData = PerSessionDataMap.Find(&InSession);
 	if(PerSessionData == nullptr)
@@ -55,7 +55,7 @@ void FGameplayTimingViewExtender::OnBeginSession(Insights::ITimingViewSession& I
 	}
 }
 
-void FGameplayTimingViewExtender::OnEndSession(Insights::ITimingViewSession& InSession)
+void FGameplayTimingViewExtender::OnEndSession(UE::Insights::Timing::ITimingViewSession& InSession)
 {
 	FPerSessionData* PerSessionData = PerSessionDataMap.Find(&InSession);
 	if(PerSessionData != nullptr)
@@ -75,7 +75,7 @@ void FGameplayTimingViewExtender::OnEndSession(Insights::ITimingViewSession& InS
 	PerSessionDataMap.Remove(&InSession);
 }
 
-void FGameplayTimingViewExtender::Tick(Insights::ITimingViewSession& InSession, const TraceServices::IAnalysisSession& InAnalysisSession)
+void FGameplayTimingViewExtender::Tick(UE::Insights::Timing::ITimingViewSession& InSession, const TraceServices::IAnalysisSession& InAnalysisSession)
 {
 	FPerSessionData* PerSessionData = PerSessionDataMap.Find(&InSession);
 	if(PerSessionData != nullptr)
@@ -85,7 +85,7 @@ void FGameplayTimingViewExtender::Tick(Insights::ITimingViewSession& InSession,
 	}
 }
 
-void FGameplayTimingViewExtender::ExtendFilterMenu(Insights::ITimingViewSession& InSession, FMenuBuilder& InMenuBuilder)
+void FGameplayTimingViewExtender::ExtendFilterMenu(UE::Insights::Timing::ITimingViewSession& InSession, FMenuBuilder& InMenuBuilder)
 {
 	FPerSessionData* PerSessionData = PerSessionDataMap.Find(&InSession);
 	if(PerSessionData != nullptr)
diff --git a/Engine/Plugins/Animation/GameplayInsights/Source/GameplayInsights/Private/GameplayTimingViewExtender.h b/Engine/Plugins/Animation/GameplayInsights/Source/GameplayInsights/Private/GameplayTimingViewExtender.h
--- a/Engine/Plugins/Animation/GameplayInsights/Source/GameplayInsights/Private/GameplayTimingViewExtender.h
+++ b/Engine/Plugins/Animation/GameplayInsights/Source/GameplayInsights/Private/GameplayTimingViewExtender.h
@@ -4,7 +4,7 @@
 
 #include "Insights/ITimingViewExtender.h"
 
-namespace Insights { class ITimingViewSession; }
+namespace UE::Insights::Timing { class ITimingViewSession; }
 class UWorld;
 class FGameplaySharedData;
 class FAnimationSharedData;
@@ -12,14 +12,14 @@ class IAnimationBlueprintEditor;
 struct FCustomDebugObject;
 class SGameplayInsightsTransportControls;
 
-class FGameplayTimingViewExtender : public Insights::ITimingViewExtender
+class FGameplayTimingViewExtender : public UE::Insights::Timing::ITimingViewExtender
 {
 public:
 	// Insights::ITimingViewExtender interface
-	virtual void OnBeginSession(Insights::ITimingViewSession& InSession) override;
-	virtual void OnEndSession(Insights::ITimingViewSession& InSession) override;
-	virtual void Tick(Insights::ITimingViewSession& InSession, const TraceServices::IAnalysisSession& InAnalysisSession) override;
-	virtual void ExtendFilterMenu(Insights::ITimingViewSession& InSession, FMenuBuilder& InMenuBuilder) override;
+	virtual void OnBeginSession(UE::Insights::Timing::ITimingViewSession& InSession) override;
+	virtual void OnEndSession(UE::Insights::Timing::ITimingViewSession& InSession) override;
+	virtual void Tick(UE::Insights::Timing::ITimingViewSession& InSession, const TraceServices::IAnalysisSession& InAnalysisSession) override;
+	virtual void ExtendFilterMenu(UE::Insights::Timing::ITimingViewSession& InSession, FMenuBuilder& InMenuBuilder) override;
 
 #if WITH_EDITOR
 	// Gets a world to perform visualizations within, depending on context
@@ -44,5 +44,5 @@ private:
 	};
 
 	// The data we host per-session
-	TMap<Insights::ITimingViewSession*, FPerSessionData> PerSessionDataMap;
+	TMap<UE::Insights::Timing::ITimingViewSession*, FPerSessionData> PerSessionDataMap;
 };
diff --git a/Engine/Plugins/Animation/GameplayInsights/Source/GameplayInsights/Private/GameplayTraceModule.h b/Engine/Plugins/Animation/GameplayInsights/Source/GameplayInsights/Private/GameplayTraceModule.h
--- a/Engine/Plugins/Animation/GameplayInsights/Source/GameplayInsights/Private/GameplayTraceModule.h
+++ b/Engine/Plugins/Animation/GameplayInsights/Source/GameplayInsights/Private/GameplayTraceModule.h
@@ -6,7 +6,6 @@
 #include "TraceServices/ModuleService.h"
 
 class IAnimationProvider;
-namespace Insights { class ITimingViewSession; }
 
 class FGameplayTraceModule : public TraceServices::IModule
 {
diff --git a/Engine/Plugins/Animation/GameplayInsights/Source/GameplayInsights/Private/GameplayTrack.cpp b/Engine/Plugins/Animation/GameplayInsights/Source/GameplayInsights/Private/GameplayTrack.cpp
--- a/Engine/Plugins/Animation/GameplayInsights/Source/GameplayInsights/Private/GameplayTrack.cpp
+++ b/Engine/Plugins/Animation/GameplayInsights/Source/GameplayInsights/Private/GameplayTrack.cpp
@@ -1,20 +1,27 @@
 // Copyright Epic Games, Inc. All Rights Reserved.
 
 #include "GameplayTrack.h"
-#include "Insights/ViewModels/TimingEventsTrack.h"
+
 #include "Algo/Sort.h"
-#include "Framework/Application/SlateApplication.h"
 #include "Fonts/FontMeasure.h"
-#include "Insights/ViewModels/TimingTrackViewport.h"
-#include "Insights/Common/PaintUtils.h"
-#include "Insights/ViewModels/ITimingViewDrawHelper.h"
-#include "GameplaySharedData.h"
+#include "Framework/Application/SlateApplication.h"
 #include "Framework/MultiBox/MultiBoxBuilder.h"
 #include "GameplayProvider.h"
-#include "TraceServices/Model/AnalysisSession.h"
+#include "GameplaySharedData.h"
 #include "IGameplayProvider.h"
+
+// TraceServices
+#include "TraceServices/Model/AnalysisSession.h"
 #include "TraceServices/Model/Frames.h"
 
+// TraceInsightsCore
+#include "InsightsCore/Common/PaintUtils.h"
+
+// TraceInsights
+#include "Insights/ViewModels/ITimingViewDrawHelper.h"
+#include "Insights/ViewModels/TimingEventsTrack.h"
+#include "Insights/ViewModels/TimingTrackViewport.h"
+
 INSIGHTS_IMPLEMENT_RTTI(FGameplayTimingEventsTrack);
 
 #define LOCTEXT_NAMESPACE "GameplayTrack"
diff --git a/Engine/Plugins/Animation/GameplayInsights/Source/GameplayInsights/Private/SAnimGraphSchematicView.h b/Engine/Plugins/Animation/GameplayInsights/Source/GameplayInsights/Private/SAnimGraphSchematicView.h
--- a/Engine/Plugins/Animation/GameplayInsights/Source/GameplayInsights/Private/SAnimGraphSchematicView.h
+++ b/Engine/Plugins/Animation/GameplayInsights/Source/GameplayInsights/Private/SAnimGraphSchematicView.h
@@ -9,9 +9,9 @@
 #include "Widgets/SCompoundWidget.h"
 #include "Widgets/Views/STreeView.h"
 
-namespace Insights { class ITimingViewSession; }
-namespace Insights { enum class ETimeChangedFlags; }
 namespace TraceServices { class IAnalysisSession; }
+namespace UE::Insights::Timing { class ITimingViewSession; }
+namespace UE::Insights::Timing { enum class ETimeChangedFlags; }
 class FAnimationSharedData;
 class IInsightsManager;
 class FAnimGraphSchematicNode;
@@ -49,7 +49,7 @@ private:
 	void HandleGetPropertyChildren(TSharedRef<FAnimGraphSchematicPropertyNode> InItem, TArray<TSharedRef<FAnimGraphSchematicPropertyNode>>& OutChildren);
 
 	// Handle the time marker being scrubbed
-	void HandleTimeMarkerChanged(Insights::ETimeChangedFlags InFlags, double InTimeMarker);
+	void HandleTimeMarkerChanged(UE::Insights::Timing::ETimeChangedFlags InFlags, double InTimeMarker);
 
 	// Handle tree selection changing
 	void HandleSelectionChanged(TSharedPtr<FAnimGraphSchematicNode> InNode, ESelectInfo::Type InSelectInfo);
diff --git a/Engine/Plugins/Animation/GameplayInsights/Source/GameplayInsights/Private/SGameplayInsightsTransportControls.cpp b/Engine/Plugins/Animation/GameplayInsights/Source/GameplayInsights/Private/SGameplayInsightsTransportControls.cpp
--- a/Engine/Plugins/Animation/GameplayInsights/Source/GameplayInsights/Private/SGameplayInsightsTransportControls.cpp
+++ b/Engine/Plugins/Animation/GameplayInsights/Source/GameplayInsights/Private/SGameplayInsightsTransportControls.cpp
@@ -316,7 +316,7 @@ void SGameplayInsightsTransportControls::SetTimeMarker(double InTime, bool bInSc
 	bSettingMarker = false;
 }
 
-void SGameplayInsightsTransportControls::HandleTimeMarkerChanged(Insights::ETimeChangedFlags InFlags, double InTimeMarker)
+void SGameplayInsightsTransportControls::HandleTimeMarkerChanged(UE::Insights::Timing::ETimeChangedFlags InFlags, double InTimeMarker)
 {
 	if(!bSettingMarker)
 	{
@@ -325,9 +325,9 @@ void SGameplayInsightsTransportControls::HandleTimeMarkerChanged(Insights::ETime
 	}
 }
 
-void SGameplayInsightsTransportControls::HandleSelectionRangeChanged(Insights::ETimeChangedFlags InFlags, double InStartTime, double InEndTime)
+void SGameplayInsightsTransportControls::HandleSelectionRangeChanged(UE::Insights::Timing::ETimeChangedFlags InFlags, double InStartTime, double InEndTime)
 {
-	if(InFlags != Insights::ETimeChangedFlags::Interactive)
+	if(InFlags != UE::Insights::Timing::ETimeChangedFlags::Interactive)
 	{
 		bSelectionRangeValid = InStartTime != std::numeric_limits<double>::infinity() && InEndTime != std::numeric_limits<double>::infinity() && InEndTime > InStartTime;
 
diff --git a/Engine/Plugins/Animation/GameplayInsights/Source/GameplayInsights/Private/SGameplayInsightsTransportControls.h b/Engine/Plugins/Animation/GameplayInsights/Source/GameplayInsights/Private/SGameplayInsightsTransportControls.h
--- a/Engine/Plugins/Animation/GameplayInsights/Source/GameplayInsights/Private/SGameplayInsightsTransportControls.h
+++ b/Engine/Plugins/Animation/GameplayInsights/Source/GameplayInsights/Private/SGameplayInsightsTransportControls.h
@@ -12,7 +12,7 @@
 #include "ITransportControl.h"
 
 class FGameplaySharedData;
-namespace Insights { enum class ETimeChangedFlags; }
+namespace UE::Insights::Timing { enum class ETimeChangedFlags; }
 
 class SGameplayInsightsTransportControls : public SCompoundWidget
 {
@@ -44,9 +44,9 @@ private:
 
 	void SetTimeMarker(double InTime, bool bInScroll);
 
-	void HandleTimeMarkerChanged(Insights::ETimeChangedFlags InFlags, double InTimeMarker);
+	void HandleTimeMarkerChanged(UE::Insights::Timing::ETimeChangedFlags InFlags, double InTimeMarker);
 
-	void HandleSelectionRangeChanged(Insights::ETimeChangedFlags InFlags, double InStartTime, double InEndTime);
+	void HandleSelectionRangeChanged(UE::Insights::Timing::ETimeChangedFlags InFlags, double InStartTime, double InEndTime);
 
 private:
 	FGameplaySharedData* SharedData;
diff --git a/Engine/Plugins/Animation/GameplayInsights/Source/GameplayInsights/Private/STrackVariantValueView.cpp b/Engine/Plugins/Animation/GameplayInsights/Source/GameplayInsights/Private/STrackVariantValueView.cpp
--- a/Engine/Plugins/Animation/GameplayInsights/Source/GameplayInsights/Private/STrackVariantValueView.cpp
+++ b/Engine/Plugins/Animation/GameplayInsights/Source/GameplayInsights/Private/STrackVariantValueView.cpp
@@ -1,13 +1,16 @@
 // Copyright Epic Games, Inc. All Rights Reserved.
 
 #include "STrackVariantValueView.h"
+
 #include "SVariantValueView.h"
-#include "Insights/ViewModels/BaseTimingTrack.h"
 #include "GameplayTrack.h"
 #include "GameplayGraphTrack.h"
+
+// TraceInsights
+#include "Insights/ViewModels/BaseTimingTrack.h"
 #include "Insights/ITimingViewSession.h"
 
-void STrackVariantValueView::Construct(const FArguments& InArgs, const TSharedRef<FBaseTimingTrack>& InTimingTrack, Insights::ITimingViewSession& InTimingViewSession, const TraceServices::IAnalysisSession& InAnalysisSession)
+void STrackVariantValueView::Construct(const FArguments& InArgs, const TSharedRef<FBaseTimingTrack>& InTimingTrack, UE::Insights::Timing::ITimingViewSession& InTimingViewSession, const TraceServices::IAnalysisSession& InAnalysisSession)
 {
 	TimingTrack = InTimingTrack;
 	AnalysisSession = &InAnalysisSession;
@@ -44,7 +47,7 @@ void STrackVariantValueView::Construct(const FArguments& InArgs, const TSharedRe
 	}
 }
 
-void STrackVariantValueView::HandleTimeMarkerChanged(Insights::ETimeChangedFlags InFlags, double InTimeMarker)
+void STrackVariantValueView::HandleTimeMarkerChanged(UE::Insights::Timing::ETimeChangedFlags InFlags, double InTimeMarker)
 {
 	TraceServices::FAnalysisSessionReadScope SessionReadScope(*AnalysisSession);
 
diff --git a/Engine/Plugins/Animation/GameplayInsights/Source/GameplayInsights/Private/STrackVariantValueView.h b/Engine/Plugins/Animation/GameplayInsights/Source/GameplayInsights/Private/STrackVariantValueView.h
--- a/Engine/Plugins/Animation/GameplayInsights/Source/GameplayInsights/Private/STrackVariantValueView.h
+++ b/Engine/Plugins/Animation/GameplayInsights/Source/GameplayInsights/Private/STrackVariantValueView.h
@@ -7,9 +7,9 @@
 #include "Widgets/DeclarativeSyntaxSupport.h"
 
 class FBaseTimingTrack;
-namespace Insights { class ITimingViewSession; }
 namespace TraceServices { class IAnalysisSession; }
-namespace Insights { enum class ETimeChangedFlags; }
+namespace UE::Insights::Timing { class ITimingViewSession; }
+namespace UE::Insights::Timing { enum class ETimeChangedFlags; }
 class SVariantValueView;
 
 // Wrapper for a variant value view for a track
@@ -20,13 +20,13 @@ public:
 
 	SLATE_END_ARGS()
 
-	void Construct(const FArguments& InArgs, const TSharedRef<FBaseTimingTrack>& InTimingTrack, Insights::ITimingViewSession& InTimingViewSession, const TraceServices::IAnalysisSession& InAnalysisSession);
+	void Construct(const FArguments& InArgs, const TSharedRef<FBaseTimingTrack>& InTimingTrack, UE::Insights::Timing::ITimingViewSession& InTimingViewSession, const TraceServices::IAnalysisSession& InAnalysisSession);
 
 	TSharedPtr<FBaseTimingTrack> GetTimingTrack() const { return TimingTrack.Pin(); }
 
 private:
 	// Handle the time marker being scrubbed
-	void HandleTimeMarkerChanged(Insights::ETimeChangedFlags InFlags, double InTimeMarker);
+	void HandleTimeMarkerChanged(UE::Insights::Timing::ETimeChangedFlags InFlags, double InTimeMarker);
 
 private:
 	TWeakPtr<FBaseTimingTrack> TimingTrack;
diff --git a/Engine/Plugins/Animation/GameplayInsights/Source/RewindDebuggerVLog/Private/VLogTraceModule.h b/Engine/Plugins/Animation/GameplayInsights/Source/RewindDebuggerVLog/Private/VLogTraceModule.h
--- a/Engine/Plugins/Animation/GameplayInsights/Source/RewindDebuggerVLog/Private/VLogTraceModule.h
+++ b/Engine/Plugins/Animation/GameplayInsights/Source/RewindDebuggerVLog/Private/VLogTraceModule.h
@@ -6,7 +6,6 @@
 #include "TraceServices/ModuleService.h"
 
 class IAnimationProvider;
-namespace Insights { class ITimingViewSession; }
 
 class FVLogTraceModule : public TraceServices::IModule
 {
diff --git a/Engine/Plugins/Cameras/GameplayCameras/Source/GameplayCamerasEditor/Private/Trace/CameraSystemTraceModule.h b/Engine/Plugins/Cameras/GameplayCameras/Source/GameplayCamerasEditor/Private/Trace/CameraSystemTraceModule.h
--- a/Engine/Plugins/Cameras/GameplayCameras/Source/GameplayCamerasEditor/Private/Trace/CameraSystemTraceModule.h
+++ b/Engine/Plugins/Cameras/GameplayCameras/Source/GameplayCamerasEditor/Private/Trace/CameraSystemTraceModule.h
@@ -8,8 +8,6 @@
 
 #if UE_GAMEPLAY_CAMERAS_TRACE
 
-namespace Insights { class ITimingViewSession; }
-
 namespace UE::Cameras
 {
 
diff --git a/Engine/Plugins/Developer/Concert/ConcertInsights/ConcertInsightsVisualizer/Source/ConcertInsightsVisualizer/Private/ConcertInsightsVisualizer.cpp b/Engine/Plugins/Developer/Concert/ConcertInsights/ConcertInsightsVisualizer/Source/ConcertInsightsVisualizer/Private/ConcertInsightsVisualizer.cpp
--- a/Engine/Plugins/Developer/Concert/ConcertInsights/ConcertInsightsVisualizer/Source/ConcertInsightsVisualizer/Private/ConcertInsightsVisualizer.cpp
+++ b/Engine/Plugins/Developer/Concert/ConcertInsights/ConcertInsightsVisualizer/Source/ConcertInsightsVisualizer/Private/ConcertInsightsVisualizer.cpp
@@ -28,13 +28,13 @@ namespace UE::ConcertInsightsVisualizer
 	void FConcertInsightsVisualizerModule::RegisterInsightsExtensions()
 	{
 		IModularFeatures::Get().RegisterModularFeature(TraceServices::ModuleFeatureName, &ConcertInsightsModule);
-		IModularFeatures::Get().RegisterModularFeature(Insights::TimingViewExtenderFeatureName, &TimingViewExtender);
+		IModularFeatures::Get().RegisterModularFeature(UE::Insights::Timing::TimingViewExtenderFeatureName, &TimingViewExtender);
 	}
 
 	void FConcertInsightsVisualizerModule::UnregisterInsightsExtensions()
 	{
 		IModularFeatures::Get().UnregisterModularFeature(TraceServices::ModuleFeatureName, &ConcertInsightsModule);
-		IModularFeatures::Get().UnregisterModularFeature(Insights::TimingViewExtenderFeatureName, &TimingViewExtender);
+		IModularFeatures::Get().UnregisterModularFeature(UE::Insights::Timing::TimingViewExtenderFeatureName, &TimingViewExtender);
 	}
 }
 
diff --git a/Engine/Plugins/Developer/Concert/ConcertInsights/ConcertInsightsVisualizer/Source/ConcertInsightsVisualizer/Private/Extension/ConcertTimingViewExtender.cpp b/Engine/Plugins/Developer/Concert/ConcertInsights/ConcertInsightsVisualizer/Source/ConcertInsightsVisualizer/Private/Extension/ConcertTimingViewExtender.cpp
--- a/Engine/Plugins/Developer/Concert/ConcertInsights/ConcertInsightsVisualizer/Source/ConcertInsightsVisualizer/Private/Extension/ConcertTimingViewExtender.cpp
+++ b/Engine/Plugins/Developer/Concert/ConcertInsights/ConcertInsightsVisualizer/Source/ConcertInsightsVisualizer/Private/Extension/ConcertTimingViewExtender.cpp
@@ -7,7 +7,7 @@
 
 namespace UE::ConcertInsightsVisualizer
 {
-	void FConcertTimingViewExtender::OnBeginSession(Insights::ITimingViewSession& InSession)
+	void FConcertTimingViewExtender::OnBeginSession(UE::Insights::Timing::ITimingViewSession& InSession)
 	{
 		if (InSession.GetName() != FInsightsManagerTabs::TimingProfilerTabId)
 		{
@@ -28,7 +28,7 @@ namespace UE::ConcertInsightsVisualizer
 		}
 	}
 
-	void FConcertTimingViewExtender::OnEndSession(Insights::ITimingViewSession& InSession)
+	void FConcertTimingViewExtender::OnEndSession(UE::Insights::Timing::ITimingViewSession& InSession)
 	{
 		if (InSession.GetName() != FInsightsManagerTabs::TimingProfilerTabId)
 		{
@@ -44,7 +44,7 @@ namespace UE::ConcertInsightsVisualizer
 		PerSessionDataMap.Remove(&InSession);
 	}
 
-	void FConcertTimingViewExtender::Tick(Insights::ITimingViewSession& InTimingSession, const TraceServices::IAnalysisSession& InAnalysisSession)
+	void FConcertTimingViewExtender::Tick(UE::Insights::Timing::ITimingViewSession& InTimingSession, const TraceServices::IAnalysisSession& InAnalysisSession)
 	{
 		if (const FPerSessionData* PerSessionData = PerSessionDataMap.Find(&InTimingSession))
 		{
@@ -52,7 +52,7 @@ namespace UE::ConcertInsightsVisualizer
 		}
 	}
 
-	void FConcertTimingViewExtender::ExtendFilterMenu(Insights::ITimingViewSession& InSession, FMenuBuilder& InMenuBuilder)
+	void FConcertTimingViewExtender::ExtendFilterMenu(UE::Insights::Timing::ITimingViewSession& InSession, FMenuBuilder& InMenuBuilder)
 	{
 		if (const FPerSessionData* PerSessionData = PerSessionDataMap.Find(&InSession))
 		{
diff --git a/Engine/Plugins/Developer/Concert/ConcertInsights/ConcertInsightsVisualizer/Source/ConcertInsightsVisualizer/Private/Extension/ConcertTimingViewExtender.h b/Engine/Plugins/Developer/Concert/ConcertInsights/ConcertInsightsVisualizer/Source/ConcertInsightsVisualizer/Private/Extension/ConcertTimingViewExtender.h
--- a/Engine/Plugins/Developer/Concert/ConcertInsights/ConcertInsightsVisualizer/Source/ConcertInsightsVisualizer/Private/Extension/ConcertTimingViewExtender.h
+++ b/Engine/Plugins/Developer/Concert/ConcertInsights/ConcertInsightsVisualizer/Source/ConcertInsightsVisualizer/Private/Extension/ConcertTimingViewExtender.h
@@ -1,3 +1,4 @@
+
 // Copyright Epic Games, Inc. All Rights Reserved.
 
 #pragma once
@@ -10,16 +11,16 @@
 namespace UE::ConcertInsightsVisualizer
 {
 	/** Keeps track of FConcertTimingViewSession per analytics session. */
-	class FConcertTimingViewExtender : public Insights::ITimingViewExtender
+	class FConcertTimingViewExtender : public UE::Insights::Timing::ITimingViewExtender
 	{
 	public:
 
-		//~ Begin Insights::ITimingViewExtender Interface
-		virtual void OnBeginSession(Insights::ITimingViewSession& InSession) override;
-		virtual void OnEndSession(Insights::ITimingViewSession& InSession) override;
-		virtual void Tick(Insights::ITimingViewSession& InTimingSession, const TraceServices::IAnalysisSession& InAnalysisSession) override;
-		virtual void ExtendFilterMenu(Insights::ITimingViewSession& InSession, FMenuBuilder& InMenuBuilder) override;
-		//~ End Insights::ITimingViewExtender Interface
+		//~ Begin UE::Insights::Timing::ITimingViewExtender Interface
+		virtual void OnBeginSession(UE::Insights::Timing::ITimingViewSession& InSession) override;
+		virtual void OnEndSession(UE::Insights::Timing::ITimingViewSession& InSession) override;
+		virtual void Tick(UE::Insights::Timing::ITimingViewSession& InTimingSession, const TraceServices::IAnalysisSession& InAnalysisSession) override;
+		virtual void ExtendFilterMenu(UE::Insights::Timing::ITimingViewSession& InSession, FMenuBuilder& InMenuBuilder) override;
+		//~ End UE::Insights::Timing::ITimingViewExtender Interface
 
 	private:
 
@@ -29,7 +30,7 @@ namespace UE::ConcertInsightsVisualizer
 		};
 
 		//** The data we host per-session */
-		TMap<Insights::ITimingViewSession*, FPerSessionData> PerSessionDataMap;
+		TMap<UE::Insights::Timing::ITimingViewSession*, FPerSessionData> PerSessionDataMap;
 	};
 }
 
diff --git a/Engine/Plugins/Developer/Concert/ConcertInsights/ConcertInsightsVisualizer/Source/ConcertInsightsVisualizer/Private/Extension/ConcertTimingViewSession.cpp b/Engine/Plugins/Developer/Concert/ConcertInsights/ConcertInsightsVisualizer/Source/ConcertInsightsVisualizer/Private/Extension/ConcertTimingViewSession.cpp
--- a/Engine/Plugins/Developer/Concert/ConcertInsights/ConcertInsightsVisualizer/Source/ConcertInsightsVisualizer/Private/Extension/ConcertTimingViewSession.cpp
+++ b/Engine/Plugins/Developer/Concert/ConcertInsights/ConcertInsightsVisualizer/Source/ConcertInsightsVisualizer/Private/Extension/ConcertTimingViewSession.cpp
@@ -15,13 +15,13 @@
 
 namespace UE::ConcertInsightsVisualizer
 {
-	void FConcertTimingViewSession::OnBeginSession(Insights::ITimingViewSession& InSession)
+	void FConcertTimingViewSession::OnBeginSession(UE::Insights::Timing::ITimingViewSession& InSession)
 	{}
 
-	void FConcertTimingViewSession::OnEndSession(Insights::ITimingViewSession& InSession)
+	void FConcertTimingViewSession::OnEndSession(UE::Insights::Timing::ITimingViewSession& InSession)
 	{}
 	
-	void FConcertTimingViewSession::Tick(Insights::ITimingViewSession& InTimingSession, const TraceServices::IAnalysisSession& InAnalysisSession)
+	void FConcertTimingViewSession::Tick(UE::Insights::Timing::ITimingViewSession& InTimingSession, const TraceServices::IAnalysisSession& InAnalysisSession)
 	{
 		const FProtocolMultiEndpointProvider* ProtocolProvider = InAnalysisSession.ReadProvider<FProtocolMultiEndpointProvider>(FProtocolMultiEndpointProvider::ProviderName);
 		if (!ensure(ProtocolProvider))
@@ -39,7 +39,7 @@ namespace UE::ConcertInsightsVisualizer
 		}
 	}
 
-	void FConcertTimingViewSession::ExtendFilterMenu(Insights::ITimingViewSession& InSession, FMenuBuilder& InMenuBuilder)
+	void FConcertTimingViewSession::ExtendFilterMenu(UE::Insights::Timing::ITimingViewSession& InSession, FMenuBuilder& InMenuBuilder)
 	{
 		if (!ObjectTrack)
 		{
diff --git a/Engine/Plugins/Developer/Concert/ConcertInsights/ConcertInsightsVisualizer/Source/ConcertInsightsVisualizer/Private/Extension/ConcertTimingViewSession.h b/Engine/Plugins/Developer/Concert/ConcertInsights/ConcertInsightsVisualizer/Source/ConcertInsightsVisualizer/Private/Extension/ConcertTimingViewSession.h
--- a/Engine/Plugins/Developer/Concert/ConcertInsights/ConcertInsightsVisualizer/Source/ConcertInsightsVisualizer/Private/Extension/ConcertTimingViewSession.h
+++ b/Engine/Plugins/Developer/Concert/ConcertInsights/ConcertInsightsVisualizer/Source/ConcertInsightsVisualizer/Private/Extension/ConcertTimingViewSession.h
@@ -7,7 +7,7 @@
 
 class FMenuBuilder;
 
-namespace Insights { class ITimingViewSession; }
+namespace UE::Insights::Timing { class ITimingViewSession; }
 namespace TraceServices { class IAnalysisSession; }
 
 namespace UE::ConcertInsightsVisualizer
@@ -21,10 +21,10 @@ namespace UE::ConcertInsightsVisualizer
 	public:
 
 		// Mirrors the ITimingViewExtender interface - called by FConcertTimingViewExtender
-		void OnBeginSession(Insights::ITimingViewSession& InSession);
-		void OnEndSession(Insights::ITimingViewSession& InSession);
-		void Tick(Insights::ITimingViewSession& InTimingSession, const TraceServices::IAnalysisSession& InAnalysisSession);
-		void ExtendFilterMenu(Insights::ITimingViewSession& InSession, FMenuBuilder& InMenuBuilder);
+		void OnBeginSession(UE::Insights::Timing::ITimingViewSession& InSession);
+		void OnEndSession(UE::Insights::Timing::ITimingViewSession& InSession);
+		void Tick(UE::Insights::Timing::ITimingViewSession& InTimingSession, const TraceServices::IAnalysisSession& InAnalysisSession);
+		void ExtendFilterMenu(UE::Insights::Timing::ITimingViewSession& InSession, FMenuBuilder& InMenuBuilder);
 
 	private:
 
diff --git a/Engine/Plugins/Developer/FieldNotificationTrace/Source/FieldNotificationTraceEditor/Private/FieldNotificationTraceServices.h b/Engine/Plugins/Developer/FieldNotificationTrace/Source/FieldNotificationTraceEditor/Private/FieldNotificationTraceServices.h
--- a/Engine/Plugins/Developer/FieldNotificationTrace/Source/FieldNotificationTraceEditor/Private/FieldNotificationTraceServices.h
+++ b/Engine/Plugins/Developer/FieldNotificationTrace/Source/FieldNotificationTraceEditor/Private/FieldNotificationTraceServices.h
@@ -6,7 +6,6 @@
 #include "TraceServices/ModuleService.h"
 
 class IAnimationProvider;
-namespace Insights { class ITimingViewSession; }
 
 namespace UE::FieldNotification
 {
diff --git a/Engine/Plugins/Experimental/Chooser/Source/ChooserEditor/Private/ChooserTraceModule.h b/Engine/Plugins/Experimental/Chooser/Source/ChooserEditor/Private/ChooserTraceModule.h
--- a/Engine/Plugins/Experimental/Chooser/Source/ChooserEditor/Private/ChooserTraceModule.h
+++ b/Engine/Plugins/Experimental/Chooser/Source/ChooserEditor/Private/ChooserTraceModule.h
@@ -6,7 +6,6 @@
 #include "TraceServices/ModuleService.h"
 
 class IAnimationProvider;
-namespace Insights { class ITimingViewSession; }
 
 class FChooserTraceModule : public TraceServices::IModule
 {
diff --git a/Engine/Plugins/RenderGraphInsights/Source/RenderGraphInsights/Private/RenderGraphInsightsModule.cpp b/Engine/Plugins/RenderGraphInsights/Source/RenderGraphInsights/Private/RenderGraphInsightsModule.cpp
--- a/Engine/Plugins/RenderGraphInsights/Source/RenderGraphInsights/Private/RenderGraphInsightsModule.cpp
+++ b/Engine/Plugins/RenderGraphInsights/Source/RenderGraphInsights/Private/RenderGraphInsightsModule.cpp
@@ -18,13 +18,13 @@ FRenderGraphInsightsModule& FRenderGraphInsightsModule::Get()
 void FRenderGraphInsightsModule::StartupModule()
 {
 	IModularFeatures::Get().RegisterModularFeature(TraceServices::ModuleFeatureName, &TraceModule);
-	IModularFeatures::Get().RegisterModularFeature(Insights::TimingViewExtenderFeatureName, &TimingViewExtender);
+	IModularFeatures::Get().RegisterModularFeature(UE::Insights::Timing::TimingViewExtenderFeatureName, &TimingViewExtender);
 }
 
 void FRenderGraphInsightsModule::ShutdownModule()
 {
 	IModularFeatures::Get().UnregisterModularFeature(TraceServices::ModuleFeatureName, &TraceModule);
-	IModularFeatures::Get().UnregisterModularFeature(Insights::TimingViewExtenderFeatureName, &TimingViewExtender);
+	IModularFeatures::Get().UnregisterModularFeature(UE::Insights::Timing::TimingViewExtenderFeatureName, &TimingViewExtender);
 }
 
 } //namespace SlateInsights
diff --git a/Engine/Plugins/RenderGraphInsights/Source/RenderGraphInsights/Private/RenderGraphProvider.h b/Engine/Plugins/RenderGraphInsights/Source/RenderGraphInsights/Private/RenderGraphProvider.h
--- a/Engine/Plugins/RenderGraphInsights/Source/RenderGraphInsights/Private/RenderGraphProvider.h
+++ b/Engine/Plugins/RenderGraphInsights/Source/RenderGraphInsights/Private/RenderGraphProvider.h
@@ -2,12 +2,12 @@
 
 #pragma once
 
-#include "Insights/Common/SimpleRtti.h"
-#include "TraceServices/Model/AnalysisSession.h"
+#include "InsightsCore/Common/SimpleRtti.h"
 #include "Model/IntervalTimeline.h"
+#include "RenderGraphResources.h"
 #include "RHITransientResourceAllocator.h"
 #include "Trace/Analyzer.h"
-#include "RenderGraphResources.h"
+#include "TraceServices/Model/AnalysisSession.h"
 
 enum class ERHIPipeline : uint8;
 
diff --git a/Engine/Plugins/RenderGraphInsights/Source/RenderGraphInsights/Private/RenderGraphTimingViewExtender.cpp b/Engine/Plugins/RenderGraphInsights/Source/RenderGraphInsights/Private/RenderGraphTimingViewExtender.cpp
--- a/Engine/Plugins/RenderGraphInsights/Source/RenderGraphInsights/Private/RenderGraphTimingViewExtender.cpp
+++ b/Engine/Plugins/RenderGraphInsights/Source/RenderGraphInsights/Private/RenderGraphTimingViewExtender.cpp
@@ -11,7 +11,7 @@ namespace UE
 namespace RenderGraphInsights
 {
 
-void FRenderGraphTimingViewExtender::OnBeginSession(Insights::ITimingViewSession& InSession)
+void FRenderGraphTimingViewExtender::OnBeginSession(UE::Insights::Timing::ITimingViewSession& InSession)
 {
 	if (InSession.GetName() != FInsightsManagerTabs::TimingProfilerTabId)
 	{
@@ -32,7 +32,7 @@ void FRenderGraphTimingViewExtender::OnBeginSession(Insights::ITimingViewSession
 	}
 }
 
-void FRenderGraphTimingViewExtender::OnEndSession(Insights::ITimingViewSession& InSession)
+void FRenderGraphTimingViewExtender::OnEndSession(UE::Insights::Timing::ITimingViewSession& InSession)
 {
 	if (InSession.GetName() != FInsightsManagerTabs::TimingProfilerTabId)
 	{
@@ -48,7 +48,7 @@ void FRenderGraphTimingViewExtender::OnEndSession(Insights::ITimingViewSession&
 	PerSessionDataMap.Remove(&InSession);
 }
 
-void FRenderGraphTimingViewExtender::Tick(Insights::ITimingViewSession& InSession, const TraceServices::IAnalysisSession& InAnalysisSession)
+void FRenderGraphTimingViewExtender::Tick(UE::Insights::Timing::ITimingViewSession& InSession, const TraceServices::IAnalysisSession& InAnalysisSession)
 {
 	FPerSessionData* PerSessionData = PerSessionDataMap.Find(&InSession);
 	if(PerSessionData != nullptr)
@@ -57,7 +57,7 @@ void FRenderGraphTimingViewExtender::Tick(Insights::ITimingViewSession& InSessio
 	}
 }
 
-void FRenderGraphTimingViewExtender::ExtendFilterMenu(Insights::ITimingViewSession& InSession, FMenuBuilder& InMenuBuilder)
+void FRenderGraphTimingViewExtender::ExtendFilterMenu(UE::Insights::Timing::ITimingViewSession& InSession, FMenuBuilder& InMenuBuilder)
 {
 	FPerSessionData* PerSessionData = PerSessionDataMap.Find(&InSession);
 	if (PerSessionData != nullptr)
diff --git a/Engine/Plugins/RenderGraphInsights/Source/RenderGraphInsights/Private/RenderGraphTimingViewExtender.h b/Engine/Plugins/RenderGraphInsights/Source/RenderGraphInsights/Private/RenderGraphTimingViewExtender.h
--- a/Engine/Plugins/RenderGraphInsights/Source/RenderGraphInsights/Private/RenderGraphTimingViewExtender.h
+++ b/Engine/Plugins/RenderGraphInsights/Source/RenderGraphInsights/Private/RenderGraphTimingViewExtender.h
@@ -5,7 +5,7 @@
 #include "Insights/ITimingViewExtender.h"
 #include "RenderGraphTimingViewSession.h" // IWYU pragma: keep
 
-namespace Insights { class ITimingViewSession; }
+namespace UE::Insights::Timing { class ITimingViewSession; }
 namespace TraceServices { class IAnalysisSession; }
 
 class FMenuBuilder;
@@ -15,15 +15,15 @@ namespace UE
 namespace RenderGraphInsights
 {
 
-class FRenderGraphTimingViewExtender : public Insights::ITimingViewExtender
+class FRenderGraphTimingViewExtender : public UE::Insights::Timing::ITimingViewExtender
 {
 public:
-	//~ Begin Insights::ITimingViewExtender interface
-	virtual void OnBeginSession(Insights::ITimingViewSession& InSession) override;
-	virtual void OnEndSession(Insights::ITimingViewSession& InSession) override;
-	virtual void Tick(Insights::ITimingViewSession& InSession, const TraceServices::IAnalysisSession& InAnalysisSession) override;
-	virtual void ExtendFilterMenu(Insights::ITimingViewSession& InSession, FMenuBuilder& InMenuBuilder) override;
-	//~ End Insights::ITimingViewExtender interface
+	//~ Begin UE::Insights::Timing::ITimingViewExtender interface
+	virtual void OnBeginSession(UE::Insights::Timing::ITimingViewSession& InSession) override;
+	virtual void OnEndSession(UE::Insights::Timing::ITimingViewSession& InSession) override;
+	virtual void Tick(UE::Insights::Timing::ITimingViewSession& InSession, const TraceServices::IAnalysisSession& InAnalysisSession) override;
+	virtual void ExtendFilterMenu(UE::Insights::Timing::ITimingViewSession& InSession, FMenuBuilder& InMenuBuilder) override;
+	//~ End UE::Insights::Timing::ITimingViewExtender interface
 
 private:
 	struct FPerSessionData
@@ -32,7 +32,7 @@ private:
 	};
 
 	// The data we host per-session
-	TMap<Insights::ITimingViewSession*, FPerSessionData> PerSessionDataMap;
+	TMap<UE::Insights::Timing::ITimingViewSession*, FPerSessionData> PerSessionDataMap;
 };
 
 } //namespace SlateInsights
diff --git a/Engine/Plugins/RenderGraphInsights/Source/RenderGraphInsights/Private/RenderGraphTimingViewSession.cpp b/Engine/Plugins/RenderGraphInsights/Source/RenderGraphInsights/Private/RenderGraphTimingViewSession.cpp
--- a/Engine/Plugins/RenderGraphInsights/Source/RenderGraphInsights/Private/RenderGraphTimingViewSession.cpp
+++ b/Engine/Plugins/RenderGraphInsights/Source/RenderGraphInsights/Private/RenderGraphTimingViewSession.cpp
@@ -15,20 +15,20 @@ namespace UE
 namespace RenderGraphInsights
 {
 
-void FRenderGraphTimingViewSession::OnBeginSession(Insights::ITimingViewSession& InTimingViewSession)
+void FRenderGraphTimingViewSession::OnBeginSession(UE::Insights::Timing::ITimingViewSession& InTimingViewSession)
 {
 	TimingViewSession = &InTimingViewSession;
 
 	Track.Reset();
 }
 
-void FRenderGraphTimingViewSession::OnEndSession(Insights::ITimingViewSession& InTimingViewSession)
+void FRenderGraphTimingViewSession::OnEndSession(UE::Insights::Timing::ITimingViewSession& InTimingViewSession)
 {
 	Track.Reset();
 	TimingViewSession = nullptr;
 }
 
-void FRenderGraphTimingViewSession::Tick(Insights::ITimingViewSession& InTimingViewSession, const TraceServices::IAnalysisSession& InAnalysisSession)
+void FRenderGraphTimingViewSession::Tick(UE::Insights::Timing::ITimingViewSession& InTimingViewSession, const TraceServices::IAnalysisSession& InAnalysisSession)
 {
 	AnalysisSession = &InAnalysisSession;
 
diff --git a/Engine/Plugins/RenderGraphInsights/Source/RenderGraphInsights/Private/RenderGraphTimingViewSession.h b/Engine/Plugins/RenderGraphInsights/Source/RenderGraphInsights/Private/RenderGraphTimingViewSession.h
--- a/Engine/Plugins/RenderGraphInsights/Source/RenderGraphInsights/Private/RenderGraphTimingViewSession.h
+++ b/Engine/Plugins/RenderGraphInsights/Source/RenderGraphInsights/Private/RenderGraphTimingViewSession.h
@@ -5,8 +5,8 @@
 #include "Templates/SharedPointer.h"
 
 namespace TraceServices { class IAnalysisSession; }
-namespace Insights { class ITimingViewSession; }
-namespace Insights { enum class ETimeChangedFlags : int32; }
+namespace UE::Insights::Timing { class ITimingViewSession; }
+namespace UE::Insights::Timing { enum class ETimeChangedFlags : int32; }
 class FMenuBuilder;
 class SDockTab;
 
@@ -21,9 +21,9 @@ class FRenderGraphTimingViewSession
 public:
 	FRenderGraphTimingViewSession() = default;
 
-	void OnBeginSession(Insights::ITimingViewSession& InTimingViewSession);
-	void OnEndSession(Insights::ITimingViewSession& InTimingViewSession);
-	void Tick(Insights::ITimingViewSession& InTimingViewSession, const TraceServices::IAnalysisSession& InAnalysisSession);
+	void OnBeginSession(UE::Insights::Timing::ITimingViewSession& InTimingViewSession);
+	void OnEndSession(UE::Insights::Timing::ITimingViewSession& InTimingViewSession);
+	void Tick(UE::Insights::Timing::ITimingViewSession& InTimingViewSession, const TraceServices::IAnalysisSession& InAnalysisSession);
 	void ExtendFilterMenu(FMenuBuilder& InMenuBuilder);
 
 	/** Get the last cached analysis session */
@@ -35,11 +35,11 @@ public:
 	/** Show/Hide the RenderGraph track */
 	void ToggleRenderGraphTrack();
 
-	Insights::ITimingViewSession* GetTimingViewSession() const { return TimingViewSession; }
+	UE::Insights::Timing::ITimingViewSession* GetTimingViewSession() const { return TimingViewSession; }
 
 private:
 	const TraceServices::IAnalysisSession* AnalysisSession = nullptr;
-	Insights::ITimingViewSession* TimingViewSession;
+	UE::Insights::Timing::ITimingViewSession* TimingViewSession;
 	TSharedPtr<FRenderGraphTrack> Track;
 	bool bTrackVisible = true;
 };
diff --git a/Engine/Plugins/RenderGraphInsights/Source/RenderGraphInsights/Private/RenderGraphTrack.cpp b/Engine/Plugins/RenderGraphInsights/Source/RenderGraphInsights/Private/RenderGraphTrack.cpp
--- a/Engine/Plugins/RenderGraphInsights/Source/RenderGraphInsights/Private/RenderGraphTrack.cpp
+++ b/Engine/Plugins/RenderGraphInsights/Source/RenderGraphInsights/Private/RenderGraphTrack.cpp
@@ -2,19 +2,18 @@
 
 #include "RenderGraphTrack.h"
 
-#include "RenderGraphTimingViewSession.h"
-
 #include "Framework/MultiBox/MultiBoxBuilder.h"
-
-#include "Insights/Common/PaintUtils.h"
-#include "Insights/ViewModels/TimingEvent.h"
-#include "Insights/ITimingViewSession.h"
-
+#include "RenderGraphTimingViewSession.h"
 #include "RenderGraphTrackDrawHelper.h"
 #include "Widgets/Input/SSearchBox.h"
 #include "Widgets/Input/SSpinBox.h"
 
+// TraceInsightsCore
+#include "InsightsCore/Common/PaintUtils.h"
 
+// TraceInsights
+#include "Insights/ITimingViewSession.h"
+#include "Insights/ViewModels/TimingEvent.h"
 
 #define LOCTEXT_NAMESPACE "RenderGraphTrack"
 
@@ -1668,7 +1667,7 @@ void FRenderGraphTrack::BuildContextMenu(FMenuBuilder& MenuBuilder)
 {
 	Super::BuildContextMenu(MenuBuilder);
 
-	Insights::ITimingViewSession* TimingViewSession = SharedData.GetTimingViewSession();
+	UE::Insights::Timing::ITimingViewSession* TimingViewSession = SharedData.GetTimingViewSession();
 
 	MenuBuilder.BeginSection("Visualizer", LOCTEXT("Visualizer", "Visualizer Mode"));
 	{
diff --git a/Engine/Plugins/RenderGraphInsights/Source/RenderGraphInsights/Private/RenderGraphTrackDrawHelper.cpp b/Engine/Plugins/RenderGraphInsights/Source/RenderGraphInsights/Private/RenderGraphTrackDrawHelper.cpp
--- a/Engine/Plugins/RenderGraphInsights/Source/RenderGraphInsights/Private/RenderGraphTrackDrawHelper.cpp
+++ b/Engine/Plugins/RenderGraphInsights/Source/RenderGraphInsights/Private/RenderGraphTrackDrawHelper.cpp
@@ -6,11 +6,14 @@
 #include "Insights/ViewModels/ITimingViewDrawHelper.h"
 #include "Styling/CoreStyle.h"
 
-#include "Insights/Common/PaintUtils.h"
-#include "Insights/Common/TimeUtils.h"
+// TraceInsightsCore
+#include "InsightsCore/Common/PaintUtils.h"
+#include "InsightsCore/Common/TimeUtils.h"
+
+// TraceInsights
+#include "Insights/ViewModels/BaseTimingTrack.h"
 #include "Insights/ViewModels/TimingEvent.h"
 #include "Insights/ViewModels/TimingTrackViewport.h"
-#include "Insights/ViewModels/BaseTimingTrack.h"
 
 namespace UE { namespace RenderGraphInsights {
 
@@ -37,7 +40,7 @@ FRenderGraphTrackDrawStateBuilder::FRenderGraphTrackDrawStateBuilder(FRenderGrap
 void FRenderGraphTrackDrawStateBuilder::AppendDurationToEventName(FString& InOutEventName, const double InDuration)
 {
 	InOutEventName += TEXT(" (");
-	InOutEventName += TimeUtils::FormatTimeAuto(InDuration);
+	InOutEventName += UE::Insights::FormatTimeAuto(InDuration);
 	InOutEventName += TEXT(")");
 }
 
diff --git a/Engine/Plugins/Runtime/NetworkPredictionInsights/Source/NetworkPredictionInsights/Private/NetworkPredictionTraceModule.h b/Engine/Plugins/Runtime/NetworkPredictionInsights/Source/NetworkPredictionInsights/Private/NetworkPredictionTraceModule.h
--- a/Engine/Plugins/Runtime/NetworkPredictionInsights/Source/NetworkPredictionInsights/Private/NetworkPredictionTraceModule.h
+++ b/Engine/Plugins/Runtime/NetworkPredictionInsights/Source/NetworkPredictionInsights/Private/NetworkPredictionTraceModule.h
@@ -5,7 +5,6 @@
 #include "TraceServices/ModuleService.h"
 
 class IAnimationProvider;
-namespace Insights { class ITimingViewSession; }
 
 class FNetworkPredictionTraceModule : public TraceServices::IModule
 {
diff --git a/Engine/Plugins/Runtime/NetworkPredictionInsights/Source/NetworkPredictionInsights/Private/UI/NetworkPredictionDrawHelpers.cpp b/Engine/Plugins/Runtime/NetworkPredictionInsights/Source/NetworkPredictionInsights/Private/UI/NetworkPredictionDrawHelpers.cpp
--- a/Engine/Plugins/Runtime/NetworkPredictionInsights/Source/NetworkPredictionInsights/Private/UI/NetworkPredictionDrawHelpers.cpp
+++ b/Engine/Plugins/Runtime/NetworkPredictionInsights/Source/NetworkPredictionInsights/Private/UI/NetworkPredictionDrawHelpers.cpp
@@ -3,7 +3,7 @@
 #include "NetworkPredictionDrawHelpers.h"
 
 
-#include "Insights/Common/PaintUtils.h"
+#include "InsightsCore/Common/PaintUtils.h"
 
 
 void FNetworkPredictionDrawHelpers::DrawBackground(
diff --git a/Engine/Plugins/Runtime/NetworkPredictionInsights/Source/NetworkPredictionInsights/Private/UI/SNPSimFrameView.cpp b/Engine/Plugins/Runtime/NetworkPredictionInsights/Source/NetworkPredictionInsights/Private/UI/SNPSimFrameView.cpp
--- a/Engine/Plugins/Runtime/NetworkPredictionInsights/Source/NetworkPredictionInsights/Private/UI/SNPSimFrameView.cpp
+++ b/Engine/Plugins/Runtime/NetworkPredictionInsights/Source/NetworkPredictionInsights/Private/UI/SNPSimFrameView.cpp
@@ -1,17 +1,22 @@
 // Copyright Epic Games, Inc. All Rights Reserved.
 
 #include "SNPSimFrameView.h"
-#include "Insights/Common/PaintUtils.h"
-#include "NetworkPredictionDrawHelpers.h"
+
+#include "Framework/MultiBox/MultiBoxBuilder.h"
 #include "Layout/WidgetPath.h"
-#include "SNPWindow.h"
-#include "Insights/IUnrealInsightsModule.h"
 #include "Modules/ModuleManager.h"
-#include "Framework/MultiBox/MultiBoxBuilder.h"
+#include "NetworkPredictionDrawHelpers.h"
+#include "SNPWindow.h"
 #include "UI/NPAxisViewportDouble.h"
 #include "UI/NPAxisViewportInt32.h"
 #include "Widgets/Layout/SScrollBar.h"
 
+// TraceInsightsCore
+#include "InsightsCore/Common/PaintUtils.h"
+
+// TraceInsights
+#include "Insights/IUnrealInsightsModule.h"
+
 #define LOCTEXT_NAMESPACE "NetworkPredictionInsights"
 
 FSimFrameViewDrawHelper::FSimFrameViewDrawHelper(const FDrawContext& InDrawContext, const FSimFrameViewport& InViewport)
diff --git a/Engine/Plugins/Slate/SlateInsights/Source/SlateInsights/Private/SSlateFrameSchematicView.cpp b/Engine/Plugins/Slate/SlateInsights/Source/SlateInsights/Private/SSlateFrameSchematicView.cpp
--- a/Engine/Plugins/Slate/SlateInsights/Source/SlateInsights/Private/SSlateFrameSchematicView.cpp
+++ b/Engine/Plugins/Slate/SlateInsights/Source/SlateInsights/Private/SSlateFrameSchematicView.cpp
@@ -1,33 +1,35 @@
 // Copyright Epic Games, Inc. All Rights Reserved.
 
 #include "SSlateFrameSchematicView.h"
+
+#include "Framework/Application/SlateApplication.h"
+#include "Framework/MultiBox/MultiBoxBuilder.h"
 #include "Framework/Views/TableViewMetadata.h"
+#include "HAL/FileManager.h"
+#include "Internationalization/Regex.h"
+#include "ISourceCodeAccessModule.h"
+#include "ISourceCodeAccessor.h"
+#include "Modules/ModuleManager.h"
+#include "SlateInsightsStyle.h"
 #include "SlateProvider.h"
 #include "SSlateTraceFlags.h"
-
-#include "Insights/Common/TimeUtils.h"
-#include "Insights/ITimingViewSession.h"
-#include "Insights/ViewModels/ITimingEvent.h"
 #include "Widgets/Input/NumericTypeInterface.h"
 #include "Widgets/Input/SMultiLineEditableTextBox.h"
+#include "Widgets/Input/SMultiLineEditableTextBox.h"
 #include "Widgets/Input/SSearchBox.h"
 #include "Widgets/Layout/SExpandableArea.h"
 #include "Widgets/Layout/SGridPanel.h"
 #include "Widgets/Layout/SHeader.h"
-#include "Widgets/Input/SMultiLineEditableTextBox.h"
-
-#include "Framework/MultiBox/MultiBoxBuilder.h"
-#include "Framework/Application/SlateApplication.h"
-
-#include "HAL/FileManager.h"
-#include "Modules/ModuleManager.h"
-#include "Internationalization/Regex.h"
-#include "SlateInsightsStyle.h"
-#include "ISourceCodeAccessModule.h"
-#include "ISourceCodeAccessor.h"
 #include "Widgets/SWindow.h"
 #include "Widgets/Views/STreeView.h"
 
+// TraceInsightsCore
+#include "InsightsCore/Common/TimeUtils.h"
+
+// TraceInsights
+#include "Insights/ITimingViewSession.h"
+#include "Insights/ViewModels/ITimingEvent.h"
+
 #define LOCTEXT_NAMESPACE "SSlateFrameSchematicView"
 
 namespace UE
@@ -177,7 +179,7 @@ namespace Private
 			else if (Column == ColumnDuration)
 			{
 				return SNew(STextBlock)
-					.Text(FText::FromString(TimeUtils::FormatTimeAuto(Info->Duration)));
+					.Text(FText::FromString(UE::Insights::FormatTimeAuto(Info->Duration)));
 			}
 			else if (Column == ColumnFlag)
 			{
@@ -333,9 +335,9 @@ namespace Private
 						DebugInfoWidget->SetText(FText::FromString(WidgetInfo->DebugInfo));
 
 						const double StartTime = SlateProvider->GetWidgetTimeline().GetEventStartTime(WidgetInfo->EventIndex);
-						CreatedTimeWidget->SetText(FText::FromString(TimeUtils::FormatTime(StartTime, TimeUtils::Milisecond)));
+						CreatedTimeWidget->SetText(FText::FromString(UE::Insights::FormatTime(StartTime, UE::Insights::FTimeValue::Millisecond)));
 						const double EndTime = SlateProvider->GetWidgetTimeline().GetEventEndTime(WidgetInfo->EventIndex);
-						DestroyedTimeWidget->SetText(FText::FromString(TimeUtils::FormatTime(EndTime, TimeUtils::Milisecond)));
+						DestroyedTimeWidget->SetText(FText::FromString(UE::Insights::FormatTime(EndTime, UE::Insights::FTimeValue::Millisecond)));
 					}
 				}
 			}
@@ -538,7 +540,7 @@ SSlateFrameSchematicView::~SSlateFrameSchematicView()
 	}
 }
 
-void SSlateFrameSchematicView::SetSession(Insights::ITimingViewSession* InTimingViewSession, const TraceServices::IAnalysisSession* InAnalysisSession)
+void SSlateFrameSchematicView::SetSession(UE::Insights::Timing::ITimingViewSession* InTimingViewSession, const TraceServices::IAnalysisSession* InAnalysisSession)
 {
 	if (TimingViewSession)
 	{
@@ -654,7 +656,7 @@ EColumnSortMode::Type SSlateFrameSchematicView::HandleWidgetUpdateGetSortMode(FN
 		: EColumnSortMode::None;
 }
 
-void SSlateFrameSchematicView::HandleTimeMarkerChanged(Insights::ETimeChangedFlags InFlags, double InTimeMarker)
+void SSlateFrameSchematicView::HandleTimeMarkerChanged(UE::Insights::Timing::ETimeChangedFlags InFlags, double InTimeMarker)
 {
 	if (!FMath::IsNearlyEqual(StartTime, InTimeMarker) && !FMath::IsNearlyEqual(EndTime, InTimeMarker))
 	{
@@ -665,7 +667,7 @@ void SSlateFrameSchematicView::HandleTimeMarkerChanged(Insights::ETimeChangedFla
 	}
 }
 
-void SSlateFrameSchematicView::HandleSelectionChanged(Insights::ETimeChangedFlags InFlags, double InStartTime, double InEndTime)
+void SSlateFrameSchematicView::HandleSelectionChanged(UE::Insights::Timing::ETimeChangedFlags InFlags, double InStartTime, double InEndTime)
 {
 	if (!FMath::IsNearlyEqual(StartTime, InStartTime) && !FMath::IsNearlyEqual(EndTime, InEndTime))
 	{
diff --git a/Engine/Plugins/Slate/SlateInsights/Source/SlateInsights/Private/SSlateFrameSchematicView.h b/Engine/Plugins/Slate/SlateInsights/Source/SlateInsights/Private/SSlateFrameSchematicView.h
--- a/Engine/Plugins/Slate/SlateInsights/Source/SlateInsights/Private/SSlateFrameSchematicView.h
+++ b/Engine/Plugins/Slate/SlateInsights/Source/SlateInsights/Private/SSlateFrameSchematicView.h
@@ -2,15 +2,14 @@
 
 #pragma once
 
-
 #include "Widgets/Views/SHeaderRow.h"
 
 template <typename ItemType> class SListView;
 template <typename ItemType> class STreeView;
 
-namespace Insights { class ITimingViewSession; }
-namespace Insights { enum class ETimeChangedFlags; }
 namespace TraceServices { class IAnalysisSession; }
+namespace UE::Insights::Timing { class ITimingViewSession; }
+namespace UE::Insights::Timing { enum class ETimeChangedFlags; }
 class FAnimationSharedData;
 class IInsightsManager;
 class ITableRow;
@@ -39,7 +38,7 @@ public:
 
 	void Construct(const FArguments& InArgs);
 
-	void SetSession(Insights::ITimingViewSession* InTimingViewSession, const TraceServices::IAnalysisSession* InAnalysisSession);
+	void SetSession(UE::Insights::Timing::ITimingViewSession* InTimingViewSession, const TraceServices::IAnalysisSession* InAnalysisSession);
 
 	bool IsSessionSet() const;
 
@@ -47,8 +46,8 @@ private:
 	TSharedRef<ITableRow> HandleUniqueInvalidatedMakeTreeRowWidget(TSharedPtr<Private::FWidgetUniqueInvalidatedInfo> InInfo, const TSharedRef<STableViewBase>& OwnerTable);
 	void HandleUniqueInvalidatedChildrenForInfo(TSharedPtr<Private::FWidgetUniqueInvalidatedInfo> InInfo, TArray<TSharedPtr<Private::FWidgetUniqueInvalidatedInfo>>& OutChildren);
 	TSharedRef<ITableRow> HandleWidgetUpdateInfoGenerateWidget(TSharedPtr<Private::FWidgetUpdateInfo> Item, const TSharedRef<STableViewBase>& OwnerTable);
-	void HandleTimeMarkerChanged(Insights::ETimeChangedFlags InFlags, double InTimeMarker);
-	void HandleSelectionChanged(Insights::ETimeChangedFlags InFlags, double StartTime, double EndTime);
+	void HandleTimeMarkerChanged(UE::Insights::Timing::ETimeChangedFlags InFlags, double InTimeMarker);
+	void HandleSelectionChanged(UE::Insights::Timing::ETimeChangedFlags InFlags, double StartTime, double EndTime);
 	void HandleSelectionEventChanged(const TSharedPtr<const ITimingEvent> InEvent);
 
 	TSharedPtr<SWidget> HandleWidgetInvalidateListContextMenu();
@@ -75,7 +74,7 @@ private:
 
 private:
 	const TraceServices::IAnalysisSession* AnalysisSession;
-	Insights::ITimingViewSession* TimingViewSession;
+	UE::Insights::Timing::ITimingViewSession* TimingViewSession;
 
 	TSharedPtr<STreeView<TSharedPtr<Private::FWidgetUniqueInvalidatedInfo>>> WidgetInvalidateInfoListView;
 	TArray<TSharedPtr<Private::FWidgetUniqueInvalidatedInfo>> WidgetInvalidationInfos;
diff --git a/Engine/Plugins/Slate/SlateInsights/Source/SlateInsights/Private/SlateFrameGraphTrack.cpp b/Engine/Plugins/Slate/SlateInsights/Source/SlateInsights/Private/SlateFrameGraphTrack.cpp
--- a/Engine/Plugins/Slate/SlateInsights/Source/SlateInsights/Private/SlateFrameGraphTrack.cpp
+++ b/Engine/Plugins/Slate/SlateInsights/Source/SlateInsights/Private/SlateFrameGraphTrack.cpp
@@ -8,7 +8,10 @@
 
 #include "Framework/MultiBox/MultiBoxBuilder.h"
 
-#include "Insights/Common/PaintUtils.h"
+// TraceInsightsCore
+#include "InsightsCore/Common/PaintUtils.h"
+
+// TraceInsights
 #include "Insights/ViewModels/GraphTrackBuilder.h"
 #include "Insights/ViewModels/ITimingViewDrawHelper.h"
 #include "Insights/ViewModels/TimingTrackViewport.h"
diff --git a/Engine/Plugins/Slate/SlateInsights/Source/SlateInsights/Private/SlateInsightsModule.cpp b/Engine/Plugins/Slate/SlateInsights/Source/SlateInsights/Private/SlateInsightsModule.cpp
--- a/Engine/Plugins/Slate/SlateInsights/Source/SlateInsights/Private/SlateInsightsModule.cpp
+++ b/Engine/Plugins/Slate/SlateInsights/Source/SlateInsights/Private/SlateInsightsModule.cpp
@@ -30,7 +30,7 @@ FSlateInsightsModule& FSlateInsightsModule::Get()
 void FSlateInsightsModule::StartupModule()
 {
 	IModularFeatures::Get().RegisterModularFeature(TraceServices::ModuleFeatureName, &TraceModule);
-	IModularFeatures::Get().RegisterModularFeature(Insights::TimingViewExtenderFeatureName, &TimingViewExtender);
+	IModularFeatures::Get().RegisterModularFeature(UE::Insights::Timing::TimingViewExtenderFeatureName, &TimingViewExtender);
 
 	IUnrealInsightsModule& UnrealInsightsModule = FModuleManager::LoadModuleChecked<IUnrealInsightsModule>("TraceInsights");
 	FOnRegisterMajorTabExtensions& TimingProfilerLayoutExtension = UnrealInsightsModule.OnRegisterMajorTabExtension(FInsightsManagerTabs::TimingProfilerTabId);
@@ -46,7 +46,7 @@ void FSlateInsightsModule::ShutdownModule()
 	}
 
 	IModularFeatures::Get().UnregisterModularFeature(TraceServices::ModuleFeatureName, &TraceModule);
-	IModularFeatures::Get().UnregisterModularFeature(Insights::TimingViewExtenderFeatureName, &TimingViewExtender);
+	IModularFeatures::Get().UnregisterModularFeature(UE::Insights::Timing::TimingViewExtenderFeatureName, &TimingViewExtender);
 }
 
 TSharedPtr<SSlateFrameSchematicView> FSlateInsightsModule::GetSlateFrameSchematicViewTab(bool bInvoke)
diff --git a/Engine/Plugins/Slate/SlateInsights/Source/SlateInsights/Private/SlateTimingViewExtender.cpp b/Engine/Plugins/Slate/SlateInsights/Source/SlateInsights/Private/SlateTimingViewExtender.cpp
--- a/Engine/Plugins/Slate/SlateInsights/Source/SlateInsights/Private/SlateTimingViewExtender.cpp
+++ b/Engine/Plugins/Slate/SlateInsights/Source/SlateInsights/Private/SlateTimingViewExtender.cpp
@@ -10,7 +10,7 @@ namespace UE
 namespace SlateInsights
 {
 
-void FSlateTimingViewExtender::OnBeginSession(Insights::ITimingViewSession& InSession)
+void FSlateTimingViewExtender::OnBeginSession(UE::Insights::Timing::ITimingViewSession& InSession)
 {
 	FPerSessionData* PerSessionData = PerSessionDataMap.Find(&InSession);
 	if(PerSessionData == nullptr)
@@ -26,7 +26,7 @@ void FSlateTimingViewExtender::OnBeginSession(Insights::ITimingViewSession& InSe
 	}
 }
 
-void FSlateTimingViewExtender::OnEndSession(Insights::ITimingViewSession& InSession)
+void FSlateTimingViewExtender::OnEndSession(UE::Insights::Timing::ITimingViewSession& InSession)
 {
 	FPerSessionData* PerSessionData = PerSessionDataMap.Find(&InSession);
 	if(PerSessionData != nullptr)
@@ -37,7 +37,7 @@ void FSlateTimingViewExtender::OnEndSession(Insights::ITimingViewSession& InSess
 	PerSessionDataMap.Remove(&InSession);
 }
 
-void FSlateTimingViewExtender::Tick(Insights::ITimingViewSession& InSession, const TraceServices::IAnalysisSession& InAnalysisSession)
+void FSlateTimingViewExtender::Tick(UE::Insights::Timing::ITimingViewSession& InSession, const TraceServices::IAnalysisSession& InAnalysisSession)
 {
 	FPerSessionData* PerSessionData = PerSessionDataMap.Find(&InSession);
 	if(PerSessionData != nullptr)
@@ -46,7 +46,7 @@ void FSlateTimingViewExtender::Tick(Insights::ITimingViewSession& InSession, con
 	}
 }
 
-void FSlateTimingViewExtender::ExtendFilterMenu(Insights::ITimingViewSession& InSession, FMenuBuilder& InMenuBuilder)
+void FSlateTimingViewExtender::ExtendFilterMenu(UE::Insights::Timing::ITimingViewSession& InSession, FMenuBuilder& InMenuBuilder)
 {
 	FPerSessionData* PerSessionData = PerSessionDataMap.Find(&InSession);
 	if (PerSessionData != nullptr)
diff --git a/Engine/Plugins/Slate/SlateInsights/Source/SlateInsights/Private/SlateTimingViewExtender.h b/Engine/Plugins/Slate/SlateInsights/Source/SlateInsights/Private/SlateTimingViewExtender.h
--- a/Engine/Plugins/Slate/SlateInsights/Source/SlateInsights/Private/SlateTimingViewExtender.h
+++ b/Engine/Plugins/Slate/SlateInsights/Source/SlateInsights/Private/SlateTimingViewExtender.h
@@ -5,8 +5,8 @@
 #include "Insights/ITimingViewExtender.h"
 #include "SlateTimingViewSession.h" // IWYU pragma: keep
 
-namespace Insights { class ITimingViewSession; }
 namespace TraceServices { class IAnalysisSession; }
+namespace UE::Insights::Timing { class ITimingViewSession; }
 
 class FMenuBuilder;
 
@@ -15,15 +15,15 @@ namespace UE
 namespace SlateInsights
 {
 
-class FSlateTimingViewExtender : public Insights::ITimingViewExtender
+class FSlateTimingViewExtender : public UE::Insights::Timing::ITimingViewExtender
 {
 public:
-	//~ Begin Insights::ITimingViewExtender interface
-	virtual void OnBeginSession(Insights::ITimingViewSession& InSession) override;
-	virtual void OnEndSession(Insights::ITimingViewSession& InSession) override;
-	virtual void Tick(Insights::ITimingViewSession& InSession, const TraceServices::IAnalysisSession& InAnalysisSession) override;
-	virtual void ExtendFilterMenu(Insights::ITimingViewSession& InSession, FMenuBuilder& InMenuBuilder) override;
-	//~ End Insights::ITimingViewExtender interface
+	//~ Begin UE::Insights::Timing::ITimingViewExtender interface
+	virtual void OnBeginSession(UE::Insights::Timing::ITimingViewSession& InSession) override;
+	virtual void OnEndSession(UE::Insights::Timing::ITimingViewSession& InSession) override;
+	virtual void Tick(UE::Insights::Timing::ITimingViewSession& InSession, const TraceServices::IAnalysisSession& InAnalysisSession) override;
+	virtual void ExtendFilterMenu(UE::Insights::Timing::ITimingViewSession& InSession, FMenuBuilder& InMenuBuilder) override;
+	//~ End UE::Insights::Timing::ITimingViewExtender interface
 
 private:
 	struct FPerSessionData
@@ -32,7 +32,7 @@ private:
 	};
 
 	// The data we host per-session
-	TMap<Insights::ITimingViewSession*, FPerSessionData> PerSessionDataMap;
+	TMap<UE::Insights::Timing::ITimingViewSession*, FPerSessionData> PerSessionDataMap;
 };
 
 } //namespace SlateInsights
diff --git a/Engine/Plugins/Slate/SlateInsights/Source/SlateInsights/Private/SlateTimingViewSession.cpp b/Engine/Plugins/Slate/SlateInsights/Source/SlateInsights/Private/SlateTimingViewSession.cpp
--- a/Engine/Plugins/Slate/SlateInsights/Source/SlateInsights/Private/SlateTimingViewSession.cpp
+++ b/Engine/Plugins/Slate/SlateInsights/Source/SlateInsights/Private/SlateTimingViewSession.cpp
@@ -27,7 +27,7 @@ FSlateTimingViewSession::FSlateTimingViewSession()
 {
 }
 
-void FSlateTimingViewSession::OnBeginSession(Insights::ITimingViewSession& InTimingViewSession)
+void FSlateTimingViewSession::OnBeginSession(UE::Insights::Timing::ITimingViewSession& InTimingViewSession)
 {
 	if (InTimingViewSession.GetName() == FInsightsManagerTabs::TimingProfilerTabId)
 	{
@@ -37,7 +37,7 @@ void FSlateTimingViewSession::OnBeginSession(Insights::ITimingViewSession& InTim
 	SlateFrameGraphTrack.Reset();
 }
 
-void FSlateTimingViewSession::OnEndSession(Insights::ITimingViewSession& InTimingViewSession)
+void FSlateTimingViewSession::OnEndSession(UE::Insights::Timing::ITimingViewSession& InTimingViewSession)
 {
 	if (&InTimingViewSession != TimingViewSession)
 	{
@@ -53,7 +53,7 @@ void FSlateTimingViewSession::OnEndSession(Insights::ITimingViewSession& InTimin
 	TimingViewSession = nullptr;
 }
 
-void FSlateTimingViewSession::Tick(Insights::ITimingViewSession& InTimingViewSession, const TraceServices::IAnalysisSession& InAnalysisSession)
+void FSlateTimingViewSession::Tick(UE::Insights::Timing::ITimingViewSession& InTimingViewSession, const TraceServices::IAnalysisSession& InAnalysisSession)
 {
 	if (&InTimingViewSession != TimingViewSession)
 	{
diff --git a/Engine/Plugins/Slate/SlateInsights/Source/SlateInsights/Private/SlateTimingViewSession.h b/Engine/Plugins/Slate/SlateInsights/Source/SlateInsights/Private/SlateTimingViewSession.h
--- a/Engine/Plugins/Slate/SlateInsights/Source/SlateInsights/Private/SlateTimingViewSession.h
+++ b/Engine/Plugins/Slate/SlateInsights/Source/SlateInsights/Private/SlateTimingViewSession.h
@@ -5,8 +5,8 @@
 #include "Templates/SharedPointer.h"
 
 namespace TraceServices { class IAnalysisSession; }
-namespace Insights { class ITimingViewSession; }
-namespace Insights { enum class ETimeChangedFlags : int32; }
+namespace UE::Insights::Timing { class ITimingViewSession; }
+namespace UE::Insights::Timing { enum class ETimeChangedFlags : int32; }
 namespace UE { namespace SlateInsights { class FSlateFrameGraphTrack; } }
 namespace UE { namespace SlateInsights { class FSlateWidgetUpdateStepsTimingTrack; } }
 class FMenuBuilder;
@@ -22,9 +22,9 @@ class FSlateTimingViewSession
 public:
 	FSlateTimingViewSession();
 
-	void OnBeginSession(Insights::ITimingViewSession& InTimingViewSession);
-	void OnEndSession(Insights::ITimingViewSession& InTimingViewSession);
-	void Tick(Insights::ITimingViewSession& InTimingViewSession, const TraceServices::IAnalysisSession& InAnalysisSession);
+	void OnBeginSession(UE::Insights::Timing::ITimingViewSession& InTimingViewSession);
+	void OnEndSession(UE::Insights::Timing::ITimingViewSession& InTimingViewSession);
+	void Tick(UE::Insights::Timing::ITimingViewSession& InTimingViewSession, const TraceServices::IAnalysisSession& InAnalysisSession);
 	void ExtendFilterMenu(FMenuBuilder& InMenuBuilder);
 
 	/** Get the last cached analysis session */
@@ -49,7 +49,7 @@ public:
 	void OpenSlateFrameTab() const;
 
 	/** The timing view for the session */
-	Insights::ITimingViewSession* GetTimingView() const
+	UE::Insights::Timing::ITimingViewSession* GetTimingView() const
 	{
 		return TimingViewSession;
 	}
@@ -59,7 +59,7 @@ private:
 	const TraceServices::IAnalysisSession* AnalysisSession;
 
 	// Cached timing view session, set in OnBeginSession/OnEndSession
-	Insights::ITimingViewSession* TimingViewSession;
+	UE::Insights::Timing::ITimingViewSession* TimingViewSession;
 
 	// All the tracks we manage
 	TSharedPtr<FSlateFrameGraphTrack> SlateFrameGraphTrack;
diff --git a/Engine/Plugins/Slate/SlateInsights/Source/SlateInsights/Private/SlateTraceModule.h b/Engine/Plugins/Slate/SlateInsights/Source/SlateInsights/Private/SlateTraceModule.h
--- a/Engine/Plugins/Slate/SlateInsights/Source/SlateInsights/Private/SlateTraceModule.h
+++ b/Engine/Plugins/Slate/SlateInsights/Source/SlateInsights/Private/SlateTraceModule.h
@@ -5,7 +5,6 @@
 #include "TraceServices/ModuleService.h"
 
 class IAnimationProvider;
-namespace Insights { class ITimingViewSession; }
 
 namespace UE
 {
diff --git a/Engine/Plugins/Slate/SlateInsights/Source/SlateInsights/Private/SlateWidgetUpdateStepsTimingTrack.cpp b/Engine/Plugins/Slate/SlateInsights/Source/SlateInsights/Private/SlateWidgetUpdateStepsTimingTrack.cpp
--- a/Engine/Plugins/Slate/SlateInsights/Source/SlateInsights/Private/SlateWidgetUpdateStepsTimingTrack.cpp
+++ b/Engine/Plugins/Slate/SlateInsights/Source/SlateInsights/Private/SlateWidgetUpdateStepsTimingTrack.cpp
@@ -6,7 +6,10 @@
 
 #include "Framework/MultiBox/MultiBoxBuilder.h"
 
-#include "Insights/Common/TimeUtils.h"
+// TraceInsightsCore
+#include "InsightsCore/Common/TimeUtils.h"
+
+// TraceInsights
 #include "Insights/ViewModels/TimingTrackViewport.h"
 #include "Insights/ViewModels/TimingEvent.h"
 #include "Insights/ViewModels/TooltipDrawState.h"
@@ -179,18 +182,18 @@ void FSlateWidgetUpdateStepsTimingTrack::InitTooltip(FTooltipDrawState& InOutToo
 	if (const Message::FWidgetInfo* WigetInfo = SlateProvider->FindWidget(Event.WidgetPaint.WidgetId))
 	{
 		InOutTooltip.AddTitle(FString::Printf(TEXT("Paint: %s"), *WigetInfo->DebugInfo));
-		InOutTooltip.AddNameValueTextLine(TEXT("Start Time:"), TimeUtils::FormatTimeAuto(InTooltipEvent.GetStartTime(), 6));
-		InOutTooltip.AddNameValueTextLine(TEXT("End Time:"), TimeUtils::FormatTimeAuto(InTooltipEvent.GetEndTime(), 6));
-		InOutTooltip.AddNameValueTextLine(TEXT("Duration:"), TimeUtils::FormatTimeAuto(InTooltipEvent.GetDuration()));
+		InOutTooltip.AddNameValueTextLine(TEXT("Start Time:"), UE::Insights::FormatTimeAuto(InTooltipEvent.GetStartTime(), 6));
+		InOutTooltip.AddNameValueTextLine(TEXT("End Time:"), UE::Insights::FormatTimeAuto(InTooltipEvent.GetEndTime(), 6));
+		InOutTooltip.AddNameValueTextLine(TEXT("Duration:"), UE::Insights::FormatTimeAuto(InTooltipEvent.GetDuration()));
 		InOutTooltip.AddNameValueTextLine(TEXT("Widget Id:"), FText::AsNumber(WigetInfo->WidgetId.GetValue()).ToString());
 		InOutTooltip.AddNameValueTextLine(TEXT("Widget Path:"), WigetInfo->Path);
 	}
 	else
 	{
 		InOutTooltip.AddTitle(TEXT("Paint: -"));
-		InOutTooltip.AddNameValueTextLine(TEXT("Start Time:"), TimeUtils::FormatTimeAuto(InTooltipEvent.GetStartTime(), 6));
-		InOutTooltip.AddNameValueTextLine(TEXT("End Time:"), TimeUtils::FormatTimeAuto(InTooltipEvent.GetEndTime(), 6));
-		InOutTooltip.AddNameValueTextLine(TEXT("Duration:"), TimeUtils::FormatTimeAuto(InTooltipEvent.GetDuration()));
+		InOutTooltip.AddNameValueTextLine(TEXT("Start Time:"), UE::Insights::FormatTimeAuto(InTooltipEvent.GetStartTime(), 6));
+		InOutTooltip.AddNameValueTextLine(TEXT("End Time:"), UE::Insights::FormatTimeAuto(InTooltipEvent.GetEndTime(), 6));
+		InOutTooltip.AddNameValueTextLine(TEXT("Duration:"), UE::Insights::FormatTimeAuto(InTooltipEvent.GetDuration()));
 		InOutTooltip.AddNameValueTextLine(TEXT("Widget Id:"), FText::AsNumber(Event.WidgetPaint.WidgetId.GetValue()).ToString());
 	}
 
@@ -235,7 +238,7 @@ void FSlateWidgetUpdateStepsTimingTrack::BuildContextMenu(FMenuBuilder& InOutMen
 				);
 			InOutMenuBuilder.AddMenuEntry(
 				LOCTEXT("HideChildEvent", "Hide Short Child Events"),
-				LOCTEXT("HideChildTooltip", "Show/Hide child events that are too short to be displed properly."),
+				LOCTEXT("HideChildTooltip", "Show/Hide child events that are too short to be displayed properly."),
 				FSlateIcon(),
 				FUIAction(
 					FExecuteAction::CreateSP(this, &FSlateWidgetUpdateStepsTimingTrack::ToggleShowChildWhenTrackIsTooSmall),
diff --git a/Engine/Source/Developer/TraceInsights/Public/Insights/Common/AvailabilityCheck.h b/Engine/Source/Developer/TraceInsights/Public/Insights/Common/AvailabilityCheck.h
--- a/Engine/Source/Developer/TraceInsights/Public/Insights/Common/AvailabilityCheck.h
+++ b/Engine/Source/Developer/TraceInsights/Public/Insights/Common/AvailabilityCheck.h
@@ -2,6 +2,9 @@
 
 #pragma once
 
+#include "Misc/CoreMiscDefines.h"
+UE_DEPRECATED_HEADER(5.5, "Use InsightsCore/Common/AvailabilityCheck.h from the TraceInsightsCore module instead.")
+
 // TraceInsightsCore
 #include "InsightsCore/Common/AvailabilityCheck.h"
 
diff --git a/Engine/Source/Developer/TraceInsights/Public/Insights/Common/PaintUtils.h b/Engine/Source/Developer/TraceInsights/Public/Insights/Common/PaintUtils.h
--- a/Engine/Source/Developer/TraceInsights/Public/Insights/Common/PaintUtils.h
+++ b/Engine/Source/Developer/TraceInsights/Public/Insights/Common/PaintUtils.h
@@ -2,5 +2,8 @@
 
 #pragma once
 
+#include "Misc/CoreMiscDefines.h"
+UE_DEPRECATED_HEADER(5.5, "Use InsightsCore/Common/PaintUtils.h from the TraceInsightsCore module instead.")
+
 // TraceInsightsCore
 #include "InsightsCore/Common/PaintUtils.h"
diff --git a/Engine/Source/Developer/TraceInsights/Public/Insights/Common/SimpleRtti.h b/Engine/Source/Developer/TraceInsights/Public/Insights/Common/SimpleRtti.h
--- a/Engine/Source/Developer/TraceInsights/Public/Insights/Common/SimpleRtti.h
+++ b/Engine/Source/Developer/TraceInsights/Public/Insights/Common/SimpleRtti.h
@@ -2,5 +2,8 @@
 
 #pragma once
 
+#include "Misc/CoreMiscDefines.h"
+UE_DEPRECATED_HEADER(5.5, "Use InsightsCore/Common/SimpleRtti.h from the TraceInsightsCore module instead.")
+
 // TraceInsightsCore
 #include "InsightsCore/Common/SimpleRtti.h"
diff --git a/Engine/Source/Developer/TraceInsights/Public/Insights/Common/TimeUtils.h b/Engine/Source/Developer/TraceInsights/Public/Insights/Common/TimeUtils.h
--- a/Engine/Source/Developer/TraceInsights/Public/Insights/Common/TimeUtils.h
+++ b/Engine/Source/Developer/TraceInsights/Public/Insights/Common/TimeUtils.h
@@ -2,6 +2,9 @@
 
 #pragma once
 
+#include "Misc/CoreMiscDefines.h"
+UE_DEPRECATED_HEADER(5.5, "Use InsightsCore/Common/TimeUtils.h from the TraceInsightsCore module instead.")
+
 // TraceInsightsCore
 #include "InsightsCore/Common/TimeUtils.h"
 
@@ -13,41 +16,41 @@
 namespace TimeUtils
 {
 
-static constexpr double /*UE_DEPRECATED(5.5, "TimeUtils namespace was merged into UE::Insights namespace")*/ Picosecond = 0.000000000001;
-static constexpr double /*UE_DEPRECATED(5.5, "TimeUtils namespace was merged into UE::Insights namespace")*/ Nanosecond = 0.000000001;
-static constexpr double /*UE_DEPRECATED(5.5, "TimeUtils namespace was merged into UE::Insights namespace")*/ Microsecond = 0.000001;
-static constexpr double /*UE_DEPRECATED(5.5, "TimeUtils namespace was merged into UE::Insights namespace")*/ Milisecond = 0.001; // !!!
-static constexpr double /*UE_DEPRECATED(5.5, "TimeUtils namespace was merged into UE::Insights namespace")*/ Second = 1.0;
-static constexpr double /*UE_DEPRECATED(5.5, "TimeUtils namespace was merged into UE::Insights namespace")*/ Minute = 60.0;
-static constexpr double /*UE_DEPRECATED(5.5, "TimeUtils namespace was merged into UE::Insights namespace")*/ Hour = 3600.0;
-static constexpr double /*UE_DEPRECATED(5.5, "TimeUtils namespace was merged into UE::Insights namespace")*/ Day = 86400.0;
-static constexpr double /*UE_DEPRECATED(5.5, "TimeUtils namespace was merged into UE::Insights namespace")*/ Week = 604800.0;
-
-//struct UE_DEPRECATED(5.5, "FTimeSplit was moved into UE::Insights namespace") FTimeSplit;
+static constexpr double UE_DEPRECATED(5.5, "TimeUtils namespace was merged into UE::Insights namespace") Picosecond = 0.000000000001;
+static constexpr double UE_DEPRECATED(5.5, "TimeUtils namespace was merged into UE::Insights namespace") Nanosecond = 0.000000001;
+static constexpr double UE_DEPRECATED(5.5, "TimeUtils namespace was merged into UE::Insights namespace") Microsecond = 0.000001;
+static constexpr double UE_DEPRECATED(5.5, "TimeUtils namespace was merged into UE::Insights namespace") Milisecond = 0.001; // !!!
+static constexpr double UE_DEPRECATED(5.5, "TimeUtils namespace was merged into UE::Insights namespace") Second = 1.0;
+static constexpr double UE_DEPRECATED(5.5, "TimeUtils namespace was merged into UE::Insights namespace") Minute = 60.0;
+static constexpr double UE_DEPRECATED(5.5, "TimeUtils namespace was merged into UE::Insights namespace") Hour = 3600.0;
+static constexpr double UE_DEPRECATED(5.5, "TimeUtils namespace was merged into UE::Insights namespace") Day = 86400.0;
+static constexpr double UE_DEPRECATED(5.5, "TimeUtils namespace was merged into UE::Insights namespace") Week = 604800.0;
+
+struct UE_DEPRECATED(5.5, "FTimeSplit was moved into UE::Insights namespace") FTimeSplit;
 struct FTimeSplit : public UE::Insights::FTimeSplit {};
 
-//UE_DEPRECATED(5.5, "TimeUtils namespace was merged into UE::Insights namespace")
+UE_DEPRECATED(5.5, "TimeUtils namespace was merged into UE::Insights namespace")
 inline FString FormatTimeValue(const double Duration, const int32 NumDigits = 1) { return UE::Insights::FormatTimeValue(Duration, NumDigits); }
 
-//UE_DEPRECATED(5.5, "TimeUtils namespace was merged into UE::Insights namespace")
+UE_DEPRECATED(5.5, "TimeUtils namespace was merged into UE::Insights namespace")
 inline FString FormatTimeAuto(const double Duration, const int32 NumDigits = 1) { return UE::Insights::FormatTimeAuto(Duration, NumDigits); }
 
-//UE_DEPRECATED(5.5, "TimeUtils namespace was merged into UE::Insights namespace")
+UE_DEPRECATED(5.5, "TimeUtils namespace was merged into UE::Insights namespace")
 inline FString FormatTimeMs(const double Duration, const int32 NumDigits = 2, bool bAddTimeUnit = false) { return UE::Insights::FormatTimeMs(Duration, NumDigits, bAddTimeUnit); }
 
-//UE_DEPRECATED(5.5, "TimeUtils namespace was merged into UE::Insights namespace")
+UE_DEPRECATED(5.5, "TimeUtils namespace was merged into UE::Insights namespace")
 inline FString FormatTime(const double Time, const double Precision = 0.0) { return UE::Insights::FormatTime(Time, Precision); }
 
-//UE_DEPRECATED(5.5, "TimeUtils namespace was merged into UE::Insights namespace")
+UE_DEPRECATED(5.5, "TimeUtils namespace was merged into UE::Insights namespace")
 inline FString FormatTimeHMS(const double Time, const double Precision = 0.0) { return UE::Insights::FormatTimeHMS(Time, Precision); }
 
-//UE_DEPRECATED(5.5, "TimeUtils namespace was merged into UE::Insights namespace")
+UE_DEPRECATED(5.5, "TimeUtils namespace was merged into UE::Insights namespace")
 inline void SplitTime(const double Time, UE::Insights::FTimeSplit& OutTimeSplit) { UE::Insights::SplitTime(Time, OutTimeSplit); }
 
-//UE_DEPRECATED(5.5, "TimeUtils namespace was merged into UE::Insights namespace")
+UE_DEPRECATED(5.5, "TimeUtils namespace was merged into UE::Insights namespace")
 inline FString FormatTimeSplit(const UE::Insights::FTimeSplit& TimeSplit, const double Precision = 0.0) { return UE::Insights::FormatTimeSplit(TimeSplit, Precision); }
 
-//UE_DEPRECATED(5.5, "TimeUtils namespace was merged into UE::Insights namespace")
+UE_DEPRECATED(5.5, "TimeUtils namespace was merged into UE::Insights namespace")
 inline FString FormatTimeSplit(const double Time, const double Precision = 0.0) { return UE::Insights::FormatTimeSplit(Time, Precision); }
 
 } // namespace TimeUtils
diff --git a/Engine/Source/Developer/TraceInsights/Public/Insights/ITimingViewExtender.h b/Engine/Source/Developer/TraceInsights/Public/Insights/ITimingViewExtender.h
--- a/Engine/Source/Developer/TraceInsights/Public/Insights/ITimingViewExtender.h
+++ b/Engine/Source/Developer/TraceInsights/Public/Insights/ITimingViewExtender.h
@@ -62,9 +62,10 @@ public:
 namespace Insights
 {
 
+UE_DEPRECATED(5.5, "TimingViewExtenderFeatureName was moved inside UE::Insights::Timing namespace")
 extern TRACEINSIGHTS_API const FName TimingViewExtenderFeatureName;
 
-//class UE_DEPRECATED(5.5, "ITimingViewExtender class was moved inside UE::Insights::Timing namespace") ITimingViewExtender;
+class UE_DEPRECATED(5.5, "ITimingViewExtender class was moved inside UE::Insights::Timing namespace") ITimingViewExtender;
 class ITimingViewExtender : public IModularFeature
 {
 public:
@@ -87,9 +88,10 @@ public:
 namespace UE::Insights
 {
 
+UE_DEPRECATED(5.5, "TimingViewExtenderFeatureName was moved inside UE::Insights::Timing namespace")
 extern TRACEINSIGHTS_API const FName TimingViewExtenderFeatureName;
 
-//class UE_DEPRECATED(5.5, "ITimingViewExtender class was moved inside UE::Insights::Timing namespace") ITimingViewExtender;
+class UE_DEPRECATED(5.5, "ITimingViewExtender class was moved inside UE::Insights::Timing namespace") ITimingViewExtender;
 class ITimingViewExtender : public IModularFeature
 {
 public:
diff --git a/Engine/Source/Developer/TraceInsights/Public/Insights/ITimingViewSession.h b/Engine/Source/Developer/TraceInsights/Public/Insights/ITimingViewSession.h
--- a/Engine/Source/Developer/TraceInsights/Public/Insights/ITimingViewSession.h
+++ b/Engine/Source/Developer/TraceInsights/Public/Insights/ITimingViewSession.h
@@ -190,14 +190,16 @@ public:
 namespace Insights
 {
 
-enum class ETimeChangedFlags : int32
+enum class UE_DEPRECATED(5.5, "ITimeMarker class was moved inside UE::Insights::Timing namespace") ETimeChangedFlags : int32
 {
 	None,
 	Interactive = (1 << 0)
 };
+PRAGMA_DISABLE_DEPRECATION_WARNINGS
 ENUM_CLASS_FLAGS(ETimeChangedFlags);
+PRAGMA_ENABLE_DEPRECATION_WARNINGS
 
-class ITimeMarker
+class UE_DEPRECATED(5.5, "ITimeMarker class was moved inside UE::Insights::Timing namespace") ITimeMarker
 {
 public:
 	virtual ~ITimeMarker() = default;
@@ -206,6 +208,7 @@ public:
 	virtual void SetTime(const double  InTime) = 0;
 };
 
+PRAGMA_DISABLE_DEPRECATION_WARNINGS
 DECLARE_MULTICAST_DELEGATE_ThreeParams(FSelectionChangedDelegate, ETimeChangedFlags /*InFlags*/, double /*StartTime*/, double /*EndTime*/);
 DECLARE_MULTICAST_DELEGATE_TwoParams(FTimeMarkerChangedDelegate, ETimeChangedFlags /*InFlags*/, double /*TimeMarker*/);
 DECLARE_MULTICAST_DELEGATE_TwoParams(FCustomTimeMarkerChangedDelegate, ETimeChangedFlags /*InFlags*/, TSharedRef<ITimeMarker> /*TimeMarker*/);
@@ -216,9 +219,9 @@ DECLARE_MULTICAST_DELEGATE_OneParam(FSelectedEventChangedDelegate, const TShared
 DECLARE_MULTICAST_DELEGATE(FTrackVisibilityChangedDelegate);
 DECLARE_MULTICAST_DELEGATE_OneParam(FTrackAddedDelegate, const TSharedPtr<const FBaseTimingTrack> /*Track*/);
 DECLARE_MULTICAST_DELEGATE_OneParam(FTrackRemovedDelegate, const TSharedPtr<const FBaseTimingTrack> /*Track*/);
+PRAGMA_ENABLE_DEPRECATION_WARNINGS
 
-//class UE_DEPRECATED(5.5, "ITimingViewSession class was moved inside UE::Insights::Timing namespace") ITimingViewSession;
-class ITimingViewSession
+class UE_DEPRECATED(5.5, "ITimingViewSession class was moved inside UE::Insights::Timing namespace") ITimingViewSession
 {
 public:
 	virtual ~ITimingViewSession() = default;
