[Insights]
  - Moved most of the TimingProfiler code into UE::Insights::TimingProfiler namespace. Corresponding files moved into /TimingProfiler/ base folder, with following sub-folder:
    * /TimingProfiler/Tracks/
    * /TimingProfiler/ViewModels/
    * /TimingProfiler/Widgets/
  - Replaced most includes of "CoreMinimal.h" with "CoreTypes.h".
  - Fixed issues re key bindings. (ex. a bug where "R" key was used for both regions track and for frames track).
  - Fixed issues re UI commands.

diff --git a/Engine/Source/Developer/TraceInsights/Private/Insights/Common/InsightsMenuBuilder.h b/Engine/Source/Developer/TraceInsights/Private/Insights/Common/InsightsMenuBuilder.h
--- a/Engine/Source/Developer/TraceInsights/Private/Insights/Common/InsightsMenuBuilder.h
+++ b/Engine/Source/Developer/TraceInsights/Private/Insights/Common/InsightsMenuBuilder.h
@@ -2,7 +2,8 @@
 
 #pragma once
 
-#include "CoreMinimal.h"
+#include "CoreTypes.h"
+
 #include "Framework/Commands/UICommandInfo.h"
 
 class FWorkspaceItem;
diff --git a/Engine/Source/Developer/TraceInsights/Private/Insights/ContextSwitches/ContextSwitchesProfilerManager.cpp b/Engine/Source/Developer/TraceInsights/Private/Insights/ContextSwitches/ContextSwitchesProfilerManager.cpp
--- a/Engine/Source/Developer/TraceInsights/Private/Insights/ContextSwitches/ContextSwitchesProfilerManager.cpp
+++ b/Engine/Source/Developer/TraceInsights/Private/Insights/ContextSwitches/ContextSwitchesProfilerManager.cpp
@@ -12,8 +12,8 @@
 // TraceInsights
 #include "Insights/ContextSwitches/ViewModels/ContextSwitchesSharedState.h"
 #include "Insights/ITimingViewExtender.h"
-#include "Insights/TimingProfilerManager.h"
-#include "Insights/Widgets/STimingProfilerWindow.h"
+#include "Insights/TimingProfiler/TimingProfilerManager.h"
+#include "Insights/TimingProfiler/Widgets/STimingProfilerWindow.h"
 #include "Insights/Widgets/STimingView.h"
 
 ////////////////////////////////////////////////////////////////////////////////////////////////////
diff --git a/Engine/Source/Developer/TraceInsights/Private/Insights/ContextSwitches/ViewModels/ContextSwitchesSharedState.cpp b/Engine/Source/Developer/TraceInsights/Private/Insights/ContextSwitches/ViewModels/ContextSwitchesSharedState.cpp
--- a/Engine/Source/Developer/TraceInsights/Private/Insights/ContextSwitches/ViewModels/ContextSwitchesSharedState.cpp
+++ b/Engine/Source/Developer/TraceInsights/Private/Insights/ContextSwitches/ViewModels/ContextSwitchesSharedState.cpp
@@ -22,9 +22,10 @@
 #include "Insights/ITimingViewSession.h"
 #include "Insights/InsightsManager.h"
 #include "Insights/InsightsStyle.h"
-#include "Insights/TimingProfilerManager.h"
-#include "Insights/ViewModels/ThreadTimingTrack.h"
-#include "Insights/Widgets/STimingProfilerWindow.h"
+#include "Insights/TimingProfiler/TimingProfilerManager.h"
+#include "Insights/TimingProfiler/Tracks/ThreadTimingTrack.h"
+#include "Insights/TimingProfiler/ViewModels/ThreadTimingSharedState.h"
+#include "Insights/TimingProfiler/Widgets/STimingProfilerWindow.h"
 #include "Insights/Widgets/STimingView.h"
 
 #define LOCTEXT_NAMESPACE "UE::Insights::ContextSwitches"
@@ -47,12 +48,6 @@ FContextSwitchesStateCommands::FContextSwitchesStateCommands()
 
 ////////////////////////////////////////////////////////////////////////////////////////////////////
 
-FContextSwitchesStateCommands::~FContextSwitchesStateCommands()
-{
-}
-
-////////////////////////////////////////////////////////////////////////////////////////////////////
-
 // UI_COMMAND takes long for the compiler to optimize
 UE_DISABLE_OPTIMIZATION_SHIP
 void FContextSwitchesStateCommands::RegisterCommands()
@@ -83,7 +78,7 @@ void FContextSwitchesStateCommands::RegisterCommands()
 
 	UI_COMMAND(Command_ShowNonTargetProcessEvents,
 		"Non-Target Process Events",
-		"Shows/hides the cpu core events that do not belong to the target process.",
+		"Shows/hides the CPU Core events that do not belong to the target process.",
 		EUserInterfaceActionType::ToggleButton,
 		FInputChord());
 
@@ -471,7 +466,7 @@ void FContextSwitchesSharedState::AddContextSwitchesChildTracks()
 
 	ThreadsSerial = NewThreadsSerial;
 
-	//TODO: Create "Cpu Thread" timing tracks also for threads without cpu timing events (i.e. only with context switch events).
+	//TODO: Create "CPU Thread" timing tracks also for threads without CPU timing events (i.e. only with context switch events).
 #endif
 
 	TSharedPtr<TimingProfiler::STimingView> TimingView = GetTimingView();
@@ -480,18 +475,18 @@ void FContextSwitchesSharedState::AddContextSwitchesChildTracks()
 		return;
 	}
 
-	TSharedPtr<FThreadTimingSharedState> TimingSharedState = TimingView->GetThreadTimingSharedState();
+	TSharedPtr<TimingProfiler::FThreadTimingSharedState> TimingSharedState = TimingView->GetThreadTimingSharedState();
 
 	if (!TimingSharedState.IsValid())
 	{
 		return;
 	}
 
-	const TMap<uint32, TSharedPtr<FCpuTimingTrack>>& CpuTracks = TimingSharedState->GetAllCpuTracks();
+	const TMap<uint32, TSharedPtr<TimingProfiler::FCpuTimingTrack>>& CpuTracks = TimingSharedState->GetAllCpuTracks();
 
-	for (const TPair<uint32, TSharedPtr<FCpuTimingTrack>>& MapEntry : CpuTracks)
+	for (const TPair<uint32, TSharedPtr<TimingProfiler::FCpuTimingTrack>>& MapEntry : CpuTracks)
 	{
-		const TSharedPtr<FCpuTimingTrack>& CpuTrack = MapEntry.Value;
+		const TSharedPtr<TimingProfiler::FCpuTimingTrack>& CpuTrack = MapEntry.Value;
 		if (CpuTrack.IsValid() && !CpuTrack->GetChildTrack().IsValid())
 		{
 			TSharedPtr<FContextSwitchesTimingTrack> ContextSwitchesTrack = MakeShared<FContextSwitchesTimingTrack>(*this, TEXT("Context Switches"), CpuTrack->GetTimelineIndex(), CpuTrack->GetThreadId());
@@ -512,18 +507,18 @@ void FContextSwitchesSharedState::RemoveContextSwitchesChildTracks()
 	{
 		return;
 	}
-	TSharedPtr<FThreadTimingSharedState> TimingSharedState = TimingView->GetThreadTimingSharedState();
+	TSharedPtr<TimingProfiler::FThreadTimingSharedState> TimingSharedState = TimingView->GetThreadTimingSharedState();
 
 	if (!TimingSharedState.IsValid())
 	{
 		return;
 	}
 
-	const TMap<uint32, TSharedPtr<FCpuTimingTrack>>& CpuTracks = TimingSharedState->GetAllCpuTracks();
+	const TMap<uint32, TSharedPtr<TimingProfiler::FCpuTimingTrack>>& CpuTracks = TimingSharedState->GetAllCpuTracks();
 
-	for (const TPair<uint32, TSharedPtr<FCpuTimingTrack>>& MapEntry : CpuTracks)
+	for (const TPair<uint32, TSharedPtr<TimingProfiler::FCpuTimingTrack>>& MapEntry : CpuTracks)
 	{
-		const TSharedPtr<FCpuTimingTrack>& CpuTrack = MapEntry.Value;
+		const TSharedPtr<TimingProfiler::FCpuTimingTrack>& CpuTrack = MapEntry.Value;
 		if (CpuTrack.IsValid() && CpuTrack->GetChildTrack().IsValid() && CpuTrack->GetChildTrack()->Is<FContextSwitchesTimingTrack>())
 		{
 			CpuTrack->SetChildTrack(nullptr);
@@ -661,7 +656,7 @@ void FContextSwitchesSharedState::Command_NavigateToCpuThreadEvent_Execute()
 			GetThreadInfo(SystemThreadId, ThreadId, ThreadName);
 			if (ThreadId != ~0)
 			{
-				TSharedPtr<FThreadTimingTrack> ThreadTimingTrack = GetThreadTimingTrack(ThreadId);
+				TSharedPtr<TimingProfiler::FThreadTimingTrack> ThreadTimingTrack = GetThreadTimingTrack(ThreadId);
 				if (ThreadTimingTrack.IsValid() && ThreadTimingTrack->IsVisible())
 				{
 					TimingView->SelectTimingTrack(ThreadTimingTrack, true);
@@ -708,7 +703,7 @@ bool FContextSwitchesSharedState::Command_NavigateToCpuThreadEvent_CanExecute()
 			GetThreadInfo(SystemThreadId, ThreadId, ThreadName);
 			if (ThreadId != ~0)
 			{
-				TSharedPtr<FThreadTimingTrack> ThreadTimingTrack = GetThreadTimingTrack(ThreadId);
+				TSharedPtr<TimingProfiler::FThreadTimingTrack> ThreadTimingTrack = GetThreadTimingTrack(ThreadId);
 				if (ThreadTimingTrack.IsValid() && ThreadTimingTrack->IsVisible())
 				{
 					return true;
@@ -742,7 +737,7 @@ void FContextSwitchesSharedState::Command_DockCpuThreadTrackToBottom_Execute()
 			GetThreadInfo(SystemThreadId, ThreadId, ThreadName);
 			if (ThreadId != ~0)
 			{
-				TSharedPtr<FThreadTimingTrack> ThreadTimingTrack = GetThreadTimingTrack(ThreadId);
+				TSharedPtr<TimingProfiler::FThreadTimingTrack> ThreadTimingTrack = GetThreadTimingTrack(ThreadId);
 				if (ThreadTimingTrack.IsValid())
 				{
 					TimingView->ChangeTrackLocation(ThreadTimingTrack.ToSharedRef(), ETimingTrackLocation::BottomDocked);
@@ -855,19 +850,19 @@ void FContextSwitchesSharedState::GetThreadInfo(uint32 InSystemThreadId, uint32&
 
 ////////////////////////////////////////////////////////////////////////////////////////////////////
 
-TSharedPtr<FThreadTimingTrack> FContextSwitchesSharedState::GetThreadTimingTrack(uint32 ThreadId) const
+TSharedPtr<TimingProfiler::FThreadTimingTrack> FContextSwitchesSharedState::GetThreadTimingTrack(uint32 ThreadId) const
 {
-	TSharedPtr<FThreadTimingTrack> FoundTrack;
+	TSharedPtr<TimingProfiler::FThreadTimingTrack> FoundTrack;
 	TSharedPtr<TimingProfiler::STimingView> TimingView = GetTimingView();
 
 	if (TimingView)
 	{
 		TimingView->EnumerateAllTracks([&FoundTrack, ThreadId](TSharedPtr<FBaseTimingTrack>& Track) -> bool
 		{
-			if (Track->Is<FThreadTimingTrack>() &&
-				Track->As<FThreadTimingTrack>().GetThreadId() == ThreadId)
+			if (Track->Is<TimingProfiler::FThreadTimingTrack>() &&
+				Track->As<TimingProfiler::FThreadTimingTrack>().GetThreadId() == ThreadId)
 			{
-				FoundTrack = StaticCastSharedPtr<FThreadTimingTrack>(Track);
+				FoundTrack = StaticCastSharedPtr<TimingProfiler::FThreadTimingTrack>(Track);
 				return false;
 			}
 			return true;
diff --git a/Engine/Source/Developer/TraceInsights/Private/Insights/ContextSwitches/ViewModels/ContextSwitchesSharedState.h b/Engine/Source/Developer/TraceInsights/Private/Insights/ContextSwitches/ViewModels/ContextSwitchesSharedState.h
--- a/Engine/Source/Developer/TraceInsights/Private/Insights/ContextSwitches/ViewModels/ContextSwitchesSharedState.h
+++ b/Engine/Source/Developer/TraceInsights/Private/Insights/ContextSwitches/ViewModels/ContextSwitchesSharedState.h
@@ -13,13 +13,18 @@
 
 class FMenuBuilder;
 
+namespace TraceServices { class IAnalysisSession; }
+
 class FBaseTimingTrack;
-class FThreadTimingTrack;
 class ITimingEvent;
 
-namespace TraceServices { class IAnalysisSession; }
 namespace UE::Insights::Timing { class ITimingViewSession; }
-namespace UE::Insights::TimingProfiler { class STimingView; }
+
+namespace UE::Insights::TimingProfiler
+{
+	class FThreadTimingTrack;
+	class STimingView;
+}
 
 namespace UE::Insights::ContextSwitches
 {
@@ -33,7 +38,7 @@ class FContextSwitchesStateCommands : public TCommands<FContextSwitchesStateComm
 {
 public:
 	FContextSwitchesStateCommands();
-	virtual ~FContextSwitchesStateCommands();
+	virtual ~FContextSwitchesStateCommands() {}
 	virtual void RegisterCommands() override;
 
 	// Commands for the Tracks Filter menu.
@@ -113,7 +118,7 @@ public:
 	void AddCommands();
 
 	void GetThreadInfo(uint32 InSystemThreadId, uint32& OutThreadId, const TCHAR*& OutThreadName) const;
-	TSharedPtr<FThreadTimingTrack> GetThreadTimingTrack(uint32 ThreadId) const;
+	TSharedPtr<TimingProfiler::FThreadTimingTrack> GetThreadTimingTrack(uint32 ThreadId) const;
 	TSharedPtr<FCpuCoreTimingTrack> GetCpuCoreTimingTrack(uint32 CoreNumber) const;
 
 private:
diff --git a/Engine/Source/Developer/TraceInsights/Private/Insights/CookProfiler/CookProfilerManager.cpp b/Engine/Source/Developer/TraceInsights/Private/Insights/CookProfiler/CookProfilerManager.cpp
--- a/Engine/Source/Developer/TraceInsights/Private/Insights/CookProfiler/CookProfilerManager.cpp
+++ b/Engine/Source/Developer/TraceInsights/Private/Insights/CookProfiler/CookProfilerManager.cpp
@@ -15,8 +15,8 @@
 #include "Insights/InsightsStyle.h"
 #include "Insights/CookProfiler/ViewModels/PackageTable.h"
 #include "Insights/CookProfiler/Widgets/SPackageTableTreeView.h"
-#include "Insights/TimingProfilerManager.h"
-#include "Insights/Widgets/STimingProfilerWindow.h"
+#include "Insights/TimingProfiler/TimingProfilerManager.h"
+#include "Insights/TimingProfiler/Widgets/STimingProfilerWindow.h"
 
 #define LOCTEXT_NAMESPACE "UE::Insights::CookProfiler"
 
diff --git a/Engine/Source/Developer/TraceInsights/Private/Insights/CookProfiler/Widgets/SPackageTableTreeView.cpp b/Engine/Source/Developer/TraceInsights/Private/Insights/CookProfiler/Widgets/SPackageTableTreeView.cpp
--- a/Engine/Source/Developer/TraceInsights/Private/Insights/CookProfiler/Widgets/SPackageTableTreeView.cpp
+++ b/Engine/Source/Developer/TraceInsights/Private/Insights/CookProfiler/Widgets/SPackageTableTreeView.cpp
@@ -18,8 +18,8 @@
 #include "Insights/CookProfiler/CookProfilerManager.h"
 #include "Insights/CookProfiler/ViewModels/PackageEntry.h"
 #include "Insights/CookProfiler/ViewModels/PackageNode.h"
-#include "Insights/TimingProfilerManager.h"
-#include "Insights/Widgets/STimingProfilerWindow.h"
+#include "Insights/TimingProfiler/TimingProfilerManager.h"
+#include "Insights/TimingProfiler/Widgets/STimingProfilerWindow.h"
 #include "Insights/Widgets/STimingView.h"
 
 #include <limits>
@@ -30,7 +30,7 @@ namespace UE::Insights::CookProfiler
 {
 
 ////////////////////////////////////////////////////////////////////////////////////////////////////
-// FTaskTableTreeViewCommands
+// FPackageTableTreeViewCommands
 ////////////////////////////////////////////////////////////////////////////////////////////////////
 
 class FPackageTableTreeViewCommands : public TCommands<FPackageTableTreeViewCommands>
diff --git a/Engine/Source/Developer/TraceInsights/Private/Insights/CookProfiler/Widgets/SPackageTableTreeView.h b/Engine/Source/Developer/TraceInsights/Private/Insights/CookProfiler/Widgets/SPackageTableTreeView.h
--- a/Engine/Source/Developer/TraceInsights/Private/Insights/CookProfiler/Widgets/SPackageTableTreeView.h
+++ b/Engine/Source/Developer/TraceInsights/Private/Insights/CookProfiler/Widgets/SPackageTableTreeView.h
@@ -2,7 +2,7 @@
 
 #pragma once
 
-#include "CoreMinimal.h"
+#include "CoreTypes.h"
 
 // TraceServices
 #include "TraceServices/Model/CookProfilerProvider.h"
diff --git a/Engine/Source/Developer/TraceInsights/Private/Insights/ImportTool/TableImportTool.h b/Engine/Source/Developer/TraceInsights/Private/Insights/ImportTool/TableImportTool.h
--- a/Engine/Source/Developer/TraceInsights/Private/Insights/ImportTool/TableImportTool.h
+++ b/Engine/Source/Developer/TraceInsights/Private/Insights/ImportTool/TableImportTool.h
@@ -2,7 +2,7 @@
 
 #pragma once
 
-#include "CoreMinimal.h"
+#include "CoreTypes.h"
 
 // TraceServices
 #include "TraceServices/Model/TableImport.h"
diff --git a/Engine/Source/Developer/TraceInsights/Private/Insights/InsightsCommands.h b/Engine/Source/Developer/TraceInsights/Private/Insights/InsightsCommands.h
--- a/Engine/Source/Developer/TraceInsights/Private/Insights/InsightsCommands.h
+++ b/Engine/Source/Developer/TraceInsights/Private/Insights/InsightsCommands.h
@@ -2,9 +2,10 @@
 
 #pragma once
 
-#include "CoreMinimal.h"
-#include "Framework/Commands/UIAction.h"
+#include "CoreTypes.h"
+
 #include "Framework/Commands/Commands.h"
+#include "Framework/Commands/UIAction.h"
 
 ////////////////////////////////////////////////////////////////////////////////////////////////////
 /**
diff --git a/Engine/Source/Developer/TraceInsights/Private/Insights/InsightsManager.cpp b/Engine/Source/Developer/TraceInsights/Private/Insights/InsightsManager.cpp
--- a/Engine/Source/Developer/TraceInsights/Private/Insights/InsightsManager.cpp
+++ b/Engine/Source/Developer/TraceInsights/Private/Insights/InsightsManager.cpp
@@ -37,10 +37,10 @@
 #include "Insights/Log.h"
 #include "Insights/NetworkingProfiler/NetworkingProfilerManager.h"
 #include "Insights/Tests/InsightsTestRunner.h"
-#include "Insights/TimingProfilerManager.h"
+#include "Insights/TimingProfiler/TimingProfilerManager.h"
+#include "Insights/TimingProfiler/Widgets/STimingProfilerWindow.h"
 #include "Insights/Widgets/SSessionInfoWindow.h"
 #include "Insights/Widgets/SStartPageWindow.h"
-#include "Insights/Widgets/STimingProfilerWindow.h"
 
 ////////////////////////////////////////////////////////////////////////////////////////////////////
 
diff --git a/Engine/Source/Developer/TraceInsights/Private/Insights/InsightsSessionBrowserSettings.h b/Engine/Source/Developer/TraceInsights/Private/Insights/InsightsSessionBrowserSettings.h
--- a/Engine/Source/Developer/TraceInsights/Private/Insights/InsightsSessionBrowserSettings.h
+++ b/Engine/Source/Developer/TraceInsights/Private/Insights/InsightsSessionBrowserSettings.h
@@ -2,10 +2,10 @@
 
 #pragma once
 
-#include "CoreMinimal.h"
+#include "CoreTypes.h"
+
 #include "Misc/ConfigCacheIni.h"
 #include "Misc/ConfigContext.h"
-#include "ProfilingDebugging/MiscTrace.h"
 
 /** Contains all settings for the Unreal Insights, accessible through the main manager. */
 class FInsightsSessionBrowserSettings
diff --git a/Engine/Source/Developer/TraceInsights/Private/Insights/InsightsStyle.h b/Engine/Source/Developer/TraceInsights/Private/Insights/InsightsStyle.h
--- a/Engine/Source/Developer/TraceInsights/Private/Insights/InsightsStyle.h
+++ b/Engine/Source/Developer/TraceInsights/Private/Insights/InsightsStyle.h
@@ -2,9 +2,15 @@
 
 #pragma once
 
-#include "CoreMinimal.h"
+#include "CoreTypes.h"
+
+#include "Styling/SlateColor.h"
 #include "Styling/SlateStyle.h"
 #include "Styling/SlateTypes.h"
+#include "Templates/SharedPointer.h"
+#include "Types/SlateEnums.h"
+
+struct FSlateBrush;
 
 /** Style data for Insights tools */
 class FInsightsStyle
diff --git a/Engine/Source/Developer/TraceInsights/Private/Insights/LoadingProfiler/Tracks/LoadingTimingTrack.cpp b/Engine/Source/Developer/TraceInsights/Private/Insights/LoadingProfiler/Tracks/LoadingTimingTrack.cpp
new file mode 100644
--- /dev/null
+++ b/Engine/Source/Developer/TraceInsights/Private/Insights/LoadingProfiler/Tracks/LoadingTimingTrack.cpp
@@ -0,0 +1,304 @@
+// Copyright Epic Games, Inc. All Rights Reserved.
+
+#include "LoadingTimingTrack.h"
+
+// TraceServices
+#include "TraceServices/Containers/Timelines.h"
+#include "TraceServices/Model/AnalysisSession.h"
+#include "TraceServices/Model/LoadTimeProfiler.h"
+
+// TraceInsightsCore
+#include "InsightsCore/Common/TimeUtils.h"
+#include "InsightsCore/Filter/ViewModels/FilterConfigurator.h"
+#include "InsightsCore/Filter/ViewModels/Filters.h"
+
+// TraceInsights
+#include "Insights/InsightsManager.h"
+#include "Insights/LoadingProfiler/ViewModels/LoadingSharedState.h"
+#include "Insights/ViewModels/TimingEvent.h"
+#include "Insights/ViewModels/TimingEventSearch.h"
+#include "Insights/ViewModels/TimingTrackViewport.h"
+#include "Insights/ViewModels/TooltipDrawState.h"
+
+#define LOCTEXT_NAMESPACE "UE::Insights::LoadingProfiler"
+
+namespace UE::Insights::LoadingProfiler
+{
+
+////////////////////////////////////////////////////////////////////////////////////////////////////
+
+INSIGHTS_IMPLEMENT_RTTI(FLoadingTimingTrack)
+
+////////////////////////////////////////////////////////////////////////////////////////////////////
+
+void FLoadingTimingTrack::BuildDrawState(ITimingEventsTrackDrawStateBuilder& Builder, const ITimingTrackUpdateContext& Context)
+{
+	TSharedPtr<const TraceServices::IAnalysisSession> Session = FInsightsManager::Get()->GetSession();
+	if (Session.IsValid() && TraceServices::ReadLoadTimeProfilerProvider(*Session.Get()))
+	{
+		TraceServices::FAnalysisSessionReadScope SessionReadScope(*Session.Get());
+
+		const TraceServices::ILoadTimeProfilerProvider& LoadTimeProfilerProvider = *TraceServices::ReadLoadTimeProfilerProvider(*Session.Get());
+
+		const FTimingTrackViewport& Viewport = Context.GetViewport();
+
+		LoadTimeProfilerProvider.ReadTimeline(TimelineIndex, [this, &Builder, &Viewport](const TraceServices::ILoadTimeProfilerProvider::CpuTimeline& Timeline)
+		{
+			if (FTimingEventsTrack::bUseDownSampling)
+			{
+				const double SecondsPerPixel = 1.0 / Viewport.GetScaleX();
+				Timeline.EnumerateEventsDownSampled(Viewport.GetStartTime(), Viewport.GetEndTime(), SecondsPerPixel, [this, &Builder](double StartTime, double EndTime, uint32 Depth, const TraceServices::FLoadTimeProfilerCpuEvent& Event)
+				{
+					if (Event.Package)
+					{
+						const TCHAR* Name = SharedState.GetEventName(Depth, Event);
+						const uint64 Type = static_cast<uint64>(Event.EventType);
+						const uint32 Color = 0;
+						Builder.AddEvent(StartTime, EndTime, Depth, Name, Type, Color);
+					}
+					return TraceServices::EEventEnumerate::Continue;
+				});
+			}
+			else
+			{
+				Timeline.EnumerateEvents(Viewport.GetStartTime(), Viewport.GetEndTime(), [this, &Builder](double StartTime, double EndTime, uint32 Depth, const TraceServices::FLoadTimeProfilerCpuEvent& Event)
+				{
+					if (Event.Package)
+					{
+						const TCHAR* Name = SharedState.GetEventName(Depth, Event);
+						const uint64 Type = static_cast<uint64>(Event.EventType);
+						const uint32 Color = 0;
+						Builder.AddEvent(StartTime, EndTime, Depth, Name, Type, Color);
+					}
+					return TraceServices::EEventEnumerate::Continue;
+				});
+			}
+		});
+	}
+}
+
+////////////////////////////////////////////////////////////////////////////////////////////////////
+
+void FLoadingTimingTrack::BuildFilteredDrawState(ITimingEventsTrackDrawStateBuilder& Builder, const ITimingTrackUpdateContext& Context)
+{
+	if (HasCustomFilter())
+	{
+		using EFilterField = UE::Insights::EFilterField;
+		UE::Insights::FFilterContext FilterConfiguratorContext;
+		FilterConfiguratorContext.SetReturnValueForUnsetFilters(false);
+		FilterConfiguratorContext.AddFilterData<double>(static_cast<int32>(EFilterField::StartTime), 0.0f);
+		FilterConfiguratorContext.AddFilterData<double>(static_cast<int32>(EFilterField::EndTime), 0.0f);
+		FilterConfiguratorContext.AddFilterData<double>(static_cast<int32>(EFilterField::Duration), 0.0f);
+		FilterConfiguratorContext.AddFilterData<FString>(static_cast<int32>(EFilterField::TrackName), this->GetName());
+
+		TSharedPtr<const TraceServices::IAnalysisSession> Session = FInsightsManager::Get()->GetSession();
+		if (Session.IsValid() && TraceServices::ReadLoadTimeProfilerProvider(*Session.Get()))
+		{
+			TraceServices::FAnalysisSessionReadScope SessionReadScope(*Session.Get());
+
+			const TraceServices::ILoadTimeProfilerProvider& LoadTimeProfilerProvider = *TraceServices::ReadLoadTimeProfilerProvider(*Session.Get());
+
+			const FTimingTrackViewport& Viewport = Context.GetViewport();
+
+			LoadTimeProfilerProvider.ReadTimeline(TimelineIndex, [this, &Builder, &Viewport, &FilterConfiguratorContext](const TraceServices::ILoadTimeProfilerProvider::CpuTimeline& Timeline)
+			{
+				Timeline.EnumerateEvents(Viewport.GetStartTime(), Viewport.GetEndTime(), [this, &Builder, &FilterConfiguratorContext](double StartTime, double EndTime, uint32 Depth, const TraceServices::FLoadTimeProfilerCpuEvent& Event)
+				{
+					if (Event.Package)
+					{
+						FilterConfiguratorContext.SetFilterData<double>(static_cast<int32>(EFilterField::StartTime), StartTime);
+						FilterConfiguratorContext.SetFilterData<double>(static_cast<int32>(EFilterField::EndTime), EndTime);
+						FilterConfiguratorContext.SetFilterData<double>(static_cast<int32>(EFilterField::Duration), EndTime - StartTime);
+
+						if (FilterConfigurator->ApplyFilters(FilterConfiguratorContext))
+						{
+							const TCHAR* Name = SharedState.GetEventName(Depth, Event);
+							const uint64 Type = static_cast<uint64>(Event.EventType);
+							const uint32 Color = 0;
+							Builder.AddEvent(StartTime, EndTime, Depth, Name, Type, Color);
+						}
+					}
+
+					return TraceServices::EEventEnumerate::Continue;
+				});
+			});
+		}
+	}
+}
+
+////////////////////////////////////////////////////////////////////////////////////////////////////
+
+void FLoadingTimingTrack::InitTooltip(FTooltipDrawState& InOutTooltip, const ITimingEvent& InTooltipEvent) const
+{
+	if (InTooltipEvent.CheckTrack(this) && InTooltipEvent.Is<FTimingEvent>())
+	{
+		const FTimingEvent& TooltipEvent = InTooltipEvent.As<FTimingEvent>();
+
+		auto MatchEvent = [&TooltipEvent](double InStartTime, double InEndTime, uint32 InDepth)
+		{
+			return InDepth == TooltipEvent.GetDepth()
+				&& InStartTime == TooltipEvent.GetStartTime()
+				&& InEndTime == TooltipEvent.GetEndTime();
+		};
+
+		FTimingEventSearchParameters SearchParameters(TooltipEvent.GetStartTime(), TooltipEvent.GetEndTime(), ETimingEventSearchFlags::StopAtFirstMatch, MatchEvent);
+		FindLoadTimeProfilerCpuEvent(SearchParameters, [this, &InOutTooltip, &TooltipEvent](double InFoundStartTime, double InFoundEndTime, uint32 InFoundDepth, const TraceServices::FLoadTimeProfilerCpuEvent& InFoundEvent)
+		{
+			InOutTooltip.ResetContent();
+
+			InOutTooltip.AddTitle(SharedState.GetEventName(TooltipEvent.GetDepth(), InFoundEvent));
+
+			const TraceServices::FPackageExportInfo* Export = InFoundEvent.Export;
+			const TraceServices::FPackageInfo* Package = InFoundEvent.Export ? InFoundEvent.Export->Package : InFoundEvent.Package;
+
+			InOutTooltip.AddNameValueTextLine(TEXT("Duration:"), UE::Insights::FormatTimeAuto(TooltipEvent.GetDuration()));
+			InOutTooltip.AddNameValueTextLine(TEXT("Depth:"), FString::Printf(TEXT("%d"), TooltipEvent.GetDepth()));
+
+			if (Package)
+			{
+				InOutTooltip.AddNameValueTextLine(TEXT("Package Name:"), Package->Name);
+				InOutTooltip.AddNameValueTextLine(TEXT("Header Size:"), FString::Printf(TEXT("%s bytes"), *FText::AsNumber(Package->Summary.TotalHeaderSize).ToString()));
+				InOutTooltip.AddNameValueTextLine(TEXT("Package Summary:"), FString::Printf(TEXT("%d imports, %d exports"), Package->Summary.ImportCount, Package->Summary.ExportCount));
+				InOutTooltip.AddNameValueTextLine(TEXT("Request Priority:"), FString::Printf(TEXT("%d"), Package->Summary.Priority));
+				if (!Export)
+				{
+					InOutTooltip.AddNameValueTextLine(TEXT("Event:"), TEXT("ProcessPackageSummary"));
+				}
+			}
+
+			if (Export)
+			{
+				InOutTooltip.AddNameValueTextLine(TEXT("Event:"), FString::Printf(TEXT("%s"), TraceServices::GetLoadTimeProfilerObjectEventTypeString(InFoundEvent.EventType)));
+				InOutTooltip.AddNameValueTextLine(TEXT("Export Class:"), Export->Class ? Export->Class->Name : TEXT("N/A"));
+				InOutTooltip.AddNameValueTextLine(TEXT("Serial Size:"), FString::Printf(TEXT("%s bytes"), *FText::AsNumber(Export->SerialSize).ToString()));
+			}
+
+			InOutTooltip.UpdateLayout();
+		});
+	}
+}
+
+////////////////////////////////////////////////////////////////////////////////////////////////////
+
+const TSharedPtr<const ITimingEvent> FLoadingTimingTrack::SearchEvent(const FTimingEventSearchParameters& InSearchParameters) const
+{
+	TSharedPtr<const ITimingEvent> FoundEvent;
+
+	FindLoadTimeProfilerCpuEvent(InSearchParameters, [this, &FoundEvent](double InFoundStartTime, double InFoundEndTime, uint32 InFoundDepth, const TraceServices::FLoadTimeProfilerCpuEvent& InFoundEvent)
+	{
+		FoundEvent = MakeShared<FTimingEvent>(SharedThis(this), InFoundStartTime, InFoundEndTime, InFoundDepth);
+	});
+
+	return FoundEvent;
+}
+
+////////////////////////////////////////////////////////////////////////////////////////////////////
+
+bool FLoadingTimingTrack::FindLoadTimeProfilerCpuEvent(const FTimingEventSearchParameters& InParameters, TFunctionRef<void(double, double, uint32, const TraceServices::FLoadTimeProfilerCpuEvent&)> InFoundPredicate) const
+{
+	using EFilterField = UE::Insights::EFilterField;
+	UE::Insights::FFilterContext FilterConfiguratorContext;
+	FilterConfiguratorContext.SetReturnValueForUnsetFilters(false);
+	FilterConfiguratorContext.AddFilterData<double>(static_cast<int32>(EFilterField::StartTime), 0.0f);
+	FilterConfiguratorContext.AddFilterData<double>(static_cast<int32>(EFilterField::EndTime), 0.0f);
+	FilterConfiguratorContext.AddFilterData<double>(static_cast<int32>(EFilterField::Duration), 0.0f);
+	FilterConfiguratorContext.AddFilterData<FString>(static_cast<int32>(EFilterField::TrackName), this->GetName());
+
+	return TTimingEventSearch<TraceServices::FLoadTimeProfilerCpuEvent>::Search(
+		InParameters,
+
+		[this](TTimingEventSearch<TraceServices::FLoadTimeProfilerCpuEvent>::FContext& InContext)
+		{
+			TSharedPtr<const TraceServices::IAnalysisSession> Session = FInsightsManager::Get()->GetSession();
+			if (Session.IsValid())
+			{
+				TraceServices::FAnalysisSessionReadScope SessionReadScope(*Session.Get());
+
+				if (TraceServices::ReadLoadTimeProfilerProvider(*Session.Get()))
+				{
+					const TraceServices::ILoadTimeProfilerProvider& LoadTimeProfilerProvider = *TraceServices::ReadLoadTimeProfilerProvider(*Session.Get());
+
+					LoadTimeProfilerProvider.ReadTimeline(TimelineIndex, [&InContext](const TraceServices::ILoadTimeProfilerProvider::CpuTimeline& Timeline)
+					{
+						auto Callback = [&InContext](double EventStartTime, double EventEndTime, uint32 EventDepth, const TraceServices::FLoadTimeProfilerCpuEvent& Event)
+						{
+							if (Event.Package)
+							{
+								InContext.Check(EventStartTime, EventEndTime, EventDepth, Event);
+								return InContext.ShouldContinueSearching() ? TraceServices::EEventEnumerate::Continue : TraceServices::EEventEnumerate::Stop;
+							}
+							else
+							{
+								return TraceServices::EEventEnumerate::Continue;
+							}
+						};
+
+						if (InContext.GetParameters().SearchDirection == FTimingEventSearchParameters::ESearchDirection::Forward)
+						{
+							Timeline.EnumerateEvents(InContext.GetParameters().StartTime, InContext.GetParameters().EndTime, Callback);
+						}
+						else
+						{
+							Timeline.EnumerateEventsBackwards(InContext.GetParameters().EndTime, InContext.GetParameters().StartTime, Callback);
+						}
+					});
+				}
+			}
+		},
+
+		[&FilterConfiguratorContext, &InParameters](double EventStartTime, double EventEndTime, uint32 EventDepth, const TraceServices::FLoadTimeProfilerCpuEvent& Event)
+		{
+			if (!InParameters.FilterExecutor.IsValid())
+			{
+				return true;
+			}
+
+			TSharedPtr<const TraceServices::IAnalysisSession> Session = FInsightsManager::Get()->GetSession();
+			if (Session.IsValid())
+			{
+				TraceServices::FAnalysisSessionReadScope SessionReadScope(*Session.Get());
+
+				if (TraceServices::ReadLoadTimeProfilerProvider(*Session.Get()))
+				{
+					FilterConfiguratorContext.SetFilterData<double>(static_cast<int32>(EFilterField::StartTime), EventStartTime);
+					FilterConfiguratorContext.SetFilterData<double>(static_cast<int32>(EFilterField::EndTime), EventEndTime);
+					FilterConfiguratorContext.SetFilterData<double>(static_cast<int32>(EFilterField::Duration), EventEndTime - EventStartTime);
+
+					return InParameters.FilterExecutor->ApplyFilters(FilterConfiguratorContext);
+				}
+			}
+
+			return false;
+		},
+
+		[&InFoundPredicate](double InFoundStartTime, double InFoundEndTime, uint32 InFoundDepth, const TraceServices::FLoadTimeProfilerCpuEvent& InEvent)
+		{
+			InFoundPredicate(InFoundStartTime, InFoundEndTime, InFoundDepth, InEvent);
+		},
+
+		TTimingEventSearch<TraceServices::FLoadTimeProfilerCpuEvent>::NoMatch);
+}
+
+////////////////////////////////////////////////////////////////////////////////////////////////////
+
+bool FLoadingTimingTrack::HasCustomFilter() const
+{
+	return FilterConfigurator.IsValid() && !FilterConfigurator->IsEmpty();
+}
+
+////////////////////////////////////////////////////////////////////////////////////////////////////
+
+void FLoadingTimingTrack::SetFilterConfigurator(TSharedPtr<UE::Insights::FFilterConfigurator> InFilterConfigurator)
+{
+	if (FilterConfigurator != InFilterConfigurator)
+	{
+		FilterConfigurator = InFilterConfigurator;
+		SetDirtyFlag();
+	}
+}
+
+////////////////////////////////////////////////////////////////////////////////////////////////////
+
+} // namespace UE::Insights::LoadingProfiler
+
+#undef LOCTEXT_NAMESPACE
diff --git a/Engine/Source/Developer/TraceInsights/Private/Insights/MemoryProfiler/Widgets/SMemInvestigationView.cpp b/Engine/Source/Developer/TraceInsights/Private/Insights/MemoryProfiler/Widgets/SMemInvestigationView.cpp
--- a/Engine/Source/Developer/TraceInsights/Private/Insights/MemoryProfiler/Widgets/SMemInvestigationView.cpp
+++ b/Engine/Source/Developer/TraceInsights/Private/Insights/MemoryProfiler/Widgets/SMemInvestigationView.cpp
@@ -20,7 +20,7 @@
 #include "Insights/MemoryProfiler/ViewModels/MemorySharedState.h"
 #include "Insights/MemoryProfiler/Widgets/SMemAllocTableTreeView.h"
 #include "Insights/MemoryProfiler/Widgets/SMemoryProfilerWindow.h"
-#include "Insights/ViewModels/TimeRulerTrack.h"
+#include "Insights/TimingProfiler/ViewModels/TimeMarker.h"
 #include "Insights/Widgets/STimingView.h"
 
 #define LOCTEXT_NAMESPACE "UE::Insights::MemoryProfiler::SMemInvestigationView"
diff --git a/Engine/Source/Developer/TraceInsights/Private/Insights/MemoryProfiler/Widgets/SMemoryProfilerWindow.cpp b/Engine/Source/Developer/TraceInsights/Private/Insights/MemoryProfiler/Widgets/SMemoryProfilerWindow.cpp
--- a/Engine/Source/Developer/TraceInsights/Private/Insights/MemoryProfiler/Widgets/SMemoryProfilerWindow.cpp
+++ b/Engine/Source/Developer/TraceInsights/Private/Insights/MemoryProfiler/Widgets/SMemoryProfilerWindow.cpp
@@ -17,7 +17,8 @@
 #include "Insights/MemoryProfiler/Widgets/SMemInvestigationView.h"
 #include "Insights/MemoryProfiler/Widgets/SMemoryProfilerToolbar.h"
 #include "Insights/MemoryProfiler/Widgets/SMemTagTreeView.h"
-#include "Insights/ViewModels/TimeRulerTrack.h"
+#include "Insights/TimingProfiler/Tracks/TimeRulerTrack.h"
+#include "Insights/TimingProfiler/ViewModels/TimeMarker.h"
 #include "Insights/Widgets/SModulesView.h"
 #include "Insights/Widgets/STimingView.h"
 
diff --git a/Engine/Source/Developer/TraceInsights/Private/Insights/TaskGraphProfiler/TaskGraphProfilerManager.cpp b/Engine/Source/Developer/TraceInsights/Private/Insights/TaskGraphProfiler/TaskGraphProfilerManager.cpp
--- a/Engine/Source/Developer/TraceInsights/Private/Insights/TaskGraphProfiler/TaskGraphProfilerManager.cpp
+++ b/Engine/Source/Developer/TraceInsights/Private/Insights/TaskGraphProfiler/TaskGraphProfilerManager.cpp
@@ -18,10 +18,11 @@
 #include "Insights/TaskGraphProfiler/ViewModels/TaskTable.h"
 #include "Insights/TaskGraphProfiler/ViewModels/TaskTimingTrack.h"
 #include "Insights/TaskGraphProfiler/Widgets/STaskTableTreeView.h"
-#include "Insights/TimingProfilerManager.h"
-#include "Insights/ViewModels/ThreadTimingTrack.h"
+#include "Insights/TimingProfiler/TimingProfilerManager.h"
+#include "Insights/TimingProfiler/Tracks/ThreadTimingTrack.h"
+#include "Insights/TimingProfiler/ViewModels/ThreadTimingSharedState.h"
+#include "Insights/TimingProfiler/Widgets/STimingProfilerWindow.h"
 #include "Insights/ViewModels/ThreadTrackEvent.h"
-#include "Insights/Widgets/STimingProfilerWindow.h"
 #include "Insights/Widgets/STimingView.h"
 
 #define LOCTEXT_NAMESPACE "UE::Insights::TaskGraphProfiler"
@@ -571,14 +572,14 @@ void FTaskGraphProfilerManager::AddRelation(const FThreadTrackEvent* InSelectedE
 		return;
 	}
 
-	TSharedPtr<FThreadTimingSharedState> ThreadSharedState = TimingView->GetThreadTimingSharedState();
+	TSharedPtr<TimingProfiler::FThreadTimingSharedState> ThreadSharedState = TimingView->GetThreadTimingSharedState();
 
 	TUniquePtr<ITimingEventRelation> Relation = MakeUnique<FTaskGraphRelation>(SourceTimestamp, SourceThreadId, TargetTimestamp, TargetThreadId, Type);
 	FTaskGraphRelation* TaskRelationPtr = StaticCast<FTaskGraphRelation*>(Relation.Get());
 
 	if (!TaskRelationPtr->GetSourceTrack().IsValid())
 	{
-		TSharedPtr<const FCpuTimingTrack> Track = ThreadSharedState->GetCpuTrack(TaskRelationPtr->GetSourceThreadId());
+		TSharedPtr<const TimingProfiler::FCpuTimingTrack> Track = ThreadSharedState->GetCpuTrack(TaskRelationPtr->GetSourceThreadId());
 		if (Track.IsValid())
 		{
 			TaskRelationPtr->SetSourceTrack(Track);
@@ -589,7 +590,7 @@ void FTaskGraphProfilerManager::AddRelation(const FThreadTrackEvent* InSelectedE
 	if (!TaskRelationPtr->GetTargetTrack().IsValid())
 
 	{
-		TSharedPtr<const FCpuTimingTrack> Track = ThreadSharedState->GetCpuTrack(TaskRelationPtr->GetTargetThreadId());
+		TSharedPtr<const TimingProfiler::FCpuTimingTrack> Track = ThreadSharedState->GetCpuTrack(TaskRelationPtr->GetTargetThreadId());
 		if (Track.IsValid())
 		{
 			TaskRelationPtr->SetTargetTrack(Track);
@@ -615,7 +616,7 @@ void FTaskGraphProfilerManager::AddRelation(const FThreadTrackEvent* InSelectedE
 
 ////////////////////////////////////////////////////////////////////////////////////////////////////
 
-int32 FTaskGraphProfilerManager::GetRelationDisplayDepth(TSharedPtr<const FThreadTimingTrack> Track, double Time, int32 KnownDepth)
+int32 FTaskGraphProfilerManager::GetRelationDisplayDepth(TSharedPtr<const TimingProfiler::FThreadTimingTrack> Track, double Time, int32 KnownDepth)
 {
 	if (KnownDepth >= 0)
 	{
@@ -659,9 +660,9 @@ int32 FTaskGraphProfilerManager::GetDepthOfTaskExecution(double TaskStartedTime,
 		return Depth;
 	}
 
-	TSharedPtr<FThreadTimingSharedState> ThreadSharedState = TimingView->GetThreadTimingSharedState();
+	TSharedPtr<TimingProfiler::FThreadTimingSharedState> ThreadSharedState = TimingView->GetThreadTimingSharedState();
 
-	TSharedPtr<FCpuTimingTrack> Track = ThreadSharedState->GetCpuTrack(ThreadId);
+	TSharedPtr<TimingProfiler::FCpuTimingTrack> Track = ThreadSharedState->GetCpuTrack(ThreadId);
 
 	if (!Track.IsValid())
 	{
diff --git a/Engine/Source/Developer/TraceInsights/Private/Insights/TaskGraphProfiler/TaskGraphProfilerManager.h b/Engine/Source/Developer/TraceInsights/Private/Insights/TaskGraphProfiler/TaskGraphProfilerManager.h
--- a/Engine/Source/Developer/TraceInsights/Private/Insights/TaskGraphProfiler/TaskGraphProfilerManager.h
+++ b/Engine/Source/Developer/TraceInsights/Private/Insights/TaskGraphProfiler/TaskGraphProfilerManager.h
@@ -2,11 +2,12 @@
 
 #pragma once
 
-#include "CoreMinimal.h"
+#include "CoreTypes.h"
 
 #include "Async/TaskTrace.h"
 #include "Containers/Ticker.h"
 #include "Framework/Commands/UICommandList.h"
+#include "Templates/SharedPointer.h"
 
 // TraceInsights
 #include "Insights/InsightsManager.h"
@@ -18,10 +19,13 @@ namespace TraceServices
 	class ITasksProvider;
 }
 
-class FThreadTimingTrack;
 class FThreadTrackEvent;
 
-namespace UE::Insights::TimingProfiler { class STimingView; }
+namespace UE::Insights::TimingProfiler
+{
+	class FThreadTimingTrack;
+	class STimingView;
+}
 
 namespace UE::Insights::TaskGraphProfiler
 {
@@ -157,7 +161,7 @@ private:
 	double GetRelationsOnCriticalPathDescendingRec(const TraceServices::FTaskInfo* Task, const TraceServices::ITasksProvider* TasksProvider, TArray<FTaskGraphRelation>& Relations);
 
 	void InitializeColorCode();
-	int32 GetRelationDisplayDepth(TSharedPtr<const FThreadTimingTrack> Track, double Time, int32 KnownDepth);
+	int32 GetRelationDisplayDepth(TSharedPtr<const TimingProfiler::FThreadTimingTrack> Track, double Time, int32 KnownDepth);
 
 	void OutputWarnings();
 
diff --git a/Engine/Source/Developer/TraceInsights/Private/Insights/TaskGraphProfiler/ViewModels/TaskGraphRelation.cpp b/Engine/Source/Developer/TraceInsights/Private/Insights/TaskGraphProfiler/ViewModels/TaskGraphRelation.cpp
--- a/Engine/Source/Developer/TraceInsights/Private/Insights/TaskGraphProfiler/ViewModels/TaskGraphRelation.cpp
+++ b/Engine/Source/Developer/TraceInsights/Private/Insights/TaskGraphProfiler/ViewModels/TaskGraphRelation.cpp
@@ -7,7 +7,7 @@
 
 // TraceInsights
 #include "Insights/TaskGraphProfiler/TaskGraphProfilerManager.h"
-#include "Insights/TimingProfilerManager.h"
+#include "Insights/TimingProfiler/TimingProfilerManager.h"
 #include "Insights/ViewModels/TimingTrackViewport.h"
 #include "Insights/ViewModels/TimingViewDrawHelper.h"
 
diff --git a/Engine/Source/Developer/TraceInsights/Private/Insights/TaskGraphProfiler/ViewModels/TaskTimingTrack.cpp b/Engine/Source/Developer/TraceInsights/Private/Insights/TaskGraphProfiler/ViewModels/TaskTimingTrack.cpp
--- a/Engine/Source/Developer/TraceInsights/Private/Insights/TaskGraphProfiler/ViewModels/TaskTimingTrack.cpp
+++ b/Engine/Source/Developer/TraceInsights/Private/Insights/TaskGraphProfiler/ViewModels/TaskTimingTrack.cpp
@@ -19,13 +19,13 @@
 #include "Insights/TaskGraphProfiler/TaskGraphProfilerManager.h"
 #include "Insights/TaskGraphProfiler/ViewModels/TaskGraphRelation.h"
 #include "Insights/TaskGraphProfiler/ViewModels/TaskTrackEvent.h"
-#include "Insights/TimingProfilerManager.h"
-#include "Insights/ViewModels/ThreadTimingTrack.h"
+#include "Insights/TimingProfiler/TimingProfilerManager.h"
+#include "Insights/TimingProfiler/Tracks/ThreadTimingTrack.h"
+#include "Insights/TimingProfiler/Widgets/STimingProfilerWindow.h"
 #include "Insights/ViewModels/ThreadTrackEvent.h"
 #include "Insights/ViewModels/TimingTrackViewport.h"
 #include "Insights/ViewModels/TimingViewDrawHelper.h"
 #include "Insights/ViewModels/TooltipDrawState.h"
-#include "Insights/Widgets/STimingProfilerWindow.h"
 #include "Insights/Widgets/STimingView.h"
 
 #define LOCTEXT_NAMESPACE "UE::Insights::TaskGraphProfiler"
@@ -868,7 +868,7 @@ void FTaskTimingTrack::OnTimingEventSelected(TSharedPtr<const ITimingEvent> InSe
 	const FThreadTrackEvent &ThreadEvent = InSelectedEvent->As<FThreadTrackEvent>();
 	GetEventRelations(ThreadEvent);
 
-	uint32 ThreadId = StaticCastSharedRef<const FThreadTimingTrack>(ThreadEvent.GetTrack())->GetThreadId();
+	uint32 ThreadId = StaticCastSharedRef<const TimingProfiler::FThreadTimingTrack>(ThreadEvent.GetTrack())->GetThreadId();
 	const TraceServices::FTaskInfo* Task = TasksProvider->TryGetTask(ThreadId, ThreadEvent.GetStartTime());
 
 	if (Task != nullptr)
@@ -1117,7 +1117,7 @@ void FTaskTimingTrack::GetEventRelations(const FThreadTrackEvent& InSelectedEven
 			return;
 		}
 
-		TSharedRef<const FThreadTimingTrack> EventTrack = StaticCastSharedRef<const FThreadTimingTrack>(InSelectedEvent.GetTrack());
+		TSharedRef<const TimingProfiler::FThreadTimingTrack> EventTrack = StaticCastSharedRef<const TimingProfiler::FThreadTimingTrack>(InSelectedEvent.GetTrack());
 		uint32 ThreadId = EventTrack->GetThreadId();
 
 		FTaskGraphProfilerManager::Get()->ShowTaskRelations(&InSelectedEvent, ThreadId);
diff --git a/Engine/Source/Developer/TraceInsights/Private/Insights/TaskGraphProfiler/Widgets/STaskTableTreeView.cpp b/Engine/Source/Developer/TraceInsights/Private/Insights/TaskGraphProfiler/Widgets/STaskTableTreeView.cpp
--- a/Engine/Source/Developer/TraceInsights/Private/Insights/TaskGraphProfiler/Widgets/STaskTableTreeView.cpp
+++ b/Engine/Source/Developer/TraceInsights/Private/Insights/TaskGraphProfiler/Widgets/STaskTableTreeView.cpp
@@ -25,9 +25,10 @@
 #include "Insights/TaskGraphProfiler/ViewModels/TaskGraphRelation.h"
 #include "Insights/TaskGraphProfiler/ViewModels/TaskNode.h"
 #include "Insights/TaskGraphProfiler/ViewModels/TaskTimingTrack.h"
-#include "Insights/TimingProfilerManager.h"
-#include "Insights/ViewModels/ThreadTimingTrack.h"
-#include "Insights/Widgets/STimingProfilerWindow.h"
+#include "Insights/TimingProfiler/TimingProfilerManager.h"
+#include "Insights/TimingProfiler/Tracks/ThreadTimingTrack.h"
+#include "Insights/TimingProfiler/ViewModels/ThreadTimingSharedState.h"
+#include "Insights/TimingProfiler/Widgets/STimingProfilerWindow.h"
 #include "Insights/Widgets/STimingView.h"
 
 #include <limits>
@@ -750,13 +751,13 @@ void STaskTableTreeView::ContextMenu_GoToTask_Execute()
 		TaskSharedState->SetTaskId(TaskEntry->GetId());
 	}
 
-	TSharedPtr<FThreadTimingSharedState> ThreadTimingState = TimingView->GetThreadTimingSharedState();
+	TSharedPtr<TimingProfiler::FThreadTimingSharedState> ThreadTimingState = TimingView->GetThreadTimingSharedState();
 	if (!ThreadTimingState.IsValid())
 	{
 		return;
 	}
 
-	TSharedPtr<FCpuTimingTrack> Track = ThreadTimingState->GetCpuTrack(TaskEntry->GetStartedThreadId());
+	TSharedPtr<TimingProfiler::FCpuTimingTrack> Track = ThreadTimingState->GetCpuTrack(TaskEntry->GetStartedThreadId());
 	if (!Track.IsValid())
 	{
 		return;
diff --git a/Engine/Source/Developer/TraceInsights/Private/Insights/Tests/FunctionalTests/UITests.cpp b/Engine/Source/Developer/TraceInsights/Private/Insights/Tests/FunctionalTests/UITests.cpp
--- a/Engine/Source/Developer/TraceInsights/Private/Insights/Tests/FunctionalTests/UITests.cpp
+++ b/Engine/Source/Developer/TraceInsights/Private/Insights/Tests/FunctionalTests/UITests.cpp
@@ -15,8 +15,8 @@
 // TraceInsights
 #include "Insights/MemoryProfiler/ViewModels/MemAllocFilterValueConverter.h"
 #include "Insights/TaskGraphProfiler/ViewModels/TaskTimingTrack.h"
-#include "Insights/TimingProfilerManager.h"
-#include "Insights/ViewModels/TimeRulerTrack.h"
+#include "Insights/TimingProfiler/TimingProfilerManager.h"
+#include "Insights/TimingProfiler/Tracks/TimeRulerTrack.h"
 #include "Insights/Widgets/STimingView.h"
 
 DECLARE_LOG_CATEGORY_EXTERN(UITests, Log, All);
diff --git a/Engine/Source/Developer/TraceInsights/Private/Insights/TimingProfiler/TimingProfilerCommands.cpp b/Engine/Source/Developer/TraceInsights/Private/Insights/TimingProfiler/TimingProfilerCommands.cpp
new file mode 100644
--- /dev/null
+++ b/Engine/Source/Developer/TraceInsights/Private/Insights/TimingProfiler/TimingProfilerCommands.cpp
@@ -0,0 +1,169 @@
+// Copyright Epic Games, Inc. All Rights Reserved.
+
+#include "Insights/TimingProfiler/TimingProfilerCommands.h"
+
+#include "DesktopPlatformModule.h"
+#include "Framework/Application/SlateApplication.h"
+#include "Framework/MultiBox/MultiBoxBuilder.h"
+
+// TraceInsights
+#include "Insights/InsightsManager.h"
+#include "Insights/InsightsStyle.h"
+#include "Insights/TimingProfiler/TimingProfilerManager.h"
+
+////////////////////////////////////////////////////////////////////////////////////////////////////
+
+#define LOCTEXT_NAMESPACE "UE::Insights::TimingProfiler"
+
+namespace UE::Insights::TimingProfiler
+{
+
+////////////////////////////////////////////////////////////////////////////////////////////////////
+// FTimingProfilerMenuBuilder
+////////////////////////////////////////////////////////////////////////////////////////////////////
+
+void FTimingProfilerMenuBuilder::AddMenuEntry(FMenuBuilder& MenuBuilder, const TSharedPtr<FUICommandInfo>& UICommandInfo, const FUIAction& UIAction)
+{
+	MenuBuilder.AddMenuEntry
+	(
+		UICommandInfo->GetLabel(),
+		UICommandInfo->GetDescription(),
+		UICommandInfo->GetIcon(),
+		UIAction,
+		NAME_None,
+		UICommandInfo->GetUserInterfaceType()
+	);
+}
+
+////////////////////////////////////////////////////////////////////////////////////////////////////
+// FTimingProfilerCommands
+////////////////////////////////////////////////////////////////////////////////////////////////////
+
+FTimingProfilerCommands::FTimingProfilerCommands()
+: TCommands<FTimingProfilerCommands>(
+	TEXT("TimingProfilerCommands"),
+	NSLOCTEXT("Contexts", "TimingProfilerCommands", "Insights - Timing Insights"),
+	NAME_None,
+	FInsightsStyle::GetStyleSetName())
+{
+}
+
+////////////////////////////////////////////////////////////////////////////////////////////////////
+
+// UI_COMMAND takes long for the compiler to optimize
+UE_DISABLE_OPTIMIZATION_SHIP
+void FTimingProfilerCommands::RegisterCommands()
+{
+	UI_COMMAND(ToggleFramesTrackVisibility, "Frames", "Toggles the visibility of the Frames track.", EUserInterfaceActionType::ToggleButton, FInputChord());
+	UI_COMMAND(ToggleTimingViewVisibility, "Timing", "Toggles the visibility of the main Timing view.", EUserInterfaceActionType::ToggleButton, FInputChord());
+	UI_COMMAND(ToggleTimersViewVisibility, "Timers", "Toggles the visibility of the Timers view.", EUserInterfaceActionType::ToggleButton, FInputChord());
+	UI_COMMAND(ToggleCallersTreeViewVisibility, "Callers", "Toggles the visibility of the Callers tree view.", EUserInterfaceActionType::ToggleButton, FInputChord());
+	UI_COMMAND(ToggleCalleesTreeViewVisibility, "Callees", "Toggles the visibility of the Callees tree view.", EUserInterfaceActionType::ToggleButton,FInputChord());
+	UI_COMMAND(ToggleStatsCountersViewVisibility, "Counters", "Toggles the visibility of the Counters view.", EUserInterfaceActionType::ToggleButton, FInputChord());
+	UI_COMMAND(ToggleLogViewVisibility, "Log", "Toggles the visibility of the Log view.", EUserInterfaceActionType::ToggleButton,FInputChord());
+}
+UE_ENABLE_OPTIMIZATION_SHIP
+
+////////////////////////////////////////////////////////////////////////////////////////////////////
+// FTimingViewCommands
+////////////////////////////////////////////////////////////////////////////////////////////////////
+
+FTimingViewCommands::FTimingViewCommands()
+: TCommands<FTimingViewCommands>(
+	TEXT("TimingViewCommands"),
+	NSLOCTEXT("Contexts", "TimingViewCommands", "Insights - Timing View"),
+	NAME_None,
+	FInsightsStyle::GetStyleSetName())
+{
+}
+
+////////////////////////////////////////////////////////////////////////////////////////////////////
+
+// UI_COMMAND takes long for the compiler to optimize
+UE_DISABLE_OPTIMIZATION_SHIP
+void FTimingViewCommands::RegisterCommands()
+{
+	UI_COMMAND(AutoHideEmptyTracks,
+		"Auto Hide Empty Tracks",
+		"Auto hide empty tracks (ex.: ones without timing events in the current viewport).\nThis option is persistent to the UnrealInsightsSettings.ini file.",
+		EUserInterfaceActionType::ToggleButton,
+		FInputChord(EKeys::V));
+
+	UI_COMMAND(PanningOnScreenEdges,
+		"Allow Panning on Screen Edges",
+		"If enabled, the panning is allowed to continue when the mouse cursor reaches the edges of the screen.\nThis option is persistent to the UnrealInsightsSettings.ini file.",
+		EUserInterfaceActionType::ToggleButton,
+		FInputChord());
+
+	UI_COMMAND(ToggleCompactMode,
+		"Compact Mode",
+		"Toggle compact mode for supporting tracks.\n(ex.: the timing tracks will be displayed with reduced height)",
+		EUserInterfaceActionType::ToggleButton,
+		FInputChord(EKeys::C));
+
+	UI_COMMAND(ShowMainGraphTrack,
+		"Graph Track",
+		"Shows/hides the main Graph track.",
+		EUserInterfaceActionType::ToggleButton,
+		FInputChord(EKeys::G));
+
+	UI_COMMAND(QuickFind,
+		"Quick Find...",
+		"Quick find or filter events in the timing view.",
+		EUserInterfaceActionType::Button,
+		FInputChord(EModifierKey::Control, EKeys::F));
+}
+UE_ENABLE_OPTIMIZATION_SHIP
+
+////////////////////////////////////////////////////////////////////////////////////////////////////
+// Toggle Commands
+////////////////////////////////////////////////////////////////////////////////////////////////////
+
+#define IMPLEMENT_TOGGLE_COMMAND(CmdName, IsEnabled, SetIsEnabled) \
+	\
+	void FTimingProfilerActionManager::Map_##CmdName##_Global()\
+	{\
+		This->CommandList->MapAction(This->GetCommands().CmdName, CmdName##_Custom());\
+	}\
+	\
+	const FUIAction FTimingProfilerActionManager::CmdName##_Custom() \
+	{\
+		FUIAction UIAction;\
+		UIAction.ExecuteAction = FExecuteAction::CreateRaw(this, &FTimingProfilerActionManager::CmdName##_Execute);\
+		UIAction.CanExecuteAction = FCanExecuteAction::CreateRaw(this, &FTimingProfilerActionManager::CmdName##_CanExecute);\
+		UIAction.GetActionCheckState = FGetActionCheckState::CreateRaw(this, &FTimingProfilerActionManager::CmdName##_GetCheckState);\
+		return UIAction;\
+	}\
+	\
+	void FTimingProfilerActionManager::CmdName##_Execute()\
+	{\
+		const bool b##IsEnabled = !This->IsEnabled();\
+		This->SetIsEnabled(b##IsEnabled);\
+	}\
+	\
+	bool FTimingProfilerActionManager::CmdName##_CanExecute() const\
+	{\
+		return FInsightsManager::Get()->GetSession().IsValid();\
+	}\
+	\
+	ECheckBoxState FTimingProfilerActionManager::CmdName##_GetCheckState() const\
+	{\
+		const bool b##IsEnabled = This->IsEnabled();\
+		return b##IsEnabled ? ECheckBoxState::Checked : ECheckBoxState::Unchecked;\
+	}
+
+IMPLEMENT_TOGGLE_COMMAND(ToggleFramesTrackVisibility, IsFramesTrackVisible, ShowHideFramesTrack)
+IMPLEMENT_TOGGLE_COMMAND(ToggleTimingViewVisibility, IsTimingViewVisible, ShowHideTimingView)
+IMPLEMENT_TOGGLE_COMMAND(ToggleTimersViewVisibility, IsTimersViewVisible, ShowHideTimersView)
+IMPLEMENT_TOGGLE_COMMAND(ToggleCallersTreeViewVisibility, IsCallersTreeViewVisible, ShowHideCallersTreeView)
+IMPLEMENT_TOGGLE_COMMAND(ToggleCalleesTreeViewVisibility, IsCalleesTreeViewVisible, ShowHideCalleesTreeView)
+IMPLEMENT_TOGGLE_COMMAND(ToggleStatsCountersViewVisibility, IsStatsCountersViewVisible, ShowHideStatsCountersView)
+IMPLEMENT_TOGGLE_COMMAND(ToggleLogViewVisibility, IsLogViewVisible, ShowHideLogView)
+
+#undef IMPLEMENT_TOGGLE_COMMAND
+
+////////////////////////////////////////////////////////////////////////////////////////////////////
+
+} // namespace UE::Insights::TimingProfiler
+
+#undef LOCTEXT_NAMESPACE
diff --git a/Engine/Source/Developer/TraceInsights/Private/Insights/TimingProfiler/TimingProfilerCommands.h b/Engine/Source/Developer/TraceInsights/Private/Insights/TimingProfiler/TimingProfilerCommands.h
new file mode 100644
--- /dev/null
+++ b/Engine/Source/Developer/TraceInsights/Private/Insights/TimingProfiler/TimingProfilerCommands.h
@@ -0,0 +1,157 @@
+// Copyright Epic Games, Inc. All Rights Reserved.
+
+#pragma once
+
+#include "CoreTypes.h"
+
+#include "Framework/Commands/UIAction.h"
+#include "Framework/Commands/Commands.h"
+
+class FMenuBuilder;
+
+namespace UE::Insights::TimingProfiler
+{
+
+////////////////////////////////////////////////////////////////////////////////////////////////////
+/**
+ * Class that holds all profiler commands.
+ */
+class FTimingProfilerCommands : public TCommands<FTimingProfilerCommands>
+{
+public:
+	FTimingProfilerCommands();
+	virtual ~FTimingProfilerCommands() {}
+	virtual void RegisterCommands() override;
+
+public:
+	//////////////////////////////////////////////////
+	// Global commands need to implement following method:
+	//     void Map_<CommandName>_Global();
+	// Custom commands needs to implement also the following method:
+	//     const FUIAction <CommandName>_Custom(...) const;
+	//////////////////////////////////////////////////
+
+	/** Toggles visibility for the Frames Track. Global and custom command. */
+	TSharedPtr<FUICommandInfo> ToggleFramesTrackVisibility;
+
+	/** Toggles visibility for the Timing View. Global and custom command. */
+	TSharedPtr<FUICommandInfo> ToggleTimingViewVisibility;
+
+	/** Toggles visibility for the Timers View. Global and custom command. */
+	TSharedPtr<FUICommandInfo> ToggleTimersViewVisibility;
+
+	/** Toggles visibility for the Callers Tree View. Global and custom command. */
+	TSharedPtr<FUICommandInfo> ToggleCallersTreeViewVisibility;
+
+	/** Toggles visibility for the Callees Tree View. Global and custom command. */
+	TSharedPtr<FUICommandInfo> ToggleCalleesTreeViewVisibility;
+
+	/** Toggles visibility for the Stats Counters View. Global and custom command. */
+	TSharedPtr<FUICommandInfo> ToggleStatsCountersViewVisibility;
+
+	/** Toggles visibility for the Log View. Global and custom command. */
+	TSharedPtr<FUICommandInfo> ToggleLogViewVisibility;
+};
+
+////////////////////////////////////////////////////////////////////////////////////////////////////
+
+class FTimingViewCommands : public TCommands<FTimingViewCommands>
+{
+public:
+	FTimingViewCommands();
+	virtual ~FTimingViewCommands() {}
+	virtual void RegisterCommands() override;
+
+public:
+	/** Toggles visibility of empty tracks. */
+	TSharedPtr<FUICommandInfo> AutoHideEmptyTracks;
+
+	/** Toggles "panning on screen edges". */
+	TSharedPtr<FUICommandInfo> PanningOnScreenEdges;
+
+	/** Toggles 'compact mode' for timing tracks. */
+	TSharedPtr<FUICommandInfo> ToggleCompactMode;
+
+	/** Toggles visibility for Main Graph track. */
+	TSharedPtr<FUICommandInfo> ShowMainGraphTrack;
+
+	/** Opens the Quick Find widget. */
+	TSharedPtr<FUICommandInfo> QuickFind;
+};
+
+////////////////////////////////////////////////////////////////////////////////////////////////////
+/**
+ * Menu builder. Helper class for adding a customized menu entry using the global UI command info.
+ */
+class FTimingProfilerMenuBuilder
+{
+public:
+	/**
+	 * Helper method for adding a customized menu entry using the global UI command info.
+	 * FUICommandInfo cannot be executed with custom parameters, so we need to create a custom FUIAction,
+	 * but sometime we have global and local version for the UI command, so reuse data from the global UI command info.
+	 * Ex:
+	 *     SessionInstance_ToggleCapture          - Global version will toggle capture process for all active session instances
+	 *     SessionInstance_ToggleCapture_OneParam - Local version will toggle capture process only for the specified session instance
+	 *
+	 * @param MenuBuilder The menu to add items to
+	 * @param FUICommandInfo A shared pointer to the UI command info
+	 * @param UIAction Customized version of the UI command info stored in an UI action
+	 */
+	static void AddMenuEntry(FMenuBuilder& MenuBuilder, const TSharedPtr<FUICommandInfo>& UICommandInfo, const FUIAction& UIAction);
+};
+
+////////////////////////////////////////////////////////////////////////////////////////////////////
+/**
+ * Class that provides helper functions for the commands to avoid cluttering profiler manager with many small functions.
+ * Can't contain any variables. Directly operates on the profiler manager instance.
+ */
+class FTimingProfilerActionManager
+{
+	friend class FTimingProfilerManager;
+
+private:
+	/** Private constructor. */
+	FTimingProfilerActionManager(class FTimingProfilerManager* Instance)
+		: This(Instance)
+	{}
+
+	//////////////////////////////////////////////////
+	// Toggle Commands
+
+#define DECLARE_TOGGLE_COMMAND(CmdName)\
+public:\
+	void Map_##CmdName##_Global(); /**< Maps UI command info CmdName with the specified UI command list. */\
+	const FUIAction CmdName##_Custom(); /**< UI action for CmdName command. */\
+protected:\
+	void CmdName##_Execute(); /**< Handles FExecuteAction for CmdName. */\
+	bool CmdName##_CanExecute() const; /**< Handles FCanExecuteAction for CmdName. */\
+	ECheckBoxState CmdName##_GetCheckState() const; /**< Handles FGetActionCheckState for CmdName. */
+
+	DECLARE_TOGGLE_COMMAND(ToggleFramesTrackVisibility)
+	DECLARE_TOGGLE_COMMAND(ToggleTimingViewVisibility)
+	DECLARE_TOGGLE_COMMAND(ToggleTimersViewVisibility)
+	DECLARE_TOGGLE_COMMAND(ToggleCallersTreeViewVisibility)
+	DECLARE_TOGGLE_COMMAND(ToggleCalleesTreeViewVisibility)
+	DECLARE_TOGGLE_COMMAND(ToggleStatsCountersViewVisibility)
+	DECLARE_TOGGLE_COMMAND(ToggleLogViewVisibility)
+#undef DECLARE_TOGGLE_COMMAND
+
+	//////////////////////////////////////////////////
+	// OpenSettings
+
+public:
+	void Map_OpenSettings_Global(); /**< Maps UI command info OpenSettings with the specified UI command list. */
+	const FUIAction OpenSettings_Custom() const; /**< UI action for OpenSettings command. */
+protected:
+	void OpenSettings_Execute(); /**< Handles FExecuteAction for OpenSettings. */
+	bool OpenSettings_CanExecute() const; /**< Handles FCanExecuteAction for OpenSettings. */
+
+	//////////////////////////////////////////////////
+
+protected:
+	/** Reference to the global instance of the profiler manager. */
+	class FTimingProfilerManager* This;
+};
+
+} // namespace UE::Insights::TimingProfiler
