[Insights]

diff --git a/Engine/Source/Developer/TraceInsights/Private/Insights/TimingProfiler/Tracks/ThreadTimingTrack.cpp b/Engine/Source/Developer/TraceInsights/Private/Insights/TimingProfiler/Tracks/ThreadTimingTrack.cpp
new file mode 100644
--- /dev/null
+++ b/Engine/Source/Developer/TraceInsights/Private/Insights/TimingProfiler/Tracks/ThreadTimingTrack.cpp
@@ -0,0 +1,1541 @@
+// Copyright Epic Games, Inc. All Rights Reserved.
+
+#include "ThreadTimingTrack.h"
+
+#include "Async/TaskGraphInterfaces.h" // for ENamedThreads
+#include "CborReader.h"
+#include "Framework/MultiBox/MultiBoxBuilder.h"
+#include "HAL/PlatformApplicationMisc.h"
+#include "Serialization/MemoryReader.h"
+
+// TraceServices
+#include "TraceServices/Containers/Timelines.h"
+#include "TraceServices/Model/AnalysisSession.h"
+#include "TraceServices/Model/TasksProfiler.h"
+#include "TraceServices/Model/Threads.h"
+#include "TraceServices/Model/TimingProfiler.h"
+
+// TraceInsightsCore
+#include "InsightsCore/Common/TimeUtils.h"
+#include "InsightsCore/Filter/ViewModels/FilterConfigurator.h"
+#include "InsightsCore/Filter/ViewModels/Filters.h"
+
+// TraceInsights
+#include "Insights/InsightsManager.h"
+#include "Insights/TimingProfiler/TimingProfilerManager.h"
+#include "Insights/TimingProfiler/ViewModels/ThreadTimingSharedState.h"
+#include "Insights/ViewModels/ThreadTrackEvent.h"
+#include "Insights/ViewModels/TimingEvent.h"
+#include "Insights/ViewModels/TimingEventSearch.h"
+#include "Insights/ViewModels/TimingTrackViewport.h"
+#include "Insights/ViewModels/TimingViewDrawHelper.h"
+#include "Insights/ViewModels/TooltipDrawState.h"
+
+#define LOCTEXT_NAMESPACE "UE::Insights::TimingProfiler::ThreadTiming"
+
+namespace UE::Insights::TimingProfiler
+{
+
+////////////////////////////////////////////////////////////////////////////////////////////////////
+
+static void AppendMetadataToTooltip(FTooltipDrawState& Tooltip, TArrayView<const uint8>& Metadata)
+{
+	FMemoryReaderView MemoryReader(Metadata);
+	FCborReader CborReader(&MemoryReader, ECborEndianness::StandardCompliant);
+	FCborContext Context;
+
+	if (!CborReader.ReadNext(Context) || Context.MajorType() != ECborCode::Map)
+	{
+		return;
+	}
+
+	Tooltip.AddTitle(TEXT("Metadata:"));
+
+	while (true)
+	{
+		// Read key
+		if (!CborReader.ReadNext(Context) || !Context.IsString())
+		{
+			break;
+		}
+
+		FString Key = FString::ConstructFromPtrSize(Context.AsCString(), static_cast<int32>(Context.AsLength()));
+		Key += TEXT(":");
+
+		if (!CborReader.ReadNext(Context))
+		{
+			break;
+		}
+
+		switch (Context.MajorType())
+		{
+		case ECborCode::Int:
+		case ECborCode::Uint:
+			{
+				uint64 Value = Context.AsUInt();
+				FString ValueStr;
+				if (Value > 999'999'999ULL)
+				{
+					ValueStr = FString::Printf(TEXT("0x%llX"), Value);
+				}
+				else
+				{
+					ValueStr = FString::Printf(TEXT("%llu"), Value);
+				}
+				Tooltip.AddNameValueTextLine(Key, ValueStr);
+				continue;
+			}
+
+		case ECborCode::TextString:
+			{
+				FString Value = Context.AsString();
+				Tooltip.AddNameValueTextLine(Key, Value);
+				continue;
+			}
+
+		case ECborCode::ByteString:
+			{
+				FAnsiStringView Value(Context.AsCString(), static_cast<int32>(Context.AsLength()));
+				FString ValueStr(Value);
+				Tooltip.AddNameValueTextLine(Key, ValueStr);
+				continue;
+			}
+		}
+
+		if (Context.RawCode() == (ECborCode::Prim|ECborCode::Value_4Bytes))
+		{
+			float Value = Context.AsFloat();
+			FString ValueStr = FString::Printf(TEXT("%f"), Value);
+			Tooltip.AddNameValueTextLine(Key, ValueStr);
+			continue;
+		}
+
+		if (Context.RawCode() == (ECborCode::Prim|ECborCode::Value_8Bytes))
+		{
+			double Value = Context.AsDouble();
+			FString ValueStr = FString::Printf(TEXT("%g"), Value);
+			Tooltip.AddNameValueTextLine(Key, ValueStr);
+			continue;
+		}
+
+		if (Context.RawCode() == (ECborCode::Prim | ECborCode::False))
+		{
+			Tooltip.AddNameValueTextLine(Key, FString(TEXT("false")));
+			continue;
+		}
+
+		if (Context.RawCode() == (ECborCode::Prim | ECborCode::True))
+		{
+			Tooltip.AddNameValueTextLine(Key, FString(TEXT("true")));
+			continue;
+		}
+
+		if (Context.IsFiniteContainer())
+		{
+			CborReader.SkipContainer(ECborCode::Array);
+		}
+	}
+}
+
+////////////////////////////////////////////////////////////////////////////////////////////////////
+
+static void AppendMetadataToString(FString& Str, TArrayView<const uint8>& Metadata)
+{
+	FMemoryReaderView MemoryReader(Metadata);
+	FCborReader CborReader(&MemoryReader, ECborEndianness::StandardCompliant);
+	FCborContext Context;
+
+	if (!CborReader.ReadNext(Context) || Context.MajorType() != ECborCode::Map)
+	{
+		return;
+	}
+
+	bool bFirst = true;
+
+	while (true)
+	{
+		// Read key
+		if (!CborReader.ReadNext(Context) || !Context.IsString())
+		{
+			break;
+		}
+
+		if (bFirst)
+		{
+			bFirst = false;
+			Str += TEXT(" - ");
+		}
+		else
+		{
+			Str += TEXT(", ");
+		}
+
+		//FString Key(Context.AsLength(), Context.AsCString());
+		//Str += Key;
+		//Str += TEXT(":");
+
+		if (!CborReader.ReadNext(Context))
+		{
+			break;
+		}
+
+		switch (Context.MajorType())
+		{
+		case ECborCode::Int:
+		case ECborCode::Uint:
+			{
+				uint64 Value = Context.AsUInt();
+				if (Value > 999'999'999ULL)
+				{
+					Str += FString::Printf(TEXT("0x%llX"), Value);
+				}
+				else
+				{
+					Str += FString::Printf(TEXT("%llu"), Value);
+				}
+				continue;
+			}
+
+		case ECborCode::TextString:
+			{
+				Str += Context.AsString();
+				continue;
+			}
+
+		case ECborCode::ByteString:
+			{
+				Str.AppendChars(Context.AsCString(), static_cast<int32>(Context.AsLength()));
+				continue;
+			}
+		}
+
+		if (Context.RawCode() == (ECborCode::Prim | ECborCode::Value_4Bytes))
+		{
+			float Value = Context.AsFloat();
+			Str += FString::Printf(TEXT("%f"), Value);
+			continue;
+		}
+
+		if (Context.RawCode() == (ECborCode::Prim | ECborCode::Value_8Bytes))
+		{
+			double Value = Context.AsDouble();
+			Str += FString::Printf(TEXT("%g"), Value);
+			continue;
+		}
+
+		if (Context.RawCode() == (ECborCode::Prim | ECborCode::False))
+		{
+			Str += TEXT("false");
+			continue;
+		}
+
+		if (Context.RawCode() == (ECborCode::Prim | ECborCode::True))
+		{
+			Str += TEXT("true");
+			continue;
+		}
+
+		if (Context.IsFiniteContainer())
+		{
+			CborReader.SkipContainer(ECborCode::Array);
+		}
+	}
+}
+
+////////////////////////////////////////////////////////////////////////////////////////////////////
+
+static void AddTimingEventToBuilder(ITimingEventsTrackDrawStateBuilder& Builder,
+									double EventStartTime, double EventEndTime, uint32 EventDepth,
+									uint32 TimerIndex, const TraceServices::FTimingProfilerTimer* Timer)
+{
+	if (EventDepth >= FTimingProfilerManager::Get()->GetEventDepthLimit())
+	{
+		return;
+	}
+
+	uint32 EventColor;
+	switch (FTimingProfilerManager::Get()->GetColoringMode())
+	{
+		case ETimingEventsColoringMode::ByTimerName:
+			EventColor = FTimingEvent::ComputeEventColor(Timer->Name);
+			break;
+		case ETimingEventsColoringMode::ByTimerId:
+			EventColor = FTimingEvent::ComputeEventColor(Timer->Id);
+			break;
+		case ETimingEventsColoringMode::BySourceFile:
+			EventColor = FTimingEvent::ComputeEventColor(Timer->File);
+			break;
+		case ETimingEventsColoringMode::ByDuration:
+		{
+			const double EventDuration = EventEndTime - EventStartTime;
+			EventColor = (EventDuration >= 0.01)     ? 0xFF883333 : // red:    >= 10ms
+						 (EventDuration >= 0.001)    ? 0xFF998833 : // yellow: [1ms .. 10ms)
+						 (EventDuration >= 0.0001)   ? 0xFF338833 : // green:  [100us .. 1ms)
+						 (EventDuration >= 0.00001)  ? 0xFF338888 : // cyan:   [10us .. 100us)
+						 (EventDuration >= 0.000001) ? 0xFF333388 : // blue:   [1us .. 10us)
+						                               0xFF888888;  // gray:   < 1us
+			break;
+		}
+		default:
+			EventColor = 0xFF000000;
+	}
+
+	Builder.AddEvent(EventStartTime, EventEndTime, EventDepth, EventColor,
+		[TimerIndex, Timer, EventStartTime, EventEndTime](float Width)
+		{
+			FString EventName = Timer->Name;
+
+			const float MinWidth = static_cast<float>(EventName.Len()) * 4.0f + 32.0f;
+			if (Width > MinWidth)
+			{
+				//EventName = TEXT("*") + EventName; // for debugging
+
+				const double Duration = EventEndTime - EventStartTime;
+				FTimingEventsTrackDrawStateBuilder::AppendDurationToEventName(EventName, Duration);
+
+				if (int32(TimerIndex) < 0) // has metadata?
+				{
+					//EventName = TEXT("!") + EventName; // for debugging
+
+					TSharedPtr<const TraceServices::IAnalysisSession> Session = FInsightsManager::Get()->GetSession();
+					check(Session.IsValid());
+
+					//TraceServices::FAnalysisSessionReadScope SessionReadScope(*Session.Get());
+
+					const TraceServices::ITimingProfilerProvider& TimingProfilerProvider = *TraceServices::ReadTimingProfilerProvider(*Session.Get());
+
+					const TraceServices::ITimingProfilerTimerReader* TimerReader;
+					TimingProfilerProvider.ReadTimers([&TimerReader](const TraceServices::ITimingProfilerTimerReader& Out) { TimerReader = &Out; });
+
+					TArrayView<const uint8> Metadata = TimerReader->GetMetadata(TimerIndex);
+					if (Metadata.Num() > 0)
+					{
+						AppendMetadataToString(EventName, Metadata);
+					}
+				}
+			}
+
+			return EventName;
+		});
+}
+
+////////////////////////////////////////////////////////////////////////////////////////////////////
+// FThreadTimingTrack
+////////////////////////////////////////////////////////////////////////////////////////////////////
+
+INSIGHTS_IMPLEMENT_RTTI(FThreadTimingTrack)
+
+////////////////////////////////////////////////////////////////////////////////////////////////////
+
+void FThreadTimingTrack::BuildDrawState(ITimingEventsTrackDrawStateBuilder& Builder, const ITimingTrackUpdateContext& Context)
+{
+	TSharedPtr<const TraceServices::IAnalysisSession> Session = FInsightsManager::Get()->GetSession();
+	if (Session.IsValid() && TraceServices::ReadTimingProfilerProvider(*Session.Get()))
+	{
+		TraceServices::FAnalysisSessionReadScope SessionReadScope(*Session.Get());
+
+		const TraceServices::ITimingProfilerProvider& TimingProfilerProvider = *TraceServices::ReadTimingProfilerProvider(*Session.Get());
+
+		const TraceServices::ITimingProfilerTimerReader* TimerReader;
+		TimingProfilerProvider.ReadTimers([&TimerReader](const TraceServices::ITimingProfilerTimerReader& Out) { TimerReader = &Out; });
+
+		const FTimingTrackViewport& Viewport = Context.GetViewport();
+
+		TimingProfilerProvider.ReadTimeline(TimelineIndex,
+			[&Viewport, this, &Builder, TimerReader](const TraceServices::ITimingProfilerProvider::Timeline& Timeline)
+			{
+				if (FTimingEventsTrack::bUseDownSampling)
+				{
+					const double SecondsPerPixel = 1.0 / Viewport.GetScaleX();
+					Timeline.EnumerateEventsDownSampled(Viewport.GetStartTime(), Viewport.GetEndTime(), SecondsPerPixel,
+						[this, &Builder, TimerReader](double StartTime, double EndTime, uint32 Depth, const TraceServices::FTimingProfilerEvent& Event)
+						{
+							const TraceServices::FTimingProfilerTimer* Timer = TimerReader->GetTimer(Event.TimerIndex);
+							if (ensure(Timer != nullptr))
+							{
+								AddTimingEventToBuilder(Builder, StartTime, EndTime, Depth, Event.TimerIndex, Timer);
+							}
+							else
+							{
+								Builder.AddEvent(StartTime, EndTime, Depth, 0xFF000000, [&Event](float) { return FString::Printf(TEXT("[%u]"), Event.TimerIndex); });
+							}
+							return TraceServices::EEventEnumerate::Continue;
+						});
+				}
+				else
+				{
+					Timeline.EnumerateEvents(Viewport.GetStartTime(), Viewport.GetEndTime(),
+						[this, &Builder, TimerReader](double StartTime, double EndTime, uint32 Depth, const TraceServices::FTimingProfilerEvent& Event)
+						{
+							const TraceServices::FTimingProfilerTimer* Timer = TimerReader->GetTimer(Event.TimerIndex);
+							if (ensure(Timer != nullptr))
+							{
+								AddTimingEventToBuilder(Builder, StartTime, EndTime, Depth, Event.TimerIndex, Timer);
+							}
+							else
+							{
+								Builder.AddEvent(StartTime, EndTime, Depth, 0xFF000000, [&Event](float) { return FString::Printf(TEXT("[%u]"), Event.TimerIndex); });
+							}
+							return TraceServices::EEventEnumerate::Continue;
+						});
+				}
+			});
+	}
+}
+
+////////////////////////////////////////////////////////////////////////////////////////////////////
+
+void FThreadTimingTrack::BuildFilteredDrawState(ITimingEventsTrackDrawStateBuilder& Builder, const ITimingTrackUpdateContext& Context)
+{
+	const TSharedPtr<ITimingEventFilter> EventFilterPtr = Context.GetEventFilter();
+	if (EventFilterPtr.IsValid() && EventFilterPtr->FilterTrack(*this))
+	{
+		bool bFilterOnlyByEventType = false; // this is the most often use case, so the below code tries to optimize it
+		uint64 FilterEventType = 0;
+		if (EventFilterPtr->Is<FTimingEventFilterByEventType>())
+		{
+			bFilterOnlyByEventType = true;
+			const FTimingEventFilterByEventType& EventFilter = EventFilterPtr->As<FTimingEventFilterByEventType>();
+			FilterEventType = EventFilter.GetEventType();
+		}
+
+		TSharedPtr<const TraceServices::IAnalysisSession> Session = FInsightsManager::Get()->GetSession();
+		if (Session.IsValid() && TraceServices::ReadTimingProfilerProvider(*Session.Get()))
+		{
+			TraceServices::FAnalysisSessionReadScope SessionReadScope(*Session.Get());
+
+			const TraceServices::ITimingProfilerProvider& TimingProfilerProvider = *TraceServices::ReadTimingProfilerProvider(*Session.Get());
+
+			const TraceServices::ITimingProfilerTimerReader* TimerReader;
+			TimingProfilerProvider.ReadTimers([&TimerReader](const TraceServices::ITimingProfilerTimerReader& Out) { TimerReader = &Out; });
+
+			const FTimingTrackViewport& Viewport = Context.GetViewport();
+
+			if (bFilterOnlyByEventType)
+			{
+				//TODO: Add a setting to switch this on/off
+				if (true)
+				{
+					TimingProfilerProvider.ReadTimeline(TimelineIndex,
+						[&Viewport, this, &Builder, TimerReader, FilterEventType](const TraceServices::ITimingProfilerProvider::Timeline& Timeline)
+						{
+							TArray<TArray<FPendingEventInfo>> FilteredEvents;
+
+							TraceServices::ITimeline<TraceServices::FTimingProfilerEvent>::EnumerateAsyncParams Params;
+							Params.IntervalStart = Viewport.GetStartTime();
+							Params.IntervalEnd = Viewport.GetEndTime();
+							Params.Resolution = 0.0;
+							Params.SetupCallback = [&FilteredEvents](uint32 NumTasks)
+							{
+								FilteredEvents.AddDefaulted(NumTasks);
+							};
+							Params.EventRangeCallback = [this, &Builder, TimerReader, FilterEventType, &FilteredEvents](double StartTime, double EndTime, uint32 Depth, const TraceServices::FTimingProfilerEvent& Event, uint32 TaskIndex)
+							{
+								const TraceServices::FTimingProfilerTimer* Timer = TimerReader->GetTimer(Event.TimerIndex);
+								if (ensure(Timer != nullptr))
+								{
+									if (Timer->Id == FilterEventType)
+									{
+										FPendingEventInfo TimelineEvent;
+										TimelineEvent.StartTime = StartTime;
+										TimelineEvent.EndTime = EndTime;
+										TimelineEvent.Depth = Depth;
+										TimelineEvent.TimerIndex = Event.TimerIndex;
+										FilteredEvents[TaskIndex].Add(TimelineEvent);
+									}
+								}
+								return TraceServices::EEventEnumerate::Continue;
+							};
+
+							// Note: Enumerating events for filtering should not use downsampling.
+							Timeline.EnumerateEventsDownSampledAsync(Params);
+
+							for (TArray<FPendingEventInfo>& Array : FilteredEvents)
+							{
+								for (FPendingEventInfo& TimelineEvent : Array)
+								{
+									const TraceServices::FTimingProfilerTimer* Timer = TimerReader->GetTimer(TimelineEvent.TimerIndex);
+									AddTimingEventToBuilder(Builder, TimelineEvent.StartTime, TimelineEvent.EndTime, TimelineEvent.Depth, TimelineEvent.TimerIndex, Timer);
+								}
+							}
+						});
+				}
+				else
+				{
+					TimingProfilerProvider.ReadTimeline(TimelineIndex,
+						[&Viewport, this, &Builder, TimerReader, FilterEventType](const TraceServices::ITimingProfilerProvider::Timeline& Timeline)
+						{
+							// Note: Enumerating events for filtering should not use downsampling.
+							Timeline.EnumerateEventsDownSampled(Viewport.GetStartTime(), Viewport.GetEndTime(), 0,
+								[this, &Builder, TimerReader, FilterEventType](double StartTime, double EndTime, uint32 Depth, const TraceServices::FTimingProfilerEvent& Event)
+								{
+									const TraceServices::FTimingProfilerTimer* Timer = TimerReader->GetTimer(Event.TimerIndex);
+									if (ensure(Timer != nullptr))
+									{
+										if (Timer->Id == FilterEventType)
+										{
+											AddTimingEventToBuilder(Builder, StartTime, EndTime, Depth, Event.TimerIndex, Timer);
+										}
+									}
+									return TraceServices::EEventEnumerate::Continue;
+								});
+						});
+				}
+			}
+			else // generic filter
+			{
+				TimingProfilerProvider.ReadTimeline(TimelineIndex,
+					[&Viewport, this, &Builder, TimerReader, &EventFilterPtr](const TraceServices::ITimingProfilerProvider::Timeline& Timeline)
+					{
+						// Note: Enumerating events for filtering should not use downsampling.
+						//const double SecondsPerPixel = 1.0 / Viewport.GetScaleX();
+						//Timeline.EnumerateEventsDownSampled(Viewport.GetStartTime(), Viewport.GetEndTime(), SecondsPerPixel,
+						Timeline.EnumerateEvents(Viewport.GetStartTime(), Viewport.GetEndTime(),
+							[this, &Builder, TimerReader, &EventFilterPtr](double StartTime, double EndTime, uint32 Depth, const TraceServices::FTimingProfilerEvent& Event)
+							{
+								const TraceServices::FTimingProfilerTimer* Timer = TimerReader->GetTimer(Event.TimerIndex);
+								if (ensure(Timer != nullptr))
+								{
+									FThreadTrackEvent TimingEvent(SharedThis(this), StartTime, EndTime, Depth);
+									TimingEvent.SetTimerId(Timer->Id);
+									TimingEvent.SetTimerIndex(Event.TimerIndex);
+
+									if (EventFilterPtr->FilterEvent(TimingEvent))
+									{
+										AddTimingEventToBuilder(Builder, StartTime, EndTime, Depth, Event.TimerIndex, Timer);
+									}
+								}
+								return TraceServices::EEventEnumerate::Continue;
+							});
+					});
+			}
+		}
+	}
+
+	if (HasCustomFilter()) // Custom filter (from the filtering widget)
+	{
+		TSharedPtr<const TraceServices::IAnalysisSession> Session = FInsightsManager::Get()->GetSession();
+		if (Session.IsValid() && TraceServices::ReadTimingProfilerProvider(*Session.Get()))
+		{
+			TraceServices::FAnalysisSessionReadScope SessionReadScope(*Session.Get());
+
+			const TraceServices::ITimingProfilerProvider& TimingProfilerProvider = *TraceServices::ReadTimingProfilerProvider(*Session.Get());
+
+			const TraceServices::ITimingProfilerTimerReader* TimerReader;
+			TimingProfilerProvider.ReadTimers([&TimerReader](const TraceServices::ITimingProfilerTimerReader& Out) { TimerReader = &Out; });
+
+			const FTimingTrackViewport& Viewport = Context.GetViewport();
+
+			TimingProfilerProvider.ReadTimeline(TimelineIndex,
+				[&Viewport, this, &Builder, TimerReader](const TraceServices::ITimingProfilerProvider::Timeline& Timeline)
+				{
+					TArray<TArray<FPendingEventInfo>> FilteredEvents;
+					TArray<FFilterContext> FilterContexts;
+
+					TraceServices::ITimeline<TraceServices::FTimingProfilerEvent>::EnumerateAsyncParams Params;
+					constexpr uint32 LargeTimelineThreshold = 50 * 1000 * 1000;
+					if (Timeline.GetEventCount() > LargeTimelineThreshold)
+					{
+						if (FilterConfigurator->IsKeyUsed(static_cast<int32>(EFilterField::Metadata)))
+						{
+							Params.MaxOccupancy = 0.75f; // This filter can be slow so reduce occupancy to avoid starvation.
+						}
+					}
+					Params.IntervalStart = Viewport.GetStartTime();
+					Params.IntervalEnd = Viewport.GetEndTime();
+
+					// Note: Enumerating events for filtering should not use downsampling.
+					Params.Resolution = 0.0;
+					Params.SetupCallback = [&FilteredEvents, &FilterContexts, this](uint32 NumTasks)
+					{
+						FilteredEvents.AddDefaulted(NumTasks);
+						FilterContexts.AddDefaulted(NumTasks);
+						for (FFilterContext& Context : FilterContexts)
+						{
+							Context.SetReturnValueForUnsetFilters(false);
+							Context.AddFilterData<double>(static_cast<int32>(EFilterField::StartTime), 0.0f);
+							Context.AddFilterData<double>(static_cast<int32>(EFilterField::EndTime), 0.0f);
+							Context.AddFilterData<double>(static_cast<int32>(EFilterField::Duration), 0.0f);
+							Context.AddFilterData<FString>(static_cast<int32>(EFilterField::TrackName), this->GetName());
+							Context.AddFilterData<int64>(static_cast<int32>(EFilterField::TimerId), 0);
+							Context.AddFilterData<int64>(static_cast<int32>(EFilterField::TimerName), 0);
+							Context.AddFilterData<int64>(static_cast<int32>(EFilterField::Metadata), 0);
+						}
+					};
+					Params.EventRangeCallback = [this, &Builder, TimerReader, &FilteredEvents, &FilterContexts](double StartTime, double EndTime, uint32 Depth, const TraceServices::FTimingProfilerEvent& Event, uint32 TaskIndex)
+					{
+						const TraceServices::FTimingProfilerTimer* Timer = TimerReader->GetTimer(Event.TimerIndex);
+						if (ensure(Timer != nullptr))
+						{
+							FFilterContext& Context = FilterContexts[TaskIndex];
+							Context.SetFilterData<double>(static_cast<int32>(EFilterField::StartTime), StartTime);
+							Context.SetFilterData<double>(static_cast<int32>(EFilterField::EndTime), EndTime);
+							Context.SetFilterData<double>(static_cast<int32>(EFilterField::Duration), EndTime - StartTime);
+							// The TimerName filter also translates to the numeric Id for performance reasons.
+							Context.SetFilterData<int64>(static_cast<int32>(EFilterField::TimerId), Timer->Id);
+							Context.SetFilterData<int64>(static_cast<int32>(EFilterField::TimerName), Timer->Id);
+							Context.SetFilterData<int64>(static_cast<int32>(EFilterField::Metadata), Event.TimerIndex);
+
+							if (FilterConfigurator->ApplyFilters(Context))
+							{
+								FPendingEventInfo TimelineEvent;
+								TimelineEvent.StartTime = StartTime;
+								TimelineEvent.EndTime = EndTime;
+								TimelineEvent.Depth = Depth;
+								TimelineEvent.TimerIndex = Event.TimerIndex;
+								FilteredEvents[TaskIndex].Add(TimelineEvent);
+							}
+						}
+						return TraceServices::EEventEnumerate::Continue;
+					};
+
+					Timeline.EnumerateEventsDownSampledAsync(Params);
+
+					for (TArray<FPendingEventInfo>& Array : FilteredEvents)
+					{
+						for (FPendingEventInfo& TimelineEvent : Array)
+						{
+							const TraceServices::FTimingProfilerTimer* Timer = TimerReader->GetTimer(TimelineEvent.TimerIndex);
+							AddTimingEventToBuilder(Builder, TimelineEvent.StartTime, TimelineEvent.EndTime, TimelineEvent.Depth, TimelineEvent.TimerIndex, Timer);
+						}
+					}
+				});
+		}
+	}
+}
+
+////////////////////////////////////////////////////////////////////////////////////////////////////
+
+void FThreadTimingTrack::PostDraw(const ITimingTrackDrawContext& Context) const
+{
+	if (ChildTrack.IsValid())
+	{
+		ChildTrack->PostDraw(Context);
+	}
+
+	const TSharedPtr<const ITimingEvent> SelectedEventPtr = Context.GetSelectedEvent();
+	if (SelectedEventPtr.IsValid() &&
+		SelectedEventPtr->CheckTrack(this) &&
+		SelectedEventPtr->Is<FThreadTrackEvent>())
+	{
+		const FThreadTrackEvent& SelectedEvent = SelectedEventPtr->As<FThreadTrackEvent>();
+		const ITimingViewDrawHelper& Helper = Context.GetHelper();
+
+		TSharedPtr<const TraceServices::IAnalysisSession> Session = FInsightsManager::Get()->GetSession();
+		check(Session.IsValid());
+
+		TraceServices::FAnalysisSessionReadScope SessionReadScope(*Session.Get());
+
+		const TraceServices::ITimingProfilerProvider& TimingProfilerProvider = *TraceServices::ReadTimingProfilerProvider(*Session.Get());
+
+		const TraceServices::ITimingProfilerTimerReader* TimerReader;
+		TimingProfilerProvider.ReadTimers([&TimerReader](const TraceServices::ITimingProfilerTimerReader& Out) { TimerReader = &Out; });
+
+		const TraceServices::FTimingProfilerTimer* Timer = TimerReader->GetTimer(SelectedEvent.GetTimerIndex());
+		if (Timer != nullptr)
+		{
+			FString TimerName(Timer->Name);
+
+			const double SelectedEventDuration = SelectedEvent.GetDuration();
+			TStringBuilder<1024> StringBuilder;
+			StringBuilder.Appendf(TEXT("  Incl.: %s"), *FormatTimeAuto(SelectedEventDuration, 2));
+			if (SelectedEventDuration != std::numeric_limits<double>::infinity())
+			{
+				StringBuilder.Appendf(TEXT("   Excl.: %s"), *FormatTimeAuto(SelectedEvent.GetExclusiveTime(), 2));
+			}
+			FString StatsText(StringBuilder.ToView());
+
+			if (Timer->File)
+			{
+				FString SourceFile = FPaths::GetCleanFilename(FString(Timer->File));
+				FString SourceFileAndLine = FString::Printf(TEXT("%s (%d)"), *SourceFile, Timer->Line);
+				DrawSelectedEventInfoEx(StatsText, TimerName, SourceFileAndLine, Context.GetViewport(), Context.GetDrawContext(), Helper.GetWhiteBrush(), Helper.GetEventFont());
+			}
+			else
+			{
+				DrawSelectedEventInfoEx(StatsText, TimerName, FString(), Context.GetViewport(), Context.GetDrawContext(), Helper.GetWhiteBrush(), Helper.GetEventFont());
+			}
+		}
+	}
+}
+
+////////////////////////////////////////////////////////////////////////////////////////////////////
+
+void FThreadTimingTrack::InitTooltip(FTooltipDrawState& InOutTooltip, const ITimingEvent& InTooltipEvent) const
+{
+	InOutTooltip.ResetContent();
+
+	if (InTooltipEvent.CheckTrack(this) && InTooltipEvent.Is<FThreadTrackEvent>())
+	{
+		const FThreadTrackEvent& TooltipEvent = InTooltipEvent.As<FThreadTrackEvent>();
+
+		TSharedPtr<FThreadTrackEvent> ParentTimingEvent;
+		TSharedPtr<FThreadTrackEvent> RootTimingEvent;
+		GetParentAndRoot(TooltipEvent, ParentTimingEvent, RootTimingEvent);
+
+		TSharedPtr<const TraceServices::IAnalysisSession> Session = FInsightsManager::Get()->GetSession();
+		check(Session.IsValid());
+
+		TraceServices::FAnalysisSessionReadScope SessionReadScope(*Session.Get());
+
+		const TraceServices::ITimingProfilerProvider& TimingProfilerProvider = *TraceServices::ReadTimingProfilerProvider(*Session.Get());
+
+		const TraceServices::ITimingProfilerTimerReader* TimerReader;
+		TimingProfilerProvider.ReadTimers([&TimerReader](const TraceServices::ITimingProfilerTimerReader& Out) { TimerReader = &Out; });
+
+		const TraceServices::FTimingProfilerTimer* Timer = TimerReader->GetTimer(TooltipEvent.GetTimerIndex());
+		const TCHAR* TimerName = (Timer != nullptr) ? Timer->Name : TEXT("N/A");
+		InOutTooltip.AddTitle(TimerName);
+
+		const double TooltipEventDuration = TooltipEvent.GetDuration();
+
+		if (TooltipEvent.GetDepth() > 0 &&
+			ParentTimingEvent.IsValid() &&
+			ParentTimingEvent->GetDuration() > 0.0 &&
+			ParentTimingEvent->GetDuration() != std::numeric_limits<double>::infinity())
+		{
+			Timer = TimerReader->GetTimer(ParentTimingEvent->GetTimerIndex());
+			const TCHAR* ParentTimerName = (Timer != nullptr) ? Timer->Name : TEXT("N/A");
+			FNumberFormattingOptions FormattingOptions;
+			FormattingOptions.MaximumFractionalDigits = 2;
+			const FString ValueStr = FString::Printf(TEXT("%s %s"), *FText::AsPercent(TooltipEventDuration / ParentTimingEvent->GetDuration(), &FormattingOptions).ToString(), ParentTimerName);
+			InOutTooltip.AddNameValueTextLine(TEXTVIEW("% of Parent:"), ValueStr);
+		}
+
+		if (TooltipEvent.GetDepth() > 1 &&
+			RootTimingEvent.IsValid() &&
+			RootTimingEvent->GetDuration() > 0.0 &&
+			RootTimingEvent->GetDuration() != std::numeric_limits<double>::infinity())
+		{
+			Timer = TimerReader->GetTimer(RootTimingEvent->GetTimerIndex());
+			const TCHAR* RootTimerName = (Timer != nullptr) ? Timer->Name : TEXT("N/A");
+			FNumberFormattingOptions FormattingOptions;
+			FormattingOptions.MaximumFractionalDigits = 2;
+			const FString ValueStr = FString::Printf(TEXT("%s %s"), *FText::AsPercent(TooltipEventDuration / RootTimingEvent->GetDuration(), &FormattingOptions).ToString(), RootTimerName);
+			InOutTooltip.AddNameValueTextLine(TEXTVIEW("% of Root:"), ValueStr);
+		}
+
+		InOutTooltip.AddNameValueTextLine(TEXTVIEW("Inclusive Time:"), FormatTimeAuto(TooltipEventDuration, 2));
+
+		if (TooltipEventDuration > 0.0 &&
+			TooltipEventDuration != std::numeric_limits<double>::infinity())
+		{
+			const double ExclusiveTimePercent = TooltipEvent.GetExclusiveTime() / TooltipEventDuration;
+			FNumberFormattingOptions FormattingOptions;
+			FormattingOptions.MaximumFractionalDigits = 2;
+			const FString ExclStr = FString::Printf(TEXT("%s (%s)"), *FormatTimeAuto(TooltipEvent.GetExclusiveTime(), 2), *FText::AsPercent(ExclusiveTimePercent, &FormattingOptions).ToString());
+			InOutTooltip.AddNameValueTextLine(TEXTVIEW("Exclusive Time:"), ExclStr);
+		}
+
+		InOutTooltip.AddNameValueTextLine(TEXTVIEW("Depth:"), FString::Printf(TEXT("%d"), TooltipEvent.GetDepth()));
+
+		TArrayView<const uint8> Metadata = TimerReader->GetMetadata(TooltipEvent.GetTimerIndex());
+		if (Metadata.Num() > 0)
+		{
+			AppendMetadataToTooltip(InOutTooltip, Metadata);
+		}
+
+		// tasks
+		const TraceServices::ITasksProvider* TasksProvider = TraceServices::ReadTasksProvider(*Session.Get());
+		if (TasksProvider != nullptr)
+		{
+			auto AddTaskInfo = [&InOutTooltip, this](const TraceServices::FTaskInfo& Task)
+			{
+				InOutTooltip.AddTextLine(FString::Printf(TEXT("-------- Task %d%s --------"), Task.Id, Task.bTracked ? TEXT("") : TEXT(" (not tracked)")), FLinearColor::Green);
+
+				if (Task.DebugName != nullptr)
+				{
+					InOutTooltip.AddTextLine(FString::Printf(TEXT("%s"), Task.DebugName), FLinearColor::Green);
+				}
+
+				ENamedThreads::Type ThreadInfo = (ENamedThreads::Type)Task.ThreadToExecuteOn;
+				ENamedThreads::Type ThreadIndex = ENamedThreads::GetThreadIndex(ThreadInfo);
+
+				auto FormatTaskTimestamp = [](double Timestamp) -> FString
+				{
+					return (Timestamp != TraceServices::FTaskInfo::InvalidTimestamp) ? FString::SanitizeFloat(Timestamp) : TEXT("[not set]");
+				};
+
+				auto FormatTaskTime = [](double Time) -> FString
+				{
+					return FormatTimeAuto(Time, 2);
+				};
+
+				auto GetTrackName = [this](uint32 InThreadId) -> FString
+				{
+					TSharedPtr<FCpuTimingTrack> Track = SharedState.GetCpuTrack(InThreadId);
+					return Track.IsValid() ? Track->GetName() : TEXT("Unknown");
+				};
+
+				if (ThreadIndex == ENamedThreads::AnyThread)
+				{
+					const TCHAR* TaskPri = ENamedThreads::GetTaskPriority(ThreadInfo) == ENamedThreads::NormalTaskPriority ? TEXT("Normal") : TEXT("High");
+
+					int32 ThreadPriIndex = ENamedThreads::GetThreadPriorityIndex(ThreadInfo);
+					const TCHAR* ThreadPriStrs[] = { TEXT("Normal"), TEXT("High"), TEXT("Low") };
+					const TCHAR* ThreadPri = ensure(ThreadPriIndex >= 0 && ThreadPriIndex < 3) ? ThreadPriStrs[ThreadPriIndex] : TEXT("Unknown");
+
+					InOutTooltip.AddTextLine(
+						FString::Printf(TEXT("%s Pri task on %s Pri worker (%s)"), TaskPri, ThreadPri, *GetTrackName(Task.StartedThreadId)),
+						FLinearColor::Green);
+				}
+				else
+				{
+					const TCHAR* QueueStr = ENamedThreads::GetQueueIndex(ThreadInfo) == ENamedThreads::MainQueue ? TEXT("Main") : TEXT("Local");
+					InOutTooltip.AddTextLine(
+						FString::Printf(TEXT("%s (%s queue)"), *GetTrackName(Task.StartedThreadId), QueueStr),
+						FLinearColor::Green);
+				}
+
+				InOutTooltip.AddNameValueTextLine(TEXTVIEW("Created:"), FString::Printf(TEXT("%s on %s"),
+					*FormatTaskTimestamp(Task.CreatedTimestamp),
+					*GetTrackName(Task.CreatedThreadId)));
+
+				InOutTooltip.AddNameValueTextLine(TEXTVIEW("Launched:"), FString::Printf(TEXT("%s (+%s) on %s"),
+					*FormatTaskTimestamp(Task.LaunchedTimestamp),
+					*FormatTaskTime(Task.LaunchedTimestamp - Task.CreatedTimestamp),
+					*GetTrackName(Task.LaunchedThreadId)));
+
+				InOutTooltip.AddNameValueTextLine(TEXTVIEW("Scheduled:"), FString::Printf(TEXT("%s (+%s) on %s"),
+					*FormatTaskTimestamp(Task.ScheduledTimestamp),
+					*FormatTaskTime(Task.ScheduledTimestamp - Task.LaunchedTimestamp),
+					*GetTrackName(Task.ScheduledThreadId)));
+
+				InOutTooltip.AddNameValueTextLine(TEXTVIEW("Started:"), FString::Printf(TEXT("%s (+%s)"),
+					*FormatTaskTimestamp(Task.StartedTimestamp),
+					*FormatTaskTime(Task.StartedTimestamp - Task.ScheduledTimestamp)));
+
+				if (Task.FinishedTimestamp != TraceServices::FTaskInfo::InvalidTimestamp)
+				{
+					InOutTooltip.AddNameValueTextLine(TEXTVIEW("Finished:"), FString::Printf(TEXT("%s (+%s)"),
+						*FormatTaskTimestamp(Task.FinishedTimestamp),
+						*FormatTaskTime(Task.FinishedTimestamp - Task.StartedTimestamp)));
+
+					if (Task.CompletedTimestamp != TraceServices::FTaskInfo::InvalidTimestamp)
+					{
+						InOutTooltip.AddNameValueTextLine(TEXTVIEW("Completed:"), FString::Printf(TEXT("%s (+%s) on %s"),
+							*FormatTaskTimestamp(Task.CompletedTimestamp),
+							*FormatTaskTime(Task.CompletedTimestamp - Task.FinishedTimestamp),
+							*GetTrackName(Task.CompletedThreadId)));
+
+						if (Task.DestroyedTimestamp != TraceServices::FTaskInfo::InvalidTimestamp)
+						{
+							InOutTooltip.AddNameValueTextLine(TEXTVIEW("Destroyed:"), FString::Printf(TEXT("%s (+%s) on %s"),
+								*FormatTaskTimestamp(Task.DestroyedTimestamp),
+								*FormatTaskTime(Task.DestroyedTimestamp - Task.CompletedTimestamp),
+								*GetTrackName(Task.DestroyedThreadId)));
+						}
+					}
+				}
+
+				InOutTooltip.AddNameValueTextLine(TEXTVIEW("Prerequisite tasks:"), FString::Printf(TEXT("%d"), Task.Prerequisites.Num()));
+				InOutTooltip.AddNameValueTextLine(TEXTVIEW("Subsequent tasks:"), FString::Printf(TEXT("%d"), Task.Subsequents.Num()));
+				InOutTooltip.AddNameValueTextLine(TEXTVIEW("Parent tasks:"), FString::Printf(TEXT("%d"), Task.ParentTasks.Num()));
+				InOutTooltip.AddNameValueTextLine(TEXTVIEW("Nested tasks:"), FString::Printf(TEXT("%d"), Task.NestedTasks.Num()));
+			};
+
+			// info about a task
+			const TraceServices::FTaskInfo* Task = TasksProvider->TryGetTask(ThreadId, TooltipEvent.GetStartTime());
+			if (Task != nullptr && Task->FinishedTimestamp >= TooltipEvent.GetEndTime())
+			{
+				AddTaskInfo(*Task);
+			}
+
+			// info about blocking
+			const TraceServices::FWaitingForTasks* Waiting = TasksProvider->TryGetWaiting(TimerName, ThreadId, TooltipEvent.GetStartTime());
+			if (Waiting != nullptr && Waiting->Tasks.Num() > 0)
+			{
+				InOutTooltip.AddTextLine(TEXT("-------- Waiting for tasks --------"), FLinearColor::Red);
+				constexpr int32 NumIdsOnRow = 4;
+				TStringBuilder<1024> StringBuilder;
+
+				// Add the first line of Task Id values.
+				for (int32 Index = 0; Index < Waiting->Tasks.Num() && Index < NumIdsOnRow; ++Index)
+				{
+					StringBuilder.Appendf(TEXT("%d, "), Waiting->Tasks[Index]);
+				}
+
+				// Remove separators from last entry.
+				if (Waiting->Tasks.Num() <= NumIdsOnRow)
+				{
+					StringBuilder.RemoveSuffix(2);
+				}
+				InOutTooltip.AddNameValueTextLine(FString::Printf(TEXT("Tasks[%d]:"), Waiting->Tasks.Num()), StringBuilder.ToView());
+				StringBuilder.Reset();
+
+				// Add the rest of the lines with an empty name so they appear as a multi line value.
+				for (int32 Index = NumIdsOnRow; Index < Waiting->Tasks.Num(); ++Index)
+				{
+					StringBuilder.Appendf(TEXT("%d, "), Waiting->Tasks[Index]);
+
+					if ((Index + 1) % NumIdsOnRow == 0)
+					{
+						InOutTooltip.AddNameValueTextLine(TEXT(""), StringBuilder.ToView());
+						StringBuilder.Reset();
+					}
+				}
+				if (StringBuilder.Len() > 1)
+				{
+					StringBuilder.RemoveSuffix(2);
+					InOutTooltip.AddNameValueTextLine(TEXT(""), StringBuilder.ToView());
+				}
+
+				InOutTooltip.AddNameValueTextLine(TEXT("Started waiting:"),
+					FString::Printf(TEXT("%s"), *FString::SanitizeFloat(Waiting->StartedTimestamp)));
+
+				if (Waiting->FinishedTimestamp != TraceServices::FTaskInfo::InvalidTimestamp)
+				{
+					InOutTooltip.AddNameValueTextLine(TEXT("Finished waiting:"),
+						FString::Printf(TEXT("%s (+%s)"),
+							*FString::SanitizeFloat(Waiting->FinishedTimestamp),
+							*FormatTimeAuto(Waiting->FinishedTimestamp - Waiting->StartedTimestamp, 2)));
+				}
+				else
+				{
+					InOutTooltip.AddNameValueTextLine(TEXT("Finished waiting:"), TEXT("[not set]"));
+				}
+
+				const int32 MaxWaitedTasksToList = 5;
+				int32 NumTasksToList = FMath::Min(Waiting->Tasks.Num(), MaxWaitedTasksToList);
+				for (int32 TaskIndex = 0; TaskIndex != NumTasksToList; ++TaskIndex)
+				{
+					const TraceServices::FTaskInfo* WaitedTask = TasksProvider->TryGetTask(Waiting->Tasks[TaskIndex]);
+					if (WaitedTask != nullptr)
+					{
+						AddTaskInfo(*WaitedTask);
+					}
+				}
+				if (NumTasksToList < Waiting->Tasks.Num())
+				{
+					InOutTooltip.AddTextLine(TEXT("[...]"), FLinearColor::Green);
+				}
+			}
+		}
+	}
+	else if (ChildTrack.IsValid())
+	{
+		ChildTrack->InitTooltip(InOutTooltip, InTooltipEvent);
+	}
+
+	InOutTooltip.UpdateLayout();
+}
+
+////////////////////////////////////////////////////////////////////////////////////////////////////
+
+void FThreadTimingTrack::GetParentAndRoot(const FThreadTrackEvent& TimingEvent, TSharedPtr<FThreadTrackEvent>& OutParentTimingEvent, TSharedPtr<FThreadTrackEvent>& OutRootTimingEvent) const
+{
+	if (TimingEvent.GetDepth() > 0)
+	{
+		TSharedPtr<const TraceServices::IAnalysisSession> Session = FInsightsManager::Get()->GetSession();
+		if (Session.IsValid())
+		{
+			TraceServices::FAnalysisSessionReadScope SessionReadScope(*Session.Get());
+
+			if (TraceServices::ReadTimingProfilerProvider(*Session.Get()))
+			{
+				const TraceServices::ITimingProfilerProvider& TimingProfilerProvider = *TraceServices::ReadTimingProfilerProvider(*Session.Get());
+
+				TimingProfilerProvider.ReadTimeline(GetTimelineIndex(), [&TimingEvent, &OutParentTimingEvent, &OutRootTimingEvent](const TraceServices::ITimingProfilerProvider::Timeline& Timeline)
+					{
+						const double Time = (TimingEvent.GetStartTime() + TimingEvent.GetEndTime()) / 2;
+						TimelineEventInfo EventInfo;
+						bool IsFound = Timeline.GetEventInfo(Time, 0, TimingEvent.GetDepth() - 1, EventInfo);
+						if (IsFound)
+						{
+							CreateFThreadTrackEventFromInfo(EventInfo, TimingEvent.GetTrack(), TimingEvent.GetDepth() - 1, OutParentTimingEvent);
+						}
+
+						IsFound = Timeline.GetEventInfo(Time, 0, 0, EventInfo);
+						if (IsFound)
+						{
+							CreateFThreadTrackEventFromInfo(EventInfo, TimingEvent.GetTrack(), 0, OutRootTimingEvent);
+						}
+					});
+			}
+		}
+	}
+}
+
+////////////////////////////////////////////////////////////////////////////////////////////////////
+
+const TSharedPtr<const ITimingEvent> FThreadTimingTrack::GetEvent(float InPosX, float InPosY, const FTimingTrackViewport& Viewport) const
+{
+	TSharedPtr<FThreadTrackEvent> TimingEvent;
+
+	const FTimingViewLayout& Layout = Viewport.GetLayout();
+
+	float TopLaneY = 0.0f;
+	float TrackLanesHeight = 0.0f;
+	if (ChildTrack.IsValid() && ChildTrack->GetHeight() > 0.0f)
+	{
+		const float HeaderDY = InPosY - ChildTrack->GetPosY();
+		if (HeaderDY >= 0.0f && HeaderDY < ChildTrack->GetHeight())
+		{
+			return ChildTrack->GetEvent(InPosX, InPosY, Viewport);
+		}
+		TopLaneY = GetPosY() + 1.0f + Layout.TimelineDY + ChildTrack->GetHeight() + Layout.ChildTimelineDY;
+		TrackLanesHeight = GetHeight() - ChildTrack->GetHeight() - 1.0f - 2.0f * Layout.TimelineDY - Layout.ChildTimelineDY;
+	}
+	else
+	{
+		if (IsChildTrack())
+		{
+			TopLaneY = GetPosY();
+			TrackLanesHeight = GetHeight();
+		}
+		else
+		{
+			TopLaneY = GetPosY() + 1.0f + Layout.TimelineDY;
+			TrackLanesHeight = GetHeight() - 1.0f - 2.0f * Layout.TimelineDY;
+		}
+	}
+
+	const float DY = InPosY - TopLaneY;
+
+	// If mouse is not above first sub-track or below last sub-track...
+	if (DY >= 0 && DY < TrackLanesHeight)
+	{
+		const int32 Depth = static_cast<int32>(DY / (Layout.EventH + Layout.EventDY));
+
+		const double SecondsPerPixel = 1.0 / Viewport.GetScaleX();
+
+		const double EventTime = Viewport.SlateUnitsToTime(InPosX);
+
+		TSharedPtr<const TraceServices::IAnalysisSession> Session = FInsightsManager::Get()->GetSession();
+		if (Session.IsValid())
+		{
+			TraceServices::FAnalysisSessionReadScope SessionReadScope(*Session.Get());
+
+			if (EventTime <= Session->GetDurationSeconds() && TraceServices::ReadTimingProfilerProvider(*Session.Get()))
+			{
+				const TraceServices::ITimingProfilerProvider& TimingProfilerProvider = *TraceServices::ReadTimingProfilerProvider(*Session.Get());
+
+				TimingProfilerProvider.ReadTimeline(GetTimelineIndex(), [this, &EventTime, &Depth, &TimingEvent, &SecondsPerPixel](const TraceServices::ITimingProfilerProvider::Timeline& Timeline)
+					{
+						TimelineEventInfo EventInfo;
+						bool IsFound = Timeline.GetEventInfo(EventTime, 2 * SecondsPerPixel, Depth, EventInfo);
+						if (IsFound)
+						{
+							CreateFThreadTrackEventFromInfo(EventInfo, SharedThis(this), Depth, TimingEvent);
+						}
+					});
+			}
+		}
+	}
+
+	return TimingEvent;
+}
+
+////////////////////////////////////////////////////////////////////////////////////////////////////
+
+const TSharedPtr<const ITimingEvent> FThreadTimingTrack::SearchEvent(const FTimingEventSearchParameters& InSearchParameters) const
+{
+	TSharedPtr<FThreadTrackEvent> FoundEvent;
+	FindTimingProfilerEvent(InSearchParameters, [this, &FoundEvent](double InFoundStartTime, double InFoundEndTime, uint32 InFoundDepth, const TraceServices::FTimingProfilerEvent& InFoundEvent)
+	{
+		FoundEvent = MakeShared<FThreadTrackEvent>(SharedThis(this), InFoundStartTime, InFoundEndTime, InFoundDepth);
+		FoundEvent->SetTimerIndex(InFoundEvent.TimerIndex);
+
+		uint32 TimerId = 0;
+		bool ret = FThreadTimingTrack::TimerIndexToTimerId(InFoundEvent.TimerIndex, TimerId);
+		if (ret)
+		{
+			FoundEvent->SetTimerId(TimerId);
+		}
+	});
+
+	return FoundEvent;
+}
+
+////////////////////////////////////////////////////////////////////////////////////////////////////
+
+void FThreadTimingTrack::UpdateEventStats(ITimingEvent& InOutEvent) const
+{
+	if (InOutEvent.CheckTrack(this) && InOutEvent.Is<FThreadTrackEvent>())
+	{
+		FThreadTrackEvent& TrackEvent = InOutEvent.As<FThreadTrackEvent>();
+		if (TrackEvent.IsExclusiveTimeComputed())
+		{
+			return;
+		}
+
+		TSharedPtr<const TraceServices::IAnalysisSession> Session = FInsightsManager::Get()->GetSession();
+		if (Session.IsValid())
+		{
+			TraceServices::FAnalysisSessionReadScope SessionReadScope(*Session.Get());
+
+			if (TraceServices::ReadTimingProfilerProvider(*Session.Get()))
+			{
+				const TraceServices::ITimingProfilerProvider& TimingProfilerProvider = *TraceServices::ReadTimingProfilerProvider(*Session.Get());
+
+				// Get Exclusive Time.
+				TimingProfilerProvider.ReadTimeline(GetTimelineIndex(), [&TrackEvent](const TraceServices::ITimingProfilerProvider::Timeline& Timeline)
+				{
+					TimelineEventInfo EventInfo;
+					bool bIsFound = Timeline.GetEventInfo(TrackEvent.GetStartTime(), 0.0, TrackEvent.GetDepth(), EventInfo);
+					if (bIsFound)
+					{
+						TrackEvent.SetExclusiveTime(EventInfo.ExclTime);
+						TrackEvent.SetIsExclusiveTimeComputed(true);
+					}
+				});
+			}
+		}
+	}
+}
+
+////////////////////////////////////////////////////////////////////////////////////////////////////
+
+void FThreadTimingTrack::OnEventSelected(const ITimingEvent& InSelectedEvent) const
+{
+	if (InSelectedEvent.CheckTrack(this) && InSelectedEvent.Is<FThreadTrackEvent>())
+	{
+		const FThreadTrackEvent& TrackEvent = InSelectedEvent.As<FThreadTrackEvent>();
+
+		// Select the timer node corresponding to timing event type of selected timing event.
+		FTimingProfilerManager::Get()->SetSelectedTimer(TrackEvent.GetTimerId());
+	}
+}
+
+////////////////////////////////////////////////////////////////////////////////////////////////////
+
+void FThreadTimingTrack::OnClipboardCopyEvent(const ITimingEvent& InSelectedEvent) const
+{
+	if (InSelectedEvent.CheckTrack(this) && InSelectedEvent.Is<FThreadTrackEvent>())
+	{
+		const FThreadTrackEvent& TrackEvent = InSelectedEvent.As<FThreadTrackEvent>();
+
+		TSharedPtr<const TraceServices::IAnalysisSession> Session = FInsightsManager::Get()->GetSession();
+		check(Session.IsValid());
+
+		TraceServices::FAnalysisSessionReadScope SessionReadScope(*Session.Get());
+		const TraceServices::ITimingProfilerProvider& TimingProfilerProvider = *TraceServices::ReadTimingProfilerProvider(*Session.Get());
+		const TraceServices::ITimingProfilerTimerReader* TimerReader = nullptr;
+		TimingProfilerProvider.ReadTimers([&TimerReader](const TraceServices::ITimingProfilerTimerReader& Out) { TimerReader = &Out; });
+		check(TimerReader);
+
+		const TraceServices::FTimingProfilerTimer* TimerPtr = TimerReader->GetTimer(TrackEvent.GetTimerIndex());
+
+		if (TimerPtr)
+		{
+			FString EventName(TimerPtr->Name);
+
+			FTimingEventsTrackDrawStateBuilder::AppendDurationToEventName(EventName, TrackEvent.GetDuration());
+
+			const uint32 TimerIndex = TrackEvent.GetTimerIndex();
+			if (int32(TimerIndex) < 0) // has metadata?
+			{
+				TArrayView<const uint8> Metadata = TimerReader->GetMetadata(TimerIndex);
+				if (Metadata.Num() > 0)
+				{
+					AppendMetadataToString(EventName, Metadata);
+				}
+			}
+
+			// Copy name of selected timing event to clipboard.
+			FPlatformApplicationMisc::ClipboardCopy(*EventName);
+		}
+	}
+}
+
+////////////////////////////////////////////////////////////////////////////////////////////////////
+
+void FThreadTimingTrack::BuildContextMenu(FMenuBuilder& MenuBuilder)
+{
+	if (GetGroupName() != nullptr)
+	{
+		MenuBuilder.BeginSection("CpuThread", LOCTEXT("ContextMenu_Section_CpuThread", "CPU Thread"));
+		{
+			MenuBuilder.AddMenuEntry(
+				FText::Format(LOCTEXT("CpuThreadGroupFmt", "Group: {0}"), FText::FromString(GetGroupName())),
+				FText(),
+				FSlateIcon(),
+				FUIAction(FExecuteAction(), FCanExecuteAction::CreateLambda([]() { return false; })),
+				NAME_None,
+				EUserInterfaceActionType::Button
+			);
+
+			const FString ThreadIdStr = FString::Printf(TEXT("%s%u (0x%X)"), ThreadId & 0x70000000 ? TEXT("*") : TEXT(""), ThreadId & ~0x70000000, ThreadId);
+			MenuBuilder.AddMenuEntry(
+				FText::Format(LOCTEXT("CpuThreadIdFmt", "Thread Id: {0}"), FText::FromString(ThreadIdStr)),
+				FText(),
+				FSlateIcon(),
+				FUIAction(FExecuteAction(), FCanExecuteAction::CreateLambda([]() { return false; })),
+				NAME_None,
+				EUserInterfaceActionType::Button
+			);
+		}
+		MenuBuilder.EndSection();
+	}
+
+	if (ChildTrack.IsValid())
+	{
+		ChildTrack->BuildContextMenu(MenuBuilder);
+	}
+}
+
+////////////////////////////////////////////////////////////////////////////////////////////////////
+
+bool FThreadTimingTrack::FindTimingProfilerEvent(const FThreadTrackEvent& InTimingEvent, TFunctionRef<void(double, double, uint32, const TraceServices::FTimingProfilerEvent&)> InFoundPredicate) const
+{
+	auto MatchEvent = [&InTimingEvent](double InStartTime, double InEndTime, uint32 InDepth)
+	{
+		return InDepth == InTimingEvent.GetDepth()
+			&& InStartTime == InTimingEvent.GetStartTime()
+			&& InEndTime == InTimingEvent.GetEndTime();
+	};
+
+	const double Time = (InTimingEvent.GetStartTime() + InTimingEvent.GetEndTime()) / 2;
+	FTimingEventSearchParameters SearchParameters(Time, Time, ETimingEventSearchFlags::StopAtFirstMatch, MatchEvent);
+	SearchParameters.SearchHandle = &InTimingEvent.GetSearchHandle();
+	return FindTimingProfilerEvent(SearchParameters, InFoundPredicate);
+}
+
+////////////////////////////////////////////////////////////////////////////////////////////////////
+
+bool FThreadTimingTrack::FindTimingProfilerEvent(const FTimingEventSearchParameters& InParameters, TFunctionRef<void(double, double, uint32, const TraceServices::FTimingProfilerEvent&)> InFoundPredicate) const
+{
+	FFilterContext FilterConfiguratorContext;
+	FilterConfiguratorContext.SetReturnValueForUnsetFilters(false);
+	FilterConfiguratorContext.AddFilterData<double>(static_cast<int32>(EFilterField::StartTime), 0.0f);
+	FilterConfiguratorContext.AddFilterData<double>(static_cast<int32>(EFilterField::EndTime), 0.0f);
+	FilterConfiguratorContext.AddFilterData<double>(static_cast<int32>(EFilterField::Duration), 0.0f);
+	FilterConfiguratorContext.AddFilterData<FString>(static_cast<int32>(EFilterField::TrackName), this->GetName());
+	FilterConfiguratorContext.AddFilterData<int64>(static_cast<int32>(EFilterField::TimerId), 0);
+	FilterConfiguratorContext.AddFilterData<int64>(static_cast<int32>(EFilterField::TimerName), 0);
+	FilterConfiguratorContext.AddFilterData<int64>(static_cast<int32>(EFilterField::Metadata), 0);
+
+	return TTimingEventSearch<TraceServices::FTimingProfilerEvent>::Search(
+		InParameters,
+
+		[this](TTimingEventSearch<TraceServices::FTimingProfilerEvent>::FContext& InContext)
+		{
+			TSharedPtr<const TraceServices::IAnalysisSession> Session = FInsightsManager::Get()->GetSession();
+			if (Session.IsValid())
+			{
+				TraceServices::FAnalysisSessionReadScope SessionReadScope(*Session.Get());
+
+				if (TraceServices::ReadTimingProfilerProvider(*Session.Get()))
+				{
+					const TraceServices::ITimingProfilerProvider& TimingProfilerProvider = *TraceServices::ReadTimingProfilerProvider(*Session.Get());
+
+					TimingProfilerProvider.ReadTimeline(GetTimelineIndex(), [&InContext](const TraceServices::ITimingProfilerProvider::Timeline& Timeline)
+					{
+						auto Callback = [&InContext](double EventStartTime, double EventEndTime, uint32 EventDepth, const TraceServices::FTimingProfilerEvent& Event)
+						{
+							InContext.Check(EventStartTime, EventEndTime, EventDepth, Event);
+							return InContext.ShouldContinueSearching() ? TraceServices::EEventEnumerate::Continue : TraceServices::EEventEnumerate::Stop;
+						};
+
+						if (InContext.GetParameters().SearchDirection == FTimingEventSearchParameters::ESearchDirection::Forward)
+						{
+							Timeline.EnumerateEvents(InContext.GetParameters().StartTime, InContext.GetParameters().EndTime, Callback);
+						}
+						else
+						{
+							Timeline.EnumerateEventsBackwards(InContext.GetParameters().EndTime, InContext.GetParameters().StartTime, Callback);
+						}
+					});
+				}
+			}
+		},
+
+		[&FilterConfiguratorContext, &InParameters](double EventStartTime, double EventEndTime, uint32 EventDepth, const TraceServices::FTimingProfilerEvent& Event)
+		{
+			if (!InParameters.FilterExecutor.IsValid())
+			{
+				return true;
+			}
+
+			TSharedPtr<const TraceServices::IAnalysisSession> Session = FInsightsManager::Get()->GetSession();
+			if (Session.IsValid())
+			{
+				TraceServices::FAnalysisSessionReadScope SessionReadScope(*Session.Get());
+				if (TraceServices::ReadTimingProfilerProvider(*Session.Get()))
+				{
+					const TraceServices::ITimingProfilerProvider& TimingProfilerProvider = *TraceServices::ReadTimingProfilerProvider(*Session.Get());
+					const TraceServices::ITimingProfilerTimerReader* TimerReader;
+					TimingProfilerProvider.ReadTimers([&TimerReader](const TraceServices::ITimingProfilerTimerReader& Out) { TimerReader = &Out; });
+
+					const TraceServices::FTimingProfilerTimer* Timer = TimerReader->GetTimer(Event.TimerIndex);
+					if (ensure(Timer != nullptr))
+					{
+						FilterConfiguratorContext.SetFilterData<double>(static_cast<int32>(EFilterField::StartTime), EventStartTime);
+						FilterConfiguratorContext.SetFilterData<double>(static_cast<int32>(EFilterField::EndTime), EventEndTime);
+						FilterConfiguratorContext.SetFilterData<double>(static_cast<int32>(EFilterField::Duration), EventEndTime - EventStartTime);
+
+						FilterConfiguratorContext.SetFilterData<int64>(static_cast<int32>(EFilterField::TimerId), Timer->Id);
+						FilterConfiguratorContext.SetFilterData<int64>(static_cast<int32>(EFilterField::TimerName), Timer->Id);
+
+						FilterConfiguratorContext.SetFilterData<int64>(static_cast<int32>(EFilterField::Metadata), Event.TimerIndex);
+
+						return InParameters.FilterExecutor->ApplyFilters(FilterConfiguratorContext);
+					}
+				}
+			}
+
+			return false;
+		},
+
+		[&InFoundPredicate](double InFoundStartTime, double InFoundEndTime, uint32 InFoundDepth, const TraceServices::FTimingProfilerEvent& InEvent)
+		{
+			InFoundPredicate(InFoundStartTime, InFoundEndTime, InFoundDepth, InEvent);
+		},
+
+		TTimingEventSearch<TraceServices::FTimingProfilerEvent>::NoMatch,
+
+		&SearchCache);
+}
+
+////////////////////////////////////////////////////////////////////////////////////////////////////
+
+void FThreadTimingTrack::CreateFThreadTrackEventFromInfo(const TimelineEventInfo& InEventInfo, const TSharedRef<const FBaseTimingTrack> InTrack, int32 InDepth, TSharedPtr<FThreadTrackEvent> &OutTimingEvent)
+{
+	OutTimingEvent = MakeShared<FThreadTrackEvent>(InTrack, InEventInfo.StartTime, InEventInfo.EndTime, InDepth);
+	FThreadTrackEvent& Event = OutTimingEvent->As<FThreadTrackEvent>();
+	Event.SetExclusiveTime(InEventInfo.ExclTime);
+	Event.SetIsExclusiveTimeComputed(true);
+	Event.SetTimerIndex(InEventInfo.Event.TimerIndex);
+
+	uint32 TimerId = 0;
+	bool ret = FThreadTimingTrack::TimerIndexToTimerId(InEventInfo.Event.TimerIndex, TimerId);
+	if (ret)
+	{
+		Event.SetTimerId(TimerId);
+	}
+
+}
+
+////////////////////////////////////////////////////////////////////////////////////////////////////
+
+bool FThreadTimingTrack::TimerIndexToTimerId(uint32 InTimerIndex, uint32& OutTimerId)
+{
+	TSharedPtr<const TraceServices::IAnalysisSession> Session = FInsightsManager::Get()->GetSession();
+	check(Session.IsValid())
+
+	TraceServices::FAnalysisSessionReadScope SessionReadScope(*Session.Get());
+
+	const TraceServices::ITimingProfilerProvider& TimingProfilerProvider = *TraceServices::ReadTimingProfilerProvider(*Session.Get());
+
+	const TraceServices::ITimingProfilerTimerReader* TimerReader;
+	TimingProfilerProvider.ReadTimers([&TimerReader](const TraceServices::ITimingProfilerTimerReader& Out) { TimerReader = &Out; });
+
+	const TraceServices::FTimingProfilerTimer* Timer = TimerReader->GetTimer(InTimerIndex);
+	if (Timer == nullptr)
+	{
+		return false;
+	}
+
+	OutTimerId = Timer->Id;
+	return true;
+}
+
+////////////////////////////////////////////////////////////////////////////////////////////////////
+
+bool FThreadTimingTrack::HasCustomFilter() const
+{
+	return FilterConfigurator.IsValid() && !FilterConfigurator->IsEmpty();
+}
+
+////////////////////////////////////////////////////////////////////////////////////////////////////
+
+int32 FThreadTimingTrack::GetDepthAt(double Time) const
+{
+	int32 Depth = 0;
+	TSharedPtr<const TraceServices::IAnalysisSession> Session = FInsightsManager::Get()->GetSession();
+	if (Session.IsValid() && TraceServices::ReadTimingProfilerProvider(*Session.Get()))
+	{
+		TraceServices::FAnalysisSessionReadScope SessionReadScope(*Session.Get());
+
+		const TraceServices::ITimingProfilerProvider& TimingProfilerProvider = *TraceServices::ReadTimingProfilerProvider(*Session.Get());
+
+		TimingProfilerProvider.ReadTimeline(TimelineIndex,
+			[Time, &Depth](const TraceServices::ITimingProfilerProvider::Timeline& Timeline)
+			{
+				Depth = Timeline.GetDepthAt(Time);
+			});
+	}
+	return Depth;
+}
+
+////////////////////////////////////////////////////////////////////////////////////////////////////
+
+void FThreadTimingTrack::SetFilterConfigurator(TSharedPtr<FFilterConfigurator> InFilterConfigurator)
+{
+	if (FilterConfigurator != InFilterConfigurator)
+	{
+		FilterConfigurator = InFilterConfigurator;
+		SetDirtyFlag();
+	}
+}
+
+////////////////////////////////////////////////////////////////////////////////////////////////////
+
+TSharedPtr<const ITimingEvent> FThreadTimingTrack::FindMaxEventInstance(uint32 InTimerId, double InStartTime, double InEndTime) const
+{
+	TSharedPtr<const TraceServices::IAnalysisSession> Session = FInsightsManager::Get()->GetSession();
+	check(Session.IsValid());
+
+	TraceServices::FAnalysisSessionReadScope SessionReadScope(*Session.Get());
+
+	const TraceServices::ITimingProfilerProvider& TimingProfilerProvider = *TraceServices::ReadTimingProfilerProvider(*Session.Get());
+
+	const TraceServices::ITimingProfilerTimerReader* TimerReader;
+	TimingProfilerProvider.ReadTimers([&TimerReader](const TraceServices::ITimingProfilerTimerReader& Out) { TimerReader = &Out; });
+
+	struct CandidateEvent
+	{
+		double StartTime = 0.0f;
+		double EndTime = -1.0f;
+		uint32 Depth = 0;
+		uint32 TimerIndex = 0;
+	};
+
+	TSharedPtr<FThreadTrackEvent> TimingEvent;
+
+	TimingProfilerProvider.ReadTimeline(TimelineIndex,
+		[TimerReader, InStartTime, InEndTime, InTimerId, &TimingEvent, this](const TraceServices::ITimingProfilerProvider::Timeline& Timeline)
+		{
+			TArray<CandidateEvent> CandidateEvents;
+
+			TraceServices::ITimeline<TraceServices::FTimingProfilerEvent>::EnumerateAsyncParams Params;
+			Params.IntervalStart = InStartTime;
+			Params.IntervalEnd = InEndTime;
+			Params.Resolution = 0.0;
+			Params.SetupCallback = [&CandidateEvents](uint32 NumTasks)
+			{
+				CandidateEvents.AddDefaulted(NumTasks);
+			};
+			Params.EventRangeCallback = [TimerReader, &CandidateEvents, InTimerId](double StartTime, double EndTime, uint32 Depth, const TraceServices::FTimingProfilerEvent& Event, uint32 TaskIndex)
+			{
+				const TraceServices::FTimingProfilerTimer* Timer = TimerReader->GetTimer(Event.TimerIndex);
+				if (ensure(Timer != nullptr))
+				{
+					if (Timer->Id == InTimerId)
+					{
+						double CandidateDuration = CandidateEvents[TaskIndex].EndTime - CandidateEvents[TaskIndex].StartTime;
+						double EventDuration = EndTime - StartTime;
+
+						if (EventDuration > CandidateDuration)
+						{
+							CandidateEvents[TaskIndex].StartTime = StartTime;
+							CandidateEvents[TaskIndex].EndTime = EndTime;
+							CandidateEvents[TaskIndex].Depth = Depth;
+							CandidateEvents[TaskIndex].TimerIndex = Event.TimerIndex;
+						}
+
+					}
+				}
+				return TraceServices::EEventEnumerate::Continue;
+			};
+
+			// Note: Enumerating events for filtering should not use downsampling.
+			Timeline.EnumerateEventsDownSampledAsync(Params);
+
+			CandidateEvent BestMatch;
+			for (const CandidateEvent& Event : CandidateEvents)
+			{
+				if ((Event.EndTime - Event.StartTime) > BestMatch.EndTime - BestMatch.StartTime)
+				{
+					BestMatch = Event;
+				}
+			}
+
+			if (BestMatch.EndTime > BestMatch.StartTime)
+			{
+				TimingEvent = MakeShared<FThreadTrackEvent>(SharedThis(this), BestMatch.StartTime, BestMatch.EndTime, BestMatch.Depth);
+				TimingEvent->SetTimerId(InTimerId);
+				TimingEvent->SetTimerIndex(BestMatch.TimerIndex);
+			}
+		});
+
+	return TimingEvent;
+}
+
+////////////////////////////////////////////////////////////////////////////////////////////////////
+
+TSharedPtr<const ITimingEvent> FThreadTimingTrack::FindMinEventInstance(uint32 InTimerId, double InStartTime, double InEndTime) const
+{
+	TSharedPtr<const TraceServices::IAnalysisSession> Session = FInsightsManager::Get()->GetSession();
+	check(Session.IsValid());
+
+	TraceServices::FAnalysisSessionReadScope SessionReadScope(*Session.Get());
+
+	const TraceServices::ITimingProfilerProvider& TimingProfilerProvider = *TraceServices::ReadTimingProfilerProvider(*Session.Get());
+
+	const TraceServices::ITimingProfilerTimerReader* TimerReader;
+	TimingProfilerProvider.ReadTimers([&TimerReader](const TraceServices::ITimingProfilerTimerReader& Out) { TimerReader = &Out; });
+
+	struct CandidateEvent
+	{
+		double StartTime = -std::numeric_limits<double>::infinity();
+		double EndTime = std::numeric_limits<double>::infinity();
+		uint32 Depth = 0;
+		uint32 TimerIndex = 0;
+	};
+
+	TSharedPtr<FThreadTrackEvent> TimingEvent;
+
+	TimingProfilerProvider.ReadTimeline(TimelineIndex,
+		[TimerReader, InStartTime, InEndTime, InTimerId, &TimingEvent, this](const TraceServices::ITimingProfilerProvider::Timeline& Timeline)
+		{
+			TArray<CandidateEvent> CandidateEvents;
+
+			TraceServices::ITimeline<TraceServices::FTimingProfilerEvent>::EnumerateAsyncParams Params;
+			Params.IntervalStart = InStartTime;
+			Params.IntervalEnd = InEndTime;
+			Params.Resolution = 0.0;
+			Params.SetupCallback = [&CandidateEvents](uint32 NumTasks)
+			{
+				CandidateEvents.AddDefaulted(NumTasks);
+			};
+			Params.EventRangeCallback = [TimerReader, &CandidateEvents, InTimerId](double StartTime, double EndTime, uint32 Depth, const TraceServices::FTimingProfilerEvent& Event, uint32 TaskIndex)
+			{
+				const TraceServices::FTimingProfilerTimer* Timer = TimerReader->GetTimer(Event.TimerIndex);
+				if (ensure(Timer != nullptr))
+				{
+					if (Timer->Id == InTimerId)
+					{
+						double CandidateDuration = CandidateEvents[TaskIndex].EndTime - CandidateEvents[TaskIndex].StartTime;
+						double EventDuration = EndTime - StartTime;
+
+						if (EventDuration < CandidateDuration)
+						{
+							CandidateEvents[TaskIndex].StartTime = StartTime;
+							CandidateEvents[TaskIndex].EndTime = EndTime;
+							CandidateEvents[TaskIndex].Depth = Depth;
+							CandidateEvents[TaskIndex].TimerIndex = Event.TimerIndex;
+						}
+					}
+				}
+				return TraceServices::EEventEnumerate::Continue;
+			};
+
+			// Note: Enumerating events for filtering should not use downsampling.
+			Timeline.EnumerateEventsDownSampledAsync(Params);
+
+			CandidateEvent BestMatch;
+			for (const CandidateEvent& Event : CandidateEvents)
+			{
+				if ((Event.EndTime - Event.StartTime) < BestMatch.EndTime - BestMatch.StartTime)
+				{
+					BestMatch = Event;
+				}
+			}
+
+			if (BestMatch.StartTime != -std::numeric_limits<double>::infinity())
+			{
+				TimingEvent = MakeShared<FThreadTrackEvent>(SharedThis(this), BestMatch.StartTime, BestMatch.EndTime, BestMatch.Depth);
+				TimingEvent->SetTimerId(InTimerId);
+				TimingEvent->SetTimerIndex(BestMatch.TimerIndex);
+			}
+		});
+
+	return TimingEvent;
+}
+
+////////////////////////////////////////////////////////////////////////////////////////////////////
+
+} // namespace UE::Insights::TimingProfiler
+
+#undef LOCTEXT_NAMESPACE
diff --git a/Engine/Source/Developer/TraceInsights/Private/Insights/TimingProfiler/Tracks/ThreadTimingTrack.h b/Engine/Source/Developer/TraceInsights/Private/Insights/TimingProfiler/Tracks/ThreadTimingTrack.h
new file mode 100644
--- /dev/null
+++ b/Engine/Source/Developer/TraceInsights/Private/Insights/TimingProfiler/Tracks/ThreadTimingTrack.h
@@ -0,0 +1,137 @@
+// Copyright Epic Games, Inc. All Rights Reserved.
+
+#pragma once
+
+#include "CoreTypes.h"
+
+// TraceServices
+#include "TraceServices/Model/TimingProfiler.h"
+
+// TraceInsights
+#include "Insights/ViewModels/TimingEventSearch.h" // for TTimingEventSearchCache
+#include "Insights/ViewModels/TimingEventsTrack.h"
+
+class FThreadTrackEvent;
+class FTimingEventSearchParameters;
+
+namespace UE::Insights { class FFilterConfigurator; }
+
+namespace UE::Insights::TimingProfiler
+{
+
+class FThreadTimingSharedState;
+
+////////////////////////////////////////////////////////////////////////////////////////////////////
+
+class FThreadTimingTrack : public FTimingEventsTrack
+{
+	INSIGHTS_DECLARE_RTTI(FThreadTimingTrack, FTimingEventsTrack)
+
+public:
+	typedef typename TraceServices::ITimeline<TraceServices::FTimingProfilerEvent>::FTimelineEventInfo TimelineEventInfo;
+
+	struct FPendingEventInfo
+	{
+		double StartTime;
+		double EndTime;
+		uint32 Depth;
+		uint32 TimerIndex;
+	};
+
+	explicit FThreadTimingTrack(FThreadTimingSharedState& InSharedState, const FString& InName, const TCHAR* InGroupName, uint32 InTimelineIndex, uint32 InThreadId)
+		: FTimingEventsTrack(InName)
+		, SharedState(InSharedState)
+		, GroupName(InGroupName)
+		, TimelineIndex(InTimelineIndex)
+		, ThreadId(InThreadId)
+	{
+	}
+
+	virtual ~FThreadTimingTrack() override {}
+
+	const TCHAR* GetGroupName() const { return GroupName; };
+
+	uint32 GetTimelineIndex() const { return TimelineIndex; }
+	//void SetTimelineIndex(uint32 InTimelineIndex) { TimelineIndex = InTimelineIndex; }
+
+	uint32 GetThreadId() const { return ThreadId; }
+	//void SetThreadId(uint32 InThreadId) { ThreadId = InThreadId; }
+
+	virtual void BuildDrawState(ITimingEventsTrackDrawStateBuilder& Builder, const ITimingTrackUpdateContext& Context) override;
+	virtual void BuildFilteredDrawState(ITimingEventsTrackDrawStateBuilder& Builder, const ITimingTrackUpdateContext& Context) override;
+
+	virtual void PostDraw(const ITimingTrackDrawContext& Context) const override;
+
+	virtual void InitTooltip(FTooltipDrawState& InOutTooltip, const ITimingEvent& InTooltipEvent) const override;
+
+	virtual const TSharedPtr<const ITimingEvent> GetEvent(float InPosX, float InPosY, const FTimingTrackViewport& Viewport) const override;
+	virtual const TSharedPtr<const ITimingEvent> SearchEvent(const FTimingEventSearchParameters& InSearchParameters) const override;
+
+	virtual void UpdateEventStats(ITimingEvent& InOutEvent) const override;
+	virtual void OnEventSelected(const ITimingEvent& InSelectedEvent) const override;
+	virtual void OnClipboardCopyEvent(const ITimingEvent& InSelectedEvent) const override;
+	virtual void BuildContextMenu(FMenuBuilder& MenuBuilder) override;
+
+	int32 GetDepthAt(double Time) const;
+
+	virtual void SetFilterConfigurator(TSharedPtr<FFilterConfigurator> InFilterConfigurator) override;
+
+	TSharedPtr<const ITimingEvent> FindMaxEventInstance(uint32 TimerId, double StartTime, double EndTime) const;
+	TSharedPtr<const ITimingEvent> FindMinEventInstance(uint32 TimerId, double StartTime, double EndTime) const;
+
+protected:
+	virtual bool HasCustomFilter() const override;
+
+private:
+	bool FindTimingProfilerEvent(const FThreadTrackEvent& InTimingEvent, TFunctionRef<void(double, double, uint32, const TraceServices::FTimingProfilerEvent&)> InFoundPredicate) const;
+	bool FindTimingProfilerEvent(const FTimingEventSearchParameters& InParameters, TFunctionRef<void(double, double, uint32, const TraceServices::FTimingProfilerEvent&)> InFoundPredicate) const;
+
+	void GetParentAndRoot(const FThreadTrackEvent& TimingEvent,
+						  TSharedPtr<FThreadTrackEvent>& OutParentTimingEvent,
+						  TSharedPtr<FThreadTrackEvent>& OutRootTimingEvent) const;
+
+	static void CreateFThreadTrackEventFromInfo(const TimelineEventInfo& InEventInfo, const TSharedRef<const FBaseTimingTrack> InTrack, int32 InDepth, TSharedPtr<FThreadTrackEvent> &OutTimingEvent);
+	static bool TimerIndexToTimerId(uint32 InTimerIndex, uint32 & OutTimerId);
+
+private:
+	FThreadTimingSharedState& SharedState;
+
+	TSharedPtr<FFilterConfigurator> FilterConfigurator;
+
+	const TCHAR* GroupName = nullptr;
+	uint32 TimelineIndex = 0;
+	uint32 ThreadId = 0;
+
+	// Search cache
+	mutable TTimingEventSearchCache<TraceServices::FTimingProfilerEvent> SearchCache;
+};
+
+////////////////////////////////////////////////////////////////////////////////////////////////////
+
+class FCpuTimingTrack : public FThreadTimingTrack
+{
+public:
+	explicit FCpuTimingTrack(FThreadTimingSharedState& InSharedState, const FString& InName, const TCHAR* InGroupName, uint32 InTimelineIndex, uint32 InThreadId)
+		: FThreadTimingTrack(InSharedState, InName, InGroupName, InTimelineIndex, InThreadId)
+	{
+	}
+};
+
+////////////////////////////////////////////////////////////////////////////////////////////////////
+
+class FGpuTimingTrack : public FThreadTimingTrack
+{
+public:
+	static constexpr uint32 Gpu1ThreadId = uint32('GPU1');
+	static constexpr uint32 Gpu2ThreadId = uint32('GPU2');
+
+public:
+	explicit FGpuTimingTrack(FThreadTimingSharedState& InSharedState, const FString& InName, const TCHAR* InGroupName, uint32 InTimelineIndex, uint32 InThreadId)
+		: FThreadTimingTrack(InSharedState, InName, InGroupName, InTimelineIndex, InThreadId)
+	{
+	}
+};
+
+////////////////////////////////////////////////////////////////////////////////////////////////////
+
+} // namespace UE::Insights::TimingProfiler
diff --git a/Engine/Source/Developer/TraceInsights/Private/Insights/TimingProfiler/ViewModels/TimerAggregator.cpp b/Engine/Source/Developer/TraceInsights/Private/Insights/TimingProfiler/ViewModels/TimerAggregator.cpp
new file mode 100644
--- /dev/null
+++ b/Engine/Source/Developer/TraceInsights/Private/Insights/TimingProfiler/ViewModels/TimerAggregator.cpp
@@ -0,0 +1,134 @@
+// Copyright Epic Games, Inc. All Rights Reserved.
+
+#include "TimerAggregator.h"
+
+// TraceInsights
+#include "Insights/TimingProfiler/TimingProfilerManager.h"
+#include "Insights/TimingProfiler/ViewModels/ThreadTimingSharedState.h"
+#include "Insights/TimingProfiler/Widgets/STimingProfilerWindow.h"
+#include "Insights/Widgets/STimingView.h"
+
+namespace UE::Insights::TimingProfiler
+{
+
+////////////////////////////////////////////////////////////////////////////////////////////////////
+// FTimerAggregationWorker
+////////////////////////////////////////////////////////////////////////////////////////////////////
+
+class FTimerAggregationWorker : public IStatsAggregationWorker
+{
+public:
+	FTimerAggregationWorker(TSharedPtr<const TraceServices::IAnalysisSession> InSession, double InStartTime, double InEndTime, const TSet<uint32>& InCpuThreads, bool bInIncludeGpuThread, ETraceFrameType InFrameType)
+		: Session(InSession)
+		, StartTime(InStartTime)
+		, EndTime(InEndTime)
+		, CpuThreads(InCpuThreads)
+		, bIncludeGpuThread(bInIncludeGpuThread)
+		, FrameType(InFrameType)
+		, ResultTable()
+	{
+	}
+
+	virtual ~FTimerAggregationWorker() {}
+
+	virtual void DoWork(TSharedPtr<TraceServices::FCancellationToken> CancellationToken) override;
+
+	TraceServices::ITable<TraceServices::FTimingProfilerAggregatedStats>* GetResultTable() const { return ResultTable.Get(); }
+	void ResetResults() { ResultTable.Reset(); }
+
+private:
+	TSharedPtr<const TraceServices::IAnalysisSession> Session;
+	double StartTime;
+	double EndTime;
+	TSet<uint32> CpuThreads;
+	bool bIncludeGpuThread;
+	ETraceFrameType FrameType;
+	TUniquePtr<TraceServices::ITable<TraceServices::FTimingProfilerAggregatedStats>> ResultTable;
+};
+
+////////////////////////////////////////////////////////////////////////////////////////////////////
+
+void FTimerAggregationWorker::DoWork(TSharedPtr<TraceServices::FCancellationToken> CancellationToken)
+{
+	if (Session.IsValid() && TraceServices::ReadTimingProfilerProvider(*Session.Get()))
+	{
+		// Suspend analysis in order to avoid write locks (ones blocked by the read lock below) to further block other read locks.
+		//FAnalysisSessionSuspensionScope SessionPauseScope(*Session.Get());
+
+		TraceServices::FAnalysisSessionReadScope SessionReadScope(*Session.Get());
+		const TraceServices::ITimingProfilerProvider& TimingProfilerProvider = *TraceServices::ReadTimingProfilerProvider(*Session.Get());
+
+		auto CpuThreadFilter = [this](uint32 ThreadId)
+		{
+			return CpuThreads.Contains(ThreadId);
+		};
+
+		TraceServices::FCreateAggreationParams Params;
+		Params.IntervalStart = StartTime;
+		Params.IntervalEnd = EndTime;
+		Params.CpuThreadFilter = CpuThreadFilter;
+		Params.IncludeGpu = bIncludeGpuThread;
+		Params.FrameType = FrameType;
+		Params.CancellationToken = CancellationToken;
+
+		ResultTable.Reset(TimingProfilerProvider.CreateAggregation(Params));
+	}
+}
+
+////////////////////////////////////////////////////////////////////////////////////////////////////
+// FTimerAggregator
+////////////////////////////////////////////////////////////////////////////////////////////////////
+
+IStatsAggregationWorker* FTimerAggregator::CreateWorker(TSharedPtr<const TraceServices::IAnalysisSession> InSession)
+{
+	bool bIsGpuTrackVisible = false;
+	TSet<uint32> CpuThreads;
+
+	TSharedPtr<STimingProfilerWindow> Wnd = FTimingProfilerManager::Get()->GetProfilerWindow();
+	if (Wnd.IsValid())
+	{
+		TSharedPtr<STimingView> TimingView = Wnd->GetTimingView();
+		if (TimingView.IsValid())
+		{
+			bIsGpuTrackVisible = TimingView->IsGpuTrackVisible();
+
+			TSharedPtr<FThreadTimingSharedState> ThreadTimingSharedState = TimingView->GetThreadTimingSharedState();
+			if (ThreadTimingSharedState.IsValid())
+			{
+				ThreadTimingSharedState->GetVisibleCpuThreads(CpuThreads);
+			}
+		}
+	}
+
+	return new FTimerAggregationWorker(InSession, GetIntervalStartTime(), GetIntervalEndTime(), CpuThreads, bIsGpuTrackVisible, FrameType);
+}
+
+////////////////////////////////////////////////////////////////////////////////////////////////////
+
+TraceServices::ITable<TraceServices::FTimingProfilerAggregatedStats>* FTimerAggregator::GetResultTable() const
+{
+	// It can only be called from OnFinishedCallback.
+	check(IsFinished());
+
+	FTimerAggregationWorker* Worker = (FTimerAggregationWorker*)GetWorker();
+	check(Worker != nullptr);
+
+	return Worker->GetResultTable();
+}
+
+////////////////////////////////////////////////////////////////////////////////////////////////////
+
+void FTimerAggregator::ResetResults()
+{
+	// It can only be called from OnFinishedCallback.
+	check(IsFinished());
+
+	FTimerAggregationWorker* Worker = (FTimerAggregationWorker*)GetWorker();
+	check(Worker != nullptr);
+
+	Worker->ResetResults();
+}
+
+////////////////////////////////////////////////////////////////////////////////////////////////////
+
+} // namespace UE::Insights::TimingProfiler
diff --git a/Engine/Source/Developer/TraceInsights/Private/Insights/TimingProfiler/ViewModels/TimerButterflyAggregator.cpp b/Engine/Source/Developer/TraceInsights/Private/Insights/TimingProfiler/ViewModels/TimerButterflyAggregator.cpp
new file mode 100644
--- /dev/null
+++ b/Engine/Source/Developer/TraceInsights/Private/Insights/TimingProfiler/ViewModels/TimerButterflyAggregator.cpp
@@ -0,0 +1,143 @@
+// Copyright Epic Games, Inc. All Rights Reserved.
+
+#include "TimerButterflyAggregator.h"
+
+// TraceServices
+#include "TraceServices/Model/AnalysisSession.h"
+#include "TraceServices/Model/TimingProfiler.h"
+
+// TraceInsights
+#include "Insights/TimingProfiler/TimingProfilerManager.h"
+#include "Insights/TimingProfiler/Tracks/ThreadTimingTrack.h"
+#include "Insights/TimingProfiler/ViewModels/ThreadTimingSharedState.h"
+#include "Insights/TimingProfiler/Widgets/STimingProfilerWindow.h"
+#include "Insights/Widgets/STimingView.h"
+
+namespace UE::Insights::TimingProfiler
+{
+
+////////////////////////////////////////////////////////////////////////////////////////////////////
+// FTimerButterflyAggregationWorker
+////////////////////////////////////////////////////////////////////////////////////////////////////
+
+class FTimerButterflyAggregationWorker : public IStatsAggregationWorker
+{
+public:
+	FTimerButterflyAggregationWorker(TSharedPtr<const TraceServices::IAnalysisSession> InSession, double InStartTime, double InEndTime, const TSet<uint32>& InCpuThreads, bool bInIncludeGpuThread)
+		: Session(InSession)
+		, StartTime(InStartTime)
+		, EndTime(InEndTime)
+		, CpuThreads(InCpuThreads)
+		, bIncludeGpuThread(bInIncludeGpuThread)
+		, ResultButterfly()
+	{
+	}
+
+	virtual ~FTimerButterflyAggregationWorker() {}
+
+	virtual void DoWork(TSharedPtr<TraceServices::FCancellationToken> CancellationToken) override;
+
+	TraceServices::ITimingProfilerButterfly* GetResultButterfly() const { return ResultButterfly.Get(); }
+	void ResetResults() { ResultButterfly.Reset(); }
+
+private:
+	TSharedPtr<const TraceServices::IAnalysisSession> Session;
+	double StartTime;
+	double EndTime;
+	TSet<uint32> CpuThreads;
+	bool bIncludeGpuThread;
+	TUniquePtr<TraceServices::ITimingProfilerButterfly> ResultButterfly;
+};
+
+////////////////////////////////////////////////////////////////////////////////////////////////////
+
+void FTimerButterflyAggregationWorker::DoWork(TSharedPtr<TraceServices::FCancellationToken> CancellationToken)
+{
+	if (Session.IsValid() && TraceServices::ReadTimingProfilerProvider(*Session.Get()))
+	{
+		// Suspend analysis in order to avoid write locks (ones blocked by the read lock below) to further block other read locks.
+		//FAnalysisSessionSuspensionScope SessionPauseScope(*Session.Get());
+
+		TraceServices::FAnalysisSessionReadScope SessionReadScope(*Session.Get());
+		const TraceServices::ITimingProfilerProvider& TimingProfilerProvider = *TraceServices::ReadTimingProfilerProvider(*Session.Get());
+
+		if (StartTime <= Session->GetDurationSeconds())
+		{
+			auto ThreadFilter = [this](uint32 ThreadId)
+			{
+				return CpuThreads.Contains(ThreadId);
+			};
+
+			double ClampedEndTime = FMath::Min(Session->GetDurationSeconds(), EndTime);
+			ResultButterfly.Reset(TimingProfilerProvider.CreateButterfly(StartTime, ClampedEndTime, ThreadFilter, bIncludeGpuThread));
+		}
+		else
+		{
+			// Exclude all threads to receive empty results in case infinite events are in the session.
+			auto ThreadFilter = [this](uint32 ThreadId)
+			{
+				return false;
+			};
+
+			ResultButterfly.Reset(TimingProfilerProvider.CreateButterfly(StartTime, EndTime, ThreadFilter, false));
+		}
+	}
+}
+
+////////////////////////////////////////////////////////////////////////////////////////////////////
+// FTimerButterflyAggregator
+////////////////////////////////////////////////////////////////////////////////////////////////////
+
+IStatsAggregationWorker* FTimerButterflyAggregator::CreateWorker(TSharedPtr<const TraceServices::IAnalysisSession> InSession)
+{
+	bool bIsGpuTrackVisible = false;
+	TSet<uint32> CpuThreads;
+
+	TSharedPtr<STimingProfilerWindow> Wnd = FTimingProfilerManager::Get()->GetProfilerWindow();
+	if (Wnd.IsValid())
+	{
+		TSharedPtr<STimingView> TimingView = Wnd->GetTimingView();
+		if (TimingView.IsValid())
+		{
+			bIsGpuTrackVisible = TimingView->IsGpuTrackVisible();
+
+			TSharedPtr<FThreadTimingSharedState> ThreadTimingSharedState = TimingView->GetThreadTimingSharedState();
+			if (ThreadTimingSharedState.IsValid())
+			{
+				ThreadTimingSharedState->GetVisibleCpuThreads(CpuThreads);
+			}
+		}
+	}
+
+	return new FTimerButterflyAggregationWorker(InSession, GetIntervalStartTime(), GetIntervalEndTime(), CpuThreads, bIsGpuTrackVisible);
+}
+
+////////////////////////////////////////////////////////////////////////////////////////////////////
+
+TraceServices::ITimingProfilerButterfly* FTimerButterflyAggregator::GetResultButterfly() const
+{
+	// It can only be called from OnFinishedCallback.
+	check(IsFinished());
+
+	FTimerButterflyAggregationWorker* Worker = (FTimerButterflyAggregationWorker*)GetWorker();
+	check(Worker != nullptr);
+
+	return Worker->GetResultButterfly();
+}
+
+////////////////////////////////////////////////////////////////////////////////////////////////////
+
+void FTimerButterflyAggregator::ResetResults()
+{
+	// It can only be called from OnFinishedCallback.
+	check(IsFinished());
+
+	FTimerButterflyAggregationWorker* Worker = (FTimerButterflyAggregationWorker*)GetWorker();
+	check(Worker != nullptr);
+
+	Worker->ResetResults();
+}
+
+////////////////////////////////////////////////////////////////////////////////////////////////////
+
+} // namespace UE::Insights::TimingProfiler
diff --git a/Engine/Source/Developer/TraceInsights/Private/Insights/TimingProfiler/Widgets/STimingProfilerToolbar.cpp b/Engine/Source/Developer/TraceInsights/Private/Insights/TimingProfiler/Widgets/STimingProfilerToolbar.cpp
new file mode 100644
--- /dev/null
+++ b/Engine/Source/Developer/TraceInsights/Private/Insights/TimingProfiler/Widgets/STimingProfilerToolbar.cpp
@@ -0,0 +1,129 @@
+// Copyright Epic Games, Inc. All Rights Reserved.
+
+#include "STimingProfilerToolbar.h"
+
+#include "Framework/MultiBox/MultiBoxExtender.h"
+#include "Framework/MultiBox/MultiBoxBuilder.h"
+#include "Widgets/SBoxPanel.h"
+
+// TraceInsights
+#include "Insights/InsightsCommands.h"
+#include "Insights/InsightsManager.h"
+#include "Insights/InsightsStyle.h"
+#include "Insights/TimingProfiler/TimingProfilerCommands.h"
+
+////////////////////////////////////////////////////////////////////////////////////////////////////
+
+#define LOCTEXT_NAMESPACE "UE::Insights::TimingProfiler"
+
+namespace UE::Insights::TimingProfiler
+{
+
+////////////////////////////////////////////////////////////////////////////////////////////////////
+
+STimingProfilerToolbar::STimingProfilerToolbar()
+{
+}
+
+////////////////////////////////////////////////////////////////////////////////////////////////////
+
+STimingProfilerToolbar::~STimingProfilerToolbar()
+{
+}
+
+////////////////////////////////////////////////////////////////////////////////////////////////////
+
+void STimingProfilerToolbar::Construct(const FArguments& InArgs)
+{
+	struct Local
+	{
+		static void FillViewToolbar(FToolBarBuilder& ToolbarBuilder, const FArguments &InArgs)
+		{
+			ToolbarBuilder.BeginSection("View");
+			{
+				ToolbarBuilder.AddToolBarButton(FTimingProfilerCommands::Get().ToggleFramesTrackVisibility,
+					NAME_None, TAttribute<FText>(), TAttribute<FText>(),
+					FSlateIcon(FInsightsStyle::GetStyleSetName(), "Icons.FramesTrack.ToolBar"));
+				ToolbarBuilder.AddToolBarButton(FTimingProfilerCommands::Get().ToggleTimingViewVisibility,
+					NAME_None, TAttribute<FText>(), TAttribute<FText>(),
+					FSlateIcon(FInsightsStyle::GetStyleSetName(), "Icons.TimingView.ToolBar"));
+				ToolbarBuilder.AddToolBarButton(FTimingProfilerCommands::Get().ToggleTimersViewVisibility,
+					NAME_None, TAttribute<FText>(), TAttribute<FText>(),
+					FSlateIcon(FInsightsStyle::GetStyleSetName(), "Icons.TimersView.ToolBar"));
+				ToolbarBuilder.AddToolBarButton(FTimingProfilerCommands::Get().ToggleCallersTreeViewVisibility,
+					NAME_None, TAttribute<FText>(), TAttribute<FText>(),
+					FSlateIcon(FInsightsStyle::GetStyleSetName(), "Icons.CallersView.ToolBar"));
+				ToolbarBuilder.AddToolBarButton(FTimingProfilerCommands::Get().ToggleCalleesTreeViewVisibility,
+					NAME_None, TAttribute<FText>(), TAttribute<FText>(),
+					FSlateIcon(FInsightsStyle::GetStyleSetName(), "Icons.CalleesView.ToolBar"));
+				ToolbarBuilder.AddToolBarButton(FTimingProfilerCommands::Get().ToggleStatsCountersViewVisibility,
+					NAME_None, TAttribute<FText>(), TAttribute<FText>(),
+					FSlateIcon(FInsightsStyle::GetStyleSetName(), "Icons.CountersView.ToolBar"));
+				ToolbarBuilder.AddToolBarButton(FTimingProfilerCommands::Get().ToggleLogViewVisibility,
+					NAME_None, TAttribute<FText>(), TAttribute<FText>(),
+					FSlateIcon(FInsightsStyle::GetStyleSetName(), "Icons.LogView.ToolBar"));
+			}
+			ToolbarBuilder.EndSection();
+
+			if (InArgs._ToolbarExtender.IsValid())
+			{
+				InArgs._ToolbarExtender->Apply("MainToolbar", EExtensionHook::First, ToolbarBuilder);
+			}
+		}
+
+		static void FillRightSideToolbar(FToolBarBuilder& ToolbarBuilder, const FArguments &InArgs)
+		{
+			ToolbarBuilder.BeginSection("Debug");
+			{
+				ToolbarBuilder.AddToolBarButton(FInsightsCommands::Get().ToggleDebugInfo,
+					NAME_None, FText(), TAttribute<FText>(),
+					FSlateIcon(FInsightsStyle::GetStyleSetName(), "Icons.Debug.ToolBar"));
+			}
+			ToolbarBuilder.EndSection();
+
+			if (InArgs._ToolbarExtender.IsValid())
+			{
+				InArgs._ToolbarExtender->Apply("RightSideToolbar", EExtensionHook::First, ToolbarBuilder);
+			}
+		}
+	};
+
+	TSharedPtr<FUICommandList> CommandList = FInsightsManager::Get()->GetCommandList();
+
+	FSlimHorizontalToolBarBuilder ToolbarBuilder(CommandList.ToSharedRef(), FMultiBoxCustomization::None);
+	ToolbarBuilder.SetStyle(&FInsightsStyle::Get(), "PrimaryToolbar");
+	Local::FillViewToolbar(ToolbarBuilder, InArgs);
+
+	FSlimHorizontalToolBarBuilder RightSideToolbarBuilder(CommandList.ToSharedRef(), FMultiBoxCustomization::None);
+	RightSideToolbarBuilder.SetStyle(&FInsightsStyle::Get(), "PrimaryToolbar");
+	Local::FillRightSideToolbar(RightSideToolbarBuilder, InArgs);
+
+	ChildSlot
+	[
+		SNew(SHorizontalBox)
+
+		+ SHorizontalBox::Slot()
+		.HAlign(HAlign_Fill)
+		.VAlign(VAlign_Center)
+		.FillWidth(1.0)
+		.Padding(0.0f)
+		[
+			ToolbarBuilder.MakeWidget()
+		]
+
+		+ SHorizontalBox::Slot()
+		.HAlign(HAlign_Right)
+		.VAlign(VAlign_Center)
+		.AutoWidth()
+		.Padding(0.0f)
+		[
+			RightSideToolbarBuilder.MakeWidget()
+		]
+	];
+}
+
+////////////////////////////////////////////////////////////////////////////////////////////////////
+
+} // namespace UE::Insights::TimingProfiler
+
+#undef LOCTEXT_NAMESPACE
diff --git a/Engine/Source/Developer/TraceInsights/Private/Insights/TraceInsightsModule.cpp b/Engine/Source/Developer/TraceInsights/Private/Insights/TraceInsightsModule.cpp
--- a/Engine/Source/Developer/TraceInsights/Private/Insights/TraceInsightsModule.cpp
+++ b/Engine/Source/Developer/TraceInsights/Private/Insights/TraceInsightsModule.cpp
@@ -36,7 +36,7 @@
 #include "Insights/NetworkingProfiler/NetworkingProfilerManager.h"
 #include "Insights/TaskGraphProfiler/TaskGraphProfilerManager.h"
 #include "Insights/Tests/InsightsTestRunner.h"
-#include "Insights/TimingProfilerManager.h"
+#include "Insights/TimingProfiler/TimingProfilerManager.h"
 #include "Insights/Widgets/SStartPageWindow.h"
 
 ////////////////////////////////////////////////////////////////////////////////////////////////////
diff --git a/Engine/Source/Developer/TraceInsights/Private/Insights/ViewModels/ThreadTrackEvent.cpp b/Engine/Source/Developer/TraceInsights/Private/Insights/ViewModels/ThreadTrackEvent.cpp
--- a/Engine/Source/Developer/TraceInsights/Private/Insights/ViewModels/ThreadTrackEvent.cpp
+++ b/Engine/Source/Developer/TraceInsights/Private/Insights/ViewModels/ThreadTrackEvent.cpp
@@ -2,8 +2,6 @@
 
 #include "Insights/ViewModels/ThreadTrackEvent.h"
 
-#include "Insights/ViewModels/GraphTrack.h"
-
 ////////////////////////////////////////////////////////////////////////////////////////////////////
 
 INSIGHTS_IMPLEMENT_RTTI(FThreadTrackEvent)
diff --git a/Engine/Source/Developer/TraceInsights/Private/Insights/ViewModels/TimingGraphTrack.cpp b/Engine/Source/Developer/TraceInsights/Private/Insights/ViewModels/TimingGraphTrack.cpp
--- a/Engine/Source/Developer/TraceInsights/Private/Insights/ViewModels/TimingGraphTrack.cpp
+++ b/Engine/Source/Developer/TraceInsights/Private/Insights/ViewModels/TimingGraphTrack.cpp
@@ -2,6 +2,8 @@
 
 #include "TimingGraphTrack.h"
 
+#include "Algo/BinarySearch.h"
+
 // TraceServices
 #include "TraceServices/Model/Counters.h"
 #include "TraceServices/Model/Frames.h"
@@ -13,15 +15,16 @@
 
 // TraceInsights
 #include "Insights/InsightsManager.h"
-#include "Insights/TimingProfilerManager.h"
+#include "Insights/TimingProfiler/TimingProfilerManager.h"
+#include "Insights/TimingProfiler/Tracks/ThreadTimingTrack.h"
+#include "Insights/TimingProfiler/ViewModels/FrameStatsHelper.h"
+#include "Insights/TimingProfiler/ViewModels/ThreadTimingSharedState.h"
+#include "Insights/TimingProfiler/Widgets/STimersView.h"
+#include "Insights/TimingProfiler/Widgets/STimingProfilerWindow.h"
 #include "Insights/ViewModels/AxisViewportDouble.h"
-#include "Insights/ViewModels/FrameStatsHelper.h"
 #include "Insights/ViewModels/GraphTrackBuilder.h"
 #include "Insights/ViewModels/ITimingViewDrawHelper.h"
-#include "Insights/ViewModels/ThreadTimingTrack.h"
 #include "Insights/ViewModels/TimingTrackViewport.h"
-#include "Insights/Widgets/STimersView.h"
-#include "Insights/Widgets/STimingProfilerWindow.h"
 #include "Insights/Widgets/STimingView.h"
 
 #include <limits>
@@ -181,10 +184,11 @@ FTimingGraphTrack::~FTimingGraphTrack()
 		}
 	}
 
+	using namespace UE::Insights::TimingProfiler;
+
 	TSharedPtr<STimersView> TimersView;
 	if (bNotifyTimersOnDestruction)
 	{
-		using namespace UE::Insights::TimingProfiler;
 		TSharedPtr<STimingProfilerWindow> ProfilerWindow = FTimingProfilerManager::Get()->GetProfilerWindow();
 		if (ProfilerWindow.IsValid())
 		{
@@ -192,7 +196,7 @@ FTimingGraphTrack::~FTimingGraphTrack()
 		}
 	}
 
-	if(TimersView)
+	if (TimersView)
 	{
 		for (const TSharedPtr<FGraphSeries>& Series : AllSeries)
 		{
@@ -223,7 +227,7 @@ void FTimingGraphTrack::Update(const ITimingTrackUpdateContext& Context)
 		{
 			auto OnTrackAddedRemovedLamda = [this](const TSharedPtr<const FBaseTimingTrack> Track)
 			{
-				if (Track->Is<FThreadTimingTrack>())
+				if (Track->Is<UE::Insights::TimingProfiler::FThreadTimingTrack>())
 				{
 					// If there are more series than the default frame series.
 					if (this->AllSeries.Num() > ETraceFrameType::TraceFrameType_Count)
@@ -597,6 +601,8 @@ void FTimingGraphTrack::RemoveFrameStatsTimerSeries(uint32 TimerId, ETraceFrameT
 
 void FTimingGraphTrack::UpdateFrameStatsTimerSeries(FTimingGraphSeries& Series, const FTimingTrackViewport& Viewport)
 {
+	using namespace UE::Insights::TimingProfiler;
+
 	FGraphTrackBuilder Builder(*this, Series, Viewport);
 	TSharedPtr<const TraceServices::IAnalysisSession> Session = FInsightsManager::Get()->GetSession();
 	if (Session.IsValid())
@@ -622,29 +628,29 @@ void FTimingGraphTrack::UpdateFrameStatsTimerSeries(FTimingGraphSeries& Series,
 
 			FramesProvider.EnumerateFrames(Series.FrameType, 0ull, FrameCount, [&Series](const TraceServices::FFrame& Frame)
 				{
-					Insights::FFrameStatsCachedEvent Event;
+					FFrameStatsCachedEvent Event;
 					Event.FrameStartTime = Frame.StartTime;
 					Event.FrameEndTime = Frame.EndTime;
 					Event.Duration.store(0.0f);
 					Series.FrameStatsCachedEvents.Add(Event);
 				});
 
-			Insights::FFrameStatsHelper::ComputeFrameStatsForTimer(Series.FrameStatsCachedEvents, Series.TimerId, VisibleTimelines);
+			FFrameStatsHelper::ComputeFrameStatsForTimer(Series.FrameStatsCachedEvents, Series.TimerId, VisibleTimelines);
 		}
 
-		int32 StartIndex = Algo::UpperBoundBy(Series.FrameStatsCachedEvents, Viewport.GetStartTime(), &Insights::FFrameStatsCachedEvent::FrameStartTime);
+		int32 StartIndex = Algo::UpperBoundBy(Series.FrameStatsCachedEvents, Viewport.GetStartTime(), &FFrameStatsCachedEvent::FrameStartTime);
 		if (StartIndex > 0)
 		{
 			StartIndex--;
 		}
-		int32 EndIndex = Algo::UpperBoundBy(Series.FrameStatsCachedEvents, Viewport.GetEndTime(), &Insights::FFrameStatsCachedEvent::FrameStartTime);
+		int32 EndIndex = Algo::UpperBoundBy(Series.FrameStatsCachedEvents, Viewport.GetEndTime(), &FFrameStatsCachedEvent::FrameStartTime);
 		if (EndIndex < Series.FrameStatsCachedEvents.Num())
 		{
 			EndIndex++;
 		}
 		for (int32 Index = StartIndex; Index < EndIndex; ++Index)
 		{
-			const Insights::FFrameStatsCachedEvent& Entry = Series.FrameStatsCachedEvents[Index];
+			const FFrameStatsCachedEvent& Entry = Series.FrameStatsCachedEvents[Index];
 			Builder.AddEvent(Entry.FrameStartTime, Entry.Duration.load(), Entry.Duration.load());
 		}
 	}
@@ -807,7 +813,7 @@ void FTimingGraphTrack::GetVisibleTimelineIndexes(TSet<uint32>& TimelineIndexes)
 		return;
 	}
 
-	TSharedPtr<FThreadTimingSharedState> ThreadSharedState = TimingViewPtr->GetThreadTimingSharedState();
+	TSharedPtr<UE::Insights::TimingProfiler::FThreadTimingSharedState> ThreadSharedState = TimingViewPtr->GetThreadTimingSharedState();
 	ThreadSharedState->GetVisibleTimelineIndexes(TimelineIndexes);
 }
 
diff --git a/Engine/Source/Developer/TraceInsights/Private/Insights/ViewModels/TimingGraphTrack.h b/Engine/Source/Developer/TraceInsights/Private/Insights/ViewModels/TimingGraphTrack.h
--- a/Engine/Source/Developer/TraceInsights/Private/Insights/ViewModels/TimingGraphTrack.h
+++ b/Engine/Source/Developer/TraceInsights/Private/Insights/ViewModels/TimingGraphTrack.h
@@ -2,19 +2,18 @@
 
 #pragma once
 
-#include "CoreMinimal.h"
+#include "CoreTypes.h"
 
 // TraceInsights
 #include "Insights/ViewModels/GraphSeries.h"
 #include "Insights/ViewModels/GraphTrack.h"
 
-namespace Insights
+namespace UE::Insights::TimingProfiler
 {
 	struct FFrameStatsCachedEvent;
+	class STimingView;
 }
 
-namespace UE::Insights::TimingProfiler { class STimingView; }
-
 ////////////////////////////////////////////////////////////////////////////////////////////////////
 
 /** The delegate to be invoked when a series visibility is changed. */
@@ -61,7 +60,7 @@ public:
 	ETraceFrameType FrameType;
 	double CachedSessionDuration;
 	TArray<FSimpleTimingEvent> CachedEvents; // used by Timer series
-	TArray<Insights::FFrameStatsCachedEvent> FrameStatsCachedEvents; // used by Frame Stats Timer series
+	TArray<UE::Insights::TimingProfiler::FFrameStatsCachedEvent> FrameStatsCachedEvents; // used by Frame Stats Timer series
 	uint32 CachedTimelinesNum = 0; // the number of timelines used to gather the data
 
 	bool bIsTime; // the unit for values is [second]
diff --git a/Engine/Source/Developer/TraceInsights/Private/Insights/ViewModels/TimingViewDrawHelper.cpp b/Engine/Source/Developer/TraceInsights/Private/Insights/ViewModels/TimingViewDrawHelper.cpp
--- a/Engine/Source/Developer/TraceInsights/Private/Insights/ViewModels/TimingViewDrawHelper.cpp
+++ b/Engine/Source/Developer/TraceInsights/Private/Insights/ViewModels/TimingViewDrawHelper.cpp
@@ -13,8 +13,8 @@
 
 // TraceInsights
 #include "Insights/InsightsStyle.h"
+#include "Insights/TimingProfiler/Tracks/MarkersTimingTrack.h"
 #include "Insights/ViewModels/DrawHelpers.h"
-#include "Insights/ViewModels/MarkersTimingTrack.h"
 #include "Insights/ViewModels/TimingEvent.h"
 #include "Insights/ViewModels/TimingEventsTrack.h"
 #include "Insights/ViewModels/TimingTrackViewport.h"
diff --git a/Engine/Source/Developer/TraceInsights/Private/Insights/Widgets/SLogView.cpp b/Engine/Source/Developer/TraceInsights/Private/Insights/Widgets/SLogView.cpp
--- a/Engine/Source/Developer/TraceInsights/Private/Insights/Widgets/SLogView.cpp
+++ b/Engine/Source/Developer/TraceInsights/Private/Insights/Widgets/SLogView.cpp
@@ -33,10 +33,10 @@
 #include "Insights/InsightsManager.h"
 #include "Insights/InsightsStyle.h"
 #include "Insights/Log.h"
-#include "Insights/TimingProfilerManager.h"
-#include "Insights/ViewModels/MarkersTimingTrack.h" // for FTimeMarkerTrackBuilder::GetColorBy*
+#include "Insights/TimingProfiler/TimingProfilerManager.h"
+#include "Insights/TimingProfiler/Tracks/MarkersTimingTrack.h" // for FTimeMarkerTrackBuilder::GetColorBy*
+#include "Insights/TimingProfiler/Widgets/STimingProfilerWindow.h"
 #include "Insights/ViewModels/TimingViewDrawHelper.h"
-#include "Insights/Widgets/STimingProfilerWindow.h"
 #include "Insights/Widgets/STimingView.h"
 
 ////////////////////////////////////////////////////////////////////////////////////////////////////
@@ -44,7 +44,7 @@
 #define LOCTEXT_NAMESPACE "SLogView"
 
 ////////////////////////////////////////////////////////////////////////////////////////////////////
-// FTableTreeViewCommands
+// FLogViewCommands
 ////////////////////////////////////////////////////////////////////////////////////////////////////
 
 class FLogViewCommands : public TCommands<FLogViewCommands>
@@ -379,7 +379,7 @@ public:
 		if (ParentWidgetPin.IsValid() && LogMessagePin.IsValid())
 		{
 			FLogMessageRecord& CacheEntry = ParentWidgetPin->GetCache().Get(LogMessagePin->GetIndex());
-			return FSlateColor(FTimeMarkerTrackBuilder::GetColorByVerbosity(CacheEntry.GetVerbosity()));
+			return FSlateColor(UE::Insights::TimingProfiler::FTimeMarkerTrackBuilder::GetColorByVerbosity(CacheEntry.GetVerbosity()));
 		}
 		else
 		{
@@ -412,7 +412,7 @@ public:
 		if (ParentWidgetPin.IsValid() && LogMessagePin.IsValid())
 		{
 			FLogMessageRecord& CacheEntry = ParentWidgetPin->GetCache().Get(LogMessagePin->GetIndex());
-			return FSlateColor(FTimeMarkerTrackBuilder::GetColorByCategory(CacheEntry.GetCategory()));
+			return FSlateColor(UE::Insights::TimingProfiler::FTimeMarkerTrackBuilder::GetColorByCategory(CacheEntry.GetCategory()));
 		}
 		else
 		{
@@ -1405,7 +1405,7 @@ void SLogView::CreateVerbosityThresholdMenuSection(FMenuBuilder& MenuBuilder)
 			[
 				SNew(STextBlock)
 				.Text(Threshold.Label)
-				.ColorAndOpacity(FSlateColor(FTimeMarkerTrackBuilder::GetColorByVerbosity(Threshold.Verbosity)))
+				.ColorAndOpacity(FSlateColor(UE::Insights::TimingProfiler::FTimeMarkerTrackBuilder::GetColorByVerbosity(Threshold.Verbosity)))
 				.ShadowColorAndOpacity(FLinearColor(0.02f, 0.02f, 0.02f, 1.0f))
 				.ShadowOffset(FVector2D(1.0f, 1.0f))
 			]
@@ -1421,7 +1421,7 @@ void SLogView::CreateVerbosityThresholdMenuSection(FMenuBuilder& MenuBuilder)
 				[
 					SNew(STextBlock)
 					.Text(Threshold.Label)
-					.ColorAndOpacity(FSlateColor(FTimeMarkerTrackBuilder::GetColorByVerbosity(Threshold.Verbosity)))
+					.ColorAndOpacity(FSlateColor(UE::Insights::TimingProfiler::FTimeMarkerTrackBuilder::GetColorByVerbosity(Threshold.Verbosity)))
 					.ShadowColorAndOpacity(FLinearColor(0.02f, 0.02f, 0.02f, 1.0f))
 					.ShadowOffset(FVector2D(1.0f, 1.0f))
 				]
@@ -1502,7 +1502,7 @@ void SLogView::CreateCategoriesFilterMenuSection(FMenuBuilder& MenuBuilder)
 			.Text(FText::AsCultureInvariant(CategoryString))
 			.ShadowColorAndOpacity(FLinearColor(0.05f, 0.05f, 0.05f, 1.0f))
 			.ShadowOffset(FVector2D(1.0f, 1.0f))
-			.ColorAndOpacity(FSlateColor(FTimeMarkerTrackBuilder::GetColorByCategory(*CategoryString)));
+			.ColorAndOpacity(FSlateColor(UE::Insights::TimingProfiler::FTimeMarkerTrackBuilder::GetColorByCategory(*CategoryString)));
 
 		MenuBuilder.AddMenuEntry
 		(
