[Insights]

diff --git a/Engine/Source/Developer/TraceInsights/Private/Insights/Widgets/STimingView.cpp b/Engine/Source/Developer/TraceInsights/Private/Insights/Widgets/STimingView.cpp
--- a/Engine/Source/Developer/TraceInsights/Private/Insights/Widgets/STimingView.cpp
+++ b/Engine/Source/Developer/TraceInsights/Private/Insights/Widgets/STimingView.cpp
@@ -34,6 +34,9 @@
 #include "Widgets/SOverlay.h"
 #include "Widgets/Text/STextBlock.h"
 
+// TraceServices
+#include "TraceServices/Model/LoadTimeProfiler.h"
+
 // TraceInsightsCore
 #include "InsightsCore/Common/PaintUtils.h"
 #include "InsightsCore/Common/Stopwatch.h"
@@ -49,34 +52,39 @@
 #include "Insights/InsightsStyle.h"
 #include "Insights/ITimingViewExtender.h"
 #include "Insights/LoadingProfiler/LoadingProfilerManager.h"
+#include "Insights/LoadingProfiler/Tracks/LoadingTimingTrack.h"
+#include "Insights/LoadingProfiler/ViewModels/LoadingSharedState.h"
 #include "Insights/LoadingProfiler/Widgets/SLoadingProfilerWindow.h"
 #include "Insights/Log.h"
 #include "Insights/TaskGraphProfiler/TaskGraphProfilerManager.h"
 #include "Insights/Tests/TimingProfilerTests.h"
+#include "Insights/TimingProfiler/TimingProfilerManager.h"
+#include "Insights/TimingProfiler/Tracks/FileActivityTimingTrack.h"
+#include "Insights/TimingProfiler/Tracks/MarkersTimingTrack.h"
+#include "Insights/TimingProfiler/Tracks/RegionsTimingTrack.h"
+#include "Insights/TimingProfiler/Tracks/ThreadTimingTrack.h"
+#include "Insights/TimingProfiler/Tracks/TimeRulerTrack.h"
+#include "Insights/TimingProfiler/ViewModels/FileActivitySharedState.h"
+#include "Insights/TimingProfiler/ViewModels/FrameTimingTrack.h"
+#include "Insights/TimingProfiler/ViewModels/ThreadTimingSharedState.h"
+#include "Insights/TimingProfiler/ViewModels/TimeMarker.h"
+#include "Insights/TimingProfiler/ViewModels/TimerFilters.h"
+#include "Insights/TimingProfiler/ViewModels/TimingRegionsSharedState.h"
+#include "Insights/TimingProfiler/Widgets/SStatsView.h"
+#include "Insights/TimingProfiler/Widgets/STimersView.h"
+#include "Insights/TimingProfiler/Widgets/STimingProfilerWindow.h"
 #include "Insights/TimingProfilerCommon.h"
-#include "Insights/TimingProfilerManager.h"
 #include "Insights/ViewModels/BaseTimingTrack.h"
 #include "Insights/ViewModels/DrawHelpers.h"
 #include "Insights/ViewModels/EventNameFilterValueConverter.h"
-#include "Insights/ViewModels/FileActivityTimingTrack.h"
-#include "Insights/ViewModels/FrameTimingTrack.h"
 #include "Insights/ViewModels/GraphSeries.h"
 #include "Insights/ViewModels/GraphTrack.h"
-#include "Insights/ViewModels/LoadingTimingTrack.h"
-#include "Insights/ViewModels/MarkersTimingTrack.h"
 #include "Insights/ViewModels/QuickFind.h"
-#include "Insights/ViewModels/RegionsTimingTrack.h"
-#include "Insights/ViewModels/ThreadTimingTrack.h"
 #include "Insights/ViewModels/ThreadTrackEvent.h"
-#include "Insights/ViewModels/TimerFilters.h"
-#include "Insights/ViewModels/TimeRulerTrack.h"
 #include "Insights/ViewModels/TimingEventSearch.h"
 #include "Insights/ViewModels/TimingGraphTrack.h"
 #include "Insights/ViewModels/TimingViewDrawHelper.h"
 #include "Insights/Widgets/SQuickFind.h"
-#include "Insights/Widgets/SStatsView.h"
-#include "Insights/Widgets/STimersView.h"
-#include "Insights/Widgets/STimingProfilerWindow.h"
 #include "Insights/Widgets/STimingViewTrackList.h"
 
 #include <limits>
@@ -116,9 +124,9 @@ STimingView::STimingView()
 	: bScrollableTracksOrderIsDirty(false)
 	, FrameSharedState(MakeShared<FFrameSharedState>(this))
 	, ThreadTimingSharedState(MakeShared<FThreadTimingSharedState>(this))
-	, LoadingSharedState(MakeShared<FLoadingSharedState>(this))
+	, LoadingSharedState(MakeShared<LoadingProfiler::FLoadingSharedState>(this))
 	, FileActivitySharedState(MakeShared<FFileActivitySharedState>(this))
-	, TimingRegionsSharedState(MakeShared<::Insights::FTimingRegionsSharedState>(this))
+	, TimingRegionsSharedState(MakeShared<FTimingRegionsSharedState>(this))
 	, TimeRulerTrack(MakeShared<FTimeRulerTrack>())
 	, DefaultTimeMarker(MakeShared<FTimeMarker>())
 	, MarkersTrack(MakeShared<FMarkersTimingTrack>())
@@ -4853,7 +4861,7 @@ void STimingView::CreateDepthLimitMenu(FMenuBuilder& MenuBuilder)
 {
 	MenuBuilder.BeginSection("DepthLimit", LOCTEXT("ContextMenu_Section_DepthLimit", "Depth Limit"));
 	{
-		// Note: We use the custom AddMenuEntry in order to set the same key binding text for multiple menu items.
+		// Note: We use the custom FInsightsMenuBuilder::AddMenuEntry in order to set the same key binding text for multiple menu items.
 
 		FInsightsMenuBuilder::AddMenuEntry(MenuBuilder,
 			FUIAction(
@@ -4894,7 +4902,11 @@ FText STimingView::GetEventDepthLimitKeybindingText(uint32 DepthLimit) const
 {
 	uint32 CurrentDepthLimit = FTimingProfilerManager::Get()->GetEventDepthLimit();
 	uint32 NextDepthLimit = GetNextEventDepthLimit(CurrentDepthLimit);
-	return DepthLimit == NextDepthLimit ? LOCTEXT("DepthLimitKeybinding", "X") : FText::GetEmpty();
+	if (DepthLimit == NextDepthLimit)
+	{
+		return FInputChord(EKeys::X).GetInputText().ToUpper();
+	}
+	return FText::GetEmpty();
 }
 
 ////////////////////////////////////////////////////////////////////////////////////////////////////
@@ -5107,8 +5119,8 @@ void STimingView::QuickFind_Execute()
 			nullptr,
 			FFilterService::Get()->GetIntegerOperators()));
 
-		NewFilterConfigurator->Add(MakeShared<::Insights::FTimerNameFilter>());
-		NewFilterConfigurator->Add(MakeShared<::Insights::FMetadataFilter>());
+		NewFilterConfigurator->Add(MakeShared<FTimerNameFilter>());
+		NewFilterConfigurator->Add(MakeShared<FMetadataFilter>());
 
 		for (Timing::ITimingViewExtender* Extender : GetExtenders())
 		{
diff --git a/Engine/Source/Developer/TraceInsights/Private/Insights/Widgets/STimingView.h b/Engine/Source/Developer/TraceInsights/Private/Insights/Widgets/STimingView.h
--- a/Engine/Source/Developer/TraceInsights/Private/Insights/Widgets/STimingView.h
+++ b/Engine/Source/Developer/TraceInsights/Private/Insights/Widgets/STimingView.h
@@ -19,8 +19,8 @@
 // TraceInsights
 #include "Insights/Config.h"
 #include "Insights/ITimingViewSession.h"
+#include "Insights/TimingProfiler/ViewModels/TimerNode.h"
 #include "Insights/ViewModels/BaseTimingTrack.h"
-#include "Insights/ViewModels/TimerNode.h"
 #include "Insights/ViewModels/TimingEvent.h"
 #include "Insights/ViewModels/TimingEventsTrack.h"
 #include "Insights/ViewModels/TimingTrackViewport.h"
@@ -38,17 +38,8 @@ class SOverlay;
 class SScrollBar;
 class FSpawnTabArgs;
 
-class FFileActivitySharedState;
-class FFrameSharedState;
-class FLoadingSharedState;
-class FMarkersTimingTrack;
-class FThreadTimingSharedState;
-class FTimingGraphTrack;
-class FTimingViewDrawHelper;
-
 namespace Insights
 {
-	class FTimingRegionsSharedState;
 	class FQuickFind;
 	class SQuickFind;
 }
@@ -64,11 +55,25 @@ namespace UE::Insights::Timing
 	class ITimingViewExtender;
 }
 
+class FTimingGraphTrack;
+class FTimingViewDrawHelper;
+
+namespace UE::Insights::LoadingProfiler
+{
+	class FLoadingSharedState;
+}
+
 namespace UE::Insights::TimingProfiler
 {
 
-class FTimeMarker;
+class FFrameSharedState;
+class FThreadTimingSharedState;
+class FTimingRegionsSharedState;
+class FFileActivitySharedState;
+
 class FTimeRulerTrack;
+class FTimeMarker;
+class FMarkersTimingTrack;
 
 enum class ESelectEventType : uint32
 {
@@ -119,7 +124,7 @@ public:
 
 	TSharedPtr<FFrameSharedState> GetFrameSharedState() const { return FrameSharedState; }
 	TSharedPtr<FThreadTimingSharedState> GetThreadTimingSharedState() const { return ThreadTimingSharedState; }
-	TSharedPtr<FLoadingSharedState> GetLoadingSharedState() const { return LoadingSharedState; }
+	TSharedPtr<LoadingProfiler::FLoadingSharedState> GetLoadingSharedState() const { return LoadingSharedState; }
 	TSharedPtr<FFileActivitySharedState> GetFileActivitySharedState() const { return FileActivitySharedState; }
 
 	void HideAllDefaultTracks();
@@ -591,13 +596,13 @@ protected:
 	TSharedPtr<FThreadTimingSharedState> ThreadTimingSharedState;
 
 	// Shared state for Asset Loading tracks
-	TSharedPtr<FLoadingSharedState> LoadingSharedState;
+	TSharedPtr<LoadingProfiler::FLoadingSharedState> LoadingSharedState;
 
 	// Shared state for File Activity (I/O) tracks
 	TSharedPtr<FFileActivitySharedState> FileActivitySharedState;
 
 	// Shared state for Regions tracks
-	TSharedPtr<::Insights::FTimingRegionsSharedState> TimingRegionsSharedState;
+	TSharedPtr<FTimingRegionsSharedState> TimingRegionsSharedState;
 	
 	////////////////////////////////////////////////////////////
 
