[Insights]
  - Added new TraceInsightsFrontend module. Moved all functionality of Unreal Insights Frontend (i.e. Session Browser) from TraceInsights module to this new module.
  - Split the functional tests between TraceInsightsFrontend and TraceInsights. Also updated some of the tests to reduce dependency to InsightsManager. Tests moved to TraceInsightsFrontend module:
        * "System.Insights.Hub.MemoryInsights" (MemoryInsightsTests.cpp)
        * "System.Insights.Hub.SessionBrowser" (TraceStoreWindowTests.cpp)
        * "System.Insights.Hub.UnrealTraceServer.Starting" (UnrealTraceServerTests.cpp)
  - TraceInsights: Updated InsightsStyle to inherit the InsightsCoreStyle (from TraceInsightsCore module). Moved more resources to FInsightsCoreStyle. Removed duplicates between TraceInsights and TraceInsightsCore modules. Removed "WhiteBrush" from FInsightsStyle (can be used the one from FAppStyle / FCoreStyle instead).
  - Moved and renamed FTraceServerControl from "Insights/Widgets/STraceServerControl.h" to "InsightsFrontend/StoreService/TraceServerControl.h". TraceUtilities now depends on TraceInsightsFrontend module (instead of TraceInsights module).
  - Added FStoreConnection in TraceAnalysis module. It manages a connection to a trace server (FStoreClient + critical section). This removes duplicated code between FInsightsManager (TraceInsights module) and in STraceStoreWindows (TraceInsightsFrontend module).
  - Moved MessageDialogUtils.h/.cpp from TraceInsights to TraceInsightsCore module (InsightsCore/Common/MessageDialogUtils.h).
  - Moved OpenUnrealInsights() into a utility static class FMiscUtils in TraceInsightsCore module. This removes some of the dependencies to InsightsManager.
  - Moved/renamed Insights::FEventNameFilterValueConverter to UE::Insights::TimingProfiler::FTimerNameFilterValueConverter (currently unused).
  - Moved "Version.h" to TraceInsightsCore module ("InsightsCore/Version.h").
  - Moved more types into UE::Insights namespace, in TraceInsights module.
  - Changed UnrealInsights program to either load the TraceInsights module or the TraceInsightsFrontend module depending on the required mode to run (based on command line arguments).
  - Updated UnrealInsights program to detect "single instance" of Unreal Insights Frontend window/process much sooner during initialization.
  - Temporary disabled "Import Table ..." and "Diff Tables ..." functionality in Unreal Insights Frontend / Trace Store window (code switched off by UE_INSIGHTS_TABLE_IMPORT_TOOL_ENABLED in STraceStoreWindow.cpp).

diff --git a/Engine/Plugins/TraceUtilities/Source/EditorTraceUtilities/EditorTraceUtilities.Build.cs b/Engine/Plugins/TraceUtilities/Source/EditorTraceUtilities/EditorTraceUtilities.Build.cs
--- a/Engine/Plugins/TraceUtilities/Source/EditorTraceUtilities/EditorTraceUtilities.Build.cs
+++ b/Engine/Plugins/TraceUtilities/Source/EditorTraceUtilities/EditorTraceUtilities.Build.cs
@@ -30,7 +30,7 @@ public class EditorTraceUtilities : ModuleRules
 				"SlateCore",
 				"ToolMenus",
 				"TraceAnalysis",
-				"TraceInsights",
+				"TraceInsightsFrontend",
 				"TraceTools",
 				"TraceLog",
 				"UATHelper",
diff --git a/Engine/Plugins/TraceUtilities/Source/EditorTraceUtilities/Private/InsightsStatusBar/SInsightsStatusBar.cpp b/Engine/Plugins/TraceUtilities/Source/EditorTraceUtilities/Private/InsightsStatusBar/SInsightsStatusBar.cpp
--- a/Engine/Plugins/TraceUtilities/Source/EditorTraceUtilities/Private/InsightsStatusBar/SInsightsStatusBar.cpp
+++ b/Engine/Plugins/TraceUtilities/Source/EditorTraceUtilities/Private/InsightsStatusBar/SInsightsStatusBar.cpp
@@ -3,8 +3,6 @@
 #include "SInsightsStatusBar.h"
 
 #include "CoreGlobals.h"
-#include "EditorTraceUtilitiesStyle.h"
-#include "EditorTraceUtilities.h"
 #include "Framework/Application/SlateApplication.h"
 #include "Framework/Commands/Commands.h"
 #include "Framework/Commands/UICommandList.h"
@@ -21,28 +19,37 @@
 #include "Misc/Paths.h"
 #include "Modules/ModuleManager.h"
 #include "ProfilingDebugging/MiscTrace.h"
-#include "ProfilingDebugging/TraceScreenshot.h"
 #include "ProfilingDebugging/PlatformEvents.h"
-#include "SRecentTracesList.h"
+#include "ProfilingDebugging/TraceScreenshot.h"
 #include "Styling/StyleColors.h"
 #include "ToolMenus.h"
-#include "TraceTools/Interfaces/ITraceToolsModule.h"
-#include "TraceTools/Widgets/SToggleTraceButton.h"
 #include "Trace/Detail/Channel.h"
 #include "Trace/StoreClient.h"
 #include "Trace/Trace.h"
-#include "UnrealInsightsLauncher.h"
-#include "Insights/Widgets/STraceServerControl.h"
-#include "Widgets/SBoxPanel.h"
 #include "Widgets/Images/SImage.h"
 #include "Widgets/Input/SButton.h"
 #include "Widgets/Input/SComboButton.h"
 #include "Widgets/Notifications/SNotificationList.h"
+#include "Widgets/SBoxPanel.h"
+
+// TraceTools
+#include "TraceTools/Interfaces/ITraceToolsModule.h"
+#include "TraceTools/Widgets/SToggleTraceButton.h"
+
+// TraceInsightsFrontend
+#include "InsightsFrontend/StoreService/TraceServerControl.h"
+
+// EditorTraceUtilities
+#include "EditorTraceUtilities.h"
+#include "EditorTraceUtilitiesStyle.h"
+#include "SRecentTracesList.h"
+#include "UnrealInsightsLauncher.h"
 
 #define LOCTEXT_NAMESPACE "InsightsEditor"
 
 namespace UE::EditorTraceUtilities
 {
+
 FStatusBarTraceSettings SInsightsStatusBarWidget::StatusBarTraceSettings;
 
 const TCHAR* SInsightsStatusBarWidget::DefaultPreset = TEXT("default");
@@ -1048,6 +1055,7 @@ void SInsightsStatusBarWidget::OpenTrace(int32 Index)
 		FUnrealInsightsLauncher::Get()->TryOpenTraceFromDestination(Traces[Index]->FilePath);
 	}
 }
+
 } // namespace UE::EditorTraceUtilities
 
 #undef LOCTEXT_NAMESPACE
diff --git a/Engine/Plugins/TraceUtilities/Source/EditorTraceUtilities/Private/InsightsStatusBar/SInsightsStatusBar.h b/Engine/Plugins/TraceUtilities/Source/EditorTraceUtilities/Private/InsightsStatusBar/SInsightsStatusBar.h
--- a/Engine/Plugins/TraceUtilities/Source/EditorTraceUtilities/Private/InsightsStatusBar/SInsightsStatusBar.h
+++ b/Engine/Plugins/TraceUtilities/Source/EditorTraceUtilities/Private/InsightsStatusBar/SInsightsStatusBar.h
@@ -11,7 +11,7 @@
 class FLiveSessionTracker;
 class FMenuBuilder;
 class FUICommandList;
-class STraceServerControl;
+namespace UE::Insights { class FTraceServerControl; }
 
 TSharedRef<SWidget> CreateInsightsStatusBarWidget();
 
@@ -154,7 +154,7 @@ private:
 
 	TSharedPtr<FUICommandList> CommandList;
 	
-	TArray<STraceServerControl> ServerControls;
+	TArray<UE::Insights::FTraceServerControl> ServerControls;
 
 	TArray<TSharedPtr<FTraceFileInfo>> Traces;
 	FName LogListingName;
diff --git a/Engine/Source/Developer/TraceInsights/Private/Insights/Common/InsightsMenuBuilder.cpp b/Engine/Source/Developer/TraceInsights/Private/Insights/Common/InsightsMenuBuilder.cpp
--- a/Engine/Source/Developer/TraceInsights/Private/Insights/Common/InsightsMenuBuilder.cpp
+++ b/Engine/Source/Developer/TraceInsights/Private/Insights/Common/InsightsMenuBuilder.cpp
@@ -8,11 +8,19 @@
 #include "WorkspaceMenuStructure.h"
 #include "WorkspaceMenuStructureModule.h"
 
+// TraceInsightsCore
+#include "InsightsCore/Common/InsightsCoreStyle.h"
+#include "InsightsCore/Common/MiscUtils.h"
+
+// TraceInsights
 #include "Insights/ImportTool/TableImportTool.h"
 #include "Insights/InsightsStyle.h"
 #include "Insights/InsightsManager.h"
 
-#define LOCTEXT_NAMESPACE "InsightsMenuBuilder"
+#define LOCTEXT_NAMESPACE "UE::Insights::MenuBuilder"
+
+namespace UE::Insights
+{
 
 ////////////////////////////////////////////////////////////////////////////////////////////////////
 // FInsightsMenuBuilder
@@ -64,7 +72,7 @@ void FInsightsMenuBuilder::PopulateMenu(FMenuBuilder& MenuBuilder)
 		LOCTEXT("OpenSessionBrowser", "Session Browser"),
 		LOCTEXT("OpenSessionBrowser_ToolTip", "Opens the Unreal Insights Session Browser window."),
 		FSlateIcon(FInsightsStyle::GetStyleSetName(), "AppIcon.Small"),
-		FUIAction(FExecuteAction::CreateLambda([] { FInsightsManager::Get()->OpenUnrealInsights(); })));
+		FUIAction(FExecuteAction::CreateLambda([] { FMiscUtils::OpenUnrealInsights(); })));
 	MenuBuilder.AddSubMenu(
 		LOCTEXT("OpenTraceFile_SubMenu", "Open Trace File"),
 		LOCTEXT("OpenTraceFile_SubMenu_Desc", "Starts analysis for a specified trace file."),
@@ -112,7 +120,7 @@ void FInsightsMenuBuilder::PopulateMenu(FMenuBuilder& MenuBuilder)
 		MenuBuilder.AddMenuEntry(
 			LOCTEXT("OpenStarshipSuite", "Starship Test Suite"),
 			LOCTEXT("OpenStarshipSuite_ToolTip", "Opens the Starship UX test suite."),
-			FSlateIcon(FInsightsStyle::GetStyleSetName(), "Icons.Test"),
+			FSlateIcon(FInsightsCoreStyle::GetStyleSetName(), "Icons.Test"),
 			OpenStarshipSuiteAction,
 			NAME_None,
 			EUserInterfaceActionType::Button
@@ -183,4 +191,6 @@ void FInsightsMenuBuilder::AddMenuEntry(
 
 ////////////////////////////////////////////////////////////////////////////////////////////////////
 
+} // namespace UE::Insights
+
 #undef LOCTEXT_NAMESPACE
diff --git a/Engine/Source/Developer/TraceInsights/Private/Insights/Common/InsightsMenuBuilder.h b/Engine/Source/Developer/TraceInsights/Private/Insights/Common/InsightsMenuBuilder.h
--- a/Engine/Source/Developer/TraceInsights/Private/Insights/Common/InsightsMenuBuilder.h
+++ b/Engine/Source/Developer/TraceInsights/Private/Insights/Common/InsightsMenuBuilder.h
@@ -11,6 +11,9 @@ class FMenuBuilder;
 
 struct FUIAction;
 
+namespace UE::Insights
+{
+
 class FInsightsMenuBuilder : public TSharedFromThis<FInsightsMenuBuilder>
 {
 public:
@@ -37,3 +40,5 @@ private:
 	TSharedRef<FWorkspaceItem> WindowsGroup;
 #endif
 };
+
+} // namespace UE::Insights
diff --git a/Engine/Source/Developer/TraceInsights/Private/Insights/ImportTool/TableImportTool.cpp b/Engine/Source/Developer/TraceInsights/Private/Insights/ImportTool/TableImportTool.cpp
--- a/Engine/Source/Developer/TraceInsights/Private/Insights/ImportTool/TableImportTool.cpp
+++ b/Engine/Source/Developer/TraceInsights/Private/Insights/ImportTool/TableImportTool.cpp
@@ -26,9 +26,9 @@
 #include "Insights/Table/Widgets/SUntypedDiffTableTreeView.h"
 #include "Insights/Table/Widgets/SUntypedTableTreeView.h"
 
-#define LOCTEXT_NAMESPACE "ImportTableTool"
+#define LOCTEXT_NAMESPACE "UE::Insights::TableImportTool"
 
-namespace Insights
+namespace UE::Insights
 {
 
 ////////////////////////////////////////////////////////////////////////////////////////////////////
@@ -330,7 +330,7 @@ TSharedRef<SDockTab> FTableImportTool::SpawnTab_TableImportTreeView(const FSpawn
 	UntypedTable->Reset();
 	UntypedTable->SetDisplayName(InDisplayName);
 
-	TSharedPtr<Insights::SUntypedTableTreeView> TableTreeView;
+	TSharedPtr<SUntypedTableTreeView> TableTreeView;
 
 	const TSharedRef<SDockTab> DockTab = SNew(SDockTab)
 		.ShouldAutosize(false)
@@ -434,6 +434,6 @@ void FTableImportTool::CloseAllOpenTabs()
 
 ////////////////////////////////////////////////////////////////////////////////////////////////////
 
-} // namespace Insights
+} // namespace UE::Insights
 
 #undef LOCTEXT_NAMESPACE
diff --git a/Engine/Source/Developer/TraceInsights/Private/Insights/ImportTool/TableImportTool.h b/Engine/Source/Developer/TraceInsights/Private/Insights/ImportTool/TableImportTool.h
--- a/Engine/Source/Developer/TraceInsights/Private/Insights/ImportTool/TableImportTool.h
+++ b/Engine/Source/Developer/TraceInsights/Private/Insights/ImportTool/TableImportTool.h
@@ -19,7 +19,7 @@ namespace TraceServices
 	struct FTableImportCallbackParams;
 }
 
-namespace Insights
+namespace UE::Insights
 {
 
 class SUntypedTableTreeView;
@@ -42,15 +42,10 @@ struct FOpenImportedTableTabData
 class FTableImportTool : public TSharedFromThis<FTableImportTool>, public IInsightsComponent
 {
 public:
-	/** Default constructor. */
 	FTableImportTool();
-
-	/** Destructor. */
 	virtual ~FTableImportTool();
 
-	/** Creates an instance of the Table Import Tool. */
 	static TSharedPtr<FTableImportTool> CreateInstance();
-
 	static TSharedPtr<FTableImportTool> Get();
 
 	// IInsightsComponent
@@ -62,6 +57,7 @@ public:
 
 	void StartImportProcess();
 	void ImportFile(const FString& Filename);
+
 	void StartDiffProcess();
 	void DiffFiles(const FString& FilenameA, const FString& FilenameB);
 
@@ -77,7 +73,10 @@ private:
 	void TableImportServiceCallback(TSharedPtr<TraceServices::FTableImportCallbackParams> Params);
 
 	FName GetTableID(const FString& Path);
+
+private:
 	static TSharedPtr<FTableImportTool> Instance;
+
 	TMap<FName, FOpenImportedTableTabData> OpenTablesMap;
 };
 
diff --git a/Engine/Source/Developer/TraceInsights/Private/Insights/InsightsCommands.cpp b/Engine/Source/Developer/TraceInsights/Private/Insights/InsightsCommands.cpp
--- a/Engine/Source/Developer/TraceInsights/Private/Insights/InsightsCommands.cpp
+++ b/Engine/Source/Developer/TraceInsights/Private/Insights/InsightsCommands.cpp
@@ -13,7 +13,10 @@
 
 ////////////////////////////////////////////////////////////////////////////////////////////////////
 
-#define LOCTEXT_NAMESPACE "FInsightsCommands"
+#define LOCTEXT_NAMESPACE "UE::Insights::FInsightsCommands"
+
+namespace UE::Insights
+{
 
 ////////////////////////////////////////////////////////////////////////////////////////////////////
 // FInsightsCommands
@@ -181,4 +184,6 @@ bool FInsightsActionManager::OpenSettings_CanExecute() const
 
 ////////////////////////////////////////////////////////////////////////////////////////////////////
 
+} // namespace UE::Insights
+
 #undef LOCTEXT_NAMESPACE
diff --git a/Engine/Source/Developer/TraceInsights/Private/Insights/InsightsCommands.h b/Engine/Source/Developer/TraceInsights/Private/Insights/InsightsCommands.h
--- a/Engine/Source/Developer/TraceInsights/Private/Insights/InsightsCommands.h
+++ b/Engine/Source/Developer/TraceInsights/Private/Insights/InsightsCommands.h
@@ -7,6 +7,11 @@
 #include "Framework/Commands/Commands.h"
 #include "Framework/Commands/UIAction.h"
 
+namespace UE::Insights
+{
+
+class FInsightsManager;
+
 ////////////////////////////////////////////////////////////////////////////////////////////////////
 /**
  * Class that holds all profiler commands.
@@ -87,5 +92,9 @@ protected:
 
 protected:
 	/** Reference to the global instance of the Insights manager. */
-	class FInsightsManager* This;
+	FInsightsManager* This;
 };
+
+////////////////////////////////////////////////////////////////////////////////////////////////////
+
+} // namespace UE::Insights
diff --git a/Engine/Source/Developer/TraceInsights/Private/Insights/InsightsManager.cpp b/Engine/Source/Developer/TraceInsights/Private/Insights/InsightsManager.cpp
--- a/Engine/Source/Developer/TraceInsights/Private/Insights/InsightsManager.cpp
+++ b/Engine/Source/Developer/TraceInsights/Private/Insights/InsightsManager.cpp
@@ -27,6 +27,7 @@
 #include "TraceServices/Model/NetProfiler.h"
 
 // TraceInsightsCore
+#include "InsightsCore/Common/MiscUtils.h"
 #include "InsightsCore/Common/TimeUtils.h"
 
 // TraceInsights
@@ -40,37 +41,28 @@
 #include "Insights/TimingProfiler/TimingProfilerManager.h"
 #include "Insights/TimingProfiler/Widgets/STimingProfilerWindow.h"
 #include "Insights/Widgets/SSessionInfoWindow.h"
-#include "Insights/Widgets/SStartPageWindow.h"
-
-////////////////////////////////////////////////////////////////////////////////////////////////////
-
-/**
- * Special tab type, that cannot be dragged/undocked from the tab bar
- */
-class SLockedTab : public SDockTab
-{
-	virtual FReply OnDragDetected(const FGeometry& MyGeometry, const FPointerEvent& MouseEvent) override
-	{
-		return FReply::Handled();
-	}
-};
 
 ////////////////////////////////////////////////////////////////////////////////////////////////////
 
 #define LOCTEXT_NAMESPACE "InsightsManager"
 
 const FName FInsightsManagerTabs::StartPageTabId(TEXT("TraceStore")); // DEPRECATED
-const FName FInsightsManagerTabs::TraceStoreTabId(TEXT("TraceStore"));
-const FName FInsightsManagerTabs::ConnectionTabId(TEXT("Connection"));
-const FName FInsightsManagerTabs::LauncherTabId(TEXT("Launcher"));
+const FName FInsightsManagerTabs::TraceStoreTabId(TEXT("TraceStore")); // DEPRECATED
+const FName FInsightsManagerTabs::ConnectionTabId(TEXT("Connection")); // DEPRECATED
+const FName FInsightsManagerTabs::LauncherTabId(TEXT("Launcher")); // DEPRECATED
+
 const FName FInsightsManagerTabs::SessionInfoTabId(TEXT("SessionInfo"));
 const FName FInsightsManagerTabs::TimingProfilerTabId(TEXT("TimingProfiler"));
 const FName FInsightsManagerTabs::LoadingProfilerTabId(TEXT("LoadingProfiler"));
-const FName FInsightsManagerTabs::NetworkingProfilerTabId(TEXT("NetworkingProfiler"));
 const FName FInsightsManagerTabs::MemoryProfilerTabId(TEXT("MemoryProfiler"));
+const FName FInsightsManagerTabs::NetworkingProfilerTabId(TEXT("NetworkingProfiler"));
+
 const FName FInsightsManagerTabs::AutomationWindowTabId(TEXT("AutomationWindow"));
 const FName FInsightsManagerTabs::MessageLogTabId(TEXT("MessageLog"));
 
+namespace UE::Insights
+{
+
 ////////////////////////////////////////////////////////////////////////////////////////////////////
 // FInsightsManager
 ////////////////////////////////////////////////////////////////////////////////////////////////////
@@ -192,50 +184,6 @@ void FInsightsManager::BindCommands()
 
 void FInsightsManager::RegisterMajorTabs(IUnrealInsightsModule& InsightsModule)
 {
-#if !WITH_EDITOR
-	const FInsightsMajorTabConfig& TraceStoreConfig = InsightsModule.FindMajorTabConfig(FInsightsManagerTabs::TraceStoreTabId);
-	if (TraceStoreConfig.bIsAvailable)
-	{
-		// Register tab spawner for the Trace Store tab.
-		FTabSpawnerEntry& TabSpawnerEntry = FGlobalTabmanager::Get()->RegisterNomadTabSpawner(FInsightsManagerTabs::TraceStoreTabId,
-			FOnSpawnTab::CreateRaw(this, &FInsightsManager::SpawnTraceStoreTab))
-			.SetDisplayName(TraceStoreConfig.TabLabel.IsSet() ? TraceStoreConfig.TabLabel.GetValue() : LOCTEXT("TraceStoreTabTitle", "Trace Store"))
-			.SetTooltipText(TraceStoreConfig.TabTooltip.IsSet() ? TraceStoreConfig.TabTooltip.GetValue() : LOCTEXT("TraceStoreTooltipText", "Open the Trace Store Browser."))
-			.SetIcon(TraceStoreConfig.TabIcon.IsSet() ? TraceStoreConfig.TabIcon.GetValue() : FSlateIcon(FInsightsStyle::GetStyleSetName(), "Icons.TraceStore"));
-
-		TSharedRef<FWorkspaceItem> Group = TraceStoreConfig.WorkspaceGroup.IsValid() ? TraceStoreConfig.WorkspaceGroup.ToSharedRef() : WorkspaceMenu::GetMenuStructure().GetDeveloperToolsProfilingCategory();
-		TabSpawnerEntry.SetGroup(Group);
-	}
-
-	const FInsightsMajorTabConfig& ConnectionConfig = InsightsModule.FindMajorTabConfig(FInsightsManagerTabs::ConnectionTabId);
-	if (ConnectionConfig.bIsAvailable)
-	{
-		// Register tab spawner for the Connection tab.
-		FTabSpawnerEntry& TabSpawnerEntry = FGlobalTabmanager::Get()->RegisterNomadTabSpawner(FInsightsManagerTabs::ConnectionTabId,
-			FOnSpawnTab::CreateRaw(this, &FInsightsManager::SpawnConnectionTab))
-			.SetDisplayName(ConnectionConfig.TabLabel.IsSet() ? ConnectionConfig.TabLabel.GetValue() : LOCTEXT("ConnectionTabTitle", "Connection"))
-			.SetTooltipText(ConnectionConfig.TabTooltip.IsSet() ? ConnectionConfig.TabTooltip.GetValue() : LOCTEXT("ConnectionTooltipText", "Open the Connection tab."))
-			.SetIcon(ConnectionConfig.TabIcon.IsSet() ? ConnectionConfig.TabIcon.GetValue() : FSlateIcon(FInsightsStyle::GetStyleSetName(), "Icons.Connection"));
-
-		TSharedRef<FWorkspaceItem> Group = ConnectionConfig.WorkspaceGroup.IsValid() ? ConnectionConfig.WorkspaceGroup.ToSharedRef() : WorkspaceMenu::GetMenuStructure().GetDeveloperToolsProfilingCategory();
-		TabSpawnerEntry.SetGroup(Group);
-	}
-
-	const FInsightsMajorTabConfig& LauncherConfig = InsightsModule.FindMajorTabConfig(FInsightsManagerTabs::LauncherTabId);
-	if (LauncherConfig.bIsAvailable)
-	{
-		// Register tab spawner for the Launcher tab.
-		FTabSpawnerEntry& TabSpawnerEntry = FGlobalTabmanager::Get()->RegisterNomadTabSpawner(FInsightsManagerTabs::LauncherTabId,
-			FOnSpawnTab::CreateRaw(this, &FInsightsManager::SpawnLauncherTab))
-			.SetDisplayName(LauncherConfig.TabLabel.IsSet() ? LauncherConfig.TabLabel.GetValue() : LOCTEXT("LauncherTabTitle", "Launcher"))
-			.SetTooltipText(LauncherConfig.TabTooltip.IsSet() ? LauncherConfig.TabTooltip.GetValue() : LOCTEXT("LauncherTooltipText", "Open the Launcher tab."))
-			.SetIcon(LauncherConfig.TabIcon.IsSet() ? LauncherConfig.TabIcon.GetValue() : FSlateIcon(FInsightsStyle::GetStyleSetName(), "Icons.Launcher"));
-
-		TSharedRef<FWorkspaceItem> Group = LauncherConfig.WorkspaceGroup.IsValid() ? LauncherConfig.WorkspaceGroup.ToSharedRef() : WorkspaceMenu::GetMenuStructure().GetDeveloperToolsProfilingCategory();
-		TabSpawnerEntry.SetGroup(Group);
-	}
-#endif // !WITH_EDITOR
-
 	const FInsightsMajorTabConfig& SessionInfoConfig = InsightsModule.FindMajorTabConfig(FInsightsManagerTabs::SessionInfoTabId);
 	if (SessionInfoConfig.bIsAvailable)
 	{
@@ -261,113 +209,6 @@ void FInsightsManager::RegisterMajorTabs(IUnrealInsightsModule& InsightsModule)
 void FInsightsManager::UnregisterMajorTabs()
 {
 	FGlobalTabmanager::Get()->UnregisterNomadTabSpawner(FInsightsManagerTabs::SessionInfoTabId);
-
-#if !WITH_EDITOR
-	FGlobalTabmanager::Get()->UnregisterNomadTabSpawner(FInsightsManagerTabs::LauncherTabId);
-	FGlobalTabmanager::Get()->UnregisterNomadTabSpawner(FInsightsManagerTabs::ConnectionTabId);
-	FGlobalTabmanager::Get()->UnregisterNomadTabSpawner(FInsightsManagerTabs::TraceStoreTabId);
-#endif
-}
-
-////////////////////////////////////////////////////////////////////////////////////////////////////
-
-TSharedRef<SDockTab> FInsightsManager::SpawnTraceStoreTab(const FSpawnTabArgs& Args)
-{
-	const TSharedRef<SDockTab> DockTab = SNew(SLockedTab)
-		.TabRole(ETabRole::MajorTab);
-
-	DockTab->SetOnTabClosed(SDockTab::FOnTabClosedCallback::CreateRaw(this, &FInsightsManager::OnTraceStoreTabClosed));
-
-	TSharedRef<STraceStoreWindow> Window = SNew(STraceStoreWindow);
-	DockTab->SetContent(Window);
-
-	AssignTraceStoreWindow(Window);
-
-	if (!bIsMainTabSet)
-	{
-		FGlobalTabmanager::Get()->SetMainTab(DockTab);
-		bIsMainTabSet = true;
-	}
-
-	return DockTab;
-}
-
-////////////////////////////////////////////////////////////////////////////////////////////////////
-
-void FInsightsManager::OnTraceStoreTabClosed(TSharedRef<SDockTab> TabBeingClosed)
-{
-	RemoveTraceStoreWindow();
-
-	// Disable TabClosed delegate.
-	TabBeingClosed->SetOnTabClosed(SDockTab::FOnTabClosedCallback());
-}
-
-////////////////////////////////////////////////////////////////////////////////////////////////////
-
-TSharedRef<SDockTab> FInsightsManager::SpawnConnectionTab(const FSpawnTabArgs& Args)
-{
-	const TSharedRef<SDockTab> DockTab = SNew(SLockedTab)
-		.TabRole(ETabRole::MajorTab)
-		.OnCanCloseTab_Lambda([]() { return false; }); // can't close this tab
-
-	DockTab->SetOnTabClosed(SDockTab::FOnTabClosedCallback::CreateRaw(this, &FInsightsManager::OnConnectionTabClosed));
-
-	TSharedRef<SConnectionWindow> Window = SNew(SConnectionWindow);
-	DockTab->SetContent(Window);
-
-	AssignConnectionWindow(Window);
-
-	//if (!bIsMainTabSet)
-	//{
-	//	FGlobalTabmanager::Get()->SetMainTab(DockTab);
-	//	bIsMainTabSet = true;
-	//}
-
-	return DockTab;
-}
-
-////////////////////////////////////////////////////////////////////////////////////////////////////
-
-void FInsightsManager::OnConnectionTabClosed(TSharedRef<SDockTab> TabBeingClosed)
-{
-	RemoveConnectionWindow();
-
-	// Disable TabClosed delegate.
-	TabBeingClosed->SetOnTabClosed(SDockTab::FOnTabClosedCallback());
-}
-
-////////////////////////////////////////////////////////////////////////////////////////////////////
-
-TSharedRef<SDockTab> FInsightsManager::SpawnLauncherTab(const FSpawnTabArgs& Args)
-{
-	const TSharedRef<SDockTab> DockTab = SNew(SLockedTab)
-		.TabRole(ETabRole::NomadTab)
-		.OnCanCloseTab_Lambda([]() { return false; }); // can't close this tab
-
-	DockTab->SetOnTabClosed(SDockTab::FOnTabClosedCallback::CreateRaw(this, &FInsightsManager::OnLauncherTabClosed));
-
-	TSharedRef<SLauncherWindow> Window = SNew(SLauncherWindow);
-	DockTab->SetContent(Window);
-
-	AssignLauncherWindow(Window);
-
-	//if (!bIsMainTabSet)
-	//{
-	//	FGlobalTabmanager::Get()->SetMainTab(DockTab);
-	//	bIsMainTabSet = true;
-	//}
-
-	return DockTab;
-}
-
-////////////////////////////////////////////////////////////////////////////////////////////////////
-
-void FInsightsManager::OnLauncherTabClosed(TSharedRef<SDockTab> TabBeingClosed)
-{
-	RemoveLauncherWindow();
-
-	// Disable TabClosed delegate.
-	TabBeingClosed->SetOnTabClosed(SDockTab::FOnTabClosedCallback());
 }
 
 ////////////////////////////////////////////////////////////////////////////////////////////////////
@@ -406,51 +247,6 @@ void FInsightsManager::OnSessionInfoTabClosed(TSharedRef<SDockTab> TabBeingClose
 
 ////////////////////////////////////////////////////////////////////////////////////////////////////
 
-FString FInsightsManager::GetStoreDir()
-{
-	using namespace UE::Trace;
-	if (!StoreClient.IsValid())
-	{
-		return FString();
-	}
-	FScopeLock _(&StoreClientCriticalSection);
-	const FStoreClient::FStatus* Status = StoreClient->GetStatus();
-	return Status ? FString(Status->GetStoreDir()) : FString();
-}
-
-////////////////////////////////////////////////////////////////////////////////////////////////////
-
-bool FInsightsManager::ConnectToStore(const TCHAR* Host, uint32 Port)
-{
-	using namespace UE::Trace;
-	StoreClient.Reset(FStoreClient::Connect(Host, Port));
-	if (!StoreClient.IsValid())
-	{
-		return false;
-	}
-
-	LastStoreHost = Host;
-	LastStorePort = Port;
-	bCanChangeStoreSettings = LastStoreHost.Equals(TEXT("localhost"), ESearchCase::IgnoreCase) ||
-	                          LastStoreHost.Equals(TEXT("127.0.0.1"));
-
-	return true;
-}
-
-////////////////////////////////////////////////////////////////////////////////////////////////////
-
-bool FInsightsManager::ReconnectToStore() const
-{
-	if (!StoreClient.IsValid())
-	{
-		return false;
-	}
-
-	return StoreClient->Reconnect(*LastStoreHost, LastStorePort);
-}
-
-////////////////////////////////////////////////////////////////////////////////////////////////////
-
 TSharedPtr<const TraceServices::IAnalysisSession> FInsightsManager::GetSession() const
 {
 	return Session;
@@ -479,12 +275,6 @@ FInsightsActionManager& FInsightsManager::GetActionManager()
 
 ////////////////////////////////////////////////////////////////////////////////////////////////////
 
-FInsightsSessionBrowserSettings& FInsightsManager::GetSessionBrowserSettings()
-{
-	return FInsightsManager::Instance->SessionBrowserSettings;
-}
-////////////////////////////////////////////////////////////////////////////////////////////////////
-
 FInsightsSettings& FInsightsManager::GetSettings()
 {
 	return FInsightsManager::Instance->Settings;
@@ -770,7 +560,7 @@ void FInsightsManager::ActivateTimingInsightsTab()
 
 bool FInsightsManager::ShowOpenTraceFileDialog(FString& OutTraceFile) const
 {
-	static FString DefaultDirectory(FPaths::ConvertRelativePathToFull(FInsightsManager::Get()->GetStoreDir()));
+	static FString DefaultDirectory(FPaths::ConvertRelativePathToFull(GetStoreDir()));
 
 	TArray<FString> OutFiles;
 	bool bOpened = false;
@@ -804,39 +594,10 @@ bool FInsightsManager::ShowOpenTraceFileDialog(FString& OutTraceFile) const
 
 ////////////////////////////////////////////////////////////////////////////////////////////////////
 
-void FInsightsManager::OpenUnrealInsights(const TCHAR* CmdLine) const
-{
-	if (CmdLine == nullptr)
-	{
-		CmdLine = TEXT("");
-	}
-
-	const TCHAR* ExecutablePath = FPlatformProcess::ExecutablePath();
-
-	constexpr bool bLaunchDetached = true;
-	constexpr bool bLaunchHidden = false;
-	constexpr bool bLaunchReallyHidden = false;
-
-	uint32 ProcessID = 0;
-	const int32 PriorityModifier = 0;
-	const TCHAR* OptionalWorkingDirectory = nullptr;
-
-	void* PipeWriteChild = nullptr;
-	void* PipeReadChild = nullptr;
-
-	FProcHandle Handle = FPlatformProcess::CreateProc(ExecutablePath, CmdLine, bLaunchDetached, bLaunchHidden, bLaunchReallyHidden, &ProcessID, PriorityModifier, OptionalWorkingDirectory, PipeWriteChild, PipeReadChild);
-	if (Handle.IsValid())
-	{
-		FPlatformProcess::CloseProc(Handle);
-	}
-}
-
-////////////////////////////////////////////////////////////////////////////////////////////////////
-
 void FInsightsManager::OpenTraceFile() const
 {
 	FString TraceFile;
-	if (FInsightsManager::Get()->ShowOpenTraceFileDialog(TraceFile))
+	if (ShowOpenTraceFileDialog(TraceFile))
 	{
 		OpenTraceFile(TraceFile);
 	}
@@ -847,7 +608,7 @@ void FInsightsManager::OpenTraceFile() const
 void FInsightsManager::OpenTraceFile(const FString& InTraceFile) const
 {
 	FString CmdLine = TEXT("-OpenTraceFile=\"") + InTraceFile + TEXT("\"");
-	OpenUnrealInsights(*CmdLine);
+	FMiscUtils::OpenUnrealInsights(*CmdLine);
 }
 
 ////////////////////////////////////////////////////////////////////////////////////////////////////
@@ -859,14 +620,15 @@ void FInsightsManager::AutoLoadLiveSession()
 		return;
 	}
 
-	if (!StoreClient.IsValid())
+	UE::Trace::FStoreClient* StoreClient = TraceStoreConnection.GetStoreClient();
+	if (!StoreClient)
 	{
 		return;
 	}
 
 	uint32 AutoLoadTraceId = 0;
 	{
-		FScopeLock _(&StoreClientCriticalSection);
+		FScopeLock _(&TraceStoreConnection.GetStoreClientCriticalSection());
 		const uint32 SessionCount = StoreClient->GetSessionCount();
 		for (uint32 SessionIndex = 0; SessionIndex < SessionCount; ++SessionIndex)
 		{
@@ -896,14 +658,15 @@ void FInsightsManager::LoadLastLiveSession(float RetryTime)
 {
 	ResetSession();
 
-	if (!StoreClient.IsValid())
+	UE::Trace::FStoreClient* StoreClient = TraceStoreConnection.GetStoreClient();
+	if (!StoreClient)
 	{
 		return;
 	}
 
 	uint32 LastLiveSessionTraceId = 0;
 	{
-		FScopeLock _(&StoreClientCriticalSection);
+		FScopeLock _(&TraceStoreConnection.GetStoreClientCriticalSection());
 		const uint32 SessionCount = StoreClient->GetSessionCount();
 		if (SessionCount != 0)
 		{
@@ -932,7 +695,8 @@ void FInsightsManager::LoadTrace(uint32 InTraceId, bool InAutoQuit)
 {
 	ResetSession();
 
-	if (!StoreClient.IsValid())
+	UE::Trace::FStoreClient* StoreClient = TraceStoreConnection.GetStoreClient();
+	if (!StoreClient)
 	{
 		if (InAutoQuit)
 		{
@@ -941,7 +705,7 @@ void FInsightsManager::LoadTrace(uint32 InTraceId, bool InAutoQuit)
 		return;
 	}
 
-	FScopeLock StoreClientLock(&StoreClientCriticalSection);
+	FScopeLock StoreClientLock(&TraceStoreConnection.GetStoreClientCriticalSection());
 
 	UE::Trace::FStoreClient::FTraceData TraceData = StoreClient->ReadTrace(InTraceId);
 	if (!TraceData)
@@ -996,7 +760,7 @@ void FInsightsManager::LoadTrace(uint32 InTraceId, bool InAutoQuit)
 void FInsightsManager::LoadTraceFile()
 {
 	FString TraceFile;
-	if (FInsightsManager::Get()->ShowOpenTraceFileDialog(TraceFile))
+	if (ShowOpenTraceFileDialog(TraceFile))
 	{
 		LoadTraceFile(TraceFile);
 	}
@@ -1135,11 +899,6 @@ void FInsightsManager::UpdateAppTitle()
 
 void FInsightsManager::OpenSettings()
 {
-	TSharedPtr<STraceStoreWindow> Wnd = GetTraceStoreWindow();
-	if (Wnd.IsValid())
-	{
-		Wnd->OpenSettings();
-	}
 }
 
 ////////////////////////////////////////////////////////////////////////////////////////////////////
@@ -1267,4 +1026,6 @@ bool FInsightsManager::HandleResponseFileCmd(const TCHAR* ResponseFile, FOutputD
 
 ////////////////////////////////////////////////////////////////////////////////////////////////////
 
+} // namespace UE::Insights
+
 #undef LOCTEXT_NAMESPACE
diff --git a/Engine/Source/Developer/TraceInsights/Private/Insights/InsightsManager.h b/Engine/Source/Developer/TraceInsights/Private/Insights/InsightsManager.h
--- a/Engine/Source/Developer/TraceInsights/Private/Insights/InsightsManager.h
+++ b/Engine/Source/Developer/TraceInsights/Private/Insights/InsightsManager.h
@@ -8,6 +8,9 @@
 #include "Framework/Commands/UICommandList.h"
 #include "Input/DragAndDrop.h"
 
+// TraceAnalysis
+#include "Trace/StoreConnection.h"
+
 // TraceInsightsCore
 #include "InsightsCore/Common/AvailabilityCheck.h"
 #include "InsightsCore/Common/Stopwatch.h"
@@ -16,7 +19,6 @@
 #include "Insights/IInsightsManager.h"
 #include "Insights/InsightsCommands.h"
 #include "Insights/InsightsSettings.h"
-#include "Insights/InsightsSessionBrowserSettings.h"
 #include "Insights/IUnrealInsightsModule.h"
 
 ////////////////////////////////////////////////////////////////////////////////////////////////////
@@ -32,11 +34,12 @@ namespace TraceServices
 	class IModuleService;
 }
 
-class STraceStoreWindow;
-class SConnectionWindow;
-class SLauncherWindow;
-class SSessionInfoWindow;
 class FInsightsTestRunner;
+
+namespace UE::Insights
+{
+
+class SSessionInfoWindow;
 class FInsightsMenuBuilder;
 
 /**
@@ -97,22 +100,26 @@ public:
 	TSharedRef<TraceServices::IAnalysisService> GetAnalysisService() const { return AnalysisService; }
 	TSharedRef<TraceServices::IModuleService> GetModuleService() const { return ModuleService; }
 
-	FString GetStoreDir();
+	//////////////////////////////////////////////////
+	// Trace Store Connection wrapper
 
-	bool ConnectToStore(const TCHAR* Host, uint32 Port=0);
+	bool ConnectToStore(const TCHAR* Host, uint32 Port = 0) { return TraceStoreConnection.ConnectToStore(Host, Port); }
+	bool ReconnectToStore() { return TraceStoreConnection.ReconnectToStore(); }
 
-	/**
-	 * Attempt to reconnect to the store if the connection was severed, without recreating the store client.
-	 * @return True on success, false on failure
-	 */
-	bool ReconnectToStore() const;
+	UE::Trace::FStoreClient* GetStoreClient() const { return TraceStoreConnection.GetStoreClient(); }
+	FCriticalSection& GetStoreClientCriticalSection() const { return TraceStoreConnection.GetStoreClientCriticalSection(); }
+
+	bool GetStoreAddressAndPort(uint32& OutStoreAddress, uint32& OutStorePort) const { return TraceStoreConnection.GetStoreAddressAndPort(OutStoreAddress, OutStorePort); }
+	FString GetStoreDir() const { return TraceStoreConnection.GetStoreDir(); }
 
-	const FString& GetLastStoreHost() const { return LastStoreHost; }
+	const FString& GetLastStoreHost() const { return TraceStoreConnection.GetLastStoreHost(); }
+	uint32 GetLastStorePort() const { return TraceStoreConnection.GetLastStorePort(); }
 
-	const bool CanChangeStoreSettings() const { return bCanChangeStoreSettings; }
-	
-	UE::Trace::FStoreClient* GetStoreClient() const { return StoreClient.Get(); }
-	FCriticalSection& GetStoreClientCriticalSection() const { return StoreClientCriticalSection; }
+	const bool IsLocalHost() const { return TraceStoreConnection.IsLocalHost(); }
+	const bool CanChangeStoreSettings() const { return TraceStoreConnection.CanChangeStoreSettings(); }
+
+	//////////////////////////////////////////////////
+	// Trace Session
 
 	/** @return an instance of the trace analysis session. */
 	TSharedPtr<const TraceServices::IAnalysisSession> GetSession() const;
@@ -123,6 +130,8 @@ public:
 	/** @return the filename of the trace being analyzed. */
 	const FString& GetTraceFilename() const { return CurrentTraceFilename; }
 
+	//////////////////////////////////////////////////
+
 	/** @return UI command list for the main manager. */
 	const TSharedRef<FUICommandList> GetCommandList() const;
 
@@ -135,75 +144,6 @@ public:
 	/** @return an instance of the main settings. */
 	static FInsightsSettings& GetSettings();
 
-	/** @return an instance of the session browser settings. */
-	static FInsightsSessionBrowserSettings& GetSessionBrowserSettings();
-
-	//////////////////////////////////////////////////
-	// Trace Store
-
-	void AssignTraceStoreWindow(const TSharedRef<STraceStoreWindow>& InTraceStoreWindow)
-	{
-		TraceStoreWindow = InTraceStoreWindow;
-	}
-
-	void RemoveTraceStoreWindow()
-	{
-		TraceStoreWindow.Reset();
-	}
-
-	/**
-	 * Converts Trace Store window weak pointer to a shared pointer and returns it.
-	 * Make sure the returned pointer is valid before trying to dereference it.
-	 */
-	TSharedPtr<class STraceStoreWindow> GetTraceStoreWindow() const
-	{
-		return TraceStoreWindow.Pin();
-	}
-
-	//////////////////////////////////////////////////
-	// Connection
-
-	void AssignConnectionWindow(const TSharedRef<SConnectionWindow>& InConnectionWindow)
-	{
-		ConnectionWindow = InConnectionWindow;
-	}
-
-	void RemoveConnectionWindow()
-	{
-		ConnectionWindow.Reset();
-	}
-
-	/**
-	 * Converts Connection window weak pointer to a shared pointer and returns it.
-	 * Make sure the returned pointer is valid before trying to dereference it.
-	 */
-	TSharedPtr<class SConnectionWindow> GetConnectionWindow() const
-	{
-		return ConnectionWindow.Pin();
-	}
-
-	//////////////////////////////////////////////////
-	// Launcher
-
-	void AssignLauncherWindow(const TSharedRef<SLauncherWindow>& InLauncherWindow)
-	{
-		LauncherWindow = InLauncherWindow;
-	}
-
-	void RemoveLauncherWindow()
-	{
-		LauncherWindow.Reset();
-	}
-
-	/**
-	 * Converts Launcher window weak pointer to a shared pointer and returns it.
-	 * Make sure the returned pointer is valid before trying to dereference it.
-	 */
-	TSharedPtr<class SLauncherWindow> GetLauncherWindow() const
-	{
-		return LauncherWindow.Pin();
-	}
-
 	//////////////////////////////////////////////////
 	// Session Info
 
@@ -242,12 +182,6 @@ public:
 	 */
 	bool ShowOpenTraceFileDialog(FString& OutTraceFile) const;
 
-	/**
-	 * Starts a new Unreal Insights instance.
-	 * @param CmdLine - The command line passed to the new UnrealInsights.exe process
-	 */
-	void OpenUnrealInsights(const TCHAR* CmdLine = nullptr) const;
-
 	/**
 	 * Shows the open file dialog and starts analysis session for the chosen trace file, in a new Unreal Insights instance.
 	 */
@@ -343,26 +277,6 @@ private:
 	/** Binds our UI commands to delegates. */
 	void BindCommands();
 
-	/** Called to spawn the Trace Store major tab. */
-	TSharedRef<SDockTab> SpawnTraceStoreTab(const FSpawnTabArgs& Args);
-
-	/** Callback called when the Trace Store major tab is closed. */
-	void OnTraceStoreTabClosed(TSharedRef<SDockTab> TabBeingClosed);
-
-	/** Called to spawn the Connection major tab. */
-	TSharedRef<SDockTab> SpawnConnectionTab(const FSpawnTabArgs& Args);
-
-	/** Callback called when the Connection major tab is closed. */
-	void OnConnectionTabClosed(TSharedRef<SDockTab> TabBeingClosed);
-
-	/** Called to spawn the Launcher major tab. */
-	TSharedRef<SDockTab> SpawnLauncherTab(const FSpawnTabArgs& Args);
-
-	TSharedRef<SDockTab> SpawnSessionBrowserAutomationWindowTab(const FSpawnTabArgs& Args);
-
-	/** Callback called when the Launcher major tab is closed. */
-	void OnLauncherTabClosed(TSharedRef<SDockTab> TabBeingClosed);
-
 	/** Called to spawn the Session Info major tab. */
 	TSharedRef<SDockTab> SpawnSessionInfoTab(const FSpawnTabArgs& Args);
 
@@ -408,12 +322,6 @@ private:
 	TSharedRef<TraceServices::IAnalysisService> AnalysisService;
 	TSharedRef<TraceServices::IModuleService> ModuleService;
 
-	/** The client used to connect to the trace store. It is not thread safe! */
-	TUniquePtr<UE::Trace::FStoreClient> StoreClient;
-
-	/** CriticalSection for using the store client's API. */
-	mutable FCriticalSection StoreClientCriticalSection;
-
 	/** The trace analysis session. */
 	TSharedPtr<const TraceServices::IAnalysisSession> Session;
 
@@ -432,18 +340,6 @@ private:
 	/** An instance of the main settings. */
 	FInsightsSettings Settings;
 
-	/** An instance of the session browser settings. */
-	FInsightsSessionBrowserSettings SessionBrowserSettings;
-
-	/** A weak pointer to the Trace Store window. */
-	TWeakPtr<class STraceStoreWindow> TraceStoreWindow;
-
-	/** A weak pointer to the Connection window. */
-	TWeakPtr<class SConnectionWindow> ConnectionWindow;
-
-	/** A weak pointer to the Launcher window. */
-	TWeakPtr<class SLauncherWindow> LauncherWindow;
-
 	/** A weak pointer to the Session Info window. */
 	TWeakPtr<class SSessionInfoWindow> SessionInfoWindow;
 
@@ -478,7 +374,8 @@ private:
 	/** A shared pointer to the global instance of the main manager. */
 	static TSharedPtr<FInsightsManager> Instance;
 
-	FString LastStoreHost;
-	uint32 LastStorePort = 0;
-	bool bCanChangeStoreSettings = false;
+	/** The Trace Store connection */
+	UE::Trace::FStoreConnection TraceStoreConnection;
 };
+
+} // namespace UE::Insights
diff --git a/Engine/Source/Developer/TraceInsights/Private/Insights/InsightsStyle.cpp b/Engine/Source/Developer/TraceInsights/Private/Insights/InsightsStyle.cpp
--- a/Engine/Source/Developer/TraceInsights/Private/Insights/InsightsStyle.cpp
+++ b/Engine/Source/Developer/TraceInsights/Private/Insights/InsightsStyle.cpp
@@ -18,11 +18,18 @@ TSharedPtr<FInsightsStyle::FStyle> FInsightsStyle::StyleInstance = nullptr;
 
 ////////////////////////////////////////////////////////////////////////////////////////////////////
 
+const ISlateStyle& FInsightsStyle::Get()
+{
+	return *StyleInstance;
+}
+
+////////////////////////////////////////////////////////////////////////////////////////////////////
+
 void FInsightsStyle::Initialize()
 {
 	LLM_SCOPE_BYNAME(TEXT("Insights/Style"));
 
-	// The core style must be initialized before the Insights style
+	// The UE Core style must be initialized before the Insights style
 	FSlateApplication::InitializeCoreStyle();
 
 	if (!StyleInstance.IsValid())
@@ -52,13 +59,6 @@ void FInsightsStyle::Shutdown()
 
 ////////////////////////////////////////////////////////////////////////////////////////////////////
 
-const ISlateStyle& FInsightsStyle::Get()
-{
-	return *StyleInstance;
-}
-
-////////////////////////////////////////////////////////////////////////////////////////////////////
-
 FName FInsightsStyle::GetStyleSetName()
 {
 	static FName StyleSetName(TEXT("InsightsStyle"));
@@ -99,14 +99,11 @@ void FInsightsStyle::FStyle::SyncParentStyles()
 
 void FInsightsStyle::FStyle::Initialize()
 {
-	SetParentStyleName("CoreStyle");
+	SetParentStyleName("InsightsCoreStyle");
 
 	// Sync styles from the parent style that will be used as templates for styles defined here
 	SyncParentStyles();
 
-	Set("Mono.9", DEFAULT_FONT("Mono", 9));
-	Set("Mono.10", DEFAULT_FONT("Mono", 10));
-
 	SetContentRoot(FPaths::EngineContentDir() / TEXT("Slate/Starship/Insights"));
 	SetCoreContentRoot(FPaths::EngineContentDir() / TEXT("Slate"));
 
@@ -121,47 +118,10 @@ void FInsightsStyle::FStyle::Initialize()
 	Set("AppIconPadding.Small", FMargin(4.0f, 4.0f, 0.0f, 0.0f));
 
 	//////////////////////////////////////////////////
-
-	Set("WhiteBrush", new FSlateColorBrush(FLinearColor::White));
-	Set("DarkGreenBrush", new FSlateColorBrush(FLinearColor(0.0f, 0.25f, 0.0f, 1.0f)));
-
-	Set("SingleBorder", new FSlateBorderBrush(NAME_None, FMargin(1.0f)));
-	Set("DoubleBorder", new FSlateBorderBrush(NAME_None, FMargin(2.0f)));
-
-	Set("EventBorder", new FSlateBorderBrush(NAME_None, FMargin(1.0f)));
-	Set("HoveredEventBorder", new FSlateBorderBrush(NAME_None, FMargin(2.0f)));
-	Set("SelectedEventBorder", new FSlateBorderBrush(NAME_None, FMargin(2.0f)));
-
-	Set("RoundedBackground", new FSlateRoundedBoxBrush(FLinearColor(1.0f, 1.0f, 1.0f, 1.0f), Icon16x16));
-
-	Set("Border.TB", new CORE_BOX_BRUSH("Icons/Profiler/Profiler_Border_TB_16x", FMargin(4.0f / 16.0f)));
-	Set("Border.L", new CORE_BOX_BRUSH("Icons/Profiler/Profiler_Border_L_16x", FMargin(4.0f / 16.0f)));
-	Set("Border.R", new CORE_BOX_BRUSH("Icons/Profiler/Profiler_Border_R_16x", FMargin(4.0f / 16.0f)));
-
-	Set("Graph.Point", new EDITOR_IMAGE_BRUSH("Old/Graph/ExecutionBubble", Icon16x16));
-
-	//////////////////////////////////////////////////
-	// Icons for major components
+	// Session Info
 
 	Set("Icons.SessionInfo", new IMAGE_BRUSH_SVG("Session", Icon16x16));
 
-	//////////////////////////////////////////////////
-	// Trace Store, Connection, Launcher
-
-	Set("Icons.TraceStore", new IMAGE_BRUSH_SVG("TraceStore", Icon16x16));
-	Set("Icons.Expand", new CORE_IMAGE_BRUSH_SVG("Starship/Common/chevron-right", Icon16x16));
-	Set("Icons.Expanded", new CORE_IMAGE_BRUSH_SVG("Starship/Common/chevron-down", Icon16x16));
-	Set("Icons.AddWatchDir", new CORE_IMAGE_BRUSH_SVG("Starship/Common/folder-plus", Icon16x16));
-	Set("Icons.RemoveWatchDir", new CORE_IMAGE_BRUSH_SVG("Starship/Common/delete", Icon16x16));
-	Set("Icons.Connection", new IMAGE_BRUSH_SVG("Connection", Icon16x16));
-	Set("Icons.Online", new CORE_IMAGE_BRUSH_SVG("Starship/Common/check-circle", Icon16x16, FStyleColors::AccentGreen));
-	Set("Icons.Offline", new CORE_IMAGE_BRUSH_SVG("Starship/Common/alert-triangle", Icon16x16, FStyleColors::Warning));
-	Set("Icons.Launcher", new TODO_IMAGE_BRUSH(Icon16x16));
-	Set("Icons.UTrace", new IMAGE_BRUSH_SVG("UTrace", Icon16x16));
-	Set("Icons.UTraceStack", new IMAGE_BRUSH_SVG("UTrace", Icon16x16));
-	Set("Icons.TraceServerStart", new CORE_IMAGE_BRUSH_SVG("Starship/Common/play", Icon16x16, FStyleColors::AccentGreen));
-	Set("Icons.TraceServerStop", new CORE_IMAGE_BRUSH_SVG("Starship/Common/close", Icon16x16, FStyleColors::AccentRed));
-
 	//////////////////////////////////////////////////
 	// Timing Insights
 
@@ -194,6 +154,9 @@ void FInsightsStyle::FStyle::Initialize()
 	Set("Icons.TasksView", new IMAGE_BRUSH_SVG("Tasks", Icon16x16));
 	Set("Icons.PackagesView", new EDITOR_IMAGE_BRUSH_SVG("Starship/Common/Spreadsheet", Icon16x16));
 
+	//////////////////////////////////////////////////
+	// Timing View
+
 	Set("Icons.AllTracksMenu.ToolBar", new IMAGE_BRUSH_SVG("AllTracks_20", Icon20x20));
 	Set("Icons.CpuGpuTracksMenu.ToolBar", new IMAGE_BRUSH_SVG("CpuGpuTracks_20", Icon20x20));
 	Set("Icons.OtherTracksMenu.ToolBar", new IMAGE_BRUSH_SVG("SpecialTracks_20", Icon20x20));
@@ -203,6 +166,22 @@ void FInsightsStyle::FStyle::Initialize()
 	Set("Icons.HighlightEvents.ToolBar", new CORE_IMAGE_BRUSH_SVG("Starship/Common/Visualizer", Icon20x20));
 	Set("Icons.ResetHighlight.ToolBar", new CORE_IMAGE_BRUSH_SVG("Starship/Common/Reject", Icon20x20));
 
+	//////////////////////////////////////////////////
+
+	Set("Icons.FindFirst.ToolBar", new IMAGE_BRUSH_SVG("ControlsFirst", Icon20x20));
+	Set("Icons.FindPrevious.ToolBar", new IMAGE_BRUSH_SVG("ControlsPrevious", Icon20x20));
+	Set("Icons.FindNext.ToolBar", new IMAGE_BRUSH_SVG("ControlsNext", Icon20x20));
+	Set("Icons.FindLast.ToolBar", new IMAGE_BRUSH_SVG("ControlsLast", Icon20x20));
+
+	//////////////////////////////////////////////////
+
+	Set("Icons.SizeSmall", new IMAGE_BRUSH_SVG("SizeSmall", Icon16x16));
+	Set("Icons.SizeSmall.ToolBar", new IMAGE_BRUSH_SVG("SizeSmall_20", Icon20x20));
+	Set("Icons.SizeMedium", new IMAGE_BRUSH_SVG("SizeMedium", Icon16x16));
+	Set("Icons.SizeMedium.ToolBar", new IMAGE_BRUSH_SVG("SizeMedium_20", Icon20x20));
+	Set("Icons.SizeLarge", new IMAGE_BRUSH_SVG("SizeLarge", Icon16x16));
+	Set("Icons.SizeLarge.ToolBar", new IMAGE_BRUSH_SVG("SizeLarge_20", Icon20x20));
+
 	//////////////////////////////////////////////////
 	// Asset Loading Insights
 
@@ -257,63 +236,14 @@ void FInsightsStyle::FStyle::Initialize()
 
 	//////////////////////////////////////////////////
 
-	Set("Icons.FindFirst.ToolBar", new IMAGE_BRUSH_SVG("ControlsFirst", Icon20x20));
-	Set("Icons.FindPrevious.ToolBar", new IMAGE_BRUSH_SVG("ControlsPrevious", Icon20x20));
-	Set("Icons.FindNext.ToolBar", new IMAGE_BRUSH_SVG("ControlsNext", Icon20x20));
-	Set("Icons.FindLast.ToolBar", new IMAGE_BRUSH_SVG("ControlsLast", Icon20x20));
-
-	//////////////////////////////////////////////////
-
-	Set("Icons.SizeSmall", new IMAGE_BRUSH_SVG("SizeSmall", Icon16x16));
-	Set("Icons.SizeSmall.ToolBar", new IMAGE_BRUSH_SVG("SizeSmall_20", Icon20x20));
-	Set("Icons.SizeMedium", new IMAGE_BRUSH_SVG("SizeMedium", Icon16x16));
-	Set("Icons.SizeMedium.ToolBar", new IMAGE_BRUSH_SVG("SizeMedium_20", Icon20x20));
-	Set("Icons.SizeLarge", new IMAGE_BRUSH_SVG("SizeLarge", Icon16x16));
-	Set("Icons.SizeLarge.ToolBar", new IMAGE_BRUSH_SVG("SizeLarge_20", Icon20x20));
-
-	//////////////////////////////////////////////////
-
 	Set("MainFrame.OpenVisualStudio", new EDITOR_IMAGE_BRUSH_SVG("Starship/Common/VisualStudio", Icon16x16));
 	Set("MainFrame.OpenSourceCodeEditor", new EDITOR_IMAGE_BRUSH_SVG("Starship/Common/SourceCodeEditor", Icon16x16));
 
 	//////////////////////////////////////////////////
 
-	Set("Icons.ResetToDefault", new CORE_IMAGE_BRUSH("Icons/Profiler/profiler_ResetToDefault_32x", Icon16x16));
-	Set("Icons.DiffersFromDefault", new EDITOR_IMAGE_BRUSH_SVG("Starship/Common/ResetToDefault", Icon16x16));
-
-	Set("Icons.Console", new CORE_IMAGE_BRUSH_SVG("Starship/Common/Console", Icon16x16));
-	//Set("Icons.Filter", new CORE_IMAGE_BRUSH_SVG("Starship/Common/filter", Icon16x16));	//-> use FAppStyle "Icons.Filter"
-	Set("Icons.Filter.ToolBar", new CORE_IMAGE_BRUSH_SVG("Starship/Common/filter", Icon20x20));
-	Set("Icons.Find", new EDITOR_IMAGE_BRUSH_SVG("Starship/Common/TraceDataFiltering", Icon16x16));
-	Set("Icons.FilterAddGroup", new EDITOR_IMAGE_BRUSH_SVG("Starship/Common/WorldOutliner", Icon16x16));
-	Set("Icons.ClassicFilter", new IMAGE_BRUSH_SVG("Filter", Icon16x16));
-	Set("Icons.ClassicFilterConfig", new IMAGE_BRUSH_SVG("FilterConfig", Icon16x16));
-
-	Set("Icons.FolderExplore", new EDITOR_IMAGE_BRUSH_SVG("Starship/Common/ContentBrowser", Icon16x16));
-	//Set("Icons.FolderOpen", new CORE_IMAGE_BRUSH_SVG("Starship/Common/folder-open", Icon16x16));		//-> use FAppStyle "Icons.FolderOpen"
-	//Set("Icons.FolderClosed", new CORE_IMAGE_BRUSH_SVG("Starship/Common/folder-closed", Icon16x16));	//-> use FAppStyle "Icons.FolderClosed"
-
-	Set("Icons.SortBy", new CORE_IMAGE_BRUSH("Icons/Profiler/profiler_SortBy_32x", Icon16x16));
-	//Set("Icons.SortUp", new CORE_IMAGE_BRUSH_SVG("Starship/Common/SortUp", Icon16x16));		//-> use FAppStyle "Icons.SortUp"
-	//Set("Icons.SortDown", new CORE_IMAGE_BRUSH_SVG("Starship/Common/SortDown", Icon16x16));	//-> use FAppStyle "Icons.SortDown"
-
-	Set("Icons.ViewColumn", new CORE_IMAGE_BRUSH("Icons/Profiler/profiler_ViewColumn_32x", Icon16x16));
-	Set("Icons.ResetColumn", new CORE_IMAGE_BRUSH("Icons/Profiler/profiler_ResetColumn_32x", Icon16x16));
-
-	Set("Icons.ExpandAll", new CORE_IMAGE_BRUSH("Icons/Profiler/profiler_ExpandAll_32x", Icon16x16));
-	Set("Icons.CollapseAll", new CORE_IMAGE_BRUSH("Icons/Profiler/profiler_CollapseAll_32x", Icon16x16));
-	Set("Icons.ExpandSelection", new CORE_IMAGE_BRUSH("Icons/Profiler/profiler_ExpandSelection_32x", Icon16x16));
-	Set("Icons.CollapseSelection", new CORE_IMAGE_BRUSH("Icons/Profiler/profiler_CollapseSelection_32x", Icon16x16));
-
 	Set("Icons.AddGraphSeries", new CORE_IMAGE_BRUSH_SVG("Starship/Common/plus", Icon16x16));
 	Set("Icons.RemoveGraphSeries", new CORE_IMAGE_BRUSH_SVG("Starship/Common/close", Icon16x16));
 
-	Set("Icons.TestAutomation", new EDITOR_IMAGE_BRUSH_SVG("Starship/Common/TestAutomation", Icon16x16));
-	Set("Icons.Test", new CORE_IMAGE_BRUSH_SVG("Starship/Common/Test", Icon16x16));
-
-	Set("Icons.Debug", new EDITOR_IMAGE_BRUSH_SVG("Starship/Common/bug", Icon16x16));
-	Set("Icons.Debug.ToolBar", new EDITOR_IMAGE_BRUSH_SVG("Starship/Common/bug", Icon20x20));
-
 	Set("Icons.AutoScroll", new IMAGE_BRUSH_SVG("AutoScrollRight_20", Icon16x16));
 
 	Set("Icons.ZeroCountFilter", new IMAGE_BRUSH_SVG("ZeroCountFilter", Icon16x16));
@@ -323,9 +253,6 @@ void FInsightsStyle::FStyle::Initialize()
 	Set("Icons.Pinned", new EDITOR_IMAGE_BRUSH_SVG("Starship/Common/Pinned", Icon16x16));
 	Set("Icons.Unpinned", new EDITOR_IMAGE_BRUSH_SVG("Starship/Common/Unpinned", Icon16x16));
 
-	Set("Icons.Rename", new CORE_IMAGE_BRUSH_SVG("Starship/Common/Rename", Icon16x16));
-	Set("Icons.Delete", new CORE_IMAGE_BRUSH_SVG("Starship/Common/Delete", Icon16x16));
-
 	Set("Icons.ImportTable", new CORE_IMAGE_BRUSH_SVG("Starship/Common/Import", Icon16x16));
 	Set("Icons.SelectEventRange", new EDITOR_IMAGE_BRUSH_SVG("Starship/Common/SelectInViewport", Icon16x16));
 
@@ -334,14 +261,9 @@ void FInsightsStyle::FStyle::Initialize()
 	Set("Icons.FindMaxInstance", new EDITOR_IMAGE_BRUSH_SVG("Starship/Common/Profile", Icon16x16));
 
 	//////////////////////////////////////////////////
+	// Icons for tree/table items
 
-	Set("TreeTable.RowBackground", new EDITOR_IMAGE_BRUSH("Old/White", Icon16x16, FLinearColor(1.0f, 1.0f, 1.0f, 0.25f)));
-	Set("TreeViewBanner.WarningIcon", new CORE_IMAGE_BRUSH_SVG("Starship/Common/alert-circle", Icon20x20, FStyleColors::Warning));
-
-	Set("Icons.Hint.TreeItem", new IMAGE_BRUSH_SVG("InfoTag_12", Icon12x12));
 	Set("Icons.HotPath.TreeItem", new IMAGE_BRUSH_SVG("HotPath_12", Icon12x12));
-	Set("Icons.Group.TreeItem", new CORE_IMAGE_BRUSH_SVG("Starship/Common/folder-closed", Icon12x12));
-	Set("Icons.Leaf.TreeItem", new CORE_IMAGE_BRUSH_SVG("Starship/Common/bullet-point", Icon12x12));
 	Set("Icons.GpuTimer.TreeItem", new TODO_IMAGE_BRUSH(Icon12x12));
 	Set("Icons.CpuTimer.TreeItem", new TODO_IMAGE_BRUSH(Icon12x12));
 	Set("Icons.Counter.TreeItem", new TODO_IMAGE_BRUSH(Icon12x12));
@@ -351,139 +273,6 @@ void FInsightsStyle::FStyle::Initialize()
 	Set("Icons.NetEvent.TreeItem", new TODO_IMAGE_BRUSH(Icon12x12));
 	Set("Icons.MemTag.TreeItem", new TODO_IMAGE_BRUSH(Icon12x12));
 
-	Set("TreeTable.TooltipBold", FTextBlockStyle(NormalText)
-		.SetFont(DEFAULT_FONT("Bold", 8))
-		.SetColorAndOpacity(FLinearColor(0.5f, 0.5f, 0.5f, 1.0f))
-		.SetShadowOffset(FVector2D(1.0f, 1.0f))
-		.SetShadowColorAndOpacity(FLinearColor(0.0f, 0.0f, 0.0f, 0.8f))
-	);
-
-	Set("TreeTable.Tooltip", FTextBlockStyle(NormalText)
-		.SetFont(DEFAULT_FONT("Regular", 8))
-		.SetColorAndOpacity(FLinearColor::White)
-		.SetShadowOffset(FVector2D(1.0f, 1.0f))
-		.SetShadowColorAndOpacity(FLinearColor(0.0f, 0.0f, 0.0f, 0.8f))
-	);
-
-	Set("TreeTable.NameText", FTextBlockStyle(NormalText)
-		.SetColorAndOpacity(FLinearColor::White)
-	);
-
-	Set("TreeTable.NormalText", FTextBlockStyle(NormalText)
-		.SetColorAndOpacity(FLinearColor::White)
-	);
-
-	//////////////////////////////////////////////////
-
-	// NormalEditableTextBox && SearchBox
-	{
-		const FEditableTextBoxStyle& NormalEditableTextBoxStyle = FAppStyle::GetWidgetStyle<FEditableTextBoxStyle>("NormalEditableTextBox");
-		const FTextBlockStyle TextBlockStyle = FTextBlockStyle()
-			.SetColorAndOpacity(NormalEditableTextBoxStyle.ForegroundColor)
-			.SetHighlightColor(NormalEditableTextBoxStyle.FocusedForegroundColor)
-			.SetFont(NormalEditableTextBoxStyle.TextStyle.Font)
-			.SetFontSize(uint16(NormalEditableTextBoxStyle.TextStyle.Font.Size));
-		const FEditableTextBoxStyle EditableTextBoxStyle = FEditableTextBoxStyle(NormalEditableTextBoxStyle)
-			.SetPadding(FMargin(6.0f, 4.0f, 6.0f, 4.0f))
-			.SetTextStyle(TextBlockStyle);
-		Set("NormalEditableTextBox", EditableTextBoxStyle);
-
-		const FSearchBoxStyle& NormalSearchBoxStyle = FAppStyle::GetWidgetStyle<FSearchBoxStyle>("SearchBox");
-		const FSearchBoxStyle SearchBoxStyle = FSearchBoxStyle(NormalSearchBoxStyle)
-			.SetTextBoxStyle(EditableTextBoxStyle);
-		Set("SearchBox", SearchBoxStyle);
-	}
-
-	// PrimaryToolbar
-	{
-		FToolBarStyle PrimaryToolbarStyle = FAppStyle::Get().GetWidgetStyle<FToolBarStyle>("SlimToolBar");
-
-		Set("PrimaryToolbar", PrimaryToolbarStyle);
-
-		Set("PrimaryToolbar.MinUniformToolbarSize", 40.0f);
-		Set("PrimaryToolbar.MaxUniformToolbarSize", 40.0f);
-	}
-
-	// SecondaryToolbar
-	{
-		FToolBarStyle SecondaryToolbarStyle = FAppStyle::Get().GetWidgetStyle<FToolBarStyle>("SlimToolBar");
-
-		SecondaryToolbarStyle.SetBackgroundPadding(         FMargin(4.0f, 4.0f));
-		SecondaryToolbarStyle.SetBlockPadding(              FMargin(2.0f, 0.0f));
-		SecondaryToolbarStyle.SetButtonPadding(             FMargin(0.0f, 0.0f));
-		SecondaryToolbarStyle.SetCheckBoxPadding(           FMargin(2.0f, 0.0f));
-		SecondaryToolbarStyle.SetComboButtonPadding(        FMargin(0.0f, 0.0f));
-		SecondaryToolbarStyle.SetIndentedBlockPadding(      FMargin(2.0f, 0.0f));
-		SecondaryToolbarStyle.SetLabelPadding(              FMargin(2.0f, 0.0f));
-		SecondaryToolbarStyle.SetSeparatorPadding(          FMargin(2.0f, -3.0f));
-
-		SecondaryToolbarStyle.ToggleButton.SetPadding(      FMargin(0.0f, 0.0f));
-
-		SecondaryToolbarStyle.ButtonStyle.SetNormalPadding( FMargin(6.0f, 2.0f, 4.0f, 2.0f));
-		SecondaryToolbarStyle.ButtonStyle.SetPressedPadding(FMargin(6.0f, 2.0f, 4.0f, 2.0f));
-
-		//SecondaryToolbarStyle.IconSize.Set(16.0f, 16.0f);
-
-		Set("SecondaryToolbar", SecondaryToolbarStyle);
-
-		Set("SecondaryToolbar.MinUniformToolbarSize", 32.0f);
-		Set("SecondaryToolbar.MaxUniformToolbarSize", 32.0f);
-	}
-
-	// SecondaryToolbar2 (used by AutoScroll and NetPacketContentView toolbars)
-	{
-		FToolBarStyle SecondaryToolbarStyle = FAppStyle::Get().GetWidgetStyle<FToolBarStyle>("SlimToolBar");
-
-		SecondaryToolbarStyle.SetBackgroundPadding(         FMargin(4.0f, 2.0f));
-		SecondaryToolbarStyle.SetBlockPadding(              FMargin(2.0f, 2.0f));
-		SecondaryToolbarStyle.SetButtonPadding(             FMargin(0.0f, 2.0f));
-		SecondaryToolbarStyle.SetCheckBoxPadding(           FMargin(2.0f, 2.0f));
-		SecondaryToolbarStyle.SetComboButtonPadding(        FMargin(0.0f, 2.0f));
-		SecondaryToolbarStyle.SetIndentedBlockPadding(      FMargin(2.0f, 2.0f));
-		SecondaryToolbarStyle.SetLabelPadding(              FMargin(2.0f, 2.0f));
-		SecondaryToolbarStyle.SetSeparatorPadding(          FMargin(2.0f, -1.0f));
-
-		SecondaryToolbarStyle.ToggleButton.SetPadding(      FMargin(0.0f, 0.0f));
-
-		SecondaryToolbarStyle.ButtonStyle.SetNormalPadding( FMargin(3.0f, 0.0f, -1.0f, 0.0f));
-		SecondaryToolbarStyle.ButtonStyle.SetPressedPadding(FMargin(3.0f, 0.0f, -1.0f, 0.0f));
-
-		//SecondaryToolbarStyle.IconSize.Set(16.0f, 16.0f);
-
-		Set("SecondaryToolbar2", SecondaryToolbarStyle);
-
-		Set("SecondaryToolbar2.MinUniformToolbarSize", 32.0f);
-		Set("SecondaryToolbar2.MaxUniformToolbarSize", 32.0f);
-	}
-
-	// ToggleButton
-	{
-		Set("ToggleButton", FButtonStyle(Button)
-			.SetNormal(FSlateNoResource())
-			.SetHovered(EDITOR_BOX_BRUSH("Common/RoundedSelection_16x", 4.0f / 16.0f, SelectionColor))
-			.SetPressed(EDITOR_BOX_BRUSH("Common/RoundedSelection_16x", 4.0f / 16.0f, SelectionColor_Pressed)));
-	}
-
-	// Common.GotoNativeCodeHyperlink
-	{
-		FTextBlockStyle InheritedFromNativeTextStyle = FTextBlockStyle(NormalText)
-			.SetFont(DEFAULT_FONT("Regular", 10));
-
-		Set("Common.InheritedFromNativeTextStyle", InheritedFromNativeTextStyle);
-
-		// Go to native class hyperlink
-		FButtonStyle EditNativeHyperlinkButton = FButtonStyle()
-			.SetNormal(EDITOR_BORDER_BRUSH("Old/HyperlinkDotted", FMargin(0, 0, 0, 3 / 16.0f)))
-			.SetPressed(FSlateNoResource())
-			.SetHovered(EDITOR_BORDER_BRUSH("Old/HyperlinkUnderline", FMargin(0, 0, 0, 3 / 16.0f)));
-		FHyperlinkStyle EditNativeHyperlinkStyle = FHyperlinkStyle()
-			.SetUnderlineStyle(EditNativeHyperlinkButton)
-			.SetTextStyle(InheritedFromNativeTextStyle)
-			.SetPadding(FMargin(0.0f));
-
-		Set("Common.GotoNativeCodeHyperlink", EditNativeHyperlinkStyle);
-	}
-
 	//////////////////////////////////////////////////
 }
 
diff --git a/Engine/Source/Developer/TraceInsights/Private/Insights/LoadingProfiler/Widgets/SLoadingProfilerWindow.cpp b/Engine/Source/Developer/TraceInsights/Private/Insights/LoadingProfiler/Widgets/SLoadingProfilerWindow.cpp
--- a/Engine/Source/Developer/TraceInsights/Private/Insights/LoadingProfiler/Widgets/SLoadingProfilerWindow.cpp
+++ b/Engine/Source/Developer/TraceInsights/Private/Insights/LoadingProfiler/Widgets/SLoadingProfilerWindow.cpp
@@ -267,14 +267,14 @@ TSharedRef<SDockTab> SLoadingProfilerWindow::SpawnTab_EventAggregationTreeView(c
 {
 	FLoadingProfilerManager::Get()->SetEventAggregationTreeViewVisible(true);
 
-	TSharedRef<::Insights::FUntypedTable> Table = MakeShared<::Insights::FUntypedTable>();
+	TSharedRef<FUntypedTable> Table = MakeShared<FUntypedTable>();
 	Table->SetDisplayName(LOCTEXT("EventAggregation_TableName", "Event Aggregation"));
 
 	const TSharedRef<SDockTab> DockTab = SNew(SDockTab)
 		.ShouldAutosize(false)
 		.TabRole(ETabRole::PanelTab)
 		[
-			SAssignNew(EventAggregationTreeView, ::Insights::SUntypedTableTreeView, Table)
+			SAssignNew(EventAggregationTreeView, SUntypedTableTreeView, Table)
 		];
 
 	EventAggregationTreeView->SetLogListingName(FLoadingProfilerManager::Get()->GetLogListingName());
@@ -304,14 +304,14 @@ TSharedRef<SDockTab> SLoadingProfilerWindow::SpawnTab_ObjectTypeAggregationTreeV
 {
 	FLoadingProfilerManager::Get()->SetObjectTypeAggregationTreeViewVisible(true);
 
-	TSharedRef<::Insights::FUntypedTable> Table = MakeShared<::Insights::FUntypedTable>();
+	TSharedRef<FUntypedTable> Table = MakeShared<FUntypedTable>();
 	Table->SetDisplayName(LOCTEXT("ObjectTypeAggregation_TableName", "Object Type Aggregation"));
 
 	const TSharedRef<SDockTab> DockTab = SNew(SDockTab)
 		.ShouldAutosize(false)
 		.TabRole(ETabRole::PanelTab)
 		[
-			SAssignNew(ObjectTypeAggregationTreeView, ::Insights::SUntypedTableTreeView, Table)
+			SAssignNew(ObjectTypeAggregationTreeView, SUntypedTableTreeView, Table)
 		];
 
 	ObjectTypeAggregationTreeView->SetLogListingName(FLoadingProfilerManager::Get()->GetLogListingName());
@@ -341,14 +341,14 @@ TSharedRef<SDockTab> SLoadingProfilerWindow::SpawnTab_PackageDetailsTreeView(con
 {
 	FLoadingProfilerManager::Get()->SetPackageDetailsTreeViewVisible(true);
 
-	TSharedRef<::Insights::FUntypedTable> Table = MakeShared<::Insights::FUntypedTable>();
+	TSharedRef<FUntypedTable> Table = MakeShared<FUntypedTable>();
 	Table->SetDisplayName(LOCTEXT("PackageDetails_TableName", "Package Details"));
 
 	const TSharedRef<SDockTab> DockTab = SNew(SDockTab)
 		.ShouldAutosize(false)
 		.TabRole(ETabRole::PanelTab)
 		[
-			SAssignNew(PackageDetailsTreeView, ::Insights::SUntypedTableTreeView, Table)
+			SAssignNew(PackageDetailsTreeView, SUntypedTableTreeView, Table)
 		];
 
 	PackageDetailsTreeView->SetLogListingName(FLoadingProfilerManager::Get()->GetLogListingName());
@@ -378,14 +378,14 @@ TSharedRef<SDockTab> SLoadingProfilerWindow::SpawnTab_ExportDetailsTreeView(cons
 {
 	FLoadingProfilerManager::Get()->SetExportDetailsTreeViewVisible(true);
 
-	TSharedRef<::Insights::FUntypedTable> Table = MakeShared<::Insights::FUntypedTable>();
+	TSharedRef<FUntypedTable> Table = MakeShared<FUntypedTable>();
 	Table->SetDisplayName(LOCTEXT("ExportDetails_TableName", "Export Details"));
 
 	const TSharedRef<SDockTab> DockTab = SNew(SDockTab)
 		.ShouldAutosize(false)
 		.TabRole(ETabRole::PanelTab)
 		[
-			SAssignNew(ExportDetailsTreeView, ::Insights::SUntypedTableTreeView, Table)
+			SAssignNew(ExportDetailsTreeView, SUntypedTableTreeView, Table)
 		];
 
 	ExportDetailsTreeView->SetLogListingName(FLoadingProfilerManager::Get()->GetLogListingName());
@@ -415,14 +415,14 @@ TSharedRef<SDockTab> SLoadingProfilerWindow::SpawnTab_RequestsTreeView(const FSp
 {
 	FLoadingProfilerManager::Get()->SetRequestsTreeViewVisible(true);
 
-	TSharedRef<::Insights::FUntypedTable> Table = MakeShared<::Insights::FUntypedTable>();
+	TSharedRef<FUntypedTable> Table = MakeShared<FUntypedTable>();
 	Table->SetDisplayName(LOCTEXT("Requests_TableName", "Requests"));
 
 	const TSharedRef<SDockTab> DockTab = SNew(SDockTab)
 		.ShouldAutosize(false)
 		.TabRole(ETabRole::PanelTab)
 		[
-			SAssignNew(RequestsTreeView, ::Insights::SUntypedTableTreeView, Table)
+			SAssignNew(RequestsTreeView, SUntypedTableTreeView, Table)
 		];
 
 	RequestsTreeView->SetLogListingName(FLoadingProfilerManager::Get()->GetLogListingName());
diff --git a/Engine/Source/Developer/TraceInsights/Private/Insights/LoadingProfiler/Widgets/SLoadingProfilerWindow.h b/Engine/Source/Developer/TraceInsights/Private/Insights/LoadingProfiler/Widgets/SLoadingProfilerWindow.h
--- a/Engine/Source/Developer/TraceInsights/Private/Insights/LoadingProfiler/Widgets/SLoadingProfilerWindow.h
+++ b/Engine/Source/Developer/TraceInsights/Private/Insights/LoadingProfiler/Widgets/SLoadingProfilerWindow.h
@@ -8,8 +8,8 @@
 #include "Insights/ITimingViewSession.h" // for ETimeChangedFlags
 #include "Insights/Widgets/SMajorTabWindow.h"
 
+namespace UE::Insights { class SUntypedTableTreeView; }
 namespace UE::Insights::TimingProfiler { class STimingView; }
-namespace Insights { class SUntypedTableTreeView; }
 
 namespace UE::Insights::LoadingProfiler
 {
@@ -55,11 +55,11 @@ public:
 	void Construct(const FArguments& InArgs, const TSharedRef<SDockTab>& ConstructUnderMajorTab, const TSharedPtr<SWindow>& ConstructUnderWindow);
 
 	TSharedPtr<TimingProfiler::STimingView> GetTimingView() const { return TimingView; }
-	TSharedPtr<::Insights::SUntypedTableTreeView> GetEventAggregationTreeView() const { return EventAggregationTreeView; }
-	TSharedPtr<::Insights::SUntypedTableTreeView> GetObjectTypeAggregationTreeView() const { return ObjectTypeAggregationTreeView; }
-	TSharedPtr<::Insights::SUntypedTableTreeView> GetPackageDetailsTreeView() const { return PackageDetailsTreeView; }
-	TSharedPtr<::Insights::SUntypedTableTreeView> GetExportDetailsTreeView() const { return ExportDetailsTreeView; }
-	TSharedPtr<::Insights::SUntypedTableTreeView> GetRequestsTreeView() const { return RequestsTreeView; }
+	TSharedPtr<SUntypedTableTreeView> GetEventAggregationTreeView() const { return EventAggregationTreeView; }
+	TSharedPtr<SUntypedTableTreeView> GetObjectTypeAggregationTreeView() const { return ObjectTypeAggregationTreeView; }
+	TSharedPtr<SUntypedTableTreeView> GetPackageDetailsTreeView() const { return PackageDetailsTreeView; }
+	TSharedPtr<SUntypedTableTreeView> GetExportDetailsTreeView() const { return ExportDetailsTreeView; }
+	TSharedPtr<SUntypedTableTreeView> GetRequestsTreeView() const { return RequestsTreeView; }
 
 protected:
 	virtual const TCHAR* GetAnalyticsEventName() const override;
@@ -94,19 +94,19 @@ private:
 	TSharedPtr<TimingProfiler::STimingView> TimingView;
 
 	/** The Event Aggregation tree view widget */
-	TSharedPtr<::Insights::SUntypedTableTreeView> EventAggregationTreeView;
+	TSharedPtr<SUntypedTableTreeView> EventAggregationTreeView;
 
 	/** The Object Type Aggregation tree view widget */
-	TSharedPtr<::Insights::SUntypedTableTreeView> ObjectTypeAggregationTreeView;
+	TSharedPtr<SUntypedTableTreeView> ObjectTypeAggregationTreeView;
 
 	/** The Package Details tree view widget */
-	TSharedPtr<::Insights::SUntypedTableTreeView> PackageDetailsTreeView;
+	TSharedPtr<SUntypedTableTreeView> PackageDetailsTreeView;
 
 	/** The Export Details tree view widget */
-	TSharedPtr<::Insights::SUntypedTableTreeView> ExportDetailsTreeView;
+	TSharedPtr<SUntypedTableTreeView> ExportDetailsTreeView;
 
 	/** The Requests tree view widget */
-	TSharedPtr<::Insights::SUntypedTableTreeView> RequestsTreeView;
+	TSharedPtr<SUntypedTableTreeView> RequestsTreeView;
 
 	double SelectionStartTime;
 	double SelectionEndTime;
diff --git a/Engine/Source/Developer/TraceInsights/Private/Insights/MemoryProfiler/Widgets/SMemTagTreeViewTableCell.cpp b/Engine/Source/Developer/TraceInsights/Private/Insights/MemoryProfiler/Widgets/SMemTagTreeViewTableCell.cpp
--- a/Engine/Source/Developer/TraceInsights/Private/Insights/MemoryProfiler/Widgets/SMemTagTreeViewTableCell.cpp
+++ b/Engine/Source/Developer/TraceInsights/Private/Insights/MemoryProfiler/Widgets/SMemTagTreeViewTableCell.cpp
@@ -99,7 +99,7 @@ TSharedRef<SWidget> SMemTagTreeViewTableCell::GenerateWidgetForNameColumn(const
 			.HeightOverride(14.0f)
 			[
 				SNew(SBorder)
-				.BorderImage(FInsightsStyle::Get().GetBrush("WhiteBrush"))
+				.BorderImage(FAppStyle::Get().GetBrush("WhiteBrush"))
 				.BorderBackgroundColor(this, &SMemTagTreeViewTableCell::GetBoxColorAndOpacity)
 				.HAlign(HAlign_Fill)
 				.VAlign(VAlign_Fill)
diff --git a/Engine/Source/Developer/TraceInsights/Private/Insights/NetworkingProfiler/ViewModels/PacketContentViewDrawHelper.cpp b/Engine/Source/Developer/TraceInsights/Private/Insights/NetworkingProfiler/ViewModels/PacketContentViewDrawHelper.cpp
--- a/Engine/Source/Developer/TraceInsights/Private/Insights/NetworkingProfiler/ViewModels/PacketContentViewDrawHelper.cpp
+++ b/Engine/Source/Developer/TraceInsights/Private/Insights/NetworkingProfiler/ViewModels/PacketContentViewDrawHelper.cpp
@@ -313,7 +313,7 @@ void FPacketContentViewDrawStateBuilder::Flush()
 FPacketContentViewDrawHelper::FPacketContentViewDrawHelper(const FDrawContext& InDrawContext, const FPacketContentViewport& InViewport)
 	: DrawContext(InDrawContext)
 	, Viewport(InViewport)
-	, WhiteBrush(FInsightsStyle::Get().GetBrush("WhiteBrush"))
+	, WhiteBrush(FAppStyle::Get().GetBrush("WhiteBrush"))
 	, EventBorderBrush(FInsightsStyle::Get().GetBrush("EventBorder"))
 	, HoveredEventBorderBrush(FInsightsStyle::Get().GetBrush("HoveredEventBorder"))
 	, SelectedEventBorderBrush(FInsightsStyle::Get().GetBrush("SelectedEventBorder"))
diff --git a/Engine/Source/Developer/TraceInsights/Private/Insights/NetworkingProfiler/ViewModels/PacketViewDrawHelper.cpp b/Engine/Source/Developer/TraceInsights/Private/Insights/NetworkingProfiler/ViewModels/PacketViewDrawHelper.cpp
--- a/Engine/Source/Developer/TraceInsights/Private/Insights/NetworkingProfiler/ViewModels/PacketViewDrawHelper.cpp
+++ b/Engine/Source/Developer/TraceInsights/Private/Insights/NetworkingProfiler/ViewModels/PacketViewDrawHelper.cpp
@@ -117,7 +117,7 @@ void FNetworkPacketSeriesBuilder::SetHighlightEventTypeIndex(int32 EventTypeInde
 FPacketViewDrawHelper::FPacketViewDrawHelper(const FDrawContext& InDrawContext, const FPacketViewport& InViewport)
 	: DrawContext(InDrawContext)
 	, Viewport(InViewport)
-	, WhiteBrush(FInsightsStyle::Get().GetBrush("WhiteBrush"))
+	, WhiteBrush(FAppStyle::Get().GetBrush("WhiteBrush"))
 	//, EventBorderBrush(FInsightsStyle::Get().GetBrush("EventBorder"))
 	, HoveredEventBorderBrush(FInsightsStyle::Get().GetBrush("HoveredEventBorder"))
 	, SelectedEventBorderBrush(FInsightsStyle::Get().GetBrush("SelectedEventBorder"))
diff --git a/Engine/Source/Developer/TraceInsights/Private/Insights/NetworkingProfiler/Widgets/SPacketContentView.cpp b/Engine/Source/Developer/TraceInsights/Private/Insights/NetworkingProfiler/Widgets/SPacketContentView.cpp
--- a/Engine/Source/Developer/TraceInsights/Private/Insights/NetworkingProfiler/Widgets/SPacketContentView.cpp
+++ b/Engine/Source/Developer/TraceInsights/Private/Insights/NetworkingProfiler/Widgets/SPacketContentView.cpp
@@ -1318,7 +1318,7 @@ int32 SPacketContentView::OnPaint(const FPaintArgs& Args, const FGeometry& Allot
 	const bool bShouldDisplayDebugInfo = FInsightsManager::Get()->IsDebugInfoEnabled();
 	if (bShouldDisplayDebugInfo)
 	{
-		const FSlateBrush* WhiteBrush = FInsightsStyle::Get().GetBrush("WhiteBrush");
+		const FSlateBrush* WhiteBrush = FAppStyle::Get().GetBrush("WhiteBrush");
 		FSlateFontInfo SummaryFont = FAppStyle::Get().GetFontStyle("SmallFont");
 
 		const TSharedRef<FSlateFontMeasure> FontMeasureService = FSlateApplication::Get().GetRenderer()->GetFontMeasureService();
diff --git a/Engine/Source/Developer/TraceInsights/Private/Insights/NetworkingProfiler/Widgets/SPacketView.cpp b/Engine/Source/Developer/TraceInsights/Private/Insights/NetworkingProfiler/Widgets/SPacketView.cpp
--- a/Engine/Source/Developer/TraceInsights/Private/Insights/NetworkingProfiler/Widgets/SPacketView.cpp
+++ b/Engine/Source/Developer/TraceInsights/Private/Insights/NetworkingProfiler/Widgets/SPacketView.cpp
@@ -763,7 +763,7 @@ int32 SPacketView::OnPaint(const FPaintArgs& Args, const FGeometry& AllottedGeom
 	FSlateFontInfo SummaryFont = FAppStyle::Get().GetFontStyle("SmallFont");
 	const float FontScale = AllottedGeometry.Scale;
 
-	const FSlateBrush* WhiteBrush = FInsightsStyle::Get().GetBrush("WhiteBrush");
+	const FSlateBrush* WhiteBrush = FAppStyle::Get().GetBrush("WhiteBrush");
 
 	const float ViewWidth = static_cast<float>(AllottedGeometry.Size.X);
 	const float ViewHeight = static_cast<float>(AllottedGeometry.Size.Y);
diff --git a/Engine/Source/Developer/TraceInsights/Private/Insights/Table/ViewModels/UntypedTable.cpp b/Engine/Source/Developer/TraceInsights/Private/Insights/Table/ViewModels/UntypedTable.cpp
--- a/Engine/Source/Developer/TraceInsights/Private/Insights/Table/ViewModels/UntypedTable.cpp
+++ b/Engine/Source/Developer/TraceInsights/Private/Insights/Table/ViewModels/UntypedTable.cpp
@@ -12,24 +12,22 @@
 #include "InsightsCore/Table/ViewModels/TableColumn.h"
 #include "InsightsCore/Table/ViewModels/TableTreeNode.h"
 
-#define LOCTEXT_NAMESPACE "Insights::FUntypedTable"
+#define LOCTEXT_NAMESPACE "UE::Insights::FUntypedTable"
 
-namespace Insights
+namespace UE::Insights
 {
 
 ////////////////////////////////////////////////////////////////////////////////////////////////////
 // FUntypedTableTreeNodeValueGetter
 ////////////////////////////////////////////////////////////////////////////////////////////////////
 
-class FUntypedTableTreeNodeValueGetter : public UE::Insights::FTableCellValueGetter
+class FUntypedTableTreeNodeValueGetter : public FTableCellValueGetter
 {
 public:
-	FUntypedTableTreeNodeValueGetter(UE::Insights::ETableCellDataType InDataType) : UE::Insights::FTableCellValueGetter(), DataType(InDataType) {}
+	FUntypedTableTreeNodeValueGetter(ETableCellDataType InDataType) : FTableCellValueGetter(), DataType(InDataType) {}
 
-	virtual const TOptional<UE::Insights::FTableCellValue> GetValue(const UE::Insights::FTableColumn& Column, const UE::Insights::FBaseTreeNode& Node) const
+	virtual const TOptional<FTableCellValue> GetValue(const FTableColumn& Column, const FBaseTreeNode& Node) const
 	{
-		using namespace UE::Insights;
-
 		ensure(Node.Is<FTableTreeNode>());
 		const FTableTreeNode& TableTreeNode = static_cast<const FTableTreeNode&>(Node);
 
@@ -72,7 +70,7 @@ public:
 	}
 
 private:
-	UE::Insights::ETableCellDataType DataType;
+	ETableCellDataType DataType;
 };
 
 ////////////////////////////////////////////////////////////////////////////////////////////////////
@@ -162,8 +160,6 @@ bool FUntypedTable::UpdateSourceTable(TSharedPtr<TraceServices::IUntypedTable> I
 
 void FUntypedTable::CreateColumns(const TraceServices::ITableLayout& TableLayout)
 {
-	using namespace UE::Insights;
-
 	ensure(GetColumnCount() == 0);
 	const int32 ColumnCount = static_cast<int32>(TableLayout.GetColumnCount());
 
@@ -331,6 +327,6 @@ void FUntypedTable::CreateColumns(const TraceServices::ITableLayout& TableLayout
 
 ////////////////////////////////////////////////////////////////////////////////////////////////////
 
-} // namespace Insights
+} // namespace UE::Insights
 
 #undef LOCTEXT_NAMESPACE
diff --git a/Engine/Source/Developer/TraceInsights/Private/Insights/Table/ViewModels/UntypedTable.h b/Engine/Source/Developer/TraceInsights/Private/Insights/Table/ViewModels/UntypedTable.h
--- a/Engine/Source/Developer/TraceInsights/Private/Insights/Table/ViewModels/UntypedTable.h
+++ b/Engine/Source/Developer/TraceInsights/Private/Insights/Table/ViewModels/UntypedTable.h
@@ -16,12 +16,12 @@ namespace TraceServices
 	class IUntypedTableReader;
 }
 
-namespace Insights
+namespace UE::Insights
 {
 
 ////////////////////////////////////////////////////////////////////////////////////////////////////
 
-class FUntypedTable : public UE::Insights::FTable
+class FUntypedTable : public FTable
 {
 public:
 	FUntypedTable();
@@ -45,4 +45,4 @@ private:
 
 ////////////////////////////////////////////////////////////////////////////////////////////////////
 
-} // namespace Insights
+} // namespace UE::Insights
diff --git a/Engine/Source/Developer/TraceInsights/Private/Insights/Table/Widgets/SSessionTableTreeView.cpp b/Engine/Source/Developer/TraceInsights/Private/Insights/Table/Widgets/SSessionTableTreeView.cpp
--- a/Engine/Source/Developer/TraceInsights/Private/Insights/Table/Widgets/SSessionTableTreeView.cpp
+++ b/Engine/Source/Developer/TraceInsights/Private/Insights/Table/Widgets/SSessionTableTreeView.cpp
@@ -5,11 +5,9 @@
 // TraceInsights
 #include "Insights/InsightsManager.h"
 
-#define LOCTEXT_NAMESPACE "Insights::SSessionTableTreeView"
+#define LOCTEXT_NAMESPACE "UE::Insights::SSessionTableTreeView"
 
-namespace UE
-{
-namespace Insights
+namespace UE::Insights
 {
 
 ////////////////////////////////////////////////////////////////////////////////////////////////////
@@ -61,7 +59,6 @@ void SSessionTableTreeView::InsightsManager_OnSessionChanged()
 
 ////////////////////////////////////////////////////////////////////////////////////////////////////
 
-} // namespace Insights
-} // namespace UE
+} // namespace UE::Insights
 
 #undef LOCTEXT_NAMESPACE
diff --git a/Engine/Source/Developer/TraceInsights/Private/Insights/Table/Widgets/SSessionTableTreeView.h b/Engine/Source/Developer/TraceInsights/Private/Insights/Table/Widgets/SSessionTableTreeView.h
--- a/Engine/Source/Developer/TraceInsights/Private/Insights/Table/Widgets/SSessionTableTreeView.h
+++ b/Engine/Source/Developer/TraceInsights/Private/Insights/Table/Widgets/SSessionTableTreeView.h
@@ -12,9 +12,7 @@ namespace TraceServices
 	class IAnalysisSession;
 }
 
-namespace UE
-{
-namespace Insights
+namespace UE::Insights
 {
 
 class FTable;
@@ -30,9 +28,9 @@ protected:
 	/** Called when the analysis session has changed. */
 	void InsightsManager_OnSessionChanged();
 
+protected:
 	/** The analysis session used to populate this widget. */
 	TSharedPtr<const TraceServices::IAnalysisSession> Session;
 };
 
-} // namespace Insights
-} // namespace UE
+} // namespace UE::Insights
diff --git a/Engine/Source/Developer/TraceInsights/Private/Insights/Table/Widgets/SUntypedDiffTableTreeView.cpp b/Engine/Source/Developer/TraceInsights/Private/Insights/Table/Widgets/SUntypedDiffTableTreeView.cpp
--- a/Engine/Source/Developer/TraceInsights/Private/Insights/Table/Widgets/SUntypedDiffTableTreeView.cpp
+++ b/Engine/Source/Developer/TraceInsights/Private/Insights/Table/Widgets/SUntypedDiffTableTreeView.cpp
@@ -11,9 +11,9 @@
 // TraceInsights
 #include "Insights/InsightsManager.h"
 
-#define LOCTEXT_NAMESPACE "SUntypedDiffTableTreeView"
+#define LOCTEXT_NAMESPACE "UE::Insights::SUntypedDiffTableTreeView"
 
-namespace Insights
+namespace UE::Insights
 {
 
 ////////////////////////////////////////////////////////////////////////////////////////////////////
@@ -117,6 +117,6 @@ void SUntypedDiffTableTreeView::RequestMergeTables()
 
 ////////////////////////////////////////////////////////////////////////////////////////////////////
 
-} // namespace Insights
+} // namespace UE::Insights
 
 #undef LOCTEXT_NAMESPACE
diff --git a/Engine/Source/Developer/TraceInsights/Private/Insights/Table/Widgets/SUntypedDiffTableTreeView.h b/Engine/Source/Developer/TraceInsights/Private/Insights/Table/Widgets/SUntypedDiffTableTreeView.h
--- a/Engine/Source/Developer/TraceInsights/Private/Insights/Table/Widgets/SUntypedDiffTableTreeView.h
+++ b/Engine/Source/Developer/TraceInsights/Private/Insights/Table/Widgets/SUntypedDiffTableTreeView.h
@@ -5,10 +5,10 @@
 #include "CoreMinimal.h"
 
 // TraceInsights
-#include "SUntypedTableTreeView.h"
 #include "Insights/Table/ViewModels/UntypedTable.h"
+#include "Insights/Table/Widgets/SUntypedTableTreeView.h"
 
-namespace Insights
+namespace UE::Insights
 {
 
 ////////////////////////////////////////////////////////////////////////////////////////////////////
@@ -36,4 +36,4 @@ private:
 
 ////////////////////////////////////////////////////////////////////////////////////////////////////
 
-} // namespace Insights
+} // namespace UE::Insights
diff --git a/Engine/Source/Developer/TraceInsights/Private/Insights/Table/Widgets/SUntypedTableTreeView.cpp b/Engine/Source/Developer/TraceInsights/Private/Insights/Table/Widgets/SUntypedTableTreeView.cpp
--- a/Engine/Source/Developer/TraceInsights/Private/Insights/Table/Widgets/SUntypedTableTreeView.cpp
+++ b/Engine/Source/Developer/TraceInsights/Private/Insights/Table/Widgets/SUntypedTableTreeView.cpp
@@ -13,9 +13,9 @@
 #include "Insights/Log.h"
 #include "Insights/Table/ViewModels/UntypedTable.h"
 
-#define LOCTEXT_NAMESPACE "SUntypedTableTreeView"
+#define LOCTEXT_NAMESPACE "UE::Insights::SUntypedTableTreeView"
 
-namespace Insights
+namespace UE::Insights
 {
 
 ////////////////////////////////////////////////////////////////////////////////////////////////////
@@ -32,7 +32,7 @@ SUntypedTableTreeView::~SUntypedTableTreeView()
 
 ////////////////////////////////////////////////////////////////////////////////////////////////////
 
-void SUntypedTableTreeView::Construct(const FArguments& InArgs, TSharedPtr<Insights::FUntypedTable> InTablePtr)
+void SUntypedTableTreeView::Construct(const FArguments& InArgs, TSharedPtr<FUntypedTable> InTablePtr)
 {
 	bRunInAsyncMode = InArgs._RunInAsyncMode;
 	ConstructWidget(InTablePtr);
@@ -42,7 +42,6 @@ void SUntypedTableTreeView::Construct(const FArguments& InArgs, TSharedPtr<Insig
 
 void SUntypedTableTreeView::Reset()
 {
-	//...
 	STableTreeView::Reset();
 }
 
@@ -50,8 +49,8 @@ void SUntypedTableTreeView::Reset()
 
 void SUntypedTableTreeView::UpdateSourceTable(TSharedPtr<TraceServices::IUntypedTable> SourceTable)
 {
-	//check(Table->Is<Insights::FUntypedTable>());
-	TSharedPtr<Insights::FUntypedTable> UntypedTable = StaticCastSharedPtr<Insights::FUntypedTable>(Table);
+	//check(Table->Is<FUntypedTable>());
+	TSharedPtr<FUntypedTable> UntypedTable = StaticCastSharedPtr<FUntypedTable>(Table);
 
 	if (UntypedTable->UpdateSourceTable(SourceTable))
 	{
@@ -65,10 +64,10 @@ void SUntypedTableTreeView::UpdateSourceTable(TSharedPtr<TraceServices::IUntyped
 
 void SUntypedTableTreeView::RebuildTree(bool bResync)
 {
-	UE::Insights::FStopwatch Stopwatch;
+	FStopwatch Stopwatch;
 	Stopwatch.Start();
 
-	UE::Insights::FStopwatch SyncStopwatch;
+	FStopwatch SyncStopwatch;
 	SyncStopwatch.Start();
 
 	if (bResync)
@@ -78,8 +77,8 @@ void SUntypedTableTreeView::RebuildTree(bool bResync)
 
 	const int32 PreviousNodeCount = TableRowNodes.Num();
 
-	//check(Table->Is<Insights::FUntypedTable>());
-	TSharedPtr<Insights::FUntypedTable> UntypedTable = StaticCastSharedPtr<Insights::FUntypedTable>(Table);
+	//check(Table->Is<FUntypedTable>());
+	TSharedPtr<FUntypedTable> UntypedTable = StaticCastSharedPtr<FUntypedTable>(Table);
 
 	TSharedPtr<TraceServices::IUntypedTable> SourceTable = UntypedTable->GetSourceTable();
 	TSharedPtr<TraceServices::IUntypedTableReader> TableReader = UntypedTable->GetTableReader();
@@ -95,7 +94,7 @@ void SUntypedTableTreeView::RebuildTree(bool bResync)
 			{
 				TableReader->SetRowIndex(RowIndex);
 				FName NodeName(BaseNodeName, RowIndex + 1);
-				UE::Insights::FTableTreeNodePtr NodePtr = MakeShared<UE::Insights::FTableTreeNode>(NodeName, Table, RowIndex);
+				FTableTreeNodePtr NodePtr = MakeShared<FTableTreeNode>(NodeName, Table, RowIndex);
 				NodePtr->SetDefaultSortOrder(RowIndex + 1);
 				TableRowNodes.Add(NodePtr);
 			}
@@ -108,7 +107,7 @@ void SUntypedTableTreeView::RebuildTree(bool bResync)
 	if (bResync || TableRowNodes.Num() != PreviousNodeCount)
 	{
 		// Save selection.
-		TArray<UE::Insights::FTableTreeNodePtr> SelectedItems;
+		TArray<FTableTreeNodePtr> SelectedItems;
 		TreeView->GetSelectedItems(SelectedItems);
 
 		UpdateTree();
@@ -119,11 +118,11 @@ void SUntypedTableTreeView::RebuildTree(bool bResync)
 		if (SelectedItems.Num() > 0)
 		{
 			TreeView->ClearSelection();
-			for (UE::Insights::FTableTreeNodePtr& NodePtr : SelectedItems)
+			for (FTableTreeNodePtr& NodePtr : SelectedItems)
 			{
 				NodePtr = GetNodeByTableRowIndex(NodePtr->GetRowIndex());
 			}
-			SelectedItems.RemoveAll([](const UE::Insights::FTableTreeNodePtr& NodePtr) { return !NodePtr.IsValid(); });
+			SelectedItems.RemoveAll([](const FTableTreeNodePtr& NodePtr) { return !NodePtr.IsValid(); });
 			if (SelectedItems.Num() > 0)
 			{
 				TreeView->SetItemSelection(SelectedItems, true);
@@ -194,6 +193,6 @@ void SUntypedTableTreeView::ClearCurrentOperationNameOverride()
 
 ////////////////////////////////////////////////////////////////////////////////////////////////////
 
-} // namespace Insights
+} // namespace UE::Insights
 
 #undef LOCTEXT_NAMESPACE
diff --git a/Engine/Source/Developer/TraceInsights/Private/Insights/Table/Widgets/SUntypedTableTreeView.h b/Engine/Source/Developer/TraceInsights/Private/Insights/Table/Widgets/SUntypedTableTreeView.h
--- a/Engine/Source/Developer/TraceInsights/Private/Insights/Table/Widgets/SUntypedTableTreeView.h
+++ b/Engine/Source/Developer/TraceInsights/Private/Insights/Table/Widgets/SUntypedTableTreeView.h
@@ -11,18 +11,15 @@
 // TraceInsights
 #include "Insights/Table/ViewModels/UntypedTable.h"
 
-namespace Insights
+namespace UE::Insights
 {
 
 ////////////////////////////////////////////////////////////////////////////////////////////////////
 
-class SUntypedTableTreeView : public UE::Insights::STableTreeView
+class SUntypedTableTreeView : public STableTreeView
 {
 public:
-	/** Default constructor. */
 	SUntypedTableTreeView();
-
-	/** Virtual destructor. */
 	virtual ~SUntypedTableTreeView();
 
 	SLATE_BEGIN_ARGS(SUntypedTableTreeView)
@@ -35,16 +32,16 @@ public:
 	 * Construct this widget
 	 * @param InArgs - The declaration data for this widget
 	 */
-	void Construct(const FArguments& InArgs, TSharedPtr<Insights::FUntypedTable> InTablePtr);
+	void Construct(const FArguments& InArgs, TSharedPtr<FUntypedTable> InTablePtr);
 
-	TSharedPtr<Insights::FUntypedTable> GetUntypedTable() const { return StaticCastSharedPtr<Insights::FUntypedTable>(GetTable()); }
+	TSharedPtr<FUntypedTable> GetUntypedTable() const { return StaticCastSharedPtr<FUntypedTable>(GetTable()); }
 
 	void UpdateSourceTable(TSharedPtr<TraceServices::IUntypedTable> SourceTable);
 
 	virtual void Reset();
 
-	////////////////////////////////////////////////////////////////////////////////////////////////////
-	// IAsyncOperationStatusProvider implementation
+	//////////////////////////////////////////////////
+	// IAsyncOperationStatusProvider
 
 	virtual bool IsRunning() const override;
 	virtual double GetAllOperationsDuration() override;
@@ -52,7 +49,8 @@ public:
 	virtual uint32 GetOperationCount() const override { return 1; }
 	virtual FText GetCurrentOperationName() const override;
 
-	////////////////////////////////////////////////////////////////////////////////////////////////////
+	//////////////////////////////////////////////////
+
 	void SetCurrentOperationNameOverride(const FText& InOperationName);
 	void ClearCurrentOperationNameOverride();
 
@@ -63,10 +61,10 @@ public:
 	virtual void RebuildTree(bool bResync);
 
 private:
-	UE::Insights::FStopwatch CurrentOperationStopwatch;
+	FStopwatch CurrentOperationStopwatch;
 	FText CurrentOperationNameOverride;
 };
 
 ////////////////////////////////////////////////////////////////////////////////////////////////////
 
-} // namespace Insights
+} // namespace UE::Insights
diff --git a/Engine/Source/Developer/TraceInsights/Private/Insights/Tests/FunctionalTests/ExportCommandsTests.cpp b/Engine/Source/Developer/TraceInsights/Private/Insights/Tests/FunctionalTests/ExportCommandsTests.cpp
--- a/Engine/Source/Developer/TraceInsights/Private/Insights/Tests/FunctionalTests/ExportCommandsTests.cpp
+++ b/Engine/Source/Developer/TraceInsights/Private/Insights/Tests/FunctionalTests/ExportCommandsTests.cpp
@@ -2,11 +2,19 @@
 
 #include "HAL/FileManager.h"
 #include "HAL/PlatformFileManager.h"
+#include "Misc/AutomationTest.h"
+#include "Misc/FileHelper.h"
+
+// TraceAnalysis
+#include "Trace/StoreClient.h"
+
+// TraceInsightsCore
+#include "InsightsCore/Common/MiscUtils.h"
+
+// TraceInsights
 #include "Insights/InsightsManager.h"
 #include "Insights/IUnrealInsightsModule.h"
 #include "Insights/Tests/InsightsTestUtils.h"
-#include "Misc/AutomationTest.h"
-#include "Misc/FileHelper.h"
 
 void VerifyExportedLines(const FString& ExportReportPath, const FString& CmdLogPath, const FString& Elements, FInsightsTestUtils Utils, FAutomationTestBase* Test, double Timeout)
 {
@@ -50,10 +58,12 @@ void VerifyExportedLines(const FString& ExportReportPath, const FString& CmdLogP
 IMPLEMENT_SIMPLE_AUTOMATION_TEST(FСommandsExportWindowsTest, "System.Insights.Trace.Analysis.ExecCmd.CommandsExport(Windows)", EAutomationTestFlags::ProgramContext | EAutomationTestFlags::EngineFilter)
 bool FСommandsExportWindowsTest::RunTest(const FString& Parameters)
 {
-	TSharedPtr<FInsightsManager> InsightsManager = FInsightsManager::Get();
 	IPlatformFile& PlatformFile = FPlatformFileManager::Get().GetPlatformFile();
+
 	FInsightsTestUtils Utils(this);
 
+	TSharedPtr<UE::Insights::FInsightsManager> InsightsManager = UE::Insights::FInsightsManager::Get();
+
 	const FString StoreDir = InsightsManager->GetStoreDir();
 	const FString SourceTracePath = FPaths::RootDir() / TEXT("EngineTest/SourceAssets/Utrace/CommandsExportTest_5.4.utrace");
 	const FString SourceExportPath = FPaths::RootDir() / TEXT("EngineTest/SourceAssets/Rsp/export.rsp");
@@ -106,13 +116,13 @@ bool FСommandsExportWindowsTest::RunTest(const FString& Parameters)
 
 	// ExportThreads
 	FString InsightsParameters = FString::Printf(TEXT("-OpenTraceFile=\"%s\" -ABSLOG=\"%s\" -AutoQuit -NoUI  -ExecOnAnalysisCompleteCmd=\"%s\" -log"), *StoreTracePath, *CmdThreadLogPath, *ExportThreadsTask);
-	InsightsManager->OpenUnrealInsights(*InsightsParameters);
+	UE::Insights::FMiscUtils::OpenUnrealInsights(*InsightsParameters);
 
 	VerifyExportedLines(ExportThreadsReportPath, CmdThreadLogPath, TEXT("threads"), Utils, this, Timeout);
 
 	// ExportTimers
 	InsightsParameters = FString::Printf(TEXT("-OpenTraceFile=\"%s\" -ABSLOG=\"%s\" -AutoQuit -NoUI  -ExecOnAnalysisCompleteCmd=\"%s\" -log"), *StoreTracePath, *CmdTimersLogPath, *ExportTimersTask);
-	InsightsManager->OpenUnrealInsights(*InsightsParameters);
+	UE::Insights::FMiscUtils::OpenUnrealInsights(*InsightsParameters);
 
 	VerifyExportedLines(ExportTimersReportPath, CmdTimersLogPath, TEXT("timers"), Utils, this, Timeout);
 
@@ -120,13 +130,13 @@ bool FСommandsExportWindowsTest::RunTest(const FString& Parameters)
 
 	// ExportTimingEvents
 	InsightsParameters = FString::Printf(TEXT("-OpenTraceFile=\"%s\" -ABSLOG=\"%s\" -AutoQuit -NoUI  -ExecOnAnalysisCompleteCmd=\"%s\" -log"), *StoreTracePath, *CmdTimingEventsLogPath, *ExportTimingEventsTask);
-	InsightsManager->OpenUnrealInsights(*InsightsParameters);
+	UE::Insights::FMiscUtils::OpenUnrealInsights(*InsightsParameters);
 
 	VerifyExportedLines(ExportTimingEventsReportPath, CmdTimingEventsLogPath, TEXT("timing events"), Utils, this, Timeout);
 
 	// ExportTimingEventsNonDefault
 	InsightsParameters = FString::Printf(TEXT("-OpenTraceFile=\"%s\" -ABSLOG=\"%s\" -AutoQuit -NoUI  -ExecOnAnalysisCompleteCmd=\"%s\" -log"), *StoreTracePath, *CmdTimingEventsNonDefaultLogPath, *ExportTimingEventsBorderedTask);
-	InsightsManager->OpenUnrealInsights(*InsightsParameters);
+	UE::Insights::FMiscUtils::OpenUnrealInsights(*InsightsParameters);
 
 	const TArray<FString> ExpectedTimingEventsBorderedElements =
 	{
@@ -154,7 +164,7 @@ bool FСommandsExportWindowsTest::RunTest(const FString& Parameters)
 	TestTrue("Rsp in log directory should exists after copy", PlatformFile.FileExists(*LogResultExportPath));
 
 	InsightsParameters = FString::Printf(TEXT("-OpenTraceFile=\"%s\" -ABSLOG=\"%s\" -AutoQuit -NoUI  -ExecOnAnalysisCompleteCmd=\"@=/TestResults/export.rsp\" -log"), *StoreTracePath, *CmdExportLogPath);
-	InsightsManager->OpenUnrealInsights(*InsightsParameters);
+	UE::Insights::FMiscUtils::OpenUnrealInsights(*InsightsParameters);
 
 	const TArray<FString> ExpectedThreadsElementsRsp =
 	{
diff --git a/Engine/Source/Developer/TraceInsights/Private/Insights/Tests/FunctionalTests/MemoryInsightsTests.cpp b/Engine/Source/Developer/TraceInsights/Private/Insights/Tests/FunctionalTests/MemoryInsightsTests.cpp
--- a/Engine/Source/Developer/TraceInsights/Private/Insights/Tests/FunctionalTests/MemoryInsightsTests.cpp
+++ b/Engine/Source/Developer/TraceInsights/Private/Insights/Tests/FunctionalTests/MemoryInsightsTests.cpp
@@ -27,127 +27,12 @@
 #include "Insights/MemoryProfiler/Widgets/SMemAllocTableTreeView.h"
 #include "Insights/MemoryProfiler/Widgets/SMemoryProfilerWindow.h"
 #include "Insights/Tests/InsightsTestUtils.h"
-#include "Insights/Widgets/SStartPageWindow.h"
 #include "Insights/Widgets/STimingView.h"
 
-
 DECLARE_LOG_CATEGORY_EXTERN(MemoryInsightsTests, Log, All);
 
 #if WITH_AUTOMATION_TESTS
 
-BEGIN_DEFINE_SPEC(FAutomationDriverUnrealInsightsHubMemoryInsightsTest, "System.Insights.Hub.MemoryInsights", EAutomationTestFlags::ProgramContext | EAutomationTestFlags::EngineFilter)
-FAutomationDriverPtr Driver;
-TSharedPtr<SWindow> AutomationWindow;
-END_DEFINE_SPEC(FAutomationDriverUnrealInsightsHubMemoryInsightsTest)
-void FAutomationDriverUnrealInsightsHubMemoryInsightsTest::Define()
-{
-	BeforeEach([this]() {
-		AutomationWindow = FSlateApplication::Get().GetActiveTopLevelWindow();
-		const FString AutomationWindowName = TEXT("Automation");
-		if (AutomationWindow->GetTitle().ToString().Contains(AutomationWindowName))
-		{
-			AutomationWindow->Minimize();
-		}
-
-		IUnrealInsightsModule& UnrealInsightsModule = FModuleManager::LoadModuleChecked<IUnrealInsightsModule>("TraceInsights");
-		if (IAutomationDriverModule::Get().IsEnabled())
-		{
-			IAutomationDriverModule::Get().Disable();
-		}
-		IAutomationDriverModule::Get().Enable();
-
-		Driver = IAutomationDriverModule::Get().CreateDriver();
-		});
-
-	Describe("XMLReportsUpload", [this]()
-		{
-			It("should verify that user can upload xml reports in Memory Insights tab", EAsyncExecution::ThreadPool, FTimespan::FromSeconds(120), [this]()
-				{
-					FInsightsTestUtils Utils(this);
-					IPlatformFile& PlatformFile = FPlatformFileManager::Get().GetPlatformFile();
-					TSharedPtr<FInsightsManager> InsightsManager = FInsightsManager::Get();
-					if (!InsightsManager.IsValid())
-					{
-						AddError("Insights manager should not be null");
-						return;
-					}
-
-					// Start tracing editor instance, not Lyra. There is no difference between them in this test.
-					FString UEPath = FPlatformProcess::GenerateApplicationPath("UnrealEditor", EBuildConfiguration::Development);
-					FString Parameters = TEXT("-trace=Bookmark,Memory -tracehost=127.0.0.1");
-					constexpr bool bLaunchDetached = true;
-					constexpr bool bLaunchHidden = false;
-					constexpr bool bLaunchReallyHidden = false;
-					uint32 ProcessID = 0;
-					const int32 PriorityModifier = 0;
-					const TCHAR* OptionalWorkingDirectory = nullptr;
-					void* PipeWriteChild = nullptr;
-					void* PipeReadChild = nullptr;
-					FProcHandle EditorHandle = FPlatformProcess::CreateProc(*UEPath, *Parameters, bLaunchDetached, bLaunchHidden, bLaunchReallyHidden, &ProcessID, PriorityModifier, OptionalWorkingDirectory, PipeWriteChild, PipeReadChild);
-					if (!EditorHandle.IsValid())
-					{
-						AddError("Editor should be started");
-						return;
-					}
-
-					// Verify that LIVE trace appeared
-					int Index = 0;
-					auto TraceWaiter = [Driver = Driver, &Index](void) -> bool
-					{
-						auto Elements = Driver->FindElements(By::Id("TraceStatusColumnList"))->GetElements();
-						for (int i = 0; i < Elements.Num(); ++i) {
-							if (Elements[i]->GetText().ToString() == TEXT("LIVE")) {
-								Index = i;
-								return true;
-							}
-						}
-						return false;
-					};
-
-					if (!Driver->Wait(Until::Condition(TraceWaiter, FWaitTimeout::InSeconds(10))))
-					{
-						AddError("Live trace should appear");
-						FPlatformProcess::TerminateProc(EditorHandle);
-						return;
-					}
-
-					FDriverElementRef TraceElement = Driver->FindElements(By::Id("TraceList"))->GetElements()[Index];
-					const FString TraceName = TraceElement->GetText().ToString();
-
-					const FString StoreDir = InsightsManager->GetStoreDir();
-					const FString ProjectDir = FPaths::ProjectDir();
-					const FString StoreTracePath = StoreDir / FString::Printf(TEXT("%s.utrace"), *TraceName);
-					const FString StoreCachePath = StoreDir / FString::Printf(TEXT("%s.ucache"), *TraceName);
-					const FString LogDirPath = ProjectDir / TEXT("TestResults");
-					const FString TestLogPath = ProjectDir / TEXT("TestResults/Log.txt");
-					const FString SuccessTestResult = TEXT("Test Completed. Result={Success}");
-
-					// Test live trace
-					FString TraceParameters = FString::Printf(TEXT("-InsightsTest -ABSLOG=\"%s\" -AutoQuit -ExecOnAnalysisCompleteCmd=\"Automation RunTests System.Insights.Trace.Analysis.MemoryInsights.UploadMemoryInsightsLLMXMLReportsTrace\" -OpenTraceFile=\"%s\""), *TestLogPath, *StoreTracePath);
-					InsightsManager->OpenUnrealInsights(*TraceParameters);
-					bool bLineFound = Utils.FileContainsString(TestLogPath, SuccessTestResult, 60.0f);
-					TestTrue("Test for live trace should pass", bLineFound);
-
-					IFileManager::Get().DeleteDirectory(*LogDirPath, false, true);
-					FPlatformProcess::TerminateProc(EditorHandle);
-
-					// Test stopped trace 
-					InsightsManager->OpenUnrealInsights(*TraceParameters);
-					bLineFound = Utils.FileContainsString(TestLogPath, SuccessTestResult, 60.0f);
-					TestTrue("Test for stopped trace should pass", bLineFound);
-
-					IFileManager::Get().DeleteDirectory(*LogDirPath, false, true);
-					IFileManager::Get().Delete(*StoreTracePath);
-					IFileManager::Get().Delete(*StoreCachePath);
-				});
-		});
-	AfterEach([this]() {
-		Driver.Reset();
-		IAutomationDriverModule::Get().Disable();
-		AutomationWindow->Restore();
-		});
-}
-
 IMPLEMENT_SIMPLE_AUTOMATION_TEST(FMemoryInsightsUploadLLMXMLReportsTraceTest, "System.Insights.Trace.Analysis.MemoryInsights.UploadMemoryInsightsLLMXMLReportsTrace", EAutomationTestFlags::ProgramContext | EAutomationTestFlags::EngineFilter)
 bool FMemoryInsightsUploadLLMXMLReportsTraceTest::RunTest(const FString& Parameters)
 {
@@ -291,7 +176,7 @@ bool MemoryInsightsAllocationsQueryTableTest(const FString& Parameters, const TM
 	double Timeout = 30.0;
 	FInsightsTestUtils InsightsTestUtils(Test);
 	TSharedPtr<SMemoryProfilerWindow> ProfilerWindow = FMemoryProfilerManager::Get()->GetProfilerWindow();
-	TSharedPtr<FInsightsManager> InsightsManager = FInsightsManager::Get();
+	TSharedPtr<UE::Insights::FInsightsManager> InsightsManager = UE::Insights::FInsightsManager::Get();
 	FMemorySharedState& SharedState = ProfilerWindow->GetSharedState();
 
 	TSharedPtr<FMemoryRuleSpec> MemoryRule = *Algo::FindByPredicate(SharedState.GetMemoryRules(),
@@ -389,4 +274,4 @@ void FMemoryInsightsAllocationsQueryTableStandaloneTest::GetTests(TArray<FString
 	}
 }
 
-#endif //WITH_AUTOMATION_TESTS
\ No newline at end of file
+#endif // WITH_AUTOMATION_TESTS
diff --git a/Engine/Source/Developer/TraceInsights/Private/Insights/Tests/FunctionalTests/TimingProfilerTests.cpp b/Engine/Source/Developer/TraceInsights/Private/Insights/Tests/FunctionalTests/TimingProfilerTests.cpp
--- a/Engine/Source/Developer/TraceInsights/Private/Insights/Tests/FunctionalTests/TimingProfilerTests.cpp
+++ b/Engine/Source/Developer/TraceInsights/Private/Insights/Tests/FunctionalTests/TimingProfilerTests.cpp
@@ -129,7 +129,7 @@ bool EnumerateEventsToFile::RunTest(const FString& Parameters)
 	double SessionTime = 0.0;
 	uint32 TimelineIndex = (uint32)-1;
 
-	TSharedPtr<const TraceServices::IAnalysisSession> Session = FInsightsManager::Get()->GetSession();
+	TSharedPtr<const TraceServices::IAnalysisSession> Session = UE::Insights::FInsightsManager::Get()->GetSession();
 	if (Session.IsValid() && TraceServices::ReadTimingProfilerProvider(*Session.Get()))
 	{
 		TraceServices::FAnalysisSessionReadScope SessionReadScope(*Session.Get());
@@ -170,7 +170,7 @@ bool EnumerateScopesToFile::RunTest(const FString& Parameters)
 	double SessionTime = 0.0;
 	uint32 TimelineIndex = (uint32)-1;
 
-	TSharedPtr<const TraceServices::IAnalysisSession> Session = FInsightsManager::Get()->GetSession();
+	TSharedPtr<const TraceServices::IAnalysisSession> Session = UE::Insights::FInsightsManager::Get()->GetSession();
 	if (Session.IsValid() && TraceServices::ReadTimingProfilerProvider(*Session.Get()))
 	{
 		TraceServices::FAnalysisSessionReadScope SessionReadScope(*Session.Get());
@@ -350,7 +350,7 @@ void FTimingProfilerTests::RunEnumerateBenchmark(const FEnumerateTestParams& InP
 	double SessionTime = 0.0;
 	uint32 TimelineIndex = (uint32)-1;
 
-	TSharedPtr<const TraceServices::IAnalysisSession> Session = FInsightsManager::Get()->GetSession();
+	TSharedPtr<const TraceServices::IAnalysisSession> Session = UE::Insights::FInsightsManager::Get()->GetSession();
 	if (Session.IsValid() && TraceServices::ReadTimingProfilerProvider(*Session.Get()))
 	{
 		TraceServices::FAnalysisSessionReadScope SessionReadScope(*Session.Get());
@@ -404,7 +404,7 @@ void FTimingProfilerTests::RunEnumerateAsyncBenchmark(const FEnumerateTestParams
 	double SessionTime = 0.0;
 	uint32 TimelineIndex = (uint32)-1;
 
-	TSharedPtr<const TraceServices::IAnalysisSession> Session = FInsightsManager::Get()->GetSession();
+	TSharedPtr<const TraceServices::IAnalysisSession> Session = UE::Insights::FInsightsManager::Get()->GetSession();
 	if (Session.IsValid() && TraceServices::ReadTimingProfilerProvider(*Session.Get()))
 	{
 		TraceServices::FAnalysisSessionReadScope SessionReadScope(*Session.Get());
@@ -478,7 +478,7 @@ void FTimingProfilerTests::RunEnumerateAllTracksBenchmark(const FEnumerateTestPa
 	double SessionTime = 0.0;
 	uint32 TimelineIndex = (uint32)-1;
 
-	TSharedPtr<const TraceServices::IAnalysisSession> Session = FInsightsManager::Get()->GetSession();
+	TSharedPtr<const TraceServices::IAnalysisSession> Session = UE::Insights::FInsightsManager::Get()->GetSession();
 	if (Session.IsValid() && TraceServices::ReadTimingProfilerProvider(*Session.Get()))
 	{
 		TraceServices::FAnalysisSessionReadScope SessionReadScope(*Session.Get());
@@ -537,7 +537,7 @@ void FTimingProfilerTests::RunEnumerateAsyncAllTracksBenchmark(const FEnumerateT
 	double SessionTime = 0.0;
 	uint32 TimelineIndex = (uint32)-1;
 
-	TSharedPtr<const TraceServices::IAnalysisSession> Session = FInsightsManager::Get()->GetSession();
+	TSharedPtr<const TraceServices::IAnalysisSession> Session = UE::Insights::FInsightsManager::Get()->GetSession();
 	if (Session.IsValid() && TraceServices::ReadTimingProfilerProvider(*Session.Get()))
 	{
 		TraceServices::FAnalysisSessionReadScope SessionReadScope(*Session.Get());
@@ -609,7 +609,7 @@ void FTimingProfilerTests::RunEnumerateAsyncAllTracksBenchmark(const FEnumerateT
 
 uint32 FTimingProfilerTests::GetTimelineIndex(const TCHAR* InName)
 {
-	TSharedPtr<const TraceServices::IAnalysisSession> Session = FInsightsManager::Get()->GetSession();
+	TSharedPtr<const TraceServices::IAnalysisSession> Session = UE::Insights::FInsightsManager::Get()->GetSession();
 	const TraceServices::ITimingProfilerProvider& TimingProfilerProvider = *TraceServices::ReadTimingProfilerProvider(*Session.Get());
 	const TraceServices::IThreadProvider& ThreadProvider = TraceServices::ReadThreadProvider(*Session.Get());
 	uint32 TimelineIndex = (uint32)-1;
