[Insights] Fix ExportTimingEvents not always finding the timer when providing a specific timer to -timers.

There can be multiple of the same timer with different Ids. Use a mutli-map when building them up and include them all in filter.

diff --git a/Engine/Source/Developer/TraceInsights/Private/Insights/TimingProfiler/ViewModels/TimingExporter.cpp b/Engine/Source/Developer/TraceInsights/Private/Insights/TimingProfiler/ViewModels/TimingExporter.cpp
--- a/Engine/Source/Developer/TraceInsights/Private/Insights/TimingProfiler/ViewModels/TimingExporter.cpp
+++ b/Engine/Source/Developer/TraceInsights/Private/Insights/TimingProfiler/ViewModels/TimingExporter.cpp
@@ -912,7 +912,7 @@ FTimingExporter::FTimingEventFilterFunc FTimingExporter::MakeTimingEventFilterBy
 
 	OutIncludedTimers.Reset();
 
-	TMap<FString, uint32> Timers;
+	TMultiMap<FString, uint32> Timers;
 
 	// Iterate the GPU & CPU timers.
 	if (TraceServices::ReadTimingProfilerProvider(Session))
@@ -935,12 +935,15 @@ FTimingExporter::FTimingEventFilterFunc FTimingExporter::MakeTimingEventFilterBy
 	TArray<FString> Filter;
 	InFilterString.ParseIntoArray(Filter, TEXT(","), true);
 
+	TArray<uint32> Ids;
 	for (const FString& TimerWildcard : Filter)
 	{
-		const uint32* Id = Timers.Find(TimerWildcard);
-		if (Id)
+		Ids.Reset();
+
+		Timers.MultiFind(TimerWildcard, Ids, false);
+		if (!Ids.IsEmpty())
 		{
-			OutIncludedTimers.Add(*Id);
+			OutIncludedTimers.Append(Ids);
 		}
 		else
 		{
