UnrealInsights - Add 'Copy Name To Clipboard' context menu option.

diff --git a/Engine/Source/Developer/TraceInsights/Private/Insights/Table/ViewModels/TableCommands.cpp b/Engine/Source/Developer/TraceInsights/Private/Insights/Table/ViewModels/TableCommands.cpp
new file mode 100644
--- /dev/null
+++ b/Engine/Source/Developer/TraceInsights/Private/Insights/Table/ViewModels/TableCommands.cpp
@@ -0,0 +1,59 @@
+// Copyright Epic Games, Inc. All Rights Reserved.
+
+#include "Insights/Table/ViewModels/TableCommands.h"
+
+#include "HAL/PlatformApplicationMisc.h"
+#include "InsightsCore/Table/ViewModels/Table.h"
+#include "InsightsCore/Table/ViewModels/TableCellValueSorter.h"
+
+namespace UE::Insights::Private
+{
+void CopyToClipboardImpl(TSharedRef<FTable> Table, TArray<Insights::FBaseTreeNodePtr> SelectedNodes, TSharedPtr<Insights::ITableCellValueSorter> Sorter, EColumnSortMode::Type ColumnSortMode)
+{
+	if (SelectedNodes.Num() == 0)
+	{
+		return;
+	}
+
+	FString ClipboardText;
+
+	if (Sorter.IsValid())
+	{
+		Sorter->Sort(SelectedNodes, ColumnSortMode == EColumnSortMode::Ascending ? Insights::ESortMode::Ascending : Insights::ESortMode::Descending);
+	}
+
+	Table->GetVisibleColumnsData(SelectedNodes, NAME_None, TEXT('\t'), true, ClipboardText);
+
+	if (ClipboardText.Len() > 0)
+	{
+		FPlatformApplicationMisc::ClipboardCopy(*ClipboardText);
+	}
+}
+
+void CopyNameToClipboardImpl(TSharedRef<FTable> Table, TArray<Insights::FBaseTreeNodePtr> SelectedNodes, TSharedPtr<ITableCellValueSorter> Sorter, EColumnSortMode::Type ColumnSortMode)
+{
+	if (SelectedNodes.Num() == 0)
+	{
+		return;
+	}
+
+	if (Sorter.IsValid())
+	{
+		Sorter->Sort(SelectedNodes, ColumnSortMode == EColumnSortMode::Ascending ? Insights::ESortMode::Ascending : Insights::ESortMode::Descending);
+	}
+
+	FString ClipboardText = FString::JoinBy(
+		SelectedNodes,
+		TEXT("\n"),
+		[](const Insights::FBaseTreeNodePtr& TimerNode)
+	{
+		return TimerNode->GetName().ToString();
+	});
+
+	if (ClipboardText.Len() > 0)
+	{
+		FPlatformApplicationMisc::ClipboardCopy(*ClipboardText);
+	}
+}
+
+} // namespace UE::Insights::Private
\ No newline at end of file
diff --git a/Engine/Source/Developer/TraceInsights/Private/Insights/Table/ViewModels/TableCommands.h b/Engine/Source/Developer/TraceInsights/Private/Insights/Table/ViewModels/TableCommands.h
new file mode 100644
--- /dev/null
+++ b/Engine/Source/Developer/TraceInsights/Private/Insights/Table/ViewModels/TableCommands.h
@@ -0,0 +1,44 @@
+// Copyright Epic Games, Inc. All Rights Reserved.
+#pragma once
+
+#include "Algo/Transform.h"
+#include "InsightsCore/Table/ViewModels/BaseTreeNode.h"
+#include "Templates/SharedPointer.h"
+
+namespace UE::Insights
+{
+
+class FTable;
+class ITableCellValueSorter;
+
+namespace Private
+{
+	template<class TreeNodePtrType>
+	TArray<FBaseTreeNodePtr> ConvertNodes(const TArray<TreeNodePtrType>& InNodes)
+	{
+		TArray<FBaseTreeNodePtr> BaseNodes;
+		Algo::Transform(InNodes, BaseNodes, [](const TreeNodePtrType& Ptr)
+		{
+			return StaticCastSharedPtr<FBaseTreeNodePtr::ElementType>(Ptr);
+		});
+		return BaseNodes;
+	}
+
+	void CopyToClipboardImpl(TSharedRef<FTable> Table, TArray<FBaseTreeNodePtr> SelectedNodes, TSharedPtr<ITableCellValueSorter> Sorter, EColumnSortMode::Type ColumnSortMode);
+	void CopyNameToClipboardImpl(TSharedRef<FTable> Table, TArray<FBaseTreeNodePtr> SelectedNodes, TSharedPtr<ITableCellValueSorter> Sorter, EColumnSortMode::Type ColumnSortMode);
+
+} // Private
+
+template<class TreeNodePtrType>
+void CopyToClipboard(TSharedRef<FTable> Table, const TArray<TreeNodePtrType>& SelectedNodes, TSharedPtr<ITableCellValueSorter> Sorter, EColumnSortMode::Type ColumnSortMode)
+{
+	Private::CopyToClipboardImpl(Table, Private::ConvertNodes(SelectedNodes), MoveTemp(Sorter), ColumnSortMode);
+}
+
+template<class TreeNodePtrType>
+void CopyNameToClipboard(TSharedRef<FTable> Table, const TArray<TreeNodePtrType>& SelectedNodes, TSharedPtr<ITableCellValueSorter> Sorter, EColumnSortMode::Type ColumnSortMode)
+{
+	Private::CopyNameToClipboardImpl(Table, Private::ConvertNodes(SelectedNodes), MoveTemp(Sorter), ColumnSortMode);
+}
+
+} // namespace UE::Insights
\ No newline at end of file
diff --git a/Engine/Source/Developer/TraceInsights/Private/Insights/TimingProfiler/Widgets/SStatsView.cpp b/Engine/Source/Developer/TraceInsights/Private/Insights/TimingProfiler/Widgets/SStatsView.cpp
--- a/Engine/Source/Developer/TraceInsights/Private/Insights/TimingProfiler/Widgets/SStatsView.cpp
+++ b/Engine/Source/Developer/TraceInsights/Private/Insights/TimingProfiler/Widgets/SStatsView.cpp
@@ -44,6 +44,7 @@
 #include "Insights/TimingProfilerCommon.h"
 #include "Insights/ViewModels/TimingGraphTrack.h"
 #include "Insights/Widgets/STimingView.h"
+#include "Insights/Table/ViewModels/TableCommands.h"
 
 #include <limits>
 
@@ -82,6 +83,12 @@ public:
 			EUserInterfaceActionType::Button,
 			FInputChord(EModifierKey::Control, EKeys::C));
 
+		UI_COMMAND(Command_CopyNameToClipboard,
+			"Copy Name To Clipboard",
+			"Copies the name of the selected counter to the clipboard.",
+			EUserInterfaceActionType::Button,
+			FInputChord(EModifierKey::Control | EModifierKey::Shift, EKeys::C));
+
 		UI_COMMAND(Command_Export,
 			"Export...",
 			"Exports the selection (counters and their aggregated statistics) to a text file (tab-separated values or comma-separated values).",
@@ -109,6 +116,7 @@ public:
 	UE_ENABLE_OPTIMIZATION_SHIP
 
 	TSharedPtr<FUICommandInfo> Command_CopyToClipboard;
+	TSharedPtr<FUICommandInfo> Command_CopyNameToClipboard;
 	TSharedPtr<FUICommandInfo> Command_Export;
 	TSharedPtr<FUICommandInfo> Command_ExportValues;
 	TSharedPtr<FUICommandInfo> Command_ExportOps;
@@ -157,6 +165,7 @@ void SStatsView::InitCommandList()
 	FStatsViewCommands::Register();
 	CommandList = MakeShared<FUICommandList>();
 	CommandList->MapAction(FStatsViewCommands::Get().Command_CopyToClipboard, FExecuteAction::CreateSP(this, &SStatsView::ContextMenu_CopySelectedToClipboard_Execute), FCanExecuteAction::CreateSP(this, &SStatsView::ContextMenu_CopySelectedToClipboard_CanExecute));
+	CommandList->MapAction(FStatsViewCommands::Get().Command_CopyNameToClipboard, FExecuteAction::CreateSP(this, &SStatsView::ContextMenu_CopySelectedNameToClipboard_Execute), FCanExecuteAction::CreateSP(this, &SStatsView::ContextMenu_CopySelectedNameToClipboard_CanExecute));
 	CommandList->MapAction(FStatsViewCommands::Get().Command_Export, FExecuteAction::CreateSP(this, &SStatsView::ContextMenu_Export_Execute), FCanExecuteAction::CreateSP(this, &SStatsView::ContextMenu_Export_CanExecute));
 	CommandList->MapAction(FStatsViewCommands::Get().Command_ExportValues, FExecuteAction::CreateSP(this, &SStatsView::ContextMenu_ExportValues_Execute), FCanExecuteAction::CreateSP(this, &SStatsView::ContextMenu_ExportValues_CanExecute));
 	CommandList->MapAction(FStatsViewCommands::Get().Command_ExportOps, FExecuteAction::CreateSP(this, &SStatsView::ContextMenu_ExportOps_Execute), FCanExecuteAction::CreateSP(this, &SStatsView::ContextMenu_ExportOps_CanExecute));
@@ -486,6 +495,15 @@ TSharedPtr<SWidget> SStatsView::TreeView_GetMenuContent()
 			FSlateIcon(FAppStyle::Get().GetStyleSetName(), "GenericCommands.Copy")
 		);
 
+		MenuBuilder.AddMenuEntry
+		(
+			FStatsViewCommands::Get().Command_CopyNameToClipboard,
+			NAME_None,
+			TAttribute<FText>(),
+			TAttribute<FText>(),
+			FSlateIcon(FAppStyle::Get().GetStyleSetName(), "GenericCommands.Copy")
+		);
+
 		MenuBuilder.AddMenuEntry
 		(
 			FStatsViewCommands::Get().Command_Export,
@@ -2198,34 +2216,26 @@ bool SStatsView::ContextMenu_CopySelectedToClipboard_CanExecute() const
 
 void SStatsView::ContextMenu_CopySelectedToClipboard_Execute()
 {
-	if (!Table->IsValid())
-	{
-		return;
-	}
-
-	TArray<FBaseTreeNodePtr> SelectedNodes;
-	for (FStatsNodePtr CounterPtr : TreeView->GetSelectedItems())
+	if (Table->IsValid())
 	{
-		SelectedNodes.Add(CounterPtr);
-	}
-
-	if (SelectedNodes.Num() == 0)
-	{
-		return;
+		UE::Insights::CopyToClipboard(Table.ToSharedRef(), TreeView->GetSelectedItems(), CurrentSorter, ColumnSortMode);
 	}
+}
 
-	FString ClipboardText;
+////////////////////////////////////////////////////////////////////////////////////////////////////
 
-	if (CurrentSorter.IsValid())
-	{
-		CurrentSorter->Sort(SelectedNodes, ColumnSortMode == EColumnSortMode::Ascending ? ESortMode::Ascending : ESortMode::Descending);
-	}
+bool SStatsView::ContextMenu_CopySelectedNameToClipboard_CanExecute() const
+{
+	return TreeView->GetSelectedItems().Num() > 0;
+}
 
-	Table->GetVisibleColumnsData(SelectedNodes, FTimingProfilerManager::Get()->GetLogListingName(), TEXT('\t'), true, ClipboardText);
+////////////////////////////////////////////////////////////////////////////////////////////////////
 
-	if (ClipboardText.Len() > 0)
+void SStatsView::ContextMenu_CopySelectedNameToClipboard_Execute()
+{
+	if (Table->IsValid())
 	{
-		FPlatformApplicationMisc::ClipboardCopy(*ClipboardText);
+		UE::Insights::CopyNameToClipboard(Table.ToSharedRef(), TreeView->GetSelectedItems(), CurrentSorter, ColumnSortMode);
 	}
 }
 
diff --git a/Engine/Source/Developer/TraceInsights/Private/Insights/TimingProfiler/Widgets/SStatsView.h b/Engine/Source/Developer/TraceInsights/Private/Insights/TimingProfiler/Widgets/SStatsView.h
--- a/Engine/Source/Developer/TraceInsights/Private/Insights/TimingProfiler/Widgets/SStatsView.h
+++ b/Engine/Source/Developer/TraceInsights/Private/Insights/TimingProfiler/Widgets/SStatsView.h
@@ -127,6 +127,9 @@ private:
 	bool ContextMenu_CopySelectedToClipboard_CanExecute() const;
 	void ContextMenu_CopySelectedToClipboard_Execute();
 
+	bool ContextMenu_CopySelectedNameToClipboard_CanExecute() const;
+	void ContextMenu_CopySelectedNameToClipboard_Execute();
+
 	bool ContextMenu_Export_CanExecute() const;
 	void ContextMenu_Export_Execute();
 
diff --git a/Engine/Source/Developer/TraceInsights/Private/Insights/TimingProfiler/Widgets/STimerTreeView.cpp b/Engine/Source/Developer/TraceInsights/Private/Insights/TimingProfiler/Widgets/STimerTreeView.cpp
--- a/Engine/Source/Developer/TraceInsights/Private/Insights/TimingProfiler/Widgets/STimerTreeView.cpp
+++ b/Engine/Source/Developer/TraceInsights/Private/Insights/TimingProfiler/Widgets/STimerTreeView.cpp
@@ -24,6 +24,7 @@
 #include "Insights/TimingProfiler/ViewModels/TimersViewColumnFactory.h"
 #include "Insights/TimingProfiler/Widgets/STimersViewTooltip.h"
 #include "Insights/TimingProfiler/Widgets/STimerTableRow.h"
+#include "Insights/Table/ViewModels/TableCommands.h"
 
 #define LOCTEXT_NAMESPACE "UE::Insights::TimingProfiler::STimerTreeView"
 
@@ -59,10 +60,17 @@ public:
 			"Copies the selection to clipboard.",
 			EUserInterfaceActionType::Button,
 			FInputChord(EModifierKey::Control, EKeys::C));
+
+		UI_COMMAND(Command_CopyNameToClipboard,
+			"Copy Name To Clipboard",
+			"Copies the name of the selected timer to the clipboard.",
+			EUserInterfaceActionType::Button,
+			FInputChord(EModifierKey::Control | EModifierKey::Shift, EKeys::C));
 	}
 	UE_ENABLE_OPTIMIZATION_SHIP
 
 	TSharedPtr<FUICommandInfo> Command_CopyToClipboard;
+	TSharedPtr<FUICommandInfo> Command_CopyNameToClipboard;
 };
 
 ////////////////////////////////////////////////////////////////////////////////////////////////////
@@ -99,6 +107,7 @@ void STimerTreeView::InitCommandList()
 	FTimerTreeViewCommands::Register();
 	CommandList = MakeShared<FUICommandList>();
 	CommandList->MapAction(FTimerTreeViewCommands::Get().Command_CopyToClipboard, FExecuteAction::CreateSP(this, &STimerTreeView::ContextMenu_CopySelectedToClipboard_Execute), FCanExecuteAction::CreateSP(this, &STimerTreeView::ContextMenu_CopySelectedToClipboard_CanExecute));
+	CommandList->MapAction(FTimerTreeViewCommands::Get().Command_CopyNameToClipboard, FExecuteAction::CreateSP(this, &STimerTreeView::ContextMenu_CopySelectedTimerNameToClipboard_Execute), FCanExecuteAction::CreateSP(this, &STimerTreeView::ContextMenu_CopySelectedTimerNameToClipboard_CanExecute));
 }
 
 ////////////////////////////////////////////////////////////////////////////////////////////////////
@@ -238,6 +247,15 @@ TSharedPtr<SWidget> STimerTreeView::TreeView_GetMenuContent()
 			TAttribute<FText>(),
 			FSlateIcon(FAppStyle::Get().GetStyleSetName(), "GenericCommands.Copy")
 		);
+
+		MenuBuilder.AddMenuEntry
+		(
+			FTimerTreeViewCommands::Get().Command_CopyNameToClipboard,
+			NAME_None,
+			TAttribute<FText>(),
+			TAttribute<FText>(),
+			FSlateIcon(FAppStyle::Get().GetStyleSetName(), "GenericCommands.Copy")
+		);
 	}
 	MenuBuilder.EndSection();
 
@@ -1102,34 +1120,26 @@ bool STimerTreeView::ContextMenu_CopySelectedToClipboard_CanExecute() const
 
 void STimerTreeView::ContextMenu_CopySelectedToClipboard_Execute()
 {
-	if (!Table->IsValid())
-	{
-		return;
-	}
-
-	TArray<FBaseTreeNodePtr> SelectedNodes;
-	for (FTimerNodePtr TimerPtr : TreeView->GetSelectedItems())
+	if (Table->IsValid())
 	{
-		SelectedNodes.Add(TimerPtr);
-	}
-
-	if (SelectedNodes.Num() == 0)
-	{
-		return;
+		UE::Insights::CopyToClipboard(Table.ToSharedRef(), TreeView->GetSelectedItems(), CurrentSorter, ColumnSortMode);
 	}
+}
 
-	FString ClipboardText;
+////////////////////////////////////////////////////////////////////////////////////////////////////
 
-	if (CurrentSorter.IsValid())
-	{
-		CurrentSorter->Sort(SelectedNodes, ColumnSortMode == EColumnSortMode::Ascending ? ESortMode::Ascending : ESortMode::Descending);
-	}
+bool STimerTreeView::ContextMenu_CopySelectedTimerNameToClipboard_CanExecute() const
+{
+	return TreeView->GetSelectedItems().Num() > 0;
+}
 
-	Table->GetVisibleColumnsData(SelectedNodes, FTimingProfilerManager::Get()->GetLogListingName(), TEXT('\t'), true, ClipboardText);
+////////////////////////////////////////////////////////////////////////////////////////////////////
 
-	if (ClipboardText.Len() > 0)
+void STimerTreeView::ContextMenu_CopySelectedTimerNameToClipboard_Execute()
+{
+	if (Table->IsValid())
 	{
-		FPlatformApplicationMisc::ClipboardCopy(*ClipboardText);
+		UE::Insights::CopyNameToClipboard(Table.ToSharedRef(), TreeView->GetSelectedItems(), CurrentSorter, ColumnSortMode);
 	}
 }
 
diff --git a/Engine/Source/Developer/TraceInsights/Private/Insights/TimingProfiler/Widgets/STimerTreeView.h b/Engine/Source/Developer/TraceInsights/Private/Insights/TimingProfiler/Widgets/STimerTreeView.h
--- a/Engine/Source/Developer/TraceInsights/Private/Insights/TimingProfiler/Widgets/STimerTreeView.h
+++ b/Engine/Source/Developer/TraceInsights/Private/Insights/TimingProfiler/Widgets/STimerTreeView.h
@@ -77,6 +77,9 @@ private:
 	bool ContextMenu_CopySelectedToClipboard_CanExecute() const;
 	void ContextMenu_CopySelectedToClipboard_Execute();
 
+	bool ContextMenu_CopySelectedTimerNameToClipboard_CanExecute() const;
+	void ContextMenu_CopySelectedTimerNameToClipboard_Execute();
+
 	////////////////////////////////////////////////////////////////////////////////////////////////////
 	// Tree View - Columns' Header
 
diff --git a/Engine/Source/Developer/TraceInsights/Private/Insights/TimingProfiler/Widgets/STimersView.cpp b/Engine/Source/Developer/TraceInsights/Private/Insights/TimingProfiler/Widgets/STimersView.cpp
--- a/Engine/Source/Developer/TraceInsights/Private/Insights/TimingProfiler/Widgets/STimersView.cpp
+++ b/Engine/Source/Developer/TraceInsights/Private/Insights/TimingProfiler/Widgets/STimersView.cpp
@@ -50,6 +50,7 @@
 #include "Insights/TimingProfilerCommon.h"
 #include "Insights/ViewModels/TimingGraphTrack.h"
 #include "Insights/Widgets/STimingView.h"
+#include "Insights/Table/ViewModels/TableCommands.h"
 
 #include <limits>
 
@@ -88,6 +89,12 @@ public:
 			EUserInterfaceActionType::Button,
 			FInputChord(EModifierKey::Control, EKeys::C));
 
+		UI_COMMAND(Command_CopyNameToClipboard,
+			"Copy Name To Clipboard",
+			"Copies the name of the selected timer to the clipboard.",
+			EUserInterfaceActionType::Button,
+			FInputChord(EModifierKey::Control | EModifierKey::Shift, EKeys::C));
+
 		UI_COMMAND(Command_Export,
 			"Export...",
 			"Exports the selection (timers and their aggregated statistics) to a text file (tab-separated values or comma-separated values).",
@@ -151,6 +158,7 @@ public:
 	UE_ENABLE_OPTIMIZATION_SHIP
 
 	TSharedPtr<FUICommandInfo> Command_CopyToClipboard;
+	TSharedPtr<FUICommandInfo> Command_CopyNameToClipboard;
 	TSharedPtr<FUICommandInfo> Command_Export;
 	TSharedPtr<FUICommandInfo> Command_ExportTimingEventsSelection;
 	TSharedPtr<FUICommandInfo> Command_ExportTimingEvents;
@@ -208,6 +216,7 @@ void STimersView::InitCommandList()
 	FTimersViewCommands::Register();
 	CommandList = MakeShared<FUICommandList>();
 	CommandList->MapAction(FTimersViewCommands::Get().Command_CopyToClipboard, FExecuteAction::CreateSP(this, &STimersView::ContextMenu_CopyToClipboard_Execute), FCanExecuteAction::CreateSP(this, &STimersView::ContextMenu_CopyToClipboard_CanExecute));
+	CommandList->MapAction(FTimersViewCommands::Get().Command_CopyNameToClipboard, FExecuteAction::CreateSP(this, &STimersView::ContextMenu_CopyTimerNameToClipboard_Execute), FCanExecuteAction::CreateSP(this, &STimersView::ContextMenu_CopyTimerNameToClipboard_CanExecute));
 	CommandList->MapAction(FTimersViewCommands::Get().Command_Export, FExecuteAction::CreateSP(this, &STimersView::ContextMenu_Export_Execute), FCanExecuteAction::CreateSP(this, &STimersView::ContextMenu_Export_CanExecute));
 	CommandList->MapAction(FTimersViewCommands::Get().Command_ExportTimingEventsSelection, FExecuteAction::CreateSP(this, &STimersView::ContextMenu_ExportTimingEventsSelection_Execute), FCanExecuteAction::CreateSP(this, &STimersView::ContextMenu_ExportTimingEventsSelection_CanExecute));
 	CommandList->MapAction(FTimersViewCommands::Get().Command_ExportTimingEvents, FExecuteAction::CreateSP(this, &STimersView::ContextMenu_ExportTimingEvents_Execute), FCanExecuteAction::CreateSP(this, &STimersView::ContextMenu_ExportTimingEvents_CanExecute));
@@ -659,6 +668,15 @@ TSharedPtr<SWidget> STimersView::TreeView_GetMenuContent()
 			FSlateIcon(FAppStyle::Get().GetStyleSetName(), "GenericCommands.Copy")
 		);
 
+		MenuBuilder.AddMenuEntry
+		(
+			FTimersViewCommands::Get().Command_CopyNameToClipboard,
+			NAME_None,
+			TAttribute<FText>(),
+			TAttribute<FText>(),
+			FSlateIcon(FAppStyle::Get().GetStyleSetName(), "GenericCommands.Copy")
+		);
+
 		MenuBuilder.AddMenuEntry
 		(
 			FTimersViewCommands::Get().Command_Export,
@@ -2830,35 +2848,26 @@ bool STimersView::ContextMenu_CopyToClipboard_CanExecute() const
 
 void STimersView::ContextMenu_CopyToClipboard_Execute()
 {
-	if (!Table->IsValid())
-	{
-		return;
-	}
-
-	const TArray<FTimerNodePtr> SelectedTimerNodes = TreeView->GetSelectedItems();
-	if (SelectedTimerNodes.Num() == 0)
+	if (Table->IsValid())
 	{
-		return;
-	}
-
-	TArray<FBaseTreeNodePtr> SelectedNodes;
-	for (FTimerNodePtr TimerPtr : SelectedTimerNodes)
-	{
-		SelectedNodes.Add(TimerPtr);
+		UE::Insights::CopyToClipboard(Table.ToSharedRef(), TreeView->GetSelectedItems(), CurrentSorter, ColumnSortMode);
 	}
+}
 
-	FString ClipboardText;
+////////////////////////////////////////////////////////////////////////////////////////////////////
 
-	if (CurrentSorter.IsValid())
-	{
-		CurrentSorter->Sort(SelectedNodes, ColumnSortMode == EColumnSortMode::Ascending ? ESortMode::Ascending : ESortMode::Descending);
-	}
+bool STimersView::ContextMenu_CopyTimerNameToClipboard_CanExecute() const
+{
+	return TreeView->GetSelectedItems().Num() > 0;
+}
 
-	Table->GetVisibleColumnsData(SelectedNodes, FTimingProfilerManager::Get()->GetLogListingName(), TEXT('\t'), true, ClipboardText);
+////////////////////////////////////////////////////////////////////////////////////////////////////
 
-	if (ClipboardText.Len() > 0)
+void STimersView::ContextMenu_CopyTimerNameToClipboard_Execute()
+{
+	if (Table->IsValid())
 	{
-		FPlatformApplicationMisc::ClipboardCopy(*ClipboardText);
+		UE::Insights::CopyNameToClipboard(Table.ToSharedRef(), TreeView->GetSelectedItems(), CurrentSorter, ColumnSortMode);
 	}
 }
 
diff --git a/Engine/Source/Developer/TraceInsights/Private/Insights/TimingProfiler/Widgets/STimersView.h b/Engine/Source/Developer/TraceInsights/Private/Insights/TimingProfiler/Widgets/STimersView.h
--- a/Engine/Source/Developer/TraceInsights/Private/Insights/TimingProfiler/Widgets/STimersView.h
+++ b/Engine/Source/Developer/TraceInsights/Private/Insights/TimingProfiler/Widgets/STimersView.h
@@ -130,6 +130,9 @@ private:
 	bool ContextMenu_CopyToClipboard_CanExecute() const;
 	void ContextMenu_CopyToClipboard_Execute();
 
+	bool ContextMenu_CopyTimerNameToClipboard_CanExecute() const;
+	void ContextMenu_CopyTimerNameToClipboard_Execute();
+
 	bool ContextMenu_Export_CanExecute() const;
 	void ContextMenu_Export_Execute();
 
