[Insights] Fix non-unity compile error due to EColumnSortMode not being defined.

Changed the TableCommands to take the internal sort mode instead so we didn't need to add a slate dependency.

#rnx

diff --git a/Engine/Source/Developer/TraceInsights/Private/Insights/Table/ViewModels/TableCommands.cpp b/Engine/Source/Developer/TraceInsights/Private/Insights/Table/ViewModels/TableCommands.cpp
--- a/Engine/Source/Developer/TraceInsights/Private/Insights/Table/ViewModels/TableCommands.cpp
+++ b/Engine/Source/Developer/TraceInsights/Private/Insights/Table/ViewModels/TableCommands.cpp
@@ -8,7 +8,7 @@
 
 namespace UE::Insights::Private
 {
-void CopyToClipboardImpl(TSharedRef<FTable> Table, TArray<Insights::FBaseTreeNodePtr> SelectedNodes, TSharedPtr<Insights::ITableCellValueSorter> Sorter, EColumnSortMode::Type ColumnSortMode)
+void CopyToClipboardImpl(TSharedRef<FTable> Table, TArray<Insights::FBaseTreeNodePtr> SelectedNodes, TSharedPtr<Insights::ITableCellValueSorter> Sorter, ESortMode ColumnSortMode)
 {
 	if (SelectedNodes.Num() == 0)
 	{
@@ -19,7 +19,7 @@ void CopyToClipboardImpl(TSharedRef<FTable> Table, TArray<Insights::FBaseTreeNod
 
 	if (Sorter.IsValid())
 	{
-		Sorter->Sort(SelectedNodes, ColumnSortMode == EColumnSortMode::Ascending ? Insights::ESortMode::Ascending : Insights::ESortMode::Descending);
+		Sorter->Sort(SelectedNodes, ColumnSortMode);
 	}
 
 	Table->GetVisibleColumnsData(SelectedNodes, NAME_None, TEXT('\t'), true, ClipboardText);
@@ -30,7 +30,7 @@ void CopyToClipboardImpl(TSharedRef<FTable> Table, TArray<Insights::FBaseTreeNod
 	}
 }
 
-void CopyNameToClipboardImpl(TSharedRef<FTable> Table, TArray<Insights::FBaseTreeNodePtr> SelectedNodes, TSharedPtr<ITableCellValueSorter> Sorter, EColumnSortMode::Type ColumnSortMode)
+void CopyNameToClipboardImpl(TSharedRef<FTable> Table, TArray<Insights::FBaseTreeNodePtr> SelectedNodes, TSharedPtr<ITableCellValueSorter> Sorter, ESortMode ColumnSortMode)
 {
 	if (SelectedNodes.Num() == 0)
 	{
@@ -39,7 +39,7 @@ void CopyNameToClipboardImpl(TSharedRef<FTable> Table, TArray<Insights::FBaseTre
 
 	if (Sorter.IsValid())
 	{
-		Sorter->Sort(SelectedNodes, ColumnSortMode == EColumnSortMode::Ascending ? Insights::ESortMode::Ascending : Insights::ESortMode::Descending);
+		Sorter->Sort(SelectedNodes, ColumnSortMode);
 	}
 
 	FString ClipboardText = FString::JoinBy(
diff --git a/Engine/Source/Developer/TraceInsights/Private/Insights/Table/ViewModels/TableCommands.h b/Engine/Source/Developer/TraceInsights/Private/Insights/Table/ViewModels/TableCommands.h
--- a/Engine/Source/Developer/TraceInsights/Private/Insights/Table/ViewModels/TableCommands.h
+++ b/Engine/Source/Developer/TraceInsights/Private/Insights/Table/ViewModels/TableCommands.h
@@ -3,6 +3,7 @@
 
 #include "Algo/Transform.h"
 #include "InsightsCore/Table/ViewModels/BaseTreeNode.h"
+#include "InsightsCore/Table/ViewModels/TableCellValueSorter.h"
 #include "Templates/SharedPointer.h"
 
 namespace UE::Insights
@@ -24,19 +25,19 @@ namespace Private
 		return BaseNodes;
 	}
 
-	void CopyToClipboardImpl(TSharedRef<FTable> Table, TArray<FBaseTreeNodePtr> SelectedNodes, TSharedPtr<ITableCellValueSorter> Sorter, EColumnSortMode::Type ColumnSortMode);
-	void CopyNameToClipboardImpl(TSharedRef<FTable> Table, TArray<FBaseTreeNodePtr> SelectedNodes, TSharedPtr<ITableCellValueSorter> Sorter, EColumnSortMode::Type ColumnSortMode);
+	void CopyToClipboardImpl(TSharedRef<FTable> Table, TArray<FBaseTreeNodePtr> SelectedNodes, TSharedPtr<ITableCellValueSorter> Sorter, ESortMode ColumnSortMode);
+	void CopyNameToClipboardImpl(TSharedRef<FTable> Table, TArray<FBaseTreeNodePtr> SelectedNodes, TSharedPtr<ITableCellValueSorter> Sorter, ESortMode ColumnSortMode);
 
 } // Private
 
 template<class TreeNodePtrType>
-void CopyToClipboard(TSharedRef<FTable> Table, const TArray<TreeNodePtrType>& SelectedNodes, TSharedPtr<ITableCellValueSorter> Sorter, EColumnSortMode::Type ColumnSortMode)
+void CopyToClipboard(TSharedRef<FTable> Table, const TArray<TreeNodePtrType>& SelectedNodes, TSharedPtr<ITableCellValueSorter> Sorter, ESortMode ColumnSortMode)
 {
 	Private::CopyToClipboardImpl(Table, Private::ConvertNodes(SelectedNodes), MoveTemp(Sorter), ColumnSortMode);
 }
 
 template<class TreeNodePtrType>
-void CopyNameToClipboard(TSharedRef<FTable> Table, const TArray<TreeNodePtrType>& SelectedNodes, TSharedPtr<ITableCellValueSorter> Sorter, EColumnSortMode::Type ColumnSortMode)
+void CopyNameToClipboard(TSharedRef<FTable> Table, const TArray<TreeNodePtrType>& SelectedNodes, TSharedPtr<ITableCellValueSorter> Sorter, ESortMode ColumnSortMode)
 {
 	Private::CopyNameToClipboardImpl(Table, Private::ConvertNodes(SelectedNodes), MoveTemp(Sorter), ColumnSortMode);
 }
diff --git a/Engine/Source/Developer/TraceInsights/Private/Insights/TimingProfiler/Widgets/SStatsView.cpp b/Engine/Source/Developer/TraceInsights/Private/Insights/TimingProfiler/Widgets/SStatsView.cpp
--- a/Engine/Source/Developer/TraceInsights/Private/Insights/TimingProfiler/Widgets/SStatsView.cpp
+++ b/Engine/Source/Developer/TraceInsights/Private/Insights/TimingProfiler/Widgets/SStatsView.cpp
@@ -27,6 +27,7 @@
 #include "InsightsCore/Common/Stopwatch.h"
 #include "InsightsCore/Common/TimeUtils.h"
 #include "InsightsCore/Table/ViewModels/Table.h"
+#include "InsightsCore/Table/ViewModels/TableCellValueSorter.h"
 #include "InsightsCore/Table/ViewModels/TableColumn.h"
 #include "InsightsCore/Table/Widgets/SAsyncOperationStatus.h"
 
@@ -2218,7 +2219,8 @@ void SStatsView::ContextMenu_CopySelectedToClipboard_Execute()
 {
 	if (Table->IsValid())
 	{
-		UE::Insights::CopyToClipboard(Table.ToSharedRef(), TreeView->GetSelectedItems(), CurrentSorter, ColumnSortMode);
+		const ESortMode SortMode = ColumnSortMode == EColumnSortMode::Ascending ? Insights::ESortMode::Ascending : Insights::ESortMode::Descending;
+		UE::Insights::CopyToClipboard(Table.ToSharedRef(), TreeView->GetSelectedItems(), CurrentSorter, SortMode);
 	}
 }
 
@@ -2235,7 +2237,8 @@ void SStatsView::ContextMenu_CopySelectedNameToClipboard_Execute()
 {
 	if (Table->IsValid())
 	{
-		UE::Insights::CopyNameToClipboard(Table.ToSharedRef(), TreeView->GetSelectedItems(), CurrentSorter, ColumnSortMode);
+		const ESortMode SortMode = ColumnSortMode == EColumnSortMode::Ascending ? Insights::ESortMode::Ascending : Insights::ESortMode::Descending;
+		UE::Insights::CopyNameToClipboard(Table.ToSharedRef(), TreeView->GetSelectedItems(), CurrentSorter, SortMode);
 	}
 }
 
diff --git a/Engine/Source/Developer/TraceInsights/Private/Insights/TimingProfiler/Widgets/STimerTreeView.cpp b/Engine/Source/Developer/TraceInsights/Private/Insights/TimingProfiler/Widgets/STimerTreeView.cpp
--- a/Engine/Source/Developer/TraceInsights/Private/Insights/TimingProfiler/Widgets/STimerTreeView.cpp
+++ b/Engine/Source/Developer/TraceInsights/Private/Insights/TimingProfiler/Widgets/STimerTreeView.cpp
@@ -14,6 +14,7 @@
 
 // TraceInsightsCore
 #include "InsightsCore/Table/ViewModels/Table.h"
+#include "InsightsCore/Table/ViewModels/TableCellValueSorter.h"
 #include "InsightsCore/Table/ViewModels/TableColumn.h"
 #include "InsightsCore/Table/Widgets/SAsyncOperationStatus.h"
 
@@ -1122,7 +1123,8 @@ void STimerTreeView::ContextMenu_CopySelectedToClipboard_Execute()
 {
 	if (Table->IsValid())
 	{
-		UE::Insights::CopyToClipboard(Table.ToSharedRef(), TreeView->GetSelectedItems(), CurrentSorter, ColumnSortMode);
+		const ESortMode SortMode = ColumnSortMode == EColumnSortMode::Ascending ? Insights::ESortMode::Ascending : Insights::ESortMode::Descending;
+		UE::Insights::CopyToClipboard(Table.ToSharedRef(), TreeView->GetSelectedItems(), CurrentSorter, SortMode);
 	}
 }
 
@@ -1139,7 +1141,8 @@ void STimerTreeView::ContextMenu_CopySelectedTimerNameToClipboard_Execute()
 {
 	if (Table->IsValid())
 	{
-		UE::Insights::CopyNameToClipboard(Table.ToSharedRef(), TreeView->GetSelectedItems(), CurrentSorter, ColumnSortMode);
+		const ESortMode SortMode = ColumnSortMode == EColumnSortMode::Ascending ? Insights::ESortMode::Ascending : Insights::ESortMode::Descending;
+		UE::Insights::CopyNameToClipboard(Table.ToSharedRef(), TreeView->GetSelectedItems(), CurrentSorter, SortMode);
 	}
 }
 
diff --git a/Engine/Source/Developer/TraceInsights/Private/Insights/TimingProfiler/Widgets/STimersView.cpp b/Engine/Source/Developer/TraceInsights/Private/Insights/TimingProfiler/Widgets/STimersView.cpp
--- a/Engine/Source/Developer/TraceInsights/Private/Insights/TimingProfiler/Widgets/STimersView.cpp
+++ b/Engine/Source/Developer/TraceInsights/Private/Insights/TimingProfiler/Widgets/STimersView.cpp
@@ -31,6 +31,7 @@
 #include "InsightsCore/Common/Stopwatch.h"
 #include "InsightsCore/Common/TimeUtils.h"
 #include "InsightsCore/Table/ViewModels/Table.h"
+#include "InsightsCore/Table/ViewModels/TableCellValueSorter.h"
 #include "InsightsCore/Table/ViewModels/TableColumn.h"
 #include "InsightsCore/Table/Widgets/SAsyncOperationStatus.h"
 
@@ -2850,7 +2851,8 @@ void STimersView::ContextMenu_CopyToClipboard_Execute()
 {
 	if (Table->IsValid())
 	{
-		UE::Insights::CopyToClipboard(Table.ToSharedRef(), TreeView->GetSelectedItems(), CurrentSorter, ColumnSortMode);
+		const ESortMode SortMode = ColumnSortMode == EColumnSortMode::Ascending ? Insights::ESortMode::Ascending : Insights::ESortMode::Descending;
+		UE::Insights::CopyToClipboard(Table.ToSharedRef(), TreeView->GetSelectedItems(), CurrentSorter, SortMode);
 	}
 }
 
@@ -2867,7 +2869,8 @@ void STimersView::ContextMenu_CopyTimerNameToClipboard_Execute()
 {
 	if (Table->IsValid())
 	{
-		UE::Insights::CopyNameToClipboard(Table.ToSharedRef(), TreeView->GetSelectedItems(), CurrentSorter, ColumnSortMode);
+		const ESortMode SortMode = ColumnSortMode == EColumnSortMode::Ascending ? Insights::ESortMode::Ascending : Insights::ESortMode::Descending;
+		UE::Insights::CopyNameToClipboard(Table.ToSharedRef(), TreeView->GetSelectedItems(), CurrentSorter, SortMode);
 	}
 }
 
