[Insights] Add TimingInsights.ExportTimerCallees command to export the callee tree of the provided timers.

diff --git a/Engine/Source/Developer/TraceInsights/Private/Insights/TimingProfiler/TimingProfilerManager.cpp b/Engine/Source/Developer/TraceInsights/Private/Insights/TimingProfiler/TimingProfilerManager.cpp
--- a/Engine/Source/Developer/TraceInsights/Private/Insights/TimingProfiler/TimingProfilerManager.cpp
+++ b/Engine/Source/Developer/TraceInsights/Private/Insights/TimingProfiler/TimingProfilerManager.cpp
@@ -900,6 +900,99 @@ bool FTimingProfilerManager::Exec(const TCHAR* Cmd, FOutputDevice& Ar)
 		return true;
 	}
 
+	if (FParse::Command(&Cmd, TEXT("TimingInsights.ExportTimerCallees")))
+	{
+		Ar.Logf(TEXT("TimingInsights.ExportTimerCallees %s"), Cmd);
+
+		check(FInsightsManager::Get().IsValid() && FInsightsManager::Get()->GetSession().IsValid());
+		FTimingExporter Exporter(*FInsightsManager::Get()->GetSession().Get());
+		FTimingExporter::FExportTimerCalleesParams Params;
+
+		// These variables needs to be in the same scope with the call to Exporter.ExportTimerStatisticsAsText().
+		TSet<uint32> IncludedThreads; // referenced in Params.ThreadFilter lambda function
+
+		//////////////////////////////////////////////////
+
+		const bool bUseEscape = true;
+		FString Filename = FParse::Token(Cmd, bUseEscape);
+		Filename.TrimQuotesInline();
+		Ar.Logf(TEXT("  Filename: \"%s\""), *Filename);
+		// '{region}' in Filename (if any) will be replaced with the resolved name of the region.
+
+		while (Cmd && Cmd[0] != TEXT('\0'))
+		{
+			FString Token;
+			if (FParse::Token(Cmd, Token, bUseEscape))
+			{
+				Ar.Logf(TEXT("  Token: %s"), *Token);
+
+				static constexpr TCHAR ThreadsToken[] = TEXT("-threads=");
+				static constexpr TCHAR TimersToken[] = TEXT("-timers=");
+				static constexpr TCHAR StartTimeToken[] = TEXT("-startTime=");
+				static constexpr TCHAR EndTimeToken[] = TEXT("-endTime=");
+				static constexpr TCHAR RegionToken[] = TEXT("-region=");
+
+				if (Token.StartsWith(ThreadsToken))
+				{
+					// Comma-delimited list of thread names. Supports *?-type wildcard.
+					// Default: -threads="*" (all threads; no filter)
+					// Example: -threads="GameThread" -threads="GPU1,GPU2,GameThread,Render*"
+					Token.RightChopInline(UE_ARRAY_COUNT(ThreadsToken) - 1);
+					Token.TrimQuotesInline();
+					Params.ThreadFilter = Exporter.MakeThreadFilterInclusive(Token, IncludedThreads);
+				}
+				else if (Token.StartsWith(TimersToken))
+				{
+					// Comma-delimited list of timer names. Supports *?-type wildcard.
+					// Default: -timers="*" (all timers; no filter)
+					// Example: -timers="A,B,*z"
+					Token.RightChopInline(UE_ARRAY_COUNT(TimersToken) - 1);
+					Token.TrimQuotesInline();
+					Exporter.MakeTimingEventFilterByTimersInclusive(Token, Params.TimerIds);
+				}
+				else if (Token.StartsWith(StartTimeToken))
+				{
+					// Default: -startTime=-infinite
+					// Example: -startTime=10.0
+					Token.RightChopInline(UE_ARRAY_COUNT(StartTimeToken) - 1);
+					Params.IntervalStartTime = FCString::Atof(*Token);
+				}
+				else if (Token.StartsWith(EndTimeToken))
+				{
+					// Default: -endTime=+infinite
+					// Example: -endTime=20.0
+					Token.RightChopInline(UE_ARRAY_COUNT(EndTimeToken) - 1);
+					Params.IntervalEndTime = FCString::Atof(*Token);
+				}
+				else if (Token.StartsWith(RegionToken))
+				{
+					// Comma-delimited list of region names. Supports *?-type wildcard.
+					// Each region is exported to a separate file.
+					// '{region}' in Filename (if any) will be replaced with the resolved name of the region.
+					// Default: -Region=
+					// Example: -Region="RegionName,Game_*,OtherRegion"
+					Token.RightChopInline(UE_ARRAY_COUNT(RegionToken) - 1);
+					Token.TrimQuotesInline();
+					Params.Region = TCHAR_TO_ANSI(*Token);
+				}
+				else
+				{
+					Ar.Logf(ELogVerbosity::Warning, TEXT("Unknown Cmd Param: %s"), *Token);
+				}
+			}
+		}
+
+		//////////////////////////////////////////////////
+
+		if (Params.ThreadFilter == nullptr)
+		{
+			Params.ThreadFilter = [](unsigned int) { return true; };
+		}
+
+		Exporter.ExportTimerCalleesAsText(Filename, Params);
+		return true;
+	}
+
 	if (FParse::Command(&Cmd, TEXT("TimingInsights.ExportCounters")))
 	{
 		Ar.Logf(TEXT("TimingInsights.ExportCounters %s"), Cmd);
diff --git a/Engine/Source/Developer/TraceInsights/Private/Insights/TimingProfiler/ViewModels/TimingExporter.cpp b/Engine/Source/Developer/TraceInsights/Private/Insights/TimingProfiler/ViewModels/TimingExporter.cpp
--- a/Engine/Source/Developer/TraceInsights/Private/Insights/TimingProfiler/ViewModels/TimingExporter.cpp
+++ b/Engine/Source/Developer/TraceInsights/Private/Insights/TimingProfiler/ViewModels/TimingExporter.cpp
@@ -9,6 +9,7 @@
 #include "Common/ProviderLock.h"
 #include "TraceServices/Model/Bookmarks.h"
 #include "TraceServices/Model/Counters.h"
+#include "TraceServices/Model/Frames.h"
 #include "TraceServices/Model/Regions.h"
 #include "TraceServices/Model/Threads.h"
 #include "TraceServices/Model/TimingProfiler.h"
@@ -796,6 +797,160 @@ int32 FTimingExporter::ExportTimerStatisticsAsText(const FString& Filename, FExp
 
 ////////////////////////////////////////////////////////////////////////////////////////////////////
 
+int32 FTimingExporter::ExportTimerCalleesByRegions(const FString& FilenamePattern, const FExportTimerCalleesParams& Params) const
+{
+	TMap<FString, FTimeRegionGroup> RegionGroups;
+	GetRegions(Params.Region, RegionGroups);
+
+	if (RegionGroups.Num() == 0)
+	{
+		UE_LOG(TraceInsights, Error, TEXT("Unable to find any region with name pattern '%s'."), *Params.Region);
+		return -1;
+	}
+
+	FStopwatch Stopwatch;
+	Stopwatch.Start();
+
+	// Export timing callees for each region.
+	FExportTimerCalleesParams RegionParams = Params;
+	RegionParams.Region.Reset();
+	int32 ExportedRegionCount = EnumerateRegions(RegionGroups, FilenamePattern,
+		[this, &RegionParams](const FString& Filename, const FString& RegionName, double IntervalStartTime, double IntervalEndTime)
+		{
+			RegionParams.IntervalStartTime = IntervalStartTime;
+			RegionParams.IntervalEndTime = IntervalEndTime;
+			UE_LOG(TraceInsights, Display, TEXT("Exporting timing callees for region '%s' [%f .. %f] to '%s'"), *RegionName, RegionParams.IntervalStartTime, RegionParams.IntervalEndTime, *Filename);
+			ExportTimerCalleesAsText(Filename, RegionParams);
+		});
+
+	Stopwatch.Stop();
+	const double TotalTime = Stopwatch.GetAccumulatedTime();
+	UE_LOG(TraceInsights, Log, TEXT("Exported timing callees for %d regions in %.3fs."), ExportedRegionCount, TotalTime);
+	return ExportedRegionCount;
+}
+
+////////////////////////////////////////////////////////////////////////////////////////////////////
+
+int32 FTimingExporter::ExportTimerCalleesAsText(const FString& Filename, const FExportTimerCalleesParams& Params) const
+{
+	if (!Params.Region.IsEmpty())
+	{
+		return ExportTimerCalleesByRegions(Filename, Params);
+	}
+
+	TUniquePtr<TraceServices::ITimingProfilerButterfly> Butterfly = [&Params, this]() -> TUniquePtr<TraceServices::ITimingProfilerButterfly>
+	{
+		TraceServices::FAnalysisSessionReadScope SessionReadScope(Session);
+
+		const TraceServices::ITimingProfilerProvider* TimingProfilerProvider = TraceServices::ReadTimingProfilerProvider(Session);
+		if (!TimingProfilerProvider)
+		{
+			UE_LOG(TraceInsights, Error, TEXT("Unable to access TimingProfilerProvider for ExportTimerCalleesAsText"));
+			return nullptr;
+		}
+
+		const bool bIncludeGpu = Params.ThreadFilter(FGpuTimingTrack::Gpu1ThreadId) || Params.ThreadFilter(FGpuTimingTrack::Gpu2ThreadId);
+	
+		return TUniquePtr<TraceServices::ITimingProfilerButterfly>
+		{
+			TimingProfilerProvider->CreateButterfly(Params.IntervalStartTime, Params.IntervalEndTime, Params.ThreadFilter, bIncludeGpu)
+		};
+	}();
+
+	if (!Butterfly)
+	{
+		UE_LOG(TraceInsights, Error, TEXT("ExportTimerCalleesAsText failed to create Butterfly for region %s [%.6f-%.6f]"), *Params.Region, Params.IntervalStartTime, Params.IntervalEndTime);
+		return -1;
+	}
+
+	TUniquePtr<IFileHandle> ExportFileHandle{ OpenExportFile(*Filename) };
+	if (!ExportFileHandle)
+	{
+		UE_LOG(TraceInsights, Error, TEXT("ExportTimerCalleesAsText failed to open export file %s."), *Filename);
+		return -1;
+	}
+
+	const bool bIsCSV = Filename.EndsWith(TEXT(".csv"));
+	FUtf8Writer Writer(ExportFileHandle.Get(), bIsCSV);
+
+	// Write header
+	{
+		TArray<FStringView> Columns
+		{
+			TEXTVIEW("TimerId"),
+			TEXTVIEW("ParentId"),
+			TEXTVIEW("TimerName"),
+			TEXTVIEW("Count"),
+			TEXTVIEW("Inc.Time"),
+			TEXTVIEW("Exc.Time"),
+			TEXTVIEW("NumFrames")
+		};
+
+		for (const FStringView ColumnName : Columns)
+		{
+			Writer.AppendString(ColumnName.GetData());
+			if (ColumnName != Columns.Last())
+			{
+				Writer.AppendSeparator();
+			}
+		}
+		Writer.AppendLineEnd();
+	}
+
+	// Count the number of frames in this region so it can be output and used to calculate frame averages.
+	uint64 NumFrames = 0;
+	{
+		TraceServices::FAnalysisSessionReadScope SessionReadScope(Session);
+		const TraceServices::IFrameProvider& FrameProvider = TraceServices::ReadFrameProvider(Session);
+		FrameProvider.EnumerateFrames(ETraceFrameType::TraceFrameType_Game, Params.IntervalStartTime, Params.IntervalEndTime, [&NumFrames](const TraceServices::FFrame& Frame)
+		{
+			++NumFrames;
+		});
+	}
+
+	// Write rows
+	{
+		FUtf8StringBuilder& StringBuilder = Writer.GetStringBuilder();
+
+		for (uint32 TimerId : Params.TimerIds)
+		{
+			TArray<const TraceServices::FTimingProfilerButterflyNode*> NodesToVisit{ &Butterfly->GenerateCalleesTree(TimerId) };
+			while (!NodesToVisit.IsEmpty())
+			{
+				// Root node can be null if the timer id is for a thread we've filtered out.
+				const TraceServices::FTimingProfilerButterflyNode* CurrentNode = NodesToVisit.Pop(EAllowShrinking::No);
+				if (!CurrentNode || !CurrentNode->Timer)
+				{
+					continue;
+				}
+
+				NodesToVisit.Append(CurrentNode->Children);
+
+				StringBuilder.Appendf(UTF8TEXT("%u"), CurrentNode->Timer->Id);
+				Writer.AppendSeparator();
+				StringBuilder.Appendf(UTF8TEXT("%u"), CurrentNode->Parent ? CurrentNode->Parent->Timer->Id : uint32(-1));
+				Writer.AppendSeparator();
+				Writer.AppendString(CurrentNode->Timer->Name);
+				Writer.AppendSeparator();
+				StringBuilder.Appendf(UTF8TEXT("%llu"), CurrentNode->Count);
+				Writer.AppendSeparator();
+				StringBuilder.Appendf(UTF8TEXT("%.9g"), CurrentNode->InclusiveTime);
+				Writer.AppendSeparator();
+				StringBuilder.Appendf(UTF8TEXT("%.9g"), CurrentNode->ExclusiveTime);
+				Writer.AppendSeparator();
+				StringBuilder.Appendf(UTF8TEXT("%llu"), NumFrames);
+				Writer.AppendLineEnd();
+			}
+		}
+	}
+
+	Writer.Flush();
+	ExportFileHandle->Flush();
+	return 1;
+}
+
+////////////////////////////////////////////////////////////////////////////////////////////////////
+
 FTimingExporter::FThreadFilterFunc FTimingExporter::MakeThreadFilterInclusive(const FString& InFilterString, TSet<uint32>& OutIncludedThreads) const
 {
 	if (InFilterString.Len() == 1 && InFilterString[0] == TEXT('*'))
diff --git a/Engine/Source/Developer/TraceInsights/Private/Insights/TimingProfiler/ViewModels/TimingExporter.h b/Engine/Source/Developer/TraceInsights/Private/Insights/TimingProfiler/ViewModels/TimingExporter.h
--- a/Engine/Source/Developer/TraceInsights/Private/Insights/TimingProfiler/ViewModels/TimingExporter.h
+++ b/Engine/Source/Developer/TraceInsights/Private/Insights/TimingProfiler/ViewModels/TimingExporter.h
@@ -93,6 +93,16 @@ public:
 	{
 	};
 
+	struct FExportTimerCalleesParams : private FExportTimingEventsParams
+	{
+		using FExportTimingEventsParams::ThreadFilter;
+		using FExportTimingEventsParams::IntervalStartTime;
+		using FExportTimingEventsParams::IntervalEndTime;
+		using FExportTimingEventsParams::Region;
+
+		TSet<uint32> TimerIds;
+	};
+
 	struct FExportCountersParams
 	{
 		/**
@@ -220,6 +230,8 @@ public:
 	 */
 	int32 ExportTimerStatisticsAsText(const FString& Filename, FExportTimerStatisticsParams& Params) const;
 
+	int32 ExportTimerCalleesAsText(const FString& Filename, const FExportTimerCalleesParams& Params) const;
+
 	int32 ExportCountersAsText(const FString& Filename, FExportCountersParams& Params) const;
 
 	int32 ExportCounterAsText(const FString& Filename, uint32 CounterId, FExportCounterParams& Params) const;
@@ -283,6 +295,7 @@ private:
 	int32 ExportTimingEvents_WriteEvents(FExportTimingEventsInternalParams& Params) const;
 	int32 ExportTimingEventsAsTextByRegions(const FString& FilenamePattern, FExportTimingEventsParams& Params) const;
 	int32 ExportTimerStatisticsAsTextByRegions(const FString& Filename, FExportTimerStatisticsParams& Params) const;
+	int32 ExportTimerCalleesByRegions(const FString& FilenamePattern, const FExportTimerCalleesParams& Params) const;
 	int32 ExportCounterAsTextByRegions(const FString& Filename, uint32 CounterId, FExportCounterParams& Params) const;
 
 private:
