[Insights]
  - Timing View: Added "Scroll Log View" to any track's context menu (i.e. not limited anymore to the context menu of Time Ruler track). This makes unnecessary PR #11471.
  - Added icon for "Scroll Log View" action in various context menus (in the Timing view and in the Frames view).

#11471

diff --git a/Engine/Source/Developer/TraceInsights/Private/Insights/TimingProfiler/Tracks/TimeRulerTrack.cpp b/Engine/Source/Developer/TraceInsights/Private/Insights/TimingProfiler/Tracks/TimeRulerTrack.cpp
--- a/Engine/Source/Developer/TraceInsights/Private/Insights/TimingProfiler/Tracks/TimeRulerTrack.cpp
+++ b/Engine/Source/Developer/TraceInsights/Private/Insights/TimingProfiler/Tracks/TimeRulerTrack.cpp
@@ -481,26 +481,6 @@ void FTimeRulerTrack::BuildContextMenu(FMenuBuilder& MenuBuilder)
 
 		MenuBuilder.EndSection();
 	}
-
-	MenuBuilder.BeginSection("Misc", LOCTEXT("ContextMenu_Section_Misc", "Misc"));
-
-	FUIAction Action_ScrollLogView
-	(
-		FExecuteAction::CreateSP(this, &FTimeRulerTrack::ContextMenu_ScrollLogView_Execute),
-		FCanExecuteAction::CreateSP(this, &FTimeRulerTrack::ContextMenu_ScrollLogView_CanExecute)
-	);
-	FText Label = FText::Format(LOCTEXT("ContextMenu_ScrollLogView_Fmt", "Scroll Log View (\u2192 {0})"), CrtMousePosTimeText);
-	MenuBuilder.AddMenuEntry
-	(
-		Label,
-		FText::Format(LOCTEXT("ContextMenu_ScrollLogView_Desc_Fmt", "Scrolls the Log View at the message with the closest timestamp to the time of the current mouse position ({0})."), CrtMousePosTimeText),
-		FSlateIcon(),
-		Action_ScrollLogView,
-		NAME_None,
-		EUserInterfaceActionType::Button
-	);
-
-	MenuBuilder.EndSection();
 }
 
 ////////////////////////////////////////////////////////////////////////////////////////////////////
@@ -512,40 +492,6 @@ void FTimeRulerTrack::ContextMenu_MoveTimeMarker_Execute(TSharedRef<FTimeMarker>
 
 ////////////////////////////////////////////////////////////////////////////////////////////////////
 
-bool FTimeRulerTrack::ContextMenu_ScrollLogView_CanExecute()
-{
-	TSharedPtr<STimingProfilerWindow> TimingWindow = FTimingProfilerManager::Get()->GetProfilerWindow();
-	if (!TimingWindow.IsValid())
-	{
-		return false;
-	}
-	TSharedPtr<SLogView> LogView = TimingWindow->GetLogView();
-	if (!LogView.IsValid())
-	{
-		return false;
-	}
-	return true;
-}
-
-////////////////////////////////////////////////////////////////////////////////////////////////////
-
-void FTimeRulerTrack::ContextMenu_ScrollLogView_Execute()
-{
-	TSharedPtr<STimingProfilerWindow> TimingWindow = FTimingProfilerManager::Get()->GetProfilerWindow();
-	if (!TimingWindow.IsValid())
-	{
-		return;
-	}
-	TSharedPtr<SLogView> LogView = TimingWindow->GetLogView();
-	if (!LogView.IsValid())
-	{
-		return;
-	}
-	LogView->SelectLogMessageByClosestTime(CrtMousePosTime);
-}
-
-////////////////////////////////////////////////////////////////////////////////////////////////////
-
 } // namespace UE::Insights::TimingProfiler
 
 #undef LOCTEXT_NAMESPACE
diff --git a/Engine/Source/Developer/TraceInsights/Private/Insights/TimingProfiler/Tracks/TimeRulerTrack.h b/Engine/Source/Developer/TraceInsights/Private/Insights/TimingProfiler/Tracks/TimeRulerTrack.h
--- a/Engine/Source/Developer/TraceInsights/Private/Insights/TimingProfiler/Tracks/TimeRulerTrack.h
+++ b/Engine/Source/Developer/TraceInsights/Private/Insights/TimingProfiler/Tracks/TimeRulerTrack.h
@@ -58,8 +58,6 @@ public:
 private:
 	void DrawTimeMarker(const ITimingTrackDrawContext& Context, const FTimeMarker& TimeMarker) const;
 	void ContextMenu_MoveTimeMarker_Execute(TSharedRef<FTimeMarker> InTimeMarker);
-	bool ContextMenu_ScrollLogView_CanExecute();
-	void ContextMenu_ScrollLogView_Execute();
 
 private:
 	// Slate resources
diff --git a/Engine/Source/Developer/TraceInsights/Private/Insights/TimingProfiler/Widgets/SFrameTrack.cpp b/Engine/Source/Developer/TraceInsights/Private/Insights/TimingProfiler/Widgets/SFrameTrack.cpp
--- a/Engine/Source/Developer/TraceInsights/Private/Insights/TimingProfiler/Widgets/SFrameTrack.cpp
+++ b/Engine/Source/Developer/TraceInsights/Private/Insights/TimingProfiler/Widgets/SFrameTrack.cpp
@@ -1881,7 +1881,7 @@ void SFrameTrack::CreateSelectedFrameMenu(FMenuBuilder& MenuBuilder)
 	(
 		Label,
 		LOCTEXT("ContextMenu_ScrollLogView_Desc", "Scrolls the Log View at the message with the closest timestamp to the start time of the selected frame."),
-		FSlateIcon(),
+		FSlateIcon(FInsightsStyle::GetStyleSetName(), "Icons.LogView"),
 		Action_ScrollLogView,
 		NAME_None,
 		EUserInterfaceActionType::Button
diff --git a/Engine/Source/Developer/TraceInsights/Private/Insights/Widgets/STimingView.cpp b/Engine/Source/Developer/TraceInsights/Private/Insights/Widgets/STimingView.cpp
--- a/Engine/Source/Developer/TraceInsights/Private/Insights/Widgets/STimingView.cpp
+++ b/Engine/Source/Developer/TraceInsights/Private/Insights/Widgets/STimingView.cpp
@@ -83,6 +83,7 @@
 #include "Insights/ViewModels/TimingEventSearch.h"
 #include "Insights/ViewModels/TimingGraphTrack.h"
 #include "Insights/ViewModels/TimingViewDrawHelper.h"
+#include "Insights/Widgets/SLogView.h"
 #include "Insights/Widgets/SQuickFind.h"
 #include "Insights/Widgets/STimingViewTrackList.h"
 
@@ -2991,6 +2992,30 @@ void STimingView::ShowContextMenu(const FPointerEvent& MouseEvent)
 			TAttribute<FText>(),
 			FSlateIcon(FInsightsStyle::GetStyleSetName(), "Icons.Find"));
 
+		TSharedPtr<SLogView> LogView = GetLogView();
+		const double MousePosTime = Viewport.SlateUnitsToTime(static_cast<float>(MousePosition.X));
+		const FText MousePosTimeText = FText::FromString(FormatTimeAuto(MousePosTime, 2));
+		const FText Label = FText::Format(LOCTEXT("ContextMenu_ScrollLogView_Fmt", "Scroll Log View (\u2192 {0})"), MousePosTimeText);
+		MenuBuilder.AddMenuEntry(
+			Label,
+			FText::Format(LOCTEXT("ContextMenu_ScrollLogView_Desc_Fmt", "Scrolls the Log View at the message with the closest timestamp to the time of the current mouse position ({0})."), MousePosTimeText),
+			FSlateIcon(FInsightsStyle::GetStyleSetName(), "Icons.LogView"),
+			FUIAction(
+				FExecuteAction::CreateLambda([LogView, MousePosTime]()
+				{
+					if (LogView.IsValid())
+					{
+						LogView->SelectLogMessageByClosestTime(MousePosTime);
+					}
+				}),
+				FCanExecuteAction::CreateLambda([LogView]() -> bool
+				{
+					return LogView.IsValid();
+				})),
+			NAME_None,
+			EUserInterfaceActionType::Button
+		);
+
 		if (HoveredEvent)
 		{
 			double RangeStart = HoveredEvent->GetStartTime();
@@ -5513,14 +5538,10 @@ void STimingView::EnumerateFilteredTracks(TSharedPtr<UE::Insights::FFilterConfig
 
 ETraceFrameType STimingView::GetFrameTypeToSnapTo()
 {
-	TSharedPtr<STimingProfilerWindow> Window = FTimingProfilerManager::Get()->GetProfilerWindow();
-	if (Window.IsValid())
+	TSharedPtr<STimersView> TimersView = GetTimersView();
+	if (TimersView.IsValid())
 	{
-		TSharedPtr<STimersView> TimersView = Window->GetTimersView();
-		if (TimersView.IsValid())
-		{
-			return TimersView->GetFrameTypeMode();
-		}
+		return TimersView->GetFrameTypeMode();
 	}
 
 	// TraceFrameType_Count is the Instance mode.
@@ -5588,13 +5609,43 @@ void STimingView::UpdateFilters()
 
 ////////////////////////////////////////////////////////////////////////////////////////////////////
 
-bool STimingView::IsInTimingProfiler()
+bool STimingView::IsInTimingProfiler() const
 {
 	return GetName() == FInsightsManagerTabs::TimingProfilerTabId;
 }
 
 ////////////////////////////////////////////////////////////////////////////////////////////////////
 
+TSharedPtr<STimersView> STimingView::GetTimersView() const
+{
+	if (IsInTimingProfiler())
+	{
+		TSharedPtr<STimingProfilerWindow> TimingWindow = FTimingProfilerManager::Get()->GetProfilerWindow();
+		if (TimingWindow.IsValid())
+		{
+			return TimingWindow->GetTimersView();
+		}
+	}
+	return nullptr;
+}
+
+////////////////////////////////////////////////////////////////////////////////////////////////////
+
+TSharedPtr<SLogView> STimingView::GetLogView() const
+{
+	if (IsInTimingProfiler())
+	{
+		TSharedPtr<STimingProfilerWindow> TimingWindow = FTimingProfilerManager::Get()->GetProfilerWindow();
+		if (TimingWindow.IsValid())
+		{
+			return TimingWindow->GetLogView();
+		}
+	}
+	return nullptr;
+}
+
+////////////////////////////////////////////////////////////////////////////////////////////////////
+
 } // namespace UE::Insights::TimingProfiler
 
 #undef INSIGHTS_ACTIVATE_BENCHMARK
diff --git a/Engine/Source/Developer/TraceInsights/Private/Insights/Widgets/STimingView.h b/Engine/Source/Developer/TraceInsights/Private/Insights/Widgets/STimingView.h
--- a/Engine/Source/Developer/TraceInsights/Private/Insights/Widgets/STimingView.h
+++ b/Engine/Source/Developer/TraceInsights/Private/Insights/Widgets/STimingView.h
@@ -48,6 +48,7 @@ namespace UE::Insights
 {
 	class FFilterConfigurator;
 	enum class ETimingEventsColoringMode : uint32;
+	class SLogView;
 }
 
 namespace UE::Insights::Timing
@@ -75,6 +76,8 @@ class FTimeRulerTrack;
 class FTimeMarker;
 class FMarkersTimingTrack;
 
+class STimersView;
+
 enum class ESelectEventType : uint32
 {
 	Min = 0,
@@ -563,7 +566,9 @@ PRAGMA_ENABLE_DEPRECATION_WARNINGS
 	 */ 
 	void UpdateFilters();
 
-	bool IsInTimingProfiler();
+	bool IsInTimingProfiler() const;
+	TSharedPtr<STimersView> GetTimersView() const;
+	TSharedPtr<SLogView> GetLogView() const;
 
 protected:
 	/** The name of the view. */
