[UnrealInsights] Fix an assert in ExportTimerCalleesAsText if the trace contained a region with an infinite interval.

#rnx

diff --git a/Engine/Source/Developer/TraceInsights/Private/Insights/TimingProfiler/ViewModels/TimingExporter.cpp b/Engine/Source/Developer/TraceInsights/Private/Insights/TimingProfiler/ViewModels/TimingExporter.cpp
--- a/Engine/Source/Developer/TraceInsights/Private/Insights/TimingProfiler/ViewModels/TimingExporter.cpp
+++ b/Engine/Source/Developer/TraceInsights/Private/Insights/TimingProfiler/ViewModels/TimingExporter.cpp
@@ -851,9 +851,12 @@ int32 FTimingExporter::ExportTimerCalleesAsText(const FString& Filename, const F
 
 		const bool bIncludeGpu = Params.ThreadFilter(FGpuTimingTrack::Gpu1ThreadId) || Params.ThreadFilter(FGpuTimingTrack::Gpu2ThreadId);
 	
+		// The region end interval may be inf if the capture ended before the region was closed.
+		const double ClampedEndTime = FMath::Min(Session.GetDurationSeconds(), Params.IntervalEndTime);
+
 		return TUniquePtr<TraceServices::ITimingProfilerButterfly>
 		{
-			TimingProfilerProvider->CreateButterfly(Params.IntervalStartTime, Params.IntervalEndTime, Params.ThreadFilter, bIncludeGpu)
+			TimingProfilerProvider->CreateButterfly(Params.IntervalStartTime, ClampedEndTime, Params.ThreadFilter, bIncludeGpu)
 		};
 	}();
 
