Unreal Insights - Commands - Export - Linux

#rnx

diff --git a/Engine/Source/Developer/TraceInsights/Private/Insights/Tests/FunctionalTests/ExportCommandsTests.cpp b/Engine/Source/Developer/TraceInsights/Private/Insights/Tests/FunctionalTests/ExportCommandsTests.cpp
--- a/Engine/Source/Developer/TraceInsights/Private/Insights/Tests/FunctionalTests/ExportCommandsTests.cpp
+++ b/Engine/Source/Developer/TraceInsights/Private/Insights/Tests/FunctionalTests/ExportCommandsTests.cpp
@@ -1,217 +1,275 @@
 // Copyright Epic Games, Inc. All Rights Reserved.
 
-#include "HAL/FileManager.h"
 #include "HAL/PlatformFileManager.h"
 #include "Misc/AutomationTest.h"
 #include "Misc/FileHelper.h"
 
-// TraceAnalysis
-#include "Trace/StoreClient.h"
-
 // TraceInsightsCore
 #include "InsightsCore/Common/MiscUtils.h"
 
 // TraceInsights
-#include "Insights/InsightsManager.h"
-#include "Insights/IUnrealInsightsModule.h"
 #include "Insights/Tests/InsightsTestUtils.h"
 
-void VerifyExportedLines(const FString& ExportReportPath, const FString& CmdLogPath, const FString& Elements, FInsightsTestUtils Utils, FAutomationTestBase* Test, double Timeout)
-{
-	bool bLineFound = false;
-	FString ExpectedResult;
+DEFINE_LATENT_AUTOMATION_COMMAND_FIVE_PARAMETER(FVerifyExportedLinesCommand, const FString, ExportReportPath, const FString, CmdLogPath, const FString, Elements, FAutomationTestBase*, Test, double, Timeout);
 
-	double StartTime = FPlatformTime::Seconds();
-	while (FPlatformTime::Seconds() - StartTime < Timeout)
+bool FVerifyExportedLinesCommand::Update()
+{
+	if (FPlatformTime::Seconds() - StartTime >= Timeout)
 	{
-		FString FileContent;
+		Test->AddError(TEXT("The FVerifyExportedLinesCommand timed out"));
+		return true;
+	}
 
-		if (FFileHelper::LoadFileToString(FileContent, *ExportReportPath))
-		{
-			TArray<FString> Lines;
-			FileContent.ParseIntoArrayLines(Lines);
-			ExpectedResult = FString::Printf(TEXT("Exported %d %s to file"), Lines.Num() - 1, *Elements);
-			bLineFound = Utils.FileContainsString(CmdLogPath, ExpectedResult, 1.0f);
-			if (bLineFound)
-			{
-				Test->TestTrue(FString::Printf(TEXT("Line '%s' from '%s' should exists in file: '%s'"), *ExpectedResult, *ExportReportPath, *CmdLogPath), bLineFound);
-				return;
-			}
-			else
-			{
-				FPlatformProcess::Sleep(0.1f);
-			}
-		}
-		else
+	FString FileContent;
+	if (FFileHelper::LoadFileToString(FileContent, *ExportReportPath))
+	{
+		TArray<FString> Lines;
+		FString ExpectedResult;
+		FInsightsTestUtils Utils(Test);
+		FileContent.ParseIntoArrayLines(Lines);
+		ExpectedResult = FString::Printf(TEXT("Exported %d %s to file"), Lines.Num() - 1, *Elements);
+		if (Utils.FileContainsString(CmdLogPath, ExpectedResult, 1.0f))
 		{
-			FPlatformProcess::Sleep(0.1f);
+			return true;
 		}
 	}
 
-	Test->AddError(FString::Printf(TEXT("VerifyExportedLines timed out while trying to find line '%s' from '%s'"), *ExpectedResult, *ExportReportPath));
+	return false;
 }
 
-#if WITH_AUTOMATION_TESTS
+DEFINE_LATENT_AUTOMATION_COMMAND_FOUR_PARAMETER(FDeleteDirectoryCommand, IPlatformFile&, PlatformFile, const FString, DirectoryPath, FAutomationTestBase*, Test, double, Timeout);
 
-// The objective of these tests is to verify if the export command generates the predicted files and if these files encompass the anticipated data.
-// This confirmation is realized through the utilization of a pre-recorded trace file, which ensures data reliability.
-IMPLEMENT_SIMPLE_AUTOMATION_TEST(FСommandsExportWindowsTest, "System.Insights.Trace.Analysis.ExecCmd.CommandsExport(Windows)", EAutomationTestFlags::ProgramContext | EAutomationTestFlags::EngineFilter)
-bool FСommandsExportWindowsTest::RunTest(const FString& Parameters)
+bool FDeleteDirectoryCommand::Update()
 {
-	IPlatformFile& PlatformFile = FPlatformFileManager::Get().GetPlatformFile();
+	if (FPlatformTime::Seconds() - StartTime >= Timeout)
+	{
+		Test->AddError(TEXT("The FDeleteDirectoryCommand timed out"));
+		return true;
+	}
 
-	FInsightsTestUtils Utils(this);
+	IFileManager::Get().DeleteDirectory(*DirectoryPath, false, true);
 
-	TSharedPtr<UE::Insights::FInsightsManager> InsightsManager = UE::Insights::FInsightsManager::Get();
+	if (!PlatformFile.DirectoryExists(*DirectoryPath))
+	{
+		return true;
+	}
 
-	const FString StoreDir = InsightsManager->GetStoreDir();
-	const FString SourceTracePath = FPaths::RootDir() / TEXT("EngineTest/SourceAssets/Utrace/CommandsExportTest_5.4.utrace");
-	const FString SourceExportPath = FPaths::RootDir() / TEXT("EngineTest/SourceAssets/Rsp/export.rsp");
-	const FString StoreTracePath = StoreDir / TEXT("CommandsExportTest_5.4.utrace");
-	const FString TestResultsDirPath = TEXT("/TestResults");
-	const FString SingleCommandDirPath = TEXT("/TestResults/SingleCommand");
-	const FString LogResultExportPath = TEXT("/TestResults/export.rsp");
-	const FString CmdLogPath = TEXT("/TestResults/Logs/cmd.log");
-	const FString CmdThreadLogPath = TEXT("/TestResults/Logs/cmd_treads.log");
-	const FString CmdTimersLogPath = TEXT("/TestResults/Logs/cmd_timers.log");
-	const FString CmdTimingEventsLogPath = TEXT("/TestResults/Logs/cmd_timing_events.log");
-	const FString CmdTimingEventsNonDefaultLogPath = TEXT("/TestResults/Logs/cmd_timing_non_default.log");
-	const FString CmdExportLogPath = TEXT("/TestResults/Logs/cmd_export.log");
-	const FString ExportThreadsTask = TEXT("TimingInsights.ExportThreads /TestResults/SingleCommand/Threads.csv");
-	const FString ExportTimersTask = TEXT("TimingInsights.ExportTimers /TestResults/SingleCommand/Timers.csv");
-	const FString ExportTimingEventsTask = TEXT("TimingInsights.ExportTimingEvents /TestResults/SingleCommand/TimingEvents.csv");
-	const FString ExportTimingEventsBorderedTask = TEXT("TimingInsights.ExportTimingEvents /TestResults/SingleCommand/TimingEventsNonDefault.CSV  -columns=ThreadId,ThreadName,TimerId,TimerName,StartTime,EndTime,Duration,Depth -threads=GameThread -timers=* -startTime=10 -endTime=20");
-	const FString ExportTimerStatisticsTask = TEXT("TimingInsights.ExportTimerStatistics  /TestResults/SingleCommand/*.csv -region=* -threads=GPU");
-	const FString ExportThreadsReportPath = TEXT("/TestResults/SingleCommand/Threads.csv");
-	const FString ExportTimersReportPath = TEXT("/TestResults/SingleCommand/Timers.csv");
-	const FString ExportTimingEventsReportPath = TEXT("/TestResults/SingleCommand/TimingEvents.csv");
-	const FString ExportTimingEventsNonDefaultReportPath = TEXT("/TestResults/SingleCommand/TimingEventsNonDefault.csv");
-	const FString ExportTimerStatisticsReportPath = TEXT("/TestResults/SingleCommand/TimerStatistics.csv");
-	bool bLineFound = false;
-	double Timeout = 150.0;
+	return false;
+}
 
-	UTEST_TRUE("Trace in project exists", PlatformFile.FileExists(*SourceTracePath));
-	UTEST_TRUE("Export in project exists", PlatformFile.FileExists(*SourceExportPath));
-	UTEST_FALSE("Trace in store should not exists before copy", PlatformFile.FileExists(*StoreTracePath));
+#if WITH_AUTOMATION_TESTS
 
+// The goal of this test is to verify that threads data can be exported from a trace
+IMPLEMENT_SIMPLE_AUTOMATION_TEST(FCommandsExportThreadsDataTest, "System.Insights.Trace.Analysis.CommandsExport.ThreadsData", EAutomationTestFlags::ProgramContext | EAutomationTestFlags::EngineFilter)
+bool FCommandsExportThreadsDataTest::RunTest(const FString& Parameters)
+{
+	IPlatformFile& PlatformFile = FPlatformFileManager::Get().GetPlatformFile();
+	const FString SourceTracePath = FPaths::RootDir() / TEXT("EngineTest/SourceAssets/Utrace/CommandsExportTest_5.5.utrace");
+	const FString StoreTracePath = TEXT("/tmp/CommandsExportTest_5.5.utrace");
+	const FString TestResultsDirPath = TEXT("/tmp/TestResults");
+	const FString CmdThreadLogPath = TestResultsDirPath / TEXT("Logs/cmd_treads.log");
+	const FString ExportThreadsReportPath = TestResultsDirPath / TEXT("SingleCommand/Threads.csv");
+	const FString ExportThreadsTask = TEXT("TimingInsights.ExportThreads ") + ExportThreadsReportPath;
+	double Timeout = 10.0;
+
+	UTEST_TRUE(TEXT("Trace in project exists"), PlatformFile.FileExists(*SourceTracePath));
 	PlatformFile.CopyFile(*StoreTracePath, *SourceTracePath);
 	UTEST_TRUE(TEXT("Trace in store should exists after copy"), PlatformFile.FileExists(*StoreTracePath));
 
 	if (PlatformFile.DirectoryExists(*TestResultsDirPath))
 	{
 		AddInfo(TEXT("The TestResults directory already exists. Deleting to avoid undefined behavior"));
-		IFileManager::Get().DeleteDirectory(*TestResultsDirPath, false, true);
-		double StartTime = FPlatformTime::Seconds();
-		while (FPlatformTime::Seconds() - StartTime < Timeout)
-		{
-			if (!PlatformFile.DirectoryExists(*TestResultsDirPath))
-			{
-				break;
-			}
-			FPlatformProcess::Sleep(0.1f);
-		}
-
-		UTEST_FALSE("Previously created TestResults directory is not deleted. Execution of the next steps will cause unstable behavior.", PlatformFile.DirectoryExists(*TestResultsDirPath));
+		ADD_LATENT_AUTOMATION_COMMAND(FDeleteDirectoryCommand(PlatformFile, TestResultsDirPath, this, Timeout));
 	}
 
-	// ExportThreads
-	FString InsightsParameters = FString::Printf(TEXT("-OpenTraceFile=\"%s\" -ABSLOG=\"%s\" -AutoQuit -NoUI  -ExecOnAnalysisCompleteCmd=\"%s\" -log"), *StoreTracePath, *CmdThreadLogPath, *ExportThreadsTask);
+	FString InsightsParameters = FString::Printf(TEXT("-OpenTraceFile=\"%s\" -ABSLOG=\"%s\" -AutoQuit -NoUI -ExecOnAnalysisCompleteCmd=\"%s\" -log"), *StoreTracePath, *CmdThreadLogPath, *ExportThreadsTask);
 	UE::Insights::FMiscUtils::OpenUnrealInsights(*InsightsParameters);
+	ADD_LATENT_AUTOMATION_COMMAND(FVerifyExportedLinesCommand(ExportThreadsReportPath, CmdThreadLogPath, TEXT("threads"), this, Timeout));
+	ADD_LATENT_AUTOMATION_COMMAND(FDeleteDirectoryCommand(PlatformFile, TestResultsDirPath, this, Timeout));
+
+	return true;
+}
 
-	VerifyExportedLines(ExportThreadsReportPath, CmdThreadLogPath, TEXT("threads"), Utils, this, Timeout);
+// The goal of this test is to verify that timers data can be exported from a trace
+IMPLEMENT_SIMPLE_AUTOMATION_TEST(FCommandsExportTimersDataTest, "System.Insights.Trace.Analysis.CommandsExport.TimersData", EAutomationTestFlags::ProgramContext | EAutomationTestFlags::EngineFilter)
+bool FCommandsExportTimersDataTest::RunTest(const FString& Parameters)
+{
+	IPlatformFile& PlatformFile = FPlatformFileManager::Get().GetPlatformFile();
+	const FString SourceTracePath = FPaths::RootDir() / TEXT("EngineTest/SourceAssets/Utrace/CommandsExportTest_5.5.utrace");
+	const FString StoreTracePath = TEXT("/tmp/CommandsExportTest_5.5.utrace");
+	const FString TestResultsDirPath = TEXT("/tmp/TestResults");
+	const FString CmdTimersLogPath = TestResultsDirPath / TEXT("Logs/cmd_timers.log");
+	const FString ExportTimersReportPath = TestResultsDirPath / TEXT("SingleCommand/Timers.csv");
+	const FString ExportTimersTask = TEXT("TimingInsights.ExportTimers ") + ExportTimersReportPath;
+	double Timeout = 10.0;
+
+	UTEST_TRUE(TEXT("Trace in project exists"), PlatformFile.FileExists(*SourceTracePath));
+	PlatformFile.CopyFile(*StoreTracePath, *SourceTracePath);
+	UTEST_TRUE(TEXT("Trace in store should exists after copy"), PlatformFile.FileExists(*StoreTracePath));
 
-	// ExportTimers
-	InsightsParameters = FString::Printf(TEXT("-OpenTraceFile=\"%s\" -ABSLOG=\"%s\" -AutoQuit -NoUI  -ExecOnAnalysisCompleteCmd=\"%s\" -log"), *StoreTracePath, *CmdTimersLogPath, *ExportTimersTask);
+	if (PlatformFile.DirectoryExists(*TestResultsDirPath))
+	{
+		AddInfo(TEXT("The TestResults directory already exists. Deleting to avoid undefined behavior"));
+		ADD_LATENT_AUTOMATION_COMMAND(FDeleteDirectoryCommand(PlatformFile, TestResultsDirPath, this, Timeout));
+	}
+
+	FString InsightsParameters = FString::Printf(TEXT("-OpenTraceFile=\"%s\" -ABSLOG=\"%s\" -AutoQuit -NoUI -ExecOnAnalysisCompleteCmd=\"%s\" -log"), *StoreTracePath, *CmdTimersLogPath, *ExportTimersTask);
 	UE::Insights::FMiscUtils::OpenUnrealInsights(*InsightsParameters);
+	ADD_LATENT_AUTOMATION_COMMAND(FVerifyExportedLinesCommand(ExportTimersReportPath, CmdTimersLogPath, TEXT("timers"), this, Timeout));
+	ADD_LATENT_AUTOMATION_COMMAND(FDeleteDirectoryCommand(PlatformFile, TestResultsDirPath, this, Timeout));
 
-	VerifyExportedLines(ExportTimersReportPath, CmdTimersLogPath, TEXT("timers"), Utils, this, Timeout);
+	return true;
+}
 
-	// UI verification cannot be executed in that case.
+// The goal of this test is to verify that timing events data can be exported from a trace
+IMPLEMENT_SIMPLE_AUTOMATION_TEST(FCommandsExportTimingEventsDataTest, "System.Insights.Trace.Analysis.CommandsExport.TimingEventsData", EAutomationTestFlags::ProgramContext | EAutomationTestFlags::EngineFilter)
+bool FCommandsExportTimingEventsDataTest::RunTest(const FString& Parameters)
+{
+	IPlatformFile& PlatformFile = FPlatformFileManager::Get().GetPlatformFile();
+	const FString SourceTracePath = FPaths::RootDir() / TEXT("EngineTest/SourceAssets/Utrace/CommandsExportTest_5.5.utrace");
+	const FString StoreTracePath = TEXT("/tmp/CommandsExportTest_5.5.utrace");
+	const FString TestResultsDirPath = TEXT("/tmp/TestResults");
+	const FString CmdTimingEventsLogPath = TestResultsDirPath / TEXT("Logs/cmd_timing_events.log");
+	const FString ExportTimingEventsReportPath = TestResultsDirPath / TEXT("SingleCommand/TimingEvents.txt");
+	const FString ExportTimingEventsTask = TEXT("TimingInsights.ExportTimingEvents " + ExportTimingEventsReportPath);
+	double Timeout = 10.0;
+
+	UTEST_TRUE(TEXT("Trace in project exists"), PlatformFile.FileExists(*SourceTracePath));
+	PlatformFile.CopyFile(*StoreTracePath, *SourceTracePath);
+	UTEST_TRUE(TEXT("Trace in store should exists after copy"), PlatformFile.FileExists(*StoreTracePath));
 
-	// ExportTimingEvents
-	InsightsParameters = FString::Printf(TEXT("-OpenTraceFile=\"%s\" -ABSLOG=\"%s\" -AutoQuit -NoUI  -ExecOnAnalysisCompleteCmd=\"%s\" -log"), *StoreTracePath, *CmdTimingEventsLogPath, *ExportTimingEventsTask);
+	if (PlatformFile.DirectoryExists(*TestResultsDirPath))
+	{
+		AddInfo(TEXT("The TestResults directory already exists. Deleting to avoid undefined behavior"));
+		ADD_LATENT_AUTOMATION_COMMAND(FDeleteDirectoryCommand(PlatformFile, TestResultsDirPath, this, Timeout));
+	}
+
+	FString InsightsParameters = FString::Printf(TEXT("-OpenTraceFile=\"%s\" -ABSLOG=\"%s\" -AutoQuit -NoUI -ExecOnAnalysisCompleteCmd=\"%s\" -log"), *StoreTracePath, *CmdTimingEventsLogPath, *ExportTimingEventsTask);
 	UE::Insights::FMiscUtils::OpenUnrealInsights(*InsightsParameters);
+	ADD_LATENT_AUTOMATION_COMMAND(FVerifyExportedLinesCommand(ExportTimingEventsReportPath, CmdTimingEventsLogPath, TEXT("timing events"), this, Timeout));
+	ADD_LATENT_AUTOMATION_COMMAND(FDeleteDirectoryCommand(PlatformFile, TestResultsDirPath, this, Timeout));
+
+	return true;
+}
 
-	VerifyExportedLines(ExportTimingEventsReportPath, CmdTimingEventsLogPath, TEXT("timing events"), Utils, this, Timeout);
+// The goal of this test is to verify that a filtered list of timing events data can be exported from a trace
+IMPLEMENT_SIMPLE_AUTOMATION_TEST(FCommandsExportFilteredTimingEventsDataTest, "System.Insights.Trace.Analysis.CommandsExport.FilteredTimingEventsData", EAutomationTestFlags::ProgramContext | EAutomationTestFlags::EngineFilter)
+bool FCommandsExportFilteredTimingEventsDataTest::RunTest(const FString& Parameters)
+{
+	IPlatformFile& PlatformFile = FPlatformFileManager::Get().GetPlatformFile();
+	const FString SourceTracePath = FPaths::RootDir() / TEXT("EngineTest/SourceAssets/Utrace/CommandsExportTest_5.5.utrace");
+	const FString StoreTracePath = TEXT("/tmp/CommandsExportTest_5.5.utrace");
+	const FString TestResultsDirPath = TEXT("/tmp/TestResults");
+	const FString CmdTimingEventsNonDefaultLogPath = TestResultsDirPath / TEXT("Logs/cmd_timing_non_default.log");
+	const FString ExportTimingEventsNonDefaultReportPath = TestResultsDirPath / TEXT("SingleCommand/TimingEventsNonDefault.csv");
+	const FString ExportTimingEventsBorderedTask = TEXT("TimingInsights.ExportTimingEvents " + ExportTimingEventsNonDefaultReportPath + " -columns=ThreadId,ThreadName,TimerId,TimerName,StartTime,EndTime,Duration,Depth -threads=GameThread -timers=* -startTime=10 -endTime=20");
+	double Timeout = 10.0;
+
+	UTEST_TRUE(TEXT("Trace in project exists"), PlatformFile.FileExists(*SourceTracePath));
+	PlatformFile.CopyFile(*StoreTracePath, *SourceTracePath);
+	UTEST_TRUE(TEXT("Trace in store should exists after copy"), PlatformFile.FileExists(*StoreTracePath));
 
-	// ExportTimingEventsNonDefault
-	InsightsParameters = FString::Printf(TEXT("-OpenTraceFile=\"%s\" -ABSLOG=\"%s\" -AutoQuit -NoUI  -ExecOnAnalysisCompleteCmd=\"%s\" -log"), *StoreTracePath, *CmdTimingEventsNonDefaultLogPath, *ExportTimingEventsBorderedTask);
+	if (PlatformFile.DirectoryExists(*TestResultsDirPath))
+	{
+		AddInfo(TEXT("The TestResults directory already exists. Deleting to avoid undefined behavior"));
+		ADD_LATENT_AUTOMATION_COMMAND(FDeleteDirectoryCommand(PlatformFile, TestResultsDirPath, this, Timeout));
+	}
+
+	FString InsightsParameters = FString::Printf(TEXT("-OpenTraceFile=\"%s\" -ABSLOG=\"%s\" -AutoQuit -NoUI -ExecOnAnalysisCompleteCmd=\"%s\" -log"), *StoreTracePath, *CmdTimingEventsNonDefaultLogPath, *ExportTimingEventsBorderedTask);
 	UE::Insights::FMiscUtils::OpenUnrealInsights(*InsightsParameters);
 
-	const TArray<FString> ExpectedTimingEventsBorderedElements =
-	{
-		TEXT("ThreadId"),
-		TEXT("ThreadName"),
-		TEXT("TimerId"),
-		TEXT("TimerName"),
-		TEXT("StartTime"),
-		TEXT("EndTime"),
-		TEXT("Duration"),
-		TEXT("Depth"),
-		TEXT("2,GameThread,1,FEngineLoop::PreInitPreStartupScreen")
+	const TArray<FStringView> ExpectedTimingEventsBorderedElements =
+	{
+		TEXTVIEW("ThreadId"),
+		TEXTVIEW("ThreadName"),
+		TEXTVIEW("TimerId"),
+		TEXTVIEW("TimerName"),
+		TEXTVIEW("StartTime"),
+		TEXTVIEW("EndTime"),
+		TEXTVIEW("Duration"),
+		TEXTVIEW("Depth"),
 	};
 
+	bool bLineFound = false;
+	FInsightsTestUtils Utils(this);
 	for (int i = 0; i < ExpectedTimingEventsBorderedElements.Num(); i++)
 	{
-		bLineFound = Utils.FileContainsString(ExportTimingEventsNonDefaultReportPath, ExpectedTimingEventsBorderedElements[i], Timeout);
-		TestTrue(FString::Printf(TEXT("Line '%s' should exists in file: '%s'"), *ExpectedTimingEventsBorderedElements[i], *ExportTimingEventsNonDefaultReportPath), bLineFound);
+		bLineFound = Utils.FileContainsString(ExportTimingEventsNonDefaultReportPath, ExpectedTimingEventsBorderedElements[i].GetData(), Timeout);
+		UTEST_TRUE(FString::Printf(TEXT("Line '%s' should exists in file: '%s'"), ExpectedTimingEventsBorderedElements[i].GetData(), *ExportTimingEventsNonDefaultReportPath), bLineFound);
 	}
 
-	VerifyExportedLines(ExportTimingEventsNonDefaultReportPath, CmdTimingEventsNonDefaultLogPath, TEXT("timing events"), Utils, this, Timeout);
+	ADD_LATENT_AUTOMATION_COMMAND(FVerifyExportedLinesCommand(ExportTimingEventsNonDefaultReportPath, CmdTimingEventsNonDefaultLogPath, TEXT("timing events"), this, Timeout));
+	ADD_LATENT_AUTOMATION_COMMAND(FDeleteDirectoryCommand(PlatformFile, TestResultsDirPath, this, Timeout));
 
-	// Export.rsp
+	return true;
+}
+
+// The goal of this test is to verify that files with data are generated when executing multiple export commands using a response file
+IMPLEMENT_SIMPLE_AUTOMATION_TEST(FCommandsExportMultipleExportCommandsTest, "System.Insights.Trace.Analysis.CommandsExport.MultipleExportCommands", EAutomationTestFlags::ProgramContext | EAutomationTestFlags::EngineFilter)
+bool FCommandsExportMultipleExportCommandsTest::RunTest(const FString& Parameters)
+{
+	IPlatformFile& PlatformFile = FPlatformFileManager::Get().GetPlatformFile();
+	const FString SourceTracePath = FPaths::RootDir() / TEXT("EngineTest/SourceAssets/Utrace/CommandsExportTest_5.5.utrace");
+	const FString SourceExportPath = FPaths::RootDir() / TEXT("EngineTest/SourceAssets/Rsp/export.rsp");
+	const FString StoreTracePath = TEXT("/tmp/CommandsExportTest_5.5.utrace");
+	const FString LogResultExportPath = TEXT("/tmp/export.rsp");
+	const FString TestResultsDirPath = TEXT("/tmp/TestResults");
+	const FString CmdExportLogPath = TestResultsDirPath / TEXT("Logs/cmd_export.log");
+	double Timeout = 10.0;
+
+	UTEST_TRUE(TEXT("Trace in project exists"), PlatformFile.FileExists(*SourceTracePath));
+	UTEST_TRUE(TEXT("Export in project exists"), PlatformFile.FileExists(*SourceExportPath));
 	PlatformFile.CopyFile(*LogResultExportPath, *SourceExportPath);
-	TestTrue("Rsp in log directory should exists after copy", PlatformFile.FileExists(*LogResultExportPath));
+	PlatformFile.CopyFile(*StoreTracePath, *SourceTracePath);
+	UTEST_TRUE(TEXT("Rsp in log directory should exists after copy"), PlatformFile.FileExists(*LogResultExportPath));
+	UTEST_TRUE(TEXT("Trace in store should exists after copy"), PlatformFile.FileExists(*StoreTracePath));
+
+	if (PlatformFile.DirectoryExists(*TestResultsDirPath))
+	{
+		AddInfo(TEXT("The TestResults directory already exists. Deleting to avoid undefined behavior"));
+		ADD_LATENT_AUTOMATION_COMMAND(FDeleteDirectoryCommand(PlatformFile, TestResultsDirPath, this, Timeout));
+	}
 
-	InsightsParameters = FString::Printf(TEXT("-OpenTraceFile=\"%s\" -ABSLOG=\"%s\" -AutoQuit -NoUI  -ExecOnAnalysisCompleteCmd=\"@=/TestResults/export.rsp\" -log"), *StoreTracePath, *CmdExportLogPath);
+	FString InsightsParameters = FString::Printf(TEXT("-OpenTraceFile=\"%s\" -ABSLOG=\"%s\" -AutoQuit -NoUI -ExecOnAnalysisCompleteCmd=\"@=%s\" -log"), *StoreTracePath, *CmdExportLogPath, *LogResultExportPath);
 	UE::Insights::FMiscUtils::OpenUnrealInsights(*InsightsParameters);
 
-	const TArray<FString> ExpectedThreadsElementsRsp =
+	const TArray<FStringView> ExpectedThreadsElementsRsp =
 	{
-		TEXT("/TestResults/RSPtest/CSV/Threads_rsp.csv"),
-		TEXT("/TestResults/RSPtest/TSV/Threads_rsp.tsv"),
-		TEXT("/TestResults/RSPtest/TXT/Threads_rsp.txt"),
+		TEXTVIEW("/tmp/TestResults/RSPtest/CSV/Threads_rsp.csv"),
+		TEXTVIEW("/tmp/TestResults/RSPtest/TSV/Threads_rsp.tsv"),
+		TEXTVIEW("/tmp/TestResults/RSPtest/TXT/Threads_rsp.txt"),
 	};
-
-	for (int i = 0; i < ExpectedThreadsElementsRsp.Num(); i++)
-	{
-		bLineFound = Utils.FileContainsString(CmdExportLogPath, ExpectedThreadsElementsRsp[i], Timeout);
-		TestTrue(FString::Printf(TEXT("Line '%s' should exists in file: '%s'"), *ExpectedThreadsElementsRsp[i], *CmdExportLogPath), bLineFound);
-
-		VerifyExportedLines(ExpectedThreadsElementsRsp[i], CmdExportLogPath, TEXT("threads"), Utils, this, Timeout);
-	}
-
-	const TArray<FString> ExpectedTimersElementsRsp =
+	const TArray<FStringView> ExpectedTimersElementsRsp =
 	{
-		TEXT("/TestResults/RSPtest/CSV/Timers_rsp.csv"),
-		TEXT("/TestResults/RSPtest/TSV/Timers_rsp.tsv"),
-		TEXT("/TestResults/RSPtest/TXT/Timers_rsp.txt"),
+		TEXTVIEW("/tmp/TestResults/RSPtest/CSV/Timers_rsp.csv"),
+		TEXTVIEW("/tmp/TestResults/RSPtest/TSV/Timers_rsp.tsv"),
+		TEXTVIEW("/tmp/TestResults/RSPtest/TXT/Timers_rsp.txt"),
 	};
-
-	for (int i = 0; i < ExpectedTimersElementsRsp.Num(); i++)
+	const TArray<FStringView> ExpectedTimingEventsRsp =
 	{
-		bLineFound = Utils.FileContainsString(CmdExportLogPath, ExpectedTimersElementsRsp[i], Timeout);
-		TestTrue(FString::Printf(TEXT("Line '%s' should exists in file: '%s'"), *ExpectedTimersElementsRsp[i], *CmdExportLogPath), bLineFound);
-
-		VerifyExportedLines(ExpectedTimersElementsRsp[i], CmdExportLogPath, TEXT("timers"), Utils, this, Timeout);
-	}
-
-	const TArray<FString> ExpectedTimingEventsRsp =
+		TEXTVIEW("/tmp/TestResults/RSPtest/CSV/TimingEvents_rsp.csv"),
+		TEXTVIEW("/tmp/TestResults/RSPtest/TSV/TimingEvents_rsp.tsv"),
+		TEXTVIEW("/tmp/TestResults/RSPtest/TXT/TimingEvents_rsp.txt"),
+	};
+	const TMap<FStringView, TArray<FStringView>> ExpectedElementsRspMap =
 	{
-		TEXT("/TestResults/RSPtest/CSV/TimingEvents_rsp.csv"),
-		TEXT("/TestResults/RSPtest/TSV/TimingEvents_rsp.tsv"),
-		TEXT("/TestResults/RSPtest/TXT/TimingEvents_rsp.txt"),
+		{TEXTVIEW("threads"), ExpectedThreadsElementsRsp},
+		{TEXTVIEW("timers"), ExpectedTimersElementsRsp},
+		{TEXTVIEW("timing events"), ExpectedTimingEventsRsp},
 	};
 
-	for (int i = 0; i < ExpectedTimingEventsRsp.Num(); i++)
+	bool bLineFound = false;
+	FInsightsTestUtils Utils(this);
+	for (const auto& [Element, ExpectedElementsRsp] : ExpectedElementsRspMap)
 	{
-		bLineFound = Utils.FileContainsString(CmdExportLogPath, ExpectedTimingEventsRsp[i], Timeout);
-		TestTrue(FString::Printf(TEXT("Line '%s' should exists in file: '%s'"), *ExpectedTimingEventsRsp[i], *CmdExportLogPath), bLineFound);
-
-		VerifyExportedLines(ExpectedTimingEventsRsp[i], CmdExportLogPath, TEXT("timing events"), Utils, this, Timeout);
+		for (int i = 0; i < ExpectedElementsRsp.Num(); i++)
+		{
+			bLineFound = Utils.FileContainsString(CmdExportLogPath, ExpectedElementsRsp[i].GetData(), Timeout);
+			UTEST_TRUE(FString::Printf(TEXT("Line '%s' should exists in file: '%s'"), ExpectedElementsRsp[i].GetData(), *CmdExportLogPath), bLineFound);
+			ADD_LATENT_AUTOMATION_COMMAND(FVerifyExportedLinesCommand(ExpectedElementsRsp[i].GetData(), CmdExportLogPath, Element.GetData(), this, Timeout));
+		}
 	}
-
-	IFileManager::Get().Delete(*StoreTracePath);
+	ADD_LATENT_AUTOMATION_COMMAND(FDeleteDirectoryCommand(PlatformFile, TestResultsDirPath, this, Timeout));
 
 	return true;
 }
