[Trivial] Ensure LogTimingProfiler is exposed

#rnx

diff --git a/Engine/Source/Developer/TraceInsights/Public/Insights/TimingProfilerCommon.h b/Engine/Source/Developer/TraceInsights/Public/Insights/TimingProfilerCommon.h
--- a/Engine/Source/Developer/TraceInsights/Public/Insights/TimingProfilerCommon.h
+++ b/Engine/Source/Developer/TraceInsights/Public/Insights/TimingProfilerCommon.h
@@ -7,7 +7,7 @@
 #include "Stats/Stats.h"
 #include "Stats/StatsMisc.h"
 
-DECLARE_LOG_CATEGORY_EXTERN(LogTimingProfiler, Log, All);
+TRACEINSIGHTS_API DECLARE_LOG_CATEGORY_EXTERN(LogTimingProfiler, Log, All);
 
 DECLARE_STATS_GROUP(TEXT("TimingProfiler"), STATGROUP_TimingProfiler, STATCAT_Advanced);
 
