[Insights] Add filters with suggestions for AllocThreadId and FreeThreadId in Memory Investigation View.

diff --git a/Engine/Source/Developer/TraceInsights/Private/Insights/MemoryProfiler/ViewModels/MemAllocTable.cpp b/Engine/Source/Developer/TraceInsights/Private/Insights/MemoryProfiler/ViewModels/MemAllocTable.cpp
--- a/Engine/Source/Developer/TraceInsights/Private/Insights/MemoryProfiler/ViewModels/MemAllocTable.cpp
+++ b/Engine/Source/Developer/TraceInsights/Private/Insights/MemoryProfiler/ViewModels/MemAllocTable.cpp
@@ -474,7 +474,7 @@ void FMemAllocTable::AddDefaultColumns()
 		Column.SetTitleName(LOCTEXT("AllocThreadColumnTitle", "Alloc Thread"));
 		Column.SetDescription(LOCTEXT("AllocThreadColumnDesc", "The thread the allocation was made on"));
 
-		Column.SetFlags(ETableColumnFlags::CanBeHidden | ETableColumnFlags::CanBeFiltered);
+		Column.SetFlags(ETableColumnFlags::CanBeHidden);
 
 		Column.SetHorizontalAlignment(HAlign_Left);
 		Column.SetInitialWidth(80.0f);
@@ -530,7 +530,7 @@ void FMemAllocTable::AddDefaultColumns()
 		Column.SetTitleName(LOCTEXT("FreeThreadColumnTitle", "Free Thread"));
 		Column.SetDescription(LOCTEXT("FreeThreadColumnDesc", "The thread the allocation was freed on"));
 
-		Column.SetFlags(ETableColumnFlags::CanBeHidden | ETableColumnFlags::CanBeFiltered);
+		Column.SetFlags(ETableColumnFlags::CanBeHidden);
 
 		Column.SetHorizontalAlignment(HAlign_Left);
 		Column.SetInitialWidth(80.0f);
diff --git a/Engine/Source/Developer/TraceInsights/Private/Insights/MemoryProfiler/ViewModels/MemoryFiltersValueConverters.cpp b/Engine/Source/Developer/TraceInsights/Private/Insights/MemoryProfiler/ViewModels/MemoryFiltersValueConverters.cpp
new file mode 100644
--- /dev/null
+++ b/Engine/Source/Developer/TraceInsights/Private/Insights/MemoryProfiler/ViewModels/MemoryFiltersValueConverters.cpp
@@ -0,0 +1,77 @@
+// Copyright Epic Games, Inc. All Rights Reserved.
+
+#include "MemoryFiltersValueConverters.h"
+
+// TraceServices
+#include "TraceServices/Model/Threads.h"
+#include "TraceServices/Model/TimingProfiler.h"
+
+// TraceInsights
+#include "Insights/InsightsManager.h"
+
+#define LOCTEXT_NAMESPACE "UE::Insights::MemoryProfiler::MemoryFilterConverters"
+
+namespace UE::Insights::MemoryProfiler
+{
+
+////////////////////////////////////////////////////////////////////////////////////////////////////
+
+bool FThreadFilterValueConverter::Convert(const FString& Input, int64& Output, FText& OutError) const
+{
+	if (Input.IsNumeric() && !Input.Contains(TEXT(".")))
+	{
+		Output = FCString::Atoi(*Input);
+		return true;
+	}
+
+	int32 Pos = Input.Find(TEXT(" (id:"));
+	FString ThreadName = Input;
+	if (Pos > 0)
+	{
+		ThreadName = Input.Left(Pos);
+	}
+
+	bool bIsFound = false;
+	TSharedPtr<const TraceServices::IAnalysisSession> Session = FInsightsManager::Get()->GetSession();
+	if (Session.IsValid())
+	{
+		TraceServices::FAnalysisSessionReadScope SessionReadScope(*Session.Get());
+		const TraceServices::IThreadProvider& ThreadProvider = TraceServices::ReadThreadProvider(*Session.Get());
+
+		ThreadProvider.EnumerateThreads([&ThreadName, &Output, &bIsFound](const TraceServices::FThreadInfo& ThreadInfo)
+		{
+			if (FCString::Strcmp(ThreadInfo.Name, *ThreadName) == 0)
+			{
+				Output = ThreadInfo.Id;
+				bIsFound = true;
+			}
+		});
+	}
+
+	if (!bIsFound)
+	{
+		OutError = LOCTEXT("NoThreadFound", "No thread with this name was found!");
+	}
+
+	return bIsFound;
+}
+
+////////////////////////////////////////////////////////////////////////////////////////////////////
+
+FText FThreadFilterValueConverter::GetTooltipText() const
+{
+	return LOCTEXT("FThreadConverterTooltipText", "Enter the name or the id of the thread.");
+}
+
+////////////////////////////////////////////////////////////////////////////////////////////////////
+
+FText FThreadFilterValueConverter::GetHintText() const
+{
+	return LOCTEXT("FThreadConverterHint", "Start typing or press arrow down or up to see options.");
+}
+
+////////////////////////////////////////////////////////////////////////////////////////////////////
+
+} // namespace UE::Insights::MemoryProfiler
+
+#undef LOCTEXT_NAMESPACE
diff --git a/Engine/Source/Developer/TraceInsights/Private/Insights/MemoryProfiler/ViewModels/MemoryFiltersValueConverters.h b/Engine/Source/Developer/TraceInsights/Private/Insights/MemoryProfiler/ViewModels/MemoryFiltersValueConverters.h
new file mode 100644
--- /dev/null
+++ b/Engine/Source/Developer/TraceInsights/Private/Insights/MemoryProfiler/ViewModels/MemoryFiltersValueConverters.h
@@ -0,0 +1,19 @@
+// Copyright Epic Games, Inc. All Rights Reserved.
+
+#pragma once
+
+// TraceInsightsCore
+#include "InsightsCore/Filter/ViewModels/Filters.h"
+
+namespace UE::Insights::MemoryProfiler
+{
+
+class FThreadFilterValueConverter : public IFilterValueConverter
+{
+public:
+	virtual bool Convert(const FString& Input, int64& Output, FText& OutError) const override;
+	virtual FText GetTooltipText() const override;
+	virtual FText GetHintText() const override;
+};
+
+} // namespace UE::Insights::MemoryProfiler
diff --git a/Engine/Source/Developer/TraceInsights/Private/Insights/MemoryProfiler/Widgets/SMemAllocTableTreeView.cpp b/Engine/Source/Developer/TraceInsights/Private/Insights/MemoryProfiler/Widgets/SMemAllocTableTreeView.cpp
--- a/Engine/Source/Developer/TraceInsights/Private/Insights/MemoryProfiler/Widgets/SMemAllocTableTreeView.cpp
+++ b/Engine/Source/Developer/TraceInsights/Private/Insights/MemoryProfiler/Widgets/SMemAllocTableTreeView.cpp
@@ -26,6 +26,7 @@
 #include "TraceServices/Model/MetadataProvider.h"
 #include "TraceServices/Model/Modules.h"
 #include "TraceServices/Model/Strings.h"
+#include "TraceServices/Model/Threads.h"
 
 // TraceInsightsCore
 #include "InsightsCore/Filter/ViewModels/FilterConfigurator.h"
@@ -45,6 +46,7 @@
 #include "Insights/MemoryProfiler/ViewModels/MemAllocGroupingByTag.h"
 #include "Insights/MemoryProfiler/ViewModels/MemAllocNode.h"
 #include "Insights/MemoryProfiler/ViewModels/MemAllocTable.h"
+#include "Insights/MemoryProfiler/ViewModels/MemoryFiltersValueConverters.h"
 #include "Insights/MemoryProfiler/ViewModels/MemorySharedState.h"
 #include "Insights/MemoryProfiler/Widgets/SMemoryProfilerWindow.h"
 #include "Insights/TimingProfilerCommon.h"
@@ -59,6 +61,8 @@ namespace UE::Insights::MemoryProfiler
 
 const int32 SMemAllocTableTreeView::FullCallStackIndex = 0x0000FFFFF;
 const int32 SMemAllocTableTreeView::LLMFilterIndex = 0x0000FFFFE;
+const int32 SMemAllocTableTreeView::AllocThreadFilterIndex = 0x0000FFFFD;
+const int32 SMemAllocTableTreeView::FreeThreadFilterIndex = 0x0000FFFFC;
 
 ////////////////////////////////////////////////////////////////////////////////////////////////////
 
@@ -1500,6 +1504,16 @@ void SMemAllocTableTreeView::UpdateFilterContext(const FFilterConfigurator& InFi
 		{
 			FilterContext.SetFilterData<FString>(LLMFilterIndex, MemNode.GetMemAlloc()->GetTag());
 		}
+
+		if (InFilterConfigurator.IsKeyUsed(AllocThreadFilterIndex))
+		{
+			FilterContext.SetFilterData<int64>(AllocThreadFilterIndex, MemNode.GetMemAlloc()->GetAllocThreadId());
+		}
+
+		if (InFilterConfigurator.IsKeyUsed(FreeThreadFilterIndex))
+		{
+			FilterContext.SetFilterData<int64>(FreeThreadFilterIndex, MemNode.GetMemAlloc()->GetFreeThreadId());
+		}
 	}
 }
 
@@ -1532,11 +1546,42 @@ void SMemAllocTableTreeView::InitFilterConfigurator(FFilterConfigurator& InOutFi
 		this->PopulateLLMTagSuggestionList(Text, OutSuggestions);
 	});
 	InOutFilterConfigurator.Add(LLMTagFilter);
+
+	TSharedPtr<TArray<TSharedPtr<IFilterOperator>>> AllocThreadFilterOperators = MakeShared<TArray<TSharedPtr<IFilterOperator>>>();
+	AllocThreadFilterOperators->Add(StaticCastSharedRef<IFilterOperator>(MakeShared<FFilterOperator<int64>>(EFilterOperator::Eq, TEXT("Is"), [](int64 lhs, int64 rhs) { return lhs == rhs; })));
+
+	TSharedRef<FFilterWithSuggestions> AllocThreadFilter = MakeShared<FFilterWithSuggestions>(
+		AllocThreadFilterIndex,
+		LOCTEXT("AllocThread", "Alloc Thread"),
+		LOCTEXT("AllocThreadDesc", "The allocation thread."),
+		EFilterDataType::StringInt64Pair,
+		MakeShared<FThreadFilterValueConverter>(),
+		AllocThreadFilterOperators);
+	FilterContext.AddFilterData<int64>(AllocThreadFilterIndex, 0);
+	AllocThreadFilter->SetCallback([this](const FString& Text, TArray<FString>& OutSuggestions)
+		{
+			this->PopulateThreadSuggestionList(Text, OutSuggestions);
+		});
+	InOutFilterConfigurator.Add(AllocThreadFilter);
+
+	TSharedRef<FFilterWithSuggestions> FreeThreadFilter = MakeShared<FFilterWithSuggestions>(
+		FreeThreadFilterIndex,
+		LOCTEXT("FreeThread", "Free Thread"),
+		LOCTEXT("FreeThreadDesc", "The thread the allocation was freed on."),
+		EFilterDataType::StringInt64Pair,
+		MakeShared<FThreadFilterValueConverter>(),
+		AllocThreadFilterOperators);
+	FilterContext.AddFilterData<int64>(FreeThreadFilterIndex, 0);
+	FreeThreadFilter->SetCallback([this](const FString& Text, TArray<FString>& OutSuggestions)
+		{
+			this->PopulateThreadSuggestionList(Text, OutSuggestions);
+		});
+	InOutFilterConfigurator.Add(FreeThreadFilter);
 }
 
 ////////////////////////////////////////////////////////////////////////////////////////////////////
 
-void SMemAllocTableTreeView::PopulateLLMTagSuggestionList(const FString& Text, TArray<FString>& OutSuggestions)
+void SMemAllocTableTreeView::PopulateLLMTagSuggestionList(const FString& Text, TArray<FString>& OutSuggestions) const
 {
 	const TraceServices::IAllocationsProvider* AllocationsProvider = TraceServices::ReadAllocationsProvider(*Session.Get());
 	if (!AllocationsProvider)
@@ -1566,6 +1611,39 @@ void SMemAllocTableTreeView::PopulateLLMTagSuggestionList(const FString& Text, T
 
 ////////////////////////////////////////////////////////////////////////////////////////////////////
 
+void SMemAllocTableTreeView::PopulateThreadSuggestionList(const FString& Text, TArray<FString>& OutSuggestions) const
+{
+	TraceServices::FAnalysisSessionReadScope SessionReadScope(*Session.Get());
+	const TraceServices::IThreadProvider& ThreadProvider = TraceServices::ReadThreadProvider(*Session.Get());
+
+	bool bIsInteger = Text.IsNumeric() && !Text.Contains(TEXT("."));
+	TStringBuilder<128> Builder;
+	ThreadProvider.EnumerateThreads([&Text, &OutSuggestions, &Builder, bIsInteger](const TraceServices::FThreadInfo& ThreadInfo)
+	{
+		bool bAddSuggestion = false;
+		if (bIsInteger)
+		{
+			if (ThreadInfo.Id == FCString::Atoi(*Text))
+			{
+				bAddSuggestion = true;
+			}
+		}
+		if (FCString::Stristr(ThreadInfo.Name, *Text))
+		{
+			bAddSuggestion = true;
+		}
+
+		if (bAddSuggestion)
+		{
+			Builder.Appendf(TEXT("%s (id:%d)"), ThreadInfo.Name, ThreadInfo.Id);
+			OutSuggestions.Add(Builder.ToString());
+			Builder.Reset();
+		}
+	});
+}
+
+////////////////////////////////////////////////////////////////////////////////////////////////////
+
 BEGIN_SLATE_FUNCTION_BUILD_OPTIMIZATION
 TSharedRef<SWidget> SMemAllocTableTreeView::ConstructFunctionToggleButton()
 {
diff --git a/Engine/Source/Developer/TraceInsights/Private/Insights/MemoryProfiler/Widgets/SMemAllocTableTreeView.h b/Engine/Source/Developer/TraceInsights/Private/Insights/MemoryProfiler/Widgets/SMemAllocTableTreeView.h
--- a/Engine/Source/Developer/TraceInsights/Private/Insights/MemoryProfiler/Widgets/SMemAllocTableTreeView.h
+++ b/Engine/Source/Developer/TraceInsights/Private/Insights/MemoryProfiler/Widgets/SMemAllocTableTreeView.h
@@ -145,11 +145,15 @@ private:
 	ECheckBoxState CallstackGroupingByFunction_IsChecked() const;
 
 	void InitAvailableViewPresets();
-	void PopulateLLMTagSuggestionList(const FString& Text, TArray<FString>& OutSuggestions);
+	void PopulateLLMTagSuggestionList(const FString& Text, TArray<FString>& OutSuggestions) const;
+	void PopulateThreadSuggestionList(const FString& Text, TArray<FString>& OutSuggestions) const;
 
 private:
 	const static int32 FullCallStackIndex;
 	const static int32 LLMFilterIndex;
+	const static int32 AllocThreadFilterIndex;
+	const static int32 FreeThreadFilterIndex;
+
 	int32 TabIndex = -1;
 	TSharedPtr<FMemoryRuleSpec> Rule = nullptr;
 	double TimeMarkers[4];
