[Insights]
  - Added a new public class FTableImporter in TraceInsightsCore module (with functionality from FTableImportTool).
  - Moved UntypedTable functionality from TraceInsights module to TraceInsightsCore module (and made the API public).
  - Re-enabled "Import Table..." and "Diff Tables..." in Unreal Insights Frontend.
  - Added "Diff Tables..." also in the main menu of Unreal Insights viewer mode.
  - Renamed "Session Browser" to "Unreal Insights Frontend", in the main menu of Unreal Insights viewer mode.

diff --git a/Engine/Source/Developer/TraceInsights/Private/Insights/Common/InsightsMenuBuilder.cpp b/Engine/Source/Developer/TraceInsights/Private/Insights/Common/InsightsMenuBuilder.cpp
--- a/Engine/Source/Developer/TraceInsights/Private/Insights/Common/InsightsMenuBuilder.cpp
+++ b/Engine/Source/Developer/TraceInsights/Private/Insights/Common/InsightsMenuBuilder.cpp
@@ -62,19 +62,27 @@ void FInsightsMenuBuilder::PopulateMenu(FMenuBuilder& MenuBuilder)
 {
 #if !WITH_EDITOR
 	MenuBuilder.BeginSection("Insights");
+
 	MenuBuilder.AddMenuEntry(
 		LOCTEXT("ImportTable", "Import Table..."),
 		LOCTEXT("ImportTable_ToolTip", "Import CSV or TSV data from a file to an Insights Table."),
-		FSlateIcon(FInsightsStyle::GetStyleSetName(), "Icons.ImportTable"),
+		FSlateIcon(FInsightsCoreStyle::GetStyleSetName(), "Icons.ImportTable"),
 		FUIAction(FExecuteAction::CreateLambda([] { Insights::FTableImportTool::Get()->StartImportProcess(); })));
 
+	MenuBuilder.AddMenuEntry(
+		LOCTEXT("DiffTables", "Diff Tables..."),
+		LOCTEXT("DiffTables_Tooltip", "Opens two table files in diff mode."),
+		FSlateIcon(FAppStyle::Get().GetStyleSetName(), "Icons.FolderOpen"),
+		FUIAction(FExecuteAction::CreateLambda([] { Insights::FTableImportTool::Get()->StartDiffProcess(); })));
+
 	MenuBuilder.AddSeparator();
 
 	MenuBuilder.AddMenuEntry(
-		LOCTEXT("OpenSessionBrowser", "Session Browser"),
-		LOCTEXT("OpenSessionBrowser_ToolTip", "Opens the Unreal Insights Session Browser window."),
+		LOCTEXT("OpenUnrealInsightsFrontend", "Unreal Insights Frontend"),
+		LOCTEXT("OpenUnrealInsightsFrontend_ToolTip", "Opens the Unreal Insights Frontend window."),
 		FSlateIcon(FInsightsStyle::GetStyleSetName(), "AppIcon.Small"),
 		FUIAction(FExecuteAction::CreateLambda([] { FMiscUtils::OpenUnrealInsights(); })));
+
 	MenuBuilder.AddSubMenu(
 		LOCTEXT("OpenTraceFile_SubMenu", "Open Trace File"),
 		LOCTEXT("OpenTraceFile_SubMenu_Desc", "Starts analysis for a specified trace file."),
diff --git a/Engine/Source/Developer/TraceInsights/Private/Insights/ImportTool/TableImportTool.cpp b/Engine/Source/Developer/TraceInsights/Private/Insights/ImportTool/TableImportTool.cpp
--- a/Engine/Source/Developer/TraceInsights/Private/Insights/ImportTool/TableImportTool.cpp
+++ b/Engine/Source/Developer/TraceInsights/Private/Insights/ImportTool/TableImportTool.cpp
@@ -2,31 +2,11 @@
 
 #include "TableImportTool.h"
 
-#include "DesktopPlatformModule.h"
-#include "Framework/Application/SlateApplication.h"
-#include "Framework/Docking/TabManager.h"
-#include "Framework/Docking/WorkspaceItem.h"
-#include "Logging/MessageLog.h"
-#include "MessageLogModule.h"
-#include "Misc/FileHelper.h"
-#include "SlateOptMacros.h"
-#include "Widgets/Docking/SDockTab.h"
-
-// TraceServices
-#include "TraceServices/Model/Threads.h"
-
 // TraceInsightsCore
-#include "InsightsCore/Common/Stopwatch.h"
+#include "InsightsCore/Table/ViewModels/TableImporter.h"
 
 // TraceInsights
 #include "Insights/InsightsManager.h"
-#include "Insights/InsightsStyle.h"
-#include "Insights/Log.h"
-#include "Insights/Table/ViewModels/UntypedTable.h"
-#include "Insights/Table/Widgets/SUntypedDiffTableTreeView.h"
-#include "Insights/Table/Widgets/SUntypedTableTreeView.h"
-
-#define LOCTEXT_NAMESPACE "UE::Insights::TableImportTool"
 
 namespace UE::Insights
 {
@@ -38,6 +18,7 @@ TSharedPtr<FTableImportTool> FTableImportTool::Instance;
 ////////////////////////////////////////////////////////////////////////////////////////////////////
 
 FTableImportTool::FTableImportTool()
+	: TableImporter(MakeShared<FTableImporter>(FInsightsManager::Get()->GetLogListingName()))
 {
 }
 
@@ -68,372 +49,37 @@ TSharedPtr<FTableImportTool> FTableImportTool::Get()
 
 void FTableImportTool::StartImportProcess()
 {
-	TArray<FString> Filenames;
-	bool bDialogResult = false;
-
-	IDesktopPlatform* DesktopPlatform = FDesktopPlatformModule::Get();
-	if (DesktopPlatform)
-	{
-		const FString DefaultPath = FPaths::ProjectSavedDir();
-		const FString DefaultFile = TEXT("");
-		bDialogResult = DesktopPlatform->OpenFileDialog(
-			FSlateApplication::Get().FindBestParentWindowHandleForDialogs(nullptr),
-			LOCTEXT("ImportTableTitle", "Import Table").ToString(),
-			DefaultPath,
-			DefaultFile,
-			TEXT("Comma-Separated Values (*.csv)|*.csv|Tab-Separated Values (*.tsv)|*.tsv"),
-			EFileDialogFlags::None,
-			Filenames
-		);
-	}
-
-	if (!bDialogResult)
-	{
-		return;
-	}
-
-	if (Filenames.Num() != 1)
-	{
-		FMessageLog ReportMessageLog(FInsightsManager::Get()->GetLogListingName());
-		ReportMessageLog.Error(LOCTEXT("NotSingleFile", "A single file can be imported at a time!"));
-		ReportMessageLog.Notify();
-		return;
-	}
-
-	ImportFile(Filenames[0]);
+	TableImporter->StartImportProcess();
 }
 
 ////////////////////////////////////////////////////////////////////////////////////////////////////
 
 void FTableImportTool::ImportFile(const FString& Filename)
 {
-	using namespace TraceServices;
-
-	FName TableId = GetTableID(Filename);
-	DisplayImportTable(TableId);
-
-	FTableImportService::ImportTable(Filename,
-		TableId,
-		[this](TSharedPtr<FTableImportCallbackParams> Params) {
-			this->TableImportServiceCallback(Params);
-		});
+	TableImporter->ImportFile(Filename);
 }
 
 ////////////////////////////////////////////////////////////////////////////////////////////////////
 
 void FTableImportTool::StartDiffProcess()
 {
-	TArray<FString> Filenames;
-	bool bDialogResult = false;
-
-	IDesktopPlatform* DesktopPlatform = FDesktopPlatformModule::Get();
-	if (DesktopPlatform)
-	{
-		const FString DefaultPath = FPaths::ProjectSavedDir();
-		const FString DefaultFile = TEXT("");
-		bDialogResult = DesktopPlatform->OpenFileDialog(
-			FSlateApplication::Get().FindBestParentWindowHandleForDialogs(nullptr),
-			LOCTEXT("DiffTableTitle", "Choose Tables").ToString(),
-			DefaultPath,
-			DefaultFile,
-			TEXT("Comma-Separated Values (*.csv)|*.csv|Tab-Separated Values (*.tsv)|*.tsv"),
-			EFileDialogFlags::Multiple,
-			Filenames
-		);
-	}
-
-	if (!bDialogResult)
-	{
-		return;
-	}
-
-	if (DesktopPlatform && Filenames.Num() <  2)
-	{
-		const FString DefaultPath = FPaths::ProjectSavedDir();
-		const FString DefaultFile = TEXT("");
-		bDialogResult = DesktopPlatform->OpenFileDialog(
-			FSlateApplication::Get().FindBestParentWindowHandleForDialogs(nullptr),
-			LOCTEXT("DiffTableBTitle", "Choose Table B").ToString(),
-			DefaultPath,
-			DefaultFile,
-			TEXT("Comma-Separated Values (*.csv)|*.csv|Tab-Separated Values (*.tsv)|*.tsv"),
-			EFileDialogFlags::None,
-			Filenames
-		);
-	}
-
-	if (!bDialogResult)
-	{
-		return;
-	}
-
-	if (Filenames.Num() != 2)
-	{
-		FMessageLog ReportMessageLog(FInsightsManager::Get()->GetLogListingName());
-		ReportMessageLog.Error(LOCTEXT("NotEnoughFilesSelected", "Exactly 2 files are required to be chosen!"));
-		ReportMessageLog.Notify();
-		return;
-	}
-
-	DiffFiles(Filenames[0], Filenames[1]);
+	TableImporter->StartDiffProcess();
 }
 
 ////////////////////////////////////////////////////////////////////////////////////////////////////
 
 void FTableImportTool::DiffFiles(const FString& FilenameA, const FString& FilenameB)
 {
-	using namespace TraceServices;
-
-	FName TableViewID = GetTableID(FilenameA+TEXT(",")+FilenameB);
-	DisplayDiffTable(TableViewID);
-
-	FTableImportService::ImportTable(FilenameA, TableViewID, 
-		[this, Name = FilenameA](TSharedPtr<FTableImportCallbackParams> Params)
-		{
-			ensure(Params.IsValid());
-
-			if (const FOpenImportedTableTabData* TableViewPtr = OpenTablesMap.Find(Params->TableId))
-			{
-				TSharedPtr<SUntypedDiffTableTreeView> TableView = StaticCastSharedPtr<SUntypedDiffTableTreeView>(TableViewPtr->TableTreeView);
-
-				if (Params->Result == TraceServices::ETableImportResult::ESuccess)
-				{
-					TableView->UpdateSourceTableA(Name, Params->Table);
-				}
-				else
-				{
-					FMessageLog ReportMessageLog(FInsightsManager::Get()->GetLogListingName());
-					ReportMessageLog.AddMessages(Params->Messages);
-					ReportMessageLog.Notify();
-				}
-			}
-		});
-	FTableImportService::ImportTable(FilenameB, TableViewID, 
-		[this, Name = FilenameB](TSharedPtr<FTableImportCallbackParams> Params)
-		{
-			ensure(Params.IsValid());
-
-			if (const FOpenImportedTableTabData* TableViewPtr = OpenTablesMap.Find(Params->TableId))
-			{
-				TSharedPtr<SUntypedDiffTableTreeView> TableView = StaticCastSharedPtr<SUntypedDiffTableTreeView>(TableViewPtr->TableTreeView);
-
-				if (Params->Result == TraceServices::ETableImportResult::ESuccess)
-				{
-					TableView->UpdateSourceTableB(Name, Params->Table);
-				}
-				else
-				{
-					FMessageLog ReportMessageLog(FInsightsManager::Get()->GetLogListingName());
-					ReportMessageLog.AddMessages(Params->Messages);
-					ReportMessageLog.Notify();
-				}
-			}
-		});
-}
-
-////////////////////////////////////////////////////////////////////////////////////////////////////
-
-void FTableImportTool::DisplayImportTable(FName TableViewID)
-{
-	FText ImportTableTreeViewTabDisplayName;
-	FString FileName = FPaths::GetCleanFilename(TableViewID.GetPlainNameString());
-	if (TableViewID.GetNumber())
-	{
-		ImportTableTreeViewTabDisplayName = FText::Format(LOCTEXT("TableImportTool.ImportTableTreeViewTabDisplayName", "{0}({1})"),
-														  FText::FromString(FileName),
-														  FText::AsNumber(TableViewID.GetNumber()));
-	}
-	else
-	{
-		ImportTableTreeViewTabDisplayName = FText::FromString(FileName);
-	}
-
-	FString TooltipText = FPaths::ConvertRelativePathToFull(TableViewID.GetPlainNameString());
-	FGlobalTabmanager::Get()->RegisterNomadTabSpawner(TableViewID, FOnSpawnTab::CreateSP(this, &FTableImportTool::SpawnTab_TableImportTreeView, TableViewID, ImportTableTreeViewTabDisplayName))
-		.SetDisplayName(ImportTableTreeViewTabDisplayName)
-		.SetTooltipText(FText::FromString(TooltipText))
-		.SetIcon(FSlateIcon(FInsightsStyle::GetStyleSetName(), "Icons.TableTreeView"));
-
-	FGlobalTabmanager::Get()->TryInvokeTab(TableViewID);
-}
-
-////////////////////////////////////////////////////////////////////////////////////////////////////
-
-void FTableImportTool::DisplayDiffTable(FName TableViewID)
-{
-	FText ImportTableTreeViewTabDisplayName;
-	TArray<FString> Filenames;
-	ensure(TableViewID.GetPlainNameString().ParseIntoArray(Filenames, TEXT(",")) == 2);
-
-	const FText FilenameA = FText::FromString(FPaths::GetCleanFilename(Filenames[0]));
-	const FText FilenameB = FText::FromString(FPaths::GetCleanFilename(Filenames[1]));
-
-	if (TableViewID.GetNumber())
-	{
-		ImportTableTreeViewTabDisplayName = FText::Format(
-			LOCTEXT("TableImportTool.DiffTablesNumericTreeViewTabDisplayName", "{0} vs {1} ({2})"),
-			FilenameA,
-			FilenameB,
-			FText::AsNumber(TableViewID.GetNumber()));
-	}
-	else
-	{
-		ImportTableTreeViewTabDisplayName = FText::Format(
-			LOCTEXT("TableImportTool.DiffTablesTableTreeViewTabDisplayName", "{0} vs {1}"),
-			FilenameA,
-			FilenameB);
-	}
-	const FText ImportTableTreeViewTabTooltipText = FText::Format(
-			LOCTEXT("TableImportTool.DiffTablesTableTreeViewTabTooltip", "{0} vs {1}"),
-			FilenameA,
-			FilenameB);
-
-	FGlobalTabmanager::Get()->RegisterNomadTabSpawner(TableViewID,
-		FOnSpawnTab::CreateSP(this, &FTableImportTool::SpawnTab_TableDiffTreeView, TableViewID, ImportTableTreeViewTabDisplayName))
-		.SetDisplayName(ImportTableTreeViewTabDisplayName)
-		.SetTooltipText(ImportTableTreeViewTabTooltipText)
-		.SetIcon(FSlateIcon(FInsightsStyle::GetStyleSetName(), "Icons.TableTreeView"));
-
-	FGlobalTabmanager::Get()->TryInvokeTab(TableViewID);
-}
-////////////////////////////////////////////////////////////////////////////////////////////////////
-
-void FTableImportTool::TableImportServiceCallback(TSharedPtr<TraceServices::FTableImportCallbackParams> Params)
-{
-	ensure(Params.IsValid());
-
-	const FOpenImportedTableTabData* TableViewPtr = OpenTablesMap.Find(Params->TableId);
-	if (TableViewPtr)
-	{
-		TSharedPtr<SUntypedTableTreeView> TableView = TableViewPtr->TableTreeView;
-		TableView->ClearCurrentOperationNameOverride();
-
-		if (Params->Result == TraceServices::ETableImportResult::ESuccess)
-		{
-			TableView->UpdateSourceTable(Params->Table);
-		}
-		else
-		{
-			FMessageLog ReportMessageLog(FInsightsManager::Get()->GetLogListingName());
-			ReportMessageLog.AddMessages(Params->Messages);
-			ReportMessageLog.Notify();
-		}
-	}
-}
-
-////////////////////////////////////////////////////////////////////////////////////////////////////
-
-BEGIN_SLATE_FUNCTION_BUILD_OPTIMIZATION
-TSharedRef<SDockTab> FTableImportTool::SpawnTab_TableImportTreeView(const FSpawnTabArgs& Args, FName TableViewID, FText InDisplayName)
-{
-	TSharedRef<FUntypedTable> UntypedTable = MakeShared<Insights::FUntypedTable>();
-	UntypedTable->Reset();
-	UntypedTable->SetDisplayName(InDisplayName);
-
-	TSharedPtr<SUntypedTableTreeView> TableTreeView;
-
-	const TSharedRef<SDockTab> DockTab = SNew(SDockTab)
-		.ShouldAutosize(false)
-		.TabRole(ETabRole::PanelTab)
-		[
-			SAssignNew(TableTreeView, SUntypedTableTreeView, UntypedTable)
-			.RunInAsyncMode(true)
-		];
-
-	TableTreeView->SetCurrentOperationNameOverride(LOCTEXT("TableImportMessage", "Importing Data"));
-
-	FOpenImportedTableTabData OpenTabData;
-	OpenTabData.TableTreeView = TableTreeView;
-	OpenTabData.Tab = DockTab;
-
-	OpenTablesMap.Add(TableViewID, OpenTabData);
-	DockTab->SetOnTabClosed(SDockTab::FOnTabClosedCallback::CreateSP(this, &FTableImportTool::OnTableImportTreeViewTabClosed));
-
-	return DockTab;
-}
-END_SLATE_FUNCTION_BUILD_OPTIMIZATION
-
-////////////////////////////////////////////////////////////////////////////////////////////////////
-
-BEGIN_SLATE_FUNCTION_BUILD_OPTIMIZATION
-TSharedRef<SDockTab> FTableImportTool::SpawnTab_TableDiffTreeView(const FSpawnTabArgs& Args, FName TableViewID, FText InDisplayName)
-{
-	TSharedRef<FUntypedTable> UntypedTable = MakeShared<Insights::FUntypedTable>();
-	UntypedTable->Reset();
-	UntypedTable->SetDisplayName(InDisplayName);
-
-	TSharedPtr<Insights::SUntypedDiffTableTreeView> TableTreeView;
-
-	const TSharedRef<SDockTab> DockTab = SNew(SDockTab)
-		.ShouldAutosize(false)
-		.TabRole(ETabRole::PanelTab)
-		[
-			SAssignNew(TableTreeView, SUntypedDiffTableTreeView, UntypedTable)
-			.RunInAsyncMode(true)
-		];
-
-	TableTreeView->SetCurrentOperationNameOverride(LOCTEXT("TableImportMessage", "Importing Data"));
-
-	FOpenImportedTableTabData OpenTabData;
-	OpenTabData.TableTreeView = TableTreeView;
-	OpenTabData.Tab = DockTab;
-
-	OpenTablesMap.Add(TableViewID, OpenTabData);
-	DockTab->SetOnTabClosed(SDockTab::FOnTabClosedCallback::CreateSP(this, &FTableImportTool::OnTableImportTreeViewTabClosed));
-
-	return DockTab;
-}
-END_SLATE_FUNCTION_BUILD_OPTIMIZATION
-
-////////////////////////////////////////////////////////////////////////////////////////////////////
-
-void FTableImportTool::OnTableImportTreeViewTabClosed(TSharedRef<SDockTab> TabBeingClosed)
-{
-	TSharedRef<SUntypedTableTreeView> TableTreeView = StaticCastSharedRef<SUntypedTableTreeView>(TabBeingClosed->GetContent());
-	FOpenImportedTableTabData OpenTabData;
-	OpenTabData.TableTreeView = TableTreeView;
-	OpenTabData.Tab = TabBeingClosed;
-	const FName* TableIdPtr = OpenTablesMap.FindKey(OpenTabData);
-	check(TableIdPtr);
-
-	FGlobalTabmanager::Get()->UnregisterNomadTabSpawner(*TableIdPtr);
-	TableTreeView->OnClose();
-	OpenTablesMap.Remove(*TableIdPtr);
-}
-
-////////////////////////////////////////////////////////////////////////////////////////////////////
-
-FName FTableImportTool::GetTableID(const FString& Path)
-{
-	FName TableId(Path);
-	while (OpenTablesMap.Contains(TableId))
-	{
-		TableId.SetNumber(TableId.GetNumber() + 1);
-	}
-
-	return TableId;
+	TableImporter->DiffFiles(FilenameA, FilenameB);
 }
 
 ////////////////////////////////////////////////////////////////////////////////////////////////////
 
 void FTableImportTool::OnWindowClosedEvent()
 {
-	CloseAllOpenTabs();
-}
-
-////////////////////////////////////////////////////////////////////////////////////////////////////
-
-void FTableImportTool::CloseAllOpenTabs()
-{
-	while (!OpenTablesMap.IsEmpty())
-	{
-		TSharedPtr<SDockTab> Tab = OpenTablesMap.CreateConstIterator().Value().Tab;
-		Tab->RequestCloseTab();
-	}
+	TableImporter->CloseAllOpenTabs();
 }
 
 ////////////////////////////////////////////////////////////////////////////////////////////////////
 
 } // namespace UE::Insights
-
-#undef LOCTEXT_NAMESPACE
diff --git a/Engine/Source/Developer/TraceInsights/Private/Insights/ImportTool/TableImportTool.h b/Engine/Source/Developer/TraceInsights/Private/Insights/ImportTool/TableImportTool.h
--- a/Engine/Source/Developer/TraceInsights/Private/Insights/ImportTool/TableImportTool.h
+++ b/Engine/Source/Developer/TraceInsights/Private/Insights/ImportTool/TableImportTool.h
@@ -4,38 +4,13 @@
 
 #include "CoreTypes.h"
 
-// TraceServices
-#include "TraceServices/Model/TableImport.h"
-#include "TraceServices/Containers/Tables.h"
-
 // TraceInsights
 #include "Insights/IUnrealInsightsModule.h"
 
-class SDockTab;
-class FSpawnTabArgs;
-
-namespace TraceServices
-{
-	struct FTableImportCallbackParams;
-}
-
 namespace UE::Insights
 {
 
-class SUntypedTableTreeView;
-
-////////////////////////////////////////////////////////////////////////////////////////////////////
-
-struct FOpenImportedTableTabData
-{
-	TSharedPtr<SDockTab> Tab;
-	TSharedPtr<SUntypedTableTreeView> TableTreeView;
-
-	bool operator ==(const FOpenImportedTableTabData& Other) const
-	{
-		return Tab == Other.Tab && TableTreeView == Other.TableTreeView;
-	}
-};
+class FTableImporter;
 
 ////////////////////////////////////////////////////////////////////////////////////////////////////
 
@@ -61,25 +36,12 @@ public:
 	void StartDiffProcess();
 	void DiffFiles(const FString& FilenameA, const FString& FilenameB);
 
-	void CloseAllOpenTabs();
-
-private:
-	TSharedRef<SDockTab> SpawnTab_TableImportTreeView(const FSpawnTabArgs& Args, FName TableViewID, FText InDisplayName);
-	TSharedRef<SDockTab> SpawnTab_TableDiffTreeView(const FSpawnTabArgs& Args, FName TableViewID, FText InDisplayName);
-	void OnTableImportTreeViewTabClosed(TSharedRef<SDockTab> TabBeingClosed);
-
-	void DisplayImportTable(FName TableViewID);
-	void DisplayDiffTable(FName TableViewID);
-	void TableImportServiceCallback(TSharedPtr<TraceServices::FTableImportCallbackParams> Params);
-
-	FName GetTableID(const FString& Path);
-
 private:
 	static TSharedPtr<FTableImportTool> Instance;
 
-	TMap<FName, FOpenImportedTableTabData> OpenTablesMap;
+	TSharedRef<FTableImporter> TableImporter;
 };
 
 ////////////////////////////////////////////////////////////////////////////////////////////////////
 
-} // namespace Insights
+} // namespace UE::Insights
diff --git a/Engine/Source/Developer/TraceInsights/Private/Insights/InsightsStyle.cpp b/Engine/Source/Developer/TraceInsights/Private/Insights/InsightsStyle.cpp
--- a/Engine/Source/Developer/TraceInsights/Private/Insights/InsightsStyle.cpp
+++ b/Engine/Source/Developer/TraceInsights/Private/Insights/InsightsStyle.cpp
@@ -253,7 +253,6 @@ void FInsightsStyle::FStyle::Initialize()
 	Set("Icons.Pinned", new EDITOR_IMAGE_BRUSH_SVG("Starship/Common/Pinned", Icon16x16));
 	Set("Icons.Unpinned", new EDITOR_IMAGE_BRUSH_SVG("Starship/Common/Unpinned", Icon16x16));
 
-	Set("Icons.ImportTable", new CORE_IMAGE_BRUSH_SVG("Starship/Common/Import", Icon16x16));
 	Set("Icons.SelectEventRange", new EDITOR_IMAGE_BRUSH_SVG("Starship/Common/SelectInViewport", Icon16x16));
 
 	Set("Icons.FindInstance", new CORE_IMAGE_BRUSH_SVG("Starship/Common/search", Icon16x16));
diff --git a/Engine/Source/Developer/TraceInsights/Private/Insights/LoadingProfiler/Widgets/SLoadingProfilerWindow.cpp b/Engine/Source/Developer/TraceInsights/Private/Insights/LoadingProfiler/Widgets/SLoadingProfilerWindow.cpp
--- a/Engine/Source/Developer/TraceInsights/Private/Insights/LoadingProfiler/Widgets/SLoadingProfilerWindow.cpp
+++ b/Engine/Source/Developer/TraceInsights/Private/Insights/LoadingProfiler/Widgets/SLoadingProfilerWindow.cpp
@@ -11,12 +11,15 @@
 // TraceServices
 #include "TraceServices/Model/LoadTimeProfiler.h"
 
+// TraceInsightsCore
+#include "InsightsCore/Table/ViewModels/UntypedTable.h"
+#include "InsightsCore/Table/Widgets/SUntypedTableTreeView.h"
+
 // TraceInsights
 #include "Insights/InsightsManager.h"
 #include "Insights/InsightsStyle.h"
 #include "Insights/LoadingProfiler/LoadingProfilerManager.h"
 #include "Insights/LoadingProfiler/Widgets/SLoadingProfilerToolbar.h"
-#include "Insights/Table/Widgets/SUntypedTableTreeView.h"
 #include "Insights/Widgets/STimingView.h"
 
 #define LOCTEXT_NAMESPACE "UE::Insights::LoadingProfiler"
diff --git a/Engine/Source/Developer/TraceInsights/Private/Insights/Table/ViewModels/UntypedTable.cpp b/Engine/Source/Developer/TraceInsights/Private/Insights/Table/ViewModels/UntypedTable.cpp
deleted file mode 100644
diff --git a/Engine/Source/Developer/TraceInsights/Private/Insights/Table/ViewModels/UntypedTable.h b/Engine/Source/Developer/TraceInsights/Private/Insights/Table/ViewModels/UntypedTable.h
deleted file mode 100644
diff --git a/Engine/Source/Developer/TraceInsights/Private/Insights/Table/Widgets/SUntypedDiffTableTreeView.cpp b/Engine/Source/Developer/TraceInsights/Private/Insights/Table/Widgets/SUntypedDiffTableTreeView.cpp
deleted file mode 100644
diff --git a/Engine/Source/Developer/TraceInsights/Private/Insights/Table/Widgets/SUntypedDiffTableTreeView.h b/Engine/Source/Developer/TraceInsights/Private/Insights/Table/Widgets/SUntypedDiffTableTreeView.h
deleted file mode 100644
diff --git a/Engine/Source/Developer/TraceInsights/Private/Insights/Table/Widgets/SUntypedTableTreeView.cpp b/Engine/Source/Developer/TraceInsights/Private/Insights/Table/Widgets/SUntypedTableTreeView.cpp
deleted file mode 100644
diff --git a/Engine/Source/Developer/TraceInsights/Private/Insights/Table/Widgets/SUntypedTableTreeView.h b/Engine/Source/Developer/TraceInsights/Private/Insights/Table/Widgets/SUntypedTableTreeView.h
deleted file mode 100644
diff --git a/Engine/Source/Developer/TraceInsightsCore/Private/InsightsCore/Common/InsightsCoreStyle.cpp b/Engine/Source/Developer/TraceInsightsCore/Private/InsightsCore/Common/InsightsCoreStyle.cpp
--- a/Engine/Source/Developer/TraceInsightsCore/Private/InsightsCore/Common/InsightsCoreStyle.cpp
+++ b/Engine/Source/Developer/TraceInsightsCore/Private/InsightsCore/Common/InsightsCoreStyle.cpp
@@ -179,13 +179,14 @@ void FInsightsCoreStyle::FStyle::Initialize()
 	Set("Icons.TableTreeView", new EDITOR_IMAGE_BRUSH_SVG("Starship/Common/Spreadsheet", Icon16x16));
 	Set("Icons.TableTreeView.ToolBar", new EDITOR_IMAGE_BRUSH_SVG("Starship/Common/Spreadsheet", Icon20x20));
 
+	Set("Icons.ImportTable", new CORE_IMAGE_BRUSH_SVG("Starship/Common/Import", Icon16x16));
+
 	Set("Icons.Filter.ToolBar", new CORE_IMAGE_BRUSH_SVG("Starship/Common/filter", Icon20x20));
 	//Set("Icons.Filter", new CORE_IMAGE_BRUSH_SVG("Starship/Common/filter", Icon16x16));	//-> use FAppStyle "Icons.Filter"
 	Set("Icons.FilterAddGroup", new EDITOR_IMAGE_BRUSH_SVG("Starship/Common/WorldOutliner", Icon16x16));
 	Set("Icons.ClassicFilter", new IMAGE_BRUSH_SVG("Filter", Icon16x16));
 	Set("Icons.ClassicFilterConfig", new IMAGE_BRUSH_SVG("FilterConfig", Icon16x16));
 
-
 	Set("Icons.SortBy", new CORE_IMAGE_BRUSH("Icons/Profiler/profiler_SortBy_32x", Icon16x16));
 	//Set("Icons.SortUp", new CORE_IMAGE_BRUSH_SVG("Starship/Common/SortUp", Icon16x16));		//-> use FAppStyle "Icons.SortUp"
 	//Set("Icons.SortDown", new CORE_IMAGE_BRUSH_SVG("Starship/Common/SortDown", Icon16x16));	//-> use FAppStyle "Icons.SortDown"
diff --git a/Engine/Source/Developer/TraceInsightsCore/Private/InsightsCore/Table/ViewModels/TableImporter.cpp b/Engine/Source/Developer/TraceInsightsCore/Private/InsightsCore/Table/ViewModels/TableImporter.cpp
new file mode 100644
--- /dev/null
+++ b/Engine/Source/Developer/TraceInsightsCore/Private/InsightsCore/Table/ViewModels/TableImporter.cpp
@@ -0,0 +1,410 @@
+// Copyright Epic Games, Inc. All Rights Reserved.
+
+#include "InsightsCore/Table/ViewModels/TableImporter.h"
+
+#include "DesktopPlatformModule.h"
+//include "Framework/Application/SlateApplication.h"
+#include "Framework/Docking/TabManager.h"
+//#include "Framework/Docking/WorkspaceItem.h"
+#include "Logging/MessageLog.h"
+#include "Misc/FileHelper.h"
+#include "SlateOptMacros.h"
+#include "Widgets/Docking/SDockTab.h"
+
+// TraceServices
+#include "TraceServices/Model/Threads.h"
+
+// TraceInsightsCore
+#include "InsightsCore/Common/InsightsCoreStyle.h"
+#include "InsightsCore/Common/Log.h"
+#include "InsightsCore/Common/Stopwatch.h"
+#include "InsightsCore/Table/ViewModels/UntypedTable.h"
+#include "InsightsCore/Table/Widgets/SUntypedDiffTableTreeView.h"
+#include "InsightsCore/Table/Widgets/SUntypedTableTreeView.h"
+
+#define LOCTEXT_NAMESPACE "UE::Insights::TableImporter"
+
+namespace UE::Insights
+{
+
+////////////////////////////////////////////////////////////////////////////////////////////////////
+
+FTableImporter::FTableImporter(FName InLogListingName)
+	: LogListingName(InLogListingName)
+{
+}
+
+////////////////////////////////////////////////////////////////////////////////////////////////////
+
+FTableImporter::~FTableImporter()
+{
+}
+
+////////////////////////////////////////////////////////////////////////////////////////////////////
+
+void FTableImporter::StartImportProcess()
+{
+	TArray<FString> Filenames;
+	bool bDialogResult = false;
+
+	IDesktopPlatform* DesktopPlatform = FDesktopPlatformModule::Get();
+	if (DesktopPlatform)
+	{
+		const FString DefaultPath = FPaths::ProjectSavedDir();
+		const FString DefaultFile = TEXT("");
+		bDialogResult = DesktopPlatform->OpenFileDialog(
+			FSlateApplication::Get().FindBestParentWindowHandleForDialogs(nullptr),
+			LOCTEXT("ImportTableTitle", "Import Table").ToString(),
+			DefaultPath,
+			DefaultFile,
+			TEXT("Comma-Separated Values (*.csv)|*.csv|Tab-Separated Values (*.tsv)|*.tsv"),
+			EFileDialogFlags::None,
+			Filenames
+		);
+	}
+
+	if (!bDialogResult)
+	{
+		return;
+	}
+
+	if (Filenames.Num() != 1)
+	{
+		FMessageLog ReportMessageLog(LogListingName);
+		ReportMessageLog.Error(LOCTEXT("NotSingleFile", "A single file can be imported at a time!"));
+		ReportMessageLog.Notify();
+		return;
+	}
+
+	ImportFile(Filenames[0]);
+}
+
+////////////////////////////////////////////////////////////////////////////////////////////////////
+
+void FTableImporter::ImportFile(const FString& Filename)
+{
+	using namespace TraceServices;
+
+	FName TableId = GetTableID(Filename);
+	DisplayImportTable(TableId);
+
+	FTableImportService::ImportTable(Filename,
+		TableId,
+		[this](TSharedPtr<FTableImportCallbackParams> Params) {
+			this->TableImportServiceCallback(Params);
+		});
+}
+
+////////////////////////////////////////////////////////////////////////////////////////////////////
+
+void FTableImporter::StartDiffProcess()
+{
+	TArray<FString> Filenames;
+	bool bDialogResult = false;
+
+	IDesktopPlatform* DesktopPlatform = FDesktopPlatformModule::Get();
+	if (DesktopPlatform)
+	{
+		const FString DefaultPath = FPaths::ProjectSavedDir();
+		const FString DefaultFile = TEXT("");
+		bDialogResult = DesktopPlatform->OpenFileDialog(
+			FSlateApplication::Get().FindBestParentWindowHandleForDialogs(nullptr),
+			LOCTEXT("DiffTableTitle", "Choose Tables").ToString(),
+			DefaultPath,
+			DefaultFile,
+			TEXT("Comma-Separated Values (*.csv)|*.csv|Tab-Separated Values (*.tsv)|*.tsv"),
+			EFileDialogFlags::Multiple,
+			Filenames
+		);
+	}
+
+	if (!bDialogResult)
+	{
+		return;
+	}
+
+	if (DesktopPlatform && Filenames.Num() <  2)
+	{
+		const FString DefaultPath = FPaths::ProjectSavedDir();
+		const FString DefaultFile = TEXT("");
+		bDialogResult = DesktopPlatform->OpenFileDialog(
+			FSlateApplication::Get().FindBestParentWindowHandleForDialogs(nullptr),
+			LOCTEXT("DiffTableBTitle", "Choose Table B").ToString(),
+			DefaultPath,
+			DefaultFile,
+			TEXT("Comma-Separated Values (*.csv)|*.csv|Tab-Separated Values (*.tsv)|*.tsv"),
+			EFileDialogFlags::None,
+			Filenames
+		);
+	}
+
+	if (!bDialogResult)
+	{
+		return;
+	}
+
+	if (Filenames.Num() != 2)
+	{
+		FMessageLog ReportMessageLog(LogListingName);
+		ReportMessageLog.Error(LOCTEXT("NotEnoughFilesSelected", "Exactly 2 files are required to be chosen!"));
+		ReportMessageLog.Notify();
+		return;
+	}
+
+	DiffFiles(Filenames[0], Filenames[1]);
+}
+
+////////////////////////////////////////////////////////////////////////////////////////////////////
+
+void FTableImporter::DiffFiles(const FString& FilenameA, const FString& FilenameB)
+{
+	using namespace TraceServices;
+
+	FName TableViewID = GetTableID(FilenameA+TEXT(",")+FilenameB);
+	DisplayDiffTable(TableViewID);
+
+	FTableImportService::ImportTable(FilenameA, TableViewID,
+		[this, Name = FilenameA](TSharedPtr<FTableImportCallbackParams> Params)
+		{
+			ensure(Params.IsValid());
+
+			if (const FOpenTableTabData* TableViewPtr = OpenTablesMap.Find(Params->TableId))
+			{
+				TSharedPtr<SUntypedDiffTableTreeView> TableView = StaticCastSharedPtr<SUntypedDiffTableTreeView>(TableViewPtr->TableTreeView);
+
+				if (Params->Result == TraceServices::ETableImportResult::ESuccess)
+				{
+					TableView->UpdateSourceTableA(Name, Params->Table);
+				}
+				else
+				{
+					FMessageLog ReportMessageLog(LogListingName);
+					ReportMessageLog.AddMessages(Params->Messages);
+					ReportMessageLog.Notify();
+				}
+			}
+		});
+
+	FTableImportService::ImportTable(FilenameB, TableViewID,
+		[this, Name = FilenameB](TSharedPtr<FTableImportCallbackParams> Params)
+		{
+			ensure(Params.IsValid());
+
+			if (const FOpenTableTabData* TableViewPtr = OpenTablesMap.Find(Params->TableId))
+			{
+				TSharedPtr<SUntypedDiffTableTreeView> TableView = StaticCastSharedPtr<SUntypedDiffTableTreeView>(TableViewPtr->TableTreeView);
+
+				if (Params->Result == TraceServices::ETableImportResult::ESuccess)
+				{
+					TableView->UpdateSourceTableB(Name, Params->Table);
+				}
+				else
+				{
+					FMessageLog ReportMessageLog(LogListingName);
+					ReportMessageLog.AddMessages(Params->Messages);
+					ReportMessageLog.Notify();
+				}
+			}
+		});
+}
+
+////////////////////////////////////////////////////////////////////////////////////////////////////
+
+void FTableImporter::DisplayImportTable(FName TableViewID)
+{
+	FText ImportTableTreeViewTabDisplayName;
+	FString FileName = FPaths::GetCleanFilename(TableViewID.GetPlainNameString());
+	if (TableViewID.GetNumber())
+	{
+		ImportTableTreeViewTabDisplayName = FText::Format(LOCTEXT("TableImportTool.ImportTableTreeViewTabDisplayName", "{0}({1})"),
+														  FText::FromString(FileName),
+														  FText::AsNumber(TableViewID.GetNumber()));
+	}
+	else
+	{
+		ImportTableTreeViewTabDisplayName = FText::FromString(FileName);
+	}
+
+	FString TooltipText = FPaths::ConvertRelativePathToFull(TableViewID.GetPlainNameString());
+	FGlobalTabmanager::Get()->RegisterNomadTabSpawner(TableViewID, FOnSpawnTab::CreateSP(this, &FTableImporter::SpawnTab_TableImportTreeView, TableViewID, ImportTableTreeViewTabDisplayName))
+		.SetDisplayName(ImportTableTreeViewTabDisplayName)
+		.SetTooltipText(FText::FromString(TooltipText))
+		.SetIcon(FSlateIcon(FInsightsCoreStyle::GetStyleSetName(), "Icons.TableTreeView"));
+
+	FGlobalTabmanager::Get()->TryInvokeTab(TableViewID);
+}
+
+////////////////////////////////////////////////////////////////////////////////////////////////////
+
+void FTableImporter::DisplayDiffTable(FName TableViewID)
+{
+	FText ImportTableTreeViewTabDisplayName;
+	TArray<FString> Filenames;
+	ensure(TableViewID.GetPlainNameString().ParseIntoArray(Filenames, TEXT(",")) == 2);
+
+	const FText FilenameA = FText::FromString(FPaths::GetCleanFilename(Filenames[0]));
+	const FText FilenameB = FText::FromString(FPaths::GetCleanFilename(Filenames[1]));
+
+	if (TableViewID.GetNumber())
+	{
+		ImportTableTreeViewTabDisplayName = FText::Format(
+			LOCTEXT("TableImportTool.DiffTablesNumericTreeViewTabDisplayName", "{0} vs {1} ({2})"),
+			FilenameA,
+			FilenameB,
+			FText::AsNumber(TableViewID.GetNumber()));
+	}
+	else
+	{
+		ImportTableTreeViewTabDisplayName = FText::Format(
+			LOCTEXT("TableImportTool.DiffTablesTableTreeViewTabDisplayName", "{0} vs {1}"),
+			FilenameA,
+			FilenameB);
+	}
+	const FText ImportTableTreeViewTabTooltipText = FText::Format(
+			LOCTEXT("TableImportTool.DiffTablesTableTreeViewTabTooltip", "{0} vs {1}"),
+			FilenameA,
+			FilenameB);
+
+	FGlobalTabmanager::Get()->RegisterNomadTabSpawner(TableViewID,
+		FOnSpawnTab::CreateSP(this, &FTableImporter::SpawnTab_TableDiffTreeView, TableViewID, ImportTableTreeViewTabDisplayName))
+		.SetDisplayName(ImportTableTreeViewTabDisplayName)
+		.SetTooltipText(ImportTableTreeViewTabTooltipText)
+		.SetIcon(FSlateIcon(FInsightsCoreStyle::GetStyleSetName(), "Icons.TableTreeView"));
+
+	FGlobalTabmanager::Get()->TryInvokeTab(TableViewID);
+}
+
+////////////////////////////////////////////////////////////////////////////////////////////////////
+
+void FTableImporter::TableImportServiceCallback(TSharedPtr<TraceServices::FTableImportCallbackParams> Params)
+{
+	ensure(Params.IsValid());
+
+	const FOpenTableTabData* TableViewPtr = OpenTablesMap.Find(Params->TableId);
+	if (TableViewPtr)
+	{
+		TSharedPtr<SUntypedTableTreeView> TableView = TableViewPtr->TableTreeView;
+		TableView->ClearCurrentOperationNameOverride();
+
+		if (Params->Result == TraceServices::ETableImportResult::ESuccess)
+		{
+			TableView->UpdateSourceTable(Params->Table);
+		}
+		else
+		{
+			FMessageLog ReportMessageLog(LogListingName);
+			ReportMessageLog.AddMessages(Params->Messages);
+			ReportMessageLog.Notify();
+		}
+	}
+}
+
+////////////////////////////////////////////////////////////////////////////////////////////////////
+
+BEGIN_SLATE_FUNCTION_BUILD_OPTIMIZATION
+TSharedRef<SDockTab> FTableImporter::SpawnTab_TableImportTreeView(const FSpawnTabArgs& Args, FName TableViewID, FText InDisplayName)
+{
+	TSharedRef<FUntypedTable> UntypedTable = MakeShared<Insights::FUntypedTable>();
+	UntypedTable->Reset();
+	UntypedTable->SetDisplayName(InDisplayName);
+
+	TSharedPtr<SUntypedTableTreeView> TableTreeView;
+
+	const TSharedRef<SDockTab> DockTab = SNew(SDockTab)
+		.ShouldAutosize(false)
+		.TabRole(ETabRole::PanelTab)
+		[
+			SAssignNew(TableTreeView, SUntypedTableTreeView, UntypedTable)
+			.RunInAsyncMode(true)
+		];
+
+	TableTreeView->SetCurrentOperationNameOverride(LOCTEXT("TableImportMessage", "Importing Data"));
+
+	FOpenTableTabData OpenTabData;
+	OpenTabData.TableTreeView = TableTreeView;
+	OpenTabData.Tab = DockTab;
+
+	OpenTablesMap.Add(TableViewID, OpenTabData);
+	DockTab->SetOnTabClosed(SDockTab::FOnTabClosedCallback::CreateSP(this, &FTableImporter::OnTableImportTreeViewTabClosed));
+
+	return DockTab;
+}
+END_SLATE_FUNCTION_BUILD_OPTIMIZATION
+
+////////////////////////////////////////////////////////////////////////////////////////////////////
+
+BEGIN_SLATE_FUNCTION_BUILD_OPTIMIZATION
+TSharedRef<SDockTab> FTableImporter::SpawnTab_TableDiffTreeView(const FSpawnTabArgs& Args, FName TableViewID, FText InDisplayName)
+{
+	TSharedRef<FUntypedTable> UntypedTable = MakeShared<Insights::FUntypedTable>();
+	UntypedTable->Reset();
+	UntypedTable->SetDisplayName(InDisplayName);
+
+	TSharedPtr<Insights::SUntypedDiffTableTreeView> TableTreeView;
+
+	const TSharedRef<SDockTab> DockTab = SNew(SDockTab)
+		.ShouldAutosize(false)
+		.TabRole(ETabRole::PanelTab)
+		[
+			SAssignNew(TableTreeView, SUntypedDiffTableTreeView, UntypedTable)
+			.RunInAsyncMode(true)
+		];
+
+	TableTreeView->SetCurrentOperationNameOverride(LOCTEXT("TableImportMessage", "Importing Data"));
+
+	FOpenTableTabData OpenTabData;
+	OpenTabData.TableTreeView = TableTreeView;
+	OpenTabData.Tab = DockTab;
+
+	OpenTablesMap.Add(TableViewID, OpenTabData);
+	DockTab->SetOnTabClosed(SDockTab::FOnTabClosedCallback::CreateSP(this, &FTableImporter::OnTableImportTreeViewTabClosed));
+
+	return DockTab;
+}
+END_SLATE_FUNCTION_BUILD_OPTIMIZATION
+
+////////////////////////////////////////////////////////////////////////////////////////////////////
+
+void FTableImporter::OnTableImportTreeViewTabClosed(TSharedRef<SDockTab> TabBeingClosed)
+{
+	TSharedRef<SUntypedTableTreeView> TableTreeView = StaticCastSharedRef<SUntypedTableTreeView>(TabBeingClosed->GetContent());
+	FOpenTableTabData OpenTabData;
+	OpenTabData.TableTreeView = TableTreeView;
+	OpenTabData.Tab = TabBeingClosed;
+	const FName* TableIdPtr = OpenTablesMap.FindKey(OpenTabData);
+	check(TableIdPtr);
+
+	FGlobalTabmanager::Get()->UnregisterNomadTabSpawner(*TableIdPtr);
+	TableTreeView->OnClose();
+	OpenTablesMap.Remove(*TableIdPtr);
+}
+
+////////////////////////////////////////////////////////////////////////////////////////////////////
+
+FName FTableImporter::GetTableID(const FString& Path)
+{
+	FName TableId(Path);
+	while (OpenTablesMap.Contains(TableId))
+	{
+		TableId.SetNumber(TableId.GetNumber() + 1);
+	}
+
+	return TableId;
+}
+
+////////////////////////////////////////////////////////////////////////////////////////////////////
+
+void FTableImporter::CloseAllOpenTabs()
+{
+	while (!OpenTablesMap.IsEmpty())
+	{
+		TSharedPtr<SDockTab> Tab = OpenTablesMap.CreateConstIterator().Value().Tab;
+		Tab->RequestCloseTab();
+	}
+}
+
+////////////////////////////////////////////////////////////////////////////////////////////////////
+
+} // namespace UE::Insights
+
+#undef LOCTEXT_NAMESPACE
diff --git a/Engine/Source/Developer/TraceInsightsCore/Private/InsightsCore/Table/ViewModels/UntypedTable.cpp b/Engine/Source/Developer/TraceInsightsCore/Private/InsightsCore/Table/ViewModels/UntypedTable.cpp
new file mode 100644
--- /dev/null
+++ b/Engine/Source/Developer/TraceInsightsCore/Private/InsightsCore/Table/ViewModels/UntypedTable.cpp
@@ -0,0 +1,332 @@
+// Copyright Epic Games, Inc. All Rights Reserved.
+
+#include "InsightsCore/Table/ViewModels/UntypedTable.h"
+
+// TraceServices
+#include "TraceServices/Containers/Tables.h"
+
+// TraceInsightsCore
+#include "InsightsCore/Table/ViewModels/TableCellValueFormatter.h"
+#include "InsightsCore/Table/ViewModels/TableCellValueGetter.h"
+#include "InsightsCore/Table/ViewModels/TableCellValueSorter.h"
+#include "InsightsCore/Table/ViewModels/TableColumn.h"
+#include "InsightsCore/Table/ViewModels/TableTreeNode.h"
+
+#define LOCTEXT_NAMESPACE "UE::Insights::FUntypedTable"
+
+namespace UE::Insights
+{
+
+////////////////////////////////////////////////////////////////////////////////////////////////////
+// FUntypedTableTreeNodeValueGetter
+////////////////////////////////////////////////////////////////////////////////////////////////////
+
+class FUntypedTableTreeNodeValueGetter : public FTableCellValueGetter
+{
+public:
+	FUntypedTableTreeNodeValueGetter(ETableCellDataType InDataType) : FTableCellValueGetter(), DataType(InDataType) {}
+
+	virtual const TOptional<FTableCellValue> GetValue(const FTableColumn& Column, const FBaseTreeNode& Node) const
+	{
+		ensure(Node.Is<FTableTreeNode>());
+		const FTableTreeNode& TableTreeNode = static_cast<const FTableTreeNode&>(Node);
+
+		if (!Node.IsGroup()) // Table Row Node
+		{
+			const TSharedPtr<FTable> TablePtr = Column.GetParentTable().Pin();
+			const TSharedPtr<FUntypedTable> UntypedTablePtr = StaticCastSharedPtr<FUntypedTable>(TablePtr);
+			if (UntypedTablePtr.IsValid())
+			{
+				TSharedPtr<TraceServices::IUntypedTableReader> Reader = UntypedTablePtr->GetTableReader();
+				if (Reader.IsValid() && TableTreeNode.GetRowId().HasValidIndex())
+				{
+					Reader->SetRowIndex(TableTreeNode.GetRowId().RowIndex);
+					const int32 ColumnIndex = Column.GetIndex();
+					switch (DataType)
+					{
+						case ETableCellDataType::Bool:    return TOptional<FTableCellValue>(FTableCellValue(Reader->GetValueBool(ColumnIndex)));
+						case ETableCellDataType::Int64:   return TOptional<FTableCellValue>(FTableCellValue(Reader->GetValueInt(ColumnIndex)));
+						case ETableCellDataType::Float:   return TOptional<FTableCellValue>(FTableCellValue(Reader->GetValueFloat(ColumnIndex)));
+						case ETableCellDataType::Double:  return TOptional<FTableCellValue>(FTableCellValue(Reader->GetValueDouble(ColumnIndex)));
+						case ETableCellDataType::CString: return TOptional<FTableCellValue>(FTableCellValue(Reader->GetValueCString(ColumnIndex)));
+					}
+				}
+			}
+		}
+		else // Aggregated Group Node
+		{
+			if (Column.GetAggregation() != ETableColumnAggregation::None)
+			{
+				const FTableCellValue* ValuePtr = TableTreeNode.FindAggregatedValue(Column.GetId());
+				if (ValuePtr != nullptr)
+				{
+					ensure(ValuePtr->DataType == DataType);
+					return TOptional<FTableCellValue>(*ValuePtr);
+				}
+			}
+		}
+
+		return TOptional<FTableCellValue>();
+	}
+
+private:
+	ETableCellDataType DataType;
+};
+
+////////////////////////////////////////////////////////////////////////////////////////////////////
+// FUntypedTable
+////////////////////////////////////////////////////////////////////////////////////////////////////
+
+FUntypedTable::FUntypedTable()
+	: SourceTable()
+	, TableReader()
+{
+}
+
+////////////////////////////////////////////////////////////////////////////////////////////////////
+
+FUntypedTable::~FUntypedTable()
+{
+	Reset();
+}
+
+////////////////////////////////////////////////////////////////////////////////////////////////////
+
+void FUntypedTable::Reset()
+{
+	SourceTable.Reset();
+	TableReader.Reset();
+
+	FTable::Reset();
+}
+
+////////////////////////////////////////////////////////////////////////////////////////////////////
+
+bool AreTableLayoutsEqual(const TraceServices::ITableLayout& TableLayoutA, const TraceServices::ITableLayout& TableLayoutB)
+{
+	if (TableLayoutA.GetColumnCount() != TableLayoutB.GetColumnCount())
+	{
+		return false;
+	}
+
+	int32 ColumnCount = static_cast<int32>(TableLayoutA.GetColumnCount());
+	for (int32 ColumnIndex = 0; ColumnIndex < ColumnCount; ++ColumnIndex)
+	{
+		if (TableLayoutA.GetColumnType(ColumnIndex) != TableLayoutB.GetColumnType(ColumnIndex))
+		{
+			return false;
+		}
+		if (FCString::Strcmp(TableLayoutA.GetColumnName(ColumnIndex), TableLayoutB.GetColumnName(ColumnIndex)) != 0)
+		{
+			return false;
+		}
+	}
+
+	return true;
+}
+
+////////////////////////////////////////////////////////////////////////////////////////////////////
+
+bool FUntypedTable::UpdateSourceTable(TSharedPtr<TraceServices::IUntypedTable> InSourceTable)
+{
+	bool bTableLayoutChanged;
+
+	if (InSourceTable.IsValid())
+	{
+		bTableLayoutChanged = !SourceTable.IsValid() || !AreTableLayoutsEqual(InSourceTable->GetLayout(), SourceTable->GetLayout());
+		SourceTable = InSourceTable;
+		TableReader = MakeShareable(SourceTable->CreateReader());
+	}
+	else
+	{
+		bTableLayoutChanged = SourceTable.IsValid();
+		SourceTable.Reset();
+		TableReader.Reset();
+	}
+
+	if (bTableLayoutChanged)
+	{
+		ResetColumns();
+		if (SourceTable.IsValid())
+		{
+			CreateColumns(SourceTable->GetLayout());
+		}
+	}
+
+	return bTableLayoutChanged;
+}
+
+////////////////////////////////////////////////////////////////////////////////////////////////////
+
+void FUntypedTable::CreateColumns(const TraceServices::ITableLayout& TableLayout)
+{
+	ensure(GetColumnCount() == 0);
+	const int32 ColumnCount = static_cast<int32>(TableLayout.GetColumnCount());
+
+	//////////////////////////////////////////////////
+	// Hierarchy Column
+
+	int32 HierarchyColumnIndex = -1;
+	const TCHAR* HierarchyColumnName = nullptr;
+
+	// Look for first string column.
+	//for (int32 ColumnIndex = 0; ColumnIndex < ColumnCount; ++ColumnIndex)
+	//{
+	//	TraceServices::ETableColumnType ColumnType = TableLayout.GetColumnType(ColumnIndex);
+	//	if (ColumnType == TraceServices::TableColumnType_CString)
+	//	{
+	//		HierarchyColumnIndex = ColumnIndex;
+	//		HierarchyColumnName = TableLayout.GetColumnName(ColumnIndex);
+	//		break;
+	//	}
+	//}
+
+	AddHierarchyColumn(HierarchyColumnIndex, HierarchyColumnName);
+
+	//////////////////////////////////////////////////
+
+	for (int32 ColumnIndex = 0; ColumnIndex < ColumnCount; ++ColumnIndex)
+	{
+		TraceServices::ETableColumnType ColumnType = TableLayout.GetColumnType(ColumnIndex);
+		uint32 ColumnDisplayHintFlags = TableLayout.GetColumnDisplayHintFlags(ColumnIndex);
+		const TCHAR* ColumnName = TableLayout.GetColumnName(ColumnIndex);
+
+		TSharedRef<FTableColumn> ColumnRef = MakeShared<FTableColumn>(FName(ColumnName));
+		FTableColumn& Column = *ColumnRef;
+
+		const FString ColumnNameStr(ColumnName);
+		const FText ColumnNameText = FText::FromString(ColumnNameStr);
+
+		ETableColumnFlags ColumnFlags = ETableColumnFlags::CanBeFiltered | ETableColumnFlags::CanBeHidden;
+		if (ColumnIndex != HierarchyColumnIndex)
+		{
+			ColumnFlags |= ETableColumnFlags::ShouldBeVisible;
+		}
+
+		EHorizontalAlignment HorizontalAlignment = HAlign_Left;
+		float InitialColumnWidth = 60.0f;
+
+		ETableColumnAggregation Aggregation = ETableColumnAggregation::None;
+
+		TSharedPtr<ITableCellValueFormatter> FormatterPtr;
+		TSharedPtr<ITableCellValueSorter> SorterPtr;
+		EColumnSortMode::Type InitialSortMode = EColumnSortMode::Ascending;
+
+		switch (ColumnType)
+		{
+		case TraceServices::TableColumnType_Bool:
+			Column.SetDataType(ETableCellDataType::Bool);
+			HorizontalAlignment = HAlign_Right;
+			InitialColumnWidth = 40.0f;
+			//TODO: if (Hint == AsOnOff)
+			//else // if (Hint == AsTrueFalse)
+			FormatterPtr = MakeShared<FBoolValueFormatterAsTrueFalse>();
+			SorterPtr = MakeShared<FSorterByBoolValue>(ColumnRef);
+			InitialSortMode = EColumnSortMode::Ascending;
+			break;
+
+		case TraceServices::TableColumnType_Int:
+			Column.SetDataType(ETableCellDataType::Int64);
+			if (ColumnDisplayHintFlags & TraceServices::TableColumnDisplayHint_Summable)
+			{
+				Aggregation = ETableColumnAggregation::Sum;
+			}
+			HorizontalAlignment = HAlign_Right;
+			InitialColumnWidth = 60.0f;
+			if (ColumnDisplayHintFlags & TraceServices::TableColumnDisplayHint_Memory)
+			{
+				FormatterPtr = MakeShared<FInt64ValueFormatterAsMemory>();
+			}
+			else
+			{
+				FormatterPtr = MakeShared<FInt64ValueFormatterAsNumber>();
+			}
+			SorterPtr = MakeShared<FSorterByInt64Value>(ColumnRef);
+			InitialSortMode = EColumnSortMode::Descending;
+			break;
+
+		case TraceServices::TableColumnType_Float:
+			Column.SetDataType(ETableCellDataType::Float);
+			if (ColumnDisplayHintFlags & TraceServices::TableColumnDisplayHint_Summable)
+			{
+				Aggregation = ETableColumnAggregation::Sum;
+			}
+			HorizontalAlignment = HAlign_Right;
+			InitialColumnWidth = 60.0f;
+			if (ColumnDisplayHintFlags & TraceServices::TableColumnDisplayHint_Time)
+			{
+				FormatterPtr = MakeShared<FFloatValueFormatterAsTimeAuto>();
+				InitialSortMode = EColumnSortMode::Ascending;
+			}
+			else
+			{
+				FormatterPtr = MakeShared<FFloatValueFormatterAsNumber>();
+				InitialSortMode = EColumnSortMode::Descending;
+			}
+			SorterPtr = MakeShared<FSorterByFloatValue>(ColumnRef);
+			break;
+
+		case TraceServices::TableColumnType_Double:
+			Column.SetDataType(ETableCellDataType::Double);
+			if (ColumnDisplayHintFlags & TraceServices::TableColumnDisplayHint_Summable)
+			{
+				Aggregation = ETableColumnAggregation::Sum;
+			}
+			HorizontalAlignment = HAlign_Right;
+			InitialColumnWidth = 80.0f;
+			if (ColumnDisplayHintFlags & TraceServices::TableColumnDisplayHint_Time)
+			{
+				FormatterPtr = MakeShared<FDoubleValueFormatterAsTimeAuto>();
+				InitialSortMode = EColumnSortMode::Ascending;
+			}
+			else
+			{
+				FormatterPtr = MakeShared<FDoubleValueFormatterAsNumber>();
+				InitialSortMode = EColumnSortMode::Descending;
+			}
+			SorterPtr = MakeShared<FSorterByDoubleValue>(ColumnRef);
+			break;
+
+		case TraceServices::TableColumnType_CString:
+			Column.SetDataType(ETableCellDataType::CString);
+			HorizontalAlignment = HAlign_Left;
+			InitialColumnWidth = FMath::Max(120.0f, 6.0f * static_cast<float>(ColumnNameStr.Len()));
+			FormatterPtr = MakeShared<FCStringValueFormatterAsText>();
+			SorterPtr = MakeShared<FSorterByCStringValue>(ColumnRef);
+			InitialSortMode = EColumnSortMode::Ascending;
+			break;
+		}
+
+		Column.SetIndex(ColumnIndex);
+
+		Column.SetShortName(ColumnNameText);
+		Column.SetTitleName(ColumnNameText);
+
+		//TODO: Column.SetDescription(...);
+
+		Column.SetFlags(ColumnFlags);
+
+		Column.SetHorizontalAlignment(HorizontalAlignment);
+		Column.SetInitialWidth(InitialColumnWidth);
+
+		Column.SetAggregation(Aggregation);
+
+		Column.SetValueGetter(MakeShared<FUntypedTableTreeNodeValueGetter>(Column.GetDataType()));
+
+		if (FormatterPtr.IsValid())
+		{
+			Column.SetValueFormatter(FormatterPtr.ToSharedRef());
+		}
+
+		Column.SetValueSorter(SorterPtr);
+		Column.SetInitialSortMode(InitialSortMode);
+
+		AddColumn(ColumnRef);
+	}
+}
+
+////////////////////////////////////////////////////////////////////////////////////////////////////
+
+} // namespace UE::Insights
+
+#undef LOCTEXT_NAMESPACE
diff --git a/Engine/Source/Developer/TraceInsightsCore/Private/InsightsCore/Table/Widgets/SUntypedDiffTableTreeView.cpp b/Engine/Source/Developer/TraceInsightsCore/Private/InsightsCore/Table/Widgets/SUntypedDiffTableTreeView.cpp
new file mode 100644
--- /dev/null
+++ b/Engine/Source/Developer/TraceInsightsCore/Private/InsightsCore/Table/Widgets/SUntypedDiffTableTreeView.cpp
@@ -0,0 +1,120 @@
+// Copyright Epic Games, Inc. All Rights Reserved.
+
+#include "InsightsCore/Table/Widgets/SUntypedDiffTableTreeView.h"
+
+#include "SlateOptMacros.h"
+#include "Logging/MessageLog.h"
+
+// TraceServices
+#include "TraceServices/Model/TableMerge.h"
+
+#define LOCTEXT_NAMESPACE "UE::Insights::SUntypedDiffTableTreeView"
+
+namespace UE::Insights
+{
+
+////////////////////////////////////////////////////////////////////////////////////////////////////
+
+void SUntypedDiffTableTreeView::UpdateSourceTableA(const FString& Name, TSharedPtr<TraceServices::IUntypedTable> SourceTable)
+{
+	TableNameA = Name;
+	SourceTableA = SourceTable;
+
+	if (SourceTableA.IsValid() && SourceTableB.IsValid())
+	{
+		RequestMergeTables();
+	}
+}
+
+////////////////////////////////////////////////////////////////////////////////////////////////////
+
+void SUntypedDiffTableTreeView::UpdateSourceTableB(const FString& Name, TSharedPtr<TraceServices::IUntypedTable> SourceTable)
+{
+	TableNameB = Name;
+	SourceTableB = SourceTable;
+
+	if (SourceTableA.IsValid() && SourceTableB.IsValid())
+	{
+		RequestMergeTables();
+	}
+}
+
+////////////////////////////////////////////////////////////////////////////////////////////////////
+
+FReply SUntypedDiffTableTreeView::SwapTables_OnClicked()
+{
+	Swap(TableNameA, TableNameB);
+	Swap(SourceTableA, SourceTableB);
+
+	if (SourceTableA.IsValid() && SourceTableB.IsValid())
+	{
+		RequestMergeTables();
+	}
+
+	return FReply::Handled();
+}
+
+////////////////////////////////////////////////////////////////////////////////////////////////////
+
+FText SUntypedDiffTableTreeView::GetSwapButtonText() const
+{
+	const FText CleanA = FText::FromString(FPaths::GetCleanFilename(TableNameA));
+	const FText CleanB = FText::FromString(FPaths::GetCleanFilename(TableNameB));
+	return FText::Format(LOCTEXT("ToolBar_SwapButtonText", "SWAP (A: {0} - B: {1})"), CleanA, CleanB);
+}
+
+////////////////////////////////////////////////////////////////////////////////////////////////////
+
+BEGIN_SLATE_FUNCTION_BUILD_OPTIMIZATION
+TSharedPtr<SWidget> SUntypedDiffTableTreeView::ConstructToolbar()
+{
+	return
+		SNew(SHorizontalBox)
+		+ SHorizontalBox::Slot()
+		[
+			SNew(SButton)
+			.ButtonStyle(&FAppStyle::Get().GetWidgetStyle<FButtonStyle>("PrimaryButton"))
+			.ForegroundColor(FSlateColor::UseStyle())
+			.HAlign(HAlign_Center)
+			.Text(this, &SUntypedDiffTableTreeView::GetSwapButtonText)
+			.OnClicked(this, &SUntypedDiffTableTreeView::SwapTables_OnClicked)
+		];
+}
+END_SLATE_FUNCTION_BUILD_OPTIMIZATION
+
+////////////////////////////////////////////////////////////////////////////////////////////////////
+
+void SUntypedDiffTableTreeView::RequestMergeTables()
+{
+	SetCurrentOperationNameOverride(LOCTEXT("TableDiffMessage", "Merging Tables"));
+
+	// Clear current table
+	if (TreeView.IsValid())
+	{
+		TreeView->SetTreeItemsSource(&DummyGroupNodes);
+		TreeView->RequestTreeRefresh();
+	}
+
+	TraceServices::FTableMergeService::MergeTables(SourceTableA, SourceTableB,
+		[this](TSharedPtr<TraceServices::FTableDiffCallbackParams> Params)
+		{
+			ClearCurrentOperationNameOverride();
+			if (Params->Result == TraceServices::ETableDiffResult::ESuccess)
+			{
+				this->UpdateSourceTable(Params->Table);
+			}
+			else
+			{
+				FName LogListingName(TEXT("Insights"));
+				FMessageLog ReportMessageLog(LogListingName);
+				ReportMessageLog.AddMessages(Params->Messages);
+				ReportMessageLog.Notify();
+			}
+		});
+}
+
+////////////////////////////////////////////////////////////////////////////////////////////////////
+
+} // namespace UE::Insights
+
+#undef LOCTEXT_NAMESPACE
diff --git a/Engine/Source/Developer/TraceInsightsCore/Private/InsightsCore/Table/Widgets/SUntypedTableTreeView.cpp b/Engine/Source/Developer/TraceInsightsCore/Private/InsightsCore/Table/Widgets/SUntypedTableTreeView.cpp
new file mode 100644
--- /dev/null
+++ b/Engine/Source/Developer/TraceInsightsCore/Private/InsightsCore/Table/Widgets/SUntypedTableTreeView.cpp
@@ -0,0 +1,195 @@
+// Copyright Epic Games, Inc. All Rights Reserved.
+
+#include "InsightsCore/Table/Widgets/SUntypedTableTreeView.h"
+
+// TraceServices
+#include "TraceServices/Containers/Tables.h"
+
+// TraceInsightsCore
+#include "InsightsCore/Common/Log.h"
+#include "InsightsCore/Common/Stopwatch.h"
+#include "InsightsCore/Table/ViewModels/UntypedTable.h"
+
+#define LOCTEXT_NAMESPACE "UE::Insights::SUntypedTableTreeView"
+
+namespace UE::Insights
+{
+
+////////////////////////////////////////////////////////////////////////////////////////////////////
+
+SUntypedTableTreeView::SUntypedTableTreeView()
+{
+}
+
+////////////////////////////////////////////////////////////////////////////////////////////////////
+
+SUntypedTableTreeView::~SUntypedTableTreeView()
+{
+}
+
+////////////////////////////////////////////////////////////////////////////////////////////////////
+
+void SUntypedTableTreeView::Construct(const FArguments& InArgs, TSharedPtr<FUntypedTable> InTablePtr)
+{
+	bRunInAsyncMode = InArgs._RunInAsyncMode;
+	ConstructWidget(InTablePtr);
+}
+
+////////////////////////////////////////////////////////////////////////////////////////////////////
+
+void SUntypedTableTreeView::Reset()
+{
+	STableTreeView::Reset();
+}
+
+////////////////////////////////////////////////////////////////////////////////////////////////////
+
+void SUntypedTableTreeView::UpdateSourceTable(TSharedPtr<TraceServices::IUntypedTable> SourceTable)
+{
+	//check(Table->Is<FUntypedTable>());
+	TSharedPtr<FUntypedTable> UntypedTable = StaticCastSharedPtr<FUntypedTable>(Table);
+
+	if (UntypedTable->UpdateSourceTable(SourceTable))
+	{
+		RebuildColumns();
+	}
+
+	RebuildTree(true);
+}
+
+////////////////////////////////////////////////////////////////////////////////////////////////////
+
+void SUntypedTableTreeView::RebuildTree(bool bResync)
+{
+	FStopwatch Stopwatch;
+	Stopwatch.Start();
+
+	FStopwatch SyncStopwatch;
+	SyncStopwatch.Start();
+
+	if (bResync)
+	{
+		TableRowNodes.Empty();
+	}
+
+	const int32 PreviousNodeCount = TableRowNodes.Num();
+
+	//check(Table->Is<FUntypedTable>());
+	TSharedPtr<FUntypedTable> UntypedTable = StaticCastSharedPtr<FUntypedTable>(Table);
+
+	TSharedPtr<TraceServices::IUntypedTable> SourceTable = UntypedTable->GetSourceTable();
+	TSharedPtr<TraceServices::IUntypedTableReader> TableReader = UntypedTable->GetTableReader();
+
+	if (SourceTable.IsValid() && TableReader.IsValid())
+	{
+		const int32 TotalRowCount = static_cast<int32>(SourceTable->GetRowCount());
+		if (TotalRowCount != TableRowNodes.Num())
+		{
+			TableRowNodes.Empty(TotalRowCount);
+			FName BaseNodeName(TEXT("row"));
+			for (int32 RowIndex = 0; RowIndex < TotalRowCount; ++RowIndex)
+			{
+				TableReader->SetRowIndex(RowIndex);
+				FName NodeName(BaseNodeName, RowIndex + 1);
+				FTableTreeNodePtr NodePtr = MakeShared<FTableTreeNode>(NodeName, Table, RowIndex);
+				NodePtr->SetDefaultSortOrder(RowIndex + 1);
+				TableRowNodes.Add(NodePtr);
+			}
+			ensure(TableRowNodes.Num() == TotalRowCount);
+		}
+	}
+
+	SyncStopwatch.Stop();
+
+	if (bResync || TableRowNodes.Num() != PreviousNodeCount)
+	{
+		// Save selection.
+		TArray<FTableTreeNodePtr> SelectedItems;
+		TreeView->GetSelectedItems(SelectedItems);
+
+		UpdateTree();
+
+		TreeView->RebuildList();
+
+		// Restore selection.
+		if (SelectedItems.Num() > 0)
+		{
+			TreeView->ClearSelection();
+			for (FTableTreeNodePtr& NodePtr : SelectedItems)
+			{
+				NodePtr = GetNodeByTableRowIndex(NodePtr->GetRowIndex());
+			}
+			SelectedItems.RemoveAll([](const FTableTreeNodePtr& NodePtr) { return !NodePtr.IsValid(); });
+			if (SelectedItems.Num() > 0)
+			{
+				TreeView->SetItemSelection(SelectedItems, true);
+				TreeView->RequestScrollIntoView(SelectedItems.Last());
+			}
+		}
+	}
+
+	Stopwatch.Stop();
+	const double TotalTime = Stopwatch.GetAccumulatedTime();
+	if (TotalTime > 0.01)
+	{
+		const double SyncTime = SyncStopwatch.GetAccumulatedTime();
+		UE_LOG(LogInsights, Log, TEXT("[Table] Tree view rebuilt in %.4fs (sync: %.4fs + update: %.4fs) --> %d rows (%d added)"),
+			TotalTime, SyncTime, TotalTime - SyncTime, TableRowNodes.Num(), TableRowNodes.Num() - PreviousNodeCount);
+	}
+}
+
+////////////////////////////////////////////////////////////////////////////////////////////////////
+
+bool SUntypedTableTreeView::IsRunning() const
+{
+	return !CurrentOperationNameOverride.IsEmpty() || STableTreeView::IsRunning();
+}
+
+////////////////////////////////////////////////////////////////////////////////////////////////////
+
+double SUntypedTableTreeView::GetAllOperationsDuration()
+{
+	if (!CurrentOperationNameOverride.IsEmpty())
+	{
+		CurrentOperationStopwatch.Update();
+		return CurrentOperationStopwatch.GetAccumulatedTime();
+	}
+
+	return STableTreeView::GetAllOperationsDuration();
+}
+
+////////////////////////////////////////////////////////////////////////////////////////////////////
+
+FText SUntypedTableTreeView::GetCurrentOperationName() const
+{
+	if (!CurrentOperationNameOverride.IsEmpty())
+	{
+		return CurrentOperationNameOverride;
+	}
+
+	return STableTreeView::GetCurrentOperationName();
+}
+
+////////////////////////////////////////////////////////////////////////////////////////////////////
+
+void SUntypedTableTreeView::SetCurrentOperationNameOverride(const FText& InOperationName)
+{
+	CurrentOperationStopwatch.Start();
+	CurrentOperationNameOverride = InOperationName;
+}
+
+////////////////////////////////////////////////////////////////////////////////////////////////////
+
+void SUntypedTableTreeView::ClearCurrentOperationNameOverride()
+{
+	CurrentOperationStopwatch.Stop();
+	CurrentOperationStopwatch.Reset();
+
+	CurrentOperationNameOverride = FText::GetEmpty();
+}
+
+////////////////////////////////////////////////////////////////////////////////////////////////////
+
+} // namespace UE::Insights
+
+#undef LOCTEXT_NAMESPACE
diff --git a/Engine/Source/Developer/TraceInsightsCore/Public/InsightsCore/Table/ViewModels/TableImporter.h b/Engine/Source/Developer/TraceInsightsCore/Public/InsightsCore/Table/ViewModels/TableImporter.h
new file mode 100644
--- /dev/null
+++ b/Engine/Source/Developer/TraceInsightsCore/Public/InsightsCore/Table/ViewModels/TableImporter.h
@@ -0,0 +1,70 @@
+// Copyright Epic Games, Inc. All Rights Reserved.
+
+#pragma once
+
+#include "CoreTypes.h"
+
+// TraceServices
+#include "TraceServices/Model/TableImport.h"
+#include "TraceServices/Containers/Tables.h"
+
+class SDockTab;
+class FSpawnTabArgs;
+
+namespace TraceServices
+{
+	struct FTableImportCallbackParams;
+}
+
+namespace UE::Insights
+{
+
+class SUntypedTableTreeView;
+
+////////////////////////////////////////////////////////////////////////////////////////////////////
+
+class TRACEINSIGHTSCORE_API FTableImporter : public TSharedFromThis<FTableImporter>
+{
+private:
+	struct FOpenTableTabData
+	{
+		TSharedPtr<SDockTab> Tab;
+		TSharedPtr<SUntypedTableTreeView> TableTreeView;
+
+		bool operator ==(const FOpenTableTabData& Other) const
+		{
+			return Tab == Other.Tab && TableTreeView == Other.TableTreeView;
+		}
+	};
+
+public:
+	FTableImporter(FName InLogListingName);
+	virtual ~FTableImporter();
+
+	void StartImportProcess();
+	void ImportFile(const FString& Filename);
+
+	void StartDiffProcess();
+	void DiffFiles(const FString& FilenameA, const FString& FilenameB);
+
+	void CloseAllOpenTabs();
+
+private:
+	TSharedRef<SDockTab> SpawnTab_TableImportTreeView(const FSpawnTabArgs& Args, FName TableViewID, FText InDisplayName);
+	TSharedRef<SDockTab> SpawnTab_TableDiffTreeView(const FSpawnTabArgs& Args, FName TableViewID, FText InDisplayName);
+	void OnTableImportTreeViewTabClosed(TSharedRef<SDockTab> TabBeingClosed);
+
+	void DisplayImportTable(FName TableViewID);
+	void DisplayDiffTable(FName TableViewID);
+	void TableImportServiceCallback(TSharedPtr<TraceServices::FTableImportCallbackParams> Params);
+
+	FName GetTableID(const FString& Path);
+
+private:
+	FName LogListingName;
+	TMap<FName, FOpenTableTabData> OpenTablesMap;
+};
+
+////////////////////////////////////////////////////////////////////////////////////////////////////
+
+} // namespace UE::Insights
diff --git a/Engine/Source/Developer/TraceInsightsCore/Public/InsightsCore/Table/ViewModels/UntypedTable.h b/Engine/Source/Developer/TraceInsightsCore/Public/InsightsCore/Table/ViewModels/UntypedTable.h
new file mode 100644
--- /dev/null
+++ b/Engine/Source/Developer/TraceInsightsCore/Public/InsightsCore/Table/ViewModels/UntypedTable.h
@@ -0,0 +1,48 @@
+// Copyright Epic Games, Inc. All Rights Reserved.
+
+#pragma once
+
+#include "CoreTypes.h"
+
+#include "Templates/SharedPointer.h"
+
+// TraceInsightsCore
+#include "InsightsCore/Table/ViewModels/Table.h"
+
+namespace TraceServices
+{
+	class ITableLayout;
+	class IUntypedTable;
+	class IUntypedTableReader;
+}
+
+namespace UE::Insights
+{
+
+////////////////////////////////////////////////////////////////////////////////////////////////////
+
+class TRACEINSIGHTSCORE_API FUntypedTable : public FTable
+{
+public:
+	FUntypedTable();
+	virtual ~FUntypedTable();
+
+	virtual void Reset();
+
+	TSharedPtr<TraceServices::IUntypedTable> GetSourceTable() const { return SourceTable; }
+	TSharedPtr<TraceServices::IUntypedTableReader> GetTableReader() const { return TableReader; }
+
+	/* Update table content. Returns true if the table layout has changed. */
+	bool UpdateSourceTable(TSharedPtr<TraceServices::IUntypedTable> InSourceTable);
+
+private:
+	void CreateColumns(const TraceServices::ITableLayout& TableLayout);
+
+private:
+	TSharedPtr<TraceServices::IUntypedTable> SourceTable;
+	TSharedPtr<TraceServices::IUntypedTableReader> TableReader;
+};
+
+////////////////////////////////////////////////////////////////////////////////////////////////////
+
+} // namespace UE::Insights
diff --git a/Engine/Source/Developer/TraceInsightsCore/Public/InsightsCore/Table/Widgets/SUntypedDiffTableTreeView.h b/Engine/Source/Developer/TraceInsightsCore/Public/InsightsCore/Table/Widgets/SUntypedDiffTableTreeView.h
new file mode 100644
--- /dev/null
+++ b/Engine/Source/Developer/TraceInsightsCore/Public/InsightsCore/Table/Widgets/SUntypedDiffTableTreeView.h
@@ -0,0 +1,43 @@
+// Copyright Epic Games, Inc. All Rights Reserved.
+
+#pragma once
+
+#include "CoreTypes.h"
+#include "Containers/UnrealString.h"
+#include "Input/Reply.h"
+#include "Internationalization/Text.h"
+#include "Templates/SharedPointer.h"
+
+// TraceInsights
+#include "InsightsCore/Table/ViewModels/UntypedTable.h"
+#include "InsightsCore/Table/Widgets/SUntypedTableTreeView.h"
+
+namespace UE::Insights
+{
+
+////////////////////////////////////////////////////////////////////////////////////////////////////
+
+class TRACEINSIGHTSCORE_API SUntypedDiffTableTreeView : public SUntypedTableTreeView
+{
+public:
+	void UpdateSourceTableA(const FString& Name, TSharedPtr<TraceServices::IUntypedTable> SourceTable);
+	void UpdateSourceTableB(const FString& Name, TSharedPtr<TraceServices::IUntypedTable> SourceTable);
+
+protected:
+	FReply SwapTables_OnClicked();
+	FText GetSwapButtonText() const;
+
+	virtual TSharedPtr<SWidget> ConstructToolbar() override;
+
+	void RequestMergeTables();
+
+private:
+	TSharedPtr<TraceServices::IUntypedTable> SourceTableA;
+	TSharedPtr<TraceServices::IUntypedTable> SourceTableB;
+	FString TableNameA;
+	FString TableNameB;
+};
+
+////////////////////////////////////////////////////////////////////////////////////////////////////
+
+} // namespace UE::Insights
diff --git a/Engine/Source/Developer/TraceInsightsCore/Public/InsightsCore/Table/Widgets/SUntypedTableTreeView.h b/Engine/Source/Developer/TraceInsightsCore/Public/InsightsCore/Table/Widgets/SUntypedTableTreeView.h
new file mode 100644
--- /dev/null
+++ b/Engine/Source/Developer/TraceInsightsCore/Public/InsightsCore/Table/Widgets/SUntypedTableTreeView.h
@@ -0,0 +1,77 @@
+// Copyright Epic Games, Inc. All Rights Reserved.
+
+#pragma once
+
+#include "CoreTypes.h"
+#include "Containers/UnrealString.h"
+#include "Input/Reply.h"
+#include "Internationalization/Text.h"
+#include "Templates/SharedPointer.h"
+
+// TraceInsightsCore
+#include "InsightsCore/Common/Stopwatch.h"
+#include "InsightsCore/Table/ViewModels/UntypedTable.h"
+#include "InsightsCore/Table/Widgets/STableTreeView.h"
+
+namespace TraceServices
+{
+	class IUntypedTable;
+}
+
+namespace UE::Insights
+{
+
+////////////////////////////////////////////////////////////////////////////////////////////////////
+
+class TRACEINSIGHTSCORE_API SUntypedTableTreeView : public STableTreeView
+{
+public:
+	SUntypedTableTreeView();
+	virtual ~SUntypedTableTreeView();
+
+	SLATE_BEGIN_ARGS(SUntypedTableTreeView)
+		: _RunInAsyncMode(false)
+		{}
+		SLATE_ARGUMENT(bool, RunInAsyncMode)
+	SLATE_END_ARGS()
+
+	/**
+	 * Construct this widget
+	 * @param InArgs - The declaration data for this widget
+	 */
+	void Construct(const FArguments& InArgs, TSharedPtr<FUntypedTable> InTablePtr);
+
+	TSharedPtr<FUntypedTable> GetUntypedTable() const { return StaticCastSharedPtr<FUntypedTable>(GetTable()); }
+
+	void UpdateSourceTable(TSharedPtr<TraceServices::IUntypedTable> SourceTable);
+
+	virtual void Reset();
+
+	//////////////////////////////////////////////////
+	// IAsyncOperationStatusProvider
+
+	virtual bool IsRunning() const override;
+	virtual double GetAllOperationsDuration() override;
+	virtual double GetCurrentOperationDuration() override { return 0.0; }
+	virtual uint32 GetOperationCount() const override { return 1; }
+	virtual FText GetCurrentOperationName() const override;
+
+	//////////////////////////////////////////////////
+
+	void SetCurrentOperationNameOverride(const FText& InOperationName);
+	void ClearCurrentOperationNameOverride();
+
+	/**
+	 * Rebuilds the tree (if necessary).
+	 * @param bResync - If true, it forces a resync even if the list did not changed since last sync.
+	 */
+	virtual void RebuildTree(bool bResync);
+
+private:
+	FStopwatch CurrentOperationStopwatch;
+	FText CurrentOperationNameOverride;
+};
+
+////////////////////////////////////////////////////////////////////////////////////////////////////
+
+} // namespace UE::Insights
diff --git a/Engine/Source/Developer/TraceInsightsCore/TraceInsightsCore.Build.cs b/Engine/Source/Developer/TraceInsightsCore/TraceInsightsCore.Build.cs
--- a/Engine/Source/Developer/TraceInsightsCore/TraceInsightsCore.Build.cs
+++ b/Engine/Source/Developer/TraceInsightsCore/TraceInsightsCore.Build.cs
@@ -18,6 +18,7 @@ public class TraceInsightsCore : ModuleRules
 				"Slate",
 				"SlateCore",
 				"ToolWidgets",
+				"TraceServices",
 			}
 		);
 	}
diff --git a/Engine/Source/Developer/TraceInsightsFrontend/Private/InsightsFrontend/Widgets/STraceStoreWindow.cpp b/Engine/Source/Developer/TraceInsightsFrontend/Private/InsightsFrontend/Widgets/STraceStoreWindow.cpp
--- a/Engine/Source/Developer/TraceInsightsFrontend/Private/InsightsFrontend/Widgets/STraceStoreWindow.cpp
+++ b/Engine/Source/Developer/TraceInsightsFrontend/Private/InsightsFrontend/Widgets/STraceStoreWindow.cpp
@@ -45,6 +45,7 @@
 #include "InsightsCore/Common/InsightsCoreStyle.h"
 #include "InsightsCore/Common/MiscUtils.h"
 #include "InsightsCore/Common/Stopwatch.h"
+#include "InsightsCore/Table/ViewModels/TableImporter.h"
 #include "InsightsCore/Version.h"
 
 // TraceInsightsFrontend
@@ -66,20 +67,19 @@
 #include "Windows/HideWindowsPlatformTypes.h"
 #endif
 
-#ifndef UE_INSIGHTS_TABLE_IMPORT_TOOL_ENABLED
-#define UE_INSIGHTS_TABLE_IMPORT_TOOL_ENABLED 0
-#endif
-
 #define LOCTEXT_NAMESPACE "UE::Insights::STraceStoreWindow"
 
 namespace UE::Insights
 {
 
+FName STraceStoreWindow::LogListingName(TEXT("InsightsFrontend"));
+
 ////////////////////////////////////////////////////////////////////////////////////////////////////
 // STraceStoreWindow
 ////////////////////////////////////////////////////////////////////////////////////////////////////
 
 STraceStoreWindow::STraceStoreWindow()
+	: TableImporter(MakeShared<FTableImporter>(LogListingName))
 {
 	SortColumn = FTraceListColumns::Date;
 	SortMode = EColumnSortMode::Ascending;
@@ -2002,13 +2002,12 @@ FReply STraceStoreWindow::OnDrop(const FGeometry& MyGeometry, const FDragDropEve
 					OpenTraceFile(Files[0]);
 					return FReply::Handled();
 				}
-#if UE_INSIGHTS_TABLE_IMPORT_TOOL_ENABLED
+
 				if (DraggedFileExtension == TEXT(".csv") || DraggedFileExtension == TEXT(".tsv"))
 				{
-					FTableImportTool::Get()->ImportFile(Files[0]);
+					TableImporter->ImportFile(Files[0]);
 					return FReply::Handled();
 				}
-#endif // UE_INSIGHTS_TABLE_IMPORT_TOOL_ENABLED
 			}
 		}
 	}
@@ -2168,22 +2167,22 @@ TSharedRef<SWidget> STraceStoreWindow::MakeTraceListMenu()
 			FUIAction(FExecuteAction::CreateSP(this, &STraceStoreWindow::OpenTraceFile)),
 			NAME_None,
 			EUserInterfaceActionType::Button);
-#if UE_INSIGHTS_TABLE_IMPORT_TOOL_ENABLED
+
 		MenuBuilder.AddMenuEntry(
 			LOCTEXT("ImportTableButtonLabel", "Import Table..."),
 			LOCTEXT("ImportTableButtonTooltip", "Opens .csv or .tsv file."),
-			FSlateIcon(FAppStyle::Get().GetStyleSetName(), "Icons.FolderOpen"),
-			FUIAction(FExecuteAction::CreateLambda([]{ FTableImportTool::Get()->StartImportProcess(); })),
+			FSlateIcon(FInsightsCoreStyle::GetStyleSetName(), "Icons.ImportTable"),
+			FUIAction(FExecuteAction::CreateLambda([this]{ TableImporter->StartImportProcess(); })),
 			NAME_None,
 			EUserInterfaceActionType::Button);
+
 		MenuBuilder.AddMenuEntry(
 			LOCTEXT("DiffTablesButtonLabel", "Diff Tables..."),
 			LOCTEXT("DiffTablesButtonTooltip", "Opens two table files in diff mode."),
 			FSlateIcon(FAppStyle::Get().GetStyleSetName(), "Icons.FolderOpen"),
-			FUIAction(FExecuteAction::CreateLambda([]{ FTableImportTool::Get()->StartDiffProcess(); })),
+			FUIAction(FExecuteAction::CreateLambda([this]{ TableImporter->StartDiffProcess(); })),
 			NAME_None,
 			EUserInterfaceActionType::Button);
-#endif // UE_INSIGHTS_TABLE_IMPORT_TOOL_ENABLED
 	}
 
 	MenuBuilder.EndSection();
@@ -2684,7 +2683,6 @@ FReply STraceStoreWindow::AddWatchDir_Clicked()
 			UE::Trace::FStoreClient* StoreClient = TraceStoreConnection->GetStoreClient();
 			if (!StoreClient || !StoreClient->SetStoreDirectories(nullptr, { (*SelectedDirectory) }, {}))
 			{
-				FName LogListingName(TEXT("UnrealInsights"));
 				FMessageLog(LogListingName).Error(LOCTEXT("StoreCommunicationFail", "Failed to change settings on the store service."));
 			}
 		}
diff --git a/Engine/Source/Developer/TraceInsightsFrontend/Private/InsightsFrontend/Widgets/STraceStoreWindow.h b/Engine/Source/Developer/TraceInsightsFrontend/Private/InsightsFrontend/Widgets/STraceStoreWindow.h
--- a/Engine/Source/Developer/TraceInsightsFrontend/Private/InsightsFrontend/Widgets/STraceStoreWindow.h
+++ b/Engine/Source/Developer/TraceInsightsFrontend/Private/InsightsFrontend/Widgets/STraceStoreWindow.h
@@ -41,6 +41,8 @@ class FInsightsFrontendSettings;
 class FStoreBrowser;
 struct FStoreBrowserTraceInfo;
 
+class FTableImporter;
+
 struct FTraceViewModel;
 struct FTraceDirectoryModel;
 
@@ -466,6 +468,9 @@ private:
 #elif PLATFORM_MAC || PLATFORM_LINUX
 	sem_t* AutoConnectEvent = SEM_FAILED;
 #endif
+
+	static FName LogListingName;
+	TSharedRef<FTableImporter> TableImporter;
 };
 
 ////////////////////////////////////////////////////////////////////////////////////////////////////
