[Insights] Fixed issue in LogView where scroll stops when encountering a big multiline message.

diff --git a/Engine/Source/Developer/TraceInsights/Private/Insights/Widgets/SLogView.cpp b/Engine/Source/Developer/TraceInsights/Private/Insights/Widgets/SLogView.cpp
--- a/Engine/Source/Developer/TraceInsights/Private/Insights/Widgets/SLogView.cpp
+++ b/Engine/Source/Developer/TraceInsights/Private/Insights/Widgets/SLogView.cpp
@@ -653,6 +653,7 @@ void SLogView::Construct(const FArguments& InArgs)
 				[
 					SAssignNew(ListView, SListView<TSharedPtr<FLogMessage>>)
 					.ExternalScrollbar(ExternalScrollbar)
+					.FixedLineScrollOffset(0.0)
 					.SelectionMode(ESelectionMode::Single)
 					.OnMouseButtonClick(this, &SLogView::OnMouseButtonClick)
 					.OnSelectionChanged(this, &SLogView::OnSelectionChanged)
