[Insights] Fixed Frames view trying to render bars larger than Slate allows (error in Slate rendering).

diff --git a/Engine/Source/Developer/TraceInsights/Private/Insights/TimingProfiler/ViewModels/FrameTrackHelper.cpp b/Engine/Source/Developer/TraceInsights/Private/Insights/TimingProfiler/ViewModels/FrameTrackHelper.cpp
--- a/Engine/Source/Developer/TraceInsights/Private/Insights/TimingProfiler/ViewModels/FrameTrackHelper.cpp
+++ b/Engine/Source/Developer/TraceInsights/Private/Insights/TimingProfiler/ViewModels/FrameTrackHelper.cpp
@@ -216,6 +216,11 @@ void FFrameTrackDrawHelper::DrawCached(const FFrameTrackSeries& Series) const
 			}
 		}
 
+		if (ValueY > BaselineY + ViewHeight)
+		{
+			ValueY = BaselineY + ViewHeight;
+		}
+
 		const float H = ValueY - BaselineY;
 		const float Y = ViewHeight - H;
 
