Automate Test Case: Unreal Insights - Memory Insights - Allocations Query Table

#rnx

diff --git a/Engine/Source/Developer/TraceInsights/Private/Insights/Tests/FunctionalTests/MemoryInsightsTests.cpp b/Engine/Source/Developer/TraceInsights/Private/Insights/Tests/FunctionalTests/MemoryInsightsTests.cpp
--- a/Engine/Source/Developer/TraceInsights/Private/Insights/Tests/FunctionalTests/MemoryInsightsTests.cpp
+++ b/Engine/Source/Developer/TraceInsights/Private/Insights/Tests/FunctionalTests/MemoryInsightsTests.cpp
@@ -89,7 +89,7 @@ bool FChangeGroupingCommand::Update()
 	TArray<TSharedPtr<UE::Insights::FTreeNodeGrouping>> CurrentGroupings;
 	for (const auto& Grouping : MemAllocTableTreeView->GetAvailableGroupings())
 	{
-		if (Grouping->GetTitleName().ToString().Contains(TEXT("By Callstack")))
+		if (Grouping->GetTitleName().ToString().Contains(TEXT("By Free Callstack")))
 		{
 			CurrentGroupings.Add(Grouping);
 		}
@@ -101,10 +101,11 @@ bool FChangeGroupingCommand::Update()
 	return true;
 }
 
-DEFINE_LATENT_AUTOMATION_COMMAND_FOUR_PARAMETER(FVerifyHierarchyCallStackCommand, TSharedPtr<UE::Insights::MemoryProfiler::SMemAllocTableTreeView>, MemAllocTableTreeView, FInsightsTestUtils, InsightsTestUtils, double, Timeout, FAutomationTestBase*, Test);
+DEFINE_LATENT_AUTOMATION_COMMAND_THREE_PARAMETER(FVerifyHierarchyCallStackCommand, TSharedPtr<UE::Insights::MemoryProfiler::SMemAllocTableTreeView>, MemAllocTableTreeView, double, Timeout, FAutomationTestBase*, Test);
 bool FVerifyHierarchyCallStackCommand::Update()
 {
 	using namespace UE::Insights::MemoryProfiler;
+	FInsightsTestUtils InsightsTestUtils(Test);
 
 	if (!MemAllocTableTreeView->IsRunningAsyncUpdate())
 	{
@@ -149,6 +150,8 @@ const TMap<TraceServices::IAllocationsProvider::EQueryRule, UE::Insights::Memory
 	{TraceServices::IAllocationsProvider::EQueryRule::aABf, {nullptr, {50.0, 51.0, 0.0, 0.0}}},
 	{TraceServices::IAllocationsProvider::EQueryRule::AafB, {nullptr, {50.0, 51.0, 0.0, 0.0}}},
 	{TraceServices::IAllocationsProvider::EQueryRule::AaB, {nullptr, {50.0, 51.0, 0.0, 0.0}}},
+	{TraceServices::IAllocationsProvider::EQueryRule::AoB, {nullptr, {10.0, 20.0, 0.0, 0.0}}},
+	{TraceServices::IAllocationsProvider::EQueryRule::AiB, {nullptr, {10.0, 20.0, 0.0, 0.0}}},
 };
 
 const TMap<TraceServices::IAllocationsProvider::EQueryRule, UE::Insights::MemoryProfiler::SMemAllocTableTreeView::FQueryParams> AllocsTimeMarkerEditorPackageGetterMap
@@ -167,6 +170,14 @@ const TMap<TraceServices::IAllocationsProvider::EQueryRule, UE::Insights::Memory
 	{TraceServices::IAllocationsProvider::EQueryRule::AafB, {nullptr, {2.0, 3.0, 0.0, 0.0}}},
 	{TraceServices::IAllocationsProvider::EQueryRule::AaB, {nullptr, {2.0, 3.0, 0.0, 0.0}}},
 	{TraceServices::IAllocationsProvider::EQueryRule::aABf, {nullptr, {2.0, 3.0, 0.0, 0.0}}},
+	{TraceServices::IAllocationsProvider::EQueryRule::AoB, {nullptr, {10.0, 20.0, 0.0, 0.0}}},
+	{TraceServices::IAllocationsProvider::EQueryRule::AiB, {nullptr, {10.0, 20.0, 0.0, 0.0}}},
+};
+
+const TSet<TraceServices::IAllocationsProvider::EQueryRule> TemporaryExcludedRules
+{
+	TraceServices::IAllocationsProvider::EQueryRule::AoB,
+	TraceServices::IAllocationsProvider::EQueryRule::AiB
 };
 
 bool MemoryInsightsAllocationsQueryTableTest(const FString& Parameters, const TMap <TraceServices::IAllocationsProvider::EQueryRule, UE::Insights::MemoryProfiler::SMemAllocTableTreeView::FQueryParams> AllocsTimeMarkerGetterMap, FAutomationTestBase* Test)
@@ -174,7 +185,6 @@ bool MemoryInsightsAllocationsQueryTableTest(const FString& Parameters, const TM
 	using namespace UE::Insights::MemoryProfiler;
 
 	double Timeout = 30.0;
-	FInsightsTestUtils InsightsTestUtils(Test);
 	TSharedPtr<SMemoryProfilerWindow> ProfilerWindow = FMemoryProfilerManager::Get()->GetProfilerWindow();
 	TSharedPtr<UE::Insights::FInsightsManager> InsightsManager = UE::Insights::FInsightsManager::Get();
 	FMemorySharedState& SharedState = ProfilerWindow->GetSharedState();
@@ -201,10 +211,10 @@ bool MemoryInsightsAllocationsQueryTableTest(const FString& Parameters, const TM
 	QueryParams.Rule = MemoryRule;
 	MemAllocTableTreeView->SetQueryParams(QueryParams);
 
-	ADD_LATENT_AUTOMATION_COMMAND(FWaitForRunningQuieryFinishedCommand(MemAllocTableTreeView, 120.0f, Test));
+	ADD_LATENT_AUTOMATION_COMMAND(FWaitForRunningQuieryFinishedCommand(MemAllocTableTreeView, Timeout, Test));
 	ADD_LATENT_AUTOMATION_COMMAND(FChangeGroupingCommand(MemAllocTableTreeView, Test));
-	ADD_LATENT_AUTOMATION_COMMAND(FWaitForRunningQuieryFinishedCommand(MemAllocTableTreeView, 120.0f, Test));
-	ADD_LATENT_AUTOMATION_COMMAND(FVerifyHierarchyCallStackCommand(MemAllocTableTreeView, InsightsTestUtils, 120.0f, Test));
+	ADD_LATENT_AUTOMATION_COMMAND(FWaitForRunningQuieryFinishedCommand(MemAllocTableTreeView, Timeout, Test));
+	ADD_LATENT_AUTOMATION_COMMAND(FVerifyHierarchyCallStackCommand(MemAllocTableTreeView, Timeout, Test));
 
 	return true;
 }
@@ -232,12 +242,15 @@ void FMemoryInsightsAllocationsQueryTableEditorPackageTest::GetTests(TArray<FStr
 
 	FMemorySharedState& SharedState = ProfilerWindow->GetSharedState();
 
-	for (const auto& MemoryRules : SharedState.GetMemoryRules())
+	for (const auto& MemoryRule : SharedState.GetMemoryRules())
 	{
-		const FString& MemoryRuleName = MemoryRules->GetShortName().ToString();
+		if (!TemporaryExcludedRules.Contains(MemoryRule->GetValue()))
+		{
+			const FString& MemoryRuleName = MemoryRule->GetShortName().ToString();
 
-		OutBeautifiedNames.Add(FString::Printf(TEXT("%s"), *MemoryRuleName));
-		OutTestCommands.Add(MemoryRuleName);
+			OutBeautifiedNames.Add(FString::Printf(TEXT("%s"), *MemoryRuleName));
+			OutTestCommands.Add(MemoryRuleName);
+		}
 	}
 }
 
@@ -263,14 +276,16 @@ void FMemoryInsightsAllocationsQueryTableStandaloneTest::GetTests(TArray<FString
 	}
 
 	FMemorySharedState& SharedState = ProfilerWindow->GetSharedState();
-	TSharedPtr<FMemoryRuleSpec> MemoryRule = SharedState.GetMemoryRules()[0];
 
-	for (const auto& MemoryRules : SharedState.GetMemoryRules())
+	for (const auto& MemoryRule : SharedState.GetMemoryRules())
 	{
-		const FString& MemoryRuleName = MemoryRules->GetShortName().ToString();
+		if (!TemporaryExcludedRules.Contains(MemoryRule->GetValue()))
+		{
+			const FString& MemoryRuleName = MemoryRule->GetShortName().ToString();
 
-		OutBeautifiedNames.Add(FString::Printf(TEXT("%s"), *MemoryRuleName));
-		OutTestCommands.Add(MemoryRuleName);
+			OutBeautifiedNames.Add(FString::Printf(TEXT("%s"), *MemoryRuleName));
+			OutTestCommands.Add(MemoryRuleName);
+		}
 	}
 }
 
