[Insights] Minor fix for a variable name.

#rnx

diff --git a/Engine/Source/Developer/TraceInsights/Private/Insights/CookProfiler/Widgets/SPackageTableTreeView.cpp b/Engine/Source/Developer/TraceInsights/Private/Insights/CookProfiler/Widgets/SPackageTableTreeView.cpp
--- a/Engine/Source/Developer/TraceInsights/Private/Insights/CookProfiler/Widgets/SPackageTableTreeView.cpp
+++ b/Engine/Source/Developer/TraceInsights/Private/Insights/CookProfiler/Widgets/SPackageTableTreeView.cpp
@@ -86,6 +86,8 @@ void SPackageTableTreeView::ExtendMenu(FMenuBuilder& MenuBuilder)
 {
 }
 
+////////////////////////////////////////////////////////////////////////////////////////////////////
+
 void SPackageTableTreeView::AddCommmands()
 {
 	FPackageTableTreeViewCommands::Register();
@@ -137,14 +139,14 @@ void SPackageTableTreeView::RebuildTree(bool bResync)
 	{
 		TraceServices::FProviderReadScopeLock ProviderReadScope(*CookProvider);
 
-		TArray64<TraceServices::FPackageData> PackageAggreagation;
-		CookProvider->CreateAggregation(PackageAggreagation);
+		TArray64<TraceServices::FPackageData> PackageAggregation;
+		CookProvider->CreateAggregation(PackageAggregation);
 		const uint32 NumPackages = CookProvider->GetNumPackages();
 		Packages.Reserve(NumPackages);
 		TableRowNodes.Reserve(NumPackages);
 
 		TArray<FTableTreeNodePtr>* Nodes = &TableRowNodes;
-		for(const TraceServices::FPackageData& Package : PackageAggreagation)
+		for (const TraceServices::FPackageData& Package : PackageAggregation)
 		{
 			Packages.Emplace(Package);
 			uint32 Index = static_cast<uint32>(Packages.Num() - 1);
