[Insights]
  - Analysis and UI for the new GPU Profiler.
  - TraceServices: Added API to add/enumerate and edit GPU Queues in IEditableTimingProfilerProvider.
  - TraceServices: Added GetReadProvider() in IEditableTimingProfilerProvider.
  - TraceServices: Added GetOriginalTimerIdFromMetadata() and GetMetadata() in both ITimingProfilerTimerReader and ITimingProfilerProvider.
  - TraceServices: Removed GetMetadata() API from IEditableTimingProfilerProvider. IEditableTimingProfilerProvider now requires GetEditableMetadata().
  - Added GetLastTimestamp() to IEditableTimeline and TMonotonicTimeline.

diff --git a/Engine/Source/Developer/TraceInsights/Private/Insights/TimingProfiler/ViewModels/ThreadTimingSharedState.cpp b/Engine/Source/Developer/TraceInsights/Private/Insights/TimingProfiler/ViewModels/ThreadTimingSharedState.cpp
--- a/Engine/Source/Developer/TraceInsights/Private/Insights/TimingProfiler/ViewModels/ThreadTimingSharedState.cpp
+++ b/Engine/Source/Developer/TraceInsights/Private/Insights/TimingProfiler/ViewModels/ThreadTimingSharedState.cpp
@@ -90,17 +90,73 @@ FThreadTimingSharedState::FThreadTimingSharedState(STimingView* InTimingView)
 
 ////////////////////////////////////////////////////////////////////////////////////////////////////
 
-TSharedPtr<FCpuTimingTrack> FThreadTimingSharedState::GetCpuTrack(uint32 InThreadId)
+TSharedPtr<FGpuTimingTrack> FThreadTimingSharedState::GetGpuTrack(uint32 InQueueId)
 {
-	TSharedPtr<FCpuTimingTrack>*const TrackPtrPtr = CpuTracks.Find(InThreadId);
+	TSharedPtr<FGpuTimingTrack>* const TrackPtrPtr = GpuTracks.Find(InQueueId);
 	return TrackPtrPtr ? *TrackPtrPtr : nullptr;
 }
 
 ////////////////////////////////////////////////////////////////////////////////////////////////////
 
-bool FThreadTimingSharedState::IsGpuTrackVisible() const
+bool FThreadTimingSharedState::IsOldGpu1TrackVisible() const
+{
+	return OldGpu1Track != nullptr && OldGpu1Track->IsVisible();
+}
+
+////////////////////////////////////////////////////////////////////////////////////////////////////
+
+bool FThreadTimingSharedState::IsOldGpu2TrackVisible() const
+{
+	return OldGpu2Track != nullptr && OldGpu2Track->IsVisible();
+}
+
+////////////////////////////////////////////////////////////////////////////////////////////////////
+
+bool FThreadTimingSharedState::IsAnyGpuTrackVisible() const
+{
+	if (IsOldGpu1TrackVisible() || IsOldGpu2TrackVisible())
+	{
+		return true;
+	}
+	for (const auto& KV : GpuTracks)
+	{
+		if (KV.Value->IsVisible())
+		{
+			return true;
+		}
+	}
+	return false;
+}
+
+////////////////////////////////////////////////////////////////////////////////////////////////////
+
+bool FThreadTimingSharedState::IsGpuTrackVisible(uint32 InQueueId) const
+{
+	const TSharedPtr<FGpuTimingTrack>* const TrackPtrPtr = GpuTracks.Find(InQueueId);
+	return TrackPtrPtr && (*TrackPtrPtr)->IsVisible();
+}
+
+////////////////////////////////////////////////////////////////////////////////////////////////////
+
+void FThreadTimingSharedState::GetVisibleGpuQueues(TSet<uint32>& OutSet) const
+{
+	OutSet.Reset();
+	for (const auto& KV : GpuTracks)
+	{
+		const FGpuTimingTrack& Track = *KV.Value;
+		if (Track.IsVisible())
+		{
+			OutSet.Add(KV.Key);
+		}
+	}
+}
+
+////////////////////////////////////////////////////////////////////////////////////////////////////
+
+TSharedPtr<FCpuTimingTrack> FThreadTimingSharedState::GetCpuTrack(uint32 InThreadId)
 {
-	return (GpuTrack != nullptr && GpuTrack->IsVisible()) || (Gpu2Track != nullptr && Gpu2Track->IsVisible());
+	TSharedPtr<FCpuTimingTrack>* const TrackPtrPtr = CpuTracks.Find(InThreadId);
+	return TrackPtrPtr ? *TrackPtrPtr : nullptr;
 }
 
 ////////////////////////////////////////////////////////////////////////////////////////////////////
@@ -140,14 +196,23 @@ void FThreadTimingSharedState::GetVisibleTimelineIndexes(TSet<uint32>& OutSet) c
 		}
 	}
 
-	if (GpuTrack.IsValid() && GpuTrack->IsVisible())
+	if (OldGpu1Track.IsValid() && OldGpu1Track->IsVisible())
 	{
-		OutSet.Add(GpuTrack->GetTimelineIndex());
+		OutSet.Add(OldGpu1Track->GetTimelineIndex());
 	}
 
-	if (Gpu2Track.IsValid() && Gpu2Track->IsVisible())
+	if (OldGpu2Track.IsValid() && OldGpu2Track->IsVisible())
 	{
-		OutSet.Add(Gpu2Track->GetTimelineIndex());
+		OutSet.Add(OldGpu2Track->GetTimelineIndex());
+	}
+
+	for (const auto& KV : GpuTracks)
+	{
+		const FGpuTimingTrack& Track = *KV.Value;
+		if (Track.IsVisible())
+		{
+			OutSet.Add(Track.GetTimelineIndex());
+		}
 	}
 }
 
@@ -171,8 +236,9 @@ void FThreadTimingSharedState::OnBeginSession(Timing::ITimingViewSession& InSess
 		bShowHideAllCpuTracks = true;
 	}
 
-	GpuTrack = nullptr;
-	Gpu2Track = nullptr;
+	OldGpu1Track = nullptr;
+	OldGpu2Track = nullptr;
+	GpuTracks.Reset();
 	CpuTracks.Reset();
 	ThreadGroups.Reset();
 
@@ -192,8 +258,9 @@ void FThreadTimingSharedState::OnEndSession(Timing::ITimingViewSession& InSessio
 	bShowHideAllGpuTracks = false;
 	bShowHideAllCpuTracks = false;
 
-	GpuTrack = nullptr;
-	Gpu2Track = nullptr;
+	OldGpu1Track = nullptr;
+	OldGpu2Track = nullptr;
+	GpuTracks.Reset();
 	CpuTracks.Reset();
 	ThreadGroups.Reset();
 
@@ -228,36 +295,88 @@ void FThreadTimingSharedState::Tick(Timing::ITimingViewSession& InSession, const
 
 			LLM_SCOPE_BYTAG(Insights);
 
-			// Check if we have a GPU track.
-			if (!GpuTrack.IsValid())
+			// Check if we have the old GPU timelines.
+			if (!OldGpu1Track.IsValid())
 			{
 				uint32 GpuTimelineIndex;
 				if (TimingProfilerProvider->GetGpuTimelineIndex(GpuTimelineIndex))
 				{
-					GpuTrack = MakeShared<FGpuTimingTrack>(*this, TEXT("GPU"), nullptr, GpuTimelineIndex, FGpuTimingTrack::Gpu1ThreadId);
-					GpuTrack->SetOrder(FTimingTrackOrder::Gpu);
-					GpuTrack->SetVisibilityFlag(bShowHideAllGpuTracks);
-					InSession.AddScrollableTrack(GpuTrack);
+					OldGpu1Track = MakeShared<FGpuTimingTrack>(*this, TEXT("GPU"), nullptr, GpuTimelineIndex, FGpuTimingTrack::Gpu1ThreadId);
+					OldGpu1Track->SetOrder(FTimingTrackOrder::Gpu);
+					OldGpu1Track->SetVisibilityFlag(bShowHideAllGpuTracks);
+					InSession.AddScrollableTrack(OldGpu1Track);
 				}
 			}
-			if (!Gpu2Track.IsValid())
+			if (!OldGpu2Track.IsValid())
 			{
 				uint32 GpuTimelineIndex;
 				if (TimingProfilerProvider->GetGpu2TimelineIndex(GpuTimelineIndex))
 				{
-					Gpu2Track = MakeShared<FGpuTimingTrack>(*this, TEXT("GPU2"), nullptr, GpuTimelineIndex, FGpuTimingTrack::Gpu2ThreadId);
-					Gpu2Track->SetOrder(FTimingTrackOrder::Gpu + 1);
-					Gpu2Track->SetVisibilityFlag(bShowHideAllGpuTracks);
-					InSession.AddScrollableTrack(Gpu2Track);
+					OldGpu2Track = MakeShared<FGpuTimingTrack>(*this, TEXT("GPU2"), nullptr, GpuTimelineIndex, FGpuTimingTrack::Gpu2ThreadId);
+					OldGpu2Track->SetOrder(FTimingTrackOrder::Gpu + 1);
+					OldGpu2Track->SetVisibilityFlag(bShowHideAllGpuTracks);
+					InSession.AddScrollableTrack(OldGpu2Track);
 				}
 			}
 
 			bool bTracksOrderChanged = false;
-			int32 Order = FTimingTrackOrder::Cpu;
+			int32 GpuTrackOrder = FTimingTrackOrder::Gpu + 100;
+			int32 CpuTrackOrder = FTimingTrackOrder::Cpu;
+
+			// Iterate through GPU queues.
+			TimingProfilerProvider->EnumerateGpuQueues([this, &InSession, &bTracksOrderChanged, &GpuTrackOrder](const TraceServices::FGpuQueueInfo& QueueInfo)
+			{
+				// Check if there is an available GPU Work track for this queue.
+				// TODO: Attach this track as an overlay/child track to the normal GPU Queue track.
+				TSharedPtr<FGpuTimingTrack>* WorkTrackPtrPtr = GpuTracks.Find(QueueInfo.Id | 0xFF000000);
+				if (WorkTrackPtrPtr == nullptr)
+				{
+					// Create new Timing Events track for the GPU queue.
+					const FString TrackName = FString::Printf(TEXT("GPU%u - %s %u - WORK"), QueueInfo.GPU, QueueInfo.Name, QueueInfo.Index);
+					TSharedPtr<FGpuTimingTrack> Track = MakeShared<FGpuTimingTrack>(*this, TrackName, nullptr, QueueInfo.WorkTimelineIndex, QueueInfo.Id | 0xFF000000);
+					Track->SetOrder(GpuTrackOrder);
+					Track->SetVisibilityFlag(bShowHideAllGpuTracks);
+					GpuTracks.Add(QueueInfo.Id | 0xFF000000, Track);
+					InSession.AddScrollableTrack(Track);
+				}
+				else
+				{
+					TSharedPtr<FGpuTimingTrack> Track = *WorkTrackPtrPtr;
+					if (Track->GetOrder() != GpuTrackOrder)
+					{
+						Track->SetOrder(GpuTrackOrder);
+						bTracksOrderChanged = true;
+					}
+				}
+				GpuTrackOrder += 1;
+
+				// Check if there is an available GPU track for this queue.
+				TSharedPtr<FGpuTimingTrack>* TrackPtrPtr = GpuTracks.Find(QueueInfo.Id);
+				if (TrackPtrPtr == nullptr)
+				{
+					// Create new Timing Events track for the GPU Queue.
+					const FString TrackName = FString::Printf(TEXT("GPU%u - %s %u"), QueueInfo.GPU, QueueInfo.Name, QueueInfo.Index);
+					TSharedPtr<FGpuTimingTrack> Track = MakeShared<FGpuTimingTrack>(*this, TrackName, nullptr, QueueInfo.TimelineIndex, QueueInfo.Id);
+					Track->SetOrder(GpuTrackOrder);
+					Track->SetVisibilityFlag(bShowHideAllGpuTracks);
+					GpuTracks.Add(QueueInfo.Id, Track);
+					InSession.AddScrollableTrack(Track);
+				}
+				else
+				{
+					TSharedPtr<FGpuTimingTrack> Track = *TrackPtrPtr;
+					if (Track->GetOrder() != GpuTrackOrder)
+					{
+						Track->SetOrder(GpuTrackOrder);
+						bTracksOrderChanged = true;
+					}
+				}
+				GpuTrackOrder += 99;
+			});
 
 			// Iterate through threads.
 			const TraceServices::IThreadProvider& ThreadProvider = TraceServices::ReadThreadProvider(InAnalysisSession);
-			ThreadProvider.EnumerateThreads([this, &InSession, &bTracksOrderChanged, &Order, TimingProfilerProvider, LoadTimeProfilerProvider](const TraceServices::FThreadInfo& ThreadInfo)
+			ThreadProvider.EnumerateThreads([this, &InSession, &bTracksOrderChanged, &CpuTrackOrder, TimingProfilerProvider, LoadTimeProfilerProvider](const TraceServices::FThreadInfo& ThreadInfo)
 			{
 				// Check if this thread is part of a group?
 				bool bIsGroupVisible = bShowHideAllCpuTracks;
@@ -273,13 +392,13 @@ void FThreadTimingSharedState::Tick(Timing::ITimingViewSession& InSession, const
 				if (!ThreadGroups.Contains(GroupName))
 				{
 					// Note: The GroupName pointer should be valid for the duration of the session.
-					ThreadGroups.Add(GroupName, { GroupName, bIsGroupVisible, 0, Order });
+					ThreadGroups.Add(GroupName, { GroupName, bIsGroupVisible, 0, CpuTrackOrder });
 				}
 				else
 				{
 					FThreadGroup& ThreadGroup = ThreadGroups[GroupName];
 					bIsGroupVisible = ThreadGroup.bIsVisible;
-					ThreadGroup.Order = Order;
+					ThreadGroup.Order = CpuTrackOrder;
 				}
 
 				// Check if there is an available Asset Loading track for this thread.
@@ -301,7 +420,7 @@ void FThreadTimingSharedState::Tick(Timing::ITimingViewSession& InSession, const
 
 						// Create new Timing Events track for the CPU thread.
 						TSharedPtr<FCpuTimingTrack> Track = MakeShared<FCpuTimingTrack>(*this, TrackName, GroupName, CpuTimelineIndex, ThreadInfo.Id);
-						Track->SetOrder(Order);
+						Track->SetOrder(CpuTrackOrder);
 						CpuTracks.Add(ThreadInfo.Id, Track);
 
 						FThreadGroup& ThreadGroup = ThreadGroups[GroupName];
@@ -323,9 +442,9 @@ void FThreadTimingSharedState::Tick(Timing::ITimingViewSession& InSession, const
 					else
 					{
 						TSharedPtr<FCpuTimingTrack> Track = *TrackPtrPtr;
-						if (Track->GetOrder() != Order)
+						if (Track->GetOrder() != CpuTrackOrder)
 						{
-							Track->SetOrder(Order);
+							Track->SetOrder(CpuTrackOrder);
 							bTracksOrderChanged = true;
 						}
 					}
@@ -333,7 +452,7 @@ void FThreadTimingSharedState::Tick(Timing::ITimingViewSession& InSession, const
 
 				constexpr int32 OrderIncrement = FTimingTrackOrder::GroupRange / 1000; // distribute max 1000 tracks in the order group range
 				static_assert(OrderIncrement >= 1, "Order group range too small");
-				Order += OrderIncrement;
+				CpuTrackOrder += OrderIncrement;
 			});
 
 			if (bTracksOrderChanged)
@@ -461,18 +580,21 @@ void FThreadTimingSharedState::SetAllGpuTracksToggle(bool bOnOff)
 {
 	bShowHideAllGpuTracks = bOnOff;
 
-	if (GpuTrack.IsValid())
+	if (OldGpu1Track.IsValid())
 	{
-		GpuTrack->SetVisibilityFlag(bShowHideAllGpuTracks);
+		OldGpu1Track->SetVisibilityFlag(bShowHideAllGpuTracks);
 	}
-	if (Gpu2Track.IsValid())
+	if (OldGpu2Track.IsValid())
 	{
-		Gpu2Track->SetVisibilityFlag(bShowHideAllGpuTracks);
+		OldGpu2Track->SetVisibilityFlag(bShowHideAllGpuTracks);
 	}
-	if (GpuTrack.IsValid() || Gpu2Track.IsValid())
+	for (const auto& KV : GpuTracks)
 	{
-		TimingView->HandleTrackVisibilityChanged();
+		FGpuTimingTrack& Track = *KV.Value;
+		Track.SetVisibilityFlag(bShowHideAllGpuTracks);
 	}
+
+	TimingView->HandleTrackVisibilityChanged();
 }
 
 ////////////////////////////////////////////////////////////////////////////////////////////////////
@@ -539,18 +661,28 @@ TSharedPtr<const ITimingEvent> FThreadTimingSharedState::FindMaxEventInstance(ui
 		}
 	}
 
-	if (GpuTrack.IsValid() && GpuTrack->IsVisible())
+	if (OldGpu1Track.IsValid() && OldGpu1Track->IsVisible())
 	{
-		TrackEvent = GpuTrack->FindMaxEventInstance(TimerId, StartTime, EndTime);
+		TrackEvent = OldGpu1Track->FindMaxEventInstance(TimerId, StartTime, EndTime);
 		CompareAndAssignEvent(TimingEvent, TrackEvent);
 	}
 
-	if (Gpu2Track.IsValid() && Gpu2Track->IsVisible())
+	if (OldGpu2Track.IsValid() && OldGpu2Track->IsVisible())
 	{
-		TrackEvent = Gpu2Track->FindMaxEventInstance(TimerId, StartTime, EndTime);
+		TrackEvent = OldGpu2Track->FindMaxEventInstance(TimerId, StartTime, EndTime);
 		CompareAndAssignEvent(TimingEvent, TrackEvent);
 	}
 
+	for (const auto& KV : GpuTracks)
+	{
+		const FGpuTimingTrack& Track = *KV.Value;
+		if (Track.IsVisible())
+		{
+			TrackEvent = Track.FindMaxEventInstance(TimerId, StartTime, EndTime);
+			CompareAndAssignEvent(TimingEvent, TrackEvent);
+		}
+	}
+
 	return TimingEvent;
 }
 
@@ -584,18 +716,28 @@ TSharedPtr<const ITimingEvent> FThreadTimingSharedState::FindMinEventInstance(ui
 		}
 	}
 
-	if (GpuTrack.IsValid() && GpuTrack->IsVisible())
+	if (OldGpu1Track.IsValid() && OldGpu1Track->IsVisible())
 	{
-		TrackEvent = GpuTrack->FindMinEventInstance(TimerId, StartTime, EndTime);
+		TrackEvent = OldGpu1Track->FindMinEventInstance(TimerId, StartTime, EndTime);
 		CompareAndAssignEvent(TimingEvent, TrackEvent);
 	}
 
-	if (Gpu2Track.IsValid() && Gpu2Track->IsVisible())
+	if (OldGpu2Track.IsValid() && OldGpu2Track->IsVisible())
 	{
-		TrackEvent = Gpu2Track->FindMinEventInstance(TimerId, StartTime, EndTime);
+		TrackEvent = OldGpu2Track->FindMinEventInstance(TimerId, StartTime, EndTime);
 		CompareAndAssignEvent(TimingEvent, TrackEvent);
 	}
 
+	for (const auto& KV : GpuTracks)
+	{
+		const FGpuTimingTrack& Track = *KV.Value;
+		if (Track.IsVisible())
+		{
+			TrackEvent = Track.FindMinEventInstance(TimerId, StartTime, EndTime);
+			CompareAndAssignEvent(TimingEvent, TrackEvent);
+		}
+	}
+
 	return TimingEvent;
 }
 
diff --git a/Engine/Source/Developer/TraceInsights/Private/Insights/TimingProfiler/ViewModels/ThreadTimingSharedState.h b/Engine/Source/Developer/TraceInsights/Private/Insights/TimingProfiler/ViewModels/ThreadTimingSharedState.h
--- a/Engine/Source/Developer/TraceInsights/Private/Insights/TimingProfiler/ViewModels/ThreadTimingSharedState.h
+++ b/Engine/Source/Developer/TraceInsights/Private/Insights/TimingProfiler/ViewModels/ThreadTimingSharedState.h
@@ -43,14 +43,19 @@ public:
 	explicit FThreadTimingSharedState(STimingView* InTimingView);
 	virtual ~FThreadTimingSharedState() override = default;
 
-	TSharedPtr<FGpuTimingTrack> GetGpuTrack() { return GpuTrack; }
-	TSharedPtr<FGpuTimingTrack> GetGpu2Track() { return Gpu2Track; }
+	TSharedPtr<FGpuTimingTrack> GetOldGpu1Track() { return OldGpu1Track; }
+	TSharedPtr<FGpuTimingTrack> GetOldGpu2Track() { return OldGpu2Track; }
+	TSharedPtr<FGpuTimingTrack> GetGpuTrack(uint32 InQueueId);
 	TSharedPtr<FCpuTimingTrack> GetCpuTrack(uint32 InThreadId);
 	const TMap<uint32, TSharedPtr<FCpuTimingTrack>> GetAllCpuTracks() { return CpuTracks; }
 
-	bool IsGpuTrackVisible() const;
+	bool IsOldGpu1TrackVisible() const;
+	bool IsOldGpu2TrackVisible() const;
+	bool IsAnyGpuTrackVisible() const;
+	bool IsGpuTrackVisible(uint32 InQueueId) const;
 	bool IsCpuTrackVisible(uint32 InThreadId) const;
 
+	void GetVisibleGpuQueues(TSet<uint32>& OutSet) const;
 	void GetVisibleCpuThreads(TSet<uint32>& OutSet) const;
 	void GetVisibleTimelineIndexes(TSet<uint32>& OutSet) const;
 
@@ -94,10 +99,13 @@ private:
 	bool bShowHideAllGpuTracks = false;
 	bool bShowHideAllCpuTracks = false;
 
-	TSharedPtr<FGpuTimingTrack> GpuTrack;
-	TSharedPtr<FGpuTimingTrack> Gpu2Track;
+	TSharedPtr<FGpuTimingTrack> OldGpu1Track;
+	TSharedPtr<FGpuTimingTrack> OldGpu2Track;
 
-	/** Maps thread id to track pointer. */
+	/** Maps GPU queue id to track pointer. */
+	TMap<uint32, TSharedPtr<FGpuTimingTrack>> GpuTracks;
+
+	/** Maps CPU thread id to track pointer. */
 	TMap<uint32, TSharedPtr<FCpuTimingTrack>> CpuTracks;
 
 	/** Maps thread group name to thread group info. */
diff --git a/Engine/Source/Developer/TraceInsights/Private/Insights/TimingProfiler/ViewModels/TimerAggregator.cpp b/Engine/Source/Developer/TraceInsights/Private/Insights/TimingProfiler/ViewModels/TimerAggregator.cpp
--- a/Engine/Source/Developer/TraceInsights/Private/Insights/TimingProfiler/ViewModels/TimerAggregator.cpp
+++ b/Engine/Source/Developer/TraceInsights/Private/Insights/TimingProfiler/ViewModels/TimerAggregator.cpp
@@ -18,12 +18,20 @@ namespace UE::Insights::TimingProfiler
 class FTimerAggregationWorker : public IStatsAggregationWorker
 {
 public:
-	FTimerAggregationWorker(TSharedPtr<const TraceServices::IAnalysisSession> InSession, double InStartTime, double InEndTime, const TSet<uint32>& InCpuThreads, bool bInIncludeGpuThread, ETraceFrameType InFrameType)
+	FTimerAggregationWorker(TSharedPtr<const TraceServices::IAnalysisSession> InSession,
+							double InStartTime, double InEndTime,
+							bool bInIncludeOldGpu1,
+							bool bInIncludeOldGpu2,
+							const TSet<uint32>& InGpuQueues,
+							const TSet<uint32>& InCpuThreads,
+							ETraceFrameType InFrameType)
 		: Session(InSession)
 		, StartTime(InStartTime)
 		, EndTime(InEndTime)
+		, bIncludeOldGpu1(bInIncludeOldGpu1)
+		, bIncludeOldGpu2(bInIncludeOldGpu2)
+		, GpuQueues(InGpuQueues)
 		, CpuThreads(InCpuThreads)
-		, bIncludeGpuThread(bInIncludeGpuThread)
 		, FrameType(InFrameType)
 		, ResultTable()
 	{
@@ -40,8 +48,10 @@ private:
 	TSharedPtr<const TraceServices::IAnalysisSession> Session;
 	double StartTime;
 	double EndTime;
+	bool bIncludeOldGpu1;
+	bool bIncludeOldGpu2;
+	TSet<uint32> GpuQueues;
 	TSet<uint32> CpuThreads;
-	bool bIncludeGpuThread;
 	ETraceFrameType FrameType;
 	TUniquePtr<TraceServices::ITable<TraceServices::FTimingProfilerAggregatedStats>> ResultTable;
 };
@@ -52,22 +62,16 @@ void FTimerAggregationWorker::DoWork(TSharedPtr<TraceServices::FCancellationToke
 {
 	if (Session.IsValid() && TraceServices::ReadTimingProfilerProvider(*Session.Get()))
 	{
-		// Suspend analysis in order to avoid write locks (ones blocked by the read lock below) to further block other read locks.
-		//FAnalysisSessionSuspensionScope SessionPauseScope(*Session.Get());
-
 		TraceServices::FAnalysisSessionReadScope SessionReadScope(*Session.Get());
 		const TraceServices::ITimingProfilerProvider& TimingProfilerProvider = *TraceServices::ReadTimingProfilerProvider(*Session.Get());
 
-		auto CpuThreadFilter = [this](uint32 ThreadId)
-		{
-			return CpuThreads.Contains(ThreadId);
-		};
-
-		TraceServices::FCreateAggreationParams Params;
+		TraceServices::FCreateAggregationParams Params;
 		Params.IntervalStart = StartTime;
 		Params.IntervalEnd = EndTime;
-		Params.CpuThreadFilter = CpuThreadFilter;
-		Params.IncludeGpu = bIncludeGpuThread;
+		Params.bIncludeOldGpu1 = bIncludeOldGpu1;
+		Params.bIncludeOldGpu2 = bIncludeOldGpu2;
+		Params.GpuQueueFilter = [this](uint32 QueueId) { return GpuQueues.Contains(QueueId); };
+		Params.CpuThreadFilter = [this](uint32 ThreadId) { return CpuThreads.Contains(ThreadId); };
 		Params.FrameType = FrameType;
 		Params.CancellationToken = CancellationToken;
 
@@ -81,7 +85,9 @@ void FTimerAggregationWorker::DoWork(TSharedPtr<TraceServices::FCancellationToke
 
 IStatsAggregationWorker* FTimerAggregator::CreateWorker(TSharedPtr<const TraceServices::IAnalysisSession> InSession)
 {
-	bool bIsGpuTrackVisible = false;
+	bool bOldGpu1TrackVisible = false;
+	bool bOldGpu2TrackVisible = false;
+	TSet<uint32> GpuQueues;
 	TSet<uint32> CpuThreads;
 
 	TSharedPtr<STimingProfilerWindow> Wnd = FTimingProfilerManager::Get()->GetProfilerWindow();
@@ -90,17 +96,18 @@ IStatsAggregationWorker* FTimerAggregator::CreateWorker(TSharedPtr<const TraceSe
 		TSharedPtr<STimingView> TimingView = Wnd->GetTimingView();
 		if (TimingView.IsValid())
 		{
-			bIsGpuTrackVisible = TimingView->IsGpuTrackVisible();
-
 			TSharedPtr<FThreadTimingSharedState> ThreadTimingSharedState = TimingView->GetThreadTimingSharedState();
 			if (ThreadTimingSharedState.IsValid())
 			{
+				bOldGpu1TrackVisible = ThreadTimingSharedState->IsOldGpu1TrackVisible();
+				bOldGpu2TrackVisible = ThreadTimingSharedState->IsOldGpu2TrackVisible();
+				ThreadTimingSharedState->GetVisibleGpuQueues(GpuQueues);
 				ThreadTimingSharedState->GetVisibleCpuThreads(CpuThreads);
 			}
 		}
 	}
 
-	return new FTimerAggregationWorker(InSession, GetIntervalStartTime(), GetIntervalEndTime(), CpuThreads, bIsGpuTrackVisible, FrameType);
+	return new FTimerAggregationWorker(InSession, GetIntervalStartTime(), GetIntervalEndTime(), bOldGpu1TrackVisible, bOldGpu2TrackVisible, GpuQueues, CpuThreads, FrameType);
 }
 
 ////////////////////////////////////////////////////////////////////////////////////////////////////
diff --git a/Engine/Source/Developer/TraceInsights/Private/Insights/TimingProfiler/ViewModels/TimerButterflyAggregator.cpp b/Engine/Source/Developer/TraceInsights/Private/Insights/TimingProfiler/ViewModels/TimerButterflyAggregator.cpp
--- a/Engine/Source/Developer/TraceInsights/Private/Insights/TimingProfiler/ViewModels/TimerButterflyAggregator.cpp
+++ b/Engine/Source/Developer/TraceInsights/Private/Insights/TimingProfiler/ViewModels/TimerButterflyAggregator.cpp
@@ -23,12 +23,20 @@ namespace UE::Insights::TimingProfiler
 class FTimerButterflyAggregationWorker : public IStatsAggregationWorker
 {
 public:
-	FTimerButterflyAggregationWorker(TSharedPtr<const TraceServices::IAnalysisSession> InSession, double InStartTime, double InEndTime, const TSet<uint32>& InCpuThreads, bool bInIncludeGpuThread)
+	FTimerButterflyAggregationWorker(TSharedPtr<const TraceServices::IAnalysisSession> InSession,
+									 double InStartTime,
+									 double InEndTime,
+									 bool bInIncludeOldGpu1,
+									 bool bInIncludeOldGpu2,
+									 const TSet<uint32>& InGpuQueues,
+									 const TSet<uint32>& InCpuThreads)
 		: Session(InSession)
 		, StartTime(InStartTime)
 		, EndTime(InEndTime)
+		, bIncludeOldGpu1(bInIncludeOldGpu1)
+		, bIncludeOldGpu2(bInIncludeOldGpu2)
+		, GpuQueues(InGpuQueues)
 		, CpuThreads(InCpuThreads)
-		, bIncludeGpuThread(bInIncludeGpuThread)
 		, ResultButterfly()
 	{
 	}
@@ -44,8 +52,10 @@ private:
 	TSharedPtr<const TraceServices::IAnalysisSession> Session;
 	double StartTime;
 	double EndTime;
+	bool bIncludeOldGpu1;
+	bool bIncludeOldGpu2;
+	TSet<uint32> GpuQueues;
 	TSet<uint32> CpuThreads;
-	bool bIncludeGpuThread;
 	TUniquePtr<TraceServices::ITimingProfilerButterfly> ResultButterfly;
 };
 
@@ -55,31 +65,29 @@ void FTimerButterflyAggregationWorker::DoWork(TSharedPtr<TraceServices::FCancell
 {
 	if (Session.IsValid() && TraceServices::ReadTimingProfilerProvider(*Session.Get()))
 	{
-		// Suspend analysis in order to avoid write locks (ones blocked by the read lock below) to further block other read locks.
-		//FAnalysisSessionSuspensionScope SessionPauseScope(*Session.Get());
-
 		TraceServices::FAnalysisSessionReadScope SessionReadScope(*Session.Get());
 		const TraceServices::ITimingProfilerProvider& TimingProfilerProvider = *TraceServices::ReadTimingProfilerProvider(*Session.Get());
 
 		if (StartTime <= Session->GetDurationSeconds())
 		{
-			auto ThreadFilter = [this](uint32 ThreadId)
-			{
-				return CpuThreads.Contains(ThreadId);
-			};
-
-			double ClampedEndTime = FMath::Min(Session->GetDurationSeconds(), EndTime);
-			ResultButterfly.Reset(TimingProfilerProvider.CreateButterfly(StartTime, ClampedEndTime, ThreadFilter, bIncludeGpuThread));
+			TraceServices::FCreateButterflyParams Params;
+			Params.IntervalStart = StartTime;
+			Params.IntervalEnd = FMath::Min(Session->GetDurationSeconds(), EndTime);
+			Params.bIncludeOldGpu1 = bIncludeOldGpu1;
+			Params.bIncludeOldGpu2 = bIncludeOldGpu2;
+			Params.GpuQueueFilter = [this](uint32 QueueId) { return GpuQueues.Contains(QueueId); };
+			Params.CpuThreadFilter = [this](uint32 ThreadId) { return CpuThreads.Contains(ThreadId); };
+
+			ResultButterfly.Reset(TimingProfilerProvider.CreateButterfly(Params));
 		}
 		else
 		{
-			// Exclude all threads to receive empty results in case infinite events are in the session.
-			auto ThreadFilter = [this](uint32 ThreadId)
-			{
-				return false;
-			};
+			// Exclude all GPU and CPU timelines, to receive empty results (in case infinite events are in the session).
+			TraceServices::FCreateButterflyParams Params;
+			Params.IntervalStart = StartTime;
+			Params.IntervalEnd = EndTime;
 
-			ResultButterfly.Reset(TimingProfilerProvider.CreateButterfly(StartTime, EndTime, ThreadFilter, false));
+			ResultButterfly.Reset(TimingProfilerProvider.CreateButterfly(Params));
 		}
 	}
 }
@@ -90,7 +98,9 @@ void FTimerButterflyAggregationWorker::DoWork(TSharedPtr<TraceServices::FCancell
 
 IStatsAggregationWorker* FTimerButterflyAggregator::CreateWorker(TSharedPtr<const TraceServices::IAnalysisSession> InSession)
 {
-	bool bIsGpuTrackVisible = false;
+	bool bIsOldGpu1TrackVisible = false;
+	bool bIsOldGpu2TrackVisible = false;
+	TSet<uint32> GpuQueues;
 	TSet<uint32> CpuThreads;
 
 	TSharedPtr<STimingProfilerWindow> Wnd = FTimingProfilerManager::Get()->GetProfilerWindow();
@@ -99,17 +109,18 @@ IStatsAggregationWorker* FTimerButterflyAggregator::CreateWorker(TSharedPtr<cons
 		TSharedPtr<STimingView> TimingView = Wnd->GetTimingView();
 		if (TimingView.IsValid())
 		{
-			bIsGpuTrackVisible = TimingView->IsGpuTrackVisible();
-
 			TSharedPtr<FThreadTimingSharedState> ThreadTimingSharedState = TimingView->GetThreadTimingSharedState();
 			if (ThreadTimingSharedState.IsValid())
 			{
+				bIsOldGpu1TrackVisible = ThreadTimingSharedState->IsOldGpu1TrackVisible();
+				bIsOldGpu2TrackVisible = ThreadTimingSharedState->IsOldGpu2TrackVisible();
+				ThreadTimingSharedState->GetVisibleGpuQueues(GpuQueues);
 				ThreadTimingSharedState->GetVisibleCpuThreads(CpuThreads);
 			}
 		}
 	}
 
-	return new FTimerButterflyAggregationWorker(InSession, GetIntervalStartTime(), GetIntervalEndTime(), CpuThreads, bIsGpuTrackVisible);
+	return new FTimerButterflyAggregationWorker(InSession, GetIntervalStartTime(), GetIntervalEndTime(), bIsOldGpu1TrackVisible, bIsOldGpu2TrackVisible, GpuQueues, CpuThreads);
 }
 
 ////////////////////////////////////////////////////////////////////////////////////////////////////
diff --git a/Engine/Source/Developer/TraceInsights/Private/Insights/TimingProfiler/ViewModels/TimingExporter.cpp b/Engine/Source/Developer/TraceInsights/Private/Insights/TimingProfiler/ViewModels/TimingExporter.cpp
--- a/Engine/Source/Developer/TraceInsights/Private/Insights/TimingProfiler/ViewModels/TimingExporter.cpp
+++ b/Engine/Source/Developer/TraceInsights/Private/Insights/TimingProfiler/ViewModels/TimingExporter.cpp
@@ -766,14 +766,16 @@ int32 FTimingExporter::ExportTimerStatisticsAsText(const FString& Filename, FExp
 			return -1;
 		}
 
-		TraceServices::FCreateAggreationParams CreateAggreagationParams;
-		CreateAggreagationParams.IntervalStart = Params.IntervalStartTime;
-		CreateAggreagationParams.IntervalEnd = Params.IntervalEndTime;
-		CreateAggreagationParams.CpuThreadFilter = Params.ThreadFilter;
-		CreateAggreagationParams.IncludeGpu = true;
+		TraceServices::FCreateAggregationParams CreateAggregationParams;
+		CreateAggregationParams.IntervalStart = Params.IntervalStartTime;
+		CreateAggregationParams.IntervalEnd = Params.IntervalEndTime;
+		CreateAggregationParams.bIncludeOldGpu1 = true;
+		CreateAggregationParams.bIncludeOldGpu2 = true;
+		CreateAggregationParams.GpuQueueFilter = [](uint32) { return true; };
+		CreateAggregationParams.CpuThreadFilter = Params.ThreadFilter;
 
 		//@Todo: this does not yet handle the -column and -timers parameters.
-		StatsTable = TimingProfilerProvider->CreateAggregation(CreateAggreagationParams);
+		StatsTable = TimingProfilerProvider->CreateAggregation(CreateAggregationParams);
 	}
 
 	FStopwatch Stopwatch;
@@ -849,15 +851,16 @@ int32 FTimingExporter::ExportTimerCalleesAsText(const FString& Filename, const F
 			return nullptr;
 		}
 
-		const bool bIncludeGpu = Params.ThreadFilter(FGpuTimingTrack::Gpu1ThreadId) || Params.ThreadFilter(FGpuTimingTrack::Gpu2ThreadId);
-	
+		TraceServices::FCreateButterflyParams ButterflyParams;
+		ButterflyParams.IntervalStart = Params.IntervalStartTime;
 		// The region end interval may be inf if the capture ended before the region was closed.
-		const double ClampedEndTime = FMath::Min(Session.GetDurationSeconds(), Params.IntervalEndTime);
+		ButterflyParams.IntervalEnd = FMath::Min(Session.GetDurationSeconds(), Params.IntervalEndTime);
+		ButterflyParams.bIncludeOldGpu1 = Params.ThreadFilter(FGpuTimingTrack::Gpu1ThreadId);
+		ButterflyParams.bIncludeOldGpu2 = Params.ThreadFilter(FGpuTimingTrack::Gpu2ThreadId);
+		ButterflyParams.GpuQueueFilter = [](uint32) { return true; };
+		ButterflyParams.CpuThreadFilter = Params.ThreadFilter;
 
-		return TUniquePtr<TraceServices::ITimingProfilerButterfly>
-		{
-			TimingProfilerProvider->CreateButterfly(Params.IntervalStartTime, ClampedEndTime, Params.ThreadFilter, bIncludeGpu)
-		};
+		return TUniquePtr<TraceServices::ITimingProfilerButterfly>{ TimingProfilerProvider->CreateButterfly(ButterflyParams) };
 	}();
 
 	if (!Butterfly)
diff --git a/Engine/Source/Developer/TraceInsights/Private/Insights/Widgets/STimingView.cpp b/Engine/Source/Developer/TraceInsights/Private/Insights/Widgets/STimingView.cpp
--- a/Engine/Source/Developer/TraceInsights/Private/Insights/Widgets/STimingView.cpp
+++ b/Engine/Source/Developer/TraceInsights/Private/Insights/Widgets/STimingView.cpp
@@ -560,9 +560,30 @@ PRAGMA_ENABLE_DEPRECATION_WARNINGS
 
 ////////////////////////////////////////////////////////////////////////////////////////////////////
 
-bool STimingView::IsGpuTrackVisible() const
+bool STimingView::IsOldGpu1TrackVisible() const
 {
-	return ThreadTimingSharedState && ThreadTimingSharedState->IsGpuTrackVisible();
+	return ThreadTimingSharedState && ThreadTimingSharedState->IsOldGpu1TrackVisible();
+}
+
+////////////////////////////////////////////////////////////////////////////////////////////////////
+
+bool STimingView::IsOldGpu2TrackVisible() const
+{
+	return ThreadTimingSharedState && ThreadTimingSharedState->IsOldGpu2TrackVisible();
+}
+
+////////////////////////////////////////////////////////////////////////////////////////////////////
+
+bool STimingView::IsAnyGpuTrackVisible() const
+{
+	return ThreadTimingSharedState && ThreadTimingSharedState->IsAnyGpuTrackVisible();
+}
+
+////////////////////////////////////////////////////////////////////////////////////////////////////
+
+bool STimingView::IsGpuTrackVisible(uint32 InQueueId) const
+{
+	return ThreadTimingSharedState && ThreadTimingSharedState->IsGpuTrackVisible(InQueueId);
 }
 
 ////////////////////////////////////////////////////////////////////////////////////////////////////
diff --git a/Engine/Source/Developer/TraceInsights/Private/Insights/Widgets/STimingView.h b/Engine/Source/Developer/TraceInsights/Private/Insights/Widgets/STimingView.h
--- a/Engine/Source/Developer/TraceInsights/Private/Insights/Widgets/STimingView.h
+++ b/Engine/Source/Developer/TraceInsights/Private/Insights/Widgets/STimingView.h
@@ -357,7 +357,10 @@ public:
 
 	void HandleTrackVisibilityChanged();
 
-	bool IsGpuTrackVisible() const;
+	bool IsOldGpu1TrackVisible() const;
+	bool IsOldGpu2TrackVisible() const;
+	bool IsAnyGpuTrackVisible() const;
+	bool IsGpuTrackVisible(uint32 InQueueId) const;
 	bool IsCpuTrackVisible(uint32 InThreadId) const;
 
 	////////////////////////////////////////////////////////////////////////////////////////////////////
@@ -544,7 +547,7 @@ PRAGMA_ENABLE_DEPRECATION_WARNINGS
 	FReply AllowTracksToProcessOnMouseButtonDoubleClick(const FGeometry& MyGeometry, const FPointerEvent& MouseEvent);
 
 	void SetTrackPosY(TSharedPtr<FBaseTimingTrack>& TrackPtr, float TrackPosY) const;
-	
+
 	void FindFirstEvent();
 	void FindPrevEvent();
 	void FindNextEvent();
@@ -560,10 +563,10 @@ PRAGMA_ENABLE_DEPRECATION_WARNINGS
 	void EnumerateFilteredTracks(TSharedPtr<UE::Insights::FFilterConfigurator> FilterConfigurator, TSharedPtr<const FBaseTimingTrack> PriorityTrack, EnumerateFilteredTracksCallback Callback);
 
 	ETraceFrameType GetFrameTypeToSnapTo();
-	
+
 	/** The FilterConfigurator needs to be updated so that custom filters work correctly when analysis is still running
-	* and the timer list can change. 
-	 */ 
+	  * and the timer list can change.
+	  */
 	void UpdateFilters();
 
 	bool IsInTimingProfiler() const;
@@ -608,7 +611,7 @@ protected:
 
 	// Shared state for Regions tracks
 	TSharedPtr<FTimingRegionsSharedState> TimingRegionsSharedState;
-	
+
 	////////////////////////////////////////////////////////////
 
 	/** The time ruler track. It includes the custom time markers (ones user can drag with mouse). */
@@ -788,7 +791,7 @@ protected:
 	static uint32 TimingViewId;
 	const FName QuickFindTabId;
 	bool bUpdateFilters = true;
-	
+
 	// Used only between the creation of the widget and the spawning of the owning tab. When the tab is spawned, we relinquish ownership.
 	TSharedPtr<::Insights::SQuickFind> QuickFindWidgetSharedPtr;
 	TWeakPtr<::Insights::SQuickFind> QuickFindWidgetWeakPtr;
diff --git a/Engine/Source/Developer/TraceServices/Private/Analyzers/CpuProfilerTraceAnalysis.cpp b/Engine/Source/Developer/TraceServices/Private/Analyzers/CpuProfilerTraceAnalysis.cpp
--- a/Engine/Source/Developer/TraceServices/Private/Analyzers/CpuProfilerTraceAnalysis.cpp
+++ b/Engine/Source/Developer/TraceServices/Private/Analyzers/CpuProfilerTraceAnalysis.cpp
@@ -497,9 +497,9 @@ void FCpuProfilerAnalyzer::ProcessBufferV2(const FEventTime& EventTime, FThreadS
 					if (ThreadState.ScopeStack.Num() > 0)
 					{
 						uint32 MetadataTimerId = ThreadState.ScopeStack.Top().EventTypeId;
-						TArrayView<const uint8> Metadata = EditableTimingProfilerProvider.GetMetadata(MetadataTimerId);
+						TArrayView<uint8> Metadata = EditableTimingProfilerProvider.GetEditableMetadata(MetadataTimerId);
 						// Change the last byte in metadata to "true".
-						const_cast<uint8*>(Metadata.GetData())[Metadata.Num() - 1] = (uint8)(ECborCode::Prim | ECborCode::True);
+						Metadata.GetData()[Metadata.Num() - 1] = (uint8)(ECborCode::Prim | ECborCode::True);
 					}
 				}
 
diff --git a/Engine/Source/Developer/TraceServices/Private/Analyzers/GpuProfilerTraceAnalysis.cpp b/Engine/Source/Developer/TraceServices/Private/Analyzers/GpuProfilerTraceAnalysis.cpp
--- a/Engine/Source/Developer/TraceServices/Private/Analyzers/GpuProfilerTraceAnalysis.cpp
+++ b/Engine/Source/Developer/TraceServices/Private/Analyzers/GpuProfilerTraceAnalysis.cpp
@@ -3,8 +3,10 @@
 #include "GpuProfilerTraceAnalysis.h"
 
 #include "AnalysisServicePrivate.h"
+#include "CborWriter.h"
 #include "Common/Utils.h"
 #include "HAL/LowLevelMemTracker.h"
+#include "Serialization/MemoryWriter.h"
 
 namespace TraceServices
 {
@@ -19,6 +21,18 @@ void FGpuProfilerAnalyzer::OnAnalysisBegin(const FOnAnalysisContext& Context)
 {
 	auto& Builder = Context.InterfaceBuilder;
 
+	// The New GPU Profiler
+	Builder.RouteEvent(RouteId_Init, "GpuProfiler", "Init");
+	Builder.RouteEvent(RouteId_QueueSpec, "GpuProfiler", "QueueSpec");
+	Builder.RouteEvent(RouteId_EventFrameBoundary, "GpuProfiler", "EventFrameBoundary");
+	Builder.RouteEvent(RouteId_EventBeginBreadcrumb, "GpuProfiler", "EventBeginBreadcrumb");
+	Builder.RouteEvent(RouteId_EventEndBreadcrumb, "GpuProfiler", "EventEndBreadcrumb");
+	Builder.RouteEvent(RouteId_EventBeginWork, "GpuProfiler", "EventBeginWork");
+	Builder.RouteEvent(RouteId_EventEndWork, "GpuProfiler", "EventEndWork");
+	Builder.RouteEvent(RouteId_EventStats, "GpuProfiler", "EventStats");
+
+	// The Old GPU Profiler (deprecated in UE 5.6)
+	// Analysis code is maintained for backward compatibility with old traces.
 	Builder.RouteEvent(RouteId_EventSpec, "GpuProfiler", "EventSpec");
 	Builder.RouteEvent(RouteId_Frame, "GpuProfiler", "Frame");
 	Builder.RouteEvent(RouteId_Frame2, "GpuProfiler", "Frame2");
@@ -41,12 +55,139 @@ bool FGpuProfilerAnalyzer::OnEvent(uint16 RouteId, EStyle Style, const FOnEventC
 {
 	LLM_SCOPE_BYNAME(TEXT("Insights/FGpuProfilerAnalyzer"));
 
-	FAnalysisSessionEditScope _(Session);
-
 	const auto& EventData = Context.EventData;
 
 	switch (RouteId)
 	{
+
+	//////////////////////////////////////////////////
+	// The New GPU Profiler
+
+	case RouteId_Init:
+	{
+		Version = (uint32)EventData.GetValue<uint8>("Version");
+		UE_LOG(LogTraceServices, Log, TEXT("[GPU] Init Version=%u"), Version);
+		break;
+	}
+	case RouteId_QueueSpec:
+	{
+		const uint32 QueueId = EventData.GetValue<uint32>("QueueId");
+
+		const uint8 GPU = (QueueId >> 8) & 0xFF;
+		const uint8 Index = (QueueId >> 16) & 0xFF;
+		const uint8 Type = QueueId & 0xFF;
+
+		FString Name;
+		EventData.GetString("TypeString", Name);
+
+		UE_LOG(LogTraceServices, Log, TEXT("[GPU] QueueSpec QueueId=%u, GPU=%u, Index=%u, Type=%d \"%s\""), QueueId, (uint32)GPU, (uint32)Index, (uint32)Type, *Name);
+
+		FAnalysisSessionEditScope _(Session);
+		const TCHAR* PersistentName = Session.StoreString(Name);
+
+		TimingProfilerProvider.AddGpuQueue(QueueId, GPU, Index, Type, PersistentName);
+		break;
+	}
+	case RouteId_EventFrameBoundary:
+	{
+		const uint32 QueueId = EventData.GetValue<uint32>("QueueId");
+		const uint32 FrameNumber = EventData.GetValue<uint32>("FrameNumber");
+
+		UE_LOG(LogTraceServices, VeryVerbose, TEXT("[GPU] EventFrameBoundary QueueId=%u, FrameNumber=%u"), QueueId, FrameNumber);
+
+		FQueue& Queue = GetOrAddQueue(QueueId);
+		Queue.FrameNumber = FrameNumber;
+		break;
+	}
+	case RouteId_EventBeginBreadcrumb:
+	{
+		const uint32 QueueId = EventData.GetValue<uint32>("QueueId");
+		const uint64 GPUTimestampTOP = EventData.GetValue<uint64>("GPUTimestampTOP");
+		const double Time = Context.EventTime.AsSeconds(GPUTimestampTOP);
+
+		FString Name;
+		EventData.GetString("Name", Name);
+
+		UE_LOG(LogTraceServices, VeryVerbose, TEXT("[GPU] EventBeginBreadcrumb QueueId=%u, Time=%f, Name=\"%s\""), QueueId, Time, *Name);
+
+		FQueue& Queue = GetOrAddQueue(QueueId);
+		const uint32 TimerId = GetOrAddTimerId(Name);
+		BeginEvent(Queue, 0, Time, TimerId);
+		break;
+	}
+	case RouteId_EventEndBreadcrumb:
+	{
+		const uint32 QueueId = EventData.GetValue<uint32>("QueueId");
+		const uint64 GPUTimestampBOP = EventData.GetValue<uint64>("GPUTimestampBOP");
+		const double Time = Context.EventTime.AsSeconds(GPUTimestampBOP);
+
+		FString Name;
+		EventData.GetString("Name", Name);
+
+		UE_LOG(LogTraceServices, VeryVerbose, TEXT("[GPU] EventEndBreadcrumb QueueId=%u, Time=%f, Name=\"%s\""), QueueId, Time, *Name);
+
+		FQueue& Queue = GetOrAddQueue(QueueId);
+		const uint32 TimerId = GetOrAddTimerId(Name);
+		EndEvent(Queue, 0, Time, TimerId);
+		break;
+	}
+	case RouteId_EventBeginWork:
+	{
+		const uint32 QueueId = EventData.GetValue<uint32>("QueueId");
+		const uint64 GPUTimestampTOP = EventData.GetValue<uint64>("GPUTimestampTOP");
+		double Time = Context.EventTime.AsSeconds(GPUTimestampTOP);
+		//const uint64 CPUTimestamp = EventData.GetValue<uint64>("CPUTimestamp"); // not used
+
+		UE_LOG(LogTraceServices, VeryVerbose, TEXT("[GPU] EventBeginWork QueueId=%u, GPUTimestampTOP=%f"), QueueId, Time);
+
+		if (GpuWorkTimerId == ~0u)
+		{
+			const FString GpuWorkTimerName(TEXT("GpuWork"));
+			GpuWorkTimerId = GetOrAddTimerId(GpuWorkTimerName);
+		}
+
+		FQueue& Queue = GetOrAddQueue(QueueId);
+		TArray<uint8> CborData;
+		{
+			CborData.Reserve(256);
+			FMemoryWriter MemoryWriter(CborData, false, true);
+			FCborWriter CborWriter(&MemoryWriter, ECborEndianness::StandardCompliant);
+			CborWriter.WriteContainerStart(ECborCode::Map, 1); // 1 is the FieldCount
+			CborWriter.WriteValue("Frame", 5);
+			CborWriter.WriteValue((int64)Queue.FrameNumber);
+		}
+		uint32 MetadataTimerId;
+		{
+			FAnalysisSessionEditScope _(Session);
+			MetadataTimerId = TimingProfilerProvider.AddMetadata(GpuWorkTimerId, MoveTemp(CborData));
+		}
+		BeginEvent(Queue, 1, Time, MetadataTimerId);
+		break;
+	}
+	case RouteId_EventEndWork:
+	{
+		uint32 QueueId = EventData.GetValue<uint32>("QueueId");
+		const uint64 GPUTimestampBOP = EventData.GetValue<uint64>("GPUTimestampBOP");
+		const double Time = Context.EventTime.AsSeconds(GPUTimestampBOP);
+
+		UE_LOG(LogTraceServices, VeryVerbose, TEXT("[GPU] EventEndWork QueueId=%u, GPUTimestampBOP=%f"), QueueId, Time);
+
+		FQueue& Queue = GetOrAddQueue(QueueId);
+		EndEvent(Queue, 1, Time, GpuWorkTimerId);
+		break;
+	}
+	case RouteId_EventStats:
+	{
+		const uint32 QueueId = EventData.GetValue<uint32>("QueueId");
+		const uint32 NumDraws = EventData.GetValue<uint32>("NumDraws");
+		const uint32 NumPrimitives = EventData.GetValue<uint32>("NumPrimitives");
+		//TODO: counters
+		break;
+	}
+
+	//////////////////////////////////////////////////
+	// The Old GPU Profiler
+
 	case RouteId_EventSpec:
 	{
 		uint32 EventType = EventData.GetValue<uint32>("EventType");
@@ -56,10 +197,12 @@ bool FGpuProfilerAnalyzer::OnEvent(uint16 RouteId, EStyle Style, const FOnEventC
 		uint32* TimerIndexPtr = EventTypeMap.Find(EventType);
 		if (!TimerIndexPtr)
 		{
+			FAnalysisSessionEditScope _(Session);
 			EventTypeMap.Add(EventType, TimingProfilerProvider.AddGpuTimer(FStringView(NameTChar.Get(), NameTChar.Length())));
 		}
 		else
 		{
+			FAnalysisSessionEditScope _(Session);
 			TimingProfilerProvider.SetTimerName(*TimerIndexPtr, FStringView(NameTChar.Get(), NameTChar.Length()));
 		}
 		break;
@@ -67,6 +210,8 @@ bool FGpuProfilerAnalyzer::OnEvent(uint16 RouteId, EStyle Style, const FOnEventC
 	case RouteId_Frame:
 	case RouteId_Frame2:
 	{
+		FAnalysisSessionEditScope _(Session);
+
 		TraceServices::FTimingProfilerProvider::TimelineInternal& ThisTimeline = (RouteId == RouteId_Frame) ?
 			TimingProfilerProvider.EditGpuTimeline() :
 			TimingProfilerProvider.EditGpu2Timeline();
@@ -150,9 +295,129 @@ bool FGpuProfilerAnalyzer::OnEvent(uint16 RouteId, EStyle Style, const FOnEventC
 		Session.UpdateDurationSeconds(LastTime);
 		break;
 	}
-	}
+
+	//////////////////////////////////////////////////
+
+	} // switch (RouteId)
 
 	return true;
 }
 
+uint32 FGpuProfilerAnalyzer::GetOrAddTimerId(const FString& Breadcrumb)
+{
+	const uint32* FoundTimerId = BreadcrumbMap.Find(Breadcrumb);
+	if (FoundTimerId)
+	{
+		return *FoundTimerId;
+	}
+	else
+	{
+		FAnalysisSessionEditScope _(Session);
+		uint32 TimerId = TimingProfilerProvider.AddGpuTimer(Breadcrumb);
+		const TCHAR* TimerName = nullptr;
+		TimingProfilerProvider.ReadTimers([TimerId, &TimerName](const ITimingProfilerTimerReader& TimerReader)
+			{
+				const FTimingProfilerTimer* Timer = TimerReader.GetTimer(TimerId);
+				if (Timer)
+				{
+					TimerName = Timer->Name;
+				}
+			});
+		BreadcrumbMap.Add(Breadcrumb, TimerId);
+		TimerMap.Add(TimerId, TimerName);
+		return TimerId;
+	}
+}
+
+const TCHAR* FGpuProfilerAnalyzer::GetTimerName(uint32 TimerId) const
+{
+	const TCHAR* const* TimerNamePtr = TimerMap.Find(TimerId);
+	return TimerNamePtr ? *TimerNamePtr : TEXT("<unknown>");
+}
+
+FGpuProfilerAnalyzer::FQueue& FGpuProfilerAnalyzer::GetOrAddQueue(uint32 QueueId)
+{
+	FQueue* FoundQueue = Queues.Find(QueueId);
+	if (FoundQueue)
+	{
+		return *FoundQueue;
+	}
+	FQueue& NewQueue = Queues.Add(QueueId);
+	NewQueue.Id = QueueId;
+	return NewQueue;
+}
+
+void FGpuProfilerAnalyzer::BeginEvent(FQueue& Queue, int32 StackIndex, double BeginEventTime, uint32 BeginEventTimerId)
+{
+	Queue.Stack[StackIndex].Push({ BeginEventTime, BeginEventTimerId });
+
+	FAnalysisSessionEditScope _(Session);
+	IEditableTimeline<FTimingProfilerEvent>* Timeline = (StackIndex == 1) ?
+		TimingProfilerProvider.GetGpuQueueWorkEditableTimeline(Queue.Id) :
+		TimingProfilerProvider.GetGpuQueueEditableTimeline(Queue.Id);
+	if (ensure(Timeline))
+	{
+		FTimingProfilerEvent Event;
+		Event.TimerIndex = BeginEventTimerId;
+		const double LastTimestamp = Timeline->GetLastTimestamp();
+		if (BeginEventTime < LastTimestamp)
+		{
+			UE_LOG(LogTraceServices, Error, TEXT("[GPU] %s BEGIN %u \"%s\" : Time %f < %f !!!"),
+				StackIndex == 1 ? TEXT("WORK") : TEXT("BREADCRUMB"),
+				BeginEventTimerId, GetTimerName(BeginEventTimerId),
+				BeginEventTime, LastTimestamp);
+			BeginEventTime = LastTimestamp;
+		}
+		Timeline->AppendBeginEvent(BeginEventTime, Event);
+		Session.UpdateDurationSeconds(BeginEventTime);
+	}
+}
+
+void FGpuProfilerAnalyzer::EndEvent(FQueue& Queue, int32 StackIndex, double EndEventTime, uint32 EndEventTimerId)
+{
+	if (Queue.Stack[StackIndex].IsEmpty())
+	{
+		return;
+	}
+	FOpenEvent BeginEvent = Queue.Stack[StackIndex].Pop(EAllowShrinking::No);
+	uint32 BeginEventTimerId;
+	{
+		FAnalysisSessionReadScope _(Session);
+		BeginEventTimerId = TimingProfilerProvider.GetOriginalTimerIdFromMetadata(BeginEvent.TimerId);
+	}
+	if (BeginEventTimerId != EndEventTimerId)
+	{
+		UE_LOG(LogTraceServices, Error, TEXT("[GPU] %s END %u \"%s\" (%f) != BEGIN %u \"%s\" (%f) !!!"),
+			StackIndex == 1 ? TEXT("WORK") : TEXT("BREADCRUMB"),
+			EndEventTimerId, GetTimerName(EndEventTimerId), EndEventTime,
+			BeginEventTimerId, GetTimerName(BeginEventTimerId), BeginEvent.Time);
+	}
+	if (BeginEvent.Time > EndEventTime)
+	{
+		UE_LOG(LogTraceServices, Error, TEXT("[GPU] %s END %u \"%s\" (%f) < BEGIN %u \"%s\" (%f) !!!"),
+			StackIndex == 1 ? TEXT("WORK") : TEXT("BREADCRUMB"),
+			EndEventTimerId, GetTimerName(EndEventTimerId), EndEventTime,
+			BeginEventTimerId, GetTimerName(BeginEventTimerId), BeginEvent.Time);
+	}
+
+	FAnalysisSessionEditScope _(Session);
+	IEditableTimeline<FTimingProfilerEvent>* Timeline = (StackIndex == 1) ?
+		TimingProfilerProvider.GetGpuQueueWorkEditableTimeline(Queue.Id) :
+		TimingProfilerProvider.GetGpuQueueEditableTimeline(Queue.Id);
+	if (ensure(Timeline))
+	{
+		const double LastTimestamp = Timeline->GetLastTimestamp();
+		if (EndEventTime < LastTimestamp)
+		{
+			UE_LOG(LogTraceServices, Error, TEXT("[GPU] %s END %u \"%s\" : Time %f < %f !!!"),
+				StackIndex == 1 ? TEXT("WORK") : TEXT("BREADCRUMB"),
+				EndEventTimerId, GetTimerName(EndEventTimerId),
+				EndEventTime, LastTimestamp);
+			EndEventTime = LastTimestamp;
+		}
+		Timeline->AppendEndEvent(EndEventTime);
+		Session.UpdateDurationSeconds(EndEventTime);
+	}
+}
+
 } // namespace TraceServices
diff --git a/Engine/Source/Developer/TraceServices/Private/Analyzers/GpuProfilerTraceAnalysis.h b/Engine/Source/Developer/TraceServices/Private/Analyzers/GpuProfilerTraceAnalysis.h
--- a/Engine/Source/Developer/TraceServices/Private/Analyzers/GpuProfilerTraceAnalysis.h
+++ b/Engine/Source/Developer/TraceServices/Private/Analyzers/GpuProfilerTraceAnalysis.h
@@ -13,22 +13,67 @@ namespace TraceServices
 class FGpuProfilerAnalyzer
 	: public UE::Trace::IAnalyzer
 {
-public:
-	FGpuProfilerAnalyzer(FAnalysisSession& Session, FTimingProfilerProvider& TimingProfilerProvider);
-	virtual void OnAnalysisBegin(const FOnAnalysisContext& Context) override;
-	virtual void OnAnalysisEnd() override;
-	virtual bool OnEvent(uint16 RouteId, EStyle Style, const FOnEventContext& Context) override;
-
 private:
 	enum : uint16
 	{
+		// The New GPU Profiler
+		RouteId_Init,
+		RouteId_QueueSpec,
+		RouteId_EventFrameBoundary,
+		RouteId_EventBeginBreadcrumb,
+		RouteId_EventEndBreadcrumb,
+		RouteId_EventBeginWork,
+		RouteId_EventEndWork,
+		RouteId_EventStats,
+
+		// The Old GPU Profiler
 		RouteId_EventSpec,
 		RouteId_Frame, // GPU Index 0
 		RouteId_Frame2, // GPU Index 1
 	};
 
+	struct FOpenEvent
+	{
+		double Time;
+		uint32 TimerId;
+	};
+
+	struct FQueue
+	{
+		uint32 Id = 0;
+		uint32 FrameNumber = 0;
+		TArray<FOpenEvent> Stack[2]; // [0] = breadcrumbs, [1] = work
+	};
+
+public:
+	FGpuProfilerAnalyzer(FAnalysisSession& Session, FTimingProfilerProvider& TimingProfilerProvider);
+	virtual void OnAnalysisBegin(const FOnAnalysisContext& Context) override;
+	virtual void OnAnalysisEnd() override;
+	virtual bool OnEvent(uint16 RouteId, EStyle Style, const FOnEventContext& Context) override;
+
+private:
+	uint32 GetOrAddTimerId(const FString& Breadcrumb);
+	const TCHAR* GetTimerName(uint32 TimerId) const;
+	FQueue& GetOrAddQueue(uint32 QueueId);
+	void BeginEvent(FQueue& Queue, int32 StackIndex, double Time, uint32 TimerId);
+	void EndEvent(FQueue& Queue, int32 StackIndex, double Time, uint32 TimerId);
+
+private:
 	FAnalysisSession& Session;
 	FTimingProfilerProvider& TimingProfilerProvider;
+
+	//////////////////////////////////////////////////
+	// The New GPU Profiler
+
+	uint32 Version = 0;
+	uint32 GpuWorkTimerId = ~0;
+	TMap<FString, uint32> BreadcrumbMap; // breadcrumb name --> GPU timer id
+	TMap<uint32, const TCHAR*> TimerMap; // GPU timer id --> persistent timer name
+	TMap<uint32, FQueue> Queues; // QueueId --> FQueue
+
+	//////////////////////////////////////////////////
+	// The Old GPU Profiler
+
 	TMap<uint64, uint32> EventTypeMap;
 	double MinTime = DBL_MIN;
 	double MinTime2 = DBL_MIN;
diff --git a/Engine/Source/Developer/TraceServices/Private/Model/TimingProfiler.cpp b/Engine/Source/Developer/TraceServices/Private/Model/TimingProfiler.cpp
--- a/Engine/Source/Developer/TraceServices/Private/Model/TimingProfiler.cpp
+++ b/Engine/Source/Developer/TraceServices/Private/Model/TimingProfiler.cpp
@@ -12,44 +12,6 @@
 namespace TraceServices
 {
 
-struct FTimingProfilerCallstackKey
-{
-	bool operator==(const FTimingProfilerCallstackKey& Other) const
-	{
-		return Other.TimerStack == TimerStack;
-	}
-
-	friend uint32 GetTypeHash(const FTimingProfilerCallstackKey& Key)
-	{
-		return Key.Hash;
-	}
-
-	TArray<uint32> TimerStack;
-	uint32 Hash;
-};
-
-class FTimingProfilerButterfly
-	: public ITimingProfilerButterfly
-{
-public:
-	FTimingProfilerButterfly();
-	virtual ~FTimingProfilerButterfly() = default;
-	virtual const FTimingProfilerButterflyNode& GenerateCallersTree(uint32 TimerId) override;
-	virtual const FTimingProfilerButterflyNode& GenerateCalleesTree(uint32 TimerId) override;
-
-private:
-	void CreateCallersTreeRecursive(const FTimingProfilerButterflyNode* TimerNode, const FTimingProfilerButterflyNode* RootNode, FTimingProfilerButterflyNode* OutputParent);
-	void CreateCalleesTreeRecursive(const FTimingProfilerButterflyNode* TimerNode, FTimingProfilerButterflyNode* OutputParent);
-
-	FSlabAllocator Allocator;
-	TPagedArray<FTimingProfilerButterflyNode> Nodes;
-	TArray<TArray<FTimingProfilerButterflyNode*>> TimerCallstacksMap;
-	TMap<uint32, FTimingProfilerButterflyNode*> CachedCallerTrees;
-	TMap<uint32, FTimingProfilerButterflyNode*> CachedCalleeTrees;
-
-	friend class FTimingProfilerProvider;
-};
-
 FTimingProfilerProvider::FTimingProfilerProvider(IAnalysisSession& InSession)
 	: Session(InSession)
 {
@@ -135,12 +97,64 @@
 	return ~MetadataId;
 }
 
+uint32 FTimingProfilerProvider::GetOriginalTimerIdFromMetadata(uint32 MetadataTimerId) const
+{
+	Session.ReadAccessCheck();
+
+	if (int32(MetadataTimerId) >= 0)
+	{
+		return MetadataTimerId;
+	}
+
+	MetadataTimerId = ~MetadataTimerId;
+	if (MetadataTimerId >= uint32(Metadatas.Num()))
+	{
+		return 0;
+	}
+
+	const FMetadata& Metadata = Metadatas[MetadataTimerId];
+	return Metadata.TimerId;
+}
+
+TArrayView<const uint8> FTimingProfilerProvider::GetMetadata(uint32 MetadataTimerId) const
+{
+	Session.ReadAccessCheck();
+
+	if (int32(MetadataTimerId) >= 0)
+	{
+		return TArrayView<const uint8>();
+	}
+
+	MetadataTimerId = ~MetadataTimerId;
+	if (MetadataTimerId >= uint32(Metadatas.Num()))
+	{
+		return TArrayView<const uint8>();
+	}
+
+	const FMetadata& Metadata = Metadatas[MetadataTimerId];
+	return Metadata.Payload;
+}
+
+TArrayView<uint8> FTimingProfilerProvider::GetEditableMetadata(uint32 MetadataTimerId)
+{
+	Session.WriteAccessCheck();
+
+	if (int32(MetadataTimerId) >= 0)
+	{
+		return TArrayView<uint8>();
+	}
+
+	MetadataTimerId = ~MetadataTimerId;
+	if (MetadataTimerId >= uint32(Metadatas.Num()))
+	{
+		return TArrayView<uint8>();
+	}
+
+	FMetadata& Metadata = Metadatas[MetadataTimerId];
+	return Metadata.Payload;
+}
+
 IEditableTimeline<FTimingProfilerEvent>& FTimingProfilerProvider::GetCpuThreadEditableTimeline(uint32 ThreadId)
 {
-	return EditCpuThreadTimeline(ThreadId);
-}
-
-FTimingProfilerProvider::TimelineInternal& FTimingProfilerProvider::EditCpuThreadTimeline(uint32 ThreadId)
-{
 	Session.WriteAccessCheck();
 
@@ -173,6 +187,63 @@
 	return Timelines[Gpu2TimelineIndex].Get();
 }
 
+void FTimingProfilerProvider::AddGpuQueue(uint32 QueueId, uint8 GPU, uint8 Index, uint8 Type, const TCHAR* Name)
+{
+	const uint32* FoundQueueIndex = GpuQueueIdToQueueIndexMap.Find(QueueId);
+	if (!FoundQueueIndex)
+	{
+		GpuQueueIdToQueueIndexMap.Add(QueueId, (uint32)GpuQueues.Num());
+
+		FGpuQueueInfo& GpuQueue = GpuQueues.AddDefaulted_GetRef();
+
+		GpuQueue.Id = QueueId;
+		GpuQueue.GPU = GPU;
+		GpuQueue.Index = Index;
+		GpuQueue.Type = Type;
+		GpuQueue.Name = Name;
+
+		GpuQueue.TimelineIndex = Timelines.Num();
+		TSharedRef<TimelineInternal> Timeline = MakeShared<TimelineInternal>(Session.GetLinearAllocator());
+		Timelines.Add(Timeline);
+
+		GpuQueue.WorkTimelineIndex = Timelines.Num();
+		TSharedRef<TimelineInternal> WorkTimeline = MakeShared<TimelineInternal>(Session.GetLinearAllocator());
+		Timelines.Add(WorkTimeline);
+	}
+}
+
+IEditableTimeline<FTimingProfilerEvent>* FTimingProfilerProvider::GetGpuQueueEditableTimeline(uint32 QueueId)
+{
+	Session.WriteAccessCheck();
+
+	const uint32* FoundQueueIndex = GpuQueueIdToQueueIndexMap.Find(QueueId);
+	if (FoundQueueIndex)
+	{
+		check(*FoundQueueIndex < (uint32)GpuQueues.Num());
+		uint32 TimelineIndex = GpuQueues[*FoundQueueIndex].TimelineIndex;
+		check(TimelineIndex < (uint32)Timelines.Num());
+		return &Timelines[TimelineIndex].Get();
+	}
+
+	return nullptr;
+}
+
+IEditableTimeline<FTimingProfilerEvent>* FTimingProfilerProvider::GetGpuQueueWorkEditableTimeline(uint32 QueueId)
+{
+	Session.WriteAccessCheck();
+
+	const uint32* FoundQueueIndex = GpuQueueIdToQueueIndexMap.Find(QueueId);
+	if (FoundQueueIndex)
+	{
+		check(*FoundQueueIndex < (uint32)GpuQueues.Num());
+		uint32 TimelineIndex = GpuQueues[*FoundQueueIndex].WorkTimelineIndex;
+		check(TimelineIndex < (uint32)Timelines.Num());
+		return &Timelines[TimelineIndex].Get();
+	}
+
+	return nullptr;
+}
+
 bool FTimingProfilerProvider::GetCpuThreadTimelineIndex(uint32 ThreadId, uint32& OutTimelineIndex) const
 {
 	Session.ReadAccessCheck();
@@ -202,6 +273,16 @@
 	return true;
 }
 
+void FTimingProfilerProvider::EnumerateGpuQueues(TFunctionRef<void(const FGpuQueueInfo&)> Callback) const
+{
+	Session.ReadAccessCheck();
+
+	for (const FGpuQueueInfo& GpuQueue : GpuQueues)
+	{
+		Callback(GpuQueue);
+	}
+}
+
 bool FTimingProfilerProvider::ReadTimeline(uint32 Index, TFunctionRef<void(const Timeline &)> Callback) const
 {
 	Session.ReadAccessCheck();
@@ -242,44 +323,43 @@
 	return Timers.Num();
 }
 
-TArrayView<const uint8> FTimingProfilerProvider::GetMetadata(uint32 TimerId) const
-{
-	if (int32(TimerId) >= 0)
-	{
-		return TArrayView<const uint8>();
-	}
-
-	TimerId = ~TimerId;
-	if (TimerId >= uint32(Metadatas.Num()))
-	{
-		return TArrayView<const uint8>();
-	}
-
-	const FMetadata& Metadata = Metadatas[TimerId];
-	return Metadata.Payload;
-}
-
 void FTimingProfilerProvider::ReadTimers(TFunctionRef<void(const ITimingProfilerTimerReader&)> Callback) const
 {
 	Session.ReadAccessCheck();
 	Callback(*this);
 }
 
-ITable<FTimingProfilerAggregatedStats>* FTimingProfilerProvider::CreateAggregation(const FCreateAggreationParams& Params) const
+ITable<FTimingProfilerAggregatedStats>* FTimingProfilerProvider::CreateAggregation(const FCreateAggregationParams& Params) const
 {
 	Session.ReadAccessCheck();
 
 	TArray<const TimelineInternal*> IncludedTimelines;
-	if (Params.IncludeGpu)
+	if (Params.bIncludeOldGpu1)
 	{
 		IncludedTimelines.Add(&Timelines[GpuTimelineIndex].Get());
+	}
+	if (Params.bIncludeOldGpu2)
+	{
 		IncludedTimelines.Add(&Timelines[Gpu2TimelineIndex].Get());
 	}
-	for (const auto& KV : CpuThreadTimelineIndexMap)
+	if (Params.GpuQueueFilter)
 	{
-		if (Params.CpuThreadFilter(KV.Key))
+		for (const FGpuQueueInfo& GpuQueue : GpuQueues)
 		{
-			IncludedTimelines.Add(&Timelines[KV.Value].Get());
+			if (Params.GpuQueueFilter(GpuQueue.Id))
+			{
+				IncludedTimelines.Add(&Timelines[GpuQueue.TimelineIndex].Get());
+			}
+		}
+	}
+	if (Params.CpuThreadFilter)
+	{
+		for (const auto& KV : CpuThreadTimelineIndexMap)
+		{
+			if (Params.CpuThreadFilter(KV.Key))
+			{
+				IncludedTimelines.Add(&Timelines[KV.Value].Get());
+			}
 		}
 	}
 
@@ -320,7 +400,7 @@
 	}
 
 	TTable<FTimingProfilerAggregatedStats>* Table = new TTable<FTimingProfilerAggregatedStats>(AggregatedStatsTableLayout);
-	
+
 	if (Params.CancellationToken.IsValid() && Params.CancellationToken->ShouldCancel())
 	{
 		return Table;
@@ -346,11 +426,36 @@
 	return Table;
 }
 
+////////////////////////////////////////////////////////////////////////////////////////////////////
+// FTimingProfilerButterfly
+////////////////////////////////////////////////////////////////////////////////////////////////////
+
+class FTimingProfilerButterfly
+	: public ITimingProfilerButterfly
+{
+public:
+	FTimingProfilerButterfly();
+	virtual ~FTimingProfilerButterfly() = default;
+	virtual const FTimingProfilerButterflyNode& GenerateCallersTree(uint32 TimerId) override;
+	virtual const FTimingProfilerButterflyNode& GenerateCalleesTree(uint32 TimerId) override;
+
+private:
+	void CreateCallersTreeRecursive(const FTimingProfilerButterflyNode* TimerNode, const FTimingProfilerButterflyNode* RootNode, FTimingProfilerButterflyNode* OutputParent);
+	void CreateCalleesTreeRecursive(const FTimingProfilerButterflyNode* TimerNode, FTimingProfilerButterflyNode* OutputParent);
+
+	FSlabAllocator Allocator;
+	TPagedArray<FTimingProfilerButterflyNode> Nodes;
+	TArray<TArray<FTimingProfilerButterflyNode*>> TimerCallstacksMap;
+	TMap<uint32, FTimingProfilerButterflyNode*> CachedCallerTrees;
+	TMap<uint32, FTimingProfilerButterflyNode*> CachedCalleeTrees;
+
+	friend class FTimingProfilerProvider;
+};
+
 FTimingProfilerButterfly::FTimingProfilerButterfly()
 	: Allocator(2 << 20)
 	, Nodes(Allocator, 1024)
 {
-
 }
 
 void FTimingProfilerButterfly::CreateCallersTreeRecursive(const FTimingProfilerButterflyNode* TimerNode, const FTimingProfilerButterflyNode* RootNode, FTimingProfilerButterflyNode* OutputParent)
@@ -455,22 +560,58 @@
 	return *Root;
 }
 
-ITimingProfilerButterfly* FTimingProfilerProvider::CreateButterfly(double IntervalStart, double IntervalEnd, TFunctionRef<bool(uint32)> CpuThreadFilter, bool IncludeGpu) const
+////////////////////////////////////////////////////////////////////////////////////////////////////
+
+struct FTimingProfilerCallstackKey
 {
+	bool operator==(const FTimingProfilerCallstackKey& Other) const
+	{
+		return Other.TimerStack == TimerStack;
+	}
+
+	friend uint32 GetTypeHash(const FTimingProfilerCallstackKey& Key)
+	{
+		return Key.Hash;
+	}
+
+	TArray<uint32> TimerStack;
+	uint32 Hash;
+};
+
+////////////////////////////////////////////////////////////////////////////////////////////////////
+
+ITimingProfilerButterfly* FTimingProfilerProvider::CreateButterfly(const FCreateButterflyParams& Params) const
+{
 	FTimingProfilerButterfly* Butterfly = new FTimingProfilerButterfly();
 	Butterfly->TimerCallstacksMap.AddDefaulted(Timers.Num());
 
 	TArray<const TimelineInternal*> IncludedTimelines;
-	if (IncludeGpu)
+	if (Params.bIncludeOldGpu1)
 	{
 		IncludedTimelines.Add(&Timelines[GpuTimelineIndex].Get());
+	}
+	if (Params.bIncludeOldGpu2)
+	{
 		IncludedTimelines.Add(&Timelines[Gpu2TimelineIndex].Get());
 	}
-	for (const auto& KV : CpuThreadTimelineIndexMap)
+	if (Params.GpuQueueFilter)
 	{
-		if (CpuThreadFilter(KV.Key))
+		for (const FGpuQueueInfo& GpuQueue : GpuQueues)
 		{
-			IncludedTimelines.Add(&Timelines[KV.Value].Get());
+			if (Params.GpuQueueFilter(GpuQueue.Id))
+			{
+				IncludedTimelines.Add(&Timelines[GpuQueue.TimelineIndex].Get());
+			}
+		}
+	}
+	if (Params.CpuThreadFilter)
+	{
+		for (const auto& KV : CpuThreadTimelineIndexMap)
+		{
+			if (Params.CpuThreadFilter(KV.Key))
+			{
+				IncludedTimelines.Add(&Timelines[KV.Value].Get());
+			}
 		}
 	}
 
@@ -494,9 +635,9 @@
 	double LastTime = 0.0;
 	for (const TimelineInternal* Timeline : IncludedTimelines)
 	{
-		Timeline->EnumerateEvents(IntervalStart, IntervalEnd, [this, IntervalStart, IntervalEnd, &CurrentCallstackKey, &CurrentCallstack, &CallstackNodeMap, &LastTime, Butterfly](bool IsEnter, double Time, const FTimingProfilerEvent& Event)
+		Timeline->EnumerateEvents(Params.IntervalStart, Params.IntervalEnd, [this, &Params, &CurrentCallstackKey, &CurrentCallstack, &CallstackNodeMap, &LastTime, Butterfly](bool IsEnter, double Time, const FTimingProfilerEvent& Event)
 		{
-			Time = FMath::Clamp(Time, IntervalStart, IntervalEnd);
+			Time = FMath::Clamp(Time, Params.IntervalStart, Params.IntervalEnd);
 			FTimingProfilerButterflyNode* ParentNode = nullptr;
 			uint32 ParentCallstackHash = 17;
 			if (CurrentCallstack.Num())
diff --git a/Engine/Source/Developer/TraceServices/Private/Model/TimingProfilerPrivate.h b/Engine/Source/Developer/TraceServices/Private/Model/TimingProfilerPrivate.h
--- a/Engine/Source/Developer/TraceServices/Private/Model/TimingProfilerPrivate.h
+++ b/Engine/Source/Developer/TraceServices/Private/Model/TimingProfilerPrivate.h
@@ -19,38 +19,51 @@ class FTimingProfilerProvider
 	, public IEditableTimingProfilerProvider
 {
 public:
-	typedef TMonotonicTimeline<FTimingProfilerEvent> 	TimelineInternal;
+	typedef TMonotonicTimeline<FTimingProfilerEvent> TimelineInternal;
 
 	explicit FTimingProfilerProvider(IAnalysisSession& InSession);
 	virtual ~FTimingProfilerProvider();
-	uint32 AddGpuTimer(FStringView Name, const TCHAR* File = nullptr, uint32 Line = 0);
-	void SetTimerName(uint32 TimerId, FStringView Name);
-	TimelineInternal& EditCpuThreadTimeline(uint32 ThreadId);
+
 	TimelineInternal& EditGpuTimeline();
 	TimelineInternal& EditGpu2Timeline();
+
+	// ITimingProfilerProvider
 	virtual bool GetCpuThreadTimelineIndex(uint32 ThreadId, uint32& OutTimelineIndex) const override;
 	virtual bool GetGpuTimelineIndex(uint32& OutTimelineIndex) const override;
 	virtual bool GetGpu2TimelineIndex(uint32& OutTimelineIndex) const override;
+	virtual void EnumerateGpuQueues(TFunctionRef<void(const FGpuQueueInfo&)> Callback) const override;
 	virtual bool ReadTimeline(uint32 Index, TFunctionRef<void(const Timeline&)> Callback) const override;
 	virtual uint32 GetTimelineCount() const override { return Timelines.Num(); }
 	virtual void EnumerateTimelines(TFunctionRef<void(const Timeline&)> Callback) const override;
 	virtual void ReadTimers(TFunctionRef<void(const ITimingProfilerTimerReader&)> Callback) const override;
-	virtual ITable<FTimingProfilerAggregatedStats>* CreateAggregation(const FCreateAggreationParams& Params) const override;
-	virtual ITimingProfilerButterfly* CreateButterfly(double IntervalStart, double IntervalEnd, TFunctionRef<bool(uint32)> CpuThreadFilter, bool IncludeGpu) const override;
+	virtual uint32 GetOriginalTimerIdFromMetadata(uint32 MetadataTimerId) const override;
+	virtual TArrayView<const uint8> GetMetadata(uint32 MetadataTimerId) const override;
+	virtual ITable<FTimingProfilerAggregatedStats>* CreateAggregation(const FCreateAggregationParams& Params) const override;
+	virtual ITimingProfilerButterfly* CreateButterfly(const FCreateButterflyParams& Params) const override;
+
+	// ITimingProfilerTimerReader
 	virtual const FTimingProfilerTimer* GetTimer(uint32 TimerId) const override;
 	virtual uint32 GetTimerCount() const override;
+	//virtual uint32 GetOriginalTimerIdFromMetadata(uint32 MetadataTimerId) const override;
+	//virtual TArrayView<const uint8> GetMetadata(uint32 MetadataTimerId) const override;
 
-	// implementing IEditableTimingProfilerProvider
+	// IEditableTimingProfilerProvider
 	virtual uint32 AddCpuTimer(FStringView Name, const TCHAR* File = nullptr, uint32 Line = 0) override;
+	virtual uint32 AddGpuTimer(FStringView Name, const TCHAR* File = nullptr, uint32 Line = 0) override;
+	virtual void SetTimerName(uint32 TimerId, FStringView Name) override;
 	virtual void SetTimerNameAndLocation(uint32 TimerId, FStringView Name, const TCHAR* File, uint32 Line) override;
 	virtual uint32 AddMetadata(uint32 OriginalTimerId, TArray<uint8>&& Metadata) override;
-	virtual TArrayView<const uint8> GetMetadata(uint32 TimerId) const override;
+	virtual TArrayView<uint8> GetEditableMetadata(uint32 MetadataTimerId) override;
+	virtual void AddGpuQueue(uint32 QueueId, uint8 GPU, uint8 Index, uint8 Type, const TCHAR* Name) override;
 	virtual IEditableTimeline<FTimingProfilerEvent>& GetCpuThreadEditableTimeline(uint32 ThreadId) override;
+	virtual IEditableTimeline<FTimingProfilerEvent>* GetGpuQueueEditableTimeline(uint32 QueueId) override;
+	virtual IEditableTimeline<FTimingProfilerEvent>* GetGpuQueueWorkEditableTimeline(uint32 QueueId) override;
+	virtual const ITimingProfilerProvider* GetReadProvider() const override { return this; }
 
 private:
 	FTimingProfilerTimer& AddTimerInternal(FStringView Name, const TCHAR* File, uint32 Line, bool IsGpuEvent);
 
-	struct FMetadata 
+	struct FMetadata
 	{
 		TArray<uint8> Payload;
 		uint32 TimerId;
@@ -60,9 +73,11 @@ private:
 	TArray<FMetadata> Metadatas;
 	TArray<FTimingProfilerTimer> Timers;
 	TArray<TSharedRef<TimelineInternal>> Timelines;
-	TMap<uint32, uint32> CpuThreadTimelineIndexMap;
-	uint32 GpuTimelineIndex = 0;
-	uint32 Gpu2TimelineIndex = 1;
+	TMap<uint32, uint32> CpuThreadTimelineIndexMap; // CPU Thread Id --> timeline index, in Timelines
+	static const uint32 GpuTimelineIndex = 0;
+	static const uint32 Gpu2TimelineIndex = 1;
+	TArray<FGpuQueueInfo> GpuQueues;
+	TMap<uint32, uint32> GpuQueueIdToQueueIndexMap; // GPU Queue Id --> queue index, in GpuQueues
 	TTableLayout<FTimingProfilerAggregatedStats> AggregatedStatsTableLayout;
 };
 
diff --git a/Engine/Source/Developer/TraceServices/Public/Model/MonotonicTimeline.h b/Engine/Source/Developer/TraceServices/Public/Model/MonotonicTimeline.h
--- a/Engine/Source/Developer/TraceServices/Public/Model/MonotonicTimeline.h
+++ b/Engine/Source/Developer/TraceServices/Public/Model/MonotonicTimeline.h
@@ -52,6 +52,8 @@ public:
 	* @param EndTime	The ending timestamp of the event in seconds.
 	*/
 	virtual void AppendEndEvent(double EndTime) = 0;
+
+	virtual double GetLastTimestamp() const { return 0.0; }
 };
 
 template<typename InEventType, typename SettingsType = FMonotonicTimelineDefaultSettings>
@@ -81,6 +83,11 @@ public:
 
 	virtual ~TMonotonicTimeline() = default;
 
+	virtual double GetLastTimestamp() const override
+	{
+		return DetailLevels[0].InsertionState.LastTime;
+	}
+
 	virtual void AppendBeginEvent(double StartTime, const EventType& Event) override
 	{
 		int32 CurrentDepth = DetailLevels[0].InsertionState.CurrentDepth;
diff --git a/Engine/Source/Developer/TraceServices/Public/TraceServices/Model/TimingProfiler.h b/Engine/Source/Developer/TraceServices/Public/TraceServices/Model/TimingProfiler.h
--- a/Engine/Source/Developer/TraceServices/Public/TraceServices/Model/TimingProfiler.h
+++ b/Engine/Source/Developer/TraceServices/Public/TraceServices/Model/TimingProfiler.h
@@ -66,16 +66,53 @@ struct FTimingProfilerButterflyNode
 	TArray<FTimingProfilerButterflyNode*> Children;
 };
 
-struct FCreateAggreationParams
+struct FCreateAggregationParams
 {
-	double IntervalStart;
-	double IntervalEnd;
+	// The start timestamp in seconds.
+	double IntervalStart = 0.0;
+
+	// The end timestamp in seconds.
+	double IntervalEnd = 0.0;
+
+	// A boolean value to specify if aggregation should include the old GPU timeline.
+	bool bIncludeOldGpu1 = false;
+
+	// A boolean value to specify if aggregation should include the old GPU timeline.
+	bool bIncludeOldGpu2 = false;
+
+	// A function to filter the GPU queues to aggregate.
+	TFunction<bool(uint32)> GpuQueueFilter;
+
+	// A function to filter the CPU threads to aggregate.
 	TFunction<bool(uint32)> CpuThreadFilter;
-	bool IncludeGpu;
+
+	// The type of frame to use for frame stats aggregation. ETraceFrameType::TraceFrameType_Count means no frame aggregation.
 	ETraceFrameType FrameType = ETraceFrameType::TraceFrameType_Count;
+
 	TSharedPtr<TraceServices::FCancellationToken> CancellationToken;
 };
 
+struct FCreateButterflyParams
+{
+	// The start timestamp in seconds.
+	double IntervalStart = 0.0;
+
+	// The end timestamp in seconds.
+	double IntervalEnd = 0.0;
+
+	// A boolean value to specify if aggregation should include the old GPU timeline.
+	bool bIncludeOldGpu1 = false;
+
+	// A boolean value to specify if aggregation should include the old GPU timeline.
+	bool bIncludeOldGpu2 = false;
+
+	// A function to filter the GPU queues to aggregate.
+	TFunction<bool(uint32)> GpuQueueFilter;
+
+	// A function to filter the CPU threads to aggregate.
+	TFunction<bool(uint32)> CpuThreadFilter;
+};
+
 class ITimingProfilerButterfly
 {
 public:
@@ -89,7 +126,19 @@ class ITimingProfilerTimerReader
 public:
 	virtual const FTimingProfilerTimer* GetTimer(uint32 TimerId) const = 0;
 	virtual uint32 GetTimerCount() const = 0;
-	virtual TArrayView<const uint8> GetMetadata(uint32 TimerId) const = 0;
+	virtual uint32 GetOriginalTimerIdFromMetadata(uint32 MetadataTimerId) const { return MetadataTimerId; }
+	virtual TArrayView<const uint8> GetMetadata(uint32 MetadataTimerId) const { return TArrayView<const uint8>(); }
+};
+
+struct FGpuQueueInfo
+{
+	uint32 Id = 0;
+	uint8 GPU = 0;
+	uint8 Index = 0;
+	uint8 Type = 0;
+	const TCHAR* Name = nullptr;
+	uint32 TimelineIndex = ~0;
+	uint32 WorkTimelineIndex = ~0;
 };
 
 class ITimingProfilerProvider
@@ -99,31 +148,46 @@ public:
 	typedef ITimeline<FTimingProfilerEvent> Timeline;
 
 	virtual ~ITimingProfilerProvider() = default;
+
 	virtual bool GetCpuThreadTimelineIndex(uint32 ThreadId, uint32& OutTimelineIndex) const = 0;
 	virtual bool GetGpuTimelineIndex(uint32& OutTimelineIndex) const = 0;
 	virtual bool GetGpu2TimelineIndex(uint32& OutTimelineIndex) const = 0;
+	virtual void EnumerateGpuQueues(TFunctionRef<void(const FGpuQueueInfo&)> Callback) const {}
+
 	virtual bool ReadTimeline(uint32 Index, TFunctionRef<void(const Timeline&)> Callback) const = 0;
 	virtual uint32 GetTimelineCount() const = 0;
 	virtual void EnumerateTimelines(TFunctionRef<void(const Timeline&)> Callback) const = 0;
+
 	virtual void ReadTimers(TFunctionRef<void(const ITimingProfilerTimerReader&)> Callback) const = 0;
+	virtual uint32 GetOriginalTimerIdFromMetadata(uint32 MetadataTimerId) const { return MetadataTimerId; }
+	virtual TArrayView<const uint8> GetMetadata(uint32 MetadataTimerId) const { return TArrayView<const uint8>(); }
 
 	/**
-	* Create a table of aggregated stats.
+	* Creates a table of aggregated stats.
 	*
-	* @param Params				The params for the aggregation.
+	* @param Params	The params for the aggregation.
 	*/
-	virtual ITable<FTimingProfilerAggregatedStats>* CreateAggregation(const FCreateAggreationParams& Params) const = 0;
+	virtual ITable<FTimingProfilerAggregatedStats>* CreateAggregation(const FCreateAggregationParams& Params) const { return nullptr; }
 
 	/**
-	* Create a table of aggregated stats.
+	* Creates a butterfly aggregation.
 	*
-	* @param IntervalStart		The start timestamp in seconds.
-	* @param IntervalEnd		The end timestamp in seconds.
-	* @param CpuThreadFilter	A function to filter the CPU threads to aggregate.
-	* @param IncludeGpu			A boolean value to specify if aggregaton should include GPU timelines.
-	* @param FrameType			The type of frame to use for frame stats aggregation. ETraceFrameType::TraceFrameType_Count means no frame aggregation.
+	* @param Params	The params for the aggregation.
 	*/
-	virtual ITimingProfilerButterfly* CreateButterfly(double IntervalStart, double IntervalEnd, TFunctionRef<bool(uint32)> CpuThreadFilter, bool IncludeGpu) const = 0;
+	virtual ITimingProfilerButterfly* CreateButterfly(const FCreateButterflyParams& Params) const { return nullptr; }
+
+	UE_DEPRECATED(5.6, "Use FCreateButterflyParams instead.")
+	virtual ITimingProfilerButterfly* CreateButterfly(double IntervalStart, double IntervalEnd, TFunctionRef<bool(uint32)> CpuThreadFilter, bool bIncludeGpu) const
+	{
+		FCreateButterflyParams Params;
+		Params.IntervalStart = IntervalStart;
+		Params.IntervalEnd = IntervalEnd;
+		Params.bIncludeOldGpu1 = bIncludeGpu;
+		Params.bIncludeOldGpu2 = bIncludeGpu;
+		//Params.GpuQueueFilter = ;
+		Params.CpuThreadFilter = [CpuThreadFilter](uint32 QueueId) { return CpuThreadFilter(QueueId); };
+		return CreateButterfly(Params);
+	}
 };
 
 /*
@@ -136,18 +200,37 @@ public:
 	virtual ~IEditableTimingProfilerProvider() = default;
 
 	/*
-	* A new CPU timer object has been found.
+	* Adds/registers a new CPU timer.
 	*
 	* @param Name	The name attached to the timer.
 	* @param File	The source file in which the timer is defined.
 	* @param Line	The line number of the source file in which the timer is defined.
 	*
-	* @return The identity of the CPU timer object.
+	* @return The identity of the CPU timer.
 	*/
 	virtual uint32 AddCpuTimer(FStringView Name, const TCHAR* File, uint32 Line) = 0;
 
 	/*
-	* Update an existing timer with information. Some information is unavailable when it's created.
+	* Adds/registers a new GPU timer.
+	*
+	* @param Name	The name attached to the timer.
+	* @param File	The source file in which the timer is defined.
+	* @param Line	The line number of the source file in which the timer is defined.
+	*
+	* @return The identity of the GPU timer.
+	*/
+	virtual uint32 AddGpuTimer(FStringView Name, const TCHAR* File, uint32 Line) { return 0; }
+
+	/*
+	* Updates an existing timer with information. Some information is unavailable when it's created.
+	*
+	* @param TimerId	The identity of the timer to update.
+	* @param Name		The name attached to the timer.
+	*/
+	virtual void SetTimerName(uint32 TimerId, FStringView Name) = 0;
+
+	/*
+	* Updates an existing timer with information. Some information is unavailable when it's created.
 	*
 	* @param TimerId	The identity of the timer to update.
 	* @param Name		The name attached to the timer.
@@ -157,7 +240,7 @@ public:
 	virtual void SetTimerNameAndLocation(uint32 TimerId, FStringView Name, const TCHAR* File, uint32 Line) = 0;
 
 	/*
-	* Add metadata to a timer.
+	* Adds metadata to a CPU or GPU timer.
 	*
 	* @param OriginalTimerId	The identity of the timer to add metadata to.
 	* @param Metadata			The metadata.
@@ -167,22 +250,45 @@ public:
 	virtual uint32 AddMetadata(uint32 OriginalTimerId, TArray<uint8>&& Metadata) = 0;
 
 	/*
-	* Get metadata for a timer.
+	* Gets metadata by id.
 	*
-	* @param TimerId	The identity of the metadata to retrieve.
+	* @param MetadataTimerId	The identity of the metadata.
 	*
 	* @return The metadata.
 	*/
-	virtual TArrayView<const uint8> GetMetadata(uint32 TimerId) const = 0;
+	virtual TArrayView<uint8> GetEditableMetadata(uint32 MetadataTimerId) = 0;
 
 	/*
-	* Get an object to receive ordered timeline events for a thread.
+	* Adds a new GPU queue.
+	*
+	* @param QueueId	The GPU queue for which the events are for.
+	*/
+	virtual void AddGpuQueue(uint32 QueueId, uint8 GPU, uint8 Index, uint8 Type, const TCHAR* Name) {}
+
+	/*
+	* Gets an object to receive ordered timeline events for a CPU thread.
 	*
 	* @param ThreadId	The thread for which the events are for.
 	*
 	* @return The object to receive the serial events for the specified thread.
 	*/
 	virtual IEditableTimeline<FTimingProfilerEvent>& GetCpuThreadEditableTimeline(uint32 ThreadId) = 0;
+
+	/*
+	* Gets an object to receive ordered timeline events for a GPU queue.
+	*
+	* @param QueueId	The queue for which the events are for.
+	*
+	* @return The object to receive the serial events for the specified queue.
+	*/
+	virtual IEditableTimeline<FTimingProfilerEvent>* GetGpuQueueEditableTimeline(uint32 QueueId) { return nullptr; }
+	virtual IEditableTimeline<FTimingProfilerEvent>* GetGpuQueueWorkEditableTimeline(uint32 QueueId) { return nullptr; }
+
+	/*
+	* Gets the read provider.
+	* @returns The read only provider or nullptr if not available.
+	*/
+	virtual const ITimingProfilerProvider* GetReadProvider() const { return nullptr; }
 };
 
 TRACESERVICES_API FName GetTimingProfilerProviderName();
diff --git a/Engine/Source/Editor/UnrealEd/Private/Commandlets/SummarizeTraceUtils.cpp b/Engine/Source/Editor/UnrealEd/Private/Commandlets/SummarizeTraceUtils.cpp
--- a/Engine/Source/Editor/UnrealEd/Private/Commandlets/SummarizeTraceUtils.cpp
+++ b/Engine/Source/Editor/UnrealEd/Private/Commandlets/SummarizeTraceUtils.cpp
@@ -328,7 +328,7 @@ uint32 FSummarizeCpuProfilerProvider::AddCpuTimer(FStringView Name, const TCHAR*
 	return TimerId;
 }
 
-void FSummarizeCpuProfilerProvider::SetTimerNameAndLocation(uint32 TimerId, FStringView Name, const TCHAR* File, uint32 Line)
+void FSummarizeCpuProfilerProvider::SetTimerName(uint32 TimerId, FStringView Name)
 {
 	check(TimerId < uint32(ScopeNames.Num()));
 	check(!Name.IsEmpty());
@@ -342,6 +342,11 @@ void FSummarizeCpuProfilerProvider::SetTimerNameAndLocation(uint32 TimerId, FStr
 	}
 }
 
+void FSummarizeCpuProfilerProvider::SetTimerNameAndLocation(uint32 TimerId, FStringView Name, const TCHAR* File, uint32 Line)
+{
+	SetTimerName(TimerId, Name);
+}
+
 uint32 FSummarizeCpuProfilerProvider::AddMetadata(uint32 MasterTimerId, TArray<uint8>&& Metadata)
 {
 	uint32 MetadataId = Metadatas.Num();
@@ -349,20 +354,20 @@ uint32 FSummarizeCpuProfilerProvider::AddMetadata(uint32 MasterTimerId, TArray<u
 	return ~MetadataId;
 }
 
-TArrayView<const uint8> FSummarizeCpuProfilerProvider::GetMetadata(uint32 TimerId) const
+TArrayView<uint8> FSummarizeCpuProfilerProvider::GetEditableMetadata(uint32 TimerId)
 {
 	if (int32(TimerId) >= 0)
 	{
-		return TArrayView<const uint8>();
+		return TArrayView<uint8>();
 	}
 
 	TimerId = ~TimerId;
 	if (TimerId >= uint32(Metadatas.Num()))
 	{
-		return TArrayView<const uint8>();
+		return TArrayView<uint8>();
 	}
 
-	const FMetadata& Metadata = Metadatas[TimerId];
+	FMetadata& Metadata = Metadatas[TimerId];
 	return Metadata.Payload;
 }
 
diff --git a/Engine/Source/Editor/UnrealEd/Private/Commandlets/SummarizeTraceUtils.h b/Engine/Source/Editor/UnrealEd/Private/Commandlets/SummarizeTraceUtils.h
--- a/Engine/Source/Editor/UnrealEd/Private/Commandlets/SummarizeTraceUtils.h
+++ b/Engine/Source/Editor/UnrealEd/Private/Commandlets/SummarizeTraceUtils.h
@@ -202,13 +202,14 @@ private:
 
 private:
 	// TraceServices::IEditableThreadProvider
-	virtual void					AddThread(uint32 Id, const TCHAR* Name, EThreadPriority Priority) override;
+	virtual void AddThread(uint32 Id, const TCHAR* Name, EThreadPriority Priority) override;
 
 	// TraceServices::IEditableTimingProfilerProvider
-	virtual uint32					AddCpuTimer(FStringView Name, const TCHAR* File, uint32 Line) override;
-	virtual void					SetTimerNameAndLocation(uint32 TimerId, FStringView Name, const TCHAR* File, uint32 Line) override;
-	virtual uint32					AddMetadata(uint32 MasterTimerId, TArray<uint8>&& Metadata) override;
-	virtual TArrayView<const uint8> GetMetadata(uint32 TimerId) const override;
+	virtual uint32 AddCpuTimer(FStringView Name, const TCHAR* File, uint32 Line) override;
+	virtual void SetTimerName(uint32 TimerId, FStringView Name) override;
+	virtual void SetTimerNameAndLocation(uint32 TimerId, FStringView Name, const TCHAR* File, uint32 Line) override;
+	virtual uint32 AddMetadata(uint32 MasterTimerId, TArray<uint8>&& Metadata) override;
+	virtual TArrayView<uint8> GetEditableMetadata(uint32 TimerId) override;
 	virtual TraceServices::IEditableTimeline<TraceServices::FTimingProfilerEvent>& GetCpuThreadEditableTimeline(uint32 ThreadId) override;
 
 	// Callbacks from FThread to forward calls to Analyzers
