[Insights] Made the GPU Work tracks to be overlays for the GPU Queue tracks.

diff --git a/Engine/Source/Developer/TraceInsights/Private/Insights/TimingProfiler/Tracks/ThreadTimingTrack.cpp b/Engine/Source/Developer/TraceInsights/Private/Insights/TimingProfiler/Tracks/ThreadTimingTrack.cpp
--- a/Engine/Source/Developer/TraceInsights/Private/Insights/TimingProfiler/Tracks/ThreadTimingTrack.cpp
+++ b/Engine/Source/Developer/TraceInsights/Private/Insights/TimingProfiler/Tracks/ThreadTimingTrack.cpp
@@ -1534,6 +1534,134 @@ TSharedPtr<const ITimingEvent> FThreadTimingTrack::FindMinEventInstance(uint32 I
 	return TimingEvent;
 }
 
+////////////////////////////////////////////////////////////////////////////////////////////////////
+// FCpuTimingTrack, FGpuTimingTrack, FGpuQueueTimingTrack
+////////////////////////////////////////////////////////////////////////////////////////////////////
+
+INSIGHTS_IMPLEMENT_RTTI(FCpuTimingTrack)
+INSIGHTS_IMPLEMENT_RTTI(FGpuTimingTrack)
+INSIGHTS_IMPLEMENT_RTTI(FGpuQueueTimingTrack)
+
+////////////////////////////////////////////////////////////////////////////////////////////////////
+// FGpuQueueWorkTimingTrack
+////////////////////////////////////////////////////////////////////////////////////////////////////
+
+INSIGHTS_IMPLEMENT_RTTI(FGpuQueueWorkTimingTrack)
+
+////////////////////////////////////////////////////////////////////////////////////////////////////
+
+void FGpuQueueWorkTimingTrack::Draw(const ITimingTrackDrawContext& Context) const
+{
+	DrawLineEvents(Context, 1.0f);
+}
+
+////////////////////////////////////////////////////////////////////////////////////////////////////
+
+void FGpuQueueWorkTimingTrack::DrawLineEvents(const ITimingTrackDrawContext& Context, const float OffsetY) const
+{
+	const FTimingViewDrawHelper& Helper = *static_cast<const FTimingViewDrawHelper*>(&Context.GetHelper());
+
+	if ((Context.GetEventFilter().IsValid() && Context.GetEventFilter()->FilterTrack(*this)) || HasCustomFilter())
+	{
+		Helper.DrawFadedLineEvents(GetDrawState(), *this, OffsetY, 0.1f);
+
+		if (UpdateFilteredDrawStateOpacity())
+		{
+			Helper.DrawLineEvents(GetFilteredDrawState(), *this, OffsetY);
+		}
+		else
+		{
+			Helper.DrawFadedLineEvents(GetFilteredDrawState(), *this, OffsetY, GetFilteredDrawStateOpacity());
+		}
+	}
+	else
+	{
+		Helper.DrawLineEvents(GetDrawState(), *this, OffsetY);
+	}
+}
+
+////////////////////////////////////////////////////////////////////////////////////////////////////
+
+void FGpuQueueWorkTimingTrack::PostDraw(const ITimingTrackDrawContext& Context) const
+{
+	const bool bAreOverlaysVisible = true; // SharedState.AreOverlaysVisible();
+	const bool bAreExtendedLinesVisible = true; // SharedState.AreExtendedLinesVisible();
+
+	if (bAreOverlaysVisible || bAreExtendedLinesVisible)
+	{
+		float LineY1 = 0.0f;
+		float LineY2 = 0.0f;
+		ETimingTrackLocation LocalLocation = ETimingTrackLocation::None;
+		TSharedPtr<FBaseTimingTrack> LocalParentTrack = GetParentTrack().Pin();
+		if (LocalParentTrack)
+		{
+			LineY1 = LocalParentTrack->GetPosY();
+			LineY2 = LineY1 + LocalParentTrack->GetHeight();
+			LocalLocation = LocalParentTrack->GetLocation();
+		}
+		else
+		{
+			LineY1 = GetPosY();
+			LineY2 = LineY1 + GetHeight();
+			LocalLocation = GetLocation();
+		}
+
+		const FTimingTrackViewport& Viewport = Context.GetViewport();
+		switch (LocalLocation)
+		{
+			case ETimingTrackLocation::Scrollable:
+			{
+				const float TopY = Viewport.GetPosY() + Viewport.GetTopOffset();
+				if (LineY1 < TopY)
+				{
+					LineY1 = TopY;
+				}
+				const float BottomY = Viewport.GetPosY() + Viewport.GetHeight() - Viewport.GetBottomOffset();
+				if (LineY2 > BottomY)
+				{
+					LineY2 = BottomY;
+				}
+				break;
+			}
+			case ETimingTrackLocation::TopDocked:
+			{
+				const float TopY = Viewport.GetPosY();
+				if (LineY1 < TopY)
+				{
+					LineY1 = TopY;
+				}
+				const float BottomY = Viewport.GetPosY() + Viewport.GetTopOffset();
+				if (LineY2 > BottomY)
+				{
+					LineY2 = BottomY;
+				}
+				break;
+			}
+			case ETimingTrackLocation::BottomDocked:
+			{
+				const float TopY = Viewport.GetPosY() + Viewport.GetHeight() - Viewport.GetBottomOffset();
+				if (LineY1 < TopY)
+				{
+					LineY1 = TopY;
+				}
+				const float BottomY = Viewport.GetPosY() + Viewport.GetHeight();
+				if (LineY2 > BottomY)
+				{
+					LineY2 = BottomY;
+				}
+				break;
+			}
+		}
+
+		const float LineH = LineY2 - LineY1;
+		if (LineH > 0.0f)
+		{
+			const FTimingViewDrawHelper& Helper = *static_cast<const FTimingViewDrawHelper*>(&Context.GetHelper());
+			Helper.DrawContextSwitchMarkers(GetDrawState(), LineY1, LineH, 0.25f, bAreOverlaysVisible, bAreExtendedLinesVisible);
+		}
+	}
+}
+
 ////////////////////////////////////////////////////////////////////////////////////////////////////
 
 } // namespace UE::Insights::TimingProfiler
diff --git a/Engine/Source/Developer/TraceInsights/Private/Insights/TimingProfiler/Tracks/ThreadTimingTrack.h b/Engine/Source/Developer/TraceInsights/Private/Insights/TimingProfiler/Tracks/ThreadTimingTrack.h
--- a/Engine/Source/Developer/TraceInsights/Private/Insights/TimingProfiler/Tracks/ThreadTimingTrack.h
+++ b/Engine/Source/Developer/TraceInsights/Private/Insights/TimingProfiler/Tracks/ThreadTimingTrack.h
@@ -110,6 +110,8 @@ private:
 
 class FCpuTimingTrack : public FThreadTimingTrack
 {
+	INSIGHTS_DECLARE_RTTI(FCpuTimingTrack, FThreadTimingTrack)
+
 public:
 	explicit FCpuTimingTrack(FThreadTimingSharedState& InSharedState, const FString& InName, const TCHAR* InGroupName, uint32 InTimelineIndex, uint32 InThreadId)
 		: FThreadTimingTrack(InSharedState, InName, InGroupName, InTimelineIndex, InThreadId)
@@ -121,6 +123,8 @@ public:
 
 class FGpuTimingTrack : public FThreadTimingTrack
 {
+	INSIGHTS_DECLARE_RTTI(FGpuTimingTrack, FThreadTimingTrack)
+
 public:
 	static constexpr uint32 Gpu1ThreadId = uint32('GPU1');
 	static constexpr uint32 Gpu2ThreadId = uint32('GPU2');
@@ -134,4 +138,40 @@ public:
 
 ////////////////////////////////////////////////////////////////////////////////////////////////////
 
+class FGpuQueueTimingTrack : public FThreadTimingTrack
+{
+	INSIGHTS_DECLARE_RTTI(FGpuQueueTimingTrack, FThreadTimingTrack)
+
+public:
+	explicit FGpuQueueTimingTrack(FThreadTimingSharedState& InSharedState, const FString& InName, uint32 InTimelineIndex, uint32 InQueueId)
+		: FThreadTimingTrack(InSharedState, InName, nullptr, InTimelineIndex, InQueueId)
+	{
+	}
+
+	uint32 GetQueueId() const { return GetThreadId(); }
+};
+
+////////////////////////////////////////////////////////////////////////////////////////////////////
+
+class FGpuQueueWorkTimingTrack : public FThreadTimingTrack
+{
+	INSIGHTS_DECLARE_RTTI(FGpuQueueWorkTimingTrack, FThreadTimingTrack)
+
+public:
+	explicit FGpuQueueWorkTimingTrack(FThreadTimingSharedState& InSharedState, const FString& InName, uint32 InTimelineIndex, uint32 InQueueId)
+		: FThreadTimingTrack(InSharedState, InName, nullptr, InTimelineIndex, InQueueId)
+	{
+	}
+
+	uint32 GetQueueId() const { return GetThreadId(); }
+
+	virtual void Draw(const ITimingTrackDrawContext& Context) const override;
+	virtual void PostDraw(const ITimingTrackDrawContext& Context) const override;
+
+private:
+	void DrawLineEvents(const ITimingTrackDrawContext& Context, const float OffsetY) const;
+};
+
+////////////////////////////////////////////////////////////////////////////////////////////////////
+
 } // namespace UE::Insights::TimingProfiler
diff --git a/Engine/Source/Developer/TraceInsights/Private/Insights/TimingProfiler/ViewModels/ThreadTimingSharedState.cpp b/Engine/Source/Developer/TraceInsights/Private/Insights/TimingProfiler/ViewModels/ThreadTimingSharedState.cpp
--- a/Engine/Source/Developer/TraceInsights/Private/Insights/TimingProfiler/ViewModels/ThreadTimingSharedState.cpp
+++ b/Engine/Source/Developer/TraceInsights/Private/Insights/TimingProfiler/ViewModels/ThreadTimingSharedState.cpp
@@ -90,9 +90,9 @@ FThreadTimingSharedState::FThreadTimingSharedState(STimingView* InTimingView)
 
 ////////////////////////////////////////////////////////////////////////////////////////////////////
 
-TSharedPtr<FGpuTimingTrack> FThreadTimingSharedState::GetGpuTrack(uint32 InQueueId)
+TSharedPtr<FGpuQueueTimingTrack> FThreadTimingSharedState::GetGpuTrack(uint32 InQueueId)
 {
-	TSharedPtr<FGpuTimingTrack>* const TrackPtrPtr = GpuTracks.Find(InQueueId);
+	TSharedPtr<FGpuQueueTimingTrack>* const TrackPtrPtr = GpuTracks.Find(InQueueId);
 	return TrackPtrPtr ? *TrackPtrPtr : nullptr;
 }
 
@@ -132,7 +132,7 @@ bool FThreadTimingSharedState::IsAnyGpuTrackVisible() const
 
 bool FThreadTimingSharedState::IsGpuTrackVisible(uint32 InQueueId) const
 {
-	const TSharedPtr<FGpuTimingTrack>* const TrackPtrPtr = GpuTracks.Find(InQueueId);
+	const TSharedPtr<FGpuQueueTimingTrack>* const TrackPtrPtr = GpuTracks.Find(InQueueId);
 	return TrackPtrPtr && (*TrackPtrPtr)->IsVisible();
 }
 
@@ -143,7 +143,7 @@ void FThreadTimingSharedState::GetVisibleGpuQueues(TSet<uint32>& OutSet) const
 	OutSet.Reset();
 	for (const auto& KV : GpuTracks)
 	{
-		const FGpuTimingTrack& Track = *KV.Value;
+		const FGpuQueueTimingTrack& Track = *KV.Value;
 		if (Track.IsVisible())
 		{
 			OutSet.Add(KV.Key);
@@ -208,7 +208,7 @@ void FThreadTimingSharedState::GetVisibleTimelineIndexes(TSet<uint32>& OutSet) c
 
 	for (const auto& KV : GpuTracks)
 	{
-		const FGpuTimingTrack& Track = *KV.Value;
+		const FGpuQueueTimingTrack& Track = *KV.Value;
 		if (Track.IsVisible())
 		{
 			OutSet.Add(Track.GetTimelineIndex());
@@ -326,52 +326,34 @@ void FThreadTimingSharedState::Tick(Timing::ITimingViewSession& InSession, const
 			// Iterate through GPU queues.
 			TimingProfilerProvider->EnumerateGpuQueues([this, &InSession, &bTracksOrderChanged, &GpuTrackOrder](const TraceServices::FGpuQueueInfo& QueueInfo)
 			{
-				// Check if there is an available GPU Work track for this queue.
-				// TODO: Attach this track as an overlay/child track to the normal GPU Queue track.
-				TSharedPtr<FGpuTimingTrack>* WorkTrackPtrPtr = GpuTracks.Find(QueueInfo.Id | 0xFF000000);
-				if (WorkTrackPtrPtr == nullptr)
-				{
-					// Create new Timing Events track for the GPU queue.
-					const FString TrackName = FString::Printf(TEXT("GPU%u - %s %u - WORK"), QueueInfo.GPU, QueueInfo.Name, QueueInfo.Index);
-					TSharedPtr<FGpuTimingTrack> Track = MakeShared<FGpuTimingTrack>(*this, TrackName, nullptr, QueueInfo.WorkTimelineIndex, QueueInfo.Id | 0xFF000000);
-					Track->SetOrder(GpuTrackOrder);
-					Track->SetVisibilityFlag(bShowHideAllGpuTracks);
-					GpuTracks.Add(QueueInfo.Id | 0xFF000000, Track);
-					InSession.AddScrollableTrack(Track);
-				}
-				else
-				{
-					TSharedPtr<FGpuTimingTrack> Track = *WorkTrackPtrPtr;
-					if (Track->GetOrder() != GpuTrackOrder)
-					{
-						Track->SetOrder(GpuTrackOrder);
-						bTracksOrderChanged = true;
-					}
-				}
-				GpuTrackOrder += 1;
-
 				// Check if there is an available GPU track for this queue.
-				TSharedPtr<FGpuTimingTrack>* TrackPtrPtr = GpuTracks.Find(QueueInfo.Id);
+				TSharedPtr<FGpuQueueTimingTrack>* TrackPtrPtr = GpuTracks.Find(QueueInfo.Id);
 				if (TrackPtrPtr == nullptr)
 				{
 					// Create new Timing Events track for the GPU Queue.
 					const FString TrackName = FString::Printf(TEXT("GPU%u - %s %u"), QueueInfo.GPU, QueueInfo.Name, QueueInfo.Index);
-					TSharedPtr<FGpuTimingTrack> Track = MakeShared<FGpuTimingTrack>(*this, TrackName, nullptr, QueueInfo.TimelineIndex, QueueInfo.Id);
+					TSharedPtr<FGpuQueueTimingTrack> Track = MakeShared<FGpuQueueTimingTrack>(*this, TrackName, QueueInfo.TimelineIndex, QueueInfo.Id);
 					Track->SetOrder(GpuTrackOrder);
 					Track->SetVisibilityFlag(bShowHideAllGpuTracks);
 					GpuTracks.Add(QueueInfo.Id, Track);
 					InSession.AddScrollableTrack(Track);
+
+					// Create the GPU Work track and attach it to the GPU queue track.
+					const FString WorkTrackName = FString::Printf(TEXT("GPU%u - %s %u - WORK"), QueueInfo.GPU, QueueInfo.Name, QueueInfo.Index);
+					TSharedPtr<FGpuQueueWorkTimingTrack> WorkTrack = MakeShared<FGpuQueueWorkTimingTrack>(*this, WorkTrackName, QueueInfo.WorkTimelineIndex, QueueInfo.Id);
+					Track->SetChildTrack(WorkTrack);
+					WorkTrack->SetParentTrack(Track);
 				}
 				else
 				{
-					TSharedPtr<FGpuTimingTrack> Track = *TrackPtrPtr;
+					TSharedPtr<FGpuQueueTimingTrack> Track = *TrackPtrPtr;
 					if (Track->GetOrder() != GpuTrackOrder)
 					{
 						Track->SetOrder(GpuTrackOrder);
 						bTracksOrderChanged = true;
 					}
 				}
-				GpuTrackOrder += 99;
+				GpuTrackOrder += 100;
 			});
 
 			// Iterate through threads.
@@ -590,7 +572,7 @@ void FThreadTimingSharedState::SetAllGpuTracksToggle(bool bOnOff)
 	}
 	for (const auto& KV : GpuTracks)
 	{
-		FGpuTimingTrack& Track = *KV.Value;
+		FGpuQueueTimingTrack& Track = *KV.Value;
 		Track.SetVisibilityFlag(bShowHideAllGpuTracks);
 	}
 
@@ -675,7 +657,7 @@ TSharedPtr<const ITimingEvent> FThreadTimingSharedState::FindMaxEventInstance(ui
 
 	for (const auto& KV : GpuTracks)
 	{
-		const FGpuTimingTrack& Track = *KV.Value;
+		const FGpuQueueTimingTrack& Track = *KV.Value;
 		if (Track.IsVisible())
 		{
 			TrackEvent = Track.FindMaxEventInstance(TimerId, StartTime, EndTime);
@@ -730,7 +712,7 @@ TSharedPtr<const ITimingEvent> FThreadTimingSharedState::FindMinEventInstance(ui
 
 	for (const auto& KV : GpuTracks)
 	{
-		const FGpuTimingTrack& Track = *KV.Value;
+		const FGpuQueueTimingTrack& Track = *KV.Value;
 		if (Track.IsVisible())
 		{
 			TrackEvent = Track.FindMinEventInstance(TimerId, StartTime, EndTime);
diff --git a/Engine/Source/Developer/TraceInsights/Private/Insights/TimingProfiler/ViewModels/ThreadTimingSharedState.h b/Engine/Source/Developer/TraceInsights/Private/Insights/TimingProfiler/ViewModels/ThreadTimingSharedState.h
--- a/Engine/Source/Developer/TraceInsights/Private/Insights/TimingProfiler/ViewModels/ThreadTimingSharedState.h
+++ b/Engine/Source/Developer/TraceInsights/Private/Insights/TimingProfiler/ViewModels/ThreadTimingSharedState.h
@@ -22,6 +22,8 @@ namespace UE::Insights::TimingProfiler
 
 class FCpuTimingTrack;
 class FGpuTimingTrack;
+class FGpuQueueTimingTrack;
+class FGpuQueueWorkTimingTrack;
 class STimingView;
 
 ////////////////////////////////////////////////////////////////////////////////////////////////////
@@ -45,7 +47,7 @@ public:
 
 	TSharedPtr<FGpuTimingTrack> GetOldGpu1Track() { return OldGpu1Track; }
 	TSharedPtr<FGpuTimingTrack> GetOldGpu2Track() { return OldGpu2Track; }
-	TSharedPtr<FGpuTimingTrack> GetGpuTrack(uint32 InQueueId);
+	TSharedPtr<FGpuQueueTimingTrack> GetGpuTrack(uint32 InQueueId);
 	TSharedPtr<FCpuTimingTrack> GetCpuTrack(uint32 InThreadId);
 	const TMap<uint32, TSharedPtr<FCpuTimingTrack>> GetAllCpuTracks() { return CpuTracks; }
 
@@ -103,7 +105,7 @@ private:
 	TSharedPtr<FGpuTimingTrack> OldGpu2Track;
 
 	/** Maps GPU queue id to track pointer. */
-	TMap<uint32, TSharedPtr<FGpuTimingTrack>> GpuTracks;
+	TMap<uint32, TSharedPtr<FGpuQueueTimingTrack>> GpuTracks;
 
 	/** Maps CPU thread id to track pointer. */
 	TMap<uint32, TSharedPtr<FCpuTimingTrack>> CpuTracks;
