[Insights] GPU Work Track Settings

diff --git a/Engine/Source/Developer/TraceInsights/Private/Insights/TimingProfiler/Tracks/ThreadTimingTrack.h b/Engine/Source/Developer/TraceInsights/Private/Insights/TimingProfiler/Tracks/ThreadTimingTrack.h
--- a/Engine/Source/Developer/TraceInsights/Private/Insights/TimingProfiler/Tracks/ThreadTimingTrack.h
+++ b/Engine/Source/Developer/TraceInsights/Private/Insights/TimingProfiler/Tracks/ThreadTimingTrack.h
@@ -81,6 +81,7 @@ public:
 
 protected:
 	virtual bool HasCustomFilter() const override;
+	FThreadTimingSharedState& GetSharedState() const { return SharedState; }
 
 private:
 	bool FindTimingProfilerEvent(const FThreadTrackEvent& InTimingEvent, TFunctionRef<void(double, double, uint32, const TraceServices::FTimingProfilerEvent&)> InFoundPredicate) const;
diff --git a/Engine/Source/Developer/TraceInsights/Private/Insights/TimingProfiler/ViewModels/ThreadTimingSharedState.cpp b/Engine/Source/Developer/TraceInsights/Private/Insights/TimingProfiler/ViewModels/ThreadTimingSharedState.cpp
--- a/Engine/Source/Developer/TraceInsights/Private/Insights/TimingProfiler/ViewModels/ThreadTimingSharedState.cpp
+++ b/Engine/Source/Developer/TraceInsights/Private/Insights/TimingProfiler/ViewModels/ThreadTimingSharedState.cpp
@@ -17,6 +17,7 @@
 #include "InsightsCore/Common/Log.h"
 
 // TraceInsights
+#include "Insights/InsightsManager.h"
 #include "Insights/InsightsStyle.h"
 #include "Insights/IUnrealInsightsModule.h"
 #include "Insights/ITimingViewSession.h"
@@ -43,6 +44,15 @@ public:
 	/** Toggles visibility for GPU thread track(s). */
 	TSharedPtr<FUICommandInfo> ShowHideAllGpuTracks;
 
+	/** Toggles visibility for GPU work header tracks. */
+	TSharedPtr<FUICommandInfo> Command_ShowWorkTracks;
+
+	/** Extends the visualization of GPU work events over the GPU timing tracks. */
+	TSharedPtr<FUICommandInfo> Command_ShowOverlays;
+
+	/** Shows/hides the extended vertical lines at the edges of each GPU work event. */
+	TSharedPtr<FUICommandInfo> Command_ShowExtendedLines;
+
 	/** Toggles visibility for all CPU thread tracks at once. */
 	TSharedPtr<FUICommandInfo> ShowHideAllCpuTracks;
 };
@@ -70,6 +80,24 @@ void FThreadTimingViewCommands::RegisterCommands()
 		EUserInterfaceActionType::ToggleButton,
 		FInputChord(EKeys::Y));
 
+	UI_COMMAND(Command_ShowWorkTracks,
+		"Show GPU Work Track(s)",
+		"Shows/hides the GPU Work header track(s).",
+		EUserInterfaceActionType::ToggleButton,
+		FInputChord());
+
+	UI_COMMAND(Command_ShowOverlays,
+		"Show Overlays",
+		"Extends the visualization of GPU work events over the GPU timing tracks.",
+		EUserInterfaceActionType::ToggleButton,
+		FInputChord());
+
+	UI_COMMAND(Command_ShowExtendedLines,
+		"Show Extended Lines",
+		"Shows/hides the extended vertical lines at the edges of each GPU work event.",
+		EUserInterfaceActionType::ToggleButton,
+		FInputChord());
+
 	UI_COMMAND(ShowHideAllCpuTracks,
 		"CPU Thread Tracks",
 		"Shows/hides all CPU tracks (and all CPU thread groups).",
@@ -228,11 +256,17 @@ void FThreadTimingSharedState::OnBeginSession(Timing::ITimingViewSession& InSess
 	if (TimingView->GetName() == FInsightsManagerTabs::LoadingProfilerTabId)
 	{
 		bShowHideAllGpuTracks = false;
+		bShowGpuWorkTracks = false;
+		bShowOverlaysInGpuQueueTracks = false;
+		bShowExtendedLinesInGpuQueueTracks = false;
 		bShowHideAllCpuTracks = false;
 	}
 	else
 	{
 		bShowHideAllGpuTracks = true;
+		bShowGpuWorkTracks = true;
+		bShowOverlaysInGpuQueueTracks = true;
+		bShowExtendedLinesInGpuQueueTracks = true;
 		bShowHideAllCpuTracks = true;
 	}
 
@@ -338,11 +372,14 @@ void FThreadTimingSharedState::Tick(Timing::ITimingViewSession& InSession, const
 					GpuTracks.Add(QueueInfo.Id, Track);
 					InSession.AddScrollableTrack(Track);
 
-					// Create the GPU Work track and attach it to the GPU queue track.
-					const FString WorkTrackName = FString::Printf(TEXT("GPU%u - %s %u - WORK"), QueueInfo.GPU, QueueInfo.Name, QueueInfo.Index);
-					TSharedPtr<FGpuQueueWorkTimingTrack> WorkTrack = MakeShared<FGpuQueueWorkTimingTrack>(*this, WorkTrackName, QueueInfo.WorkTimelineIndex, QueueInfo.Id);
-					Track->SetChildTrack(WorkTrack);
-					WorkTrack->SetParentTrack(Track);
+					if (bShowGpuWorkTracks)
+					{
+						// Create the GPU Work track and attach it to the GPU queue track.
+						const FString WorkTrackName = FString::Printf(TEXT("GPU%u - %s %u - WORK"), QueueInfo.GPU, QueueInfo.Name, QueueInfo.Index);
+						TSharedPtr<FGpuQueueWorkTimingTrack> WorkTrack = MakeShared<FGpuQueueWorkTimingTrack>(*this, WorkTrackName, QueueInfo.WorkTimelineIndex, QueueInfo.Id);
+						Track->SetChildTrack(WorkTrack);
+						WorkTrack->SetParentTrack(Track);
+					}
 				}
 				else
 				{
@@ -457,6 +494,9 @@ void FThreadTimingSharedState::ExtendGpuTracksFilterMenu(Timing::ITimingViewSess
 	InOutMenuBuilder.BeginSection("GpuTracks", LOCTEXT("ContextMenu_Section_GpuTracks", "GPU Tracks"));
 	{
 		InOutMenuBuilder.AddMenuEntry(FThreadTimingViewCommands::Get().ShowHideAllGpuTracks);
+		InOutMenuBuilder.AddMenuEntry(FThreadTimingViewCommands::Get().Command_ShowWorkTracks);
+		InOutMenuBuilder.AddMenuEntry(FThreadTimingViewCommands::Get().Command_ShowOverlays);
+		InOutMenuBuilder.AddMenuEntry(FThreadTimingViewCommands::Get().Command_ShowExtendedLines);
 	}
 	InOutMenuBuilder.EndSection();
 }
@@ -496,6 +536,24 @@ void FThreadTimingSharedState::BindCommands()
 		FCanExecuteAction(),
 		FIsActionChecked::CreateSP(this, &FThreadTimingSharedState::IsAllGpuTracksToggleOn));
 
+	CommandList->MapAction(
+		FThreadTimingViewCommands::Get().Command_ShowWorkTracks,
+		FExecuteAction::CreateSP(this, &FThreadTimingSharedState::Command_ShowGpuWorkTracks_Execute),
+		FCanExecuteAction::CreateSP(this, &FThreadTimingSharedState::Command_ShowGpuWorkTracks_CanExecute),
+		FIsActionChecked::CreateSP(this, &FThreadTimingSharedState::Command_ShowGpuWorkTracks_IsChecked));
+
+	CommandList->MapAction(
+		FThreadTimingViewCommands::Get().Command_ShowOverlays,
+		FExecuteAction::CreateSP(this, &FThreadTimingSharedState::Command_ShowOverlays_Execute),
+		FCanExecuteAction::CreateSP(this, &FThreadTimingSharedState::Command_ShowOverlays_CanExecute),
+		FIsActionChecked::CreateSP(this, &FThreadTimingSharedState::Command_ShowOverlays_IsChecked));
+
+	CommandList->MapAction(
+		FThreadTimingViewCommands::Get().Command_ShowExtendedLines,
+		FExecuteAction::CreateSP(this, &FThreadTimingSharedState::Command_ShowExtendedLines_Execute),
+		FCanExecuteAction::CreateSP(this, &FThreadTimingSharedState::Command_ShowExtendedLines_CanExecute),
+		FIsActionChecked::CreateSP(this, &FThreadTimingSharedState::Command_ShowExtendedLines_IsChecked));
+
 	CommandList->MapAction(
 		FThreadTimingViewCommands::Get().ShowHideAllCpuTracks,
 		FExecuteAction::CreateSP(this, &FThreadTimingSharedState::ShowHideAllCpuTracks),
@@ -725,6 +783,66 @@ TSharedPtr<const ITimingEvent> FThreadTimingSharedState::FindMinEventInstance(ui
 
 ////////////////////////////////////////////////////////////////////////////////////////////////////
 
+void FThreadTimingSharedState::SetGpuWorkTracksVisibility(bool bOnOff)
+{
+	if (bShowGpuWorkTracks != bOnOff)
+	{
+		bShowGpuWorkTracks = bOnOff;
+
+		if (bShowGpuWorkTracks)
+		{
+			AddGpuWorkChildTracks();
+		}
+		else
+		{
+			RemoveGpuWorkChildTracks();
+		}
+	}
+}
+
+////////////////////////////////////////////////////////////////////////////////////////////////////
+
+void FThreadTimingSharedState::AddGpuWorkChildTracks()
+{
+	TSharedPtr<const TraceServices::IAnalysisSession> Session = FInsightsManager::Get()->GetSession();
+	if (!Session.IsValid())
+	{
+		return;
+	}
+
+	const TraceServices::ITimingProfilerProvider* TimingProfilerProvider = TraceServices::ReadTimingProfilerProvider(*Session);
+
+	TraceServices::FAnalysisSessionReadScope SessionReadScope(*Session);
+
+	TimingProfilerProvider->EnumerateGpuQueues([this](const TraceServices::FGpuQueueInfo& QueueInfo)
+	{
+		// Check if there is an available GPU track for this queue.
+		TSharedPtr<FGpuQueueTimingTrack>* TrackPtrPtr = GpuTracks.Find(QueueInfo.Id);
+
+		if (TrackPtrPtr != nullptr)
+		{
+			// Create the GPU Work track and attach it to the GPU queue track.
+			const FString WorkTrackName = FString::Printf(TEXT("GPU%u - %s %u - WORK"), QueueInfo.GPU, QueueInfo.Name, QueueInfo.Index);
+			TSharedPtr<FGpuQueueWorkTimingTrack> WorkTrack = MakeShared<FGpuQueueWorkTimingTrack>(*this, WorkTrackName, QueueInfo.WorkTimelineIndex, QueueInfo.Id);
+			(*TrackPtrPtr)->SetChildTrack(WorkTrack);
+			WorkTrack->SetParentTrack(*TrackPtrPtr);
+		}
+	});
+}
+
+////////////////////////////////////////////////////////////////////////////////////////////////////
+
+void FThreadTimingSharedState::RemoveGpuWorkChildTracks()
+{
+	for (auto& Pair : GpuTracks)
+	{
+		TSharedPtr<FGpuQueueTimingTrack> Track = Pair.Value;
+		Track->SetChildTrack(nullptr);
+	}
+}
+
+////////////////////////////////////////////////////////////////////////////////////////////////////
+
 } // namespace UE::Insights::TimingProfiler
 
 #undef LOCTEXT_NAMESPACE
diff --git a/Engine/Source/Developer/TraceInsights/Private/Insights/TimingProfiler/ViewModels/ThreadTimingSharedState.h b/Engine/Source/Developer/TraceInsights/Private/Insights/TimingProfiler/ViewModels/ThreadTimingSharedState.h
--- a/Engine/Source/Developer/TraceInsights/Private/Insights/TimingProfiler/ViewModels/ThreadTimingSharedState.h
+++ b/Engine/Source/Developer/TraceInsights/Private/Insights/TimingProfiler/ViewModels/ThreadTimingSharedState.h
@@ -86,6 +86,12 @@ public:
 	void HideAllCpuTracks() { SetAllCpuTracksToggle(false); }
 	void ShowHideAllCpuTracks() { SetAllCpuTracksToggle(!IsAllCpuTracksToggleOn()); }
 
+	bool AreOverlaysVisibleInGpuQueueTracks() { return bShowOverlaysInGpuQueueTracks; }
+	bool AreExtendedLinesVisibleInGpuQueueTracks() { return bShowExtendedLinesInGpuQueueTracks; }
+
+	bool AreGpuWorkTracksVisible() const { return bShowGpuWorkTracks; }
+	void SetGpuWorkTracksVisibility(bool bOnOff);
+
 	TSharedPtr<const ITimingEvent> FindMaxEventInstance(uint32 TimerId, double StartTime, double EndTime);
 	TSharedPtr<const ITimingEvent> FindMinEventInstance(uint32 TimerId, double StartTime, double EndTime);
 
@@ -95,10 +101,28 @@ private:
 	bool ToggleTrackVisibilityByGroup_IsChecked(const TCHAR* InGroupName) const;
 	void ToggleTrackVisibilityByGroup_Execute(const TCHAR* InGroupName);
 
+	void Command_ShowGpuWorkTracks_Execute() { SetGpuWorkTracksVisibility(!bShowGpuWorkTracks); }
+	bool Command_ShowGpuWorkTracks_CanExecute() { return GpuTracks.Num() > 0; }
+	bool Command_ShowGpuWorkTracks_IsChecked() { return bShowGpuWorkTracks; }
+
+	void Command_ShowOverlays_Execute() { bShowOverlaysInGpuQueueTracks = !bShowOverlaysInGpuQueueTracks; }
+	bool Command_ShowOverlays_CanExecute() { return bShowGpuWorkTracks && GpuTracks.Num() > 0; }
+	bool Command_ShowOverlays_IsChecked() { return bShowOverlaysInGpuQueueTracks; }
+
+	void Command_ShowExtendedLines_Execute() { bShowExtendedLinesInGpuQueueTracks = !bShowExtendedLinesInGpuQueueTracks; }
+	bool Command_ShowExtendedLines_CanExecute() { return bShowGpuWorkTracks && GpuTracks.Num() > 0; }
+	bool Command_ShowExtendedLines_IsChecked() { return bShowExtendedLinesInGpuQueueTracks; }
+
+	void AddGpuWorkChildTracks();
+	void RemoveGpuWorkChildTracks();
+
 private:
 	STimingView* TimingView = nullptr;
 
 	bool bShowHideAllGpuTracks = false;
+	bool bShowGpuWorkTracks = true;
+	bool bShowOverlaysInGpuQueueTracks = true;
+	bool bShowExtendedLinesInGpuQueueTracks = true;
 	bool bShowHideAllCpuTracks = false;
 
 	TSharedPtr<FGpuTimingTrack> OldGpu1Track;
