Insights Utils refactoring
Moved latent commands to a separate file
Reorganised latent commands.
Minor fixes.

diff --git a/Engine/Source/Developer/TraceInsights/Private/Insights/Tests/InsightsLatentCommands.cpp b/Engine/Source/Developer/TraceInsights/Private/Insights/Tests/InsightsLatentCommands.cpp
new file mode 100644
--- /dev/null
+++ b/Engine/Source/Developer/TraceInsights/Private/Insights/Tests/InsightsLatentCommands.cpp
@@ -0,0 +1,324 @@
+// Copyright Epic Games, Inc. All Rights Reserved.
+#include "Insights/Tests/InsightsLatentCommands.h"
+
+#include "CoreMinimal.h"
+#include "HAL/PlatformOutputDevices.h"
+#include "Modules/ModuleManager.h"
+
+#include "Insights/IUnrealInsightsModule.h"
+
+
+#if WITH_DEV_AUTOMATION_TESTS
+
+namespace
+{
+
+	const FString GetTestProfilingDirPath()
+	{
+		static const FString ProfilingDirPath = FPaths::ProfilingDir();
+
+		return ProfilingDirPath;
+	}
+
+} //anonymous
+
+namespace UE::Insights::Tests::LatentCommands
+{
+#if WITH_EDITOR
+	bool FExecuteTraceCommand::Update()
+	{
+		IConsoleManager::Get().ProcessUserConsoleInput(*ExpectedCommand, *GLog, GEngine->GetWorld());
+
+		return true;
+	}
+#endif //WITH_EDITOR
+
+	bool FResetSessionCommand::Update()
+	{
+		Utils.ResetSession();
+
+		return true;
+	}
+
+	bool FVerifyEnabledChannelsCommand::Update()
+	{
+		FStringBuilderBase ActiveChannels;
+		FTraceAuxiliary::GetActiveChannelsString(ActiveChannels);
+		const FStringView ActiveChannelsV = ActiveChannels.ToView();
+
+		if (ExpectedChannelsList.IsEmpty())
+		{
+			return ActiveChannelsV.IsEmpty();
+		}
+
+		for (const FString& Channel : ExpectedChannelsList)
+		{
+			if (!ActiveChannelsV.Contains(FStringView(*Channel), ESearchCase::IgnoreCase))
+			{
+				Test->AddError(FString::Printf(TEXT("Channel %s doesn't exist in the list of active channels"), *Channel));
+			}
+		}
+
+		return true;
+	}
+
+	bool FCompareTraceAndSnapshotFileSizeCommand::Update()
+	{
+		TArray<FString> SnapshotFilePaths;
+		Constants.PlatformFile.FindFiles(SnapshotFilePaths, *GetTestProfilingDirPath(), TEXT(".utrace"));
+		if (SnapshotFilePaths.IsEmpty())
+		{
+			Test->AddError(TEXT("The profiling directory shouldn't be empty"));
+			return true;
+		}
+		const FString SnapshotName = FPaths::GetCleanFilename(SnapshotFilePaths.Last());
+		int64 SnapshotFileSize = Constants.PlatformFile.FileSize(*(GetTestProfilingDirPath() / SnapshotName));
+
+		// two values below are in MB
+		static constexpr int64 MinAllowedSnapshotSize = 8;
+		static constexpr int64 MaxAllowedSnapshotSize = 50;
+		int64 ScaleFactor = 1000000;
+		int64 RoundedSnapshotSizeValue = FMath::FloorToInt(static_cast<float>(SnapshotFileSize) / ScaleFactor);
+
+		if (FMath::IsWithin<int64>(RoundedSnapshotSizeValue, MinAllowedSnapshotSize, MaxAllowedSnapshotSize))
+		{
+			Test->AddInfo(FString::Printf(TEXT("The snapshot %s has expected size %d MB"), *SnapshotName, RoundedSnapshotSizeValue));
+		}
+		else
+		{
+			Test->AddError(FString::Printf(TEXT("The snapshot should be from %d MB to %d  MB in size, but it has %d MB size"), MinAllowedSnapshotSize, MaxAllowedSnapshotSize, RoundedSnapshotSizeValue));
+		}
+
+		return true;
+	}
+
+
+	bool FVerifyTraceSnapshotSizeCommand::Update()
+	{
+		TArray<FString> SnapshotFilePaths;
+		Constants.PlatformFile.FindFiles(SnapshotFilePaths, *GetTestProfilingDirPath(), TEXT(".utrace"));
+		if (SnapshotFilePaths.IsEmpty())
+		{
+			Test->AddError(TEXT("The snapshot directory is empty"));
+			return true;
+		}
+
+		FString SnapshotName = FPaths::GetCleanFilename(SnapshotFilePaths.Last());
+		int64 SnapshotFileSize = Constants.PlatformFile.FileSize(*(GetTestProfilingDirPath() / SnapshotName));
+		Test->TestTrue(TEXT("The snapshot shouldn't has zero size, but it has %d size"), SnapshotFileSize != 0);
+		Test->AddInfo(FString::Printf(TEXT("The snapshot size is %d"), SnapshotFileSize));
+
+		return true;
+	}
+
+	bool FStopSessionCommand::Update()
+	{
+		FTraceAuxiliary::Stop();
+		if (FPlatformTime::Seconds() - StartTime >= Constants.Timeout)
+		{
+			Test->AddError(TEXT("The FStopSessionCommand timed out"));
+			return true;
+		}
+
+		if (!FTraceAuxiliary::IsConnected())
+		{
+			Test->AddInfo(TEXT("Tracing stopped"));
+			return true;
+		}
+
+		return false;
+	}
+
+	bool FVerifyTraceSizeCommand::Update()
+	{
+		TArray<FString> StoreFilePaths;
+		Constants.PlatformFile.FindFiles(StoreFilePaths, *StoreDirPath, TEXT(".utrace"));
+		if (StoreFilePaths.IsEmpty())
+		{
+			Test->AddError(TEXT("The trace store directory is empty"));
+			return true;
+		}
+
+		FString TraceName = FPaths::GetCleanFilename(StoreFilePaths.Last());
+		int64 TraceFileSize = Constants.PlatformFile.FileSize(*(StoreDirPath / TraceName));
+		Test->TestTrue(FString::Printf(TEXT("The trace should has zero size, but it has %d size"), TraceFileSize), TraceFileSize == 0);
+
+		return true;
+	}
+
+	bool FWaitForConnectionStatusCommand::Update()
+	{
+		if (FPlatformTime::Seconds() - StartTime >= Constants.Timeout)
+		{
+			Test->AddError(TEXT("The FWaitForConnectionStatusCommand timed out"));
+			return true;
+		}
+
+		if (FTraceAuxiliary::IsConnected() == bExpectedConnectionResult)
+		{
+			return true;
+		}
+
+		return false;
+	}
+
+	bool FKillUTSCommand::Update()
+	{
+		Utils.KillUTS();
+
+		return true;
+	}
+
+	bool FDeleteTraceFile(const FString& TraceDirPath, FAutomationTestBase* Test, double StartTime)
+	{
+		TArray<FString> TraceFilePaths;
+		const FString UnrealTraceServerName = TEXT("UnrealTraceServer");
+		const FString AutomationToolPath = FPaths::RootDir() / ("Engine/Programs/AutomationTool/Saved/Logs/");
+		Constants.PlatformFile.FindFiles(TraceFilePaths, *TraceDirPath, TEXT(".utrace"));
+		if (TraceFilePaths.IsEmpty())
+		{
+			Test->AddError(TEXT("The trace directory is empty"));
+			return true;
+		}
+
+		if (FPlatformProcess::IsApplicationRunning(*UnrealTraceServerName) || FTraceAuxiliary::IsConnected())
+		{
+			FTraceAuxiliary::Stop();
+		}
+
+		FString FileName = FPaths::GetCleanFilename(TraceFilePaths.Last());
+		IFileManager::Get().Copy(*(AutomationToolPath + FileName), *TraceFilePaths.Last());
+		IFileManager::Get().Delete(*TraceFilePaths.Last(), false, true);
+
+		if (!Constants.PlatformFile.FileExists(*TraceFilePaths.Last()))
+		{
+			Test->AddInfo(FString::Printf(TEXT("The file %s is deleted"), *FileName));
+			return true;
+		}
+
+		if (FPlatformTime::Seconds() - StartTime >= Constants.Timeout)
+		{
+			Test->AddWarning(FString::Printf(TEXT("The FDeleteTraceFileCommand timed out. The trace file %s isn't deleted"), *TraceFilePaths.Last()));
+			return true;
+		}
+
+		return false;
+	}
+
+	bool FDeleteCustomTraceFileCommand::Update()
+	{
+
+		return FDeleteTraceFile(TraceDirPath, Test, StartTime);
+	}
+
+	bool FDeleteDefaultTraceFileCommand::Update()
+	{
+		return FDeleteTraceFile(GetTestProfilingDirPath(), Test, StartTime);
+	}
+
+	bool FWaitTraceSnapshotNameAppearInLogCommand::Update()
+	{
+		FInsightsTestUtils Utils(Test);
+		TArray<FString> SnapshotFilePaths;
+		Constants.PlatformFile.FindFiles(SnapshotFilePaths, *(GetTestProfilingDirPath()), TEXT(".utrace"));
+		if (SnapshotFilePaths.IsEmpty())
+		{
+			Test->AddError(TEXT("The profiling directory is empty"));
+			return true;
+		}
+
+		FString SnapshotName = FPaths::GetCleanFilename(SnapshotFilePaths.Last());
+		const FString LogFilename = FPlatformOutputDevices::GetAbsoluteLogFilename();
+		if (!Constants.PlatformFile.FileExists(*LogFilename))
+		{
+			Test->AddError(FString::Printf(TEXT("The log file %s does not exist"), *LogFilename));
+			return true;
+		}
+		bool bIsLogContainsTraceStart = Utils.FileContainsString(LogFilename, SnapshotName, 1);
+		if (bIsLogContainsTraceStart)
+		{
+			Test->AddInfo(TEXT("The log contains trace info"));
+			return true;
+		}
+
+		if (FPlatformTime::Seconds() - StartTime >= Constants.Timeout)
+		{
+			Test->AddError(TEXT("The FWaitTraceSnapshotNameAppearInLogCommand timed out"));
+			return true;
+		}
+
+		return false;
+	}
+
+	bool FVerifyExpectedTraceSnapshotSizeCommand::Update()
+	{
+		TArray<FString> SnapshotFilePaths;
+		Constants.PlatformFile.FindFiles(SnapshotFilePaths, *(GetTestProfilingDirPath()), TEXT(".utrace"));
+		if (SnapshotFilePaths.IsEmpty())
+		{
+			Test->AddError(TEXT("The snapshot directory is empty"));
+			return true;
+		}
+
+		FString SnapshotName = FPaths::GetCleanFilename(SnapshotFilePaths.Last());
+		int64 SnapshotFileSize = Constants.PlatformFile.FileSize(*(GetTestProfilingDirPath() / SnapshotName));
+		Test->AddInfo(FString::Printf(TEXT("The snapshot size is %d"), SnapshotFileSize));
+		Test->AddErrorIfFalse(SnapshotFileSize > 0, TEXT("The snapshot size must be greater than zero"));
+		Test->AddErrorIfFalse(SnapshotFileSize < Constants.ExpectedTraceSize, TEXT("The snapshot shouldn't has size bigger than %lld"), static_cast<int32>(Constants.ExpectedTraceSize));
+
+		return true;
+	}
+
+	bool FWaitExpectedTraceFilesizeCommand::Update()
+	{
+		IUnrealInsightsModule& UnrealInsightsModule = FModuleManager::LoadModuleChecked<IUnrealInsightsModule>("TraceInsights");
+		const FString InsightsStoreDirPath = UnrealInsightsModule.GetDefaultStoreDir();
+		FString LiveTraceName = Utils.GetLiveTrace();
+		if (LiveTraceName.IsEmpty())
+		{
+			Test->AddInfo(TEXT("The live trace currently empty. Waiting another"));
+		}
+		int64 ActualFilesize = Constants.PlatformFile.FileSize(*(InsightsStoreDirPath / LiveTraceName));
+		// Converting from B to MiB
+		int64 ScaleFactor = 1000000;
+		int64 RoundedActualSizeValue = FMath::FloorToInt(static_cast<float>(ActualFilesize) / ScaleFactor);
+		int64 RoundedExpectedSizeValue = FMath::FloorToInt(static_cast<float>(Constants.ExpectedTraceSize) / ScaleFactor);
+		if (RoundedActualSizeValue > RoundedExpectedSizeValue)
+		{
+			Test->AddInfo(TEXT("Actual trace size is larger than expected"));
+			return true;
+		}
+
+		if (FPlatformTime::Seconds() - StartTime >= 1200.0)
+		{
+			Test->AddInfo(FString::Printf(TEXT("The WaitExpectedTraceFilesizeCommand timed out. Expected value - %d. Actual value - %d"), RoundedExpectedSizeValue, RoundedActualSizeValue));
+			return true;
+		}
+
+		return false;
+	}
+
+	bool FVerifyFileCreated::Update()
+	{
+		TArray<FString> MatchedTraceFiles;
+		Constants.PlatformFile.FindFiles(MatchedTraceFiles, *GetTestProfilingDirPath(), TEXT(".utrace"));
+		Test->AddInfo("Searching for " + FileNameToCheck + " file in the " + GetTestProfilingDirPath() + " directory");
+		if (MatchedTraceFiles.IsEmpty())
+		{
+			Test->AddError(TEXT("The " + GetTestProfilingDirPath() + " directory is empty"));
+			return true;
+		}
+
+		for (const FString& FileName : MatchedTraceFiles) {
+			if (FPaths::GetCleanFilename(FileName) == FileNameToCheck) {
+				return true;
+			}
+		}
+
+		Test->AddError(TEXT("The " + FileNameToCheck + " file was not found in the " + GetTestProfilingDirPath() + " directory"));
+
+		return false;
+	}
+}
+
+#endif // WITH_DEV_AUTOMATION_TESTS
\ No newline at end of file
diff --git a/Engine/Source/Developer/TraceInsights/Private/Insights/Tests/InsightsTestUtils.cpp b/Engine/Source/Developer/TraceInsights/Private/Insights/Tests/InsightsTestUtils.cpp
--- a/Engine/Source/Developer/TraceInsights/Private/Insights/Tests/InsightsTestUtils.cpp
+++ b/Engine/Source/Developer/TraceInsights/Private/Insights/Tests/InsightsTestUtils.cpp
@@ -23,7 +23,6 @@
 #include "Insights/InsightsManager.h"
 #include "Insights/IUnrealInsightsModule.h"
 
-
 ////////////////////////////////////////////////////////////////////////////////////////////////////
 
 FInsightsTestUtils::FInsightsTestUtils(FAutomationTestBase* InTest) :
@@ -412,4 +411,24 @@ FString FInsightsTestUtils::FInsightsTestUtils::GetUTSPath() const
 	}
 
 	return EntireUTSPath;
-}
\ No newline at end of file
+}
+
+const TArray<FString>& FInsightsTestUtils::GetDefaultTraceChannels()
+{
+	static TArray<FString> DefaultTraceChannels;
+	if (DefaultTraceChannels.IsEmpty())
+	{
+		FString DefaultChannelList;
+		auto GetDefaultPreset = [&DefaultChannelList](const FTraceAuxiliary::FChannelPreset& Preset)
+			{
+				if (FString(Preset.Name) == TEXT("Default"))
+				{
+					DefaultChannelList = Preset.ChannelList;
+				}
+				return FTraceAuxiliary::EEnumerateResult::Stop;
+			};
+		FTraceAuxiliary::EnumerateFixedChannelPresets(GetDefaultPreset);
+		DefaultChannelList.ParseIntoArray(DefaultTraceChannels, TEXT(","));
+	}
+	return DefaultTraceChannels;
+}
diff --git a/Engine/Source/Developer/TraceInsights/Public/Insights/Tests/InsightsLatentCommands.h b/Engine/Source/Developer/TraceInsights/Public/Insights/Tests/InsightsLatentCommands.h
new file mode 100644
--- /dev/null
+++ b/Engine/Source/Developer/TraceInsights/Public/Insights/Tests/InsightsLatentCommands.h
@@ -0,0 +1,67 @@
+// Copyright Epic Games, Inc. All Rights Reserved.
+
+#pragma once
+
+#if WITH_DEV_AUTOMATION_TESTS
+
+
+#include "HAL/PlatformFileManager.h"
+#include "Misc/AutomationTest.h"
+
+#include "Insights/Tests/InsightsTestUtils.h"
+
+namespace UE::Insights::Tests::LatentCommands
+{
+	struct {
+		const double Timeout = 10.0;
+		const int64 ExpectedTraceSize = 3  * 1024 * 1024; // 3MB
+		IPlatformFile& PlatformFile = FPlatformFileManager::Get().GetPlatformFile();
+	} Constants;
+
+
+	DEFINE_EXPORTED_LATENT_AUTOMATION_COMMAND_ONE_PARAMETER(TRACEINSIGHTS_API, FExecuteTraceCommand, const FString, ExpectedCommand);
+
+
+	DEFINE_EXPORTED_LATENT_AUTOMATION_COMMAND_ONE_PARAMETER(TRACEINSIGHTS_API, FResetSessionCommand, FInsightsTestUtils, Utils);
+
+
+	DEFINE_EXPORTED_LATENT_AUTOMATION_COMMAND_TWO_PARAMETER(TRACEINSIGHTS_API, FVerifyEnabledChannelsCommand, const TArray<FString>, ExpectedChannelsList, FAutomationTestBase*, Test);
+
+
+	DEFINE_EXPORTED_LATENT_AUTOMATION_COMMAND_ONE_PARAMETER(TRACEINSIGHTS_API, FCompareTraceAndSnapshotFileSizeCommand, FAutomationTestBase*, Test);
+
+
+	DEFINE_EXPORTED_LATENT_AUTOMATION_COMMAND_ONE_PARAMETER(TRACEINSIGHTS_API, FVerifyTraceSnapshotSizeCommand, FAutomationTestBase*, Test);
+
+
+	DEFINE_EXPORTED_LATENT_AUTOMATION_COMMAND_TWO_PARAMETER(TRACEINSIGHTS_API, FStopSessionCommand, FInsightsTestUtils, Utils, FAutomationTestBase*, Test);
+
+
+	DEFINE_EXPORTED_LATENT_AUTOMATION_COMMAND_TWO_PARAMETER(TRACEINSIGHTS_API, FVerifyTraceSizeCommand, const FString, StoreDirPath, FAutomationTestBase*, Test);
+
+	
+	DEFINE_EXPORTED_LATENT_AUTOMATION_COMMAND_TWO_PARAMETER(TRACEINSIGHTS_API, FWaitForConnectionStatusCommand, const bool, bExpectedConnectionResult, FAutomationTestBase*, Test);
+
+
+	DEFINE_EXPORTED_LATENT_AUTOMATION_COMMAND_TWO_PARAMETER(TRACEINSIGHTS_API, FKillUTSCommand, FAutomationTestBase*, Test, FInsightsTestUtils, Utils);
+
+
+	DEFINE_EXPORTED_LATENT_AUTOMATION_COMMAND_TWO_PARAMETER(TRACEINSIGHTS_API, FDeleteCustomTraceFileCommand, const FString, TraceDirPath, FAutomationTestBase*, Test);
+
+
+	DEFINE_EXPORTED_LATENT_AUTOMATION_COMMAND_ONE_PARAMETER(TRACEINSIGHTS_API, FDeleteDefaultTraceFileCommand, FAutomationTestBase*, Test);
+
+
+	DEFINE_EXPORTED_LATENT_AUTOMATION_COMMAND_ONE_PARAMETER(TRACEINSIGHTS_API, FWaitTraceSnapshotNameAppearInLogCommand, FAutomationTestBase*, Test);
+
+
+	DEFINE_EXPORTED_LATENT_AUTOMATION_COMMAND_ONE_PARAMETER(TRACEINSIGHTS_API, FVerifyExpectedTraceSnapshotSizeCommand, FAutomationTestBase*, Test);
+
+
+	DEFINE_EXPORTED_LATENT_AUTOMATION_COMMAND_TWO_PARAMETER(TRACEINSIGHTS_API, FWaitExpectedTraceFilesizeCommand, FInsightsTestUtils, Utils, FAutomationTestBase*, Test);
+
+
+	DEFINE_EXPORTED_LATENT_AUTOMATION_COMMAND_TWO_PARAMETER(TRACEINSIGHTS_API, FVerifyFileCreated, const FString, FileNameToCheck, FAutomationTestBase*, Test);
+}
+
+#endif // WITH_DEV_AUTOMATION_TESTS
diff --git a/Engine/Source/Developer/TraceInsights/Public/Insights/Tests/InsightsTestUtils.h b/Engine/Source/Developer/TraceInsights/Public/Insights/Tests/InsightsTestUtils.h
--- a/Engine/Source/Developer/TraceInsights/Public/Insights/Tests/InsightsTestUtils.h
+++ b/Engine/Source/Developer/TraceInsights/Public/Insights/Tests/InsightsTestUtils.h
@@ -24,6 +24,7 @@ public:
 	void ResetSession() const;
 	FString GetLiveTrace(const TCHAR* Host = TEXT("localhost"), int32 Port = 0U) const;
 	FString GetUTSPath() const;
+	const TArray<FString>& GetDefaultTraceChannels();
 
 private:
 	FAutomationTestBase* Test;
