[Backout] - CL37467563
Original CL Desc
-----------------------------------------------------------------
Insights Utils refactoring
Moved latent commands to a separate file
Reorganised latent commands.
Minor fixes.

diff --git a/Engine/Source/Developer/TraceInsights/Private/Insights/Tests/InsightsLatentCommands.cpp b/Engine/Source/Developer/TraceInsights/Private/Insights/Tests/InsightsLatentCommands.cpp
deleted file mode 100644
diff --git a/Engine/Source/Developer/TraceInsights/Private/Insights/Tests/InsightsTestUtils.cpp b/Engine/Source/Developer/TraceInsights/Private/Insights/Tests/InsightsTestUtils.cpp
--- a/Engine/Source/Developer/TraceInsights/Private/Insights/Tests/InsightsTestUtils.cpp
+++ b/Engine/Source/Developer/TraceInsights/Private/Insights/Tests/InsightsTestUtils.cpp
@@ -23,6 +23,7 @@
 #include "Insights/InsightsManager.h"
 #include "Insights/IUnrealInsightsModule.h"
 
+
 ////////////////////////////////////////////////////////////////////////////////////////////////////
 
 FInsightsTestUtils::FInsightsTestUtils(FAutomationTestBase* InTest) :
@@ -411,24 +412,4 @@ FString FInsightsTestUtils::FInsightsTestUtils::GetUTSPath() const
 	}
 
 	return EntireUTSPath;
-}
-
-const TArray<FString>& FInsightsTestUtils::GetDefaultTraceChannels()
-{
-	static TArray<FString> DefaultTraceChannels;
-	if (DefaultTraceChannels.IsEmpty())
-	{
-		FString DefaultChannelList;
-		auto GetDefaultPreset = [&DefaultChannelList](const FTraceAuxiliary::FChannelPreset& Preset)
-			{
-				if (FString(Preset.Name) == TEXT("Default"))
-				{
-					DefaultChannelList = Preset.ChannelList;
-				}
-				return FTraceAuxiliary::EEnumerateResult::Stop;
-			};
-		FTraceAuxiliary::EnumerateFixedChannelPresets(GetDefaultPreset);
-		DefaultChannelList.ParseIntoArray(DefaultTraceChannels, TEXT(","));
-	}
-	return DefaultTraceChannels;
-}
+}
\ No newline at end of file
diff --git a/Engine/Source/Developer/TraceInsights/Public/Insights/Tests/InsightsLatentCommands.h b/Engine/Source/Developer/TraceInsights/Public/Insights/Tests/InsightsLatentCommands.h
deleted file mode 100644
diff --git a/Engine/Source/Developer/TraceInsights/Public/Insights/Tests/InsightsTestUtils.h b/Engine/Source/Developer/TraceInsights/Public/Insights/Tests/InsightsTestUtils.h
--- a/Engine/Source/Developer/TraceInsights/Public/Insights/Tests/InsightsTestUtils.h
+++ b/Engine/Source/Developer/TraceInsights/Public/Insights/Tests/InsightsTestUtils.h
@@ -24,7 +24,6 @@ public:
 	void ResetSession() const;
 	FString GetLiveTrace(const TCHAR* Host = TEXT("localhost"), int32 Port = 0U) const;
 	FString GetUTSPath() const;
-	const TArray<FString>& GetDefaultTraceChannels();
 
 private:
 	FAutomationTestBase* Test;
