Network Insights - Display all HugeObject packet contents.

diff --git a/Engine/Source/Developer/TraceInsights/Private/Insights/NetworkingProfiler/Widgets/SPacketContentView.cpp b/Engine/Source/Developer/TraceInsights/Private/Insights/NetworkingProfiler/Widgets/SPacketContentView.cpp
--- a/Engine/Source/Developer/TraceInsights/Private/Insights/NetworkingProfiler/Widgets/SPacketContentView.cpp
+++ b/Engine/Source/Developer/TraceInsights/Private/Insights/NetworkingProfiler/Widgets/SPacketContentView.cpp
@@ -1047,18 +1047,22 @@ void SPacketContentView::UpdateState(float FontScale)
 			const TraceServices::INetProfilerProvider* NetProfilerProvider = TraceServices::ReadNetProfilerProvider(*Session.Get());
 			if (NetProfilerProvider)
 			{
-				const FAxisViewportDouble& ViewportX = Viewport.GetHorizontalAxisViewport();
-
 				// Count all events in packet, including split data
 				const uint32 StartPos = 0U;
 				const uint32 EndPos = ~0U;
 				uint32 EndNetIdMatchPos = ~0U;
 				uint32 EndEventTypeMatchPos = ~0U;
+				uint32 HighestEndPos = 0;
 
-				NetProfilerProvider->EnumeratePacketContentEventsByPosition(ConnectionIndex, ConnectionMode, PacketIndex, StartPos, EndPos, [this, &Builder, &FilteredDrawStateBuilder, NetProfilerProvider, &EndNetIdMatchPos, &EndEventTypeMatchPos](const TraceServices::FNetProfilerContentEvent& Event)
+				NetProfilerProvider->EnumeratePacketContentEventsByPosition(ConnectionIndex, ConnectionMode, PacketIndex, StartPos, EndPos, [this, &Builder, &FilteredDrawStateBuilder, NetProfilerProvider, &EndNetIdMatchPos, &EndEventTypeMatchPos, &HighestEndPos](const TraceServices::FNetProfilerContentEvent& Event)
 				{
 					const TCHAR* Name = nullptr;
 
+					if (Event.EndPos > HighestEndPos)
+					{
+						HighestEndPos = Event.EndPos;
+					}
+
 					uint32 NameIndex = Event.NameIndex;
 					uint64 NetId = 0;
 					if (Event.ObjectInstanceIndex != 0)
@@ -1108,6 +1112,15 @@ void SPacketContentView::UpdateState(float FontScale)
 						FilteredDrawStateBuilder.AddEvent(Event, Name, NetId);
 					}
 				});
+
+				FAxisViewportDouble& ViewportX = Viewport.GetHorizontalAxisViewport();
+				if (HighestEndPos > ViewportX.GetMaxValue())
+				{
+					PacketBitSize = HighestEndPos;
+					ViewportX.SetMinMaxValueInterval(0.0, static_cast<double>(PacketBitSize));
+					ViewportX.CenterOnValueInterval(0.0, static_cast<double>(PacketBitSize));
+					UpdateHorizontalScrollBar();
+				}
 			}
 		}
 
