[Insights] Fix for GPU Work track settings.

diff --git a/Engine/Source/Developer/TraceInsights/Private/Insights/TimingProfiler/Tracks/ThreadTimingTrack.cpp b/Engine/Source/Developer/TraceInsights/Private/Insights/TimingProfiler/Tracks/ThreadTimingTrack.cpp
--- a/Engine/Source/Developer/TraceInsights/Private/Insights/TimingProfiler/Tracks/ThreadTimingTrack.cpp
+++ b/Engine/Source/Developer/TraceInsights/Private/Insights/TimingProfiler/Tracks/ThreadTimingTrack.cpp
@@ -1584,8 +1584,8 @@ void FGpuQueueWorkTimingTrack::DrawLineEvents(const ITimingTrackDrawContext& Con
 
 void FGpuQueueWorkTimingTrack::PostDraw(const ITimingTrackDrawContext& Context) const
 {
-	const bool bAreOverlaysVisible = true; // SharedState.AreOverlaysVisible();
-	const bool bAreExtendedLinesVisible = true; // SharedState.AreExtendedLinesVisible();
+	const bool bAreOverlaysVisible = GetSharedState().AreOverlaysVisibleInGpuQueueTracks();
+	const bool bAreExtendedLinesVisible = GetSharedState().AreExtendedLinesVisibleInGpuQueueTracks();
 
 	if (bAreOverlaysVisible || bAreExtendedLinesVisible)
 	{
