[Insights ] Fix duplicate menu entry

diff --git a/Engine/Source/Developer/TraceInsights/Private/Insights/TimingProfiler/ViewModels/TimingRegionsSharedState.cpp b/Engine/Source/Developer/TraceInsights/Private/Insights/TimingProfiler/ViewModels/TimingRegionsSharedState.cpp
--- a/Engine/Source/Developer/TraceInsights/Private/Insights/TimingProfiler/ViewModels/TimingRegionsSharedState.cpp
+++ b/Engine/Source/Developer/TraceInsights/Private/Insights/TimingProfiler/ViewModels/TimingRegionsSharedState.cpp
@@ -139,6 +139,11 @@ void FTimingRegionsSharedState::ShowHideRegionsTrack()
 
 void FTimingRegionsSharedState::ExtendOtherTracksFilterMenu(Timing::ITimingViewSession& InSession, FMenuBuilder& InOutMenuBuilder)
 {
+	if (&InSession != TimingView)
+	{
+		return;
+	}
+
 	InOutMenuBuilder.BeginSection("Timing Regions", LOCTEXT("ContextMenu_Section_Regions", "Timing Regions"));
 	{
 		InOutMenuBuilder.AddMenuEntry(FTimingRegionsViewCommands::Get().ShowHideTimingRegionsTrack);
