[Insights] Fixed relative track order for common tracks (from top to bottom): Main Graph Track, Timing Regions track, I/O Overview track, Frame tracks and then the Main Memory Graph track and the related memory tracks. This fixes issue in Memory Insights where the main memory graph track and the timing regions track randomly changed order.

diff --git a/Engine/Source/Developer/TraceInsights/Private/Insights/MemoryProfiler/ViewModels/MemorySharedState.cpp b/Engine/Source/Developer/TraceInsights/Private/Insights/MemoryProfiler/ViewModels/MemorySharedState.cpp
--- a/Engine/Source/Developer/TraceInsights/Private/Insights/MemoryProfiler/ViewModels/MemorySharedState.cpp
+++ b/Engine/Source/Developer/TraceInsights/Private/Insights/MemoryProfiler/ViewModels/MemorySharedState.cpp
@@ -168,12 +168,14 @@ void FMemorySharedState::Tick(Timing::ITimingViewSession& InSession, const Trace
 		return;
 	}
 
+	constexpr int32 GraphTrackOrder = FTimingTrackOrder::First + FTimingTrackOrder::GroupRange / 2;
+
 	if (!MainGraphTrack.IsValid())
 	{
 		MainGraphTrack = CreateMemoryGraphTrack();
 		check(MainGraphTrack);
 
-		MainGraphTrack->SetOrder(FTimingTrackOrder::First);
+		MainGraphTrack->SetOrder(GraphTrackOrder);
 		MainGraphTrack->SetName(TEXT("Main Memory Graph"));
 
 		MainGraphTrack->AddTimelineSeries(FMemoryGraphSeries::ETimelineType::MaxTotalMem);
@@ -194,7 +196,7 @@ void FMemorySharedState::Tick(Timing::ITimingViewSession& InSession, const Trace
 		LiveAllocsGraphTrack = CreateMemoryGraphTrack();
 		check(LiveAllocsGraphTrack);
 
-		LiveAllocsGraphTrack->SetOrder(FTimingTrackOrder::First + 1);
+		LiveAllocsGraphTrack->SetOrder(GraphTrackOrder + 1);
 		LiveAllocsGraphTrack->SetName(TEXT("Live Allocation Count"));
 		LiveAllocsGraphTrack->SetLabelUnit(EGraphTrackLabelUnit::Count, 0);
 
@@ -216,7 +218,7 @@ void FMemorySharedState::Tick(Timing::ITimingViewSession& InSession, const Trace
 		AllocFreeGraphTrack = CreateMemoryGraphTrack();
 		check(AllocFreeGraphTrack);
 
-		AllocFreeGraphTrack->SetOrder(FTimingTrackOrder::First + 2);
+		AllocFreeGraphTrack->SetOrder(GraphTrackOrder + 2);
 		AllocFreeGraphTrack->SetName(TEXT("Alloc/Free Event Count"));
 		AllocFreeGraphTrack->SetLabelUnit(EGraphTrackLabelUnit::Count, 0);
 
@@ -238,7 +240,7 @@ void FMemorySharedState::Tick(Timing::ITimingViewSession& InSession, const Trace
 		SwapMemoryGraphTrack = CreateMemoryGraphTrack();
 		check(SwapMemoryGraphTrack);
 
-		SwapMemoryGraphTrack->SetOrder(FTimingTrackOrder::First + 3);
+		SwapMemoryGraphTrack->SetOrder(GraphTrackOrder + 3);
 		SwapMemoryGraphTrack->SetName(TEXT("Swap Memory Graph"));
 
 		SwapMemoryGraphTrack->AddTimelineSeries(FMemoryGraphSeries::ETimelineType::MaxSwapMem);
@@ -262,7 +264,7 @@ void FMemorySharedState::Tick(Timing::ITimingViewSession& InSession, const Trace
 		PageSwapGraphTrack = CreateMemoryGraphTrack();
 		check(PageSwapGraphTrack);
 
-		PageSwapGraphTrack->SetOrder(FTimingTrackOrder::First + 4);
+		PageSwapGraphTrack->SetOrder(GraphTrackOrder + 4);
 		PageSwapGraphTrack->SetName(TEXT("Page In/Out Event Count"));
 		PageSwapGraphTrack->SetLabelUnit(EGraphTrackLabelUnit::Count, 0);
 
diff --git a/Engine/Source/Developer/TraceInsights/Private/Insights/TimingProfiler/ViewModels/FileActivitySharedState.cpp b/Engine/Source/Developer/TraceInsights/Private/Insights/TimingProfiler/ViewModels/FileActivitySharedState.cpp
--- a/Engine/Source/Developer/TraceInsights/Private/Insights/TimingProfiler/ViewModels/FileActivitySharedState.cpp
+++ b/Engine/Source/Developer/TraceInsights/Private/Insights/TimingProfiler/ViewModels/FileActivitySharedState.cpp
@@ -161,7 +161,7 @@ void FFileActivitySharedState::Tick(Timing::ITimingViewSession& InSession, const
 	if (!IoOverviewTrack.IsValid())
 	{
 		IoOverviewTrack = MakeShared<FOverviewFileActivityTimingTrack>(*this);
-		IoOverviewTrack->SetOrder(FTimingTrackOrder::First);
+		IoOverviewTrack->SetOrder(FTimingTrackOrder::First + 200);
 		IoOverviewTrack->SetVisibilityFlag(bShowHideAllIoTracks);
 		InSession.AddScrollableTrack(IoOverviewTrack);
 	}
diff --git a/Engine/Source/Developer/TraceInsights/Private/Insights/TimingProfiler/ViewModels/FrameTimingTrack.cpp b/Engine/Source/Developer/TraceInsights/Private/Insights/TimingProfiler/ViewModels/FrameTimingTrack.cpp
--- a/Engine/Source/Developer/TraceInsights/Private/Insights/TimingProfiler/ViewModels/FrameTimingTrack.cpp
+++ b/Engine/Source/Developer/TraceInsights/Private/Insights/TimingProfiler/ViewModels/FrameTimingTrack.cpp
@@ -118,7 +118,7 @@ void FFrameSharedState::Tick(Timing::ITimingViewSession& InSession, const TraceS
 
 	{
 		bool bTracksOrderChanged = false;
-		int32 Order = FTimingTrackOrder::First + 100;
+		int32 Order = FTimingTrackOrder::First + 300;
 		constexpr int32 OrderIncrement = 10;
 		static_assert(FTimingTrackOrder::GroupRange > 100 + OrderIncrement * TraceFrameType_Count, "Order group range too small");
 
diff --git a/Engine/Source/Developer/TraceInsights/Private/Insights/TimingProfiler/ViewModels/TimingRegionsSharedState.cpp b/Engine/Source/Developer/TraceInsights/Private/Insights/TimingProfiler/ViewModels/TimingRegionsSharedState.cpp
--- a/Engine/Source/Developer/TraceInsights/Private/Insights/TimingProfiler/ViewModels/TimingRegionsSharedState.cpp
+++ b/Engine/Source/Developer/TraceInsights/Private/Insights/TimingProfiler/ViewModels/TimingRegionsSharedState.cpp
@@ -107,7 +107,7 @@ void FTimingRegionsSharedState::Tick(Timing::ITimingViewSession& InSession, cons
 	if (!TimingRegionsTrack.IsValid())
 	{
 		TimingRegionsTrack = MakeShared<FTimingRegionsTrack>(*this);
-		TimingRegionsTrack->SetOrder(FTimingTrackOrder::First);
+		TimingRegionsTrack->SetOrder(FTimingTrackOrder::First + 100);
 		TimingRegionsTrack->SetVisibilityFlag(true);
 		InSession.AddScrollableTrack(TimingRegionsTrack);
 	}
