Add direct trace option to Insights manager

Allows users to create an direct trace analysis session. Connection will have to be done on the port returned by IUnrealInsightsManager::StartAnalysisWithDirectConnection. Note that traced data will not be stored to disk.

diff --git a/Engine/Source/Developer/TraceInsights/Private/Insights/InsightsManager.cpp b/Engine/Source/Developer/TraceInsights/Private/Insights/InsightsManager.cpp
--- a/Engine/Source/Developer/TraceInsights/Private/Insights/InsightsManager.cpp
+++ b/Engine/Source/Developer/TraceInsights/Private/Insights/InsightsManager.cpp
@@ -812,6 +812,28 @@ void FInsightsManager::LoadTraceFile(const FString& InTraceFilename, bool InAuto
 
 ////////////////////////////////////////////////////////////////////////////////////////////////////
 
+uint16 FInsightsManager::ListenForDirectTrace()
+{
+	constexpr uint32 TraceId = 0; 
+	FString TraceName(TEXT("DirectTrace"));
+	TUniquePtr<UE::Trace::FDirectSocketStream> TraceData = MakeUnique<UE::Trace::FDirectSocketStream>();
+	const uint16 Port = TraceData->StartListening();
+	
+	Session = AnalysisService->StartAnalysis(TraceId, *TraceName, MoveTemp(TraceData));
+
+	if (Session)
+	{
+		CurrentTraceId = TraceId;
+		CurrentTraceFilename = TraceName;
+		bIsSessionInfoSet = false;
+		OnSessionChanged();
+	}
+
+	return Port;
+}
+
+////////////////////////////////////////////////////////////////////////////////////////////////////
+
 bool FInsightsManager::OnDragOver(const FDragDropEvent& DragDropEvent)
 {
 	TSharedPtr<FExternalDragOperation> DragDropOp = DragDropEvent.GetOperationAs<FExternalDragOperation>();
diff --git a/Engine/Source/Developer/TraceInsights/Private/Insights/InsightsManager.h b/Engine/Source/Developer/TraceInsights/Private/Insights/InsightsManager.h
--- a/Engine/Source/Developer/TraceInsights/Private/Insights/InsightsManager.h
+++ b/Engine/Source/Developer/TraceInsights/Private/Insights/InsightsManager.h
@@ -226,6 +226,9 @@ public:
 	 */
 	void LoadTraceFile(const FString& TraceFilename, bool InAutoQuit = false);
 
+
+	uint16 ListenForDirectTrace();
+
 	bool OnDragOver(const FDragDropEvent& DragDropEvent);
 	bool OnDrop(const FDragDropEvent& DragDropEvent);
 
diff --git a/Engine/Source/Developer/TraceInsights/Private/Insights/TraceInsightsModule.cpp b/Engine/Source/Developer/TraceInsights/Private/Insights/TraceInsightsModule.cpp
--- a/Engine/Source/Developer/TraceInsights/Private/Insights/TraceInsightsModule.cpp
+++ b/Engine/Source/Developer/TraceInsights/Private/Insights/TraceInsightsModule.cpp
@@ -495,6 +495,13 @@ void FTraceInsightsModule::StartAnalysisForTraceFile(const TCHAR* InTraceFile, b
 
 ////////////////////////////////////////////////////////////////////////////////////////////////////
 
+uint16 FTraceInsightsModule::StartAnalysisWithDirectTrace()
+{
+	return FInsightsManager::Get()->ListenForDirectTrace();
+}
+
+////////////////////////////////////////////////////////////////////////////////////////////////////
+
 void FTraceInsightsModule::OnWindowClosedEvent(const TSharedRef<SWindow>&)
 {
 	for (TSharedRef<IInsightsComponent>& Component : Components)
diff --git a/Engine/Source/Developer/TraceInsights/Private/Insights/TraceInsightsModule.h b/Engine/Source/Developer/TraceInsights/Private/Insights/TraceInsightsModule.h
--- a/Engine/Source/Developer/TraceInsights/Private/Insights/TraceInsightsModule.h
+++ b/Engine/Source/Developer/TraceInsights/Private/Insights/TraceInsightsModule.h
@@ -58,6 +58,7 @@ public:
 	virtual void StartAnalysisForTrace(uint32 InTraceId, bool InAutoQuit = false) override;
 	virtual void StartAnalysisForLastLiveSession(float InRetryTime = 1.0f) override;
 	virtual void StartAnalysisForTraceFile(const TCHAR* InTraceFile, bool InAutoQuit = false) override;
+	virtual uint16 StartAnalysisWithDirectTrace() override;
 
 	virtual void ShutdownUserInterface() override;
 
diff --git a/Engine/Source/Developer/TraceInsights/Public/Insights/IUnrealInsightsModule.h b/Engine/Source/Developer/TraceInsights/Public/Insights/IUnrealInsightsModule.h
--- a/Engine/Source/Developer/TraceInsights/Public/Insights/IUnrealInsightsModule.h
+++ b/Engine/Source/Developer/TraceInsights/Public/Insights/IUnrealInsightsModule.h
@@ -229,6 +229,14 @@ public:
 	 */
 	virtual void StartAnalysisForTraceFile(const TCHAR* InTraceFile, bool InAutoQuit = false) = 0;
 
+	/**
+	 * Starts analysis for a direct connection from the tracing application (might be self). Notice that
+	 * analysis doesn't actually start until a connection has been made.
+	 * 
+	 * @return Port used for listening to incoming connections. 0 if failed to open port.
+	 */
+	virtual uint16 StartAnalysisWithDirectTrace() = 0;
+
 	//////////////////////////////////////////////////
 
 	/**
