[Insights] Memory Insights / LLM Tags: Added initial support in UI for additional LLM tag sets (assets and asset classes).

diff --git a/Engine/Source/Developer/TraceInsights/Private/Insights/MemoryProfiler/ViewModels/MemTagNode.cpp b/Engine/Source/Developer/TraceInsights/Private/Insights/MemoryProfiler/ViewModels/MemTagNode.cpp
--- a/Engine/Source/Developer/TraceInsights/Private/Insights/MemoryProfiler/ViewModels/MemTagNode.cpp
+++ b/Engine/Source/Developer/TraceInsights/Private/Insights/MemoryProfiler/ViewModels/MemTagNode.cpp
@@ -32,6 +32,24 @@ FText FMemTagNode::GetTrackerText() const
 
 ////////////////////////////////////////////////////////////////////////////////////////////////////
 
+FText FMemTagNode::GetTagSetText() const
+{
+	FMemorySharedState* SharedState = FMemoryProfilerManager::Get()->GetSharedState();
+	if (SharedState)
+	{
+		FMemoryTagSetId TagSetId = GetMemTagSetId();
+		const FMemoryTagSet* TagSet = SharedState->GetTagSetById(TagSetId);
+		if (TagSet)
+		{
+			return FText::FromString(TagSet->GetName());
+		}
+	}
+
+	return FText::GetEmpty();
+}
+
+////////////////////////////////////////////////////////////////////////////////////////////////////
+
 void FMemTagNode::ResetAggregatedStats()
 {
 	//AggregatedStats = TraceServices::FMemoryProfilerAggregatedStats();
diff --git a/Engine/Source/Developer/TraceInsights/Private/Insights/MemoryProfiler/ViewModels/MemTagNode.h b/Engine/Source/Developer/TraceInsights/Private/Insights/MemoryProfiler/ViewModels/MemTagNode.h
--- a/Engine/Source/Developer/TraceInsights/Private/Insights/MemoryProfiler/ViewModels/MemTagNode.h
+++ b/Engine/Source/Developer/TraceInsights/Private/Insights/MemoryProfiler/ViewModels/MemTagNode.h
@@ -101,6 +101,9 @@ public:
 	FMemoryTrackerId GetMemTrackerId() const { return MemTag ? MemTag->GetTrackerId() : FMemoryTracker::InvalidTrackerId; }
 	FText GetTrackerText() const;
 
+	FMemoryTagSetId GetMemTagSetId() const { return MemTag ? MemTag->GetTagSetId() : FMemoryTagSet::InvalidTagSetId; }
+	FText GetTagSetText() const;
+
 	FLinearColor GetColor() const { return MemTag ? MemTag->GetColor() : FLinearColor(0.5f, 0.5f, 0.5f, 1.0f); }
 	bool IsAddedToGraph() const { return MemTag ? MemTag->IsAddedToGraph() : false; }
 
diff --git a/Engine/Source/Developer/TraceInsights/Private/Insights/MemoryProfiler/ViewModels/MemTagTreeViewColumnFactory.cpp b/Engine/Source/Developer/TraceInsights/Private/Insights/MemoryProfiler/ViewModels/MemTagTreeViewColumnFactory.cpp
--- a/Engine/Source/Developer/TraceInsights/Private/Insights/MemoryProfiler/ViewModels/MemTagTreeViewColumnFactory.cpp
+++ b/Engine/Source/Developer/TraceInsights/Private/Insights/MemoryProfiler/ViewModels/MemTagTreeViewColumnFactory.cpp
@@ -22,6 +22,7 @@ namespace UE::Insights::MemoryProfiler
 const FName FMemTagTreeViewColumns::NameColumnID(TEXT("Name")); // TEXT("_Hierarchy")
 const FName FMemTagTreeViewColumns::TypeColumnID(TEXT("Type"));
 const FName FMemTagTreeViewColumns::TrackerColumnID(TEXT("Tracker"));
+const FName FMemTagTreeViewColumns::TagSetColumnID(TEXT("TagSet"));
 const FName FMemTagTreeViewColumns::InstanceCountColumnID(TEXT("Count"));
 const FName FMemTagTreeViewColumns::MinValueColumnID(TEXT("Min"));
 const FName FMemTagTreeViewColumns::MaxValueColumnID(TEXT("Max"));
@@ -36,6 +37,7 @@ void FMemTagTreeViewColumnFactory::CreateMemTagTreeViewColumns(TArray<TSharedRef
 	Columns.Add(CreateNameColumn());
 	Columns.Add(CreateTypeColumn());
 	Columns.Add(CreateTrackerColumn());
+	Columns.Add(CreateTagSetColumn());
 	//Columns.Add(CreateInstanceCountColumn());
 	//Columns.Add(CreateMinValueColumn());
 	//Columns.Add(CreateMaxValueColumn());
@@ -139,7 +141,7 @@ TSharedRef<FTableColumn> FMemTagTreeViewColumnFactory::CreateTrackerColumn()
 
 	Column.SetDataType(ETableCellDataType::Text);
 
-	class FMemTagNodeTypeValueGetter : public FTableCellValueGetter
+	class FMemTagNodeTrackerValueGetter : public FTableCellValueGetter
 	{
 	public:
 		virtual const TOptional<FTableCellValue> GetValue(const FTableColumn& Column, const FBaseTreeNode& Node) const
@@ -150,7 +152,7 @@ TSharedRef<FTableColumn> FMemTagTreeViewColumnFactory::CreateTrackerColumn()
 		}
 	};
 
-	TSharedRef<ITableCellValueGetter> Getter = MakeShared<FMemTagNodeTypeValueGetter>();
+	TSharedRef<ITableCellValueGetter> Getter = MakeShared<FMemTagNodeTrackerValueGetter>();
 	Column.SetValueGetter(Getter);
 
 	TSharedRef<ITableCellValueFormatter> Formatter = MakeShared<FTextValueFormatter>();
@@ -164,6 +166,48 @@ TSharedRef<FTableColumn> FMemTagTreeViewColumnFactory::CreateTrackerColumn()
 
 ////////////////////////////////////////////////////////////////////////////////////////////////////
 
+TSharedRef<FTableColumn> FMemTagTreeViewColumnFactory::CreateTagSetColumn()
+{
+	TSharedRef<FTableColumn> ColumnRef = MakeShared<FTableColumn>(FMemTagTreeViewColumns::TagSetColumnID);
+	FTableColumn& Column = *ColumnRef;
+
+	Column.SetShortName(LOCTEXT("TagSet_ColumnName", "Tag Set"));
+	Column.SetTitleName(LOCTEXT("TagSet_ColumnTitle", "Tag Set"));
+	Column.SetDescription(LOCTEXT("TagSet_ColumnDesc", "Tag set of the LLM tag"));
+
+	Column.SetFlags(ETableColumnFlags::CanBeHidden |
+		ETableColumnFlags::CanBeFiltered);
+
+	Column.SetHorizontalAlignment(HAlign_Left);
+	Column.SetInitialWidth(60.0f);
+
+	Column.SetDataType(ETableCellDataType::Text);
+
+	class FMemTagNodeTagSetValueGetter : public FTableCellValueGetter
+	{
+	public:
+		virtual const TOptional<FTableCellValue> GetValue(const FTableColumn& Column, const FBaseTreeNode& Node) const
+		{
+			ensure(Column.GetId() == FMemTagTreeViewColumns::TagSetColumnID);
+			const FMemTagNode& MemTagNode = static_cast<const FMemTagNode&>(Node);
+			return FTableCellValue(MemTagNode.GetTagSetText());
+		}
+	};
+
+	TSharedRef<ITableCellValueGetter> Getter = MakeShared<FMemTagNodeTagSetValueGetter>();
+	Column.SetValueGetter(Getter);
+
+	TSharedRef<ITableCellValueFormatter> Formatter = MakeShared<FTextValueFormatter>();
+	Column.SetValueFormatter(Formatter);
+
+	//TSharedRef<ITableCellValueSorter> Sorter = MakeShared<FMemTagNodeSortingByTagSet>(ColumnRef);
+	//Column.SetValueSorter(Sorter);
+
+	return ColumnRef;
+}
+
+////////////////////////////////////////////////////////////////////////////////////////////////////
+
 TSharedRef<FTableColumn> FMemTagTreeViewColumnFactory::CreateInstanceCountColumn()
 {
 	TSharedRef<FTableColumn> ColumnRef = MakeShared<FTableColumn>(FMemTagTreeViewColumns::InstanceCountColumnID);
diff --git a/Engine/Source/Developer/TraceInsights/Private/Insights/MemoryProfiler/ViewModels/MemTagTreeViewColumnFactory.h b/Engine/Source/Developer/TraceInsights/Private/Insights/MemoryProfiler/ViewModels/MemTagTreeViewColumnFactory.h
--- a/Engine/Source/Developer/TraceInsights/Private/Insights/MemoryProfiler/ViewModels/MemTagTreeViewColumnFactory.h
+++ b/Engine/Source/Developer/TraceInsights/Private/Insights/MemoryProfiler/ViewModels/MemTagTreeViewColumnFactory.h
@@ -18,6 +18,7 @@ struct FMemTagTreeViewColumns
 	static const FName NameColumnID;
 	static const FName TypeColumnID;
 	static const FName TrackerColumnID;
+	static const FName TagSetColumnID;
 	static const FName InstanceCountColumnID;
 	static const FName MinValueColumnID;
 	static const FName MaxValueColumnID;
@@ -32,6 +33,7 @@ public:
 	static TSharedRef<FTableColumn> CreateNameColumn();
 	static TSharedRef<FTableColumn> CreateTypeColumn();
 	static TSharedRef<FTableColumn> CreateTrackerColumn();
+	static TSharedRef<FTableColumn> CreateTagSetColumn();
 	static TSharedRef<FTableColumn> CreateInstanceCountColumn();
 	static TSharedRef<FTableColumn> CreateMinValueColumn();
 	static TSharedRef<FTableColumn> CreateMaxValueColumn();
diff --git a/Engine/Source/Developer/TraceInsights/Private/Insights/MemoryProfiler/ViewModels/MemorySharedState.cpp b/Engine/Source/Developer/TraceInsights/Private/Insights/MemoryProfiler/ViewModels/MemorySharedState.cpp
--- a/Engine/Source/Developer/TraceInsights/Private/Insights/MemoryProfiler/ViewModels/MemorySharedState.cpp
+++ b/Engine/Source/Developer/TraceInsights/Private/Insights/MemoryProfiler/ViewModels/MemorySharedState.cpp
@@ -120,6 +120,8 @@ void FMemorySharedState::OnBeginSession(Timing::ITimingViewSession& InSession)
 	DefaultTracker = nullptr;
 	PlatformTracker = nullptr;
 
+	TagSets.Reset();
+
 	MainGraphTrack = nullptr;
 	LiveAllocsGraphTrack = nullptr;
 	AllocFreeGraphTrack = nullptr;
@@ -147,6 +149,8 @@ void FMemorySharedState::OnEndSession(Timing::ITimingViewSession& InSession)
 	DefaultTracker = nullptr;
 	PlatformTracker = nullptr;
 
+	TagSets.Reset();
+
 	MainGraphTrack = nullptr;
 	LiveAllocsGraphTrack = nullptr;
 	AllocFreeGraphTrack = nullptr;
@@ -291,6 +295,8 @@ void FMemorySharedState::Tick(Timing::ITimingViewSession& InSession, const Trace
 		SyncTrackers();
 	}
 
+	UpdateTagSets();
+
 	// Scan for mem tags to show as default, but only when new mem tags are added.
 	const int32 NewTagCount = TagList.GetTags().Num();
 	if (NewTagCount > PrevTagCount)
@@ -431,6 +437,47 @@ void FMemorySharedState::SyncTrackers()
 
 ////////////////////////////////////////////////////////////////////////////////////////////////////
 
+void FMemorySharedState::UpdateTagSets()
+{
+	TSharedPtr<const TraceServices::IAnalysisSession> Session = FInsightsManager::Get()->GetSession();
+	if (Session.IsValid())
+	{
+		TraceServices::FAnalysisSessionReadScope SessionReadScope(*Session.Get());
+		const TraceServices::IMemoryProvider* MemoryProvider = TraceServices::ReadMemoryProvider(*Session.Get());
+		if (MemoryProvider)
+		{
+			if ((int32)MemoryProvider->GetTagSetCount() != TagSets.Num())
+			{
+				MemoryProvider->EnumerateTagSets([this](const TraceServices::FMemoryTagSetInfo& TagSet)
+				{
+					int32 Index = (int32)TagSet.Id;
+					if (Index >= 0)
+					{
+						if (TagSets.Num() <= Index)
+						{
+							TagSets.AddDefaulted(Index - TagSets.Num() + 1);
+						}
+						TagSets[Index] = MakeShared<FMemoryTagSet>(static_cast<FMemoryTagSetId>(TagSet.Id), TagSet.Name);
+					}
+				});
+			}
+		}
+	}
+	else
+	{
+		TagSets.Reset();
+	}
+}
+
+////////////////////////////////////////////////////////////////////////////////////////////////////
+
+const FMemoryTagSet* FMemorySharedState::GetTagSetById(FMemoryTagSetId InMemTagSetId) const
+{
+	return (int32)InMemTagSetId >= 0 && (int32)InMemTagSetId < TagSets.Num() ? TagSets[(int32)InMemTagSetId].Get() : nullptr;
+}
+
+////////////////////////////////////////////////////////////////////////////////////////////////////
+
 void FMemorySharedState::SetTrackHeightMode(EMemoryTrackHeightMode InTrackHeightMode)
 {
 	TrackHeightMode = InTrackHeightMode;
diff --git a/Engine/Source/Developer/TraceInsights/Private/Insights/MemoryProfiler/ViewModels/MemorySharedState.h b/Engine/Source/Developer/TraceInsights/Private/Insights/MemoryProfiler/ViewModels/MemorySharedState.h
--- a/Engine/Source/Developer/TraceInsights/Private/Insights/MemoryProfiler/ViewModels/MemorySharedState.h
+++ b/Engine/Source/Developer/TraceInsights/Private/Insights/MemoryProfiler/ViewModels/MemorySharedState.h
@@ -105,10 +105,13 @@ public:
 
 	const FMemoryTagList& GetTagList() { return TagList; }
 
-	const TArray<TSharedPtr<FMemoryTracker>>& GetTrackers()  const { return Trackers; }
+	const TArray<TSharedPtr<FMemoryTracker>>& GetTrackers() const { return Trackers; }
 	FString TrackersToString(uint64 Flags, const TCHAR* Conjunction) const;
 	const FMemoryTracker* GetTrackerById(FMemoryTrackerId InMemTrackerId) const;
 
+	const TArray<TSharedPtr<FMemoryTagSet>>& GetTagSets() const { return TagSets; }
+	const FMemoryTagSet* GetTagSetById(FMemoryTagSetId InMemTagSetId) const;
+
 	TSharedPtr<FMemoryGraphTrack> GetMainGraphTrack() const { return MainGraphTrack; }
 
 	EMemoryTrackHeightMode GetTrackHeightMode() const { return TrackHeightMode; }
@@ -163,6 +166,7 @@ public:
 
 private:
 	void SyncTrackers();
+	void UpdateTagSets();
 	int32 GetNextMemoryGraphTrackOrder();
 	TSharedPtr<FMemoryGraphTrack> CreateGraphTrack(const FReportTypeGraphConfig& ReportTypeGraphConfig, bool bIsPlatformTracker);
 	void InitMemoryRules();
@@ -177,6 +181,8 @@ private:
 	TSharedPtr<FMemoryTracker> DefaultTracker;
 	TSharedPtr<FMemoryTracker> PlatformTracker;
 
+	TArray<TSharedPtr<FMemoryTagSet>> TagSets;
+
 	TSharedPtr<FMemoryGraphTrack> MainGraphTrack; // the Main Memory Graph track; also hosts the Total Allocated Memory series
 	TSharedPtr<FMemoryGraphTrack> LiveAllocsGraphTrack; // the graph track for the Live Allocation Count series
 	TSharedPtr<FMemoryGraphTrack> AllocFreeGraphTrack; // the graph track for the Alloc Event Count and the Free Event Count series
diff --git a/Engine/Source/Developer/TraceInsights/Private/Insights/MemoryProfiler/ViewModels/MemoryTag.cpp b/Engine/Source/Developer/TraceInsights/Private/Insights/MemoryProfiler/ViewModels/MemoryTag.cpp
--- a/Engine/Source/Developer/TraceInsights/Private/Insights/MemoryProfiler/ViewModels/MemoryTag.cpp
+++ b/Engine/Source/Developer/TraceInsights/Private/Insights/MemoryProfiler/ViewModels/MemoryTag.cpp
@@ -175,7 +175,8 @@ void FMemoryTagList::UpdateInternal()
 								Tag.ParentId = static_cast<FMemoryTagId>(TraceTag.ParentId);
 								Tag.StatName = TraceTag.Name;
 								Tag.StatFullName = TraceTag.Name;
-								Tag.TrackerId = FMemoryTrackerId(TrackerId);
+								Tag.TrackerId = static_cast<FMemoryTrackerId>(TrackerId);
+								Tag.TagSetId = static_cast<FMemoryTagSetId>(TraceTag.TagSetId);
 								Tag.SetColorAuto();
 
 								// Skip the parent prefix if it is already included in the name.
diff --git a/Engine/Source/Developer/TraceInsights/Private/Insights/MemoryProfiler/ViewModels/MemoryTag.h b/Engine/Source/Developer/TraceInsights/Private/Insights/MemoryProfiler/ViewModels/MemoryTag.h
--- a/Engine/Source/Developer/TraceInsights/Private/Insights/MemoryProfiler/ViewModels/MemoryTag.h
+++ b/Engine/Source/Developer/TraceInsights/Private/Insights/MemoryProfiler/ViewModels/MemoryTag.h
@@ -17,10 +17,32 @@ namespace UE::Insights::MemoryProfiler
 
 class FMemoryGraphTrack;
 
+typedef int64 FMemoryTagSetId;
 typedef int64 FMemoryTagId;
 
 ////////////////////////////////////////////////////////////////////////////////////////////////////
 
+class FMemoryTagSet
+{
+public:
+	static const FMemoryTagSetId InvalidTagSetId = -1;
+	static const FMemoryTagSetId DefaultTagSetId = 0;
+
+public:
+	FMemoryTagSet(FMemoryTagSetId InId, const FString& InName) : Id(InId), Name(InName) {}
+	virtual ~FMemoryTagSet() {}
+
+	int32 GetIndex() const { return (int32)Id; }
+	FMemoryTagSetId GetId() const { return Id; }
+	const FString& GetName() const { return Name; }
+
+private:
+	FMemoryTagSetId Id = InvalidTagSetId;
+	FString Name;
+};
+
+////////////////////////////////////////////////////////////////////////////////////////////////////
+
 class FMemoryTag
 {
 	friend class FMemoryTagList;
@@ -43,6 +65,7 @@ public:
 	bool MatchesWildcard(const TArray<FString>& FullNames) const;
 
 	FMemoryTrackerId GetTrackerId() const { return TrackerId; }
+	FMemoryTagSetId GetTagSetId() const { return TagSetId; }
 
 	const FLinearColor& GetColor() const { return Color; }
 	void SetColor(FLinearColor InColor) { Color = InColor; }
@@ -65,6 +88,8 @@ private:
 	FString StatName;
 	FString StatFullName; // includes the parent prefix
 	FMemoryTrackerId TrackerId = FMemoryTracker::InvalidTrackerId;
+	FMemoryTagSetId TagSetId = FMemoryTagSet::InvalidTagSetId;
+
 	FLinearColor Color;
 	FMemoryTag* Parent = nullptr;
 	TSet<FMemoryTag*> Children;
diff --git a/Engine/Source/Developer/TraceInsights/Private/Insights/MemoryProfiler/Widgets/SMemTagTreeView.cpp b/Engine/Source/Developer/TraceInsights/Private/Insights/MemoryProfiler/Widgets/SMemTagTreeView.cpp
--- a/Engine/Source/Developer/TraceInsights/Private/Insights/MemoryProfiler/Widgets/SMemTagTreeView.cpp
+++ b/Engine/Source/Developer/TraceInsights/Private/Insights/MemoryProfiler/Widgets/SMemTagTreeView.cpp
@@ -46,6 +46,7 @@ SMemTagTreeView::SMemTagTreeView()
 	, bExpansionSaved(false)
 	, bFilterOutZeroCountMemTags(false)
 	, TrackersFilter(uint64(-1))
+	, TagSetFilter(FMemoryTagSet::DefaultTagSetId)
 	, GroupingMode(EMemTagNodeGroupingMode::Flat)
 	, AvailableSorters()
 	, CurrentSorter(nullptr)
@@ -233,6 +234,62 @@ bool SMemTagTreeView::IsTrackerChecked(FMemoryTrackerId InTrackerId) const
 
 ////////////////////////////////////////////////////////////////////////////////////////////////////
 
+void SMemTagTreeView::TagSet_OnSelectionChanged(TSharedPtr<FMemoryTagSetId> NewTagSet, ESelectInfo::Type SelectInfo)
+{
+	if (NewTagSet.IsValid())
+	{
+		TagSetFilter = *NewTagSet;
+		ApplyFiltering();
+	}
+}
+
+////////////////////////////////////////////////////////////////////////////////////////////////////
+
+TSharedRef<SWidget> SMemTagTreeView::TagSet_OnGenerateWidget(TSharedPtr<FMemoryTagSetId> InTagSet) const
+{
+	FText TagSetName;
+	if (InTagSet.IsValid())
+	{
+		TSharedPtr<SMemoryProfilerWindow> ProfilerWindow = GetProfilerWindow();
+		if (ProfilerWindow.IsValid())
+		{
+			FMemorySharedState& SharedState = ProfilerWindow->GetSharedState();
+			const FMemoryTagSet* TagSet = SharedState.GetTagSetById(*InTagSet);
+			if (TagSet)
+			{
+				TagSetName = FText::FromString(TagSet->GetName());
+			}
+		}
+
+		if (TagSetName.IsEmpty())
+		{
+			TagSetName = FText::AsNumber((int64)(*InTagSet));
+		}
+	}
+
+	return SNew(STextBlock)
+		.Text(TagSetName);
+}
+
+////////////////////////////////////////////////////////////////////////////////////////////////////
+
+FText SMemTagTreeView::TagSet_GetSelectedText() const
+{
+	TSharedPtr<SMemoryProfilerWindow> ProfilerWindow = GetProfilerWindow();
+	if (ProfilerWindow.IsValid())
+	{
+		FMemorySharedState& SharedState = ProfilerWindow->GetSharedState();
+		const FMemoryTagSet* TagSet = SharedState.GetTagSetById(TagSetFilter);
+		if (TagSet)
+		{
+			return FText::FromString(TagSet->GetName());
+		}
+	}
+	return FText::GetEmpty();
+}
+
+////////////////////////////////////////////////////////////////////////////////////////////////////
+
 TSharedRef<SWidget> SMemTagTreeView::ConstructTagsFilteringWidgetArea()
 {
 	FSlimHorizontalToolBarBuilder ToolbarBuilder(TSharedPtr<const FUICommandList>(), FMultiBoxCustomization::None);
@@ -258,6 +315,18 @@ TSharedRef<SWidget> SMemTagTreeView::ConstructTagsFilteringWidgetArea()
 			FSlateIcon(FInsightsStyle::GetStyleSetName(), "Icons.Filter.ToolBar"),
 			false
 		);
+
+		// Select Tag Set
+		ToolbarBuilder.AddWidget(
+			SAssignNew(TagSetComboBox, SComboBox<TSharedPtr<FMemoryTagSetId>>)
+				.ToolTipText(LOCTEXT("TagSetToolTip", "Select the LLM tag set."))
+				.OptionsSource(&AvailableTagSets)
+				.OnSelectionChanged(this, &SMemTagTreeView::TagSet_OnSelectionChanged)
+				.OnGenerateWidget(this, &SMemTagTreeView::TagSet_OnGenerateWidget)
+				[
+					SNew(STextBlock)
+					.Text(this, &SMemTagTreeView::TagSet_GetSelectedText)
+				]);
 	}
 	ToolbarBuilder.EndSection();
 
@@ -295,7 +364,7 @@ TSharedRef<SWidget> SMemTagTreeView::ConstructTagsFilteringWidgetArea()
 	//	]
 	//]
 
-	// Filter LLM tags by tracker.
+	// Trackers and Tag Sets
 	+ SHorizontalBox::Slot()
 	.VAlign(VAlign_Center)
 	.Padding(FMargin(-2.0f, -5.0f))
@@ -1007,6 +1076,7 @@ void SMemTagTreeView::ApplyFiltering()
 			const FMemTagNodePtr& NodePtr = StaticCastSharedPtr<FMemTagNode, FBaseTreeNode>(Child);
 			const bool bIsChildVisible = (!bFilterOutZeroCountMemTags || NodePtr->GetAggregatedStats().InstanceCount > 0)
 									  && (!FMemoryTracker::IsValidTrackerId(NodePtr->GetMemTrackerId()) || ((FMemoryTracker::AsFlag(NodePtr->GetMemTrackerId()) & TrackersFilter) != 0))
+									  && (NodePtr->GetMemTagSetId() == TagSetFilter)
 									  && Filters->PassesAllFilters(NodePtr);
 			if (bIsChildVisible)
 			{
@@ -1919,6 +1989,16 @@ void SMemTagTreeView::RebuildTree(bool bResync)
 				}
 			}
 		}
+
+		const auto& TagSets = SharedState.GetTagSets();
+		if (AvailableTagSets.Num() != TagSets.Num())
+		{
+			AvailableTagSets.Reset();
+			for (const auto& TagSet : TagSets)
+			{
+				AvailableTagSets.Add(MakeShared<FMemoryTagSetId>(TagSet->GetId()));
+			}
+		}
 	}
 
 	SyncStopwatch.Stop();
diff --git a/Engine/Source/Developer/TraceInsights/Private/Insights/MemoryProfiler/Widgets/SMemTagTreeView.h b/Engine/Source/Developer/TraceInsights/Private/Insights/MemoryProfiler/Widgets/SMemTagTreeView.h
--- a/Engine/Source/Developer/TraceInsights/Private/Insights/MemoryProfiler/Widgets/SMemTagTreeView.h
+++ b/Engine/Source/Developer/TraceInsights/Private/Insights/MemoryProfiler/Widgets/SMemTagTreeView.h
@@ -183,6 +183,10 @@ private:
 	void ToggleTracker(FMemoryTrackerId InTrackerId);
 	bool IsTrackerChecked(FMemoryTrackerId InTrackerId) const;
 
+	void TagSet_OnSelectionChanged(TSharedPtr<FMemoryTagSetId> NewTagSet, ESelectInfo::Type SelectInfo);
+	TSharedRef<SWidget> TagSet_OnGenerateWidget(TSharedPtr<FMemoryTagSetId> InTagSet) const;
+	FText TagSet_GetSelectedText() const;
+
 	////////////////////////////////////////////////////////////////////////////////////////////////////
 	// Grouping
 
@@ -389,6 +393,11 @@ private:
 	/** Filter the LLM tags by tracker. */
 	uint64 TrackersFilter;
 
+	/** Filter the LLM tags by tag set. */
+	FMemoryTagSetId TagSetFilter;
+	TArray<TSharedPtr<FMemoryTagSetId>> AvailableTagSets;
+	TSharedPtr<SComboBox<TSharedPtr<FMemoryTagSetId>>> TagSetComboBox;
+
 	//////////////////////////////////////////////////
 	// Grouping
 
