[Insights] Added EnumerateTracks() API to ITimingViewSession.

diff --git a/Engine/Source/Developer/TraceInsights/Private/Insights/Widgets/STimingView.cpp b/Engine/Source/Developer/TraceInsights/Private/Insights/Widgets/STimingView.cpp
--- a/Engine/Source/Developer/TraceInsights/Private/Insights/Widgets/STimingView.cpp
+++ b/Engine/Source/Developer/TraceInsights/Private/Insights/Widgets/STimingView.cpp
@@ -5287,6 +5287,16 @@ TSharedPtr<FBaseTimingTrack> STimingView::FindTrack(uint64 InTrackId)
 
 ////////////////////////////////////////////////////////////////////////////////////////////////////
 
+void STimingView::EnumerateTracks(TFunctionRef<void(TSharedPtr<FBaseTimingTrack> Track)> Callback)
+{
+	for (auto& Entry : AllTracks)
+	{
+		Callback(Entry.Value);
+	}
+}
+
+////////////////////////////////////////////////////////////////////////////////////////////////////
+
 TArray<Timing::ITimingViewExtender*> STimingView::GetExtenders() const
 {
 	return IModularFeatures::Get().GetModularFeatureImplementations<Timing::ITimingViewExtender>(Timing::TimingViewExtenderFeatureName);
diff --git a/Engine/Source/Developer/TraceInsights/Private/Insights/Widgets/STimingView.h b/Engine/Source/Developer/TraceInsights/Private/Insights/Widgets/STimingView.h
--- a/Engine/Source/Developer/TraceInsights/Private/Insights/Widgets/STimingView.h
+++ b/Engine/Source/Developer/TraceInsights/Private/Insights/Widgets/STimingView.h
@@ -295,6 +295,7 @@ public:
 	virtual bool RemoveTrack(TSharedPtr<FBaseTimingTrack> Track) override;
 
 	virtual TSharedPtr<FBaseTimingTrack> FindTrack(uint64 InTrackId) override;
+	virtual void EnumerateTracks(TFunctionRef<void(TSharedPtr<FBaseTimingTrack> Track)> Callback) override;
 
 	virtual double GetTimeMarker() const override;
 	virtual void SetTimeMarker(double InTimeMarker) override;
diff --git a/Engine/Source/Developer/TraceInsights/Public/Insights/ITimingViewSession.h b/Engine/Source/Developer/TraceInsights/Public/Insights/ITimingViewSession.h
--- a/Engine/Source/Developer/TraceInsights/Public/Insights/ITimingViewSession.h
+++ b/Engine/Source/Developer/TraceInsights/Public/Insights/ITimingViewSession.h
@@ -126,6 +126,9 @@ public:
 	/** Finds a track has been added via Add*Track(). */
 	virtual TSharedPtr<FBaseTimingTrack> FindTrack(uint64 InTrackId) = 0;
 
+	/** Enumerate all available tracks (for all locations). */
+	virtual void EnumerateTracks(TFunctionRef<void(TSharedPtr<FBaseTimingTrack> Track)> Callback) = 0;
+
 	//////////////////////////////////////////////////
 
 	/** Gets the current marker time. */
