[Insights]
- Added public abstract class for FThreadTimingTrack (implemented by CPU timing event tracks).
- Moved the CPU/GPU tracks in separate source files.

diff --git a/Engine/Source/Developer/TraceInsights/Private/Insights/ContextSwitches/ViewModels/ContextSwitchesSharedState.cpp b/Engine/Source/Developer/TraceInsights/Private/Insights/ContextSwitches/ViewModels/ContextSwitchesSharedState.cpp
--- a/Engine/Source/Developer/TraceInsights/Private/Insights/ContextSwitches/ViewModels/ContextSwitchesSharedState.cpp
+++ b/Engine/Source/Developer/TraceInsights/Private/Insights/ContextSwitches/ViewModels/ContextSwitchesSharedState.cpp
@@ -23,7 +23,7 @@
 #include "Insights/InsightsManager.h"
 #include "Insights/InsightsStyle.h"
 #include "Insights/TimingProfiler/TimingProfilerManager.h"
-#include "Insights/TimingProfiler/Tracks/ThreadTimingTrack.h"
+#include "Insights/TimingProfiler/Tracks/CpuTimingTrack.h"
 #include "Insights/TimingProfiler/ViewModels/ThreadTimingSharedState.h"
 #include "Insights/TimingProfiler/Widgets/STimingProfilerWindow.h"
 #include "Insights/Widgets/STimingView.h"
diff --git a/Engine/Source/Developer/TraceInsights/Private/Insights/TaskGraphProfiler/TaskGraphProfilerManager.cpp b/Engine/Source/Developer/TraceInsights/Private/Insights/TaskGraphProfiler/TaskGraphProfilerManager.cpp
--- a/Engine/Source/Developer/TraceInsights/Private/Insights/TaskGraphProfiler/TaskGraphProfilerManager.cpp
+++ b/Engine/Source/Developer/TraceInsights/Private/Insights/TaskGraphProfiler/TaskGraphProfilerManager.cpp
@@ -19,7 +19,7 @@
 #include "Insights/TaskGraphProfiler/ViewModels/TaskTimingTrack.h"
 #include "Insights/TaskGraphProfiler/Widgets/STaskTableTreeView.h"
 #include "Insights/TimingProfiler/TimingProfilerManager.h"
-#include "Insights/TimingProfiler/Tracks/ThreadTimingTrack.h"
+#include "Insights/TimingProfiler/Tracks/CpuTimingTrack.h"
 #include "Insights/TimingProfiler/ViewModels/ThreadTimingSharedState.h"
 #include "Insights/TimingProfiler/Widgets/STimingProfilerWindow.h"
 #include "Insights/ViewModels/ThreadTrackEvent.h"
diff --git a/Engine/Source/Developer/TraceInsights/Private/Insights/TaskGraphProfiler/Widgets/STaskTableTreeView.cpp b/Engine/Source/Developer/TraceInsights/Private/Insights/TaskGraphProfiler/Widgets/STaskTableTreeView.cpp
--- a/Engine/Source/Developer/TraceInsights/Private/Insights/TaskGraphProfiler/Widgets/STaskTableTreeView.cpp
+++ b/Engine/Source/Developer/TraceInsights/Private/Insights/TaskGraphProfiler/Widgets/STaskTableTreeView.cpp
@@ -26,7 +26,7 @@
 #include "Insights/TaskGraphProfiler/ViewModels/TaskNode.h"
 #include "Insights/TaskGraphProfiler/ViewModels/TaskTimingTrack.h"
 #include "Insights/TimingProfiler/TimingProfilerManager.h"
-#include "Insights/TimingProfiler/Tracks/ThreadTimingTrack.h"
+#include "Insights/TimingProfiler/Tracks/CpuTimingTrack.h"
 #include "Insights/TimingProfiler/ViewModels/ThreadTimingSharedState.h"
 #include "Insights/TimingProfiler/Widgets/STimingProfilerWindow.h"
 #include "Insights/Widgets/STimingView.h"
diff --git a/Engine/Source/Developer/TraceInsights/Private/Insights/TimingProfiler/Tracks/CpuTimingTrack.cpp b/Engine/Source/Developer/TraceInsights/Private/Insights/TimingProfiler/Tracks/CpuTimingTrack.cpp
new file mode 100644
--- /dev/null
+++ b/Engine/Source/Developer/TraceInsights/Private/Insights/TimingProfiler/Tracks/CpuTimingTrack.cpp
@@ -0,0 +1,216 @@
+// Copyright Epic Games, Inc. All Rights Reserved.
+
+#include "CpuTimingTrack.h"
+
+// TraceServices
+#include "TraceServices/Model/AnalysisSession.h"
+#include "TraceServices/Model/TasksProfiler.h"
+
+// TraceInsightsCore
+#include "InsightsCore/Common/TimeUtils.h"
+
+// TraceInsights
+#include "Insights/TimingProfiler/ViewModels/ThreadTimingSharedState.h"
+#include "Insights/ViewModels/ThreadTrackEvent.h"
+#include "Insights/ViewModels/TooltipDrawState.h"
+
+#define LOCTEXT_NAMESPACE "UE::Insights::TimingProfiler::CpuTimingTrack"
+
+namespace UE::Insights::TimingProfiler
+{
+
+////////////////////////////////////////////////////////////////////////////////////////////////////
+// FCpuTimingTrack
+////////////////////////////////////////////////////////////////////////////////////////////////////
+
+INSIGHTS_IMPLEMENT_RTTI(FCpuTimingTrack)
+
+////////////////////////////////////////////////////////////////////////////////////////////////////
+
+void FCpuTimingTrack::AddTaskInfo(FTooltipDrawState& InOutTooltip, const TraceServices::FTaskInfo& Task) const
+{
+	InOutTooltip.AddTextLine(FString::Printf(TEXT("-------- Task %d%s --------"), Task.Id, Task.bTracked ? TEXT("") : TEXT(" (not tracked)")), FLinearColor::Green);
+
+	if (Task.DebugName != nullptr)
+	{
+		InOutTooltip.AddTextLine(FString::Printf(TEXT("%s"), Task.DebugName), FLinearColor::Green);
+	}
+
+	ENamedThreads::Type ThreadInfo = (ENamedThreads::Type)Task.ThreadToExecuteOn;
+	ENamedThreads::Type ThreadIndex = ENamedThreads::GetThreadIndex(ThreadInfo);
+
+	auto FormatTaskTimestamp = [](double Timestamp) -> FString
+	{
+		return (Timestamp != TraceServices::FTaskInfo::InvalidTimestamp) ? FString::SanitizeFloat(Timestamp) : TEXT("[not set]");
+	};
+
+	auto FormatTaskTime = [](double Time) -> FString
+	{
+		return FormatTimeAuto(Time, 2);
+	};
+
+	auto GetTrackName = [this](uint32 InThreadId) -> FString
+	{
+		TSharedPtr<FCpuTimingTrack> Track = GetSharedState().GetCpuTrack(InThreadId);
+		return Track.IsValid() ? Track->GetName() : TEXT("Unknown");
+	};
+
+	if (ThreadIndex == ENamedThreads::AnyThread)
+	{
+		const TCHAR* TaskPri = ENamedThreads::GetTaskPriority(ThreadInfo) == ENamedThreads::NormalTaskPriority ? TEXT("Normal") : TEXT("High");
+
+		int32 ThreadPriIndex = ENamedThreads::GetThreadPriorityIndex(ThreadInfo);
+		const TCHAR* ThreadPriStrs[] = { TEXT("Normal"), TEXT("High"), TEXT("Low") };
+		const TCHAR* ThreadPri = ensure(ThreadPriIndex >= 0 && ThreadPriIndex < 3) ? ThreadPriStrs[ThreadPriIndex] : TEXT("Unknown");
+
+		InOutTooltip.AddTextLine(
+			FString::Printf(TEXT("%s Pri task on %s Pri worker (%s)"), TaskPri, ThreadPri, *GetTrackName(Task.StartedThreadId)),
+			FLinearColor::Green);
+	}
+	else
+	{
+		const TCHAR* QueueStr = ENamedThreads::GetQueueIndex(ThreadInfo) == ENamedThreads::MainQueue ? TEXT("Main") : TEXT("Local");
+		InOutTooltip.AddTextLine(
+			FString::Printf(TEXT("%s (%s queue)"), *GetTrackName(Task.StartedThreadId), QueueStr),
+			FLinearColor::Green);
+	}
+
+	InOutTooltip.AddNameValueTextLine(TEXTVIEW("Created:"), FString::Printf(TEXT("%s on %s"),
+		*FormatTaskTimestamp(Task.CreatedTimestamp),
+		*GetTrackName(Task.CreatedThreadId)));
+
+	InOutTooltip.AddNameValueTextLine(TEXTVIEW("Launched:"), FString::Printf(TEXT("%s (+%s) on %s"),
+		*FormatTaskTimestamp(Task.LaunchedTimestamp),
+		*FormatTaskTime(Task.LaunchedTimestamp - Task.CreatedTimestamp),
+		*GetTrackName(Task.LaunchedThreadId)));
+
+	InOutTooltip.AddNameValueTextLine(TEXTVIEW("Scheduled:"), FString::Printf(TEXT("%s (+%s) on %s"),
+		*FormatTaskTimestamp(Task.ScheduledTimestamp),
+		*FormatTaskTime(Task.ScheduledTimestamp - Task.LaunchedTimestamp),
+		*GetTrackName(Task.ScheduledThreadId)));
+
+	InOutTooltip.AddNameValueTextLine(TEXTVIEW("Started:"), FString::Printf(TEXT("%s (+%s)"),
+		*FormatTaskTimestamp(Task.StartedTimestamp),
+		*FormatTaskTime(Task.StartedTimestamp - Task.ScheduledTimestamp)));
+
+	if (Task.FinishedTimestamp != TraceServices::FTaskInfo::InvalidTimestamp)
+	{
+		InOutTooltip.AddNameValueTextLine(TEXTVIEW("Finished:"), FString::Printf(TEXT("%s (+%s)"),
+			*FormatTaskTimestamp(Task.FinishedTimestamp),
+			*FormatTaskTime(Task.FinishedTimestamp - Task.StartedTimestamp)));
+
+		if (Task.CompletedTimestamp != TraceServices::FTaskInfo::InvalidTimestamp)
+		{
+			InOutTooltip.AddNameValueTextLine(TEXTVIEW("Completed:"), FString::Printf(TEXT("%s (+%s) on %s"),
+				*FormatTaskTimestamp(Task.CompletedTimestamp),
+				*FormatTaskTime(Task.CompletedTimestamp - Task.FinishedTimestamp),
+				*GetTrackName(Task.CompletedThreadId)));
+
+			if (Task.DestroyedTimestamp != TraceServices::FTaskInfo::InvalidTimestamp)
+			{
+				InOutTooltip.AddNameValueTextLine(TEXTVIEW("Destroyed:"), FString::Printf(TEXT("%s (+%s) on %s"),
+					*FormatTaskTimestamp(Task.DestroyedTimestamp),
+					*FormatTaskTime(Task.DestroyedTimestamp - Task.CompletedTimestamp),
+					*GetTrackName(Task.DestroyedThreadId)));
+			}
+		}
+	}
+
+	InOutTooltip.AddNameValueTextLine(TEXTVIEW("Prerequisite tasks:"), FString::Printf(TEXT("%d"), Task.Prerequisites.Num()));
+	InOutTooltip.AddNameValueTextLine(TEXTVIEW("Subsequent tasks:"), FString::Printf(TEXT("%d"), Task.Subsequents.Num()));
+	InOutTooltip.AddNameValueTextLine(TEXTVIEW("Parent tasks:"), FString::Printf(TEXT("%d"), Task.ParentTasks.Num()));
+	InOutTooltip.AddNameValueTextLine(TEXTVIEW("Nested tasks:"), FString::Printf(TEXT("%d"), Task.NestedTasks.Num()));
+}
+
+////////////////////////////////////////////////////////////////////////////////////////////////////
+
+void FCpuTimingTrack::PostInitTooltip(FTooltipDrawState& InOutTooltip, const FThreadTrackEvent& TooltipEvent, const TraceServices::IAnalysisSession& Session, const TCHAR* TimerName) const
+{
+	// tasks
+	const TraceServices::ITasksProvider* TasksProvider = TraceServices::ReadTasksProvider(Session);
+	if (TasksProvider != nullptr)
+	{
+		// info about a task
+		const TraceServices::FTaskInfo* Task = TasksProvider->TryGetTask(GetThreadId(), TooltipEvent.GetStartTime());
+		if (Task != nullptr && Task->FinishedTimestamp >= TooltipEvent.GetEndTime())
+		{
+			AddTaskInfo(InOutTooltip, *Task);
+		}
+
+		// info about blocking
+		const TraceServices::FWaitingForTasks* Waiting = TasksProvider->TryGetWaiting(TimerName, GetThreadId(), TooltipEvent.GetStartTime());
+		if (Waiting != nullptr && Waiting->Tasks.Num() > 0)
+		{
+			InOutTooltip.AddTextLine(TEXT("-------- Waiting for tasks --------"), FLinearColor::Red);
+			constexpr int32 NumIdsOnRow = 4;
+			TStringBuilder<1024> StringBuilder;
+
+			// Add the first line of Task Id values.
+			for (int32 Index = 0; Index < Waiting->Tasks.Num() && Index < NumIdsOnRow; ++Index)
+			{
+				StringBuilder.Appendf(TEXT("%d, "), Waiting->Tasks[Index]);
+			}
+
+			// Remove separators from last entry.
+			if (Waiting->Tasks.Num() <= NumIdsOnRow)
+			{
+				StringBuilder.RemoveSuffix(2);
+			}
+			InOutTooltip.AddNameValueTextLine(FString::Printf(TEXT("Tasks[%d]:"), Waiting->Tasks.Num()), StringBuilder.ToView());
+			StringBuilder.Reset();
+
+			// Add the rest of the lines with an empty name so they appear as a multi line value.
+			for (int32 Index = NumIdsOnRow; Index < Waiting->Tasks.Num(); ++Index)
+			{
+				StringBuilder.Appendf(TEXT("%d, "), Waiting->Tasks[Index]);
+
+				if ((Index + 1) % NumIdsOnRow == 0)
+				{
+					InOutTooltip.AddNameValueTextLine(TEXT(""), StringBuilder.ToView());
+					StringBuilder.Reset();
+				}
+			}
+			if (StringBuilder.Len() > 1)
+			{
+				StringBuilder.RemoveSuffix(2);
+				InOutTooltip.AddNameValueTextLine(TEXT(""), StringBuilder.ToView());
+			}
+
+			InOutTooltip.AddNameValueTextLine(TEXT("Started waiting:"),
+				FString::Printf(TEXT("%s"), *FString::SanitizeFloat(Waiting->StartedTimestamp)));
+
+			if (Waiting->FinishedTimestamp != TraceServices::FTaskInfo::InvalidTimestamp)
+			{
+				InOutTooltip.AddNameValueTextLine(TEXT("Finished waiting:"),
+					FString::Printf(TEXT("%s (+%s)"),
+						*FString::SanitizeFloat(Waiting->FinishedTimestamp),
+						*FormatTimeAuto(Waiting->FinishedTimestamp - Waiting->StartedTimestamp, 2)));
+			}
+			else
+			{
+				InOutTooltip.AddNameValueTextLine(TEXT("Finished waiting:"), TEXT("[not set]"));
+			}
+
+			const int32 MaxWaitedTasksToList = 5;
+			int32 NumTasksToList = FMath::Min(Waiting->Tasks.Num(), MaxWaitedTasksToList);
+			for (int32 TaskIndex = 0; TaskIndex != NumTasksToList; ++TaskIndex)
+			{
+				const TraceServices::FTaskInfo* WaitedTask = TasksProvider->TryGetTask(Waiting->Tasks[TaskIndex]);
+				if (WaitedTask != nullptr)
+				{
+					AddTaskInfo(InOutTooltip , *WaitedTask);
+				}
+			}
+			if (NumTasksToList < Waiting->Tasks.Num())
+			{
+				InOutTooltip.AddTextLine(TEXT("[...]"), FLinearColor::Green);
+			}
+		}
+	}
+}
+
+////////////////////////////////////////////////////////////////////////////////////////////////////
+
+} // namespace UE::Insights::TimingProfiler
+
+#undef LOCTEXT_NAMESPACE
diff --git a/Engine/Source/Developer/TraceInsights/Private/Insights/TimingProfiler/Tracks/CpuTimingTrack.h b/Engine/Source/Developer/TraceInsights/Private/Insights/TimingProfiler/Tracks/CpuTimingTrack.h
new file mode 100644
--- /dev/null
+++ b/Engine/Source/Developer/TraceInsights/Private/Insights/TimingProfiler/Tracks/CpuTimingTrack.h
@@ -0,0 +1,39 @@
+// Copyright Epic Games, Inc. All Rights Reserved.
+
+#pragma once
+
+#include "CoreTypes.h"
+
+// TraceInsights
+#include "Insights/TimingProfiler/Tracks/ThreadTimingTrackPrivate.h"
+
+namespace TraceServices
+{
+	struct FTaskInfo;
+}
+
+namespace UE::Insights::TimingProfiler
+{
+
+////////////////////////////////////////////////////////////////////////////////////////////////////
+
+class FCpuTimingTrack : public FThreadTimingTrackImpl
+{
+	INSIGHTS_DECLARE_RTTI(FCpuTimingTrack, FThreadTimingTrackImpl)
+
+public:
+	explicit FCpuTimingTrack(FThreadTimingSharedState& InSharedState, const FString& InName, const TCHAR* InGroupName, uint32 InTimelineIndex, uint32 InThreadId)
+		: FThreadTimingTrackImpl(InSharedState, InName, InGroupName, InTimelineIndex, InThreadId)
+	{
+	}
+
+protected:
+	virtual void PostInitTooltip(FTooltipDrawState& InOutTooltip, const FThreadTrackEvent& TooltipEvent, const TraceServices::IAnalysisSession& Session, const TCHAR* TimerName) const override;
+
+private:
+	void AddTaskInfo(FTooltipDrawState& InOutTooltip, const TraceServices::FTaskInfo& Task) const;
+};
+
+////////////////////////////////////////////////////////////////////////////////////////////////////
+
+} // namespace UE::Insights::TimingProfiler
diff --git a/Engine/Source/Developer/TraceInsights/Private/Insights/TimingProfiler/Tracks/GpuTimingTrack.cpp b/Engine/Source/Developer/TraceInsights/Private/Insights/TimingProfiler/Tracks/GpuTimingTrack.cpp
new file mode 100644
--- /dev/null
+++ b/Engine/Source/Developer/TraceInsights/Private/Insights/TimingProfiler/Tracks/GpuTimingTrack.cpp
@@ -0,0 +1,146 @@
+// Copyright Epic Games, Inc. All Rights Reserved.
+
+#include "GpuTimingTrack.h"
+
+// TraceInsights
+#include "Insights/TimingProfiler/ViewModels/ThreadTimingSharedState.h"
+#include "Insights/ViewModels/TimingTrackViewport.h"
+#include "Insights/ViewModels/TimingViewDrawHelper.h"
+
+#define LOCTEXT_NAMESPACE "UE::Insights::TimingProfiler::GpuTimingTrack"
+
+namespace UE::Insights::TimingProfiler
+{
+
+////////////////////////////////////////////////////////////////////////////////////////////////////
+// FGpuTimingTrack, FGpuQueueTimingTrack
+////////////////////////////////////////////////////////////////////////////////////////////////////
+
+INSIGHTS_IMPLEMENT_RTTI(FGpuTimingTrack)
+INSIGHTS_IMPLEMENT_RTTI(FGpuQueueTimingTrack)
+
+////////////////////////////////////////////////////////////////////////////////////////////////////
+// FGpuQueueWorkTimingTrack
+////////////////////////////////////////////////////////////////////////////////////////////////////
+
+INSIGHTS_IMPLEMENT_RTTI(FGpuQueueWorkTimingTrack)
+
+////////////////////////////////////////////////////////////////////////////////////////////////////
+
+void FGpuQueueWorkTimingTrack::Draw(const ITimingTrackDrawContext& Context) const
+{
+	DrawLineEvents(Context, 1.0f);
+}
+
+////////////////////////////////////////////////////////////////////////////////////////////////////
+
+void FGpuQueueWorkTimingTrack::DrawLineEvents(const ITimingTrackDrawContext& Context, const float OffsetY) const
+{
+	const FTimingViewDrawHelper& Helper = *static_cast<const FTimingViewDrawHelper*>(&Context.GetHelper());
+
+	if ((Context.GetEventFilter().IsValid() && Context.GetEventFilter()->FilterTrack(*this)) || HasCustomFilter())
+	{
+		Helper.DrawFadedLineEvents(GetDrawState(), *this, OffsetY, 0.1f);
+
+		if (UpdateFilteredDrawStateOpacity())
+		{
+			Helper.DrawLineEvents(GetFilteredDrawState(), *this, OffsetY);
+		}
+		else
+		{
+			Helper.DrawFadedLineEvents(GetFilteredDrawState(), *this, OffsetY, GetFilteredDrawStateOpacity());
+		}
+	}
+	else
+	{
+		Helper.DrawLineEvents(GetDrawState(), *this, OffsetY);
+	}
+}
+
+////////////////////////////////////////////////////////////////////////////////////////////////////
+
+void FGpuQueueWorkTimingTrack::PostDraw(const ITimingTrackDrawContext& Context) const
+{
+	const bool bAreOverlaysVisible = GetSharedState().AreOverlaysVisibleInGpuQueueTracks();
+	const bool bAreExtendedLinesVisible = GetSharedState().AreExtendedLinesVisibleInGpuQueueTracks();
+
+	if (bAreOverlaysVisible || bAreExtendedLinesVisible)
+	{
+		float LineY1 = 0.0f;
+		float LineY2 = 0.0f;
+		ETimingTrackLocation LocalLocation = ETimingTrackLocation::None;
+		TSharedPtr<FBaseTimingTrack> LocalParentTrack = GetParentTrack().Pin();
+		if (LocalParentTrack)
+		{
+			LineY1 = LocalParentTrack->GetPosY();
+			LineY2 = LineY1 + LocalParentTrack->GetHeight();
+			LocalLocation = LocalParentTrack->GetLocation();
+		}
+		else
+		{
+			LineY1 = GetPosY();
+			LineY2 = LineY1 + GetHeight();
+			LocalLocation = GetLocation();
+		}
+
+		const FTimingTrackViewport& Viewport = Context.GetViewport();
+		switch (LocalLocation)
+		{
+			case ETimingTrackLocation::Scrollable:
+			{
+				const float TopY = Viewport.GetPosY() + Viewport.GetTopOffset();
+				if (LineY1 < TopY)
+				{
+					LineY1 = TopY;
+				}
+				const float BottomY = Viewport.GetPosY() + Viewport.GetHeight() - Viewport.GetBottomOffset();
+				if (LineY2 > BottomY)
+				{
+					LineY2 = BottomY;
+				}
+				break;
+			}
+			case ETimingTrackLocation::TopDocked:
+			{
+				const float TopY = Viewport.GetPosY();
+				if (LineY1 < TopY)
+				{
+					LineY1 = TopY;
+				}
+				const float BottomY = Viewport.GetPosY() + Viewport.GetTopOffset();
+				if (LineY2 > BottomY)
+				{
+					LineY2 = BottomY;
+				}
+				break;
+			}
+			case ETimingTrackLocation::BottomDocked:
+			{
+				const float TopY = Viewport.GetPosY() + Viewport.GetHeight() - Viewport.GetBottomOffset();
+				if (LineY1 < TopY)
+				{
+					LineY1 = TopY;
+				}
+				const float BottomY = Viewport.GetPosY() + Viewport.GetHeight();
+				if (LineY2 > BottomY)
+				{
+					LineY2 = BottomY;
+				}
+				break;
+			}
+		}
+
+		const float LineH = LineY2 - LineY1;
+		if (LineH > 0.0f)
+		{
+			const FTimingViewDrawHelper& Helper = *static_cast<const FTimingViewDrawHelper*>(&Context.GetHelper());
+			Helper.DrawContextSwitchMarkers(GetDrawState(), LineY1, LineH, 0.25f, bAreOverlaysVisible, bAreExtendedLinesVisible);
+		}
+	}
+}
+
+////////////////////////////////////////////////////////////////////////////////////////////////////
+
+} // namespace UE::Insights::TimingProfiler
+
+#undef LOCTEXT_NAMESPACE
diff --git a/Engine/Source/Developer/TraceInsights/Private/Insights/TimingProfiler/Tracks/GpuTimingTrack.h b/Engine/Source/Developer/TraceInsights/Private/Insights/TimingProfiler/Tracks/GpuTimingTrack.h
new file mode 100644
--- /dev/null
+++ b/Engine/Source/Developer/TraceInsights/Private/Insights/TimingProfiler/Tracks/GpuTimingTrack.h
@@ -0,0 +1,68 @@
+// Copyright Epic Games, Inc. All Rights Reserved.
+
+#pragma once
+
+#include "CoreTypes.h"
+
+// TraceInsights
+#include "Insights/TimingProfiler/Tracks/ThreadTimingTrackPrivate.h"
+
+namespace UE::Insights::TimingProfiler
+{
+
+////////////////////////////////////////////////////////////////////////////////////////////////////
+
+class FGpuTimingTrack : public FThreadTimingTrackImpl
+{
+	INSIGHTS_DECLARE_RTTI(FGpuTimingTrack, FThreadTimingTrackImpl)
+
+public:
+	static constexpr uint32 Gpu1ThreadId = uint32('GPU1');
+	static constexpr uint32 Gpu2ThreadId = uint32('GPU2');
+
+public:
+	explicit FGpuTimingTrack(FThreadTimingSharedState& InSharedState, const FString& InName, const TCHAR* InGroupName, uint32 InTimelineIndex, uint32 InThreadId)
+		: FThreadTimingTrackImpl(InSharedState, InName, InGroupName, InTimelineIndex, InThreadId)
+	{
+	}
+};
+
+////////////////////////////////////////////////////////////////////////////////////////////////////
+
+class FGpuQueueTimingTrack : public FThreadTimingTrackImpl
+{
+	INSIGHTS_DECLARE_RTTI(FGpuQueueTimingTrack, FThreadTimingTrackImpl)
+
+public:
+	explicit FGpuQueueTimingTrack(FThreadTimingSharedState& InSharedState, const FString& InName, uint32 InTimelineIndex, uint32 InQueueId)
+		: FThreadTimingTrackImpl(InSharedState, InName, nullptr, InTimelineIndex, InQueueId)
+	{
+	}
+
+	uint32 GetQueueId() const { return GetThreadId(); }
+};
+
+////////////////////////////////////////////////////////////////////////////////////////////////////
+
+class FGpuQueueWorkTimingTrack : public FThreadTimingTrackImpl
+{
+	INSIGHTS_DECLARE_RTTI(FGpuQueueWorkTimingTrack, FThreadTimingTrackImpl)
+
+public:
+	explicit FGpuQueueWorkTimingTrack(FThreadTimingSharedState& InSharedState, const FString& InName, uint32 InTimelineIndex, uint32 InQueueId)
+		: FThreadTimingTrackImpl(InSharedState, InName, nullptr, InTimelineIndex, InQueueId)
+	{
+	}
+
+	uint32 GetQueueId() const { return GetThreadId(); }
+
+	virtual void Draw(const ITimingTrackDrawContext& Context) const override;
+	virtual void PostDraw(const ITimingTrackDrawContext& Context) const override;
+
+private:
+	void DrawLineEvents(const ITimingTrackDrawContext& Context, const float OffsetY) const;
+};
+
+////////////////////////////////////////////////////////////////////////////////////////////////////
+
+} // namespace UE::Insights::TimingProfiler
diff --git a/Engine/Source/Developer/TraceInsights/Private/Insights/TimingProfiler/Tracks/ThreadTimingTrack.cpp b/Engine/Source/Developer/TraceInsights/Private/Insights/TimingProfiler/Tracks/ThreadTimingTrack.cpp
--- a/Engine/Source/Developer/TraceInsights/Private/Insights/TimingProfiler/Tracks/ThreadTimingTrack.cpp
+++ b/Engine/Source/Developer/TraceInsights/Private/Insights/TimingProfiler/Tracks/ThreadTimingTrack.cpp
@@ -1,6 +1,7 @@
 // Copyright Epic Games, Inc. All Rights Reserved.
 
-#include "ThreadTimingTrack.h"
+#include "Insights/TimingProfiler/Tracks/ThreadTimingTrack.h"
+#include "ThreadTimingTrackPrivate.h"
 
 #include "Async/TaskGraphInterfaces.h" // for ENamedThreads
 #include "CborReader.h"
@@ -11,7 +12,6 @@
 // TraceServices
 #include "TraceServices/Containers/Timelines.h"
 #include "TraceServices/Model/AnalysisSession.h"
-#include "TraceServices/Model/TasksProfiler.h"
 #include "TraceServices/Model/Threads.h"
 #include "TraceServices/Model/TimingProfiler.h"
 
@@ -23,6 +23,7 @@
 // TraceInsights
 #include "Insights/InsightsManager.h"
 #include "Insights/TimingProfiler/TimingProfilerManager.h"
+#include "Insights/TimingProfiler/Tracks/CpuTimingTrack.h"
 #include "Insights/TimingProfiler/ViewModels/ThreadTimingSharedState.h"
 #include "Insights/ViewModels/ThreadTrackEvent.h"
 #include "Insights/ViewModels/TimingEvent.h"
@@ -36,6 +37,8 @@
 namespace UE::Insights::TimingProfiler
 {
 
+INSIGHTS_IMPLEMENT_RTTI(FThreadTimingTrack);
+
 ////////////////////////////////////////////////////////////////////////////////////////////////////
 
 static void AppendMetadataToTooltip(FTooltipDrawState& Tooltip, TArrayView<const uint8>& Metadata)
@@ -320,14 +323,14 @@ static void AddTimingEventToBuilder(ITimingEventsTrackDrawStateBuilder& Builder,
 }
 
 ////////////////////////////////////////////////////////////////////////////////////////////////////
-// FThreadTimingTrack
+// FThreadTimingTrackImpl
 ////////////////////////////////////////////////////////////////////////////////////////////////////
 
-INSIGHTS_IMPLEMENT_RTTI(FThreadTimingTrack)
+INSIGHTS_IMPLEMENT_RTTI(FThreadTimingTrackImpl)
 
 ////////////////////////////////////////////////////////////////////////////////////////////////////
 
-void FThreadTimingTrack::BuildDrawState(ITimingEventsTrackDrawStateBuilder& Builder, const ITimingTrackUpdateContext& Context)
+void FThreadTimingTrackImpl::BuildDrawState(ITimingEventsTrackDrawStateBuilder& Builder, const ITimingTrackUpdateContext& Context)
 {
 	TSharedPtr<const TraceServices::IAnalysisSession> Session = FInsightsManager::Get()->GetSession();
 	if (Session.IsValid() && TraceServices::ReadTimingProfilerProvider(*Session.Get()))
@@ -385,8 +388,16 @@ void FThreadTimingTrack::BuildDrawState(ITimingEventsTrackDrawStateBuilder& Buil
 
 ////////////////////////////////////////////////////////////////////////////////////////////////////
 
-void FThreadTimingTrack::BuildFilteredDrawState(ITimingEventsTrackDrawStateBuilder& Builder, const ITimingTrackUpdateContext& Context)
+void FThreadTimingTrackImpl::BuildFilteredDrawState(ITimingEventsTrackDrawStateBuilder& Builder, const ITimingTrackUpdateContext& Context)
 {
+	struct FPendingEventInfo
+	{
+		double StartTime;
+		double EndTime;
+		uint32 Depth;
+		uint32 TimerIndex;
+	};
+
 	const TSharedPtr<ITimingEventFilter> EventFilterPtr = Context.GetEventFilter();
 	if (EventFilterPtr.IsValid() && EventFilterPtr->FilterTrack(*this))
 	{
@@ -606,7 +617,7 @@ void FThreadTimingTrack::BuildFilteredDrawState(ITimingEventsTrackDrawStateBuild
 
 ////////////////////////////////////////////////////////////////////////////////////////////////////
 
-void FThreadTimingTrack::PostDraw(const ITimingTrackDrawContext& Context) const
+void FThreadTimingTrackImpl::PostDraw(const ITimingTrackDrawContext& Context) const
 {
 	if (ChildTrack.IsValid())
 	{
@@ -661,7 +672,7 @@ void FThreadTimingTrack::PostDraw(const ITimingTrackDrawContext& Context) const
 
 ////////////////////////////////////////////////////////////////////////////////////////////////////
 
-void FThreadTimingTrack::InitTooltip(FTooltipDrawState& InOutTooltip, const ITimingEvent& InTooltipEvent) const
+void FThreadTimingTrackImpl::InitTooltip(FTooltipDrawState& InOutTooltip, const ITimingEvent& InTooltipEvent) const
 {
 	InOutTooltip.ResetContent();
 
@@ -694,8 +705,8 @@ void FThreadTimingTrack::InitTooltip(FTooltipDrawState& InOutTooltip, const ITim
 			ParentTimingEvent->GetDuration() > 0.0 &&
 			ParentTimingEvent->GetDuration() != std::numeric_limits<double>::infinity())
 		{
-			Timer = TimerReader->GetTimer(ParentTimingEvent->GetTimerIndex());
-			const TCHAR* ParentTimerName = (Timer != nullptr) ? Timer->Name : TEXT("N/A");
+			const TraceServices::FTimingProfilerTimer* ParentTimer = TimerReader->GetTimer(ParentTimingEvent->GetTimerIndex());
+			const TCHAR* ParentTimerName = (ParentTimer != nullptr) ? ParentTimer->Name : TEXT("N/A");
 			FNumberFormattingOptions FormattingOptions;
 			FormattingOptions.MaximumFractionalDigits = 2;
 			const FString ValueStr = FString::Printf(TEXT("%s %s"), *FText::AsPercent(TooltipEventDuration / ParentTimingEvent->GetDuration(), &FormattingOptions).ToString(), ParentTimerName);
@@ -707,8 +718,8 @@ void FThreadTimingTrack::InitTooltip(FTooltipDrawState& InOutTooltip, const ITim
 			RootTimingEvent->GetDuration() > 0.0 &&
 			RootTimingEvent->GetDuration() != std::numeric_limits<double>::infinity())
 		{
-			Timer = TimerReader->GetTimer(RootTimingEvent->GetTimerIndex());
-			const TCHAR* RootTimerName = (Timer != nullptr) ? Timer->Name : TEXT("N/A");
+			const TraceServices::FTimingProfilerTimer* RootTimer = TimerReader->GetTimer(RootTimingEvent->GetTimerIndex());
+			const TCHAR* RootTimerName = (RootTimer != nullptr) ? RootTimer->Name : TEXT("N/A");
 			FNumberFormattingOptions FormattingOptions;
 			FormattingOptions.MaximumFractionalDigits = 2;
 			const FString ValueStr = FString::Printf(TEXT("%s %s"), *FText::AsPercent(TooltipEventDuration / RootTimingEvent->GetDuration(), &FormattingOptions).ToString(), RootTimerName);
@@ -735,182 +746,7 @@ void FThreadTimingTrack::InitTooltip(FTooltipDrawState& InOutTooltip, const ITim
 			AppendMetadataToTooltip(InOutTooltip, Metadata);
 		}
 
-		// tasks
-		const TraceServices::ITasksProvider* TasksProvider = TraceServices::ReadTasksProvider(*Session.Get());
-		if (TasksProvider != nullptr)
-		{
-			auto AddTaskInfo = [&InOutTooltip, this](const TraceServices::FTaskInfo& Task)
-			{
-				InOutTooltip.AddTextLine(FString::Printf(TEXT("-------- Task %d%s --------"), Task.Id, Task.bTracked ? TEXT("") : TEXT(" (not tracked)")), FLinearColor::Green);
-
-				if (Task.DebugName != nullptr)
-				{
-					InOutTooltip.AddTextLine(FString::Printf(TEXT("%s"), Task.DebugName), FLinearColor::Green);
-				}
-
-				ENamedThreads::Type ThreadInfo = (ENamedThreads::Type)Task.ThreadToExecuteOn;
-				ENamedThreads::Type ThreadIndex = ENamedThreads::GetThreadIndex(ThreadInfo);
-
-				auto FormatTaskTimestamp = [](double Timestamp) -> FString
-				{
-					return (Timestamp != TraceServices::FTaskInfo::InvalidTimestamp) ? FString::SanitizeFloat(Timestamp) : TEXT("[not set]");
-				};
-
-				auto FormatTaskTime = [](double Time) -> FString
-				{
-					return FormatTimeAuto(Time, 2);
-				};
-
-				auto GetTrackName = [this](uint32 InThreadId) -> FString
-				{
-					TSharedPtr<FCpuTimingTrack> Track = SharedState.GetCpuTrack(InThreadId);
-					return Track.IsValid() ? Track->GetName() : TEXT("Unknown");
-				};
-
-				if (ThreadIndex == ENamedThreads::AnyThread)
-				{
-					const TCHAR* TaskPri = ENamedThreads::GetTaskPriority(ThreadInfo) == ENamedThreads::NormalTaskPriority ? TEXT("Normal") : TEXT("High");
-
-					int32 ThreadPriIndex = ENamedThreads::GetThreadPriorityIndex(ThreadInfo);
-					const TCHAR* ThreadPriStrs[] = { TEXT("Normal"), TEXT("High"), TEXT("Low") };
-					const TCHAR* ThreadPri = ensure(ThreadPriIndex >= 0 && ThreadPriIndex < 3) ? ThreadPriStrs[ThreadPriIndex] : TEXT("Unknown");
-
-					InOutTooltip.AddTextLine(
-						FString::Printf(TEXT("%s Pri task on %s Pri worker (%s)"), TaskPri, ThreadPri, *GetTrackName(Task.StartedThreadId)),
-						FLinearColor::Green);
-				}
-				else
-				{
-					const TCHAR* QueueStr = ENamedThreads::GetQueueIndex(ThreadInfo) == ENamedThreads::MainQueue ? TEXT("Main") : TEXT("Local");
-					InOutTooltip.AddTextLine(
-						FString::Printf(TEXT("%s (%s queue)"), *GetTrackName(Task.StartedThreadId), QueueStr),
-						FLinearColor::Green);
-				}
-
-				InOutTooltip.AddNameValueTextLine(TEXTVIEW("Created:"), FString::Printf(TEXT("%s on %s"),
-					*FormatTaskTimestamp(Task.CreatedTimestamp),
-					*GetTrackName(Task.CreatedThreadId)));
-
-				InOutTooltip.AddNameValueTextLine(TEXTVIEW("Launched:"), FString::Printf(TEXT("%s (+%s) on %s"),
-					*FormatTaskTimestamp(Task.LaunchedTimestamp),
-					*FormatTaskTime(Task.LaunchedTimestamp - Task.CreatedTimestamp),
-					*GetTrackName(Task.LaunchedThreadId)));
-
-				InOutTooltip.AddNameValueTextLine(TEXTVIEW("Scheduled:"), FString::Printf(TEXT("%s (+%s) on %s"),
-					*FormatTaskTimestamp(Task.ScheduledTimestamp),
-					*FormatTaskTime(Task.ScheduledTimestamp - Task.LaunchedTimestamp),
-					*GetTrackName(Task.ScheduledThreadId)));
-
-				InOutTooltip.AddNameValueTextLine(TEXTVIEW("Started:"), FString::Printf(TEXT("%s (+%s)"),
-					*FormatTaskTimestamp(Task.StartedTimestamp),
-					*FormatTaskTime(Task.StartedTimestamp - Task.ScheduledTimestamp)));
-
-				if (Task.FinishedTimestamp != TraceServices::FTaskInfo::InvalidTimestamp)
-				{
-					InOutTooltip.AddNameValueTextLine(TEXTVIEW("Finished:"), FString::Printf(TEXT("%s (+%s)"),
-						*FormatTaskTimestamp(Task.FinishedTimestamp),
-						*FormatTaskTime(Task.FinishedTimestamp - Task.StartedTimestamp)));
-
-					if (Task.CompletedTimestamp != TraceServices::FTaskInfo::InvalidTimestamp)
-					{
-						InOutTooltip.AddNameValueTextLine(TEXTVIEW("Completed:"), FString::Printf(TEXT("%s (+%s) on %s"),
-							*FormatTaskTimestamp(Task.CompletedTimestamp),
-							*FormatTaskTime(Task.CompletedTimestamp - Task.FinishedTimestamp),
-							*GetTrackName(Task.CompletedThreadId)));
-
-						if (Task.DestroyedTimestamp != TraceServices::FTaskInfo::InvalidTimestamp)
-						{
-							InOutTooltip.AddNameValueTextLine(TEXTVIEW("Destroyed:"), FString::Printf(TEXT("%s (+%s) on %s"),
-								*FormatTaskTimestamp(Task.DestroyedTimestamp),
-								*FormatTaskTime(Task.DestroyedTimestamp - Task.CompletedTimestamp),
-								*GetTrackName(Task.DestroyedThreadId)));
-						}
-					}
-				}
-
-				InOutTooltip.AddNameValueTextLine(TEXTVIEW("Prerequisite tasks:"), FString::Printf(TEXT("%d"), Task.Prerequisites.Num()));
-				InOutTooltip.AddNameValueTextLine(TEXTVIEW("Subsequent tasks:"), FString::Printf(TEXT("%d"), Task.Subsequents.Num()));
-				InOutTooltip.AddNameValueTextLine(TEXTVIEW("Parent tasks:"), FString::Printf(TEXT("%d"), Task.ParentTasks.Num()));
-				InOutTooltip.AddNameValueTextLine(TEXTVIEW("Nested tasks:"), FString::Printf(TEXT("%d"), Task.NestedTasks.Num()));
-			};
-
-			// info about a task
-			const TraceServices::FTaskInfo* Task = TasksProvider->TryGetTask(ThreadId, TooltipEvent.GetStartTime());
-			if (Task != nullptr && Task->FinishedTimestamp >= TooltipEvent.GetEndTime())
-			{
-				AddTaskInfo(*Task);
-			}
-
-			// info about blocking
-			const TraceServices::FWaitingForTasks* Waiting = TasksProvider->TryGetWaiting(TimerName, ThreadId, TooltipEvent.GetStartTime());
-			if (Waiting != nullptr && Waiting->Tasks.Num() > 0)
-			{
-				InOutTooltip.AddTextLine(TEXT("-------- Waiting for tasks --------"), FLinearColor::Red);
-				constexpr int32 NumIdsOnRow = 4;
-				TStringBuilder<1024> StringBuilder;
-
-				// Add the first line of Task Id values.
-				for (int32 Index = 0; Index < Waiting->Tasks.Num() && Index < NumIdsOnRow; ++Index)
-				{
-					StringBuilder.Appendf(TEXT("%d, "), Waiting->Tasks[Index]);
-				}
-
-				// Remove separators from last entry.
-				if (Waiting->Tasks.Num() <= NumIdsOnRow)
-				{
-					StringBuilder.RemoveSuffix(2);
-				}
-				InOutTooltip.AddNameValueTextLine(FString::Printf(TEXT("Tasks[%d]:"), Waiting->Tasks.Num()), StringBuilder.ToView());
-				StringBuilder.Reset();
-
-				// Add the rest of the lines with an empty name so they appear as a multi line value.
-				for (int32 Index = NumIdsOnRow; Index < Waiting->Tasks.Num(); ++Index)
-				{
-					StringBuilder.Appendf(TEXT("%d, "), Waiting->Tasks[Index]);
-
-					if ((Index + 1) % NumIdsOnRow == 0)
-					{
-						InOutTooltip.AddNameValueTextLine(TEXT(""), StringBuilder.ToView());
-						StringBuilder.Reset();
-					}
-				}
-				if (StringBuilder.Len() > 1)
-				{
-					StringBuilder.RemoveSuffix(2);
-					InOutTooltip.AddNameValueTextLine(TEXT(""), StringBuilder.ToView());
-				}
-
-				InOutTooltip.AddNameValueTextLine(TEXT("Started waiting:"),
-					FString::Printf(TEXT("%s"), *FString::SanitizeFloat(Waiting->StartedTimestamp)));
-
-				if (Waiting->FinishedTimestamp != TraceServices::FTaskInfo::InvalidTimestamp)
-				{
-					InOutTooltip.AddNameValueTextLine(TEXT("Finished waiting:"),
-						FString::Printf(TEXT("%s (+%s)"),
-							*FString::SanitizeFloat(Waiting->FinishedTimestamp),
-							*FormatTimeAuto(Waiting->FinishedTimestamp - Waiting->StartedTimestamp, 2)));
-				}
-				else
-				{
-					InOutTooltip.AddNameValueTextLine(TEXT("Finished waiting:"), TEXT("[not set]"));
-				}
-
-				const int32 MaxWaitedTasksToList = 5;
-				int32 NumTasksToList = FMath::Min(Waiting->Tasks.Num(), MaxWaitedTasksToList);
-				for (int32 TaskIndex = 0; TaskIndex != NumTasksToList; ++TaskIndex)
-				{
-					const TraceServices::FTaskInfo* WaitedTask = TasksProvider->TryGetTask(Waiting->Tasks[TaskIndex]);
-					if (WaitedTask != nullptr)
-					{
-						AddTaskInfo(*WaitedTask);
-					}
-				}
-				if (NumTasksToList < Waiting->Tasks.Num())
-				{
-					InOutTooltip.AddTextLine(TEXT("[...]"), FLinearColor::Green);
-				}
-			}
-		}
+		PostInitTooltip(InOutTooltip, TooltipEvent, *Session.Get(), TimerName);
 	}
 	else if (ChildTrack.IsValid())
 	{
@@ -922,7 +758,7 @@ void FThreadTimingTrack::InitTooltip(FTooltipDrawState& InOutTooltip, const ITim
 
 ////////////////////////////////////////////////////////////////////////////////////////////////////
 
-void FThreadTimingTrack::GetParentAndRoot(const FThreadTrackEvent& TimingEvent, TSharedPtr<FThreadTrackEvent>& OutParentTimingEvent, TSharedPtr<FThreadTrackEvent>& OutRootTimingEvent) const
+void FThreadTimingTrackImpl::GetParentAndRoot(const FThreadTrackEvent& TimingEvent, TSharedPtr<FThreadTrackEvent>& OutParentTimingEvent, TSharedPtr<FThreadTrackEvent>& OutRootTimingEvent) const
 {
 	if (TimingEvent.GetDepth() > 0)
 	{
@@ -958,7 +794,7 @@ void FThreadTimingTrack::GetParentAndRoot(const FThreadTrackEvent& TimingEvent,
 
 ////////////////////////////////////////////////////////////////////////////////////////////////////
 
-const TSharedPtr<const ITimingEvent> FThreadTimingTrack::GetEvent(float InPosX, float InPosY, const FTimingTrackViewport& Viewport) const
+const TSharedPtr<const ITimingEvent> FThreadTimingTrackImpl::GetEvent(float InPosX, float InPosY, const FTimingTrackViewport& Viewport) const
 {
 	TSharedPtr<FThreadTrackEvent> TimingEvent;
 
@@ -1028,7 +864,7 @@ const TSharedPtr<const ITimingEvent> FThreadTimingTrack::GetEvent(float InPosX,
 
 ////////////////////////////////////////////////////////////////////////////////////////////////////
 
-const TSharedPtr<const ITimingEvent> FThreadTimingTrack::SearchEvent(const FTimingEventSearchParameters& InSearchParameters) const
+const TSharedPtr<const ITimingEvent> FThreadTimingTrackImpl::SearchEvent(const FTimingEventSearchParameters& InSearchParameters) const
 {
 	TSharedPtr<FThreadTrackEvent> FoundEvent;
 	FindTimingProfilerEvent(InSearchParameters, [this, &FoundEvent](double InFoundStartTime, double InFoundEndTime, uint32 InFoundDepth, const TraceServices::FTimingProfilerEvent& InFoundEvent)
@@ -1037,7 +873,7 @@ const TSharedPtr<const ITimingEvent> FThreadTimingTrack::SearchEvent(const FTimi
 		FoundEvent->SetTimerIndex(InFoundEvent.TimerIndex);
 
 		uint32 TimerId = 0;
-		bool ret = FThreadTimingTrack::TimerIndexToTimerId(InFoundEvent.TimerIndex, TimerId);
+		bool ret = FThreadTimingTrackImpl::TimerIndexToTimerId(InFoundEvent.TimerIndex, TimerId);
 		if (ret)
 		{
 			FoundEvent->SetTimerId(TimerId);
@@ -1049,7 +885,7 @@ const TSharedPtr<const ITimingEvent> FThreadTimingTrack::SearchEvent(const FTimi
 
 ////////////////////////////////////////////////////////////////////////////////////////////////////
 
-void FThreadTimingTrack::UpdateEventStats(ITimingEvent& InOutEvent) const
+void FThreadTimingTrackImpl::UpdateEventStats(ITimingEvent& InOutEvent) const
 {
 	if (InOutEvent.CheckTrack(this) && InOutEvent.Is<FThreadTrackEvent>())
 	{
@@ -1086,7 +922,7 @@ void FThreadTimingTrack::UpdateEventStats(ITimingEvent& InOutEvent) const
 
 ////////////////////////////////////////////////////////////////////////////////////////////////////
 
-void FThreadTimingTrack::OnEventSelected(const ITimingEvent& InSelectedEvent) const
+void FThreadTimingTrackImpl::OnEventSelected(const ITimingEvent& InSelectedEvent) const
 {
 	if (InSelectedEvent.CheckTrack(this) && InSelectedEvent.Is<FThreadTrackEvent>())
 	{
@@ -1099,7 +935,7 @@ void FThreadTimingTrack::OnEventSelected(const ITimingEvent& InSelectedEvent) co
 
 ////////////////////////////////////////////////////////////////////////////////////////////////////
 
-void FThreadTimingTrack::OnClipboardCopyEvent(const ITimingEvent& InSelectedEvent) const
+void FThreadTimingTrackImpl::OnClipboardCopyEvent(const ITimingEvent& InSelectedEvent) const
 {
 	if (InSelectedEvent.CheckTrack(this) && InSelectedEvent.Is<FThreadTrackEvent>())
 	{
@@ -1140,7 +976,7 @@ void FThreadTimingTrack::OnClipboardCopyEvent(const ITimingEvent& InSelectedEven
 
 ////////////////////////////////////////////////////////////////////////////////////////////////////
 
-void FThreadTimingTrack::BuildContextMenu(FMenuBuilder& MenuBuilder)
+void FThreadTimingTrackImpl::BuildContextMenu(FMenuBuilder& MenuBuilder)
 {
 	if (GetGroupName() != nullptr)
 	{
@@ -1176,7 +1012,7 @@ void FThreadTimingTrack::BuildContextMenu(FMenuBuilder& MenuBuilder)
 
 ////////////////////////////////////////////////////////////////////////////////////////////////////
 
-bool FThreadTimingTrack::FindTimingProfilerEvent(const FThreadTrackEvent& InTimingEvent, TFunctionRef<void(double, double, uint32, const TraceServices::FTimingProfilerEvent&)> InFoundPredicate) const
+bool FThreadTimingTrackImpl::FindTimingProfilerEvent(const FThreadTrackEvent& InTimingEvent, TFunctionRef<void(double, double, uint32, const TraceServices::FTimingProfilerEvent&)> InFoundPredicate) const
 {
 	auto MatchEvent = [&InTimingEvent](double InStartTime, double InEndTime, uint32 InDepth)
 	{
@@ -1193,7 +1029,7 @@ bool FThreadTimingTrack::FindTimingProfilerEvent(const FThreadTrackEvent& InTimi
 
 ////////////////////////////////////////////////////////////////////////////////////////////////////
 
-bool FThreadTimingTrack::FindTimingProfilerEvent(const FTimingEventSearchParameters& InParameters, TFunctionRef<void(double, double, uint32, const TraceServices::FTimingProfilerEvent&)> InFoundPredicate) const
+bool FThreadTimingTrackImpl::FindTimingProfilerEvent(const FTimingEventSearchParameters& InParameters, TFunctionRef<void(double, double, uint32, const TraceServices::FTimingProfilerEvent&)> InFoundPredicate) const
 {
 	FFilterContext FilterConfiguratorContext;
 	FilterConfiguratorContext.SetReturnValueForUnsetFilters(false);
@@ -1289,7 +1125,7 @@ bool FThreadTimingTrack::FindTimingProfilerEvent(const FTimingEventSearchParamet
 
 ////////////////////////////////////////////////////////////////////////////////////////////////////
 
-void FThreadTimingTrack::CreateFThreadTrackEventFromInfo(const TimelineEventInfo& InEventInfo, const TSharedRef<const FBaseTimingTrack> InTrack, int32 InDepth, TSharedPtr<FThreadTrackEvent> &OutTimingEvent)
+void FThreadTimingTrackImpl::CreateFThreadTrackEventFromInfo(const TimelineEventInfo& InEventInfo, const TSharedRef<const FBaseTimingTrack> InTrack, int32 InDepth, TSharedPtr<FThreadTrackEvent> &OutTimingEvent)
 {
 	OutTimingEvent = MakeShared<FThreadTrackEvent>(InTrack, InEventInfo.StartTime, InEventInfo.EndTime, InDepth);
 	FThreadTrackEvent& Event = OutTimingEvent->As<FThreadTrackEvent>();
@@ -1298,7 +1134,7 @@ void FThreadTimingTrack::CreateFThreadTrackEventFromInfo(const TimelineEventInfo
 	Event.SetTimerIndex(InEventInfo.Event.TimerIndex);
 
 	uint32 TimerId = 0;
-	bool ret = FThreadTimingTrack::TimerIndexToTimerId(InEventInfo.Event.TimerIndex, TimerId);
+	bool ret = FThreadTimingTrackImpl::TimerIndexToTimerId(InEventInfo.Event.TimerIndex, TimerId);
 	if (ret)
 	{
 		Event.SetTimerId(TimerId);
@@ -1308,7 +1144,7 @@ void FThreadTimingTrack::CreateFThreadTrackEventFromInfo(const TimelineEventInfo
 
 ////////////////////////////////////////////////////////////////////////////////////////////////////
 
-bool FThreadTimingTrack::TimerIndexToTimerId(uint32 InTimerIndex, uint32& OutTimerId)
+bool FThreadTimingTrackImpl::TimerIndexToTimerId(uint32 InTimerIndex, uint32& OutTimerId)
 {
 	TSharedPtr<const TraceServices::IAnalysisSession> Session = FInsightsManager::Get()->GetSession();
 	check(Session.IsValid())
@@ -1332,14 +1168,14 @@ bool FThreadTimingTrack::TimerIndexToTimerId(uint32 InTimerIndex, uint32& OutTim
 
 ////////////////////////////////////////////////////////////////////////////////////////////////////
 
-bool FThreadTimingTrack::HasCustomFilter() const
+bool FThreadTimingTrackImpl::HasCustomFilter() const
 {
 	return FilterConfigurator.IsValid() && !FilterConfigurator->IsEmpty();
 }
 
 ////////////////////////////////////////////////////////////////////////////////////////////////////
 
-int32 FThreadTimingTrack::GetDepthAt(double Time) const
+int32 FThreadTimingTrackImpl::GetDepthAt(double Time) const
 {
 	int32 Depth = 0;
 	TSharedPtr<const TraceServices::IAnalysisSession> Session = FInsightsManager::Get()->GetSession();
@@ -1360,7 +1196,7 @@ int32 FThreadTimingTrack::GetDepthAt(double Time) const
 
 ////////////////////////////////////////////////////////////////////////////////////////////////////
 
-void FThreadTimingTrack::SetFilterConfigurator(TSharedPtr<FFilterConfigurator> InFilterConfigurator)
+void FThreadTimingTrackImpl::SetFilterConfigurator(TSharedPtr<FFilterConfigurator> InFilterConfigurator)
 {
 	if (FilterConfigurator != InFilterConfigurator)
 	{
@@ -1371,7 +1207,7 @@ void FThreadTimingTrack::SetFilterConfigurator(TSharedPtr<FFilterConfigurator> I
 
 ////////////////////////////////////////////////////////////////////////////////////////////////////
 
-TSharedPtr<const ITimingEvent> FThreadTimingTrack::FindMaxEventInstance(uint32 InTimerId, double InStartTime, double InEndTime) const
+TSharedPtr<const ITimingEvent> FThreadTimingTrackImpl::FindMaxEventInstance(uint32 InTimerId, double InStartTime, double InEndTime) const
 {
 	TSharedPtr<const TraceServices::IAnalysisSession> Session = FInsightsManager::Get()->GetSession();
 	check(Session.IsValid());
@@ -1454,7 +1290,7 @@ TSharedPtr<const ITimingEvent> FThreadTimingTrack::FindMaxEventInstance(uint32 I
 
 ////////////////////////////////////////////////////////////////////////////////////////////////////
 
-TSharedPtr<const ITimingEvent> FThreadTimingTrack::FindMinEventInstance(uint32 InTimerId, double InStartTime, double InEndTime) const
+TSharedPtr<const ITimingEvent> FThreadTimingTrackImpl::FindMinEventInstance(uint32 InTimerId, double InStartTime, double InEndTime) const
 {
 	TSharedPtr<const TraceServices::IAnalysisSession> Session = FInsightsManager::Get()->GetSession();
 	check(Session.IsValid());
@@ -1534,134 +1370,6 @@ TSharedPtr<const ITimingEvent> FThreadTimingTrack::FindMinEventInstance(uint32 I
 	return TimingEvent;
 }
 
-////////////////////////////////////////////////////////////////////////////////////////////////////
-// FCpuTimingTrack, FGpuTimingTrack, FGpuQueueTimingTrack
-////////////////////////////////////////////////////////////////////////////////////////////////////
-
-INSIGHTS_IMPLEMENT_RTTI(FCpuTimingTrack)
-INSIGHTS_IMPLEMENT_RTTI(FGpuTimingTrack)
-INSIGHTS_IMPLEMENT_RTTI(FGpuQueueTimingTrack)
-
-////////////////////////////////////////////////////////////////////////////////////////////////////
-// FGpuQueueWorkTimingTrack
-////////////////////////////////////////////////////////////////////////////////////////////////////
-
-INSIGHTS_IMPLEMENT_RTTI(FGpuQueueWorkTimingTrack)
-
-////////////////////////////////////////////////////////////////////////////////////////////////////
-
-void FGpuQueueWorkTimingTrack::Draw(const ITimingTrackDrawContext& Context) const
-{
-	DrawLineEvents(Context, 1.0f);
-}
-
-////////////////////////////////////////////////////////////////////////////////////////////////////
-
-void FGpuQueueWorkTimingTrack::DrawLineEvents(const ITimingTrackDrawContext& Context, const float OffsetY) const
-{
-	const FTimingViewDrawHelper& Helper = *static_cast<const FTimingViewDrawHelper*>(&Context.GetHelper());
-
-	if ((Context.GetEventFilter().IsValid() && Context.GetEventFilter()->FilterTrack(*this)) || HasCustomFilter())
-	{
-		Helper.DrawFadedLineEvents(GetDrawState(), *this, OffsetY, 0.1f);
-
-		if (UpdateFilteredDrawStateOpacity())
-		{
-			Helper.DrawLineEvents(GetFilteredDrawState(), *this, OffsetY);
-		}
-		else
-		{
-			Helper.DrawFadedLineEvents(GetFilteredDrawState(), *this, OffsetY, GetFilteredDrawStateOpacity());
-		}
-	}
-	else
-	{
-		Helper.DrawLineEvents(GetDrawState(), *this, OffsetY);
-	}
-}
-
-////////////////////////////////////////////////////////////////////////////////////////////////////
-
-void FGpuQueueWorkTimingTrack::PostDraw(const ITimingTrackDrawContext& Context) const
-{
-	const bool bAreOverlaysVisible = GetSharedState().AreOverlaysVisibleInGpuQueueTracks();
-	const bool bAreExtendedLinesVisible = GetSharedState().AreExtendedLinesVisibleInGpuQueueTracks();
-
-	if (bAreOverlaysVisible || bAreExtendedLinesVisible)
-	{
-		float LineY1 = 0.0f;
-		float LineY2 = 0.0f;
-		ETimingTrackLocation LocalLocation = ETimingTrackLocation::None;
-		TSharedPtr<FBaseTimingTrack> LocalParentTrack = GetParentTrack().Pin();
-		if (LocalParentTrack)
-		{
-			LineY1 = LocalParentTrack->GetPosY();
-			LineY2 = LineY1 + LocalParentTrack->GetHeight();
-			LocalLocation = LocalParentTrack->GetLocation();
-		}
-		else
-		{
-			LineY1 = GetPosY();
-			LineY2 = LineY1 + GetHeight();
-			LocalLocation = GetLocation();
-		}
-
-		const FTimingTrackViewport& Viewport = Context.GetViewport();
-		switch (LocalLocation)
-		{
-			case ETimingTrackLocation::Scrollable:
-			{
-				const float TopY = Viewport.GetPosY() + Viewport.GetTopOffset();
-				if (LineY1 < TopY)
-				{
-					LineY1 = TopY;
-				}
-				const float BottomY = Viewport.GetPosY() + Viewport.GetHeight() - Viewport.GetBottomOffset();
-				if (LineY2 > BottomY)
-				{
-					LineY2 = BottomY;
-				}
-				break;
-			}
-			case ETimingTrackLocation::TopDocked:
-			{
-				const float TopY = Viewport.GetPosY();
-				if (LineY1 < TopY)
-				{
-					LineY1 = TopY;
-				}
-				const float BottomY = Viewport.GetPosY() + Viewport.GetTopOffset();
-				if (LineY2 > BottomY)
-				{
-					LineY2 = BottomY;
-				}
-				break;
-			}
-			case ETimingTrackLocation::BottomDocked:
-			{
-				const float TopY = Viewport.GetPosY() + Viewport.GetHeight() - Viewport.GetBottomOffset();
-				if (LineY1 < TopY)
-				{
-					LineY1 = TopY;
-				}
-				const float BottomY = Viewport.GetPosY() + Viewport.GetHeight();
-				if (LineY2 > BottomY)
-				{
-					LineY2 = BottomY;
-				}
-				break;
-			}
-		}
-
-		const float LineH = LineY2 - LineY1;
-		if (LineH > 0.0f)
-		{
-			const FTimingViewDrawHelper& Helper = *static_cast<const FTimingViewDrawHelper*>(&Context.GetHelper());
-			Helper.DrawContextSwitchMarkers(GetDrawState(), LineY1, LineH, 0.25f, bAreOverlaysVisible, bAreExtendedLinesVisible);
-		}
-	}
-}
-
 ////////////////////////////////////////////////////////////////////////////////////////////////////
 
 } // namespace UE::Insights::TimingProfiler
diff --git a/Engine/Source/Developer/TraceInsights/Private/Insights/TimingProfiler/Tracks/ThreadTimingTrack.h b/Engine/Source/Developer/TraceInsights/Private/Insights/TimingProfiler/Tracks/ThreadTimingTrack.h
deleted file mode 100644
diff --git a/Engine/Source/Developer/TraceInsights/Private/Insights/TimingProfiler/Tracks/ThreadTimingTrackPrivate.h b/Engine/Source/Developer/TraceInsights/Private/Insights/TimingProfiler/Tracks/ThreadTimingTrackPrivate.h
new file mode 100644
--- /dev/null
+++ b/Engine/Source/Developer/TraceInsights/Private/Insights/TimingProfiler/Tracks/ThreadTimingTrackPrivate.h
@@ -0,0 +1,122 @@
+// Copyright Epic Games, Inc. All Rights Reserved.
+
+#pragma once
+
+#include "CoreTypes.h"
+
+// TraceServices
+#include "TraceServices/Model/TimingProfiler.h"
+
+// TraceInsights
+#include "Insights/TimingProfiler/Tracks/ThreadTimingTrack.h"
+#include "Insights/ViewModels/TimingEventSearch.h" // for TTimingEventSearchCache
+#include "Insights/ViewModels/TimingEventsTrack.h"
+
+class FThreadTrackEvent;
+class FTimingEventSearchParameters;
+
+namespace UE::Insights { class FFilterConfigurator; }
+
+namespace UE::Insights::TimingProfiler
+{
+
+class FThreadTimingSharedState;
+
+////////////////////////////////////////////////////////////////////////////////////////////////////
+
+class FThreadTimingTrackImpl : public FThreadTimingTrack
+{
+	INSIGHTS_DECLARE_RTTI(FThreadTimingTrackImpl, FThreadTimingTrack)
+
+public:
+	typedef typename TraceServices::ITimeline<TraceServices::FTimingProfilerEvent>::FTimelineEventInfo TimelineEventInfo;
+
+	explicit FThreadTimingTrackImpl(FThreadTimingSharedState& InSharedState, const FString& InName, const TCHAR* InGroupName, uint32 InTimelineIndex, uint32 InThreadId)
+		: FThreadTimingTrack(InName)
+		, SharedState(InSharedState)
+		, GroupName(InGroupName)
+		, TimelineIndex(InTimelineIndex)
+		, ThreadId(InThreadId)
+	{
+	}
+
+	virtual ~FThreadTimingTrackImpl() override {}
+
+	const TCHAR* GetGroupName() const { return GroupName; };
+
+	uint32 GetTimelineIndex() const { return TimelineIndex; }
+
+	//////////////////////////////////////////////////
+	// FThreadTimingTrack
+
+	virtual uint32 GetThreadId() const override { return ThreadId; }
+
+	virtual int32 GetDepthAt(double Time) const override;
+
+	//////////////////////////////////////////////////
+	// FBaseTimingTrack
+
+	virtual void PostDraw(const ITimingTrackDrawContext& Context) const override;
+
+	virtual void InitTooltip(FTooltipDrawState& InOutTooltip, const ITimingEvent& InTooltipEvent) const override;
+
+	virtual const TSharedPtr<const ITimingEvent> GetEvent(float InPosX, float InPosY, const FTimingTrackViewport& Viewport) const override;
+	virtual const TSharedPtr<const ITimingEvent> SearchEvent(const FTimingEventSearchParameters& InSearchParameters) const override;
+
+	virtual void UpdateEventStats(ITimingEvent& InOutEvent) const override;
+	virtual void OnEventSelected(const ITimingEvent& InSelectedEvent) const override;
+	virtual void OnClipboardCopyEvent(const ITimingEvent& InSelectedEvent) const override;
+	virtual void BuildContextMenu(FMenuBuilder& MenuBuilder) override;
+
+	virtual void SetFilterConfigurator(TSharedPtr<FFilterConfigurator> InFilterConfigurator) override;
+
+	//////////////////////////////////////////////////
+	// FTimingEventsTrack
+
+	virtual void BuildDrawState(ITimingEventsTrackDrawStateBuilder& Builder, const ITimingTrackUpdateContext& Context) override;
+	virtual void BuildFilteredDrawState(ITimingEventsTrackDrawStateBuilder& Builder, const ITimingTrackUpdateContext& Context) override;
+
+	//////////////////////////////////////////////////
+
+	TSharedPtr<const ITimingEvent> FindMaxEventInstance(uint32 TimerId, double StartTime, double EndTime) const;
+	TSharedPtr<const ITimingEvent> FindMinEventInstance(uint32 TimerId, double StartTime, double EndTime) const;
+
+protected:
+	FThreadTimingSharedState& GetSharedState() const { return SharedState; }
+
+	//////////////////////////////////////////////////
+	// FTimingEventsTrack
+
+	virtual bool HasCustomFilter() const override;
+
+	//////////////////////////////////////////////////
+
+	virtual void PostInitTooltip(FTooltipDrawState& InOutTooltip, const FThreadTrackEvent& TooltipEvent, const TraceServices::IAnalysisSession& Session, const TCHAR* TimerName) const {}
+
+private:
+	bool FindTimingProfilerEvent(const FThreadTrackEvent& InTimingEvent, TFunctionRef<void(double, double, uint32, const TraceServices::FTimingProfilerEvent&)> InFoundPredicate) const;
+	bool FindTimingProfilerEvent(const FTimingEventSearchParameters& InParameters, TFunctionRef<void(double, double, uint32, const TraceServices::FTimingProfilerEvent&)> InFoundPredicate) const;
+
+	void GetParentAndRoot(const FThreadTrackEvent& TimingEvent,
+						  TSharedPtr<FThreadTrackEvent>& OutParentTimingEvent,
+						  TSharedPtr<FThreadTrackEvent>& OutRootTimingEvent) const;
+
+	static void CreateFThreadTrackEventFromInfo(const TimelineEventInfo& InEventInfo, const TSharedRef<const FBaseTimingTrack> InTrack, int32 InDepth, TSharedPtr<FThreadTrackEvent> &OutTimingEvent);
+	static bool TimerIndexToTimerId(uint32 InTimerIndex, uint32 & OutTimerId);
+
+private:
+	FThreadTimingSharedState& SharedState;
+
+	TSharedPtr<FFilterConfigurator> FilterConfigurator;
+
+	const TCHAR* GroupName = nullptr;
+	uint32 TimelineIndex = 0;
+	uint32 ThreadId = 0;
+
+	// Search cache
+	mutable TTimingEventSearchCache<TraceServices::FTimingProfilerEvent> SearchCache;
+};
+
+////////////////////////////////////////////////////////////////////////////////////////////////////
+
+} // namespace UE::Insights::TimingProfiler
diff --git a/Engine/Source/Developer/TraceInsights/Private/Insights/TimingProfiler/ViewModels/ThreadTimingSharedState.cpp b/Engine/Source/Developer/TraceInsights/Private/Insights/TimingProfiler/ViewModels/ThreadTimingSharedState.cpp
--- a/Engine/Source/Developer/TraceInsights/Private/Insights/TimingProfiler/ViewModels/ThreadTimingSharedState.cpp
+++ b/Engine/Source/Developer/TraceInsights/Private/Insights/TimingProfiler/ViewModels/ThreadTimingSharedState.cpp
@@ -21,7 +21,8 @@
 #include "Insights/InsightsStyle.h"
 #include "Insights/IUnrealInsightsModule.h"
 #include "Insights/ITimingViewSession.h"
-#include "Insights/TimingProfiler/Tracks/ThreadTimingTrack.h"
+#include "Insights/TimingProfiler/Tracks/CpuTimingTrack.h"
+#include "Insights/TimingProfiler/Tracks/GpuTimingTrack.h"
 #include "Insights/Widgets/STimingView.h"
 
 #define LOCTEXT_NAMESPACE "UE::Insights::TimingProfiler::ThreadTiming"
diff --git a/Engine/Source/Developer/TraceInsights/Private/Insights/TimingProfiler/ViewModels/TimingExporter.cpp b/Engine/Source/Developer/TraceInsights/Private/Insights/TimingProfiler/ViewModels/TimingExporter.cpp
--- a/Engine/Source/Developer/TraceInsights/Private/Insights/TimingProfiler/ViewModels/TimingExporter.cpp
+++ b/Engine/Source/Developer/TraceInsights/Private/Insights/TimingProfiler/ViewModels/TimingExporter.cpp
@@ -20,7 +20,7 @@
 // TraceInsights
 #include "Insights/Log.h"
 #include "Insights/TimingProfiler/TimingProfilerManager.h"
-#include "Insights/TimingProfiler/Tracks/ThreadTimingTrack.h"
+#include "Insights/TimingProfiler/Tracks/GpuTimingTrack.h"
 
 #define LOCTEXT_NAMESPACE "UE::Insights::TimingProfiler::FTimingExporter"
 
diff --git a/Engine/Source/Developer/TraceInsights/Private/Insights/TimingProfiler/Widgets/STimersView.cpp b/Engine/Source/Developer/TraceInsights/Private/Insights/TimingProfiler/Widgets/STimersView.cpp
--- a/Engine/Source/Developer/TraceInsights/Private/Insights/TimingProfiler/Widgets/STimersView.cpp
+++ b/Engine/Source/Developer/TraceInsights/Private/Insights/TimingProfiler/Widgets/STimersView.cpp
@@ -39,7 +39,7 @@
 #include "Insights/InsightsStyle.h"
 #include "Insights/Log.h"
 #include "Insights/TimingProfiler/TimingProfilerManager.h"
-#include "Insights/TimingProfiler/Tracks/ThreadTimingTrack.h"
+#include "Insights/TimingProfiler/Tracks/GpuTimingTrack.h"
 #include "Insights/TimingProfiler/ViewModels/TimerAggregator.h"
 #include "Insights/TimingProfiler/ViewModels/TimerNodeHelper.h"
 #include "Insights/TimingProfiler/ViewModels/TimersViewColumnFactory.h"
diff --git a/Engine/Source/Developer/TraceInsights/Public/Insights/TimingProfiler/Tracks/ThreadTimingTrack.h b/Engine/Source/Developer/TraceInsights/Public/Insights/TimingProfiler/Tracks/ThreadTimingTrack.h
new file mode 100644
--- /dev/null
+++ b/Engine/Source/Developer/TraceInsights/Public/Insights/TimingProfiler/Tracks/ThreadTimingTrack.h
@@ -0,0 +1,26 @@
+// Copyright Epic Games, Inc. All Rights Reserved.
+
+#pragma once
+
+#include "CoreTypes.h"
+
+// TraceInsights
+#include "Insights/ViewModels/TimingEventsTrack.h"
+
+namespace UE::Insights::TimingProfiler
+{
+
+class FThreadTimingTrack : public FTimingEventsTrack
+{
+	INSIGHTS_DECLARE_RTTI(FThreadTimingTrack, FTimingEventsTrack)
+
+public:
+	explicit FThreadTimingTrack() {}
+	explicit FThreadTimingTrack(const FString& InTrackName) : FTimingEventsTrack(InTrackName) {}
+	virtual ~FThreadTimingTrack() {}
+
+	virtual uint32 GetThreadId() const = 0;
+	virtual int32 GetDepthAt(double Time) const = 0;
+};
+
+} // namespace UE::Insights::TimingProfiler
