UnrealInsights | Add support for nesting child tracks

Makes it so that nested child tracks (i.e. child tracks in child tracks) render correctly. Previously height and Y position calculations did not happen recursively, which caused them to render in the wrong spot in the timeline.

Also changes the spacing between child tracks in the compact layout from 0 to 1 Slate unit. This is to help establish a more sensible visual hierarchy between the tracks; this does make child tracks have the same spacing as lanes on a single track, but increasing the value further would make contextswitch timing tracks look out of place.

diff --git a/Engine/Source/Developer/TraceInsights/Private/Insights/ViewModels/TimingEventsTrack.cpp b/Engine/Source/Developer/TraceInsights/Private/Insights/ViewModels/TimingEventsTrack.cpp
--- a/Engine/Source/Developer/TraceInsights/Private/Insights/ViewModels/TimingEventsTrack.cpp
+++ b/Engine/Source/Developer/TraceInsights/Private/Insights/ViewModels/TimingEventsTrack.cpp
@@ -175,14 +175,15 @@ void FTimingEventsTrack::UpdateTrackHeight(const ITimingTrackUpdateContext& Cont
 	else
 	{
 		DesiredTrackHeight = Viewport.GetLayout().ComputeTrackHeight(NumLanes);
-		if (ChildTrack.IsValid() && ChildTrack->GetHeight() > 0.0f)
+	}
+
+	if (ChildTrack.IsValid() && ChildTrack->GetHeight() > 0.0f)
+	{
+		if (DesiredTrackHeight == 0.0f)
 		{
-			if (DesiredTrackHeight == 0.0f)
-			{
-				DesiredTrackHeight += 1.0f + 2.0f * Viewport.GetLayout().TimelineDY;
-			}
-			DesiredTrackHeight += ChildTrack->GetHeight() + Viewport.GetLayout().ChildTimelineDY;
+			DesiredTrackHeight += 1.0f + 2.0f * Viewport.GetLayout().TimelineDY;
 		}
+		DesiredTrackHeight += ChildTrack->GetHeight() + Viewport.GetLayout().ChildTimelineDY;
 	}
 
 	if (CurrentTrackHeight < DesiredTrackHeight)
@@ -351,14 +352,7 @@ void FTimingEventsTrack::DrawEvent(const ITimingTrackDrawContext& Context, const
 		{
 			Y += ChildTrack->GetHeight() + Layout.ChildTimelineDY;
 		}
-		if (IsChildTrack())
-		{
-			Y += Layout.GetChildLaneY(TrackEvent.GetDepth());
-		}
-		else
-		{
-			Y += Layout.GetLaneY(TrackEvent.GetDepth());
-		}
+		Y += Layout.GetLaneY(TrackEvent.GetDepth());
 
 		const FTimingViewDrawHelper& Helper = *static_cast<const FTimingViewDrawHelper*>(&Context.GetHelper());
 		Helper.DrawTimingEventHighlight(TrackEvent.GetStartTime(), TrackEvent.GetEndTime(), Y, InDrawMode);
@@ -371,30 +365,17 @@ const TSharedPtr<const ITimingEvent> FTimingEventsTrack::GetEvent(float InPosX,
 {
 	const FTimingViewLayout& Layout = Viewport.GetLayout();
 
-	float TopLaneY = 0.0f;
-	float TrackLanesHeight = 0.0f;
+	float TopLaneY = GetPosY() + 1.0f + Layout.TimelineDY;
+	float TrackLanesHeight = GetHeight();
 	if (ChildTrack.IsValid() && ChildTrack->GetHeight() > 0.0f)
 	{
 		const float HeaderDY = InPosY - ChildTrack->GetPosY();
-		if (HeaderDY >= 0 && HeaderDY < ChildTrack->GetHeight())
+		if (HeaderDY >= 0.0f && HeaderDY < (ChildTrack->GetHeight() + Layout.ChildTimelineDY))
 		{
 			return ChildTrack->GetEvent(InPosX, InPosY, Viewport);
 		}
-		TopLaneY = GetPosY() + 1.0f + Layout.TimelineDY + ChildTrack->GetHeight() + Layout.ChildTimelineDY;
-		TrackLanesHeight = GetHeight() - ChildTrack->GetHeight() - 1.0f - 2 * Layout.TimelineDY - Layout.ChildTimelineDY;
-	}
-	else
-	{
-		if (IsChildTrack())
-		{
-			TopLaneY = GetPosY();
-			TrackLanesHeight = GetHeight();
-		}
-		else
-		{
-			TopLaneY = GetPosY() + 1.0f + Layout.TimelineDY;
-			TrackLanesHeight = GetHeight() - 1.0f - 2 * Layout.TimelineDY;
-		}
+		TopLaneY += ChildTrack->GetHeight() + Layout.ChildTimelineDY;
+		TrackLanesHeight -= ChildTrack->GetHeight();
 	}
 
 	const float DY = InPosY - TopLaneY;
diff --git a/Engine/Source/Developer/TraceInsights/Private/Insights/ViewModels/TimingViewDrawHelper.cpp b/Engine/Source/Developer/TraceInsights/Private/Insights/ViewModels/TimingViewDrawHelper.cpp
--- a/Engine/Source/Developer/TraceInsights/Private/Insights/ViewModels/TimingViewDrawHelper.cpp
+++ b/Engine/Source/Developer/TraceInsights/Private/Insights/ViewModels/TimingViewDrawHelper.cpp
@@ -327,10 +327,7 @@ void FTimingViewDrawHelper::DrawEvents(const FTimingEventsTrackDrawState& DrawSt
 		NumMergedBoxes += DrawState.GetNumMergedBoxes();
 
 		float TopLaneY = Track.GetPosY();
-		if (!Track.IsChildTrack())
-		{
-			TopLaneY += OffsetY + Layout.TimelineDY;
-		}
+		TopLaneY += OffsetY + Layout.TimelineDY;
 		if (Track.GetChildTrack().IsValid() && Track.ChildTrack->GetHeight() > 0.0f)
 		{
 			TopLaneY += Track.ChildTrack->GetHeight() + Layout.ChildTimelineDY;
diff --git a/Engine/Source/Developer/TraceInsights/Private/Insights/Widgets/STimingView.cpp b/Engine/Source/Developer/TraceInsights/Private/Insights/Widgets/STimingView.cpp
--- a/Engine/Source/Developer/TraceInsights/Private/Insights/Widgets/STimingView.cpp
+++ b/Engine/Source/Developer/TraceInsights/Private/Insights/Widgets/STimingView.cpp
@@ -856,7 +856,7 @@ PRAGMA_ENABLE_DEPRECATION_WARNINGS
 	float TopOffset = 0.0f;
 	for (TSharedPtr<FBaseTimingTrack>& TrackPtr : TopDockedTracks)
 	{
-		SetTrackPosY(TrackPtr, Viewport.GetPosY() + TopOffset);
+		SetTrackPosY(*TrackPtr, Viewport.GetPosY() + TopOffset);
 		if (TrackPtr->IsVisible())
 		{
 			TopOffset += TrackPtr->GetHeight();
@@ -898,7 +898,7 @@ PRAGMA_ENABLE_DEPRECATION_WARNINGS
 		float BottomDockedTrackPosY = Viewport.GetPosY() + Viewport.GetHeight() - BottomOffset + 2.0f;
 		for (TSharedPtr<FBaseTimingTrack>& TrackPtr : BottomDockedTracks)
 		{
-			SetTrackPosY(TrackPtr, BottomDockedTrackPosY);
+			SetTrackPosY(*TrackPtr, BottomDockedTrackPosY);
 			if (TrackPtr->IsVisible())
 			{
 				BottomDockedTrackPosY += TrackPtr->GetHeight();
@@ -1190,7 +1190,7 @@ void STimingView::UpdatePositionForScrollableTracks()
 	float ScrollableTrackPosY = Viewport.GetPosY() + Viewport.GetTopOffset() - Viewport.GetScrollPosY();
 	for (TSharedPtr<FBaseTimingTrack>& TrackPtr : ScrollableTracks)
 	{
-		SetTrackPosY(TrackPtr, ScrollableTrackPosY);
+		SetTrackPosY(*TrackPtr, ScrollableTrackPosY);
 		if (TrackPtr->IsVisible())
 		{
 			ScrollableTrackPosY += TrackPtr->GetHeight();
@@ -1200,12 +1200,12 @@ void STimingView::UpdatePositionForScrollableTracks()
 
 ////////////////////////////////////////////////////////////////////////////////////////////////////
 
-void STimingView::SetTrackPosY(TSharedPtr<FBaseTimingTrack>& TrackPtr, float TrackPosY) const
+void STimingView::SetTrackPosY(FBaseTimingTrack& Track, float TrackPosY) const
 {
-	TrackPtr->SetPosY(TrackPosY);
-	if (TrackPtr->GetChildTrack().IsValid())
+	Track.SetPosY(TrackPosY);
+	if (Track.GetChildTrack().IsValid())
 	{
-		TrackPtr->GetChildTrack()->SetPosY(TrackPosY + this->GetViewport().GetLayout().TimelineDY + 1.0f);
+		SetTrackPosY(*Track.GetChildTrack(), TrackPosY);
 	}
 }
 
@@ -3865,6 +3865,22 @@ const TSharedPtr<FBaseTimingTrack> STimingView::GetTrackAt(float InPosX, float I
 
 ////////////////////////////////////////////////////////////////////////////////////////////////////
 
+static bool EventBelongsToTrack(const ITimingEvent& Event, const FBaseTimingTrack* Track)
+{
+	const FBaseTimingTrack* Current = Track;
+	while (Current != nullptr)
+	{
+		if (&Event.GetTrack().Get() == Current)
+		{
+			return true;
+		}
+		Current = Current->GetChildTrack().Get();
+	}
+	return false;
+}
+
+////////////////////////////////////////////////////////////////////////////////////////////////////
+
 void STimingView::UpdateHoveredTimingEvent(float InMousePosX, float InMousePosY)
 {
 	TSharedPtr<FBaseTimingTrack> NewHoveredTrack = GetTrackAt(InMousePosX, InMousePosY);
@@ -3898,7 +3914,7 @@ void STimingView::UpdateHoveredTimingEvent(float InMousePosX, float InMousePosY)
 			Stopwatch.Start();
 
 			HoveredEvent = NewHoveredEvent;
-			ensure(HoveredTrack == HoveredEvent->GetTrack() || HoveredTrack->GetChildTrack() == HoveredEvent->GetTrack());
+			ensure(EventBelongsToTrack(*HoveredEvent, HoveredTrack.Get()));
 			HoveredTrack->UpdateEventStats(const_cast<ITimingEvent&>(*HoveredEvent));
 
 			Stopwatch.Update();
diff --git a/Engine/Source/Developer/TraceInsights/Private/Insights/Widgets/STimingView.h b/Engine/Source/Developer/TraceInsights/Private/Insights/Widgets/STimingView.h
--- a/Engine/Source/Developer/TraceInsights/Private/Insights/Widgets/STimingView.h
+++ b/Engine/Source/Developer/TraceInsights/Private/Insights/Widgets/STimingView.h
@@ -547,7 +547,7 @@ PRAGMA_ENABLE_DEPRECATION_WARNINGS
 	FReply AllowTracksToProcessOnMouseButtonUp(const FGeometry& MyGeometry, const FPointerEvent& MouseEvent);
 	FReply AllowTracksToProcessOnMouseButtonDoubleClick(const FGeometry& MyGeometry, const FPointerEvent& MouseEvent);
 
-	void SetTrackPosY(TSharedPtr<FBaseTimingTrack>& TrackPtr, float TrackPosY) const;
+	void SetTrackPosY(FBaseTimingTrack& Track, float TrackPosY) const;
 
 	void FindFirstEvent();
 	void FindPrevEvent();
diff --git a/Engine/Source/Developer/TraceInsights/Public/Insights/ViewModels/TimingViewLayout.h b/Engine/Source/Developer/TraceInsights/Public/Insights/ViewModels/TimingViewLayout.h
--- a/Engine/Source/Developer/TraceInsights/Public/Insights/ViewModels/TimingViewLayout.h
+++ b/Engine/Source/Developer/TraceInsights/Public/Insights/ViewModels/TimingViewLayout.h
@@ -20,7 +20,7 @@ struct TRACEINSIGHTS_API FTimingViewLayout
 	static constexpr float CompactLayoutEventH = 4.0f;
 	static constexpr float CompactLayoutEventDY = 1.0f;
 	static constexpr float CompactLayoutTimelineDY = 1.0f;
-	static constexpr float CompactLayoutChildTimelineDY = 0.0f;
+	static constexpr float CompactLayoutChildTimelineDY = 1.0f;
 	static constexpr float CompactLayoutMinTimelineH = 0.0f;
 
 	//////////////////////////////////////////////////
