[Insights] Fix crash when using the Metadata filter

diff --git a/Engine/Source/Developer/TraceInsights/Private/Insights/TimingProfiler/ViewModels/TimerFilters.cpp b/Engine/Source/Developer/TraceInsights/Private/Insights/TimingProfiler/ViewModels/TimerFilters.cpp
--- a/Engine/Source/Developer/TraceInsights/Private/Insights/TimingProfiler/ViewModels/TimerFilters.cpp
+++ b/Engine/Source/Developer/TraceInsights/Private/Insights/TimingProfiler/ViewModels/TimerFilters.cpp
@@ -306,7 +306,10 @@ bool FMetadataFilterState::ApplyFilter(const FFilterContext& Context) const
 	int64 TimerIndex;
 	Context.GetFilterData<int64>(Filter->GetKey(), TimerIndex);
 
+	TSharedPtr<const TraceServices::IAnalysisSession> Session = FInsightsManager::Get()->GetSession();
+	TraceServices::FAnalysisSessionReadScope SessionReadScope(*Session.Get());
 	TArrayView<const uint8> Metadata = TimerReader->GetMetadata(static_cast<uint32>(TimerIndex));
+
 	if (Metadata.Num() > 0)
 	{
 		if (!bShowAllMetadataEvents)
