[Insights] Minor fix re FRandomGraphTrack test.

#rnx

diff --git a/Engine/Source/Developer/TraceInsights/Private/Insights/ViewModels/GraphTrack.cpp b/Engine/Source/Developer/TraceInsights/Private/Insights/ViewModels/GraphTrack.cpp
--- a/Engine/Source/Developer/TraceInsights/Private/Insights/ViewModels/GraphTrack.cpp
+++ b/Engine/Source/Developer/TraceInsights/Private/Insights/ViewModels/GraphTrack.cpp
@@ -928,7 +928,7 @@ void FRandomGraphTrack::Update(const ITimingTrackUpdateContext& Context)
 {
 	FGraphTrack::Update(Context);
 
-	if (IsDirty())
+	if (IsDirty() || Context.GetViewport().IsHorizontalViewportDirty())
 	{
 		ClearDirtyFlag();
 
