[Insights]
- Added simple RTTI to FGraphSeries class hierarchy.
- Split FMemoryGraphSeries class into separate classes (FMemTagGraphSeries and FAllocationsGraphSeries).

#rnx

diff --git a/Engine/Source/Developer/TraceInsights/Private/Insights/MemoryProfiler/ViewModels/MemoryGraphTrack.cpp b/Engine/Source/Developer/TraceInsights/Private/Insights/MemoryProfiler/ViewModels/MemoryGraphTrack.cpp
--- a/Engine/Source/Developer/TraceInsights/Private/Insights/MemoryProfiler/ViewModels/MemoryGraphTrack.cpp
+++ b/Engine/Source/Developer/TraceInsights/Private/Insights/MemoryProfiler/ViewModels/MemoryGraphTrack.cpp
@@ -37,8 +37,12 @@ namespace UE::Insights::MemoryProfiler
 // FMemoryGraphSeries
 ////////////////////////////////////////////////////////////////////////////////////////////////////
 
+INSIGHTS_IMPLEMENT_RTTI(FMemoryGraphSeries)
+
 constexpr int32 DefaultDecimalDigitCount = 2;
 
+////////////////////////////////////////////////////////////////////////////////////////////////////
+
 FString FMemoryGraphSeries::FormatValue(double Value) const
 {
 	const int64 ValueInt64 = static_cast<int64>(Value);
@@ -47,8 +51,215 @@ FString FMemoryGraphSeries::FormatValue(double Value) const
 		return TEXT("0");
 	}
 
-	if (TimelineType <= ETimelineType::MaxTotalMem)
+	return FText::AsNumber(ValueInt64).ToString();
+}
+
+////////////////////////////////////////////////////////////////////////////////////////////////////
+// FMemTagGraphSeries
+////////////////////////////////////////////////////////////////////////////////////////////////////
+
+INSIGHTS_IMPLEMENT_RTTI(FMemTagGraphSeries)
+
+////////////////////////////////////////////////////////////////////////////////////////////////////
+
+FMemTagGraphSeries::FMemTagGraphSeries(FMemoryTrackerId InTrackerId, FMemoryTagSetId InTagSetId, FMemoryTagId InTagId)
+	: TrackerId(InTrackerId)
+	, TagSetId(InTagSetId)
+	, TagId(InTagId)
+{
+	SetName(TEXT("LLM Tag"));
+	SetDescription(TEXT("Low Level Memory Tag"));
+}
+
+////////////////////////////////////////////////////////////////////////////////////////////////////
+
+FMemTagGraphSeries::~FMemTagGraphSeries()
+{
+}
+
+////////////////////////////////////////////////////////////////////////////////////////////////////
+
+FString FMemTagGraphSeries::FormatValue(double Value) const
+{
+	const int64 ValueInt64 = static_cast<int64>(Value);
+	if (ValueInt64 == 0)
+	{
+		return TEXT("0");
+	}
+
+	double UnitValue;
+	const TCHAR* UnitText;
+	FMemoryGraphTrack::GetUnit(EGraphTrackLabelUnit::Auto, FMath::Abs(Value), UnitValue, UnitText);
+
+	if (ValueInt64 < 0)
 	{
+		FString Auto = FMemoryGraphTrack::FormatValue(-Value, UnitValue, UnitText, DefaultDecimalDigitCount);
+		return FString::Printf(TEXT("-%s (%s bytes)"), *Auto, *FText::AsNumber(ValueInt64).ToString());
+	}
+	else
+	{
+		FString Auto = FMemoryGraphTrack::FormatValue(Value, UnitValue, UnitText, DefaultDecimalDigitCount);
+		return FString::Printf(TEXT("%s (%s bytes)"), *Auto, *FText::AsNumber(ValueInt64).ToString());
+	}
+}
+
+////////////////////////////////////////////////////////////////////////////////////////////////////
+// FAllocationsGraphSeries
+////////////////////////////////////////////////////////////////////////////////////////////////////
+
+INSIGHTS_IMPLEMENT_RTTI(FAllocationsGraphSeries)
+
+////////////////////////////////////////////////////////////////////////////////////////////////////
+
+FAllocationsGraphSeries::FAllocationsGraphSeries(ETimeline InTimeline)
+	: Timeline(InTimeline)
+{
+	Initialize();
+}
+
+////////////////////////////////////////////////////////////////////////////////////////////////////
+
+FAllocationsGraphSeries::~FAllocationsGraphSeries()
+{
+}
+
+////////////////////////////////////////////////////////////////////////////////////////////////////
+
+void FAllocationsGraphSeries::Initialize()
+{
+	switch (Timeline)
+	{
+		case ETimeline::MinTotalMem:
+		{
+			ValueType = EValueType::IntegerBytes;
+			SetName(LOCTEXT("TotalAllocatedMemoryMin_TrackName", "Total Allocated Memory (Min)"));
+			SetDescription(LOCTEXT("TotalAllocatedMemoryMin_TrackDesc", "Minimum value per sample for the Total Allocated Memory"));
+			SetColor(FLinearColor(0.0f, 0.5f, 1.0f, 1.0f));
+			break;
+		}
+
+		case ETimeline::MaxTotalMem:
+		{
+			ValueType = EValueType::IntegerBytes;
+			SetName(LOCTEXT("TotalAllocatedMemoryMax_TrackName", "Total Allocated Memory (Max)"));
+			SetDescription(LOCTEXT("TotalAllocatedMemoryMax_TrackDesc", "Maximum value per sample for the Total Allocated Memory"));
+			SetColor(FLinearColor(1.0f, 0.25f, 1.0f, 1.0f));
+			break;
+		}
+
+		case ETimeline::MinLiveAllocs:
+		{
+			ValueType = EValueType::IntegerCounter;
+			SetName(LOCTEXT("LiveAllocationCountMin_TrackName", "Live Allocation Count (Min)"));
+			SetDescription(LOCTEXT("LiveAllocationCountMin_TrackDesc", "Minimum value per sample for the Live Allocation Count"));
+			SetColor(FLinearColor(1.0f, 1.0f, 0.25f, 1.0f));
+			break;
+		}
+
+		case ETimeline::MaxLiveAllocs:
+		{
+			ValueType = EValueType::IntegerCounter;
+			SetName(LOCTEXT("LiveAllocationCountMax_TrackDesc", "Live Allocation Count (Max)"));
+			SetDescription(LOCTEXT("LiveAllocationCountMax_TrackDesc", "Maximum value per sample for the Live Allocation Count"));
+			SetColor(FLinearColor(1.0f, 0.25f, 1.0f, 1.0f));
+			break;
+		}
+
+		case ETimeline::MinSwapMem:
+		{
+			ValueType = EValueType::IntegerBytes;
+			SetName(LOCTEXT("TotalSwapMemoryMin_TrackName", "Total Swap Memory (Min)"));
+			SetDescription(LOCTEXT("TotalSwapMemoryMin_TrackDesc", "Minimum value per sample for the Total Swap Memory"));
+			SetColor(FLinearColor(0.0f, 0.5f, 1.0f, 1.0f));
+			break;
+		}
+
+		case ETimeline::MaxSwapMem:
+		{
+			ValueType = EValueType::IntegerBytes;
+			SetName(LOCTEXT("TotalSwapMemoryMax_TrackName", "Total Swap Memory (Max)"));
+			SetDescription(LOCTEXT("TotalSwapMemoryMax_TrackDesc", "Maximum value per sample for the Total Swap Memory"));
+			SetColor(FLinearColor(1.0f, 0.25f, 1.0f, 1.0f));
+			break;
+		}
+
+		case ETimeline::MinCompressedSwapMem:
+		{
+			ValueType = EValueType::IntegerBytes;
+			SetName(LOCTEXT("TotalCompressedSwapMemoryMin_TrackName", "Total Compressed Swap Memory (Min)"));
+			SetDescription(LOCTEXT("TotalCompressedSwapMemoryMin_TrackDesc", "Minimum value per sample for the Total Compressed Swap Memory"));
+			SetColor(FLinearColor(1.0f, 1.0f, 0.25f, 1.0f));
+			break;
+		}
+
+		case ETimeline::MaxCompressedSwapMem:
+		{
+			ValueType = EValueType::IntegerBytes;
+			SetName(LOCTEXT("TotalCompressedSwapMemoryMax_TrackName", "Total Compressed Swap Memory (Max)"));
+			SetDescription(LOCTEXT("TotalCompressedSwapMemoryMax_TrackDesc", "Maximum value per sample for the Total Compressed Swap Memory"));
+			SetColor(FLinearColor(1.0f, 0.25f, 1.0f, 1.0f));
+			break;
+		}
+
+		case ETimeline::AllocEvents:
+		{
+			ValueType = EValueType::IntegerCounter;
+			SetName(LOCTEXT("AllocEventCount_TrackName", "Alloc Event Count"));
+			SetDescription(LOCTEXT("AllocEventCount_TrackDesc", "Number of alloc events per sample"));
+			SetColor(FLinearColor(0.0f, 1.0f, 0.5f, 1.0f));
+			break;
+		}
+
+		case ETimeline::FreeEvents:
+		{
+			ValueType = EValueType::IntegerCounter;
+			SetName(LOCTEXT("FreeEventCount_TrackName", "Free Event Count"));
+			SetDescription(LOCTEXT("FreeEventCount_TrackDesc", "Number of free events per sample"));
+			SetColor(FLinearColor(1.0f, 0.5f, 0.25f, 1.0f));
+			break;
+		}
+
+		case ETimeline::PageInEvents:
+		{
+			ValueType = EValueType::IntegerCounter;
+			SetName(LOCTEXT("PageInEventCount_TrackName", "Page In Event Count"));
+			SetDescription(LOCTEXT("PageInEventCount_TrackDesc", "Number of page in events per sample"));
+			SetColor(FLinearColor(0.0f, 1.0f, 0.5f, 1.0f));
+			break;
+		}
+
+		case ETimeline::PageOutEvents:
+		{
+			ValueType = EValueType::IntegerCounter;
+			SetName(LOCTEXT("PageOutEventCount_TrackName", "Page Out Event Count"));
+			SetDescription(LOCTEXT("PageOutEventCount_TrackDesc", "Number of page out events per sample"));
+			SetColor(FLinearColor(1.0f, 0.5f, 0.25f, 1.0f));
+			break;
+		}
+
+		case ETimeline::SwapFreeEvents:
+		{
+			ValueType = EValueType::IntegerCounter;
+			SetName(LOCTEXT("SwapFreeEventCount_TrackName", "Swap Free Event Count"));
+			SetDescription(LOCTEXT("SwapFreeEventCount_TrackDesc", "Number of swap free events per sample"));
+			SetColor(FLinearColor(0.25f, 0.5f, 1.0f, 1.0f));
+			break;
+		}
+	}
+}
+
+////////////////////////////////////////////////////////////////////////////////////////////////////
+
+FString FAllocationsGraphSeries::FormatValue(double Value) const
+{
+	if (ValueType == EValueType::IntegerBytes)
+	{
+		const int64 ValueInt64 = static_cast<int64>(Value);
+		if (ValueInt64 == 0)
+		{
+			return TEXT("0");
+		}
+
 		double UnitValue;
 		const TCHAR* UnitText;
 		FMemoryGraphTrack::GetUnit(EGraphTrackLabelUnit::Auto, FMath::Abs(Value), UnitValue, UnitText);
@@ -64,10 +275,20 @@ FString FMemoryGraphSeries::FormatValue(double Value) const
 			return FString::Printf(TEXT("%s (%s bytes)"), *Auto, *FText::AsNumber(ValueInt64).ToString());
 		}
 	}
-	else
+	else if (ValueType == EValueType::IntegerCounter)
 	{
+		const int64 ValueInt64 = static_cast<int64>(Value);
+		if (ValueInt64 == 0)
+		{
+			return TEXT("0");
+		}
+
 		return FText::AsNumber(ValueInt64).ToString();
 	}
+	else
+	{
+		return FText::AsNumber(Value).ToString();
+	}
 }
 
 ////////////////////////////////////////////////////////////////////////////////////////////////////
@@ -158,25 +379,18 @@ void FMemoryGraphTrack::Update(const ITimingTrackUpdateContext& Context)
 
 		for (TSharedPtr<FGraphSeries>& Series : AllSeries)
 		{
-			if (Series->IsVisible())
+			if (Series->IsVisible() &&
+				Series->Is<FMemoryGraphSeries>())
 			{
-				//TODO: if (Series->Is<FMemoryGraphSeries>())
-				TSharedPtr<FMemoryGraphSeries> MemorySeries = StaticCastSharedPtr<FMemoryGraphSeries>(Series);
+				FMemoryGraphSeries& MemorySeries = Series->As<FMemoryGraphSeries>();
 
-				if (bIsEntireGraphTrackDirty || Series->IsDirty())
+				if (bIsEntireGraphTrackDirty || MemorySeries.IsDirty())
 				{
-					if (MemorySeries->GetTimelineType() == FMemoryGraphSeries::ETimelineType::MemTag)
-					{
-						PreUpdateMemTagSeries(*MemorySeries, Viewport);
-					}
-					else
-					{
-						PreUpdateAllocationsTimelineSeries(*MemorySeries, Viewport);
-					}
+					MemorySeries.PreUpdate(*this, Viewport);
 				}
 
-				const double CurrentSeriesMinValue = MemorySeries->GetMinValue();
-				const double CurrentSeriesMaxValue = MemorySeries->GetMaxValue();
+				const double CurrentSeriesMinValue = MemorySeries.GetMinValue();
+				const double CurrentSeriesMaxValue = MemorySeries.GetMaxValue();
 				if (CurrentSeriesMinValue < AllSeriesMinValue)
 				{
 					AllSeriesMinValue = CurrentSeriesMinValue;
@@ -214,16 +428,10 @@ void FMemoryGraphTrack::Update(const ITimingTrackUpdateContext& Context)
 				// Clear the flag before updating, because the update itself may further need to set the series as dirty.
 				Series->ClearDirtyFlag();
 
-				//TODO: if (Series->Is<FMemoryGraphSeries>())
-				TSharedPtr<FMemoryGraphSeries> MemorySeries = StaticCastSharedPtr<FMemoryGraphSeries>(Series);
-
-				if (MemorySeries->GetTimelineType() == FMemoryGraphSeries::ETimelineType::MemTag)
+				if (Series->Is<FMemoryGraphSeries>())
 				{
-					UpdateMemTagSeries(*MemorySeries, Viewport);
-				}
-				else
-				{
-					UpdateAllocationsTimelineSeries(*MemorySeries, Viewport);
+					FMemoryGraphSeries& MemorySeries = Series->As<FMemoryGraphSeries>();
+					MemorySeries.Update(*this, Viewport);
 				}
 
 				if (Series->IsAutoZoomDirty())
@@ -241,34 +449,33 @@ void FMemoryGraphTrack::Update(const ITimingTrackUpdateContext& Context)
 // LLM Tag Series
 ////////////////////////////////////////////////////////////////////////////////////////////////////
 
-TSharedPtr<FMemoryGraphSeries> FMemoryGraphTrack::GetMemTagSeries(FMemoryTrackerId InMemTrackerId, FMemoryTagId InMemTagId)
+TSharedPtr<FMemTagGraphSeries> FMemoryGraphTrack::GetMemTagSeries(FMemoryTrackerId InMemTrackerId, FMemoryTagSetId InMemTagSetId, FMemoryTagId InMemTagId)
 {
-	TSharedPtr<FGraphSeries>* Ptr = AllSeries.FindByPredicate([InMemTrackerId, InMemTagId](const TSharedPtr<FGraphSeries>& Series)
-	{
-		//TODO: if (Series->Is<FMemoryGraphSeries>())
-		const TSharedPtr<FMemoryGraphSeries> MemorySeries = StaticCastSharedPtr<FMemoryGraphSeries>(Series);
-		return MemorySeries->GetTimelineType() == FMemoryGraphSeries::ETimelineType::MemTag &&
-			   MemorySeries->GetTrackerId() == InMemTrackerId &&
-			   MemorySeries->GetTagId() == InMemTagId;
-	});
-	return (Ptr != nullptr) ? StaticCastSharedPtr<FMemoryGraphSeries>(*Ptr) : nullptr;
+	TSharedPtr<FGraphSeries>* FoundSeries = AllSeries.FindByPredicate(
+		[InMemTrackerId, InMemTagSetId, InMemTagId](const TSharedPtr<FGraphSeries>& GraphSeries)
+		{
+			if (!GraphSeries->Is<FMemTagGraphSeries>())
+			{
+				return false;
+			}
+			const FMemTagGraphSeries& MemTagSeries = GraphSeries->As<FMemTagGraphSeries>();
+			return MemTagSeries.GetTrackerId() == InMemTrackerId &&
+				   MemTagSeries.GetTagSetId() == InMemTagSetId &&
+				   MemTagSeries.GetTagId() == InMemTagId;
+		});
+	return FoundSeries ? StaticCastSharedPtr<FMemTagGraphSeries>(*FoundSeries) : nullptr;
 }
 
 ////////////////////////////////////////////////////////////////////////////////////////////////////
 
-TSharedPtr<FMemoryGraphSeries> FMemoryGraphTrack::AddMemTagSeries(FMemoryTrackerId InMemTrackerId, FMemoryTagId InMemTagId)
+TSharedPtr<FMemTagGraphSeries> FMemoryGraphTrack::AddMemTagSeries(FMemoryTrackerId InMemTrackerId, FMemoryTagSetId InMemTagSetId, FMemoryTagId InMemTagId)
 {
-	TSharedPtr<FMemoryGraphSeries> Series = GetMemTagSeries(InMemTrackerId, InMemTagId);
+	TSharedPtr<FMemTagGraphSeries> Series = GetMemTagSeries(InMemTrackerId, InMemTagSetId, InMemTagId);
 
 	if (!Series.IsValid())
 	{
-		Series = MakeShared<FMemoryGraphSeries>();
+		Series = MakeShared<FMemTagGraphSeries>(InMemTrackerId, InMemTagSetId, InMemTagId);
 
-		Series->SetName(TEXT("LLM Tag"));
-		Series->SetDescription(TEXT("Low Level Memory Tag"));
-
-		Series->SetTrackerId(InMemTrackerId);
-		Series->SetTagId(InMemTagId);
 		Series->SetValueRange(0.0f, 0.0f);
 
 		Series->SetBaselineY(GetHeight() - 1.0f);
@@ -283,17 +490,21 @@ TSharedPtr<FMemoryGraphSeries> FMemoryGraphTrack::AddMemTagSeries(FMemoryTracker
 
 ////////////////////////////////////////////////////////////////////////////////////////////////////
 
-int32 FMemoryGraphTrack::RemoveMemTagSeries(FMemoryTrackerId InMemTrackerId, FMemoryTagId InMemTagId)
+int32 FMemoryGraphTrack::RemoveMemTagSeries(FMemoryTrackerId InMemTrackerId, FMemoryTagSetId InMemTagSetId, FMemoryTagId InMemTagId)
 {
 	SetDirtyFlag();
-	return AllSeries.RemoveAll([InMemTrackerId, InMemTagId](const TSharedPtr<FGraphSeries>& GraphSeries)
-	{
-		//TODO: if (GraphSeries->Is<FMemoryGraphSeries>())
-		const TSharedPtr<FMemoryGraphSeries> Series = StaticCastSharedPtr<FMemoryGraphSeries>(GraphSeries);
-		return Series->GetTimelineType() == FMemoryGraphSeries::ETimelineType::MemTag &&
-			   Series->GetTrackerId() == InMemTrackerId &&
-			   Series->GetTagId() == InMemTagId;
-	});
+	return AllSeries.RemoveAll(
+		[InMemTrackerId, InMemTagSetId, InMemTagId](const TSharedPtr<FGraphSeries>& GraphSeries)
+		{
+			if (!GraphSeries->Is<FMemTagGraphSeries>())
+			{
+				return false;
+			}
+			const FMemTagGraphSeries& MemTagSeries = GraphSeries->As<FMemTagGraphSeries>();
+			return MemTagSeries.GetTrackerId() == InMemTrackerId &&
+				   MemTagSeries.GetTagSetId() == InMemTagSetId &&
+				   MemTagSeries.GetTagId() == InMemTagId;
+		});
 }
 
 ////////////////////////////////////////////////////////////////////////////////////////////////////
@@ -301,17 +512,12 @@ int32 FMemoryGraphTrack::RemoveMemTagSeries(FMemoryTrackerId InMemTrackerId, FMe
 int32 FMemoryGraphTrack::RemoveAllMemTagSeries()
 {
 	SetDirtyFlag();
-	return AllSeries.RemoveAll([](const TSharedPtr<FGraphSeries>& GraphSeries)
-	{
-		//TODO: if (GraphSeries->Is<FMemoryGraphSeries>())
-		const TSharedPtr<FMemoryGraphSeries> Series = StaticCastSharedPtr<FMemoryGraphSeries>(GraphSeries);
-		return Series->GetTimelineType() == FMemoryGraphSeries::ETimelineType::MemTag;
-	});
+	return AllSeries.RemoveAll([](const TSharedPtr<FGraphSeries>& GraphSeries) { return GraphSeries->Is<FMemTagGraphSeries>(); });
 }
 
 ////////////////////////////////////////////////////////////////////////////////////////////////////
 
-void FMemoryGraphTrack::PreUpdateMemTagSeries(FMemoryGraphSeries& Series, const FTimingTrackViewport& Viewport)
+void FMemTagGraphSeries::PreUpdate(FGraphTrack& GraphTrack, const FTimingTrackViewport& Viewport)
 {
 	TSharedPtr<const TraceServices::IAnalysisSession> Session = FInsightsManager::Get()->GetSession();
 	if (Session.IsValid())
@@ -320,28 +526,28 @@ void FMemoryGraphTrack::PreUpdateMemTagSeries(FMemoryGraphSeries& Series, const
 		const TraceServices::IMemoryProvider* MemoryProvider = TraceServices::ReadMemoryProvider(*Session.Get());
 		if (MemoryProvider)
 		{
-			const uint64 TotalSampleCount = MemoryProvider->GetTagSampleCount(Series.GetTrackerId(), Series.GetTagId());
+			const uint64 TotalSampleCount = MemoryProvider->GetTagSampleCount(TrackerId, TagId);
 			if (TotalSampleCount > 0)
 			{
-				double MinValue = +std::numeric_limits<double>::infinity();
-				double MaxValue = -std::numeric_limits<double>::infinity();
+				double LocalMinValue = +std::numeric_limits<double>::infinity();
+				double LocalMaxValue = -std::numeric_limits<double>::infinity();
 
 				// Compute Min/Max values.
-				MemoryProvider->EnumerateTagSamples(Series.GetTrackerId(), Series.GetTagId(), Viewport.GetStartTime(), Viewport.GetEndTime(), true,
-					[&MinValue, &MaxValue](double Time, double Duration, const TraceServices::FMemoryTagSample& Sample)
-				{
-					const double Value = static_cast<double>(Sample.Value);
-					if (Value < MinValue)
+				MemoryProvider->EnumerateTagSamples(TrackerId, TagId, Viewport.GetStartTime(), Viewport.GetEndTime(), true,
+					[&LocalMinValue, &LocalMaxValue](double Time, double Duration, const TraceServices::FMemoryTagSample& Sample)
 					{
-						MinValue = Value;
-					}
-					if (Value > MaxValue)
-					{
-						MaxValue = Value;
-					}
-				});
+						const double Value = static_cast<double>(Sample.Value);
+						if (Value < LocalMinValue)
+						{
+							LocalMinValue = Value;
+						}
+						if (Value > LocalMaxValue)
+						{
+							LocalMaxValue = Value;
+						}
+					});
 
-				Series.SetValueRange(MinValue, MaxValue);
+				SetValueRange(LocalMinValue, LocalMaxValue);
 			}
 		}
 	}
@@ -349,9 +555,9 @@ void FMemoryGraphTrack::PreUpdateMemTagSeries(FMemoryGraphSeries& Series, const
 
 ////////////////////////////////////////////////////////////////////////////////////////////////////
 
-void FMemoryGraphTrack::UpdateMemTagSeries(FMemoryGraphSeries& Series, const FTimingTrackViewport& Viewport)
+void FMemTagGraphSeries::Update(FGraphTrack& GraphTrack, const FTimingTrackViewport& Viewport)
 {
-	FGraphTrackBuilder Builder(*this, Series, Viewport);
+	FGraphTrackBuilder Builder(GraphTrack, *this, Viewport);
 
 	TSharedPtr<const TraceServices::IAnalysisSession> Session = FInsightsManager::Get()->GetSession();
 	if (Session.IsValid())
@@ -361,23 +567,23 @@ void FMemoryGraphTrack::UpdateMemTagSeries(FMemoryGraphSeries& Series, const FTi
 		const TraceServices::IMemoryProvider* MemoryProvider = TraceServices::ReadMemoryProvider(*Session.Get());
 		if (MemoryProvider)
 		{
-			const uint64 TotalSampleCount = MemoryProvider->GetTagSampleCount(Series.GetTrackerId(), Series.GetTagId());
+			const uint64 TotalSampleCount = MemoryProvider->GetTagSampleCount(TrackerId, TagId);
 
 			if (TotalSampleCount > 0)
 			{
 				const float TopY = 4.0f;
-				const float BottomY = GetHeight() - 4.0f;
+				const float BottomY = GraphTrack.GetHeight() - 4.0f;
 
-				if (Series.IsAutoZoomEnabled() && TopY < BottomY)
+				if (IsAutoZoomEnabled() && TopY < BottomY)
 				{
-					Series.UpdateAutoZoom(TopY, BottomY, Series.GetMinValue(), Series.GetMaxValue());
+					UpdateAutoZoom(TopY, BottomY, GetMinValue(), GetMaxValue());
 				}
 
-				MemoryProvider->EnumerateTagSamples(Series.GetTrackerId(), Series.GetTagId(), Viewport.GetStartTime(), Viewport.GetEndTime(), true,
+				MemoryProvider->EnumerateTagSamples(TrackerId, TagId, Viewport.GetStartTime(), Viewport.GetEndTime(), true,
 					[this, &Builder](double Time, double Duration, const TraceServices::FMemoryTagSample& Sample)
-				{
-					Builder.AddEvent(Time, Duration, static_cast<double>(Sample.Value));
-				});
+					{
+						Builder.AddEvent(Time, Duration, static_cast<double>(Sample.Value));
+					});
 			}
 		}
 	}
@@ -385,155 +591,31 @@ void FMemoryGraphTrack::UpdateMemTagSeries(FMemoryGraphSeries& Series, const FTi
 
 ////////////////////////////////////////////////////////////////////////////////////////////////////
 
-TSharedPtr<FMemoryGraphSeries> FMemoryGraphTrack::GetTimelineSeries(FMemoryGraphSeries::ETimelineType InTimelineType)
+TSharedPtr<FAllocationsGraphSeries> FMemoryGraphTrack::GetTimelineSeries(FAllocationsGraphSeries::ETimeline InTimeline)
 {
-	return nullptr; //TODO
+	TSharedPtr<FGraphSeries>* FoundSeries = AllSeries.FindByPredicate(
+		[InTimeline](const TSharedPtr<FGraphSeries>& GraphSeries)
+		{
+			if (!GraphSeries->Is<FAllocationsGraphSeries>())
+			{
+				return false;
+			}
+			const FAllocationsGraphSeries& AllocationsSeries = GraphSeries->As<FAllocationsGraphSeries>();
+			return AllocationsSeries.GetTimeline() == InTimeline;
+		});
+	return FoundSeries ? StaticCastSharedPtr<FAllocationsGraphSeries>(*FoundSeries) : nullptr;
 }
 
 ////////////////////////////////////////////////////////////////////////////////////////////////////
 
-TSharedPtr<FMemoryGraphSeries> FMemoryGraphTrack::AddTimelineSeries(FMemoryGraphSeries::ETimelineType InTimelineType)
+TSharedPtr<FAllocationsGraphSeries> FMemoryGraphTrack::AddTimelineSeries(FAllocationsGraphSeries::ETimeline InTimeline)
 {
-	TSharedPtr<FMemoryGraphSeries> Series = GetTimelineSeries(InTimelineType);
+	TSharedPtr<FAllocationsGraphSeries> Series = GetTimelineSeries(InTimeline);
 
 	if (!Series.IsValid())
 	{
-		Series = MakeShared<FMemoryGraphSeries>();
-
-		switch (InTimelineType)
-		{
-			case FMemoryGraphSeries::ETimelineType::MinTotalMem:
-			{
-				Series->SetName(TEXT("Total Allocated Memory (Min)"));
-				Series->SetDescription(TEXT("Minimum value per sample for the Total Allocated Memory"));
-				const FLinearColor Color = FLinearColor(0.0f, 0.5f, 1.0f, 1.0f);
-				const FLinearColor BorderColor(FMath::Min(Color.R + 0.4f, 1.0f), FMath::Min(Color.G + 0.4f, 1.0f), FMath::Min(Color.B + 0.4f, 1.0f), 1.0f);
-				Series->SetColor(Color, BorderColor, Color.CopyWithNewOpacity(0.1f));
-				break;
-			}
+		Series = MakeShared<FAllocationsGraphSeries>(InTimeline);
 
-			case FMemoryGraphSeries::ETimelineType::MaxTotalMem:
-			{
-				Series->SetName(TEXT("Total Allocated Memory (Max)"));
-				Series->SetDescription(TEXT("Maximum value per sample for the Total Allocated Memory"));
-				const FLinearColor Color = FLinearColor(1.0f, 0.25f, 1.0f, 1.0f);
-				const FLinearColor BorderColor(FMath::Min(Color.R + 0.4f, 1.0f), FMath::Min(Color.G + 0.4f, 1.0f), FMath::Min(Color.B + 0.4f, 1.0f), 1.0f);
-				Series->SetColor(Color, BorderColor, Color.CopyWithNewOpacity(0.1f));
-				break;
-			}
-
-			case FMemoryGraphSeries::ETimelineType::MinLiveAllocs:
-			{
-				Series->SetName(TEXT("Live Allocation Count (Min)"));
-				Series->SetDescription(TEXT("Minimum value per sample for the Live Allocation Count"));
-				const FLinearColor Color = FLinearColor(1.0f, 1.0f, 0.25f, 1.0f);
-				const FLinearColor BorderColor(FMath::Min(Color.R + 0.4f, 1.0f), FMath::Min(Color.G + 0.4f, 1.0f), FMath::Min(Color.B + 0.4f, 1.0f), 1.0f);
-				Series->SetColor(Color, BorderColor, Color.CopyWithNewOpacity(0.1f));
-				break;
-			}
-
-			case FMemoryGraphSeries::ETimelineType::MaxLiveAllocs:
-			{
-				Series->SetName(TEXT("Live Allocation Count (Max)"));
-				Series->SetDescription(TEXT("Maximum value per sample for the Live Allocation Count"));
-				const FLinearColor Color = FLinearColor(1.0f, 0.25f, 1.0f, 1.0f);
-				const FLinearColor BorderColor(FMath::Min(Color.R + 0.4f, 1.0f), FMath::Min(Color.G + 0.4f, 1.0f), FMath::Min(Color.B + 0.4f, 1.0f), 1.0f);
-				Series->SetColor(Color, BorderColor, Color.CopyWithNewOpacity(0.1f));
-				break;
-			}
-
-			case FMemoryGraphSeries::ETimelineType::MinSwapMem:
-			{
-				Series->SetName(TEXT("Total Swap Memory (Min)"));
-				Series->SetDescription(TEXT("Minimum value per sample for the Total Swap Memory"));
-				const FLinearColor Color = FLinearColor(0.0f, 0.5f, 1.0f, 1.0f);
-				const FLinearColor BorderColor(FMath::Min(Color.R + 0.4f, 1.0f), FMath::Min(Color.G + 0.4f, 1.0f), FMath::Min(Color.B + 0.4f, 1.0f), 1.0f);
-				Series->SetColor(Color, BorderColor, Color.CopyWithNewOpacity(0.1f));
-				break;
-			}
-
-			case FMemoryGraphSeries::ETimelineType::MaxSwapMem:
-			{
-				Series->SetName(TEXT("Total Swap Memory (Max)"));
-				Series->SetDescription(TEXT("Maximum value per sample for the Total Swap Memory"));
-				const FLinearColor Color = FLinearColor(1.0f, 0.25f, 1.0f, 1.0f);
-				const FLinearColor BorderColor(FMath::Min(Color.R + 0.4f, 1.0f), FMath::Min(Color.G + 0.4f, 1.0f), FMath::Min(Color.B + 0.4f, 1.0f), 1.0f);
-				Series->SetColor(Color, BorderColor, Color.CopyWithNewOpacity(0.1f));
-				break;
-			}
-
-			case FMemoryGraphSeries::ETimelineType::MinCompressedSwapMem:
-			{
-				Series->SetName(TEXT("Total Compressed Swap Memory (Min)"));
-				Series->SetDescription(TEXT("Minimum value per sample for the Total Compressed Swap Memory"));
-				const FLinearColor Color = FLinearColor(1.0f, 1.0f, 0.25f, 1.0f);
-				const FLinearColor BorderColor(FMath::Min(Color.R + 0.4f, 1.0f), FMath::Min(Color.G + 0.4f, 1.0f), FMath::Min(Color.B + 0.4f, 1.0f), 1.0f);
-				Series->SetColor(Color, BorderColor, Color.CopyWithNewOpacity(0.1f));
-				break;
-			}
-
-			case FMemoryGraphSeries::ETimelineType::MaxCompressedSwapMem:
-			{
-				Series->SetName(TEXT("Total Compressed Swap Memory (Max)"));
-				Series->SetDescription(TEXT("Maximum value per sample for the Total Compressed Swap Memory"));
-				const FLinearColor Color = FLinearColor(1.0f, 0.25f, 1.0f, 1.0f);
-				const FLinearColor BorderColor(FMath::Min(Color.R + 0.4f, 1.0f), FMath::Min(Color.G + 0.4f, 1.0f), FMath::Min(Color.B + 0.4f, 1.0f), 1.0f);
-				Series->SetColor(Color, BorderColor, Color.CopyWithNewOpacity(0.1f));
-				break;
-			}
-
-			case FMemoryGraphSeries::ETimelineType::AllocEvents:
-			{
-				Series->SetName(TEXT("Alloc Event Count"));
-				Series->SetDescription(TEXT("Number of alloc events per sample"));
-				const FLinearColor Color = FLinearColor(0.0f, 1.0f, 0.5f, 1.0f);
-				const FLinearColor BorderColor(FMath::Min(Color.R + 0.4f, 1.0f), FMath::Min(Color.G + 0.4f, 1.0f), FMath::Min(Color.B + 0.4f, 1.0f), 1.0f);
-				Series->SetColor(Color, BorderColor, Color.CopyWithNewOpacity(0.1f));
-				break;
-			}
-
-			case FMemoryGraphSeries::ETimelineType::FreeEvents:
-			{
-				Series->SetName(TEXT("Free Event Count"));
-				Series->SetDescription(TEXT("Number of free events per sample"));
-				const FLinearColor Color = FLinearColor(1.0f, 0.5f, 0.25f, 1.0f);
-				const FLinearColor BorderColor(FMath::Min(Color.R + 0.4f, 1.0f), FMath::Min(Color.G + 0.4f, 1.0f), FMath::Min(Color.B + 0.4f, 1.0f), 1.0f);
-				Series->SetColor(Color, BorderColor, Color.CopyWithNewOpacity(0.1f));
-				break;
-			}
-
-			case FMemoryGraphSeries::ETimelineType::PageInEvents:
-			{
-				Series->SetName(TEXT("Page In Event Count"));
-				Series->SetDescription(TEXT("Number of page in events per sample"));
-				const FLinearColor Color = FLinearColor(0.0f, 1.0f, 0.5f, 1.0f);
-				const FLinearColor BorderColor(FMath::Min(Color.R + 0.4f, 1.0f), FMath::Min(Color.G + 0.4f, 1.0f), FMath::Min(Color.B + 0.4f, 1.0f), 1.0f);
-				Series->SetColor(Color, BorderColor, Color.CopyWithNewOpacity(0.1f));
-				break;
-			}
-
-			case FMemoryGraphSeries::ETimelineType::PageOutEvents:
-			{
-				Series->SetName(TEXT("Page Out Event Count"));
-				Series->SetDescription(TEXT("Number of page out events per sample"));
-				const FLinearColor Color = FLinearColor(1.0f, 0.5f, 0.25f, 1.0f);
-				const FLinearColor BorderColor(FMath::Min(Color.R + 0.4f, 1.0f), FMath::Min(Color.G + 0.4f, 1.0f), FMath::Min(Color.B + 0.4f, 1.0f), 1.0f);
-				Series->SetColor(Color, BorderColor, Color.CopyWithNewOpacity(0.1f));
-				break;
-			}
-
-			case FMemoryGraphSeries::ETimelineType::SwapFreeEvents:
-			{
-				Series->SetName(TEXT("Swap Free Event Count"));
-				Series->SetDescription(TEXT("Number of swap free events per sample"));
-				const FLinearColor Color = FLinearColor(0.25f, 0.5f, 1.0f, 1.0f);
-				const FLinearColor BorderColor(FMath::Min(Color.R + 0.4f, 1.0f), FMath::Min(Color.G + 0.4f, 1.0f), FMath::Min(Color.B + 0.4f, 1.0f), 1.0f);
-				Series->SetColor(Color, BorderColor, Color.CopyWithNewOpacity(0.1f));
-				break;
-			}
-		}
-
-		Series->SetTimelineType(InTimelineType);
 		Series->SetValueRange(0.0f, 0.0f);
 
 		Series->SetBaselineY(GetHeight() - 1.0f);
@@ -548,7 +630,7 @@ TSharedPtr<FMemoryGraphSeries> FMemoryGraphTrack::AddTimelineSeries(FMemoryGraph
 
 ////////////////////////////////////////////////////////////////////////////////////////////////////
 
-void FMemoryGraphTrack::PreUpdateAllocationsTimelineSeries(FMemoryGraphSeries& Series, const FTimingTrackViewport& Viewport)
+void FAllocationsGraphSeries::PreUpdate(FGraphTrack& GraphTrack, const FTimingTrackViewport& Viewport)
 {
 	TSharedPtr<const TraceServices::IAnalysisSession> Session = FInsightsManager::Get()->GetSession();
 	if (Session.IsValid())
@@ -567,84 +649,84 @@ void FMemoryGraphTrack::PreUpdateAllocationsTimelineSeries(FMemoryGraphSeries& S
 				++EndIndex; // include one more point on the right side
 			}
 
-			uint64 MinValue = std::numeric_limits<uint64>::max();
-			uint64 MaxValue = 0;
+			uint64 LocalMinValue = std::numeric_limits<uint64>::max();
+			uint64 LocalMaxValue = 0;
 
 			// Compute Min/Max values.
-			auto Callback64 = [&MinValue, &MaxValue](double Time, double Duration, uint64 Value)
+			auto Callback64 = [&LocalMinValue, &LocalMaxValue](double Time, double Duration, uint64 Value)
 			{
-				if (Value < MinValue)
+				if (Value < LocalMinValue)
 				{
-					MinValue = Value;
+					LocalMinValue = Value;
 				}
-				if (Value > MaxValue)
+				if (Value > LocalMaxValue)
 				{
-					MaxValue = Value;
+					LocalMaxValue = Value;
 				}
 			};
-			auto Callback32 = [&MinValue, &MaxValue](double Time, double Duration, uint32 Value)
+			auto Callback32 = [&LocalMinValue, &LocalMaxValue](double Time, double Duration, uint32 Value)
 			{
-				if (Value < MinValue)
+				if (Value < LocalMinValue)
 				{
-					MinValue = Value;
+					LocalMinValue = Value;
 				}
-				if (Value > MaxValue)
+				if (Value > LocalMaxValue)
 				{
-					MaxValue = Value;
+					LocalMaxValue = Value;
 				}
 			};
 
-			switch (Series.GetTimelineType())
+			switch (Timeline)
 			{
-			case FMemoryGraphSeries::ETimelineType::MinTotalMem:
+			case ETimeline::MinTotalMem:
 				AllocationsProvider->EnumerateMinTotalAllocatedMemoryTimeline(StartIndex, EndIndex, Callback64);
 				break;
-			case FMemoryGraphSeries::ETimelineType::MaxTotalMem:
+			case ETimeline::MaxTotalMem:
 				AllocationsProvider->EnumerateMaxTotalAllocatedMemoryTimeline(StartIndex, EndIndex, Callback64);
 				break;
-			case FMemoryGraphSeries::ETimelineType::MinLiveAllocs:
+			case ETimeline::MinLiveAllocs:
 				AllocationsProvider->EnumerateMinLiveAllocationsTimeline(StartIndex, EndIndex, Callback32);
 				break;
-			case FMemoryGraphSeries::ETimelineType::MaxLiveAllocs:
+			case ETimeline::MaxLiveAllocs:
 				AllocationsProvider->EnumerateMaxLiveAllocationsTimeline(StartIndex, EndIndex, Callback32);
 				break;
-			case FMemoryGraphSeries::ETimelineType::MinSwapMem:
+			case ETimeline::MinSwapMem:
 				AllocationsProvider->EnumerateMinTotalSwapMemoryTimeline(StartIndex, EndIndex, Callback64);
 				break;
-			case FMemoryGraphSeries::ETimelineType::MaxSwapMem:
+			case ETimeline::MaxSwapMem:
 				AllocationsProvider->EnumerateMaxTotalSwapMemoryTimeline(StartIndex, EndIndex, Callback64);
 				break;
-			case FMemoryGraphSeries::ETimelineType::MinCompressedSwapMem:
+			case ETimeline::MinCompressedSwapMem:
 				AllocationsProvider->EnumerateMinTotalCompressedSwapMemoryTimeline(StartIndex, EndIndex, Callback64);
 				break;
-			case FMemoryGraphSeries::ETimelineType::MaxCompressedSwapMem:
+			case ETimeline::MaxCompressedSwapMem:
 				AllocationsProvider->EnumerateMaxTotalCompressedSwapMemoryTimeline(StartIndex, EndIndex, Callback64);
 				break;
-			case FMemoryGraphSeries::ETimelineType::AllocEvents:
+			case ETimeline::AllocEvents:
 				AllocationsProvider->EnumerateAllocEventsTimeline(StartIndex, EndIndex, Callback32);
 				break;
-			case FMemoryGraphSeries::ETimelineType::FreeEvents:
+			case ETimeline::FreeEvents:
 				AllocationsProvider->EnumerateFreeEventsTimeline(StartIndex, EndIndex, Callback32);
 				break;
-			case FMemoryGraphSeries::ETimelineType::PageInEvents:
+			case ETimeline::PageInEvents:
 				AllocationsProvider->EnumeratePageInEventsTimeline(StartIndex, EndIndex, Callback32);
 				break;
-			case FMemoryGraphSeries::ETimelineType::PageOutEvents:
+			case ETimeline::PageOutEvents:
 				AllocationsProvider->EnumeratePageOutEventsTimeline(StartIndex, EndIndex, Callback32);
 				break;
-			case FMemoryGraphSeries::ETimelineType::SwapFreeEvents:
+			case ETimeline::SwapFreeEvents:
 				AllocationsProvider->EnumerateSwapFreeEventsTimeline(StartIndex, EndIndex, Callback32);
 				break;
 			}
 
-			if (Series.GetTimelineType() == FMemoryGraphSeries::ETimelineType::FreeEvents)
+			if (Timeline == ETimeline::FreeEvents)
 			{
 				// Shows FreeEvents as negative values in order to be displayed on same graph as AllocEvents.
-				Series.SetValueRange(-static_cast<double>(MaxValue), -static_cast<double>(MinValue));
+				SetValueRange(-static_cast<double>(LocalMaxValue), -static_cast<double>(LocalMinValue));
 			}
 			else
 			{
-				Series.SetValueRange(static_cast<double>(MinValue), static_cast<double>(MaxValue));
+				SetValueRange(static_cast<double>(LocalMinValue), static_cast<double>(LocalMaxValue));
 			}
 		}
 	}
@@ -652,9 +734,9 @@ void FMemoryGraphTrack::PreUpdateAllocationsTimelineSeries(FMemoryGraphSeries& S
 
 ////////////////////////////////////////////////////////////////////////////////////////////////////
 
-void FMemoryGraphTrack::UpdateAllocationsTimelineSeries(FMemoryGraphSeries& Series, const FTimingTrackViewport& Viewport)
+void FAllocationsGraphSeries::Update(FGraphTrack& GraphTrack, const FTimingTrackViewport& Viewport)
 {
-	FGraphTrackBuilder Builder(*this, Series, Viewport);
+	FGraphTrackBuilder Builder(GraphTrack, *this, Viewport);
 
 	TSharedPtr<const TraceServices::IAnalysisSession> Session = FInsightsManager::Get()->GetSession();
 	if (Session.IsValid())
@@ -693,11 +775,11 @@ void FMemoryGraphTrack::UpdateAllocationsTimelineSeries(FMemoryGraphSeries& Seri
 				TimelineEvents.Reserve((EndIndex - StartIndex) + 1);
 
 				const float TopY = 4.0f;
-				const float BottomY = GetHeight() - 4.0f;
+				const float BottomY = GraphTrack.GetHeight() - 4.0f;
 
-				if (Series.IsAutoZoomEnabled() && TopY < BottomY)
+				if (IsAutoZoomEnabled() && TopY < BottomY)
 				{
-					Series.UpdateAutoZoom(TopY, BottomY, Series.GetMinValue(), Series.GetMaxValue());
+					UpdateAutoZoom(TopY, BottomY, GetMinValue(), GetMaxValue());
 				}
 
 				auto Callback64 = [&TimelineEvents](double Time, double Duration, uint64 Value)
@@ -714,45 +796,45 @@ void FMemoryGraphTrack::UpdateAllocationsTimelineSeries(FMemoryGraphSeries& Seri
 					TimelineEvents.Emplace(Time, Duration, -static_cast<double>(Value));
 				};
 
-				switch (Series.GetTimelineType())
+				switch (Timeline)
 				{
-				case FMemoryGraphSeries::ETimelineType::MinTotalMem:
+				case ETimeline::MinTotalMem:
 					AllocationsProvider->EnumerateMinTotalAllocatedMemoryTimeline(StartIndex, EndIndex, Callback64);
 					break;
-				case FMemoryGraphSeries::ETimelineType::MaxTotalMem:
+				case ETimeline::MaxTotalMem:
 					AllocationsProvider->EnumerateMaxTotalAllocatedMemoryTimeline(StartIndex, EndIndex, Callback64);
 					break;
-				case FMemoryGraphSeries::ETimelineType::MinLiveAllocs:
+				case ETimeline::MinLiveAllocs:
 					AllocationsProvider->EnumerateMinLiveAllocationsTimeline(StartIndex, EndIndex, Callback32);
 					break;
-				case FMemoryGraphSeries::ETimelineType::MaxLiveAllocs:
+				case ETimeline::MaxLiveAllocs:
 					AllocationsProvider->EnumerateMaxLiveAllocationsTimeline(StartIndex, EndIndex, Callback32);
 					break;
-				case FMemoryGraphSeries::ETimelineType::MinSwapMem:
+				case ETimeline::MinSwapMem:
 					AllocationsProvider->EnumerateMinTotalSwapMemoryTimeline(StartIndex, EndIndex, Callback64);
 					break;
-				case FMemoryGraphSeries::ETimelineType::MaxSwapMem:
+				case ETimeline::MaxSwapMem:
 					AllocationsProvider->EnumerateMaxTotalSwapMemoryTimeline(StartIndex, EndIndex, Callback64);
 					break;
-				case FMemoryGraphSeries::ETimelineType::MinCompressedSwapMem:
+				case ETimeline::MinCompressedSwapMem:
 					AllocationsProvider->EnumerateMinTotalCompressedSwapMemoryTimeline(StartIndex, EndIndex, Callback64);
 					break;
-				case FMemoryGraphSeries::ETimelineType::MaxCompressedSwapMem:
+				case ETimeline::MaxCompressedSwapMem:
 					AllocationsProvider->EnumerateMaxTotalCompressedSwapMemoryTimeline(StartIndex, EndIndex, Callback64);
 					break;
-				case FMemoryGraphSeries::ETimelineType::AllocEvents:
+				case ETimeline::AllocEvents:
 					AllocationsProvider->EnumerateAllocEventsTimeline(StartIndex, EndIndex, Callback32);
 					break;
-				case FMemoryGraphSeries::ETimelineType::FreeEvents:
+				case ETimeline::FreeEvents:
 					AllocationsProvider->EnumerateFreeEventsTimeline(StartIndex, EndIndex, Callback32Negative);
 					break;
-				case FMemoryGraphSeries::ETimelineType::PageInEvents:
+				case ETimeline::PageInEvents:
 					AllocationsProvider->EnumeratePageInEventsTimeline(StartIndex, EndIndex, Callback32);
 					break;
-				case FMemoryGraphSeries::ETimelineType::PageOutEvents:
+				case ETimeline::PageOutEvents:
 					AllocationsProvider->EnumeratePageOutEventsTimeline(StartIndex, EndIndex, Callback32Negative);
 					break;
-				case FMemoryGraphSeries::ETimelineType::SwapFreeEvents:
+				case ETimeline::SwapFreeEvents:
 					AllocationsProvider->EnumerateSwapFreeEventsTimeline(StartIndex, EndIndex, Callback32);
 					break;
 				}
@@ -777,9 +859,8 @@ void FMemoryGraphTrack::DrawVerticalAxisGrid(const ITimingTrackDrawContext& Cont
 		// Use the first visible series.
 		for (const TSharedPtr<FGraphSeries>& GraphSeries : AllSeries)
 		{
-			if (GraphSeries->IsVisible())
+			if (GraphSeries->IsVisible() && GraphSeries->Is<FMemoryGraphSeries>())
 			{
-				//TODO: if (GraphSeries->Is<FMemoryGraphSeries>())
 				Series = StaticCastSharedPtr<FMemoryGraphSeries>(GraphSeries);
 				break;
 			}
@@ -848,7 +929,9 @@ void FMemoryGraphTrack::DrawVerticalAxisGrid(const ITimingTrackDrawContext& Cont
 	{
 		double Grid;
 
-		if (Series->GetTimelineType() <= FMemoryGraphSeries::ETimelineType::MaxTotalMem)
+		if (Series->Is<FMemTagGraphSeries>() ||
+			(Series->Is<FAllocationsGraphSeries>() &&
+			 Series->As<FAllocationsGraphSeries>().GetTimeline() <= FAllocationsGraphSeries::ETimeline::MaxTotalMem))
 		{
 			const uint64 DeltaBytes = FMath::Max(1ULL, static_cast<uint64>(Delta));
 			Grid = static_cast<double>(FMath::RoundUpToPowerOfTwo64(DeltaBytes));
@@ -953,7 +1036,6 @@ void FMemoryGraphTrack::DrawVerticalAxisGrid(const ITimingTrackDrawContext& Cont
 			const float MX = static_cast<float>(Context.GetMousePosition().X);
 			const float MY = static_cast<float>(Context.GetMousePosition().Y);
 
-			//constexpr float MX1 = 80.0f; // start fading out
 			constexpr float MX2 = 120.0f; // completely faded out
 
 			if (MX > ViewWidth - MX2 && MY >= MaxValueY && MY < MinValueY + TextH)
@@ -1232,25 +1314,32 @@ void FMemoryGraphTrack::InitTooltip(FTooltipDrawState& InOutTooltip, const ITimi
 	if (InTooltipEvent.CheckTrack(this) && InTooltipEvent.Is<FGraphTrackEvent>())
 	{
 		const FGraphTrackEvent& TooltipEvent = InTooltipEvent.As<FGraphTrackEvent>();
-		const TSharedRef<const FMemoryGraphSeries> Series = StaticCastSharedRef<const FMemoryGraphSeries, const FGraphSeries>(TooltipEvent.GetSeries());
+		const TSharedRef<const FGraphSeries> GraphSeries = TooltipEvent.GetSeries();
+		if (GraphSeries->Is<FMemoryGraphSeries>())
+		{
+			const FMemoryGraphSeries& Series = GraphSeries->As<FMemoryGraphSeries>();
 
-		InOutTooltip.ResetContent();
-		InOutTooltip.AddTitle(Series->GetName().ToString(), Series->GetColor());
+			InOutTooltip.ResetContent();
 
-		if (Series->GetTimelineType() == FMemoryGraphSeries::ETimelineType::MemTag)
-		{
-			FString SubTitle = FString::Printf(TEXT("(tag id 0x%llX, tracker id %i)"), (uint64)Series->GetTagId(), (int32)Series->GetTrackerId());
-			InOutTooltip.AddTitle(SubTitle, Series->GetColor());
-		}
+			InOutTooltip.AddTitle(Series.GetName().ToString(), Series.GetColor());
 
-		const double Precision = FMath::Max(1.0 / TimeScaleX, FTimeValue::Nanosecond);
-		InOutTooltip.AddNameValueTextLine(TEXT("Time:"), FormatTime(TooltipEvent.GetStartTime(), Precision));
-		if (Series->HasEventDuration())
-		{
-			InOutTooltip.AddNameValueTextLine(TEXT("Duration:"), FormatTimeAuto(TooltipEvent.GetDuration()));
+			if (Series.Is<FMemTagGraphSeries>())
+			{
+				const FMemTagGraphSeries& MemTagSeries = Series.As<FMemTagGraphSeries>();
+				FString SubTitle = FString::Printf(TEXT("(tag id 0x%llX, tag set id %i, tracker id %i)"), (uint64)MemTagSeries.GetTagId(), (int32)MemTagSeries.GetTagSetId(), (int32)MemTagSeries.GetTrackerId());
+				InOutTooltip.AddTitle(SubTitle, Series.GetColor());
+			}
+
+			const double Precision = FMath::Max(1.0 / TimeScaleX, FTimeValue::Nanosecond);
+			InOutTooltip.AddNameValueTextLine(TEXT("Time:"), FormatTime(TooltipEvent.GetStartTime(), Precision));
+			if (Series.HasEventDuration())
+			{
+				InOutTooltip.AddNameValueTextLine(TEXT("Duration:"), FormatTimeAuto(TooltipEvent.GetDuration()));
+			}
+			InOutTooltip.AddNameValueTextLine(TEXT("Value:"), Series.FormatValue(TooltipEvent.GetValue()));
+
+			InOutTooltip.UpdateLayout();
 		}
-		InOutTooltip.AddNameValueTextLine(TEXT("Value:"), Series->FormatValue(TooltipEvent.GetValue()));
-		InOutTooltip.UpdateLayout();
 	}
 }
 
diff --git a/Engine/Source/Developer/TraceInsights/Private/Insights/MemoryProfiler/ViewModels/MemoryGraphTrack.h b/Engine/Source/Developer/TraceInsights/Private/Insights/MemoryProfiler/ViewModels/MemoryGraphTrack.h
--- a/Engine/Source/Developer/TraceInsights/Private/Insights/MemoryProfiler/ViewModels/MemoryGraphTrack.h
+++ b/Engine/Source/Developer/TraceInsights/Private/Insights/MemoryProfiler/ViewModels/MemoryGraphTrack.h
@@ -40,11 +40,65 @@ enum class EGraphTrackLabelUnit
 
 class FMemoryGraphSeries : public FGraphSeries
 {
+	INSIGHTS_DECLARE_RTTI(FMemoryGraphSeries, FGraphSeries)
+
 public:
-	enum class ETimelineType
+	FMemoryGraphSeries()
+	{
+	}
+	virtual ~FMemoryGraphSeries()
 	{
-		MemTag,
-		MinTotalMem,
+	}
+
+	virtual FString FormatValue(double Value) const override;
+	virtual void PreUpdate(FGraphTrack& GraphTrack, const FTimingTrackViewport& Viewport) = 0;
+	virtual void Update(FGraphTrack& GraphTrack, const FTimingTrackViewport& Viewport) = 0;
+
+	double GetMinValue() const { return MinValue; }
+	double GetMaxValue() const { return MaxValue; }
+	void SetValueRange(double Min, double Max) { MinValue = Min; MaxValue = Max; }
+
+private:
+	double MinValue = 0.0;
+	double MaxValue = 0.0;
+};
+
+////////////////////////////////////////////////////////////////////////////////////////////////////
+
+class FMemTagGraphSeries : public FMemoryGraphSeries
+{
+	INSIGHTS_DECLARE_RTTI(FMemTagGraphSeries, FMemoryGraphSeries)
+
+public:
+	FMemTagGraphSeries(FMemoryTrackerId InTrackerId, FMemoryTagSetId InTagSetId, FMemoryTagId InTagId);
+	virtual ~FMemTagGraphSeries();
+
+	virtual FString FormatValue(double Value) const override;
+	virtual void PreUpdate(FGraphTrack& GraphTrack, const FTimingTrackViewport& Viewport) override;
+	virtual void Update(FGraphTrack& GraphTrack, const FTimingTrackViewport& Viewport) override;
+
+	FMemoryTrackerId GetTrackerId() const { return TrackerId; }
+	FMemoryTagSetId GetTagSetId() const { return TagSetId; }
+	FMemoryTagId GetTagId() const { return TagId; }
+
+private:
+	FMemoryTrackerId TrackerId = FMemoryTracker::InvalidTrackerId; // LLM tracker id
+	FMemoryTagSetId TagSetId = FMemoryTagSet::InvalidTagSetId; // LLM tag set id
+	FMemoryTagId TagId = FMemoryTag::InvalidTagId; // LLM tag id
+};
+
+////////////////////////////////////////////////////////////////////////////////////////////////////
+
+class FAllocationsGraphSeries : public FMemoryGraphSeries
+{
+	INSIGHTS_DECLARE_RTTI(FAllocationsGraphSeries, FMemoryGraphSeries)
+
+public:
+	enum class ETimeline
+	{
+		Unknown = -1,
+
+		MinTotalMem = 0,
 		MaxTotalMem,
 		MinLiveAllocs,
 		MaxLiveAllocs,
@@ -56,31 +110,35 @@ public:
 		FreeEvents,
 		PageInEvents,
 		PageOutEvents,
-		SwapFreeEvents
+		SwapFreeEvents,
+
+		/** Not an actual parameter. Number of known timelines. */
+		Count,
 	};
 
-public:
-	virtual FString FormatValue(double Value) const override;
+	enum class EValueType
+	{
+		Unknown,
+		IntegerBytes,
+		IntegerCounter,
+	};
 
-	FMemoryTrackerId GetTrackerId() const { return TrackerId; }
-	void SetTrackerId(FMemoryTrackerId InTrackerId) { TrackerId = InTrackerId; }
+public:
+	FAllocationsGraphSeries(ETimeline InTimeline);
+	virtual ~FAllocationsGraphSeries();
 
-	FMemoryTagId GetTagId() const { return TagId; }
-	void SetTagId(FMemoryTagId InTagId) { TagId = InTagId; }
+	virtual FString FormatValue(double Value) const override;
+	virtual void PreUpdate(FGraphTrack& GraphTrack, const FTimingTrackViewport& Viewport) override;
+	virtual void Update(FGraphTrack& GraphTrack, const FTimingTrackViewport& Viewport) override;
 
-	ETimelineType GetTimelineType() const { return TimelineType; }
-	void SetTimelineType(ETimelineType InTimelineType) { TimelineType = InTimelineType; }
+	ETimeline GetTimeline() const { return Timeline; }
 
-	double GetMinValue() const { return MinValue; }
-	double GetMaxValue() const { return MaxValue; }
-	void SetValueRange(double Min, double Max) { MinValue = Min; MaxValue = Max; }
+private:
+	void Initialize();
 
 private:
-	FMemoryTrackerId TrackerId = FMemoryTracker::InvalidTrackerId; // LLM tracker id
-	FMemoryTagId TagId = FMemoryTag::InvalidTagId; // LLM tag id
-	double MinValue = 0.0;
-	double MaxValue = 0.0;
-	ETimelineType TimelineType = ETimelineType::MemTag;
+	ETimeline Timeline = ETimeline::Unknown;
+	EValueType ValueType = EValueType::Unknown;
 };
 
 ////////////////////////////////////////////////////////////////////////////////////////////////////
@@ -116,19 +174,20 @@ public:
 
 	bool IsStacked() const { return bIsStacked; }
 	void SetStacked(bool bOnOff) { bIsStacked = bOnOff; }
+
 	TSharedPtr<FMemoryGraphSeries> GetMainSeries() const { return MainSeries; }
 	void SetMainSeries(TSharedPtr<FMemoryGraphSeries> InMainSeries) { MainSeries = InMainSeries; }
 
 	virtual void Update(const ITimingTrackUpdateContext& Context) override;
 	virtual void InitTooltip(FTooltipDrawState& InOutTooltip, const ITimingEvent& InTooltipEvent) const override;
 
-	TSharedPtr<FMemoryGraphSeries> GetMemTagSeries(FMemoryTrackerId InMemTrackerId, FMemoryTagId InMemTagId);
-	TSharedPtr<FMemoryGraphSeries> AddMemTagSeries(FMemoryTrackerId InMemTrackerId, FMemoryTagId InMemTagId);
-	int32 RemoveMemTagSeries(FMemoryTrackerId InMemTrackerId, FMemoryTagId InMemTagId);
+	TSharedPtr<FMemTagGraphSeries> GetMemTagSeries(FMemoryTrackerId InMemTrackerId, FMemoryTagSetId InMemTagSetId, FMemoryTagId InMemTagId);
+	TSharedPtr<FMemTagGraphSeries> AddMemTagSeries(FMemoryTrackerId InMemTrackerId, FMemoryTagSetId InMemTagSetId, FMemoryTagId InMemTagId);
+	int32 RemoveMemTagSeries(FMemoryTrackerId InMemTrackerId, FMemoryTagSetId InMemTagSetId, FMemoryTagId InMemTagId);
 	int32 RemoveAllMemTagSeries();
 
-	TSharedPtr<FMemoryGraphSeries> GetTimelineSeries(FMemoryGraphSeries::ETimelineType InTimelineType);
-	TSharedPtr<FMemoryGraphSeries> AddTimelineSeries(FMemoryGraphSeries::ETimelineType InTimelineType);
+	TSharedPtr<FAllocationsGraphSeries> GetTimelineSeries(FAllocationsGraphSeries::ETimeline InTimeline);
+	TSharedPtr<FAllocationsGraphSeries> AddTimelineSeries(FAllocationsGraphSeries::ETimeline InTimeline);
 
 	void SetAvailableTrackHeight(EMemoryTrackHeightMode InMode, float InTrackHeight);
 	void SetCurrentTrackHeight(EMemoryTrackHeightMode InMode);
@@ -137,12 +196,6 @@ public:
 	static FString FormatValue(const double InValue, const double InUnitValue, const TCHAR* InUnitText, const int32 InDecimalDigitCount);
 
 protected:
-	void PreUpdateMemTagSeries(FMemoryGraphSeries& Series, const FTimingTrackViewport& Viewport);
-	void UpdateMemTagSeries(FMemoryGraphSeries& Series, const FTimingTrackViewport& Viewport);
-
-	void PreUpdateAllocationsTimelineSeries(FMemoryGraphSeries& Series, const FTimingTrackViewport& Viewport);
-	void UpdateAllocationsTimelineSeries(FMemoryGraphSeries& Series, const FTimingTrackViewport& Viewport);
-
 	virtual void DrawVerticalAxisGrid(const ITimingTrackDrawContext& Context) const override;
 
 	struct TDrawHorizontalAxisLabelParams
diff --git a/Engine/Source/Developer/TraceInsights/Private/Insights/MemoryProfiler/ViewModels/MemorySharedState.cpp b/Engine/Source/Developer/TraceInsights/Private/Insights/MemoryProfiler/ViewModels/MemorySharedState.cpp
--- a/Engine/Source/Developer/TraceInsights/Private/Insights/MemoryProfiler/ViewModels/MemorySharedState.cpp
+++ b/Engine/Source/Developer/TraceInsights/Private/Insights/MemoryProfiler/ViewModels/MemorySharedState.cpp
@@ -182,8 +182,8 @@ void FMemorySharedState::Tick(Timing::ITimingViewSession& InSession, const Trace
 		MainGraphTrack->SetOrder(GraphTrackOrder);
 		MainGraphTrack->SetName(TEXT("Main Memory Graph"));
 
-		MainGraphTrack->AddTimelineSeries(FMemoryGraphSeries::ETimelineType::MaxTotalMem);
-		MainGraphTrack->AddTimelineSeries(FMemoryGraphSeries::ETimelineType::MinTotalMem);
+		MainGraphTrack->AddTimelineSeries(FAllocationsGraphSeries::ETimeline::MaxTotalMem);
+		MainGraphTrack->AddTimelineSeries(FAllocationsGraphSeries::ETimeline::MinTotalMem);
 
 		MainGraphTrack->SetVisibilityFlag(bShowHideAllMemoryTracks);
 
@@ -204,8 +204,8 @@ void FMemorySharedState::Tick(Timing::ITimingViewSession& InSession, const Trace
 		LiveAllocsGraphTrack->SetName(TEXT("Live Allocation Count"));
 		LiveAllocsGraphTrack->SetLabelUnit(EGraphTrackLabelUnit::Count, 0);
 
-		LiveAllocsGraphTrack->AddTimelineSeries(FMemoryGraphSeries::ETimelineType::MaxLiveAllocs);
-		LiveAllocsGraphTrack->AddTimelineSeries(FMemoryGraphSeries::ETimelineType::MinLiveAllocs);
+		LiveAllocsGraphTrack->AddTimelineSeries(FAllocationsGraphSeries::ETimeline::MaxLiveAllocs);
+		LiveAllocsGraphTrack->AddTimelineSeries(FAllocationsGraphSeries::ETimeline::MinLiveAllocs);
 
 		LiveAllocsGraphTrack->SetVisibilityFlag(bShowHideAllMemoryTracks);
 
@@ -226,8 +226,8 @@ void FMemorySharedState::Tick(Timing::ITimingViewSession& InSession, const Trace
 		AllocFreeGraphTrack->SetName(TEXT("Alloc/Free Event Count"));
 		AllocFreeGraphTrack->SetLabelUnit(EGraphTrackLabelUnit::Count, 0);
 
-		AllocFreeGraphTrack->AddTimelineSeries(FMemoryGraphSeries::ETimelineType::AllocEvents);
-		AllocFreeGraphTrack->AddTimelineSeries(FMemoryGraphSeries::ETimelineType::FreeEvents);
+		AllocFreeGraphTrack->AddTimelineSeries(FAllocationsGraphSeries::ETimeline::AllocEvents);
+		AllocFreeGraphTrack->AddTimelineSeries(FAllocationsGraphSeries::ETimeline::FreeEvents);
 
 		AllocFreeGraphTrack->SetVisibilityFlag(bShowHideAllMemoryTracks);
 
@@ -247,10 +247,10 @@ void FMemorySharedState::Tick(Timing::ITimingViewSession& InSession, const Trace
 		SwapMemoryGraphTrack->SetOrder(GraphTrackOrder + 3);
 		SwapMemoryGraphTrack->SetName(TEXT("Swap Memory Graph"));
 
-		SwapMemoryGraphTrack->AddTimelineSeries(FMemoryGraphSeries::ETimelineType::MaxSwapMem);
-		SwapMemoryGraphTrack->AddTimelineSeries(FMemoryGraphSeries::ETimelineType::MinSwapMem);
-		SwapMemoryGraphTrack->AddTimelineSeries(FMemoryGraphSeries::ETimelineType::MaxCompressedSwapMem);
-		SwapMemoryGraphTrack->AddTimelineSeries(FMemoryGraphSeries::ETimelineType::MinCompressedSwapMem);
+		SwapMemoryGraphTrack->AddTimelineSeries(FAllocationsGraphSeries::ETimeline::MaxSwapMem);
+		SwapMemoryGraphTrack->AddTimelineSeries(FAllocationsGraphSeries::ETimeline::MinSwapMem);
+		SwapMemoryGraphTrack->AddTimelineSeries(FAllocationsGraphSeries::ETimeline::MaxCompressedSwapMem);
+		SwapMemoryGraphTrack->AddTimelineSeries(FAllocationsGraphSeries::ETimeline::MinCompressedSwapMem);
 		SwapMemoryGraphTrack->SetLabelUnit(EGraphTrackLabelUnit::MiB, 1);
 
 		SwapMemoryGraphTrack->SetVisibilityFlag(false); // don't show swaps track by default
@@ -272,9 +272,9 @@ void FMemorySharedState::Tick(Timing::ITimingViewSession& InSession, const Trace
 		PageSwapGraphTrack->SetName(TEXT("Page In/Out Event Count"));
 		PageSwapGraphTrack->SetLabelUnit(EGraphTrackLabelUnit::Count, 0);
 
-		PageSwapGraphTrack->AddTimelineSeries(FMemoryGraphSeries::ETimelineType::PageInEvents);
-		PageSwapGraphTrack->AddTimelineSeries(FMemoryGraphSeries::ETimelineType::PageOutEvents);
-		PageSwapGraphTrack->AddTimelineSeries(FMemoryGraphSeries::ETimelineType::SwapFreeEvents);
+		PageSwapGraphTrack->AddTimelineSeries(FAllocationsGraphSeries::ETimeline::PageInEvents);
+		PageSwapGraphTrack->AddTimelineSeries(FAllocationsGraphSeries::ETimeline::PageOutEvents);
+		PageSwapGraphTrack->AddTimelineSeries(FAllocationsGraphSeries::ETimeline::SwapFreeEvents);
 
 		PageSwapGraphTrack->SetVisibilityFlag(false); // don't show swap tracks by default
 
@@ -617,12 +617,14 @@ void FMemorySharedState::RemoveTrackFromMemTags(TSharedPtr<FMemoryGraphTrack>& G
 {
 	for (TSharedPtr<FGraphSeries>& Series : GraphTrack->GetSeries())
 	{
-		//TODO: if (Series->Is<FMemoryGraphSeries>())
-		TSharedPtr<FMemoryGraphSeries> MemorySeries = StaticCastSharedPtr<FMemoryGraphSeries>(Series);
-		FMemoryTag* TagPtr = TagList.GetTagById(MemorySeries->GetTrackerId(), MemorySeries->GetTagId());
-		if (TagPtr)
+		if (Series->Is<FMemTagGraphSeries>())
 		{
-			TagPtr->RemoveTrack(GraphTrack);
+			FMemTagGraphSeries& MemTagSeries = Series->As<FMemTagGraphSeries>();
+			FMemoryTag* TagPtr = TagList.GetTagById(MemTagSeries.GetTrackerId(), MemTagSeries.GetTagId());
+			if (TagPtr)
+			{
+				TagPtr->RemoveTrack(GraphTrack);
+			}
 		}
 	}
 }
@@ -641,7 +643,9 @@ TSharedPtr<FMemoryGraphTrack> FMemorySharedState::GetMemTagGraphTrack(FMemoryTra
 	{
 		for (TSharedPtr<FMemoryGraphTrack> MemoryGraph : TagPtr->GetGraphTracks())
 		{
-			if (MemoryGraph != MainGraphTrack && MemoryGraph->GetSeries().Num() == 1)
+			if (MemoryGraph != MainGraphTrack &&
+				MemoryGraph->GetSeries().Num() == 1 &&
+				MemoryGraph->GetSeries()[0]->Is<FMemTagGraphSeries>())
 			{
 				return MemoryGraph;
 			}
@@ -661,6 +665,7 @@ TSharedPtr<FMemoryGraphTrack> FMemorySharedState::CreateMemTagGraphTrack(FMemory
 	}
 
 	FMemoryTag* TagPtr = TagList.GetTagById(InMemTrackerId, InMemTagId);
+	FMemoryTagSetId MemTagSetId = FMemoryTagSet::InvalidTagSetId;
 
 	FString SeriesName;
 	if (TagPtr)
@@ -674,6 +679,7 @@ TSharedPtr<FMemoryGraphTrack> FMemorySharedState::CreateMemTagGraphTrack(FMemory
 		{
 			SeriesName = FString::Printf(TEXT("LLM %s"), *TagPtr->GetStatFullName());
 		}
+		MemTagSetId = TagPtr->GetTagSetId();
 	}
 	else
 	{
@@ -686,7 +692,7 @@ TSharedPtr<FMemoryGraphTrack> FMemorySharedState::CreateMemTagGraphTrack(FMemory
 	// Also create a series in the MainGraphTrack.
 	if (MainGraphTrack.IsValid())
 	{
-		TSharedPtr<FMemoryGraphSeries> Series = MainGraphTrack->AddMemTagSeries(InMemTrackerId, InMemTagId);
+		TSharedPtr<FMemoryGraphSeries> Series = MainGraphTrack->AddMemTagSeries(InMemTrackerId, MemTagSetId, InMemTagId);
 		Series->SetName(SeriesName);
 		Series->SetColor(Color, BorderColor, Color.CopyWithNewOpacity(0.1f));
 		Series->DisableAutoZoom();
@@ -722,7 +728,7 @@ TSharedPtr<FMemoryGraphTrack> FMemorySharedState::CreateMemTagGraphTrack(FMemory
 		GraphTrack->EnableAutoZoom();
 
 		// Create series.
-		TSharedPtr<FMemoryGraphSeries> Series = GraphTrack->AddMemTagSeries(InMemTrackerId, InMemTagId);
+		TSharedPtr<FMemTagGraphSeries> Series = GraphTrack->AddMemTagSeries(InMemTrackerId, MemTagSetId, InMemTagId);
 		Series->SetName(SeriesName);
 		Series->SetColor(Color, BorderColor);
 		Series->SetBaselineY(GraphTrack->GetHeight() - 1.0f);
@@ -762,7 +768,7 @@ int32 FMemorySharedState::RemoveMemTagGraphTrack(FMemoryTrackerId InMemTrackerId
 	{
 		for (TSharedPtr<FMemoryGraphTrack> GraphTrack : TagPtr->GetGraphTracks())
 		{
-			GraphTrack->RemoveMemTagSeries(InMemTrackerId, InMemTagId);
+			GraphTrack->RemoveMemTagSeries(InMemTrackerId, TagPtr->GetTagSetId(), InMemTagId);
 			if (GraphTrack->GetSeries().Num() == 0)
 			{
 				if (GraphTrack == MainGraphTrack)
@@ -837,12 +843,13 @@ TSharedPtr<FMemoryGraphSeries> FMemorySharedState::ToggleMemTagGraphSeries(TShar
 	}
 
 	FMemoryTag* TagPtr = TagList.GetTagById(InMemTrackerId, InMemTagId);
+	FMemoryTagSetId TagSetId = TagPtr ? TagPtr->GetTagSetId() : FMemoryTagSet::InvalidTagSetId;
 
-	TSharedPtr<FMemoryGraphSeries> Series = InGraphTrack->GetMemTagSeries(InMemTrackerId, InMemTagId);
+	TSharedPtr<FMemoryGraphSeries> Series = InGraphTrack->GetMemTagSeries(InMemTrackerId, TagSetId, InMemTagId);
 	if (Series.IsValid())
 	{
 		// Remove existing series.
-		InGraphTrack->RemoveMemTagSeries(InMemTrackerId, InMemTagId);
+		InGraphTrack->RemoveMemTagSeries(InMemTrackerId, TagSetId, InMemTagId);
 		InGraphTrack->SetDirtyFlag();
 		TimingView->HandleTrackVisibilityChanged();
 
@@ -856,7 +863,7 @@ TSharedPtr<FMemoryGraphSeries> FMemorySharedState::ToggleMemTagGraphSeries(TShar
 	else
 	{
 		// Add new series.
-		Series = InGraphTrack->AddMemTagSeries(InMemTrackerId, InMemTagId);
+		Series = InGraphTrack->AddMemTagSeries(InMemTrackerId, TagSetId, InMemTagId);
 		Series->DisableAutoZoom();
 
 		if (TagPtr)
@@ -1022,7 +1029,7 @@ TSharedPtr<FMemoryGraphTrack> FMemorySharedState::CreateGraphTrack(const FReport
 		{
 			FMemoryTag& Tag = *TagPtr;
 
-			TSharedPtr<FMemoryGraphSeries> Series = GraphTrack->AddMemTagSeries(MemTrackerId, Tag.GetId());
+			TSharedPtr<FMemoryGraphSeries> Series = GraphTrack->AddMemTagSeries(MemTrackerId, Tag.GetTagSetId(), Tag.GetId());
 			Series->SetName(FText::FromString(FString::Printf(TEXT("LLM %s"), *Tag.GetStatFullName())));
 			const FLinearColor Color = Tag.GetColor();
 			const FLinearColor BorderColor(FMath::Min(Color.R + 0.4f, 1.0f), FMath::Min(Color.G + 0.4f, 1.0f), FMath::Min(Color.B + 0.4f, 1.0f), 1.0f);
diff --git a/Engine/Source/Developer/TraceInsights/Private/Insights/MemoryProfiler/Widgets/SMemTagTreeView.cpp b/Engine/Source/Developer/TraceInsights/Private/Insights/MemoryProfiler/Widgets/SMemTagTreeView.cpp
--- a/Engine/Source/Developer/TraceInsights/Private/Insights/MemoryProfiler/Widgets/SMemTagTreeView.cpp
+++ b/Engine/Source/Developer/TraceInsights/Private/Insights/MemoryProfiler/Widgets/SMemTagTreeView.cpp
@@ -2439,6 +2439,8 @@ void SMemTagTreeView::SetColorToNode(const FMemTagNodePtr& MemTagNode, FLinearCo
 
 	const FLinearColor BorderColor(FMath::Min(Color.R + 0.4f, 1.0f), FMath::Min(Color.G + 0.4f, 1.0f), FMath::Min(Color.B + 0.4f, 1.0f), 1.0f);
 
+	const FMemoryTrackerId MemTrackerId = MemTagNode->GetMemTrackerId();
+	const FMemoryTagSetId MemTagSetId = MemTagNode->GetMemTagSetId();
 	const FMemoryTagId MemTagId = MemTagNode->GetMemTagId();
 
 	TSharedPtr<FMemoryGraphTrack> MainGraphTrack;
@@ -2451,20 +2453,16 @@ void SMemTagTreeView::SetColorToNode(const FMemTagNodePtr& MemTagNode, FLinearCo
 
 	for (const TSharedPtr<FMemoryGraphTrack>& GraphTrack : MemTag->GetGraphTracks())
 	{
-		for (TSharedPtr<FGraphSeries>& Series : GraphTrack->GetSeries())
+		TSharedPtr<FMemTagGraphSeries> MemTagSeries = GraphTrack->GetMemTagSeries(MemTrackerId, MemTagSetId, MemTagId);
+		if (MemTagSeries)
 		{
-			//TODO: if (Series->Is<FMemoryGraphSeries>())
-			TSharedPtr<FMemoryGraphSeries> MemorySeries = StaticCastSharedPtr<FMemoryGraphSeries>(Series);
-			if (MemorySeries->GetTagId() == MemTagId)
+			if (GraphTrack == MainGraphTrack)
 			{
-				if (GraphTrack == MainGraphTrack)
-				{
-					MemorySeries->SetColor(Color, BorderColor, Color.CopyWithNewOpacity(0.1f));
-				}
-				else
-				{
-					MemorySeries->SetColor(Color, BorderColor);
-				}
+				MemTagSeries->SetColor(Color, BorderColor, Color.CopyWithNewOpacity(0.1f));
+			}
+			else
+			{
+				MemTagSeries->SetColor(Color, BorderColor);
 			}
 		}
 	}
diff --git a/Engine/Source/Developer/TraceInsights/Private/Insights/ViewModels/GraphSeries.cpp b/Engine/Source/Developer/TraceInsights/Private/Insights/ViewModels/GraphSeries.cpp
--- a/Engine/Source/Developer/TraceInsights/Private/Insights/ViewModels/GraphSeries.cpp
+++ b/Engine/Source/Developer/TraceInsights/Private/Insights/ViewModels/GraphSeries.cpp
@@ -8,6 +8,12 @@
 
 #define LOCTEXT_NAMESPACE "GraphTrack"
 
+////////////////////////////////////////////////////////////////////////////////////////////////////
+// FGraphSeries
+////////////////////////////////////////////////////////////////////////////////////////////////////
+
+INSIGHTS_IMPLEMENT_RTTI(FGraphSeries)
+
 ////////////////////////////////////////////////////////////////////////////////////////////////////
 
 FGraphSeries::FGraphSeries()
diff --git a/Engine/Source/Developer/TraceInsights/Public/Insights/ViewModels/GraphSeries.h b/Engine/Source/Developer/TraceInsights/Public/Insights/ViewModels/GraphSeries.h
--- a/Engine/Source/Developer/TraceInsights/Public/Insights/ViewModels/GraphSeries.h
+++ b/Engine/Source/Developer/TraceInsights/Public/Insights/ViewModels/GraphSeries.h
@@ -2,17 +2,22 @@
 
 #pragma once
 
+#include "CoreTypes.h"
 #include "Containers/Array.h"
 #include "Containers/UnrealString.h"
-#include "CoreMinimal.h"
 #include "HAL/PlatformMath.h"
-#include "Insights/ViewModels/GraphTrackEvent.h"
 #include "Internationalization/Text.h"
 #include "Math/Color.h"
 #include "Math/UnrealMathUtility.h"
 #include "Math/Vector2D.h"
 #include "Misc/AssertionMacros.h"
 
+// TraceInsightsCore
+#include "InsightsCore/Common/SimpleRtti.h"
+
+// TraceInsights
+#include "Insights/ViewModels/GraphTrackEvent.h"
+
 class FTimingTrackViewport;
 struct FGraphSeriesEvent;
 
@@ -69,6 +74,8 @@ class TRACEINSIGHTS_API FGraphSeries
 	friend class FGraphTrack;
 	friend class FGraphTrackBuilder;
 
+	INSIGHTS_DECLARE_RTTI_BASE(FGraphSeries)
+
 public:
 	struct FBox
 	{
@@ -99,21 +106,31 @@ public:
 	void ClearDirtyFlag() { bIsDirty = false; }
 
 	const FLinearColor& GetColor() const { return Color; }
+	const FLinearColor& GetFillColor() const { return FillColor; }
 	const FLinearColor& GetBorderColor() const { return BorderColor; }
 
-	void SetColor(FLinearColor InColor, FLinearColor InBorderColor)
+	void SetColor(const FLinearColor& InColor)
 	{
 		Color = InColor;
-		BorderColor = InBorderColor;
-		FillColor = InColor;
-		FillColor.A = 0.5f;
+		FillColor = InColor.CopyWithNewOpacity(0.1f);
+		BorderColor = FLinearColor(FMath::Min(InColor.R + 0.4f, 1.0f),
+								   FMath::Min(InColor.G + 0.4f, 1.0f),
+								   FMath::Min(InColor.B + 0.4f, 1.0f),
+								   InColor.A);
 	}
 
-	void SetColor(FLinearColor InColor, FLinearColor InBorderColor, FLinearColor InFillColor)
+	void SetColor(const FLinearColor& InColor, const FLinearColor& InBorderColor)
 	{
 		Color = InColor;
+		FillColor = InColor.CopyWithNewOpacity(0.5f);
 		BorderColor = InBorderColor;
+	}
+
+	void SetColor(const FLinearColor& InColor, const FLinearColor& InBorderColor, const FLinearColor& InFillColor)
+	{
+		Color = InColor;
 		FillColor = InFillColor;
+		BorderColor = InBorderColor;
 	}
 
 	bool HasEventDuration() const { return bHasEventDuration; }
