[Insights] Fixed a duplicated LOCTEXT.

#rnx

diff --git a/Engine/Source/Developer/TraceInsights/Private/Insights/MemoryProfiler/ViewModels/MemoryGraphTrack.cpp b/Engine/Source/Developer/TraceInsights/Private/Insights/MemoryProfiler/ViewModels/MemoryGraphTrack.cpp
--- a/Engine/Source/Developer/TraceInsights/Private/Insights/MemoryProfiler/ViewModels/MemoryGraphTrack.cpp
+++ b/Engine/Source/Developer/TraceInsights/Private/Insights/MemoryProfiler/ViewModels/MemoryGraphTrack.cpp
@@ -159,7 +159,7 @@ void FAllocationsGraphSeries::Initialize()
 		case ETimeline::MaxLiveAllocs:
 		{
 			ValueType = EValueType::IntegerCounter;
-			SetName(LOCTEXT("LiveAllocationCountMax_TrackDesc", "Live Allocation Count (Max)"));
+			SetName(LOCTEXT("LiveAllocationCountMax_TrackName", "Live Allocation Count (Max)"));
 			SetDescription(LOCTEXT("LiveAllocationCountMax_TrackDesc", "Maximum value per sample for the Live Allocation Count"));
 			SetColor(FLinearColor(1.0f, 0.25f, 1.0f, 1.0f));
 			break;
