[Insights] Code refactoring for MemAllocTable implementation:
- Added FAggregatedValueGetter helper class to facilitate creation of value getters for FMemAllocTable.
- Fixed all static_cast assumptions for tree nodes to be FMemAllocNode. Now is using SimpleRTTI casts.
- Moved the declaration of MemAlloc table presets to a new class and source files (MemAllocTableViewPresets.h/.cpp).

#rnx

diff --git a/Engine/Source/Developer/TraceInsights/Private/Insights/MemoryProfiler/ViewModels/MemAllocGroupingByCallstack.h b/Engine/Source/Developer/TraceInsights/Private/Insights/MemoryProfiler/ViewModels/MemAllocGroupingByCallstack.h
--- a/Engine/Source/Developer/TraceInsights/Private/Insights/MemoryProfiler/ViewModels/MemAllocGroupingByCallstack.h
+++ b/Engine/Source/Developer/TraceInsights/Private/Insights/MemoryProfiler/ViewModels/MemAllocGroupingByCallstack.h
@@ -5,7 +5,7 @@
 #include "CoreTypes.h"
 
 // TraceInsightsCore
-#include "InsightsCore/Table/ViewModels/BaseTreeNode.h"
+#include "InsightsCore/Table/ViewModels/TableTreeNode.h"
 #include "InsightsCore/Table/ViewModels/TreeNodeGrouping.h"
 
 namespace TraceServices { struct FStackFrame; }
diff --git a/Engine/Source/Developer/TraceInsights/Private/Insights/MemoryProfiler/ViewModels/MemAllocTable.cpp b/Engine/Source/Developer/TraceInsights/Private/Insights/MemoryProfiler/ViewModels/MemAllocTable.cpp
--- a/Engine/Source/Developer/TraceInsights/Private/Insights/MemoryProfiler/ViewModels/MemAllocTable.cpp
+++ b/Engine/Source/Developer/TraceInsights/Private/Insights/MemoryProfiler/ViewModels/MemAllocTable.cpp
@@ -80,8 +80,6 @@ FMemAllocTable::~FMemAllocTable()
 
 void FMemAllocTable::Reset()
 {
-	//...
-
 	FTable::Reset();
 
 	AddDefaultColumns();
@@ -91,6 +89,28 @@ void FMemAllocTable::Reset()
 
 void FMemAllocTable::AddDefaultColumns()
 {
+	class FAggregatedValueGetter : public FTableCellValueGetter
+	{
+	public:
+		virtual const TOptional<FTableCellValue> GetNodeValue(const FBaseTreeNode& Node) const
+		{
+			return TOptional<FTableCellValue>();
+		}
+
+		virtual const TOptional<FTableCellValue> GetValue(const FTableColumn& Column, const FBaseTreeNode& Node) const override
+		{
+			if (Node.IsGroup() && Node.Is<FTableTreeNode>())
+			{
+				const FTableTreeNode& TableTreeNode = Node.As<FTableTreeNode>();
+				if (TableTreeNode.HasAggregatedValue(Column.GetId()))
+				{
+					return TableTreeNode.GetAggregatedValue(Column.GetId());
+				}
+			}
+			return GetNodeValue(Node);
+		}
+	};
+
 	//////////////////////////////////////////////////
 	// Hierarchy Column
 	{
@@ -98,11 +118,11 @@ void FMemAllocTable::AddDefaultColumns()
 		const TCHAR* HierarchyColumnName = nullptr;
 		AddHierarchyColumn(HierarchyColumnIndex, HierarchyColumnName);
 
-		const TSharedRef<FTableColumn>& ColumnRef = GetColumns()[0];
-		ColumnRef->SetInitialWidth(200.0f);
-		ColumnRef->SetShortName(LOCTEXT("AllocationColumnName", "Hierarchy"));
-		ColumnRef->SetTitleName(LOCTEXT("AllocationColumnTitle", "Allocation Hierarchy"));
-		ColumnRef->SetDescription(LOCTEXT("AllocationColumnDesc", "Hierarchy of the allocation's tree"));
+		FTableColumn& HierarchyColumn = *GetColumns()[0];
+		HierarchyColumn.SetInitialWidth(200.0f);
+		HierarchyColumn.SetShortName(LOCTEXT("AllocationColumnName", "Hierarchy"));
+		HierarchyColumn.SetTitleName(LOCTEXT("AllocationColumnTitle", "Allocation Hierarchy"));
+		HierarchyColumn.SetDescription(LOCTEXT("AllocationColumnDesc", "Hierarchy of the allocation's tree"));
 	}
 
 	int32 ColumnIndex = 0;
@@ -126,29 +146,20 @@ void FMemAllocTable::AddDefaultColumns()
 
 		Column.SetDataType(ETableCellDataType::Int64);
 
-		class FMemAllocStartEventIndexValueGetter : public FTableCellValueGetter
+		class FMemAllocStartEventIndexValueGetter : public FAggregatedValueGetter
 		{
 		public:
-			virtual const TOptional<FTableCellValue> GetValue(const FTableColumn& Column, const FBaseTreeNode& Node) const override
+			virtual const TOptional<FTableCellValue> GetNodeValue(const FBaseTreeNode& Node) const override
 			{
-				if (Node.IsGroup())
+				if (Node.Is<FMemAllocNode>())
 				{
-					const FTableTreeNode& NodePtr = static_cast<const FTableTreeNode&>(Node);
-					if (NodePtr.HasAggregatedValue(Column.GetId()))
-					{
-						return NodePtr.GetAggregatedValue(Column.GetId());
-					}
-				}
-				else //if (Node->Is<FMemAllocNode>())
-				{
-					const FMemAllocNode& MemAllocNode = static_cast<const FMemAllocNode&>(Node);
+					const FMemAllocNode& MemAllocNode = Node.As<FMemAllocNode>();
 					const FMemoryAlloc* Alloc = MemAllocNode.GetMemAlloc();
 					if (Alloc)
 					{
 						return FTableCellValue(Alloc->GetStartEventIndex());
 					}
 				}
-
 				return TOptional<FTableCellValue>();
 			}
 		};
@@ -184,29 +195,20 @@ void FMemAllocTable::AddDefaultColumns()
 
 		Column.SetDataType(ETableCellDataType::Int64);
 
-		class FMemAllocEndEventIndexValueGetter : public FTableCellValueGetter
+		class FMemAllocEndEventIndexValueGetter : public FAggregatedValueGetter
 		{
 		public:
-			virtual const TOptional<FTableCellValue> GetValue(const FTableColumn& Column, const FBaseTreeNode& Node) const override
+			virtual const TOptional<FTableCellValue> GetNodeValue(const FBaseTreeNode& Node) const override
 			{
-				if (Node.IsGroup())
+				if (Node.Is<FMemAllocNode>())
 				{
-					const FTableTreeNode& NodePtr = static_cast<const FTableTreeNode&>(Node);
-					if (NodePtr.HasAggregatedValue(Column.GetId()))
-					{
-						return NodePtr.GetAggregatedValue(Column.GetId());
-					}
-				}
-				else //if (Node->Is<FMemAllocNode>())
-				{
-					const FMemAllocNode& MemAllocNode = static_cast<const FMemAllocNode&>(Node);
+					const FMemAllocNode& MemAllocNode = Node.As<FMemAllocNode>();
 					const FMemoryAlloc* Alloc = MemAllocNode.GetMemAlloc();
 					if (Alloc)
 					{
 						return FTableCellValue(Alloc->GetEndEventIndex());
 					}
 				}
-
 				return TOptional<FTableCellValue>();
 			}
 		};
@@ -242,29 +244,20 @@ void FMemAllocTable::AddDefaultColumns()
 
 		Column.SetDataType(ETableCellDataType::Int64);
 
-		class FMemAllocEventDistanceValueGetter : public FTableCellValueGetter
+		class FMemAllocEventDistanceValueGetter : public FAggregatedValueGetter
 		{
 		public:
-			virtual const TOptional<FTableCellValue> GetValue(const FTableColumn& Column, const FBaseTreeNode& Node) const override
+			virtual const TOptional<FTableCellValue> GetNodeValue(const FBaseTreeNode& Node) const override
 			{
-				if (Node.IsGroup())
+				if (Node.Is<FMemAllocNode>())
 				{
-					const FTableTreeNode& NodePtr = static_cast<const FTableTreeNode&>(Node);
-					if (NodePtr.HasAggregatedValue(Column.GetId()))
-					{
-						return NodePtr.GetAggregatedValue(Column.GetId());
-					}
-				}
-				else //if (Node->Is<FMemAllocNode>())
-				{
-					const FMemAllocNode& MemAllocNode = static_cast<const FMemAllocNode&>(Node);
+					const FMemAllocNode& MemAllocNode = Node.As<FMemAllocNode>();
 					const FMemoryAlloc* Alloc = MemAllocNode.GetMemAlloc();
 					if (Alloc)
 					{
 						return FTableCellValue(Alloc->GetEventDistance());
 					}
 				}
-
 				return TOptional<FTableCellValue>();
 			}
 		};
@@ -300,29 +293,20 @@ void FMemAllocTable::AddDefaultColumns()
 
 		Column.SetDataType(ETableCellDataType::Double);
 
-		class FMemAllocStartTimeValueGetter : public FTableCellValueGetter
+		class FMemAllocStartTimeValueGetter : public FAggregatedValueGetter
 		{
 		public:
-			virtual const TOptional<FTableCellValue> GetValue(const FTableColumn& Column, const FBaseTreeNode& Node) const override
+			virtual const TOptional<FTableCellValue> GetNodeValue(const FBaseTreeNode& Node) const override
 			{
-				if (Node.IsGroup())
+				if (Node.Is<FMemAllocNode>())
 				{
-					const FTableTreeNode& NodePtr = static_cast<const FTableTreeNode&>(Node);
-					if (NodePtr.HasAggregatedValue(Column.GetId()))
-					{
-						return NodePtr.GetAggregatedValue(Column.GetId());
-					}
-				}
-				else //if (Node->Is<FMemAllocNode>())
-				{
-					const FMemAllocNode& MemAllocNode = static_cast<const FMemAllocNode&>(Node);
+					const FMemAllocNode& MemAllocNode = Node.As<FMemAllocNode>();
 					const FMemoryAlloc* Alloc = MemAllocNode.GetMemAlloc();
 					if (Alloc)
 					{
 						return FTableCellValue(Alloc->GetStartTime());
 					}
 				}
-
 				return TOptional<FTableCellValue>();
 			}
 		};
@@ -361,29 +345,20 @@ void FMemAllocTable::AddDefaultColumns()
 
 		Column.SetDataType(ETableCellDataType::Double);
 
-		class FMemAllocEndTimeValueGetter : public FTableCellValueGetter
+		class FMemAllocEndTimeValueGetter : public FAggregatedValueGetter
 		{
 		public:
-			virtual const TOptional<FTableCellValue> GetValue(const FTableColumn& Column, const FBaseTreeNode& Node) const override
+			virtual const TOptional<FTableCellValue> GetNodeValue(const FBaseTreeNode& Node) const override
 			{
-				if (Node.IsGroup())
+				if (Node.Is<FMemAllocNode>())
 				{
-					const FTableTreeNode& NodePtr = static_cast<const FTableTreeNode&>(Node);
-					if (NodePtr.HasAggregatedValue(Column.GetId()))
-					{
-						return NodePtr.GetAggregatedValue(Column.GetId());
-					}
-				}
-				else //if (Node->Is<FMemAllocNode>())
-				{
-					const FMemAllocNode& MemAllocNode = static_cast<const FMemAllocNode&>(Node);
+					const FMemAllocNode& MemAllocNode = Node.As<FMemAllocNode>();
 					const FMemoryAlloc* Alloc = MemAllocNode.GetMemAlloc();
 					if (Alloc)
 					{
 						return FTableCellValue(Alloc->GetEndTime());
 					}
 				}
-
 				return TOptional<FTableCellValue>();
 			}
 		};
@@ -422,29 +397,20 @@ void FMemAllocTable::AddDefaultColumns()
 
 		Column.SetDataType(ETableCellDataType::Double);
 
-		class FMemAllocDurationValueGetter : public FTableCellValueGetter
+		class FMemAllocDurationValueGetter : public FAggregatedValueGetter
 		{
 		public:
-			virtual const TOptional<FTableCellValue> GetValue(const FTableColumn& Column, const FBaseTreeNode& Node) const override
+			virtual const TOptional<FTableCellValue> GetNodeValue(const FBaseTreeNode& Node) const override
 			{
-				if (Node.IsGroup())
-				{
-					const FTableTreeNode& NodePtr = static_cast<const FTableTreeNode&>(Node);
-					if (NodePtr.HasAggregatedValue(Column.GetId()))
-					{
-						return NodePtr.GetAggregatedValue(Column.GetId());
-					}
-				}
-				else //if (Node->Is<FMemAllocNode>())
+				if (Node.Is<FMemAllocNode>())
 				{
-					const FMemAllocNode& MemAllocNode = static_cast<const FMemAllocNode&>(Node);
+					const FMemAllocNode& MemAllocNode = Node.As<FMemAllocNode>();
 					const FMemoryAlloc* Alloc = MemAllocNode.GetMemAlloc();
 					if (Alloc)
 					{
 						return FTableCellValue(Alloc->GetDuration());
 					}
 				}
-
 				return TOptional<FTableCellValue>();
 			}
 		};
@@ -481,29 +447,20 @@ void FMemAllocTable::AddDefaultColumns()
 
 		Column.SetDataType(ETableCellDataType::Int64);
 
-		class FAllocThreadValueGetter : public FTableCellValueGetter
+		class FAllocThreadValueGetter : public FAggregatedValueGetter
 		{
 		public:
-			virtual const TOptional<FTableCellValue> GetValue(const FTableColumn& Column, const FBaseTreeNode& Node) const override
+			virtual const TOptional<FTableCellValue> GetNodeValue(const FBaseTreeNode& Node) const override
 			{
-				if (Node.IsGroup())
-				{
-					const FTableTreeNode& NodePtr = static_cast<const FTableTreeNode&>(Node);
-					if (NodePtr.HasAggregatedValue(Column.GetId()))
-					{
-						return NodePtr.GetAggregatedValue(Column.GetId());
-					}
-				}
-				else //if (Node->Is<FMemAllocNode>())
+				if (Node.Is<FMemAllocNode>())
 				{
-					const FMemAllocNode& MemAllocNode = static_cast<const FMemAllocNode&>(Node);
+					const FMemAllocNode& MemAllocNode = Node.As<FMemAllocNode>();
 					const FMemoryAlloc* Alloc = MemAllocNode.GetMemAlloc();
 					if (Alloc)
 					{
 						return FTableCellValue(static_cast<int64>(Alloc->GetAllocThreadId()));
 					}
 				}
-
 				return TOptional<FTableCellValue>();
 			}
 		};
@@ -537,29 +494,20 @@ void FMemAllocTable::AddDefaultColumns()
 
 		Column.SetDataType(ETableCellDataType::Int64);
 
-		class FAllocThreadValueGetter : public FTableCellValueGetter
+		class FAllocThreadValueGetter : public FAggregatedValueGetter
 		{
 		public:
-			virtual const TOptional<FTableCellValue> GetValue(const FTableColumn& Column, const FBaseTreeNode& Node) const override
+			virtual const TOptional<FTableCellValue> GetNodeValue(const FBaseTreeNode& Node) const override
 			{
-				if (Node.IsGroup())
-				{
-					const FTableTreeNode& NodePtr = static_cast<const FTableTreeNode&>(Node);
-					if (NodePtr.HasAggregatedValue(Column.GetId()))
-					{
-						return NodePtr.GetAggregatedValue(Column.GetId());
-					}
-				}
-				else //if (Node->Is<FMemAllocNode>())
+				if (Node.Is<FMemAllocNode>())
 				{
-					const FMemAllocNode& MemAllocNode = static_cast<const FMemAllocNode&>(Node);
+					const FMemAllocNode& MemAllocNode = Node.As<FMemAllocNode>();
 					const FMemoryAlloc* Alloc = MemAllocNode.GetMemAlloc();
 					if (Alloc)
 					{
 						return FTableCellValue(static_cast<int64>(Alloc->GetFreeThreadId()));
 					}
 				}
-
 				return TOptional<FTableCellValue>();
 			}
 		};
@@ -593,29 +541,20 @@ void FMemAllocTable::AddDefaultColumns()
 
 		Column.SetDataType(ETableCellDataType::Int64);
 
-		class FMemAllocAddressValueGetter : public FTableCellValueGetter
+		class FMemAllocAddressValueGetter : public FAggregatedValueGetter
 		{
 		public:
-			virtual const TOptional<FTableCellValue> GetValue(const FTableColumn& Column, const FBaseTreeNode& Node) const override
+			virtual const TOptional<FTableCellValue> GetNodeValue(const FBaseTreeNode& Node) const override
 			{
-				if (Node.IsGroup())
-				{
-					const FTableTreeNode& NodePtr = static_cast<const FTableTreeNode&>(Node);
-					if (NodePtr.HasAggregatedValue(Column.GetId()))
-					{
-						return NodePtr.GetAggregatedValue(Column.GetId());
-					}
-				}
-				else //if (Node->Is<FMemAllocNode>())
+				if (Node.Is<FMemAllocNode>())
 				{
-					const FMemAllocNode& MemAllocNode = static_cast<const FMemAllocNode&>(Node);
+					const FMemAllocNode& MemAllocNode = Node.As<FMemAllocNode>();
 					const FMemoryAlloc* Alloc = MemAllocNode.GetMemAlloc();
 					if (Alloc)
 					{
 						return FTableCellValue(static_cast<int64>(Alloc->GetAddress()));
 					}
 				}
-
 				return TOptional<FTableCellValue>();
 			}
 		};
@@ -649,30 +588,21 @@ void FMemAllocTable::AddDefaultColumns()
 
 		Column.SetDataType(ETableCellDataType::Int64);
 
-		class FMemAllocPageValueGetter : public FTableCellValueGetter
+		class FMemAllocPageValueGetter : public FAggregatedValueGetter
 		{
 		public:
-			virtual const TOptional<FTableCellValue> GetValue(const FTableColumn& Column, const FBaseTreeNode& Node) const override
+			virtual const TOptional<FTableCellValue> GetNodeValue(const FBaseTreeNode& Node) const override
 			{
-				if (Node.IsGroup())
-				{
-					const FTableTreeNode& NodePtr = static_cast<const FTableTreeNode&>(Node);
-					if (NodePtr.HasAggregatedValue(Column.GetId()))
-					{
-						return NodePtr.GetAggregatedValue(Column.GetId());
-					}
-				}
-				else
+				if (Node.Is<FMemAllocNode>())
 				{
-					const FMemAllocNode& MemAllocNode = static_cast<const FMemAllocNode&>(Node);
-					const FMemAllocTable& MemAllocTable = MemAllocNode.GetMemTableChecked();
+					const FMemAllocNode& MemAllocNode = Node.As<FMemAllocNode>();
 					const FMemoryAlloc* Alloc = MemAllocNode.GetMemAlloc();
 					if (Alloc)
 					{
+						const FMemAllocTable& MemAllocTable = MemAllocNode.GetMemTableChecked();
 						return FTableCellValue(static_cast<int64>(MemAllocTable.GetAddressPage(Alloc->GetAddress())));
 					}
 				}
-
 				return TOptional<FTableCellValue>();
 			}
 		};
@@ -706,29 +636,20 @@ void FMemAllocTable::AddDefaultColumns()
 
 		Column.SetDataType(ETableCellDataType::Int64);
 
-		class FMemAllocCountValueGetter : public FTableCellValueGetter
+		class FMemAllocCountValueGetter : public FAggregatedValueGetter
 		{
 		public:
-			virtual const TOptional<FTableCellValue> GetValue(const FTableColumn& Column, const FBaseTreeNode& Node) const override
+			virtual const TOptional<FTableCellValue> GetNodeValue(const FBaseTreeNode& Node) const override
 			{
-				if (Node.IsGroup())
-				{
-					const FTableTreeNode& NodePtr = static_cast<const FTableTreeNode&>(Node);
-					if (NodePtr.HasAggregatedValue(Column.GetId()))
-					{
-						return NodePtr.GetAggregatedValue(Column.GetId());
-					}
-				}
-				else //if (Node->Is<FMemAllocNode>())
+				if (Node.Is<FMemAllocNode>())
 				{
-					const FMemAllocNode& MemAllocNode = static_cast<const FMemAllocNode&>(Node);
+					const FMemAllocNode& MemAllocNode = Node.As<FMemAllocNode>();
 					const FMemoryAlloc* Alloc = MemAllocNode.GetMemAlloc();
 					if (Alloc)
 					{
 						return FTableCellValue(static_cast<int64>(1));
 					}
 				}
-
 				return TOptional<FTableCellValue>();
 			}
 		};
@@ -765,33 +686,23 @@ void FMemAllocTable::AddDefaultColumns()
 
 		Column.SetDataType(ETableCellDataType::Int64);
 
-		class FSwapCompressedSizeColumnValueGetter : public FTableCellValueGetter
+		class FSwapCompressedSizeColumnValueGetter : public FAggregatedValueGetter
 		{
 		public:
-			virtual const TOptional<FTableCellValue> GetValue(const FTableColumn& Column, const FBaseTreeNode& Node) const override
+			virtual const TOptional<FTableCellValue> GetNodeValue(const FBaseTreeNode& Node) const override
 			{
-				if (Node.IsGroup())
-				{
-					const FTableTreeNode& NodePtr = static_cast<const FTableTreeNode&>(Node);
-					if (NodePtr.HasAggregatedValue(Column.GetId()))
-					{
-						return NodePtr.GetAggregatedValue(Column.GetId());
-					}
-				}
-				else if (Node.Is<FMemAllocNode>())
+				if (Node.Is<FMemAllocNode>())
 				{
-					const FMemAllocNode& MemAllocNode = static_cast<const FMemAllocNode&>(Node);
+					const FMemAllocNode& MemAllocNode = Node.As<FMemAllocNode>();
 					const FMemoryAlloc* Alloc = MemAllocNode.GetMemAlloc();
 					if (Alloc && Alloc->IsSwap())
 					{
 						return FTableCellValue(static_cast<int64>(Alloc->GetSize()));
 					}
 				}
-
 				return TOptional<FTableCellValue>();
 			}
 		};
-
 		TSharedRef<ITableCellValueGetter> Getter = MakeShared<FSwapCompressedSizeColumnValueGetter>();
 		Column.SetValueGetter(Getter);
 
@@ -828,29 +739,19 @@ void FMemAllocTable::AddDefaultColumns()
 
 		Column.SetDataType(ETableCellDataType::Int64);
 
-		class FSizeInSwapColumnValueGetter : public FTableCellValueGetter
+		class FSizeInSwapColumnValueGetter : public FAggregatedValueGetter
 		{
 		public:
-			virtual const TOptional<FTableCellValue> GetValue(const FTableColumn& Column, const FBaseTreeNode& Node) const override
+			virtual const TOptional<FTableCellValue> GetNodeValue(const FBaseTreeNode& Node) const override
 			{
-				if (Node.IsGroup())
+				if (Node.Is<FMemAllocInSwapNode>())
 				{
-					const FTableTreeNode& NodePtr = static_cast<const FTableTreeNode&>(Node);
-					if (NodePtr.HasAggregatedValue(Column.GetId()))
-					{
-						return NodePtr.GetAggregatedValue(Column.GetId());
-					}
-				}
-				else if (Node.Is<FMemAllocInSwapNode>())
-				{
-					const FMemAllocInSwapNode& MemAllocInSwapNode = static_cast<const FMemAllocInSwapNode&>(Node);
+					const FMemAllocInSwapNode& MemAllocInSwapNode = Node.As<FMemAllocInSwapNode>();
 					return FTableCellValue(static_cast<int64>(MemAllocInSwapNode.GetBytesInSwapPage()));
 				}
-
 				return TOptional<FTableCellValue>();
 			}
 		};
-
 		TSharedRef<ITableCellValueGetter> Getter = MakeShared<FSizeInSwapColumnValueGetter>();
 		Column.SetValueGetter(Getter);
 
@@ -887,28 +788,21 @@ void FMemAllocTable::AddDefaultColumns()
 
 		Column.SetDataType(ETableCellDataType::Int64);
 
-		class FMemAllocSizeValueGetter : public FTableCellValueGetter
+		class FMemAllocSizeValueGetter : public FAggregatedValueGetter
 		{
 		public:
-			virtual const TOptional<FTableCellValue> GetValue(const FTableColumn& Column, const FBaseTreeNode& Node) const override
+			virtual const TOptional<FTableCellValue> GetNodeValue(const FBaseTreeNode& Node) const override
 			{
-				if (Node.IsGroup())
+				if (Node.Is<FMemAllocNode>())
 				{
-					const FTableTreeNode& NodePtr = static_cast<const FTableTreeNode&>(Node);
-					if (NodePtr.HasAggregatedValue(Column.GetId()))
-					{
-						return NodePtr.GetAggregatedValue(Column.GetId());
-					}
-				}
-				else //if (Node->Is<FMemAllocNode>())
-				{
-					const FMemAllocNode& MemAllocNode = static_cast<const FMemAllocNode&>(Node);
+					const FMemAllocNode& MemAllocNode = Node.As<FMemAllocNode>();
 					const FMemoryAlloc* Alloc = MemAllocNode.GetMemAlloc();
 					if (Alloc)
 					{
 						if (Alloc->IsSwap())
 						{
-							return FTableCellValue(static_cast<int64>(MemAllocNode.GetMemTableChecked().GetPlatformPageSize()));
+							const FMemAllocTable& MemAllocTable = MemAllocNode.GetMemTableChecked();
+							return FTableCellValue(static_cast<int64>(MemAllocTable.GetPlatformPageSize()));
 						}
 						else
 						{
@@ -916,7 +810,6 @@ void FMemAllocTable::AddDefaultColumns()
 						}
 					}
 				}
-
 				return TOptional<FTableCellValue>();
 			}
 		};
@@ -965,15 +858,16 @@ void FMemAllocTable::AddDefaultColumns()
 				{
 					if (Node.Is<FMemTagTableTreeNode>())
 					{
-						int64 LLMSize = Node.As<FMemTagTableTreeNode>().GetLLMSize();
+						const FMemTagTableTreeNode& MemTagNode = Node.As<FMemTagTableTreeNode>();
+						int64 LLMSize = MemTagNode.GetLLMSize();
 						return FTableCellValue(LLMSize);
 					}
-					else
+					else if (Node.Is<FTableTreeNode>())
 					{
-						const FTableTreeNode& NodePtr = static_cast<const FTableTreeNode&>(Node);
-						if (NodePtr.HasAggregatedValue(Column.GetId()))
+						const FTableTreeNode& TableTreeNode = Node.As<FTableTreeNode>();
+						if (TableTreeNode.HasAggregatedValue(Column.GetId()))
 						{
-							return NodePtr.GetAggregatedValue(Column.GetId());
+							return TableTreeNode.GetAggregatedValue(Column.GetId());
 						}
 					}
 				}
@@ -1025,12 +919,12 @@ void FMemAllocTable::AddDefaultColumns()
 				{
 					if (Node.Is<FMemTagTableTreeNode>())
 					{
-						int64 LLMSize = Node.As<FMemTagTableTreeNode>().GetLLMSize();
+						const FMemTagTableTreeNode& MemTagNode = Node.As<FMemTagTableTreeNode>();
+						int64 LLMSize = MemTagNode.GetLLMSize();
 						int64 Size = 0;
-						const FTableTreeNode& NodePtr = static_cast<const FTableTreeNode&>(Node);
-						if (NodePtr.HasAggregatedValue(FMemAllocTableColumns::SizeColumnId))
+						if (MemTagNode.HasAggregatedValue(FMemAllocTableColumns::SizeColumnId))
 						{
-							const FTableCellValue& Value = NodePtr.GetAggregatedValue(FMemAllocTableColumns::SizeColumnId);
+							const FTableCellValue& Value = MemTagNode.GetAggregatedValue(FMemAllocTableColumns::SizeColumnId);
 							if (Value.DataType == ETableCellDataType::Int64)
 							{
 								Size = Value.AsInt64();
@@ -1038,12 +932,12 @@ void FMemAllocTable::AddDefaultColumns()
 						}
 						return FTableCellValue(Size - LLMSize);
 					}
-					else
+					else if (Node.Is<FTableTreeNode>())
 					{
-						const FTableTreeNode& NodePtr = static_cast<const FTableTreeNode&>(Node);
-						if (NodePtr.HasAggregatedValue(Column.GetId()))
+						const FTableTreeNode& TableTreeNode = Node.As<FTableTreeNode>();
+						if (TableTreeNode.HasAggregatedValue(Column.GetId()))
 						{
-							return NodePtr.GetAggregatedValue(Column.GetId());
+							return TableTreeNode.GetAggregatedValue(Column.GetId());
 						}
 					}
 				}
@@ -1088,29 +982,20 @@ void FMemAllocTable::AddDefaultColumns()
 
 		Column.SetDataType(ETableCellDataType::CString);
 
-		class FMemTagValueGetter : public FTableCellValueGetter
+		class FMemTagValueGetter : public FAggregatedValueGetter
 		{
 		public:
-			virtual const TOptional<FTableCellValue> GetValue(const FTableColumn& Column, const FBaseTreeNode& Node) const
+			virtual const TOptional<FTableCellValue> GetNodeValue(const FBaseTreeNode& Node) const
 			{
-				if (Node.IsGroup())
-				{
-					const FTableTreeNode& NodePtr = static_cast<const FTableTreeNode&>(Node);
-					if (NodePtr.HasAggregatedValue(Column.GetId()))
-					{
-						return NodePtr.GetAggregatedValue(Column.GetId());
-					}
-				}
-				else //if (Node->Is<FMemAllocNode>())
+				if (Node.Is<FMemAllocNode>())
 				{
-					const FMemAllocNode& MemAllocNode = static_cast<const FMemAllocNode&>(Node);
+					const FMemAllocNode& MemAllocNode = Node.As<FMemAllocNode>();
 					const FMemoryAlloc* Alloc = MemAllocNode.GetMemAlloc();
 					if (Alloc)
 					{
 						return FTableCellValue(Alloc->GetTag());
 					}
 				}
-
 				return TOptional<FTableCellValue>();
 			}
 		};
@@ -1146,29 +1031,20 @@ void FMemAllocTable::AddDefaultColumns()
 
 		Column.SetDataType(ETableCellDataType::CString);
 
-		class FPackageMetadataValueGetter : public FTableCellValueGetter
+		class FPackageMetadataValueGetter : public FAggregatedValueGetter
 		{
 		public:
-			virtual const TOptional<FTableCellValue> GetValue(const FTableColumn& Column, const FBaseTreeNode& Node) const
+			virtual const TOptional<FTableCellValue> GetNodeValue(const FBaseTreeNode& Node) const
 			{
-				if (Node.IsGroup())
-				{
-					const FTableTreeNode& NodePtr = static_cast<const FTableTreeNode&>(Node);
-					if (NodePtr.HasAggregatedValue(Column.GetId()))
-					{
-						return NodePtr.GetAggregatedValue(Column.GetId());
-					}
-				}
-				else //if (Node->Is<FMemAllocNode>())
+				if (Node.Is<FMemAllocNode>())
 				{
-					const FMemAllocNode& MemAllocNode = static_cast<const FMemAllocNode&>(Node);
+					const FMemAllocNode& MemAllocNode = Node.As<FMemAllocNode>();
 					const FMemoryAlloc* Alloc = MemAllocNode.GetMemAlloc();
 					if (Alloc)
 					{
 						return FTableCellValue(Alloc->GetPackage());
 					}
 				}
-
 				return TOptional<FTableCellValue>();
 			}
 		};
@@ -1204,29 +1080,20 @@ void FMemAllocTable::AddDefaultColumns()
 
 		Column.SetDataType(ETableCellDataType::CString);
 
-		class FAssetMetadataValueGetter : public FTableCellValueGetter
+		class FAssetMetadataValueGetter : public FAggregatedValueGetter
 		{
 		public:
-			virtual const TOptional<FTableCellValue> GetValue(const FTableColumn& Column, const FBaseTreeNode& Node) const
+			virtual const TOptional<FTableCellValue> GetNodeValue(const FBaseTreeNode& Node) const
 			{
-				if (Node.IsGroup())
-				{
-					const FTableTreeNode& NodePtr = static_cast<const FTableTreeNode&>(Node);
-					if (NodePtr.HasAggregatedValue(Column.GetId()))
-					{
-						return NodePtr.GetAggregatedValue(Column.GetId());
-					}
-				}
-				else //if (Node->Is<FMemAllocNode>())
+				if (Node.Is<FMemAllocNode>())
 				{
-					const FMemAllocNode& MemAllocNode = static_cast<const FMemAllocNode&>(Node);
+					const FMemAllocNode& MemAllocNode = Node.As<FMemAllocNode>();
 					const FMemoryAlloc* Alloc = MemAllocNode.GetMemAlloc();
 					if (Alloc)
 					{
 						return FTableCellValue(Alloc->GetAsset());
 					}
 				}
-
 				return TOptional<FTableCellValue>();
 			}
 		};
@@ -1262,29 +1129,20 @@ void FMemAllocTable::AddDefaultColumns()
 
 		Column.SetDataType(ETableCellDataType::CString);
 
-		class FClassNameMetadataValueGetter : public FTableCellValueGetter
+		class FClassNameMetadataValueGetter : public FAggregatedValueGetter
 		{
 		public:
-			virtual const TOptional<FTableCellValue> GetValue(const FTableColumn& Column, const FBaseTreeNode& Node) const
+			virtual const TOptional<FTableCellValue> GetNodeValue(const FBaseTreeNode& Node) const
 			{
-				if (Node.IsGroup())
+				if (Node.Is<FMemAllocNode>())
 				{
-					const FTableTreeNode& NodePtr = static_cast<const FTableTreeNode&>(Node);
-					if (NodePtr.HasAggregatedValue(Column.GetId()))
-					{
-						return NodePtr.GetAggregatedValue(Column.GetId());
-					}
-				}
-				else //if (Node->Is<FMemAllocNode>())
-				{
-					const FMemAllocNode& MemAllocNode = static_cast<const FMemAllocNode&>(Node);
+					const FMemAllocNode& MemAllocNode = Node.As<FMemAllocNode>();
 					const FMemoryAlloc* Alloc = MemAllocNode.GetMemAlloc();
 					if (Alloc)
 					{
 						return FTableCellValue(Alloc->GetClassName());
 					}
 				}
-
 				return TOptional<FTableCellValue>();
 			}
 		};
@@ -1356,7 +1214,11 @@ void FMemAllocTable::AddDefaultColumns()
 
 		virtual TSharedPtr<IToolTip> GetCustomTooltip(const FTableColumn& Column, const FBaseTreeNode& Node) const override
 		{
-			const FMemAllocNode& MemAllocNode = static_cast<const FMemAllocNode&>(Node);
+			if (!Node.Is<FMemAllocNode>())
+			{
+				return FTableCellValueFormatter::GetCustomTooltip(Column, Node);
+			}
+			const FMemAllocNode& MemAllocNode = Node.As<const FMemAllocNode>();
 
 			FText HeaderText = GetHeaderText(MemAllocNode);
 			return SNew(SToolTip)
@@ -1388,8 +1250,15 @@ void FMemAllocTable::AddDefaultColumns()
 
 		virtual FText CopyTooltip(const FTableColumn& Column, const FBaseTreeNode& Node) const override
 		{
-			const FMemAllocNode& MemAllocNode = static_cast<const FMemAllocNode&>(Node);
-			return MemAllocNode.GetFullCallstack(CallstackType);
+			if (Node.Is<FMemAllocNode>())
+			{
+				const FMemAllocNode& MemAllocNode = Node.As<const FMemAllocNode>();
+				return MemAllocNode.GetFullCallstack(CallstackType);
+			}
+			else
+			{
+				return FTableCellValueFormatter::CopyTooltip(Column, Node);
+			}
 		}
 
 	protected:
@@ -1408,7 +1277,11 @@ void FMemAllocTable::AddDefaultColumns()
 
 		virtual TSharedPtr<IToolTip> GetCustomTooltip(const FTableColumn& Column, const FBaseTreeNode& Node) const override
 		{
-			const FMemAllocNode& MemAllocNode = static_cast<const FMemAllocNode&>(Node);
+			if (!Node.Is<FMemAllocNode>())
+			{
+				return FTextValueFormatter::GetCustomTooltip(Column, Node);
+			}
+			const FMemAllocNode& MemAllocNode = Node.As<const FMemAllocNode>();
 
 			return SNew(SToolTip)
 				.Visibility(TAttribute<EVisibility>::Create(TAttribute<EVisibility>::FGetter::CreateStatic(&FTableCellValueFormatter::GetTooltipVisibility)))
@@ -1454,7 +1327,11 @@ void FMemAllocTable::AddDefaultColumns()
 
 		virtual TSharedPtr<IToolTip> GetCustomTooltip(const FTableColumn& Column, const FBaseTreeNode& Node) const override
 		{
-			const FMemAllocNode& MemAllocNode = static_cast<const FMemAllocNode&>(Node);
+			if (!Node.Is<FMemAllocNode>())
+			{
+				return FTextValueFormatter::GetCustomTooltip(Column, Node);
+			}
+			const FMemAllocNode& MemAllocNode = Node.As<const FMemAllocNode>();
 
 			return SNew(SToolTip)
 				.Visibility(TAttribute<EVisibility>::Create(TAttribute<EVisibility>::FGetter::CreateStatic(&FTableCellValueFormatter::GetTooltipVisibility)))
@@ -1507,38 +1384,27 @@ void FMemAllocTable::AddDefaultColumns()
 
 		Column.SetDataType(ETableCellDataType::Text);
 
-		class FFunctionValueGetter : public FTableCellValueGetter
+		class FFunctionValueGetter : public FAggregatedValueGetter
 		{
 		public:
-			virtual const TOptional<FTableCellValue> GetValue(const FTableColumn& Column, const FBaseTreeNode& Node) const override
+			virtual const TOptional<FTableCellValue> GetNodeValue(const FBaseTreeNode& Node) const override
 			{
-				if (Node.IsGroup())
-				{
-					const FTableTreeNode& NodePtr = static_cast<const FTableTreeNode&>(Node);
-					if (NodePtr.HasAggregatedValue(Column.GetId()))
-					{
-						return NodePtr.GetAggregatedValue(Column.GetId());
-					}
-				}
-				else //if (Node->Is<FMemAllocNode>())
+				if (Node.Is<FMemAllocNode>())
 				{
-					const FMemAllocNode& MemAllocNode = static_cast<const FMemAllocNode&>(Node);
+					const FMemAllocNode& MemAllocNode = Node.As<FMemAllocNode>();
 					const uint64 CallstackId = MemAllocNode.GetAllocCallstackId();
 					const FText CallstackText = MemAllocNode.GetTopFunction(FMemAllocNode::ECallstackType::AllocCallstack);
 					return FTableCellValue(CallstackText, CallstackId);
 				}
-
 				return TOptional<FTableCellValue>();
 			}
-
 			virtual uint64 GetValueId(const FTableColumn& Column, const FBaseTreeNode& Node) const override
 			{
-				if (!Node.IsGroup()) //if (Node->Is<FMemAllocNode>())
+				if (Node.Is<FMemAllocNode>())
 				{
-					const FMemAllocNode& MemAllocNode = static_cast<const FMemAllocNode&>(Node);
+					const FMemAllocNode& MemAllocNode = Node.As<FMemAllocNode>();
 					return MemAllocNode.GetAllocCallstackId();
 				}
-
 				return 0;
 			}
 		};
@@ -1572,35 +1438,25 @@ void FMemAllocTable::AddDefaultColumns()
 
 		Column.SetDataType(ETableCellDataType::Text);
 
-		class FSourceValueGetter : public FTableCellValueGetter
+		class FSourceValueGetter : public FAggregatedValueGetter
 		{
 		public:
-			virtual const TOptional<FTableCellValue> GetValue(const FTableColumn& Column, const FBaseTreeNode& Node) const override
+			virtual const TOptional<FTableCellValue> GetNodeValue(const FBaseTreeNode& Node) const override
 			{
-				if (Node.IsGroup())
+				if (Node.Is<FMemAllocNode>())
 				{
-					const FTableTreeNode& NodePtr = static_cast<const FTableTreeNode&>(Node);
-					if (NodePtr.HasAggregatedValue(Column.GetId()))
-					{
-						return NodePtr.GetAggregatedValue(Column.GetId());
-					}
-				}
-				else //if (Node->Is<FMemAllocNode>())
-				{
-					const FMemAllocNode& MemAllocNode = static_cast<const FMemAllocNode&>(Node);
+					const FMemAllocNode& MemAllocNode = Node.As<FMemAllocNode>();
 					const uint64 CallstackId = MemAllocNode.GetAllocCallstackId();
 					const FText CallstackText = MemAllocNode.GetTopSourceFile(FMemAllocNode::ECallstackType::AllocCallstack);
 					return FTableCellValue(CallstackText, CallstackId);
 				}
-
 				return TOptional<FTableCellValue>();
 			}
-
 			virtual uint64 GetValueId(const FTableColumn& Column, const FBaseTreeNode& Node) const override
 			{
-				if (!Node.IsGroup()) //if (Node->Is<FMemAllocNode>())
+				if (Node.Is<FMemAllocNode>())
 				{
-					const FMemAllocNode& MemAllocNode = static_cast<const FMemAllocNode&>(Node);
+					const FMemAllocNode& MemAllocNode = Node.As<FMemAllocNode>();
 					return MemAllocNode.GetAllocCallstackId();
 				}
 				return 0;
@@ -1636,37 +1492,27 @@ void FMemAllocTable::AddDefaultColumns()
 
 		Column.SetDataType(ETableCellDataType::Int64);
 
-		class FCallstackIdValueGetter : public FTableCellValueGetter
+		class FCallstackIdValueGetter : public FAggregatedValueGetter
 		{
 		public:
-			virtual const TOptional<FTableCellValue> GetValue(const FTableColumn& Column, const FBaseTreeNode& Node) const override
+			virtual const TOptional<FTableCellValue> GetNodeValue(const FBaseTreeNode& Node) const override
 			{
-				if (Node.IsGroup())
+				if (Node.Is<FMemAllocNode>())
 				{
-					const FTableTreeNode& NodePtr = static_cast<const FTableTreeNode&>(Node);
-					if (NodePtr.HasAggregatedValue(Column.GetId()))
-					{
-						return NodePtr.GetAggregatedValue(Column.GetId());
-					}
-				}
-				else //if (Node->Is<FMemAllocNode>())
-				{
-					const FMemAllocNode& MemAllocNode = static_cast<const FMemAllocNode&>(Node);
+					const FMemAllocNode& MemAllocNode = Node.As<FMemAllocNode>();
 					const FMemoryAlloc* Alloc = MemAllocNode.GetMemAlloc();
 					if (Alloc)
 					{
 						return FTableCellValue(int64(Alloc->GetAllocCallstackId()));
 					}
 				}
-
 				return TOptional<FTableCellValue>();
 			}
-
 			virtual uint64 GetValueId(const FTableColumn& Column, const FBaseTreeNode& Node) const override
 			{
-				if (!Node.IsGroup()) //if (Node->Is<FMemAllocNode>())
+				if (Node.Is<FMemAllocNode>())
 				{
-					const FMemAllocNode& MemAllocNode = static_cast<const FMemAllocNode&>(Node);
+					const FMemAllocNode& MemAllocNode = Node.As<FMemAllocNode>();
 					return MemAllocNode.GetAllocCallstackId();
 				}
 				return 0;
@@ -1702,22 +1548,14 @@ void FMemAllocTable::AddDefaultColumns()
 
 		Column.SetDataType(ETableCellDataType::Int64);
 
-		class FCallstackSizeValueGetter : public FTableCellValueGetter
+		class FCallstackSizeValueGetter : public FAggregatedValueGetter
 		{
 		public:
-			virtual const TOptional<FTableCellValue> GetValue(const FTableColumn& Column, const FBaseTreeNode& Node) const override
+			virtual const TOptional<FTableCellValue> GetNodeValue(const FBaseTreeNode& Node) const override
 			{
-				if (Node.IsGroup())
+				if (Node.Is<FMemAllocNode>())
 				{
-					const FTableTreeNode& NodePtr = static_cast<const FTableTreeNode&>(Node);
-					if (NodePtr.HasAggregatedValue(Column.GetId()))
-					{
-						return NodePtr.GetAggregatedValue(Column.GetId());
-					}
-				}
-				else //if (Node->Is<FMemAllocNode>())
-				{
-					const FMemAllocNode& MemAllocNode = static_cast<const FMemAllocNode&>(Node);
+					const FMemAllocNode& MemAllocNode = Node.As<FMemAllocNode>();
 					const FMemoryAlloc* Alloc = MemAllocNode.GetMemAlloc();
 					if (Alloc)
 					{
@@ -1726,15 +1564,13 @@ void FMemAllocTable::AddDefaultColumns()
 						return FTableCellValue(CallstackSize);
 					}
 				}
-
 				return TOptional<FTableCellValue>();
 			}
-
 			virtual uint64 GetValueId(const FTableColumn& Column, const FBaseTreeNode& Node) const override
 			{
-				if (!Node.IsGroup()) //if (Node->Is<FMemAllocNode>())
+				if (Node.Is<FMemAllocNode>())
 				{
-					const FMemAllocNode& MemAllocNode = static_cast<const FMemAllocNode&>(Node);
+					const FMemAllocNode& MemAllocNode = Node.As<FMemAllocNode>();
 					return MemAllocNode.GetAllocCallstackId();
 				}
 				return 0;
@@ -1770,38 +1606,27 @@ void FMemAllocTable::AddDefaultColumns()
 
 		Column.SetDataType(ETableCellDataType::Text);
 
-		class FFunctionValueGetter : public FTableCellValueGetter
+		class FFunctionValueGetter : public FAggregatedValueGetter
 		{
 		public:
-			virtual const TOptional<FTableCellValue> GetValue(const FTableColumn& Column, const FBaseTreeNode& Node) const override
+			virtual const TOptional<FTableCellValue> GetNodeValue(const FBaseTreeNode& Node) const override
 			{
-				if (Node.IsGroup())
-				{
-					const FTableTreeNode& NodePtr = static_cast<const FTableTreeNode&>(Node);
-					if (NodePtr.HasAggregatedValue(Column.GetId()))
-					{
-						return NodePtr.GetAggregatedValue(Column.GetId());
-					}
-				}
-				else //if (Node->Is<FMemAllocNode>())
+				if (Node.Is<FMemAllocNode>())
 				{
-					const FMemAllocNode& MemAllocNode = static_cast<const FMemAllocNode&>(Node);
+					const FMemAllocNode& MemAllocNode = Node.As<FMemAllocNode>();
 					const uint64 CallstackId = MemAllocNode.GetFreeCallstackId();
 					const FText CallstackText = MemAllocNode.GetTopFunction(FMemAllocNode::ECallstackType::FreeCallstack);
 					return FTableCellValue(CallstackText, CallstackId);
 				}
-
 				return TOptional<FTableCellValue>();
 			}
-
 			virtual uint64 GetValueId(const FTableColumn& Column, const FBaseTreeNode& Node) const override
 			{
-				if (!Node.IsGroup()) //if (Node->Is<FMemAllocNode>())
+				if (Node.Is<FMemAllocNode>())
 				{
-					const FMemAllocNode& MemAllocNode = static_cast<const FMemAllocNode&>(Node);
+					const FMemAllocNode& MemAllocNode = Node.As<FMemAllocNode>();
 					return MemAllocNode.GetFreeCallstackId();
 				}
-
 				return 0;
 			}
 		};
@@ -1835,35 +1660,25 @@ void FMemAllocTable::AddDefaultColumns()
 
 		Column.SetDataType(ETableCellDataType::Text);
 
-		class FSourceValueGetter : public FTableCellValueGetter
+		class FSourceValueGetter : public FAggregatedValueGetter
 		{
 		public:
-			virtual const TOptional<FTableCellValue> GetValue(const FTableColumn& Column, const FBaseTreeNode& Node) const override
+			virtual const TOptional<FTableCellValue> GetNodeValue(const FBaseTreeNode& Node) const override
 			{
-				if (Node.IsGroup())
-				{
-					const FTableTreeNode& NodePtr = static_cast<const FTableTreeNode&>(Node);
-					if (NodePtr.HasAggregatedValue(Column.GetId()))
-					{
-						return NodePtr.GetAggregatedValue(Column.GetId());
-					}
-				}
-				else //if (Node->Is<FMemAllocNode>())
+				if (Node.Is<FMemAllocNode>())
 				{
-					const FMemAllocNode& MemAllocNode = static_cast<const FMemAllocNode&>(Node);
+					const FMemAllocNode& MemAllocNode = Node.As<FMemAllocNode>();
 					const uint64 CallstackId = MemAllocNode.GetFreeCallstackId();
 					const FText CallstackText = MemAllocNode.GetTopSourceFile(FMemAllocNode::ECallstackType::FreeCallstack);
 					return FTableCellValue(CallstackText, CallstackId);
 				}
-
 				return TOptional<FTableCellValue>();
 			}
-
 			virtual uint64 GetValueId(const FTableColumn& Column, const FBaseTreeNode& Node) const override
 			{
-				if (!Node.IsGroup()) //if (Node->Is<FMemAllocNode>())
+				if (Node.Is<FMemAllocNode>())
 				{
-					const FMemAllocNode& MemAllocNode = static_cast<const FMemAllocNode&>(Node);
+					const FMemAllocNode& MemAllocNode = Node.As<FMemAllocNode>();
 					return MemAllocNode.GetFreeCallstackId();
 				}
 				return 0;
@@ -1899,37 +1714,27 @@ void FMemAllocTable::AddDefaultColumns()
 
 		Column.SetDataType(ETableCellDataType::Int64);
 
-		class FCallstackIdValueGetter : public FTableCellValueGetter
+		class FCallstackIdValueGetter : public FAggregatedValueGetter
 		{
 		public:
-			virtual const TOptional<FTableCellValue> GetValue(const FTableColumn& Column, const FBaseTreeNode& Node) const override
+			virtual const TOptional<FTableCellValue> GetNodeValue(const FBaseTreeNode& Node) const override
 			{
-				if (Node.IsGroup())
+				if (Node.Is<FMemAllocNode>())
 				{
-					const FTableTreeNode& NodePtr = static_cast<const FTableTreeNode&>(Node);
-					if (NodePtr.HasAggregatedValue(Column.GetId()))
-					{
-						return NodePtr.GetAggregatedValue(Column.GetId());
-					}
-				}
-				else //if (Node->Is<FMemAllocNode>())
-				{
-					const FMemAllocNode& MemAllocNode = static_cast<const FMemAllocNode&>(Node);
+					const FMemAllocNode& MemAllocNode = Node.As<FMemAllocNode>();
 					const FMemoryAlloc* Alloc = MemAllocNode.GetMemAlloc();
 					if (Alloc)
 					{
 						return FTableCellValue(int64(Alloc->GetFreeCallstackId()));
 					}
 				}
-
 				return TOptional<FTableCellValue>();
 			}
-
 			virtual uint64 GetValueId(const FTableColumn& Column, const FBaseTreeNode& Node) const override
 			{
-				if (!Node.IsGroup()) //if (Node->Is<FMemAllocNode>())
+				if (Node.Is<FMemAllocNode>())
 				{
-					const FMemAllocNode& MemAllocNode = static_cast<const FMemAllocNode&>(Node);
+					const FMemAllocNode& MemAllocNode = Node.As<FMemAllocNode>();
 					return MemAllocNode.GetFreeCallstackId();
 				}
 				return 0;
@@ -1965,22 +1770,14 @@ void FMemAllocTable::AddDefaultColumns()
 
 		Column.SetDataType(ETableCellDataType::Int64);
 
-		class FCallstackSizeValueGetter : public FTableCellValueGetter
+		class FCallstackSizeValueGetter : public FAggregatedValueGetter
 		{
 		public:
-			virtual const TOptional<FTableCellValue> GetValue(const FTableColumn& Column, const FBaseTreeNode& Node) const override
+			virtual const TOptional<FTableCellValue> GetNodeValue(const FBaseTreeNode& Node) const override
 			{
-				if (Node.IsGroup())
+				if (Node.Is<FMemAllocNode>())
 				{
-					const FTableTreeNode& NodePtr = static_cast<const FTableTreeNode&>(Node);
-					if (NodePtr.HasAggregatedValue(Column.GetId()))
-					{
-						return NodePtr.GetAggregatedValue(Column.GetId());
-					}
-				}
-				else //if (Node->Is<FMemAllocNode>())
-				{
-					const FMemAllocNode& MemAllocNode = static_cast<const FMemAllocNode&>(Node);
+					const FMemAllocNode& MemAllocNode = Node.As<FMemAllocNode>();
 					const FMemoryAlloc* Alloc = MemAllocNode.GetMemAlloc();
 					if (Alloc)
 					{
@@ -1989,15 +1786,13 @@ void FMemAllocTable::AddDefaultColumns()
 						return FTableCellValue(CallstackSize);
 					}
 				}
-
 				return TOptional<FTableCellValue>();
 			}
-
 			virtual uint64 GetValueId(const FTableColumn& Column, const FBaseTreeNode& Node) const override
 			{
-				if (!Node.IsGroup()) //if (Node->Is<FMemAllocNode>())
+				if (Node.Is<FMemAllocNode>())
 				{
-					const FMemAllocNode& MemAllocNode = static_cast<const FMemAllocNode&>(Node);
+					const FMemAllocNode& MemAllocNode = Node.As<FMemAllocNode>();
 					return MemAllocNode.GetFreeCallstackId();
 				}
 				return 0;
diff --git a/Engine/Source/Developer/TraceInsights/Private/Insights/MemoryProfiler/ViewModels/MemAllocTableViewPresets.cpp b/Engine/Source/Developer/TraceInsights/Private/Insights/MemoryProfiler/ViewModels/MemAllocTableViewPresets.cpp
new file mode 100644
--- /dev/null
+++ b/Engine/Source/Developer/TraceInsights/Private/Insights/MemoryProfiler/ViewModels/MemAllocTableViewPresets.cpp
@@ -0,0 +1,621 @@
+#include "MemAllocTableViewPresets.h"
+
+// TraceInsightsCore
+#include "InsightsCore/Table/ViewModels/TreeNodeGrouping.h"
+#include "InsightsCore/Table/Widgets/STableTreeView.h"
+
+// TraceInsights
+#include "Insights/MemoryProfiler/ViewModels/MemAllocGroupingByCallstack.h"
+#include "Insights/MemoryProfiler/ViewModels/MemAllocGroupingByHeap.h"
+#include "Insights/MemoryProfiler/ViewModels/MemAllocGroupingBySize.h"
+#include "Insights/MemoryProfiler/ViewModels/MemAllocGroupingBySwapPage.h"
+#include "Insights/MemoryProfiler/ViewModels/MemAllocGroupingByTag.h"
+#include "Insights/MemoryProfiler/ViewModels/MemAllocTable.h"
+#include "Insights/MemoryProfiler/Widgets/SMemAllocTableTreeView.h"
+
+#define LOCTEXT_NAMESPACE "UE::Insights::MemoryProfiler::SMemAllocTableTreeView"
+
+namespace UE::Insights::MemoryProfiler
+{
+
+////////////////////////////////////////////////////////////////////////////////////////////////////
+// Default View
+
+TSharedRef<ITableTreeViewPreset> FMemAllocTableViewPresets::CreateDefaultViewPreset(SMemAllocTableTreeView& TableTreeView)
+{
+	class FDefaultViewPreset : public ITableTreeViewPreset
+	{
+	public:
+		virtual FText GetName() const override
+		{
+			return LOCTEXT("Default_PresetName", "Default");
+		}
+		virtual FText GetToolTip() const override
+		{
+			return LOCTEXT("Default_PresetToolTip", "Default View\nConfigure the tree view to show default allocation info.");
+		}
+		virtual FName GetSortColumn() const override
+		{
+			return FTable::GetHierarchyColumnId();
+		}
+		virtual EColumnSortMode::Type GetSortMode() const override
+		{
+			return EColumnSortMode::Type::Ascending;
+		}
+		virtual void SetCurrentGroupings(const TArray<TSharedPtr<FTreeNodeGrouping>>& InAvailableGroupings, TArray<TSharedPtr<FTreeNodeGrouping>>& InOutCurrentGroupings) const override
+		{
+			InOutCurrentGroupings.Reset();
+
+			check(InAvailableGroupings[0]->Is<FTreeNodeGroupingFlat>());
+			InOutCurrentGroupings.Add(InAvailableGroupings[0]);
+		}
+		virtual void GetColumnConfigSet(TArray<FTableColumnConfig>& InOutConfigSet) const override
+		{
+			InOutConfigSet.Add({ FTable::GetHierarchyColumnId(),                    true, 200.0f });
+			InOutConfigSet.Add({ FMemAllocTableColumns::CountColumnId,              true, 100.0f });
+			InOutConfigSet.Add({ FMemAllocTableColumns::SizeColumnId,               true, 100.0f });
+			InOutConfigSet.Add({ FMemAllocTableColumns::TagColumnId,                true, 120.0f });
+			InOutConfigSet.Add({ FMemAllocTableColumns::AllocFunctionColumnId,      true, 550.0f });
+			InOutConfigSet.Add({ FMemAllocTableColumns::AllocCallstackSizeColumnId, true, 100.0f });
+		}
+	};
+	return MakeShared<FDefaultViewPreset>();
+}
+
+////////////////////////////////////////////////////////////////////////////////////////////////////
+// Detailed View
+
+TSharedRef<ITableTreeViewPreset> FMemAllocTableViewPresets::CreateDetailedViewPreset(SMemAllocTableTreeView& TableTreeView)
+{
+	class FDetailedViewPreset : public ITableTreeViewPreset
+	{
+	public:
+		virtual FText GetName() const override
+		{
+			return LOCTEXT("Detailed_PresetName", "Detailed");
+		}
+		virtual FText GetToolTip() const override
+		{
+			return LOCTEXT("Detailed_PresetToolTip", "Detailed View\nConfigure the tree view to show detailed allocation info.");
+		}
+		virtual FName GetSortColumn() const override
+		{
+			return FTable::GetHierarchyColumnId();
+		}
+		virtual EColumnSortMode::Type GetSortMode() const override
+		{
+			return EColumnSortMode::Type::Ascending;
+		}
+		virtual void SetCurrentGroupings(const TArray<TSharedPtr<FTreeNodeGrouping>>& InAvailableGroupings, TArray<TSharedPtr<FTreeNodeGrouping>>& InOutCurrentGroupings) const override
+		{
+			InOutCurrentGroupings.Reset();
+
+			check(InAvailableGroupings[0]->Is<FTreeNodeGroupingFlat>());
+			InOutCurrentGroupings.Add(InAvailableGroupings[0]);
+		}
+		virtual void GetColumnConfigSet(TArray<FTableColumnConfig>& InOutConfigSet) const override
+		{
+			InOutConfigSet.Add({ FTable::GetHierarchyColumnId(),                    true, 200.0f });
+			InOutConfigSet.Add({ FMemAllocTableColumns::StartEventIndexColumnId,    true, 100.0f });
+			InOutConfigSet.Add({ FMemAllocTableColumns::EndEventIndexColumnId,      true, 100.0f });
+			InOutConfigSet.Add({ FMemAllocTableColumns::EventDistanceColumnId,      true, 100.0f });
+			InOutConfigSet.Add({ FMemAllocTableColumns::StartTimeColumnId,          true, 100.0f });
+			InOutConfigSet.Add({ FMemAllocTableColumns::EndTimeColumnId,            true, 100.0f });
+			InOutConfigSet.Add({ FMemAllocTableColumns::DurationColumnId,           true, 100.0f });
+			InOutConfigSet.Add({ FMemAllocTableColumns::AddressColumnId,            true, 120.0f });
+			InOutConfigSet.Add({ FMemAllocTableColumns::MemoryPageColumnId,         true, 120.0f });
+			InOutConfigSet.Add({ FMemAllocTableColumns::CountColumnId,              true, 100.0f });
+			InOutConfigSet.Add({ FMemAllocTableColumns::SizeColumnId,               true, 100.0f });
+			InOutConfigSet.Add({ FMemAllocTableColumns::TagColumnId,                true, 120.0f });
+			InOutConfigSet.Add({ FMemAllocTableColumns::AllocFunctionColumnId,      true, 550.0f });
+			InOutConfigSet.Add({ FMemAllocTableColumns::AllocSourceFileColumnId,    true, 550.0f });
+			InOutConfigSet.Add({ FMemAllocTableColumns::AllocCallstackSizeColumnId, true, 100.0f });
+		}
+	};
+	return MakeShared<FDetailedViewPreset>();
+}
+
+////////////////////////////////////////////////////////////////////////////////////////////////////
+// Heap Breakdown View
+
+TSharedRef<ITableTreeViewPreset> FMemAllocTableViewPresets::CreateHeapViewPreset(SMemAllocTableTreeView& TableTreeView)
+{
+	class FHeapViewPreset : public ITableTreeViewPreset
+	{
+	public:
+		virtual FText GetName() const override
+		{
+			return LOCTEXT("Heap_PresetName", "Heap");
+		}
+		virtual FText GetToolTip() const override
+		{
+			return LOCTEXT("Heap_PresetToolTip", "Heap Breakdown View\nConfigure the tree view to show a breakdown of allocations by their parent heap type.");
+		}
+		virtual FName GetSortColumn() const override
+		{
+			return FMemAllocTableColumns::SizeColumnId;
+		}
+		virtual EColumnSortMode::Type GetSortMode() const override
+		{
+			return EColumnSortMode::Type::Descending;
+		}
+		virtual void SetCurrentGroupings(const TArray<TSharedPtr<FTreeNodeGrouping>>& InAvailableGroupings, TArray<TSharedPtr<FTreeNodeGrouping>>& InOutCurrentGroupings) const override
+		{
+			InOutCurrentGroupings.Reset();
+
+			//check(InAvailableGroupings[0]->Is<FTreeNodeGroupingFlat>());
+			//InOutCurrentGroupings.Add(InAvailableGroupings[0]);
+
+			const TSharedPtr<FTreeNodeGrouping>* HeapGrouping = InAvailableGroupings.FindByPredicate(
+				[](TSharedPtr<FTreeNodeGrouping>& Grouping)
+				{
+					return Grouping->Is<FMemAllocGroupingByHeap>();
+				});
+			if (HeapGrouping)
+			{
+				InOutCurrentGroupings.Add(*HeapGrouping);
+			}
+		}
+		virtual void GetColumnConfigSet(TArray<FTableColumnConfig>& InOutConfigSet) const override
+		{
+			InOutConfigSet.Add({ FTable::GetHierarchyColumnId(),               true, 400.0f });
+			InOutConfigSet.Add({ FMemAllocTableColumns::CountColumnId,         true, 100.0f });
+			InOutConfigSet.Add({ FMemAllocTableColumns::SizeColumnId,          true, 100.0f });
+			InOutConfigSet.Add({ FMemAllocTableColumns::TagColumnId,           true, 200.0f });
+			InOutConfigSet.Add({ FMemAllocTableColumns::AllocFunctionColumnId, true, 200.0f });
+		}
+	};
+	return MakeShared<FHeapViewPreset>();
+}
+
+////////////////////////////////////////////////////////////////////////////////////////////////////
+// Size Breakdown View
+
+TSharedRef<ITableTreeViewPreset> FMemAllocTableViewPresets::CreateSizeViewPreset(SMemAllocTableTreeView& TableTreeView)
+{
+	class FSizeViewPreset : public ITableTreeViewPreset
+	{
+	public:
+		virtual FText GetName() const override
+		{
+			return LOCTEXT("Size_PresetName", "Size");
+		}
+		virtual FText GetToolTip() const override
+		{
+			return LOCTEXT("Size_PresetToolTip", "Size Breakdown View\nConfigure the tree view to show a breakdown of allocations by their size.");
+		}
+		virtual FName GetSortColumn() const override
+		{
+			return FMemAllocTableColumns::SizeColumnId;
+		}
+		virtual EColumnSortMode::Type GetSortMode() const override
+		{
+			return EColumnSortMode::Type::Descending;
+		}
+		virtual void SetCurrentGroupings(const TArray<TSharedPtr<FTreeNodeGrouping>>& InAvailableGroupings, TArray<TSharedPtr<FTreeNodeGrouping>>& InOutCurrentGroupings) const override
+		{
+			InOutCurrentGroupings.Reset();
+
+			check(InAvailableGroupings[0]->Is<FTreeNodeGroupingFlat>());
+			InOutCurrentGroupings.Add(InAvailableGroupings[0]);
+
+			const TSharedPtr<FTreeNodeGrouping>* SizeGrouping = InAvailableGroupings.FindByPredicate(
+				[](TSharedPtr<FTreeNodeGrouping>& Grouping)
+				{
+					return Grouping->Is<FMemAllocGroupingBySize>();
+				});
+			if (SizeGrouping)
+			{
+				InOutCurrentGroupings.Add(*SizeGrouping);
+			}
+		}
+		virtual void GetColumnConfigSet(TArray<FTableColumnConfig>& InOutConfigSet) const override
+		{
+			InOutConfigSet.Add({ FTable::GetHierarchyColumnId(),               true, 200.0f });
+			InOutConfigSet.Add({ FMemAllocTableColumns::AddressColumnId,       true, 120.0f });
+			InOutConfigSet.Add({ FMemAllocTableColumns::CountColumnId,         true, 100.0f });
+			InOutConfigSet.Add({ FMemAllocTableColumns::SizeColumnId,          true, 100.0f });
+			InOutConfigSet.Add({ FMemAllocTableColumns::TagColumnId,           true, 120.0f });
+			InOutConfigSet.Add({ FMemAllocTableColumns::AllocFunctionColumnId, true, 400.0f });
+		}
+	};
+	return MakeShared<FSizeViewPreset>();
+}
+
+////////////////////////////////////////////////////////////////////////////////////////////////////
+// Tag Breakdown View
+
+TSharedRef<ITableTreeViewPreset> FMemAllocTableViewPresets::CreateTagViewPreset(SMemAllocTableTreeView& TableTreeView)
+{
+	class FTagViewPreset : public ITableTreeViewPreset
+	{
+	public:
+		virtual FText GetName() const override
+		{
+			return LOCTEXT("Tag_PresetName", "Tags");
+		}
+		virtual FText GetToolTip() const override
+		{
+			return LOCTEXT("Tag_PresetToolTip", "Tag Breakdown View\nConfigure the tree view to show a breakdown of allocations by their LLM tag.");
+		}
+		virtual FName GetSortColumn() const override
+		{
+			return FTable::GetHierarchyColumnId();
+		}
+		virtual EColumnSortMode::Type GetSortMode() const override
+		{
+			return EColumnSortMode::Type::Ascending;
+		}
+		virtual void SetCurrentGroupings(const TArray<TSharedPtr<FTreeNodeGrouping>>& InAvailableGroupings, TArray<TSharedPtr<FTreeNodeGrouping>>& InOutCurrentGroupings) const override
+		{
+			InOutCurrentGroupings.Reset();
+
+			check(InAvailableGroupings[0]->Is<FTreeNodeGroupingFlat>());
+			InOutCurrentGroupings.Add(InAvailableGroupings[0]);
+
+			const TSharedPtr<FTreeNodeGrouping>* TagGrouping = InAvailableGroupings.FindByPredicate(
+				[](TSharedPtr<FTreeNodeGrouping>& Grouping)
+				{
+					return Grouping->Is<FMemAllocGroupingByTag>();
+				});
+			if (TagGrouping)
+			{
+				InOutCurrentGroupings.Add(*TagGrouping);
+			}
+		}
+		virtual void GetColumnConfigSet(TArray<FTableColumnConfig>& InOutConfigSet) const override
+		{
+			InOutConfigSet.Add({ FTable::GetHierarchyColumnId(),               true, 200.0f });
+			InOutConfigSet.Add({ FMemAllocTableColumns::CountColumnId,         true, 100.0f });
+			InOutConfigSet.Add({ FMemAllocTableColumns::SizeColumnId,          true, 100.0f });
+			InOutConfigSet.Add({ FMemAllocTableColumns::AllocFunctionColumnId, true, 400.0f });
+		}
+	};
+	return MakeShared<FTagViewPreset>();
+}
+
+////////////////////////////////////////////////////////////////////////////////////////////////////
+// Asset Breakdown View
+
+TSharedRef<ITableTreeViewPreset> FMemAllocTableViewPresets::CreateAssetViewPreset(SMemAllocTableTreeView& TableTreeView)
+{
+	class FAssetViewPreset : public ITableTreeViewPreset
+	{
+	public:
+		virtual FText GetName() const override
+		{
+			return LOCTEXT("Asset_PresetName", "Asset (Package)");
+		}
+		virtual FText GetToolTip() const override
+		{
+			return LOCTEXT("Asset_PresetToolTip", "Asset (Package) Breakdown View\nConfigure the tree view to show a breakdown of allocations by Package and Asset Name metadata.");
+		}
+		virtual FName GetSortColumn() const override
+		{
+			return FTable::GetHierarchyColumnId();
+		}
+		virtual EColumnSortMode::Type GetSortMode() const override
+		{
+			return EColumnSortMode::Type::Ascending;
+		}
+		virtual void SetCurrentGroupings(const TArray<TSharedPtr<FTreeNodeGrouping>>& InAvailableGroupings, TArray<TSharedPtr<FTreeNodeGrouping>>& InOutCurrentGroupings) const override
+		{
+			InOutCurrentGroupings.Reset();
+
+			check(InAvailableGroupings[0]->Is<FTreeNodeGroupingFlat>());
+			InOutCurrentGroupings.Add(InAvailableGroupings[0]);
+
+			const TSharedPtr<FTreeNodeGrouping>* PackageGrouping = InAvailableGroupings.FindByPredicate(
+				[](TSharedPtr<FTreeNodeGrouping>& Grouping)
+				{
+					return Grouping->Is<FTreeNodeGroupingByPathBreakdown>() &&
+						Grouping->As<FTreeNodeGroupingByPathBreakdown>().GetColumnId() == FMemAllocTableColumns::PackageColumnId;
+				});
+			if (PackageGrouping)
+			{
+				InOutCurrentGroupings.Add(*PackageGrouping);
+			}
+
+			const TSharedPtr<FTreeNodeGrouping>* AssetGrouping = InAvailableGroupings.FindByPredicate(
+				[](TSharedPtr<FTreeNodeGrouping>& Grouping)
+				{
+					return Grouping->Is<FTreeNodeGroupingByUniqueValueCString>() &&
+						Grouping->As<FTreeNodeGroupingByUniqueValueCString>().GetColumnId() == FMemAllocTableColumns::AssetColumnId;
+				});
+			if (AssetGrouping)
+			{
+				InOutCurrentGroupings.Add(*AssetGrouping);
+			}
+		}
+		virtual void GetColumnConfigSet(TArray<FTableColumnConfig>& InOutConfigSet) const override
+		{
+			InOutConfigSet.Add({ FTable::GetHierarchyColumnId(),               true, 200.0f });
+			InOutConfigSet.Add({ FMemAllocTableColumns::CountColumnId,         true, 100.0f });
+			InOutConfigSet.Add({ FMemAllocTableColumns::SizeColumnId,          true, 100.0f });
+			InOutConfigSet.Add({ FMemAllocTableColumns::TagColumnId,           true, 120.0f });
+			InOutConfigSet.Add({ FMemAllocTableColumns::ClassNameColumnId,     true, 120.0f });
+			InOutConfigSet.Add({ FMemAllocTableColumns::AllocFunctionColumnId, true, 300.0f });
+		}
+	};
+	return MakeShared<FAssetViewPreset>();
+}
+
+////////////////////////////////////////////////////////////////////////////////////////////////////
+// Class Name Breakdown View
+
+TSharedRef<ITableTreeViewPreset> FMemAllocTableViewPresets::CreateClassNameViewPreset(SMemAllocTableTreeView& TableTreeView)
+{
+	class FClassNameViewPreset : public ITableTreeViewPreset
+	{
+	public:
+		virtual FText GetName() const override
+		{
+			return LOCTEXT("ClassName_PresetName", "Class Name");
+		}
+		virtual FText GetToolTip() const override
+		{
+			return LOCTEXT("ClassName_PresetToolTip", "Class Name Breakdown View\nConfigure the tree view to show a breakdown of allocations by Asset's Class Name metadata.");
+		}
+		virtual FName GetSortColumn() const override
+		{
+			return FTable::GetHierarchyColumnId();
+		}
+		virtual EColumnSortMode::Type GetSortMode() const override
+		{
+			return EColumnSortMode::Type::Ascending;
+		}
+		virtual void SetCurrentGroupings(const TArray<TSharedPtr<FTreeNodeGrouping>>& InAvailableGroupings, TArray<TSharedPtr<FTreeNodeGrouping>>& InOutCurrentGroupings) const override
+		{
+			InOutCurrentGroupings.Reset();
+
+			check(InAvailableGroupings[0]->Is<FTreeNodeGroupingFlat>());
+			InOutCurrentGroupings.Add(InAvailableGroupings[0]);
+
+			const TSharedPtr<FTreeNodeGrouping>* ClassNameGrouping = InAvailableGroupings.FindByPredicate(
+				[](TSharedPtr<FTreeNodeGrouping>& Grouping)
+				{
+					return Grouping->Is<FTreeNodeGroupingByUniqueValueCString>() &&
+						Grouping->As<FTreeNodeGroupingByUniqueValueCString>().GetColumnId() == FMemAllocTableColumns::ClassNameColumnId;
+				});
+			if (ClassNameGrouping)
+			{
+				InOutCurrentGroupings.Add(*ClassNameGrouping);
+			}
+
+			const TSharedPtr<FTreeNodeGrouping>* PackageGrouping = InAvailableGroupings.FindByPredicate(
+				[](TSharedPtr<FTreeNodeGrouping>& Grouping)
+				{
+					return Grouping->Is<FTreeNodeGroupingByUniqueValueCString>() &&
+						Grouping->As<FTreeNodeGroupingByUniqueValueCString>().GetColumnId() == FMemAllocTableColumns::PackageColumnId;
+				});
+			if (PackageGrouping)
+			{
+				InOutCurrentGroupings.Add(*PackageGrouping);
+			}
+
+			const TSharedPtr<FTreeNodeGrouping>* AssetGrouping = InAvailableGroupings.FindByPredicate(
+				[](TSharedPtr<FTreeNodeGrouping>& Grouping)
+				{
+					return Grouping->Is<FTreeNodeGroupingByUniqueValueCString>() &&
+						Grouping->As<FTreeNodeGroupingByUniqueValueCString>().GetColumnId() == FMemAllocTableColumns::AssetColumnId;
+				});
+			if (AssetGrouping)
+			{
+				InOutCurrentGroupings.Add(*AssetGrouping);
+			}
+		}
+		virtual void GetColumnConfigSet(TArray<FTableColumnConfig>& InOutConfigSet) const override
+		{
+			InOutConfigSet.Add({ FTable::GetHierarchyColumnId(),               true, 200.0f });
+			InOutConfigSet.Add({ FMemAllocTableColumns::CountColumnId,         true, 100.0f });
+			InOutConfigSet.Add({ FMemAllocTableColumns::SizeColumnId,          true, 100.0f });
+			InOutConfigSet.Add({ FMemAllocTableColumns::TagColumnId,           true, 120.0f });
+			InOutConfigSet.Add({ FMemAllocTableColumns::AllocFunctionColumnId, true, 400.0f });
+		}
+	};
+	return MakeShared<FClassNameViewPreset>();
+}
+
+////////////////////////////////////////////////////////////////////////////////////////////////////
+// Callstack Breakdown View
+
+TSharedRef<ITableTreeViewPreset> FMemAllocTableViewPresets::CreateCallstackViewPreset(SMemAllocTableTreeView& TableTreeView, bool bIsInverted, bool bIsAlloc)
+{
+	class FCallstackViewPreset : public ITableTreeViewPreset
+	{
+	public:
+		FCallstackViewPreset(bool bIsInverted, bool bIsAlloc)
+			: bIsInvertedCallstack(bIsInverted)
+			, bIsAllocCallstack(bIsAlloc)
+		{
+		}
+
+		virtual FText GetName() const override
+		{
+			return
+				bIsAllocCallstack
+				? (bIsInvertedCallstack ?
+					LOCTEXT("InvertedCallstack_Alloc_PresetName", "Inverted Alloc Callstack") :
+					LOCTEXT("Callstack_Alloc_PresetName", "Alloc Callstack"))
+				: (bIsInvertedCallstack ?
+					LOCTEXT("InvertedCallstack_Free_PresetName", "Inverted Free Callstack") :
+					LOCTEXT("Callstack_Free_PresetName", "Free Callstack"));
+		}
+		virtual FText GetToolTip() const override
+		{
+			return
+				bIsAllocCallstack
+				? (bIsInvertedCallstack ?
+					LOCTEXT("InvertedCallstack_Alloc_PresetToolTip", "Inverted Alloc Callstack Breakdown View\nConfigure the tree view to show a breakdown of allocations by inverted callstack.") :
+					LOCTEXT("Callstack_Alloc_PresetToolTip", "Alloc Callstack Breakdown View\nConfigure the tree view to show a breakdown of allocations by callstack."))
+				: (bIsInvertedCallstack ?
+					LOCTEXT("InvertedCallstack_Free_PresetToolTip", "Inverted Free Callstack Breakdown View\nConfigure the tree view to show a breakdown of allocations by inverted callstack.") :
+					LOCTEXT("Callstack_Free_PresetToolTip", "Free Callstack Breakdown View\nConfigure the tree view to show a breakdown of allocations by callstack."));
+		}
+		virtual FName GetSortColumn() const override
+		{
+			return FMemAllocTableColumns::SizeColumnId;
+		}
+		virtual EColumnSortMode::Type GetSortMode() const override
+		{
+			return EColumnSortMode::Type::Descending;
+		}
+		virtual void SetCurrentGroupings(const TArray<TSharedPtr<FTreeNodeGrouping>>& InAvailableGroupings, TArray<TSharedPtr<FTreeNodeGrouping>>& InOutCurrentGroupings) const override
+		{
+			InOutCurrentGroupings.Reset();
+
+			check(InAvailableGroupings[0]->Is<FTreeNodeGroupingFlat>());
+			InOutCurrentGroupings.Add(InAvailableGroupings[0]);
+
+			const bool bIsInverted = bIsInvertedCallstack;
+			const bool bIsAlloc = bIsAllocCallstack;
+			const TSharedPtr<FTreeNodeGrouping>* CallstackGrouping = InAvailableGroupings.FindByPredicate(
+				[bIsInverted, bIsAlloc](TSharedPtr<FTreeNodeGrouping>& Grouping)
+				{
+					return Grouping->Is<FMemAllocGroupingByCallstack>() &&
+						Grouping->As<FMemAllocGroupingByCallstack>().IsInverted() == bIsInverted &&
+						Grouping->As<FMemAllocGroupingByCallstack>().IsAllocCallstack() == bIsAlloc;
+				});
+			if (CallstackGrouping)
+			{
+				InOutCurrentGroupings.Add(*CallstackGrouping);
+			}
+		}
+		virtual void GetColumnConfigSet(TArray<FTableColumnConfig>& InOutConfigSet) const override
+		{
+			InOutConfigSet.Add({ FTable::GetHierarchyColumnId(),       true, 400.0f });
+			InOutConfigSet.Add({ FMemAllocTableColumns::CountColumnId, true, 100.0f });
+			InOutConfigSet.Add({ FMemAllocTableColumns::SizeColumnId,  true, 100.0f });
+			InOutConfigSet.Add({ FMemAllocTableColumns::TagColumnId,   true, 200.0f });
+			if (bIsAllocCallstack)
+			{
+				InOutConfigSet.Add({ FMemAllocTableColumns::AllocFunctionColumnId, true, 200.0f });
+			}
+			else
+			{
+				InOutConfigSet.Add({ FMemAllocTableColumns::FreeFunctionColumnId, true, 200.0f });
+			}
+		}
+
+	private:
+		bool bIsInvertedCallstack;
+		bool bIsAllocCallstack;
+	};
+	return MakeShared<FCallstackViewPreset>(bIsInverted, bIsAlloc);
+}
+
+////////////////////////////////////////////////////////////////////////////////////////////////////
+// Address (Platform Page) Breakdown View
+
+TSharedRef<ITableTreeViewPreset> FMemAllocTableViewPresets::CreatePlatformPageViewPreset(SMemAllocTableTreeView& TableTreeView)
+{
+	class FPageViewPreset : public ITableTreeViewPreset
+	{
+	public:
+		virtual FText GetName() const override
+		{
+			return LOCTEXT("Page_PresetName", "Address (Platform Page)");
+		}
+		virtual FText GetToolTip() const override
+		{
+			return LOCTEXT("Page_PresetToolTip", "Platform Page Breakdown View\nConfigure the tree view to show a breakdown of allocations by their address.\nIt groups allocs into platform page size aligned memory pages.");
+		}
+		virtual FName GetSortColumn() const override
+		{
+			return FTable::GetHierarchyColumnId();
+		}
+		virtual EColumnSortMode::Type GetSortMode() const override
+		{
+			return EColumnSortMode::Type::Ascending;
+		}
+		virtual void SetCurrentGroupings(const TArray<TSharedPtr<FTreeNodeGrouping>>& InAvailableGroupings, TArray<TSharedPtr<FTreeNodeGrouping>>& InOutCurrentGroupings) const override
+		{
+			InOutCurrentGroupings.Reset();
+
+			check(InAvailableGroupings[0]->Is<FTreeNodeGroupingFlat>());
+			InOutCurrentGroupings.Add(InAvailableGroupings[0]);
+
+			const TSharedPtr<FTreeNodeGrouping>* MemoryPageGrouping = InAvailableGroupings.FindByPredicate(
+				[](TSharedPtr<FTreeNodeGrouping>& Grouping)
+				{
+					return Grouping->Is<FTreeNodeGroupingByUniqueValueInt64>() &&
+						Grouping->As<FTreeNodeGroupingByUniqueValueInt64>().GetColumnId() == FMemAllocTableColumns::MemoryPageColumnId;
+				});
+			if (MemoryPageGrouping)
+			{
+				InOutCurrentGroupings.Add(*MemoryPageGrouping);
+			}
+		}
+		virtual void GetColumnConfigSet(TArray<FTableColumnConfig>& InOutConfigSet) const override
+		{
+			InOutConfigSet.Add({ FTable::GetHierarchyColumnId(),               true, 200.0f });
+			InOutConfigSet.Add({ FMemAllocTableColumns::AddressColumnId,       true, 120.0f });
+			InOutConfigSet.Add({ FMemAllocTableColumns::CountColumnId,         true, 100.0f });
+			InOutConfigSet.Add({ FMemAllocTableColumns::SizeColumnId,          true, 100.0f });
+			InOutConfigSet.Add({ FMemAllocTableColumns::TagColumnId,           true, 120.0f });
+			InOutConfigSet.Add({ FMemAllocTableColumns::AllocFunctionColumnId, true, 400.0f });
+		}
+	};
+	return MakeShared<FPageViewPreset>();
+}
+
+////////////////////////////////////////////////////////////////////////////////////////////////////
+// Swap Breakdown View
+
+TSharedRef<ITableTreeViewPreset> FMemAllocTableViewPresets::CreateSwapViewPreset(SMemAllocTableTreeView& TableTreeView)
+{
+	class FSwapViewPreset : public ITableTreeViewPreset
+	{
+	public:
+		virtual FText GetName() const override
+		{
+			return LOCTEXT("Swap_PresetName", "Swap");
+		}
+		virtual FText GetToolTip() const override
+		{
+			return LOCTEXT("Swap_PresetToolTip", "Swap Usage Breakdown View\nConfigure the tree view to show a breakdown of allocations by their swap page.\nIt groups allocs into corresponding swap pages.");
+		}
+		virtual FName GetSortColumn() const override
+		{
+			return FTable::GetHierarchyColumnId();
+		}
+		virtual EColumnSortMode::Type GetSortMode() const override
+		{
+			return EColumnSortMode::Type::Ascending;
+		}
+		virtual void SetCurrentGroupings(const TArray<TSharedPtr<FTreeNodeGrouping>>& InAvailableGroupings, TArray<TSharedPtr<FTreeNodeGrouping>>& InOutCurrentGroupings) const override
+		{
+			InOutCurrentGroupings.Reset();
+
+			const TSharedPtr<FTreeNodeGrouping>* HeapGrouping = InAvailableGroupings.FindByPredicate(
+				[](TSharedPtr<FTreeNodeGrouping>& Grouping)
+				{
+					return Grouping->Is<FMemAllocGroupingBySwapPage>();
+				});
+			if (HeapGrouping)
+			{
+				InOutCurrentGroupings.Add(*HeapGrouping);
+			}
+		}
+		virtual void GetColumnConfigSet(TArray<FTableColumnConfig>& InOutConfigSet) const override
+		{
+			InOutConfigSet.Add({ FTable::GetHierarchyColumnId(),                    true, 200.0f });
+			InOutConfigSet.Add({ FMemAllocTableColumns::StartTimeColumnId,          true, 100.0f });
+			InOutConfigSet.Add({ FMemAllocTableColumns::EndTimeColumnId,            true, 100.0f });
+			InOutConfigSet.Add({ FMemAllocTableColumns::AddressColumnId,            true, 120.0f });
+			InOutConfigSet.Add({ FMemAllocTableColumns::CountColumnId,              true, 100.0f });
+			InOutConfigSet.Add({ FMemAllocTableColumns::CompressedSwapSizeColumnId, true, 100.0f });
+			InOutConfigSet.Add({ FMemAllocTableColumns::SwapSizeColumnId,           true, 100.0f });
+			InOutConfigSet.Add({ FMemAllocTableColumns::SizeColumnId,               true, 100.0f });
+			InOutConfigSet.Add({ FMemAllocTableColumns::TagColumnId,                true, 120.0f });
+			InOutConfigSet.Add({ FMemAllocTableColumns::AllocFunctionColumnId,      true, 400.0f });
+		}
+	};
+	return MakeShared<FSwapViewPreset>();
+}
+
+////////////////////////////////////////////////////////////////////////////////////////////////////
+
+} // namespace UE::Insights::MemoryProfiler
+
+#undef LOCTEXT_NAMESPACE
diff --git a/Engine/Source/Developer/TraceInsights/Private/Insights/MemoryProfiler/ViewModels/MemAllocTableViewPresets.h b/Engine/Source/Developer/TraceInsights/Private/Insights/MemoryProfiler/ViewModels/MemAllocTableViewPresets.h
new file mode 100644
--- /dev/null
+++ b/Engine/Source/Developer/TraceInsights/Private/Insights/MemoryProfiler/ViewModels/MemAllocTableViewPresets.h
@@ -0,0 +1,29 @@
+#pragma once
+
+#include "CoreTypes.h"
+
+#include "Templates/SharedPointer.h"
+
+namespace UE::Insights { class ITableTreeViewPreset; }
+
+namespace UE::Insights::MemoryProfiler
+{
+	
+class SMemAllocTableTreeView;
+
+class FMemAllocTableViewPresets
+{
+public:
+	static TSharedRef<ITableTreeViewPreset> CreateDefaultViewPreset(SMemAllocTableTreeView& TableTreeView);
+	static TSharedRef<ITableTreeViewPreset> CreateDetailedViewPreset(SMemAllocTableTreeView& TableTreeView);
+	static TSharedRef<ITableTreeViewPreset> CreateHeapViewPreset(SMemAllocTableTreeView& TableTreeView);
+	static TSharedRef<ITableTreeViewPreset> CreateSizeViewPreset(SMemAllocTableTreeView& TableTreeView);
+	static TSharedRef<ITableTreeViewPreset> CreateTagViewPreset(SMemAllocTableTreeView& TableTreeView);
+	static TSharedRef<ITableTreeViewPreset> CreateAssetViewPreset(SMemAllocTableTreeView& TableTreeView);
+	static TSharedRef<ITableTreeViewPreset> CreateClassNameViewPreset(SMemAllocTableTreeView& TableTreeView);
+	static TSharedRef<ITableTreeViewPreset> CreateCallstackViewPreset(SMemAllocTableTreeView& TableTreeView, bool bIsInverted, bool bIsAlloc);
+	static TSharedRef<ITableTreeViewPreset> CreatePlatformPageViewPreset(SMemAllocTableTreeView& TableTreeView);
+	static TSharedRef<ITableTreeViewPreset> CreateSwapViewPreset(SMemAllocTableTreeView& TableTreeView);
+};
+
+} // namespace UE::Insights::MemoryProfiler
diff --git a/Engine/Source/Developer/TraceInsights/Private/Insights/MemoryProfiler/Widgets/SMemAllocTableTreeView.cpp b/Engine/Source/Developer/TraceInsights/Private/Insights/MemoryProfiler/Widgets/SMemAllocTableTreeView.cpp
--- a/Engine/Source/Developer/TraceInsights/Private/Insights/MemoryProfiler/Widgets/SMemAllocTableTreeView.cpp
+++ b/Engine/Source/Developer/TraceInsights/Private/Insights/MemoryProfiler/Widgets/SMemAllocTableTreeView.cpp
@@ -31,7 +31,6 @@
 // TraceInsightsCore
 #include "InsightsCore/Filter/ViewModels/FilterConfigurator.h"
 #include "InsightsCore/Filter/ViewModels/Filters.h"
-#include "InsightsCore/Table/ViewModels/TableCellValueFormatter.h"
 #include "InsightsCore/Table/ViewModels/TableColumn.h"
 
 // TraceInsights
@@ -45,13 +44,10 @@
 #include "Insights/MemoryProfiler/ViewModels/MemAllocGroupingBySwapPage.h"
 #include "Insights/MemoryProfiler/ViewModels/MemAllocGroupingByTag.h"
 #include "Insights/MemoryProfiler/ViewModels/MemAllocNode.h"
-#include "Insights/MemoryProfiler/ViewModels/MemAllocTable.h"
+#include "Insights/MemoryProfiler/ViewModels/MemAllocTableViewPresets.h"
 #include "Insights/MemoryProfiler/ViewModels/MemoryFiltersValueConverters.h"
 #include "Insights/MemoryProfiler/ViewModels/MemorySharedState.h"
-#include "Insights/MemoryProfiler/Widgets/SMemoryProfilerWindow.h"
-#include "Insights/TimingProfilerCommon.h"
 
-#include <limits>
 #include <memory>
 
 #define LOCTEXT_NAMESPACE "UE::Insights::MemoryProfiler::SMemAllocTableTreeView"
@@ -82,6 +78,162 @@ SMemAllocTableTreeView::~SMemAllocTableTreeView()
 void SMemAllocTableTreeView::Construct(const FArguments& InArgs, TSharedPtr<FMemAllocTable> InTablePtr)
 {
 	ConstructWidget(InTablePtr);
+
+	// Apply the default preset.
+	ApplyViewPreset(*(*GetAvailableViewPresets())[0]);
+}
+
+////////////////////////////////////////////////////////////////////////////////////////////////////
+
+BEGIN_SLATE_FUNCTION_BUILD_OPTIMIZATION
+TSharedPtr<SWidget> SMemAllocTableTreeView::ConstructToolbar()
+{
+	TSharedPtr<SHorizontalBox> Box = SNew(SHorizontalBox);
+
+	ConstructViewPreset(Box);
+
+	Box->AddSlot()
+		.AutoWidth()
+		.Padding(4.0f, 0.0f, 0.0f, 0.0f)
+		[
+			ConstructFunctionToggleButton()
+		];
+
+	return Box;
+}
+END_SLATE_FUNCTION_BUILD_OPTIMIZATION
+
+////////////////////////////////////////////////////////////////////////////////////////////////////
+
+BEGIN_SLATE_FUNCTION_BUILD_OPTIMIZATION
+TSharedPtr<SWidget> SMemAllocTableTreeView::ConstructFooter()
+{
+	return
+		SNew(SHorizontalBox)
+
+		+ SHorizontalBox::Slot()
+		.AutoWidth()
+		.HAlign(HAlign_Left)
+		.Padding(2.0f, 2.0f, 0.0f, 2.0f)
+		[
+			SNew(STextBlock)
+			.Text(this, &SMemAllocTableTreeView::GetQueryInfo)
+			.ToolTipText(this, &SMemAllocTableTreeView::GetQueryInfoTooltip)
+		]
+
+		+ SHorizontalBox::Slot()
+		.AutoWidth()
+		.HAlign(HAlign_Left)
+		.Padding(0.0f, 2.0f, 0.0f, 2.0f)
+		[
+			SNew(STextBlock)
+			.Text(LOCTEXT("FooterSeparator", " : "))
+			.ColorAndOpacity(FSlateColor(EStyleColor::White25))
+		]
+
+		+ SHorizontalBox::Slot()
+		.AutoWidth()
+		.HAlign(HAlign_Left)
+		.Padding(0.0f, 2.0f, 8.0f, 2.0f)
+		[
+			SNew(STextBlock)
+			.Text(this, &SMemAllocTableTreeView::GetFooterLeftText)
+		]
+
+		+ SHorizontalBox::Slot()
+		.FillWidth(1.0f)
+		.HAlign(HAlign_Center)
+		.Padding(2.0f, 2.0f, 2.0f, 2.0f)
+		[
+			SNew(STextBlock)
+			.Text(this, &SMemAllocTableTreeView::GetFooterCenterText)
+		]
+
+		+ SHorizontalBox::Slot()
+		.AutoWidth()
+		.HAlign(HAlign_Right)
+		.Padding(8.0f, 2.0f, 2.0f, 2.0f)
+		[
+			SNew(STextBlock)
+			.Text(this, &SMemAllocTableTreeView::GetSymbolResolutionStatus)
+			.ToolTipText(this, &SMemAllocTableTreeView::GetSymbolResolutionTooltip)
+		];
+}
+END_SLATE_FUNCTION_BUILD_OPTIMIZATION
+
+////////////////////////////////////////////////////////////////////////////////////////////////////
+void SMemAllocTableTreeView::InternalCreateGroupings()
+{
+	STableTreeView::InternalCreateGroupings();
+
+	AvailableGroupings.RemoveAll(
+		[](TSharedPtr<FTreeNodeGrouping>& Grouping)
+		{
+			if (Grouping->Is<FTreeNodeGroupingByUniqueValue>())
+			{
+				const FName ColumnId = Grouping->As<FTreeNodeGroupingByUniqueValue>().GetColumnId();
+				if (ColumnId == FMemAllocTableColumns::StartEventIndexColumnId ||
+					ColumnId == FMemAllocTableColumns::EndEventIndexColumnId ||
+					ColumnId == FMemAllocTableColumns::CountColumnId)
+				{
+					return true;
+				}
+			}
+			else if (Grouping->Is<FTreeNodeGroupingByPathBreakdown>())
+			{
+				const FName ColumnId = Grouping->As<FTreeNodeGroupingByPathBreakdown>().GetColumnId();
+				if (ColumnId == FMemAllocTableColumns::AllocFunctionColumnId ||
+					ColumnId == FMemAllocTableColumns::FreeFunctionColumnId ||
+					ColumnId == FMemAllocTableColumns::ClassNameColumnId)
+				{
+					return true;
+				}
+			}
+			return false;
+		});
+
+	int32 Index = 1; // after the Flat ("All") grouping
+
+	AvailableGroupings.Insert(MakeShared<FMemAllocGroupingBySize>(), Index++);
+
+	const TraceServices::IAllocationsProvider* AllocationsProvider = Session.IsValid() ? TraceServices::ReadAllocationsProvider(*Session.Get()) : nullptr;
+
+	if (AllocationsProvider)
+	{
+		AvailableGroupings.Insert(MakeShared<FMemAllocGroupingByTag>(*AllocationsProvider), Index++);
+	}
+
+	AvailableGroupings.Insert(MakeShared<FMemAllocGroupingByCallstack>(true, false, bIsCallstackGroupingByFunction), Index++);
+	AvailableGroupings.Insert(MakeShared<FMemAllocGroupingByCallstack>(true, true, bIsCallstackGroupingByFunction), Index++);
+	AvailableGroupings.Insert(MakeShared<FMemAllocGroupingByCallstack>(false, false, bIsCallstackGroupingByFunction), Index++);
+	AvailableGroupings.Insert(MakeShared<FMemAllocGroupingByCallstack>(false, true, bIsCallstackGroupingByFunction), Index++);
+
+	if (AllocationsProvider)
+	{
+		AvailableGroupings.Insert(MakeShared<FMemAllocGroupingByHeap>(*AllocationsProvider), Index++);
+		AvailableGroupings.Insert(MakeShared<FMemAllocGroupingBySwapPage>(*AllocationsProvider), Index++);
+	}
+}
+
+////////////////////////////////////////////////////////////////////////////////////////////////////
+
+void SMemAllocTableTreeView::InitAvailableViewPresets()
+{
+	AvailableViewPresets.Add(FMemAllocTableViewPresets::CreateDefaultViewPreset(*this));
+	AvailableViewPresets.Add(FMemAllocTableViewPresets::CreateDetailedViewPreset(*this));
+	AvailableViewPresets.Add(FMemAllocTableViewPresets::CreateHeapViewPreset(*this));
+	AvailableViewPresets.Add(FMemAllocTableViewPresets::CreateSizeViewPreset(*this));
+	AvailableViewPresets.Add(FMemAllocTableViewPresets::CreateTagViewPreset(*this));
+	AvailableViewPresets.Add(FMemAllocTableViewPresets::CreateAssetViewPreset(*this));
+	AvailableViewPresets.Add(FMemAllocTableViewPresets::CreateClassNameViewPreset(*this));
+	AvailableViewPresets.Add(FMemAllocTableViewPresets::CreateCallstackViewPreset(*this, false, true));
+	AvailableViewPresets.Add(FMemAllocTableViewPresets::CreateCallstackViewPreset(*this, true, true));
+	AvailableViewPresets.Add(FMemAllocTableViewPresets::CreateCallstackViewPreset(*this, false, false));
+	AvailableViewPresets.Add(FMemAllocTableViewPresets::CreateCallstackViewPreset(*this, true, false));
+	AvailableViewPresets.Add(FMemAllocTableViewPresets::CreatePlatformPageViewPreset(*this));
+	AvailableViewPresets.Add(FMemAllocTableViewPresets::CreateSwapViewPreset(*this));
+
+	SelectedViewPreset = AvailableViewPresets[0];
 }
 
 ////////////////////////////////////////////////////////////////////////////////////////////////////
@@ -588,685 +740,122 @@ FText SMemAllocTableTreeView::GetCurrentOperationName() const
 
 ////////////////////////////////////////////////////////////////////////////////////////////////////
 
-BEGIN_SLATE_FUNCTION_BUILD_OPTIMIZATION
-TSharedPtr<SWidget> SMemAllocTableTreeView::ConstructToolbar()
+void SMemAllocTableTreeView::UpdateQueryInfo()
 {
-	TSharedPtr<SHorizontalBox> Box = SNew(SHorizontalBox);
+	if (Rule.IsValid())
+	{
+		FText TimeMarkersText;
 
-	Box->AddSlot()
-		.AutoWidth()
-		.VAlign(VAlign_Center)
-		.Padding(0.0f, 0.0f, 4.0f, 0.0f)
-		[
-			SNew(STextBlock)
-			.Text(LOCTEXT("Preset", "Preset:"))
-		];
+		const int NumMarkers = Rule->GetNumTimeMarkers();
+		switch (NumMarkers)
+		{
+		case 1:
+			TimeMarkersText = FText::Format(LOCTEXT("OneTimeMarkersFmt", "A={0}"), FText::AsNumber(TimeMarkers[0]));
+			break;
+		case 2:
+			TimeMarkersText = FText::Format(LOCTEXT("TwoTimeMarkersFmt", "A={0}  B={1}"), FText::AsNumber(TimeMarkers[0]), FText::AsNumber(TimeMarkers[1]));
+			break;
+		case 3:
+			TimeMarkersText = FText::Format(LOCTEXT("ThreeTimeMarkersFmt", "A={0}  B={1}  C={2}"), FText::AsNumber(TimeMarkers[0]), FText::AsNumber(TimeMarkers[1]), FText::AsNumber(TimeMarkers[2]));
+			break;
+		case 4:
+			TimeMarkersText = FText::Format(LOCTEXT("FourTimeMarkersFmt", "A={0}  B={1}  C={2}  D={3}"), FText::AsNumber(TimeMarkers[0]), FText::AsNumber(TimeMarkers[1]), FText::AsNumber(TimeMarkers[2]), FText::AsNumber(TimeMarkers[3]));
+			break;
+		default:
+			// Unhandled value
+			check(false);
+		}
 
-	Box->AddSlot()
-		.AutoWidth()
-		.VAlign(VAlign_Center)
-		[
-			SNew(SBox)
-			.MinDesiredWidth(150.0f)
-			[
-				SAssignNew(PresetComboBox, SComboBox<TSharedRef<ITableTreeViewPreset>>)
-				.ToolTipText(this, &SMemAllocTableTreeView::ViewPreset_GetSelectedToolTipText)
-				.OptionsSource(GetAvailableViewPresets())
-				.OnSelectionChanged(this, &SMemAllocTableTreeView::ViewPreset_OnSelectionChanged)
-				.OnGenerateWidget(this, &SMemAllocTableTreeView::ViewPreset_OnGenerateWidget)
-				[
-					SNew(STextBlock)
-					.Text(this, &SMemAllocTableTreeView::ViewPreset_GetSelectedText)
-				]
-			]
-		];
+		QueryInfo = FText::Format(LOCTEXT("QueryInfoFmt", "{0} ({1})"), Rule->GetVerboseName(), TimeMarkersText);
+	}
+}
 
-	Box->AddSlot()
-		.AutoWidth()
-		.Padding(4.0f, 0.0f, 0.0f, 0.0f)
-		[
-			ConstructFunctionToggleButton()
-		];
+////////////////////////////////////////////////////////////////////////////////////////////////////
 
-	return Box;
+FText SMemAllocTableTreeView::GetQueryInfo() const
+{
+	return QueryInfo;
 }
-END_SLATE_FUNCTION_BUILD_OPTIMIZATION
 
 ////////////////////////////////////////////////////////////////////////////////////////////////////
 
-void SMemAllocTableTreeView::InitAvailableViewPresets()
+FText SMemAllocTableTreeView::GetQueryInfoTooltip() const
 {
-	//////////////////////////////////////////////////
-	// Default View
+	return Rule.IsValid() ? Rule->GetDescription() : FText::GetEmpty();
+}
 
-	class FDefaultViewPreset : public ITableTreeViewPreset
-	{
-	public:
-		virtual FText GetName() const override
-		{
-			return LOCTEXT("Default_PresetName", "Default");
-		}
-		virtual FText GetToolTip() const override
-		{
-			return LOCTEXT("Default_PresetToolTip", "Default View\nConfigure the tree view to show default allocation info.");
-		}
-		virtual FName GetSortColumn() const override
-		{
-			return FTable::GetHierarchyColumnId();
-		}
-		virtual EColumnSortMode::Type GetSortMode() const override
-		{
-			return EColumnSortMode::Type::Ascending;
-		}
-		virtual void SetCurrentGroupings(const TArray<TSharedPtr<FTreeNodeGrouping>>& InAvailableGroupings, TArray<TSharedPtr<FTreeNodeGrouping>>& InOutCurrentGroupings) const override
-		{
-			InOutCurrentGroupings.Reset();
+////////////////////////////////////////////////////////////////////////////////////////////////////
 
-			check(InAvailableGroupings[0]->Is<FTreeNodeGroupingFlat>());
-			InOutCurrentGroupings.Add(InAvailableGroupings[0]);
-		}
-		virtual void GetColumnConfigSet(TArray<FTableColumnConfig>& InOutConfigSet) const override
-		{
-			InOutConfigSet.Add({ FTable::GetHierarchyColumnId(),                    true, 200.0f });
-			InOutConfigSet.Add({ FMemAllocTableColumns::CountColumnId,              true, 100.0f });
-			InOutConfigSet.Add({ FMemAllocTableColumns::SizeColumnId,               true, 100.0f });
-			InOutConfigSet.Add({ FMemAllocTableColumns::TagColumnId,                true, 120.0f });
-			InOutConfigSet.Add({ FMemAllocTableColumns::AllocFunctionColumnId,      true, 550.0f });
-			InOutConfigSet.Add({ FMemAllocTableColumns::AllocCallstackSizeColumnId, true, 100.0f });
-		}
-	};
-	AvailableViewPresets.Add(MakeShared<FDefaultViewPreset>());
+FText SMemAllocTableTreeView::GetFooterLeftText() const
+{
+	if (!Rule.IsValid())
+	{
+		return FText::GetEmpty();
+	}
 
-	//////////////////////////////////////////////////
-	// Detailed View
+	if (Query != 0)
+	{
+		return LOCTEXT("FooterLeftTextRunningQuery", "running query...");
+	}
 
-	class FDetailedViewPreset : public ITableTreeViewPreset
+	if (FilteredNodesPtr->Num() == TableRowNodes.Num())
 	{
-	public:
-		virtual FText GetName() const override
-		{
-			return LOCTEXT("Detailed_PresetName", "Detailed");
-		}
-		virtual FText GetToolTip() const override
-		{
-			return LOCTEXT("Detailed_PresetToolTip", "Detailed View\nConfigure the tree view to show detailed allocation info.");
-		}
-		virtual FName GetSortColumn() const override
-		{
-			return FTable::GetHierarchyColumnId();
-		}
-		virtual EColumnSortMode::Type GetSortMode() const override
-		{
-			return EColumnSortMode::Type::Ascending;
-		}
-		virtual void SetCurrentGroupings(const TArray<TSharedPtr<FTreeNodeGrouping>>& InAvailableGroupings, TArray<TSharedPtr<FTreeNodeGrouping>>& InOutCurrentGroupings) const override
-		{
-			InOutCurrentGroupings.Reset();
+		return FText::Format(LOCTEXT("FooterLeftTextFmt1", "{0} {0}|plural(one=alloc,other=allocs)"), FText::AsNumber(TableRowNodes.Num()));
+	}
+	else
+	{
+		return FText::Format(LOCTEXT("FooterLeftTextFmt2", "{0} / {1} {1}|plural(one=alloc,other=allocs)"), FilteredNodesPtr->Num(), FText::AsNumber(TableRowNodes.Num()));
+	}
+}
 
-			check(InAvailableGroupings[0]->Is<FTreeNodeGroupingFlat>());
-			InOutCurrentGroupings.Add(InAvailableGroupings[0]);
-		}
-		virtual void GetColumnConfigSet(TArray<FTableColumnConfig>& InOutConfigSet) const override
-		{
-			InOutConfigSet.Add({ FTable::GetHierarchyColumnId(),                    true, 200.0f });
-			InOutConfigSet.Add({ FMemAllocTableColumns::StartEventIndexColumnId,    true, 100.0f });
-			InOutConfigSet.Add({ FMemAllocTableColumns::EndEventIndexColumnId,      true, 100.0f });
-			InOutConfigSet.Add({ FMemAllocTableColumns::EventDistanceColumnId,      true, 100.0f });
-			InOutConfigSet.Add({ FMemAllocTableColumns::StartTimeColumnId,          true, 100.0f });
-			InOutConfigSet.Add({ FMemAllocTableColumns::EndTimeColumnId,            true, 100.0f });
-			InOutConfigSet.Add({ FMemAllocTableColumns::DurationColumnId,           true, 100.0f });
-			InOutConfigSet.Add({ FMemAllocTableColumns::AddressColumnId,            true, 120.0f });
-			InOutConfigSet.Add({ FMemAllocTableColumns::MemoryPageColumnId,         true, 120.0f });
-			InOutConfigSet.Add({ FMemAllocTableColumns::CountColumnId,              true, 100.0f });
-			InOutConfigSet.Add({ FMemAllocTableColumns::SizeColumnId,               true, 100.0f });
-			InOutConfigSet.Add({ FMemAllocTableColumns::TagColumnId,                true, 120.0f });
-			InOutConfigSet.Add({ FMemAllocTableColumns::AllocFunctionColumnId,      true, 550.0f });
-			InOutConfigSet.Add({ FMemAllocTableColumns::AllocSourceFileColumnId,    true, 550.0f });
-			InOutConfigSet.Add({ FMemAllocTableColumns::AllocCallstackSizeColumnId, true, 100.0f });
-		}
-	};
-	AvailableViewPresets.Add(MakeShared<FDetailedViewPreset>());
+////////////////////////////////////////////////////////////////////////////////////////////////////
 
-	//////////////////////////////////////////////////
-	// Heap Breakdown View
+void SMemAllocTableTreeView::UpdateSelectionStatsText()
+{
+	TArray<FTableTreeNodePtr> SelectedNodes;
+	const int32 NumSelectedNodes = TreeView->GetSelectedItems(SelectedNodes);
 
-	class FHeapViewPreset : public ITableTreeViewPreset
+	if (NumSelectedNodes > 0)
 	{
-	public:
-		virtual FText GetName() const override
-		{
-			return LOCTEXT("Heap_PresetName", "Heap");
-		}
-		virtual FText GetToolTip() const override
+		int64 TotalAllocCount = 0;
+		int64 TotalAllocSize = 0;
+
+		TSharedRef<FTableColumn> CountColumn = Table->FindColumnChecked(FMemAllocTableColumns::CountColumnId);
+		TSharedRef<FTableColumn> SizeColumn = Table->FindColumnChecked(FMemAllocTableColumns::SizeColumnId);
+
+		for (const FTableTreeNodePtr& Node : SelectedNodes)
 		{
-			return LOCTEXT("Heap_PresetToolTip", "Heap Breakdown View\nConfigure the tree view to show a breakdown of allocations by their parent heap type.");
-		}
-		virtual FName GetSortColumn() const override
-		{
-			return FMemAllocTableColumns::SizeColumnId;
-		}
-		virtual EColumnSortMode::Type GetSortMode() const override
-		{
-			return EColumnSortMode::Type::Descending;
-		}
-		virtual void SetCurrentGroupings(const TArray<TSharedPtr<FTreeNodeGrouping>>& InAvailableGroupings, TArray<TSharedPtr<FTreeNodeGrouping>>& InOutCurrentGroupings) const override
-		{
-			InOutCurrentGroupings.Reset();
-
-			//check(InAvailableGroupings[0]->Is<FTreeNodeGroupingFlat>());
-			//InOutCurrentGroupings.Add(InAvailableGroupings[0]);
-
-			const TSharedPtr<FTreeNodeGrouping>* HeapGrouping = InAvailableGroupings.FindByPredicate(
-				[](TSharedPtr<FTreeNodeGrouping>& Grouping)
-				{
-					return Grouping->Is<FMemAllocGroupingByHeap>();
-				});
-			if (HeapGrouping)
-			{
-				InOutCurrentGroupings.Add(*HeapGrouping);
-			}
-		}
-		virtual void GetColumnConfigSet(TArray<FTableColumnConfig>& InOutConfigSet) const override
-		{
-			InOutConfigSet.Add({ FTable::GetHierarchyColumnId(),               true, 400.0f });
-			InOutConfigSet.Add({ FMemAllocTableColumns::CountColumnId,         true, 100.0f });
-			InOutConfigSet.Add({ FMemAllocTableColumns::SizeColumnId,          true, 100.0f });
-			InOutConfigSet.Add({ FMemAllocTableColumns::TagColumnId,           true, 200.0f });
-			InOutConfigSet.Add({ FMemAllocTableColumns::AllocFunctionColumnId, true, 200.0f });
-		}
-	};
-	AvailableViewPresets.Add(MakeShared<FHeapViewPreset>());
-
-	//////////////////////////////////////////////////
-	// Size Breakdown View
-
-	class FSizeViewPreset : public ITableTreeViewPreset
-	{
-	public:
-		virtual FText GetName() const override
-		{
-			return LOCTEXT("Size_PresetName", "Size");
-		}
-		virtual FText GetToolTip() const override
-		{
-			return LOCTEXT("Size_PresetToolTip", "Size Breakdown View\nConfigure the tree view to show a breakdown of allocations by their size.");
-		}
-		virtual FName GetSortColumn() const override
-		{
-			return FMemAllocTableColumns::SizeColumnId;
-		}
-		virtual EColumnSortMode::Type GetSortMode() const override
-		{
-			return EColumnSortMode::Type::Descending;
-		}
-		virtual void SetCurrentGroupings(const TArray<TSharedPtr<FTreeNodeGrouping>>& InAvailableGroupings, TArray<TSharedPtr<FTreeNodeGrouping>>& InOutCurrentGroupings) const override
-		{
-			InOutCurrentGroupings.Reset();
-
-			check(InAvailableGroupings[0]->Is<FTreeNodeGroupingFlat>());
-			InOutCurrentGroupings.Add(InAvailableGroupings[0]);
-
-			const TSharedPtr<FTreeNodeGrouping>* SizeGrouping = InAvailableGroupings.FindByPredicate(
-				[](TSharedPtr<FTreeNodeGrouping>& Grouping)
-				{
-					return Grouping->Is<FMemAllocGroupingBySize>();
-				});
-			if (SizeGrouping)
-			{
-				InOutCurrentGroupings.Add(*SizeGrouping);
-			}
-		}
-		virtual void GetColumnConfigSet(TArray<FTableColumnConfig>& InOutConfigSet) const override
-		{
-			InOutConfigSet.Add({ FTable::GetHierarchyColumnId(),               true, 200.0f });
-			InOutConfigSet.Add({ FMemAllocTableColumns::AddressColumnId,       true, 120.0f });
-			InOutConfigSet.Add({ FMemAllocTableColumns::CountColumnId,         true, 100.0f });
-			InOutConfigSet.Add({ FMemAllocTableColumns::SizeColumnId,          true, 100.0f });
-			InOutConfigSet.Add({ FMemAllocTableColumns::TagColumnId,           true, 120.0f });
-			InOutConfigSet.Add({ FMemAllocTableColumns::AllocFunctionColumnId, true, 400.0f });
-		}
-	};
-	AvailableViewPresets.Add(MakeShared<FSizeViewPreset>());
-
-	//////////////////////////////////////////////////
-	// Tag Breakdown View
-
-	class FTagViewPreset : public ITableTreeViewPreset
-	{
-	public:
-		virtual FText GetName() const override
-		{
-			return LOCTEXT("Tag_PresetName", "Tags");
-		}
-		virtual FText GetToolTip() const override
-		{
-			return LOCTEXT("Tag_PresetToolTip", "Tag Breakdown View\nConfigure the tree view to show a breakdown of allocations by their LLM tag.");
-		}
-		virtual FName GetSortColumn() const override
-		{
-			return FTable::GetHierarchyColumnId();
-		}
-		virtual EColumnSortMode::Type GetSortMode() const override
-		{
-			return EColumnSortMode::Type::Ascending;
-		}
-		virtual void SetCurrentGroupings(const TArray<TSharedPtr<FTreeNodeGrouping>>& InAvailableGroupings, TArray<TSharedPtr<FTreeNodeGrouping>>& InOutCurrentGroupings) const override
-		{
-			InOutCurrentGroupings.Reset();
-
-			check(InAvailableGroupings[0]->Is<FTreeNodeGroupingFlat>());
-			InOutCurrentGroupings.Add(InAvailableGroupings[0]);
-
-			const TSharedPtr<FTreeNodeGrouping>* TagGrouping = InAvailableGroupings.FindByPredicate(
-				[](TSharedPtr<FTreeNodeGrouping>& Grouping)
-				{
-					return Grouping->Is<FMemAllocGroupingByTag>();
-				});
-			if (TagGrouping)
-			{
-				InOutCurrentGroupings.Add(*TagGrouping);
-			}
-		}
-		virtual void GetColumnConfigSet(TArray<FTableColumnConfig>& InOutConfigSet) const override
-		{
-			InOutConfigSet.Add({ FTable::GetHierarchyColumnId(),               true, 200.0f });
-			InOutConfigSet.Add({ FMemAllocTableColumns::CountColumnId,         true, 100.0f });
-			InOutConfigSet.Add({ FMemAllocTableColumns::SizeColumnId,          true, 100.0f });
-			InOutConfigSet.Add({ FMemAllocTableColumns::AllocFunctionColumnId, true, 400.0f });
-		}
-	};
-	AvailableViewPresets.Add(MakeShared<FTagViewPreset>());
-
-	//////////////////////////////////////////////////
-	// Asset Breakdown View
-
-	class FAssetViewPreset : public ITableTreeViewPreset
-	{
-	public:
-		virtual FText GetName() const override
-		{
-			return LOCTEXT("Asset_PresetName", "Asset (Package)");
-		}
-		virtual FText GetToolTip() const override
-		{
-			return LOCTEXT("Asset_PresetToolTip", "Asset (Package) Breakdown View\nConfigure the tree view to show a breakdown of allocations by Package and Asset Name metadata.");
-		}
-		virtual FName GetSortColumn() const override
-		{
-			return FTable::GetHierarchyColumnId();
-		}
-		virtual EColumnSortMode::Type GetSortMode() const override
-		{
-			return EColumnSortMode::Type::Ascending;
-		}
-		virtual void SetCurrentGroupings(const TArray<TSharedPtr<FTreeNodeGrouping>>& InAvailableGroupings, TArray<TSharedPtr<FTreeNodeGrouping>>& InOutCurrentGroupings) const override
-		{
-			InOutCurrentGroupings.Reset();
-
-			check(InAvailableGroupings[0]->Is<FTreeNodeGroupingFlat>());
-			InOutCurrentGroupings.Add(InAvailableGroupings[0]);
-
-			const TSharedPtr<FTreeNodeGrouping>* PackageGrouping = InAvailableGroupings.FindByPredicate(
-				[](TSharedPtr<FTreeNodeGrouping>& Grouping)
-				{
-					return Grouping->Is<FTreeNodeGroupingByPathBreakdown>() &&
-						   Grouping->As<FTreeNodeGroupingByPathBreakdown>().GetColumnId() == FMemAllocTableColumns::PackageColumnId;
-				});
-			if (PackageGrouping)
-			{
-				InOutCurrentGroupings.Add(*PackageGrouping);
-			}
-
-			const TSharedPtr<FTreeNodeGrouping>* AssetGrouping = InAvailableGroupings.FindByPredicate(
-				[](TSharedPtr<FTreeNodeGrouping>& Grouping)
-				{
-					return Grouping->Is<FTreeNodeGroupingByUniqueValueCString>() &&
-						   Grouping->As<FTreeNodeGroupingByUniqueValueCString>().GetColumnId() == FMemAllocTableColumns::AssetColumnId;
-				});
-			if (AssetGrouping)
-			{
-				InOutCurrentGroupings.Add(*AssetGrouping);
-			}
-		}
-		virtual void GetColumnConfigSet(TArray<FTableColumnConfig>& InOutConfigSet) const override
-		{
-			InOutConfigSet.Add({ FTable::GetHierarchyColumnId(),               true, 200.0f });
-			InOutConfigSet.Add({ FMemAllocTableColumns::CountColumnId,         true, 100.0f });
-			InOutConfigSet.Add({ FMemAllocTableColumns::SizeColumnId,          true, 100.0f });
-			InOutConfigSet.Add({ FMemAllocTableColumns::TagColumnId,           true, 120.0f });
-			InOutConfigSet.Add({ FMemAllocTableColumns::ClassNameColumnId,     true, 120.0f });
-			InOutConfigSet.Add({ FMemAllocTableColumns::AllocFunctionColumnId, true, 300.0f });
-		}
-	};
-	AvailableViewPresets.Add(MakeShared<FAssetViewPreset>());
-
-	//////////////////////////////////////////////////
-	// Class Name Breakdown View
-
-	class FClassNameViewPreset : public ITableTreeViewPreset
-	{
-	public:
-		virtual FText GetName() const override
-		{
-			return LOCTEXT("ClassName_PresetName", "Class Name");
-		}
-		virtual FText GetToolTip() const override
-		{
-			return LOCTEXT("ClassName_PresetToolTip", "Class Name Breakdown View\nConfigure the tree view to show a breakdown of allocations by Asset's Class Name metadata.");
-		}
-		virtual FName GetSortColumn() const override
-		{
-			return FTable::GetHierarchyColumnId();
-		}
-		virtual EColumnSortMode::Type GetSortMode() const override
-		{
-			return EColumnSortMode::Type::Ascending;
-		}
-		virtual void SetCurrentGroupings(const TArray<TSharedPtr<FTreeNodeGrouping>>& InAvailableGroupings, TArray<TSharedPtr<FTreeNodeGrouping>>& InOutCurrentGroupings) const override
-		{
-			InOutCurrentGroupings.Reset();
-
-			check(InAvailableGroupings[0]->Is<FTreeNodeGroupingFlat>());
-			InOutCurrentGroupings.Add(InAvailableGroupings[0]);
-
-			const TSharedPtr<FTreeNodeGrouping>* ClassNameGrouping = InAvailableGroupings.FindByPredicate(
-				[](TSharedPtr<FTreeNodeGrouping>& Grouping)
-				{
-					return Grouping->Is<FTreeNodeGroupingByUniqueValueCString>() &&
-						   Grouping->As<FTreeNodeGroupingByUniqueValueCString>().GetColumnId() == FMemAllocTableColumns::ClassNameColumnId;
-				});
-			if (ClassNameGrouping)
-			{
-				InOutCurrentGroupings.Add(*ClassNameGrouping);
-			}
-
-			const TSharedPtr<FTreeNodeGrouping>* PackageGrouping = InAvailableGroupings.FindByPredicate(
-				[](TSharedPtr<FTreeNodeGrouping>& Grouping)
-				{
-					return Grouping->Is<FTreeNodeGroupingByUniqueValueCString>() &&
-						   Grouping->As<FTreeNodeGroupingByUniqueValueCString>().GetColumnId() == FMemAllocTableColumns::PackageColumnId;
-				});
-			if (PackageGrouping)
-			{
-				InOutCurrentGroupings.Add(*PackageGrouping);
-			}
-
-			const TSharedPtr<FTreeNodeGrouping>* AssetGrouping = InAvailableGroupings.FindByPredicate(
-				[](TSharedPtr<FTreeNodeGrouping>& Grouping)
-				{
-					return Grouping->Is<FTreeNodeGroupingByUniqueValueCString>() &&
-						   Grouping->As<FTreeNodeGroupingByUniqueValueCString>().GetColumnId() == FMemAllocTableColumns::AssetColumnId;
-				});
-			if (AssetGrouping)
-			{
-				InOutCurrentGroupings.Add(*AssetGrouping);
-			}
-		}
-		virtual void GetColumnConfigSet(TArray<FTableColumnConfig>& InOutConfigSet) const override
-		{
-			InOutConfigSet.Add({ FTable::GetHierarchyColumnId(),               true, 200.0f });
-			InOutConfigSet.Add({ FMemAllocTableColumns::CountColumnId,         true, 100.0f });
-			InOutConfigSet.Add({ FMemAllocTableColumns::SizeColumnId,          true, 100.0f });
-			InOutConfigSet.Add({ FMemAllocTableColumns::TagColumnId,           true, 120.0f });
-			InOutConfigSet.Add({ FMemAllocTableColumns::AllocFunctionColumnId, true, 400.0f });
-		}
-	};
-	AvailableViewPresets.Add(MakeShared<FClassNameViewPreset>());
-
-	//////////////////////////////////////////////////
-	// (Inverted) Callstack Breakdown View
-
-	class FCallstackViewPreset : public ITableTreeViewPreset
-	{
-	public:
-		FCallstackViewPreset(bool bIsInverted, bool bIsAlloc)
-			: bIsInvertedCallstack(bIsInverted)
-			, bIsAllocCallstack(bIsAlloc)
-		{
-		}
-
-		virtual FText GetName() const override
-		{
-			return
-				bIsAllocCallstack
-				? (bIsInvertedCallstack ?
-				LOCTEXT("InvertedCallstack_Alloc_PresetName", "Inverted Alloc Callstack") :
-				LOCTEXT("Callstack_Alloc_PresetName", "Alloc Callstack"))
-				: (bIsInvertedCallstack ?
-				LOCTEXT("InvertedCallstack_Free_PresetName", "Inverted Free Callstack") :
-				LOCTEXT("Callstack_Free_PresetName", "Free Callstack"));
-		}
-		virtual FText GetToolTip() const override
-		{
-			return
-				bIsAllocCallstack
-				? (bIsInvertedCallstack ?
-				LOCTEXT("InvertedCallstack_Alloc_PresetToolTip", "Inverted Alloc Callstack Breakdown View\nConfigure the tree view to show a breakdown of allocations by inverted callstack.") :
-				LOCTEXT("Callstack_Alloc_PresetToolTip", "Alloc Callstack Breakdown View\nConfigure the tree view to show a breakdown of allocations by callstack."))
-				: (bIsInvertedCallstack ?
-				LOCTEXT("InvertedCallstack_Free_PresetToolTip", "Inverted Free Callstack Breakdown View\nConfigure the tree view to show a breakdown of allocations by inverted callstack.") :
-				LOCTEXT("Callstack_Free_PresetToolTip", "Free Callstack Breakdown View\nConfigure the tree view to show a breakdown of allocations by callstack."));
-		}
-		virtual FName GetSortColumn() const override
-		{
-			return FMemAllocTableColumns::SizeColumnId;
-		}
-		virtual EColumnSortMode::Type GetSortMode() const override
-		{
-			return EColumnSortMode::Type::Descending;
-		}
-		virtual void SetCurrentGroupings(const TArray<TSharedPtr<FTreeNodeGrouping>>& InAvailableGroupings, TArray<TSharedPtr<FTreeNodeGrouping>>& InOutCurrentGroupings) const override
-		{
-			InOutCurrentGroupings.Reset();
-
-			check(InAvailableGroupings[0]->Is<FTreeNodeGroupingFlat>());
-			InOutCurrentGroupings.Add(InAvailableGroupings[0]);
-
-			const bool bIsInverted = bIsInvertedCallstack;
-			const bool bIsAlloc = bIsAllocCallstack;
-			const TSharedPtr<FTreeNodeGrouping>* CallstackGrouping = InAvailableGroupings.FindByPredicate(
-				[bIsInverted, bIsAlloc](TSharedPtr<FTreeNodeGrouping>& Grouping)
-				{
-					return Grouping->Is<FMemAllocGroupingByCallstack>() &&
-						   Grouping->As<FMemAllocGroupingByCallstack>().IsInverted() == bIsInverted &&
-						   Grouping->As<FMemAllocGroupingByCallstack>().IsAllocCallstack() == bIsAlloc;
-				});
-			if (CallstackGrouping)
-			{
-				InOutCurrentGroupings.Add(*CallstackGrouping);
-			}
-		}
-		virtual void GetColumnConfigSet(TArray<FTableColumnConfig>& InOutConfigSet) const override
-		{
-			InOutConfigSet.Add({ FTable::GetHierarchyColumnId(),       true, 400.0f });
-			InOutConfigSet.Add({ FMemAllocTableColumns::CountColumnId, true, 100.0f });
-			InOutConfigSet.Add({ FMemAllocTableColumns::SizeColumnId,  true, 100.0f });
-			InOutConfigSet.Add({ FMemAllocTableColumns::TagColumnId,   true, 200.0f });
-			if (bIsAllocCallstack)
-			{
-				InOutConfigSet.Add({ FMemAllocTableColumns::AllocFunctionColumnId, true, 200.0f });
-			}
-			else
+			TOptional<FTableCellValue> CountValue = CountColumn->GetValue(*Node.Get());
+			if (CountValue.IsSet())
 			{
-				InOutConfigSet.Add({ FMemAllocTableColumns::FreeFunctionColumnId, true, 200.0f });
+				TotalAllocCount += CountValue.GetValue().AsInt64();
 			}
-		}
 
-	private:
-		bool bIsInvertedCallstack;
-		bool bIsAllocCallstack;
-	};
-	AvailableViewPresets.Add(MakeShared<FCallstackViewPreset>(false, true));
-	AvailableViewPresets.Add(MakeShared<FCallstackViewPreset>(true, true));
-	AvailableViewPresets.Add(MakeShared<FCallstackViewPreset>(false, false));
-	AvailableViewPresets.Add(MakeShared<FCallstackViewPreset>(true, false));
-
-	//////////////////////////////////////////////////
-	// Address (Platform Page) Breakdown View
-
-	class FPageViewPreset : public ITableTreeViewPreset
-	{
-	public:
-		virtual FText GetName() const override
-		{
-			return LOCTEXT("Page_PresetName", "Address (Platform Page)");
-		}
-		virtual FText GetToolTip() const override
-		{
-			return LOCTEXT("Page_PresetToolTip", "Platform Page Breakdown View\nConfigure the tree view to show a breakdown of allocations by their address.\nIt groups allocs into platform page size aligned memory pages.");
-		}
-		virtual FName GetSortColumn() const override
-		{
-			return FTable::GetHierarchyColumnId();
-		}
-		virtual EColumnSortMode::Type GetSortMode() const override
-		{
-			return EColumnSortMode::Type::Ascending;
-		}
-		virtual void SetCurrentGroupings(const TArray<TSharedPtr<FTreeNodeGrouping>>& InAvailableGroupings, TArray<TSharedPtr<FTreeNodeGrouping>>& InOutCurrentGroupings) const override
-		{
-			InOutCurrentGroupings.Reset();
-
-			check(InAvailableGroupings[0]->Is<FTreeNodeGroupingFlat>());
-			InOutCurrentGroupings.Add(InAvailableGroupings[0]);
-
-			const TSharedPtr<FTreeNodeGrouping>* MemoryPageGrouping = InAvailableGroupings.FindByPredicate(
-				[](TSharedPtr<FTreeNodeGrouping>& Grouping)
-				{
-					return Grouping->Is<FTreeNodeGroupingByUniqueValueInt64>() &&
-						   Grouping->As<FTreeNodeGroupingByUniqueValueInt64>().GetColumnId() == FMemAllocTableColumns::MemoryPageColumnId;
-				});
-			if (MemoryPageGrouping)
+			TOptional<FTableCellValue> SizeValue = SizeColumn->GetValue(*Node.Get());
+			if (SizeValue.IsSet())
 			{
-				InOutCurrentGroupings.Add(*MemoryPageGrouping);
+				TotalAllocSize += SizeValue.GetValue().AsInt64();
 			}
 		}
-		virtual void GetColumnConfigSet(TArray<FTableColumnConfig>& InOutConfigSet) const override
-		{
-			InOutConfigSet.Add({ FTable::GetHierarchyColumnId(),               true, 200.0f });
-			InOutConfigSet.Add({ FMemAllocTableColumns::AddressColumnId,       true, 120.0f });
-			InOutConfigSet.Add({ FMemAllocTableColumns::CountColumnId,         true, 100.0f });
-			InOutConfigSet.Add({ FMemAllocTableColumns::SizeColumnId,          true, 100.0f });
-			InOutConfigSet.Add({ FMemAllocTableColumns::TagColumnId,           true, 120.0f });
-			InOutConfigSet.Add({ FMemAllocTableColumns::AllocFunctionColumnId, true, 400.0f });
-		}
-	};
-	AvailableViewPresets.Add(MakeShared<FPageViewPreset>());
 
-	//////////////////////////////////////////////////
-	// Swap Breakdown View
+		FNumberFormattingOptions FormattingOptionsMem;
+		FormattingOptionsMem.MaximumFractionalDigits = 2;
 
-	class FSwapViewPreset : public ITableTreeViewPreset
+		SelectionStatsText = FText::Format(LOCTEXT("SelectionStatsFmt", "{0} selected {0}|plural(one=item,other=items) ({1} {1}|plural(one=alloc,other=allocs), {2})"),
+			FText::AsNumber(NumSelectedNodes), FText::AsNumber(TotalAllocCount), FText::AsMemory(TotalAllocSize, &FormattingOptionsMem));
+	}
+	else
 	{
-	public:
-		virtual FText GetName() const override
-		{
-			return LOCTEXT("Swap_PresetName", "Swap");
-		}
-		virtual FText GetToolTip() const override
-		{
-			return LOCTEXT("Swap_PresetToolTip", "Swap Usage Breakdown View\nConfigure the tree view to show a breakdown of allocations by their swap page.\nIt groups allocs into corresponding swap pages.");
-		}
-		virtual FName GetSortColumn() const override
-		{
-			return FTable::GetHierarchyColumnId();
-		}
-		virtual EColumnSortMode::Type GetSortMode() const override
-		{
-			return EColumnSortMode::Type::Ascending;
-		}
-		virtual void SetCurrentGroupings(const TArray<TSharedPtr<FTreeNodeGrouping>>& InAvailableGroupings, TArray<TSharedPtr<FTreeNodeGrouping>>& InOutCurrentGroupings) const override
-		{
-			InOutCurrentGroupings.Reset();
-
-			const TSharedPtr<FTreeNodeGrouping>* HeapGrouping = InAvailableGroupings.FindByPredicate(
-				[](TSharedPtr<FTreeNodeGrouping>& Grouping)
-				{
-					return Grouping->Is<FMemAllocGroupingBySwapPage>();
-				});
-			if (HeapGrouping)
-			{
-				InOutCurrentGroupings.Add(*HeapGrouping);
-			}
-		}
-		virtual void GetColumnConfigSet(TArray<FTableColumnConfig>& InOutConfigSet) const override
-		{
-			InOutConfigSet.Add({ FTable::GetHierarchyColumnId(),                    true, 200.0f });
-			InOutConfigSet.Add({ FMemAllocTableColumns::StartTimeColumnId,          true, 100.0f });
-			InOutConfigSet.Add({ FMemAllocTableColumns::EndTimeColumnId,            true, 100.0f });
-			InOutConfigSet.Add({ FMemAllocTableColumns::AddressColumnId,            true, 120.0f });
-			InOutConfigSet.Add({ FMemAllocTableColumns::CountColumnId,              true, 100.0f });
-			InOutConfigSet.Add({ FMemAllocTableColumns::CompressedSwapSizeColumnId, true, 100.0f });
-			InOutConfigSet.Add({ FMemAllocTableColumns::SwapSizeColumnId,           true, 100.0f });
-			InOutConfigSet.Add({ FMemAllocTableColumns::SizeColumnId,               true, 100.0f });
-			InOutConfigSet.Add({ FMemAllocTableColumns::TagColumnId,                true, 120.0f });
-			InOutConfigSet.Add({ FMemAllocTableColumns::AllocFunctionColumnId,      true, 400.0f });
-		}
-	};
-	AvailableViewPresets.Add(MakeShared<FSwapViewPreset>());
-
-	//////////////////////////////////////////////////
-
-	SelectedViewPreset = AvailableViewPresets[0];
+		SelectionStatsText = FText::GetEmpty();
+	}
 }
 
 ////////////////////////////////////////////////////////////////////////////////////////////////////
 
-BEGIN_SLATE_FUNCTION_BUILD_OPTIMIZATION
-TSharedPtr<SWidget> SMemAllocTableTreeView::ConstructFooter()
+FText SMemAllocTableTreeView::GetFooterCenterText() const
 {
-	return
-		SNew(SHorizontalBox)
-
-		+ SHorizontalBox::Slot()
-		.AutoWidth()
-		.HAlign(HAlign_Left)
-		.Padding(2.0f, 2.0f, 0.0f, 2.0f)
-		[
-			SNew(STextBlock)
-			.Text(this, &SMemAllocTableTreeView::GetQueryInfo)
-			.ToolTipText(this, &SMemAllocTableTreeView::GetQueryInfoTooltip)
-		]
-
-		+ SHorizontalBox::Slot()
-		.AutoWidth()
-		.HAlign(HAlign_Left)
-		.Padding(0.0f, 2.0f, 0.0f, 2.0f)
-		[
-			SNew(STextBlock)
-			.Text(LOCTEXT("FooterSeparator", " : "))
-			.ColorAndOpacity(FSlateColor(EStyleColor::White25))
-		]
-
-		+ SHorizontalBox::Slot()
-		.AutoWidth()
-		.HAlign(HAlign_Left)
-		.Padding(0.0f, 2.0f, 8.0f, 2.0f)
-		[
-			SNew(STextBlock)
-			.Text(this, &SMemAllocTableTreeView::GetFooterLeftText)
-		]
-
-		+ SHorizontalBox::Slot()
-		.FillWidth(1.0f)
-		.HAlign(HAlign_Center)
-		.Padding(2.0f, 2.0f, 2.0f, 2.0f)
-		[
-			SNew(STextBlock)
-			.Text(this, &SMemAllocTableTreeView::GetFooterCenterText)
-		]
-
-		+ SHorizontalBox::Slot()
-		.AutoWidth()
-		.HAlign(HAlign_Right)
-		.Padding(8.0f, 2.0f, 2.0f, 2.0f)
-		[
-			SNew(STextBlock)
-			.Text(this, &SMemAllocTableTreeView::GetSymbolResolutionStatus)
-			.ToolTipText(this, &SMemAllocTableTreeView::GetSymbolResolutionTooltip)
-		];
+	return SelectionStatsText;
 }
-END_SLATE_FUNCTION_BUILD_OPTIMIZATION
 
 ////////////////////////////////////////////////////////////////////////////////////////////////////
 
@@ -1312,176 +901,9 @@ FText SMemAllocTableTreeView::GetSymbolResolutionTooltip() const
 
 ////////////////////////////////////////////////////////////////////////////////////////////////////
 
-FText SMemAllocTableTreeView::GetQueryInfo() const
-{
-	return QueryInfo;
-}
-
-////////////////////////////////////////////////////////////////////////////////////////////////////
-
-FText SMemAllocTableTreeView::GetQueryInfoTooltip() const
-{
-	return Rule.IsValid() ? Rule->GetDescription() : FText::GetEmpty();
-}
-
-////////////////////////////////////////////////////////////////////////////////////////////////////
-
-FText SMemAllocTableTreeView::GetFooterLeftText() const
-{
-	if (!Rule.IsValid())
-	{
-		return FText::GetEmpty();
-	}
-
-	if (Query != 0)
-	{
-		return LOCTEXT("FooterLeftTextRunningQuery", "running query...");
-	}
-
-	if (FilteredNodesPtr->Num() == TableRowNodes.Num())
-	{
-		return FText::Format(LOCTEXT("FooterLeftTextFmt1", "{0} {0}|plural(one=alloc,other=allocs)"), FText::AsNumber(TableRowNodes.Num()));
-	}
-	else
-	{
-		return FText::Format(LOCTEXT("FooterLeftTextFmt2", "{0} / {1} {1}|plural(one=alloc,other=allocs)"), FilteredNodesPtr->Num(), FText::AsNumber(TableRowNodes.Num()));
-	}
-}
-
-////////////////////////////////////////////////////////////////////////////////////////////////////
-
-FText SMemAllocTableTreeView::GetFooterCenterText() const
-{
-	return SelectionStatsText;
-}
-
-////////////////////////////////////////////////////////////////////////////////////////////////////
-
 void SMemAllocTableTreeView::TreeView_OnSelectionChanged(FTableTreeNodePtr SelectedItem, ESelectInfo::Type SelectInfo)
 {
-	TArray<FTableTreeNodePtr> SelectedNodes;
-	const int32 NumSelectedNodes = TreeView->GetSelectedItems(SelectedNodes);
-
-	if (NumSelectedNodes > 0)
-	{
-		int64 TotalAllocCount = 0;
-		int64 TotalAllocSize = 0;
-
-		TSharedRef<FTableColumn> CountColumn = Table->FindColumnChecked(FMemAllocTableColumns::CountColumnId);
-		TSharedRef<FTableColumn> SizeColumn = Table->FindColumnChecked(FMemAllocTableColumns::SizeColumnId);
-
-		for (const FTableTreeNodePtr& Node : SelectedNodes)
-		{
-			TOptional<FTableCellValue> CountValue = CountColumn->GetValue(*Node.Get());
-			if (CountValue.IsSet())
-			{
-				TotalAllocCount += CountValue.GetValue().AsInt64();
-			}
-
-			TOptional<FTableCellValue> SizeValue = SizeColumn->GetValue(*Node.Get());
-			if (SizeValue.IsSet())
-			{
-				TotalAllocSize += SizeValue.GetValue().AsInt64();
-			}
-		}
-
-		FNumberFormattingOptions FormattingOptionsMem;
-		FormattingOptionsMem.MaximumFractionalDigits = 2;
-
-		SelectionStatsText = FText::Format(LOCTEXT("SelectionStatsFmt", "{0} selected {0}|plural(one=item,other=items) ({1} {1}|plural(one=alloc,other=allocs), {2})"),
-			FText::AsNumber(NumSelectedNodes), FText::AsNumber(TotalAllocCount), FText::AsMemory(TotalAllocSize, &FormattingOptionsMem));
-	}
-	else
-	{
-		SelectionStatsText = FText::GetEmpty();
-	}
-}
-
-////////////////////////////////////////////////////////////////////////////////////////////////////
-
-void SMemAllocTableTreeView::InternalCreateGroupings()
-{
-	STableTreeView::InternalCreateGroupings();
-
-	AvailableGroupings.RemoveAll(
-		[](TSharedPtr<FTreeNodeGrouping>& Grouping)
-		{
-			if (Grouping->Is<FTreeNodeGroupingByUniqueValue>())
-			{
-				const FName ColumnId = Grouping->As<FTreeNodeGroupingByUniqueValue>().GetColumnId();
-				if (ColumnId == FMemAllocTableColumns::StartEventIndexColumnId ||
-					ColumnId == FMemAllocTableColumns::EndEventIndexColumnId ||
-					ColumnId == FMemAllocTableColumns::CountColumnId)
-				{
-					return true;
-				}
-			}
-			else if (Grouping->Is<FTreeNodeGroupingByPathBreakdown>())
-			{
-				const FName ColumnId = Grouping->As<FTreeNodeGroupingByPathBreakdown>().GetColumnId();
-				if (ColumnId == FMemAllocTableColumns::AllocFunctionColumnId ||
-					ColumnId == FMemAllocTableColumns::FreeFunctionColumnId ||
-					ColumnId == FMemAllocTableColumns::ClassNameColumnId)
-				{
-					return true;
-				}
-			}
-			return false;
-		});
-
-	int32 Index = 1; // after the Flat ("All") grouping
-
-	AvailableGroupings.Insert(MakeShared<FMemAllocGroupingBySize>(), Index++);
-
-	const TraceServices::IAllocationsProvider* AllocationsProvider = Session.IsValid() ? TraceServices::ReadAllocationsProvider(*Session.Get()) : nullptr;
-
-	if (AllocationsProvider)
-	{
-		AvailableGroupings.Insert(MakeShared<FMemAllocGroupingByTag>(*AllocationsProvider), Index++);
-	}
-
-	AvailableGroupings.Insert(MakeShared<FMemAllocGroupingByCallstack>(true, false, bIsCallstackGroupingByFunction), Index++);
-	AvailableGroupings.Insert(MakeShared<FMemAllocGroupingByCallstack>(true, true, bIsCallstackGroupingByFunction), Index++);
-	AvailableGroupings.Insert(MakeShared<FMemAllocGroupingByCallstack>(false, false, bIsCallstackGroupingByFunction), Index++);
-	AvailableGroupings.Insert(MakeShared<FMemAllocGroupingByCallstack>(false, true, bIsCallstackGroupingByFunction), Index++);
-
-	if (AllocationsProvider)
-	{
-		AvailableGroupings.Insert(MakeShared<FMemAllocGroupingByHeap>(*AllocationsProvider), Index++);
-		AvailableGroupings.Insert(MakeShared<FMemAllocGroupingBySwapPage>(*AllocationsProvider), Index++);
-	}
-}
-
-////////////////////////////////////////////////////////////////////////////////////////////////////
-
-void SMemAllocTableTreeView::UpdateQueryInfo()
-{
-	if (Rule.IsValid())
-	{
-		FText TimeMarkersText;
-
-		const int NumMarkers = Rule->GetNumTimeMarkers();
-		switch (NumMarkers)
-		{
-		case 1:
-			TimeMarkersText = FText::Format(LOCTEXT("OneTimeMarkersFmt", "A={0}"), FText::AsNumber(TimeMarkers[0]));
-			break;
-		case 2:
-			TimeMarkersText = FText::Format(LOCTEXT("TwoTimeMarkersFmt", "A={0}  B={1}"), FText::AsNumber(TimeMarkers[0]), FText::AsNumber(TimeMarkers[1]));
-			break;
-		case 3:
-			TimeMarkersText = FText::Format(LOCTEXT("ThreeTimeMarkersFmt", "A={0}  B={1}  C={2}"), FText::AsNumber(TimeMarkers[0]), FText::AsNumber(TimeMarkers[1]), FText::AsNumber(TimeMarkers[2]));
-			break;
-		case 4:
-			TimeMarkersText = FText::Format(LOCTEXT("FourTimeMarkersFmt", "A={0}  B={1}  C={2}  D={3}"), FText::AsNumber(TimeMarkers[0]), FText::AsNumber(TimeMarkers[1]), FText::AsNumber(TimeMarkers[2]), FText::AsNumber(TimeMarkers[3]));
-			break;
-		default:
-			// Unhandled value
-			check(false);
-		}
-
-		QueryInfo = FText::Format(LOCTEXT("QueryInfoFmt", "{0} ({1})"), Rule->GetVerboseName(), TimeMarkersText);
-	}
+	UpdateSelectionStatsText();
 }
 
 ////////////////////////////////////////////////////////////////////////////////////////////////////
diff --git a/Engine/Source/Developer/TraceInsights/Private/Insights/MemoryProfiler/Widgets/SMemAllocTableTreeView.h b/Engine/Source/Developer/TraceInsights/Private/Insights/MemoryProfiler/Widgets/SMemAllocTableTreeView.h
--- a/Engine/Source/Developer/TraceInsights/Private/Insights/MemoryProfiler/Widgets/SMemAllocTableTreeView.h
+++ b/Engine/Source/Developer/TraceInsights/Private/Insights/MemoryProfiler/Widgets/SMemAllocTableTreeView.h
@@ -2,7 +2,7 @@
 
 #pragma once
 
-#include "CoreMinimal.h"
+#include "CoreTypes.h"
 
 // TraceServices
 #include "TraceServices/Model/AllocationsProvider.h"
@@ -14,6 +14,11 @@
 #include "Insights/MemoryProfiler/ViewModels/MemAllocTable.h"
 #include "Insights/Table/Widgets/SSessionTableTreeView.h"
 
+namespace TraceServices
+{
+	struct FCallstack;
+}
+
 namespace UE::Insights::MemoryProfiler
 {
 
@@ -44,10 +49,15 @@ public:
 	virtual TSharedPtr<SWidget> ConstructToolbar() override;
 	virtual TSharedPtr<SWidget> ConstructFooter() override;
 
-	TSharedPtr<FMemAllocTable> GetMemAllocTable() { return StaticCastSharedPtr<FMemAllocTable>(GetTable()); }
-	const TSharedPtr<FMemAllocTable> GetMemAllocTable() const { return StaticCastSharedPtr<FMemAllocTable>(GetTable()); }
+	TSharedPtr<FMemAllocTable> GetMemAllocTable()
+	{
+		return StaticCastSharedPtr<FMemAllocTable>(GetTable());
+	}
 
-	//void UpdateSourceTable(TSharedPtr<TraceServices::IMemAllocTable> SourceTable);
+	TSharedPtr<const FMemAllocTable> GetMemAllocTable() const
+	{
+		return StaticCastSharedPtr<const FMemAllocTable>(GetTable());
+	}
 
 	virtual void Reset();
 
@@ -115,36 +125,48 @@ protected:
 	void ExportMemorySnapshot() const;
 	bool IsExportMemorySnapshotAvailable() const;
 
+	virtual void TreeView_OnSelectionChanged(FTableTreeNodePtr SelectedItem, ESelectInfo::Type SelectInfo) override;
+
+	virtual void UpdateFilterContext(const FFilterConfigurator& InFilterConfigurator, const FTableTreeNode& InNode) const override;
+	virtual void InitFilterConfigurator(FFilterConfigurator& InOutFilterConfigurator) override;
+
 private:
+	void InitAvailableViewPresets();
+
 	void ExtendMenuAllocation(FMenuBuilder& MenuBuilder);
 	void ExtendMenuCallstackFrame(FMenuBuilder& MenuBuilder);
 	void ExtendMenuExportSnapshot(FMenuBuilder& MenuBuilder);
 
+	//////////////////////////////////////////////////
+	// Queries
+
 	void OnQueryInvalidated();
 	void StartQuery();
 	void UpdateQuery(TraceServices::IAllocationsProvider::EQueryStatus& OutStatus);
 	void CancelQuery();
 	void ResetAndStartQuery();
 
-	FText GetSymbolResolutionStatus() const;
-	FText GetSymbolResolutionTooltip() const;
-	
+	//////////////////////////////////////////////////
+	// Footer / Status Bar
+
 	void UpdateQueryInfo();
 	FText GetQueryInfo() const;
 	FText GetQueryInfoTooltip() const;
+
 	FText GetFooterLeftText() const;
+
+	void UpdateSelectionStatsText();
 	FText GetFooterCenterText() const;
 
-	virtual void TreeView_OnSelectionChanged(FTableTreeNodePtr SelectedItem, ESelectInfo::Type SelectInfo) override;
+	FText GetSymbolResolutionStatus() const;
+	FText GetSymbolResolutionTooltip() const;
 
-	virtual void UpdateFilterContext(const FFilterConfigurator& InFilterConfigurator, const FTableTreeNode& InNode) const override;
-	virtual void InitFilterConfigurator(FFilterConfigurator& InOutFilterConfigurator) override;
+	//////////////////////////////////////////////////
 
 	TSharedRef<SWidget> ConstructFunctionToggleButton();
 	void CallstackGroupingByFunction_OnCheckStateChanged(ECheckBoxState NewRadioState);
 	ECheckBoxState CallstackGroupingByFunction_IsChecked() const;
 
-	void InitAvailableViewPresets();
 	void PopulateLLMTagSuggestionList(const FString& Text, TArray<FString>& OutSuggestions) const;
 	void PopulateThreadSuggestionList(const FString& Text, TArray<FString>& OutSuggestions) const;
 
