[Insights]
- Memory Insights: Upgraded the LLM Tags panel to use the new Tree View implementation.
- Memory Insights: Added "Size A",  "Size B", "Diff (B-A)" columns to LLM Tags table. Shows the size value of each LLM tag at two time markers A and B and variation between the two snapshots (B-A).
- Memory Insights: Added "Min",  "Max", "Average" columns to LLM Tags table. Shows min, max and average size values for selected time range, for each LLM tag.

diff --git a/Engine/Source/Developer/TraceInsights/Private/Insights/InsightsStyle.cpp b/Engine/Source/Developer/TraceInsights/Private/Insights/InsightsStyle.cpp
--- a/Engine/Source/Developer/TraceInsights/Private/Insights/InsightsStyle.cpp
+++ b/Engine/Source/Developer/TraceInsights/Private/Insights/InsightsStyle.cpp
@@ -270,7 +270,8 @@ void FInsightsStyle::FStyle::Initialize()
 	Set("Icons.DataTypeDouble.TreeItem", new TODO_IMAGE_BRUSH(Icon12x12));
 	Set("Icons.DataTypeInt64.TreeItem", new TODO_IMAGE_BRUSH(Icon12x12));
 	Set("Icons.NetEvent.TreeItem", new TODO_IMAGE_BRUSH(Icon12x12));
-	Set("Icons.MemTag.TreeItem", new TODO_IMAGE_BRUSH(Icon12x12));
+	Set("Icons.MemTag.TreeItem", new IMAGE_BRUSH_SVG("MemTags", Icon12x12));
+	Set("Icons.MemTagWithGraph.TreeItem", new IMAGE_BRUSH_SVG("MemTags", Icon12x12));
 
 	//////////////////////////////////////////////////
 	// Trace Control
diff --git a/Engine/Source/Developer/TraceInsights/Private/Insights/MemoryProfiler/ViewModels/MemTagNode.cpp b/Engine/Source/Developer/TraceInsights/Private/Insights/MemoryProfiler/ViewModels/MemTagNode.cpp
--- a/Engine/Source/Developer/TraceInsights/Private/Insights/MemoryProfiler/ViewModels/MemTagNode.cpp
+++ b/Engine/Source/Developer/TraceInsights/Private/Insights/MemoryProfiler/ViewModels/MemTagNode.cpp
@@ -2,7 +2,10 @@
 
 #include "MemTagNode.h"
 
+#include "HAL/UnrealMemory.h"
+
 // TraceInsights
+#include "Insights/InsightsStyle.h"
 #include "Insights/MemoryProfiler/MemoryProfilerManager.h"
 
 #define LOCTEXT_NAMESPACE "UE::Insights::MemoryProfiler::FMemTagNode"
@@ -14,6 +17,30 @@ INSIGHTS_IMPLEMENT_RTTI(FMemTagNode)
 
 ////////////////////////////////////////////////////////////////////////////////////////////////////
 
+FText FMemTagNode::GetTagText() const
+{
+	if (MemTag)
+	{
+		return FText::FromString(MemTag->GetStatFullName());
+	}
+
+	return FText::GetEmpty();
+}
+
+////////////////////////////////////////////////////////////////////////////////////////////////////
+
+const TCHAR* FMemTagNode::GetTagName() const
+{
+	if (MemTag)
+	{
+		return *MemTag->GetStatFullName();
+	}
+
+	return nullptr;
+}
+
+////////////////////////////////////////////////////////////////////////////////////////////////////
+
 FText FMemTagNode::GetTrackerText() const
 {
 	FMemorySharedState* SharedState = FMemoryProfilerManager::Get()->GetSharedState();
@@ -32,6 +59,24 @@ FText FMemTagNode::GetTrackerText() const
 
 ////////////////////////////////////////////////////////////////////////////////////////////////////
 
+const TCHAR* FMemTagNode::GetTrackerName() const
+{
+	FMemorySharedState* SharedState = FMemoryProfilerManager::Get()->GetSharedState();
+	if (SharedState)
+	{
+		FMemoryTrackerId TrackerId = GetMemTrackerId();
+		const FMemoryTracker* Tracker = SharedState->GetTrackerById(TrackerId);
+		if (Tracker)
+		{
+			return *Tracker->GetName();
+		}
+	}
+
+	return nullptr;
+}
+
+////////////////////////////////////////////////////////////////////////////////////////////////////
+
 FText FMemTagNode::GetTagSetText() const
 {
 	FMemorySharedState* SharedState = FMemoryProfilerManager::Get()->GetSharedState();
@@ -50,18 +95,115 @@ FText FMemTagNode::GetTagSetText() const
 
 ////////////////////////////////////////////////////////////////////////////////////////////////////
 
-void FMemTagNode::ResetAggregatedStats()
+const TCHAR* FMemTagNode::GetTagSetName() const
+{
+	FMemorySharedState* SharedState = FMemoryProfilerManager::Get()->GetSharedState();
+	if (SharedState)
+	{
+		FMemoryTagSetId TagSetId = GetMemTagSetId();
+		const FMemoryTagSet* TagSet = SharedState->GetTagSetById(TagSetId);
+		if (TagSet)
+		{
+			return *TagSet->GetName();
+		}
+	}
+
+	return nullptr;
+}
+
+////////////////////////////////////////////////////////////////////////////////////////////////////
+
+const FText FMemTagNode::GetDisplayName() const
+{
+	return FText::FromString(TEXT("tag"));
+}
+
+////////////////////////////////////////////////////////////////////////////////////////////////////
+
+const FText FMemTagNode::GetExtraDisplayName() const
 {
-	//AggregatedStats = TraceServices::FMemoryProfilerAggregatedStats();
+	FMemorySharedState* SharedState = FMemoryProfilerManager::Get()->GetSharedState();
+	if (SharedState)
+	{
+		if (MemTag && MemTag->GetTrackerId() != SharedState->GetDefaultTrackerId())
+		{
+			return FText::Format(LOCTEXT("MemTagNodeExtraDisplayNameFmt", "({0})"), GetTrackerText());
+		}
+	}
+	return FText::GetEmpty();
+}
+
+////////////////////////////////////////////////////////////////////////////////////////////////////
+
+bool FMemTagNode::HasExtraDisplayName() const
+{
+	FMemorySharedState* SharedState = FMemoryProfilerManager::Get()->GetSharedState();
+	if (SharedState)
+	{
+		if (MemTag && MemTag->GetTrackerId() != SharedState->GetDefaultTrackerId())
+		{
+			return true;
+		}
+	}
+	return false;
+}
+
+////////////////////////////////////////////////////////////////////////////////////////////////////
+
+const FText FMemTagNode::GetTooltipText() const
+{
+	return FText::GetEmpty();
+}
+
+////////////////////////////////////////////////////////////////////////////////////////////////////
+
+const FSlateBrush* FMemTagNode::GetIcon() const
+{
+	if (MemTag && MemTag->IsAddedToGraph())
+	{
+		return FInsightsStyle::GetBrush("Icons.MemTagWithGraph.TreeItem");
+	}
+	else
+	{
+		return FInsightsStyle::GetBrush("Icons.MemTag.TreeItem");
+	}
 }
 
 ////////////////////////////////////////////////////////////////////////////////////////////////////
-/*
-void FMemTagNode::SetAggregatedStats(const TraceServices::FMemoryProfilerAggregatedStats& InAggregatedStats)
+
+FLinearColor FMemTagNode::GetIconColor() const
 {
-	AggregatedStats = InAggregatedStats;
+	if (MemTag && MemTag->IsAddedToGraph())
+	{
+		return MemTag->GetColor();
+	}
+	else
+	{
+		return FLinearColor(0.1f, 0.1f, 0.1f, 1.0f);
+	}
 }
-*/
+
+////////////////////////////////////////////////////////////////////////////////////////////////////
+
+FLinearColor FMemTagNode::GetColor() const
+{
+	if (MemTag && MemTag->IsAddedToGraph())
+	{
+		return MemTag->GetColor();
+	}
+	else
+	{
+		return GetDefaultColor(IsGroup());
+	}
+}
+
+////////////////////////////////////////////////////////////////////////////////////////////////////
+
+void FMemTagNode::ResetAggregatedStats()
+{
+	FMemory::Memzero<FMemTagStats>(Stats);
+}
+
 ////////////////////////////////////////////////////////////////////////////////////////////////////
 
 } // namespace UE::Insights::MemoryProfiler
diff --git a/Engine/Source/Developer/TraceInsights/Private/Insights/MemoryProfiler/ViewModels/MemTagNode.h b/Engine/Source/Developer/TraceInsights/Private/Insights/MemoryProfiler/ViewModels/MemTagNode.h
--- a/Engine/Source/Developer/TraceInsights/Private/Insights/MemoryProfiler/ViewModels/MemTagNode.h
+++ b/Engine/Source/Developer/TraceInsights/Private/Insights/MemoryProfiler/ViewModels/MemTagNode.h
@@ -8,12 +8,14 @@
 #include "TraceServices/Model/Memory.h"
 
 // TraceInsightsCore
-#include "InsightsCore/Table/ViewModels/BaseTreeNode.h"
+#include "InsightsCore/Table/ViewModels/TableTreeNode.h"
+#include "InsightsCore/Table/ViewModels/TreeNodeGrouping.h"
 
 // TraceInsights
 #include "Insights/MemoryProfiler/ViewModels/MemorySharedState.h"
 #include "Insights/MemoryProfiler/ViewModels/MemoryTag.h"
 #include "Insights/MemoryProfiler/ViewModels/MemoryTracker.h"
+#include "Insights/MemoryProfiler/ViewModels/MemTagTable.h"
 
 namespace TraceServices
 {
@@ -33,16 +35,15 @@ namespace UE::Insights::MemoryProfiler
 
 ////////////////////////////////////////////////////////////////////////////////////////////////////
 
-enum class EMemTagNodeType
+struct FMemTagStats
 {
-	/** The MemTagNode is a Low Level Memory Tag. */
-	MemTag,
-
-	/** The MemTagNode is a group node. */
-	Group,
-
-	/** Invalid enum type, may be used as a number of enumerations. */
-	InvalidOrMax,
+	int64 SizeA = 0; // memory size, in [bytes], for current LLM tag, at time marker A
+	int64 SizeB = 0; // memory size, in [bytes], for current LLM tag, at time marker B
+	int64 SizeBudget = 0;
+	int64 SampleCount = 0;
+	int64 SizeMin = 0;
+	int64 SizeMax = 0;
+	int64 SizeAverage = 0;
 };
 
 ////////////////////////////////////////////////////////////////////////////////////////////////////
@@ -63,67 +64,75 @@ typedef TWeakPtr<class FMemTagNode> FMemTagNodeWeak;
 
 ////////////////////////////////////////////////////////////////////////////////////////////////////
 /**
- * Class used to store information about an llm tag node (used in the SMemTagTreeView).
+ * Class used to store information about an LLM tag node (used in the SMemTagTreeView).
  */
-class FMemTagNode : public FBaseTreeNode
+class FMemTagNode : public FTableTreeNode
 {
-	INSIGHTS_DECLARE_RTTI(FMemTagNode, FBaseTreeNode)
+	INSIGHTS_DECLARE_RTTI(FMemTagNode, FTableTreeNode)
 
 public:
 	/** Initialization constructor for the MemTag node. */
-	explicit FMemTagNode(FMemoryTag* InMemTag)
-		: FBaseTreeNode(FName(InMemTag->GetStatFullName(), 0), false)
-		, Type(EMemTagNodeType::MemTag)
-		, MemTag(InMemTag)
+	explicit FMemTagNode(TWeakPtr<FMemTagTable> InParentTable, FMemoryTag& InMemTag)
+		: FTableTreeNode(FName("tag", InMemTag.GetIndex() + 1), InParentTable, InMemTag.GetIndex())
+		, MemTag(&InMemTag)
 	{
-		ResetAggregatedStats();
 	}
 
 	/** Initialization constructor for the group node. */
-	explicit FMemTagNode(const FName InGroupName)
-		: FBaseTreeNode(InGroupName, true)
-		, Type(EMemTagNodeType::Group)
+	explicit FMemTagNode(TWeakPtr<FMemTagTable> InParentTable, const FName InGroupName)
+		: FTableTreeNode(InGroupName, InParentTable)
 		, MemTag(nullptr)
 	{
-		ResetAggregatedStats();
 	}
 
-	/**
-	 * @return a type of this MemTag node or EMemTagNodeType::Group for group nodes.
-	 */
-	EMemTagNodeType GetType() const { return Type; }
-
-	bool IsValidStat() const { return MemTag != nullptr; }
+	bool IsValidMemTag() const { return MemTag != nullptr; }
 	FMemoryTag* GetMemTag() const { return MemTag; }
 
 	FMemoryTagId GetMemTagId() const { return MemTag ? MemTag->GetId() : FMemoryTag::InvalidTagId; }
+	FText GetTagText() const;
+	const TCHAR* GetTagName() const;
 
 	FMemoryTrackerId GetMemTrackerId() const { return MemTag ? MemTag->GetTrackerId() : FMemoryTracker::InvalidTrackerId; }
 	FText GetTrackerText() const;
+	const TCHAR* GetTrackerName() const;
 
 	FMemoryTagSetId GetMemTagSetId() const { return MemTag ? MemTag->GetTagSetId() : FMemoryTagSet::InvalidTagSetId; }
 	FText GetTagSetText() const;
+	const TCHAR* GetTagSetName() const;
 
-	FLinearColor GetColor() const { return MemTag ? MemTag->GetColor() : FLinearColor(0.5f, 0.5f, 0.5f, 1.0f); }
-	bool IsAddedToGraph() const { return MemTag ? MemTag->IsAddedToGraph() : false; }
+	virtual const FText GetDisplayName() const override;
+	virtual const FText GetExtraDisplayName() const override;
+	virtual bool HasExtraDisplayName() const override;
+	virtual const FText GetTooltipText() const override;
+
+	bool IsAddedToGraph() const { return MemTag && MemTag->IsAddedToGraph(); }
+
+	virtual const FSlateBrush* GetIcon() const override;
+	virtual FLinearColor GetIconColor() const override;
+	virtual FLinearColor GetColor() const override;
 
 	FMemTagNodePtr GetParentTagNode() const { return ParentTagNode; }
 	FMemoryTag* GetParentMemTag() const { return ParentTagNode.IsValid() ? ParentTagNode->GetMemTag() : nullptr; }
 	void SetParentTagNode(FMemTagNodePtr NodePtr) { ParentTagNode = NodePtr; }
 
-	/**
-	 * @return the aggregated stats for this MemTag node.
-	 */
-	const TraceServices::FMemoryProfilerAggregatedStats& GetAggregatedStats() const { return AggregatedStats; }
+	const FMemTagStats& GetStats() const { return Stats; }
+	FMemTagStats& GetStats() { return Stats; }
+
+	int64 GetSizeA() const { return Stats.SizeA; }
+	int64 GetSizeB() const { return Stats.SizeB; }
+	int64 GetSizeDiff() const { return Stats.SizeB - Stats.SizeA; }
+	int64 GetSizeBudget() const { return Stats.SizeBudget; }
+	int64 GetSampleCount() const { return Stats.SampleCount; }
+	int64 GetSizeMin() const { return Stats.SizeMin; }
+	int64 GetSizeMax() const { return Stats.SizeMax; }
+	int64 GetSizeAverage() const { return Stats.SizeAverage; }
 
 	void ResetAggregatedStats();
-	//TODO: void SetAggregatedStats(const TraceServices::FMemoryProfilerAggregatedStats& AggregatedStats);
 
 private:
-	const EMemTagNodeType Type;
 	FMemoryTag* MemTag;
-	FMemTagNodePtr ParentTagNode;
-	TraceServices::FMemoryProfilerAggregatedStats AggregatedStats;
+	FMemTagNodePtr ParentTagNode = nullptr;
+	FMemTagStats Stats;
 };
 
 ////////////////////////////////////////////////////////////////////////////////////////////////////
diff --git a/Engine/Source/Developer/TraceInsights/Private/Insights/MemoryProfiler/ViewModels/MemTagNodeGroupingAndSorting.cpp b/Engine/Source/Developer/TraceInsights/Private/Insights/MemoryProfiler/ViewModels/MemTagNodeGroupingAndSorting.cpp
--- a/Engine/Source/Developer/TraceInsights/Private/Insights/MemoryProfiler/ViewModels/MemTagNodeGroupingAndSorting.cpp
+++ b/Engine/Source/Developer/TraceInsights/Private/Insights/MemoryProfiler/ViewModels/MemTagNodeGroupingAndSorting.cpp
@@ -5,9 +5,6 @@
 // TraceInsightsCore
 #include "InsightsCore/Table/ViewModels/TableColumn.h"
 
-// TraceInsights
-#include "Insights/MemoryProfiler/ViewModels/MemTagNodeHelper.h"
-
 #define LOCTEXT_NAMESPACE "UE::Insights::MemoryProfiler::FMemTagNode"
 
 #define INSIGHTS_ENSURE ensure
@@ -33,72 +30,6 @@
 namespace UE::Insights::MemoryProfiler
 {
 
-////////////////////////////////////////////////////////////////////////////////////////////////////
-// Sorting by Type
-////////////////////////////////////////////////////////////////////////////////////////////////////
-
-FMemTagNodeSortingByType::FMemTagNodeSortingByType(TSharedRef<FTableColumn> InColumnRef)
-	: FTableCellValueSorter(
-		FName(TEXT("ByType")),
-		LOCTEXT("Sorting_ByType_Name", "By Type"),
-		LOCTEXT("Sorting_ByType_Title", "Sort By Type"),
-		LOCTEXT("Sorting_ByType_Desc", "Sort by type of tree nodes."),
-		InColumnRef)
-{
-}
-
-////////////////////////////////////////////////////////////////////////////////////////////////////
-
-void FMemTagNodeSortingByType::Sort(TArray<FBaseTreeNodePtr>& NodesToSort, ESortMode SortMode) const
-{
-	if (SortMode == ESortMode::Ascending)
-	{
-		NodesToSort.Sort([this](const FBaseTreeNodePtr& A, const FBaseTreeNodePtr& B) -> bool
-		{
-			INSIGHTS_DEFAULT_PRESORTING_NODES(A, B);
-
-			INSIGHTS_ENSURE(A.IsValid() && A->Is<FMemTagNode>());
-			const FMemTagNodePtr MemTagNodeA = StaticCastSharedPtr<FMemTagNode, FBaseTreeNode>(A);
-
-			INSIGHTS_ENSURE(B.IsValid() && B->Is<FMemTagNode>());
-			const FMemTagNodePtr MemTagNodeB = StaticCastSharedPtr<FMemTagNode, FBaseTreeNode>(B);
-
-			if (MemTagNodeA->GetType() == MemTagNodeB->GetType())
-			{
-				INSIGHTS_DEFAULT_SORTING_NODES(A, B)
-			}
-			else
-			{
-				// Sort by type (ascending).
-				return MemTagNodeA->GetType() < MemTagNodeB->GetType();
-			}
-		});
-	}
-	else // if (SortMode == ESortMode::Descending)
-	{
-		NodesToSort.Sort([this](const FBaseTreeNodePtr& A, const FBaseTreeNodePtr& B) -> bool
-		{
-			INSIGHTS_DEFAULT_PRESORTING_NODES(A, B);
-
-			INSIGHTS_ENSURE(A.IsValid() && A->Is<FMemTagNode>());
-			const FMemTagNodePtr MemTagNodeA = StaticCastSharedPtr<FMemTagNode, FBaseTreeNode>(A);
-
-			INSIGHTS_ENSURE(B.IsValid() && B->Is<FMemTagNode>());
-			const FMemTagNodePtr MemTagNodeB = StaticCastSharedPtr<FMemTagNode, FBaseTreeNode>(B);
-
-			if (MemTagNodeA->GetType() == MemTagNodeB->GetType())
-			{
-				INSIGHTS_DEFAULT_SORTING_NODES(A, B)
-			}
-			else
-			{
-				// Sort by type (descending).
-				return MemTagNodeB->GetType() < MemTagNodeA->GetType();
-			}
-		});
-	}
-}
-
 ////////////////////////////////////////////////////////////////////////////////////////////////////
 // Sorting by Tracker(s)
 ////////////////////////////////////////////////////////////////////////////////////////////////////
@@ -165,76 +96,6 @@ void FMemTagNodeSortingByTracker::Sort(TArray<FBaseTreeNodePtr>& NodesToSort, ES
 	}
 }
 
-////////////////////////////////////////////////////////////////////////////////////////////////////
-// Sort by Instance Count
-////////////////////////////////////////////////////////////////////////////////////////////////////
-
-FMemTagNodeSortingByInstanceCount::FMemTagNodeSortingByInstanceCount(TSharedRef<FTableColumn> InColumnRef)
-	: FTableCellValueSorter(
-		FName(TEXT("ByInstanceCount")),
-		LOCTEXT("Sorting_ByInstanceCount_Name", "By Instance Count"),
-		LOCTEXT("Sorting_ByInstanceCount_Title", "Sort By Instance Count"),
-		LOCTEXT("Sorting_ByInstanceCount_Desc", "Sort by aggregated instance count."),
-		InColumnRef)
-{
-}
-
-////////////////////////////////////////////////////////////////////////////////////////////////////
-
-void FMemTagNodeSortingByInstanceCount::Sort(TArray<FBaseTreeNodePtr>& NodesToSort, ESortMode SortMode) const
-{
-	if (SortMode == ESortMode::Ascending)
-	{
-		NodesToSort.Sort([this](const FBaseTreeNodePtr& A, const FBaseTreeNodePtr& B) -> bool
-		{
-			INSIGHTS_DEFAULT_PRESORTING_NODES(A, B);
-
-			INSIGHTS_ENSURE(A.IsValid() && A->Is<FMemTagNode>());
-			const FMemTagNodePtr MemTagNodeA = StaticCastSharedPtr<FMemTagNode, FBaseTreeNode>(A);
-			const uint64 ValueA = MemTagNodeA->GetAggregatedStats().InstanceCount;
-
-			INSIGHTS_ENSURE(B.IsValid() && B->Is<FMemTagNode>());
-			const FMemTagNodePtr MemTagNodeB = StaticCastSharedPtr<FMemTagNode, FBaseTreeNode>(B);
-			const uint64 ValueB = MemTagNodeB->GetAggregatedStats().InstanceCount;
-
-			if (ValueA == ValueB)
-			{
-				INSIGHTS_DEFAULT_SORTING_NODES(A, B)
-			}
-			else
-			{
-				// Sort by instance count (ascending).
-				return ValueA < ValueB;
-			}
-		});
-	}
-	else // if (SortMode == ESortMode::Descending)
-	{
-		NodesToSort.Sort([this](const FBaseTreeNodePtr& A, const FBaseTreeNodePtr& B) -> bool
-		{
-			INSIGHTS_DEFAULT_PRESORTING_NODES(A, B);
-
-			INSIGHTS_ENSURE(A.IsValid() && A->Is<FMemTagNode>());
-			const FMemTagNodePtr MemTagNodeA = StaticCastSharedPtr<FMemTagNode, FBaseTreeNode>(A);
-			const uint64 ValueA = MemTagNodeA->GetAggregatedStats().InstanceCount;
-
-			INSIGHTS_ENSURE(B.IsValid() && B->Is<FMemTagNode>());
-			const FMemTagNodePtr MemTagNodeB = StaticCastSharedPtr<FMemTagNode, FBaseTreeNode>(B);
-			const uint64 ValueB = MemTagNodeB->GetAggregatedStats().InstanceCount;
-
-			if (ValueA == ValueB)
-			{
-				INSIGHTS_DEFAULT_SORTING_NODES(A, B)
-			}
-			else
-			{
-				// Sort by instance count (descending).
-				return ValueB < ValueA;
-			}
-		});
-	}
-}
-
 ////////////////////////////////////////////////////////////////////////////////////////////////////
 
 } // namespace UE::Insights::MemoryProfiler
diff --git a/Engine/Source/Developer/TraceInsights/Private/Insights/MemoryProfiler/ViewModels/MemTagNodeGroupingAndSorting.h b/Engine/Source/Developer/TraceInsights/Private/Insights/MemoryProfiler/ViewModels/MemTagNodeGroupingAndSorting.h
--- a/Engine/Source/Developer/TraceInsights/Private/Insights/MemoryProfiler/ViewModels/MemTagNodeGroupingAndSorting.h
+++ b/Engine/Source/Developer/TraceInsights/Private/Insights/MemoryProfiler/ViewModels/MemTagNodeGroupingAndSorting.h
@@ -17,18 +17,6 @@
 namespace UE::Insights::MemoryProfiler
 {
 
-////////////////////////////////////////////////////////////////////////////////////////////////////
-// Sorters
-////////////////////////////////////////////////////////////////////////////////////////////////////
-
-class FMemTagNodeSortingByType: public FTableCellValueSorter
-{
-public:
-	FMemTagNodeSortingByType(TSharedRef<FTableColumn> InColumnRef);
-
-	virtual void Sort(TArray<FBaseTreeNodePtr>& NodesToSort, ESortMode SortMode) const override;
-};
-
 ////////////////////////////////////////////////////////////////////////////////////////////////////
 
 class FMemTagNodeSortingByTracker : public FTableCellValueSorter
@@ -41,63 +29,4 @@ public:
 
 ////////////////////////////////////////////////////////////////////////////////////////////////////
 
-class FMemTagNodeSortingByInstanceCount : public FTableCellValueSorter
-{
-public:
-	FMemTagNodeSortingByInstanceCount(TSharedRef<FTableColumn> InColumnRef);
-
-	virtual void Sort(TArray<FBaseTreeNodePtr>& NodesToSort, ESortMode SortMode) const override;
-};
-
-////////////////////////////////////////////////////////////////////////////////////////////////////
-
-class FMemTagNodeSortingByTotalInclusiveSize : public FTableCellValueSorter
-{
-public:
-	FMemTagNodeSortingByTotalInclusiveSize(TSharedRef<FTableColumn> InColumnRef);
-
-	virtual void Sort(TArray<FBaseTreeNodePtr>& NodesToSort, ESortMode SortMode) const override;
-};
-
-////////////////////////////////////////////////////////////////////////////////////////////////////
-
-class FMemTagNodeSortingByTotalExclusiveSize : public FTableCellValueSorter
-{
-public:
-	FMemTagNodeSortingByTotalExclusiveSize(TSharedRef<FTableColumn> InColumnRef);
-
-	virtual void Sort(TArray<FBaseTreeNodePtr>& NodesToSort, ESortMode SortMode) const override;
-};
-
-////////////////////////////////////////////////////////////////////////////////////////////////////
-// Organizers
-////////////////////////////////////////////////////////////////////////////////////////////////////
-
-/** Enumerates types of grouping or sorting for the LLM tag nodes. */
-enum class EMemTagNodeGroupingMode
-{
-	/** Creates a single group for all LLM tags. */
-	Flat,
-
-	/** Creates one group for one letter. */
-	ByName,
-
-	/** Creates one group for each event type. */
-	ByType,
-
-	/** Creates one group for each tracker. */
-	ByTracker,
-
-	/** Group LLM tags by their hierarchy. */
-	ByParent,
-
-	/** Invalid enum type, may be used as a number of enumerations. */
-	InvalidOrMax,
-};
-
-/** Type definition for shared pointers to instances of EMemTagNodeGroupingMode. */
-typedef TSharedPtr<EMemTagNodeGroupingMode> EMemTagNodeGroupingModePtr;
-
-////////////////////////////////////////////////////////////////////////////////////////////////////
-
 } // namespace UE::Insights::MemoryProfiler
diff --git a/Engine/Source/Developer/TraceInsights/Private/Insights/MemoryProfiler/ViewModels/MemTagNodeHelper.cpp b/Engine/Source/Developer/TraceInsights/Private/Insights/MemoryProfiler/ViewModels/MemTagNodeHelper.cpp
deleted file mode 100644
diff --git a/Engine/Source/Developer/TraceInsights/Private/Insights/MemoryProfiler/ViewModels/MemTagNodeHelper.h b/Engine/Source/Developer/TraceInsights/Private/Insights/MemoryProfiler/ViewModels/MemTagNodeHelper.h
deleted file mode 100644
diff --git a/Engine/Source/Developer/TraceInsights/Private/Insights/MemoryProfiler/ViewModels/MemTagTable.cpp b/Engine/Source/Developer/TraceInsights/Private/Insights/MemoryProfiler/ViewModels/MemTagTable.cpp
new file mode 100644
--- /dev/null
+++ b/Engine/Source/Developer/TraceInsights/Private/Insights/MemoryProfiler/ViewModels/MemTagTable.cpp
@@ -0,0 +1,712 @@
+// Copyright Epic Games, Inc. All Rights Reserved.
+
+#include "MemTagTable.h"
+
+#include "Styling/StyleColors.h"
+#include "Widgets/SBoxPanel.h"
+#include "Widgets/SToolTip.h"
+#include "Widgets/Text/STextBlock.h"
+
+// TraceInsightsCore
+#include "InsightsCore/Table/ViewModels/TableCellValueFormatter.h"
+#include "InsightsCore/Table/ViewModels/TableCellValueGetter.h"
+#include "InsightsCore/Table/ViewModels/TableCellValueSorter.h"
+#include "InsightsCore/Table/ViewModels/TableColumn.h"
+
+// TraceInsights
+#include "Insights/MemoryProfiler/ViewModels/MemTagNode.h"
+#include "Insights/MemoryProfiler/ViewModels/MemTagNodeGroupingAndSorting.h"
+#include "Insights/MemoryProfiler/ViewModels/MemTagTable.h"
+
+#define LOCTEXT_NAMESPACE "UE::Insights::MemoryProfiler::FMemTagTable"
+
+namespace UE::Insights::MemoryProfiler
+{
+
+////////////////////////////////////////////////////////////////////////////////////////////////////
+// Column identifiers
+
+const FName FMemTagTableColumns::TagNameColumnId("TagName");
+const FName FMemTagTableColumns::TagIdColumnId("TagId");
+const FName FMemTagTableColumns::SizeAColumnId("SizeA");
+const FName FMemTagTableColumns::SizeBColumnId("SizeB");
+const FName FMemTagTableColumns::SizeDiffColumnId("SizeDiff");
+const FName FMemTagTableColumns::SampleCountColumnId("SampleCount");
+const FName FMemTagTableColumns::SizeMinColumnId("SizeMin");
+const FName FMemTagTableColumns::SizeMaxColumnId("SizeMax");
+const FName FMemTagTableColumns::SizeAverageColumnId("SizeAverage");
+const FName FMemTagTableColumns::SizeBudgetColumnId("SizeBudget");
+const FName FMemTagTableColumns::TrackerColumnId("Tracker");
+const FName FMemTagTableColumns::TagSetColumnId("TagSet");
+
+////////////////////////////////////////////////////////////////////////////////////////////////////
+// FMemTagTable
+////////////////////////////////////////////////////////////////////////////////////////////////////
+
+FMemTagTable::FMemTagTable()
+{
+}
+
+////////////////////////////////////////////////////////////////////////////////////////////////////
+
+FMemTagTable::~FMemTagTable()
+{
+}
+
+////////////////////////////////////////////////////////////////////////////////////////////////////
+
+void FMemTagTable::Reset()
+{
+	FTable::Reset();
+
+	AddDefaultColumns();
+}
+
+////////////////////////////////////////////////////////////////////////////////////////////////////
+
+void FMemTagTable::AddDefaultColumns()
+{
+	class FAggregatedValueGetter : public FTableCellValueGetter
+	{
+	public:
+		virtual const TOptional<FTableCellValue> GetNodeValue(const FBaseTreeNode& Node) const
+		{
+			return TOptional<FTableCellValue>();
+		}
+
+		virtual const TOptional<FTableCellValue> GetValue(const FTableColumn& Column, const FBaseTreeNode& Node) const override
+		{
+			if (Node.IsGroup() && Node.Is<FTableTreeNode>())
+			{
+				const FTableTreeNode& TableTreeNode = Node.As<FTableTreeNode>();
+				if (TableTreeNode.HasAggregatedValue(Column.GetId()))
+				{
+					return TableTreeNode.GetAggregatedValue(Column.GetId());
+				}
+			}
+			return GetNodeValue(Node);
+		}
+	};
+
+	//////////////////////////////////////////////////
+	// Hierarchy Column
+	{
+		const int32 HierarchyColumnIndex = -1;
+		const TCHAR* HierarchyColumnName = nullptr;
+		AddHierarchyColumn(HierarchyColumnIndex, HierarchyColumnName);
+
+		FTableColumn& HierarchyColumn = *GetColumns()[0];
+		HierarchyColumn.SetInitialWidth(60.0f);
+		HierarchyColumn.SetShortName(LOCTEXT("HierarchyColumnName", "Hierarchy"));
+		HierarchyColumn.SetTitleName(LOCTEXT("HierarchyColumnTitle", "LLM Tag Hierarchy"));
+		HierarchyColumn.SetDescription(LOCTEXT("HierarchyColumnDesc", "Hierarchy of the LLM tag's tree"));
+	}
+
+	int32 ColumnIndex = 0;
+
+	//////////////////////////////////////////////////
+	// Tag Id Column
+	{
+		TSharedRef<FTableColumn> ColumnRef = MakeShared<FTableColumn>(FMemTagTableColumns::TagNameColumnId);
+		FTableColumn& Column = *ColumnRef;
+
+		Column.SetIndex(ColumnIndex++);
+
+		Column.SetShortName(LOCTEXT("TagNameColumnName", "Name"));
+		Column.SetTitleName(LOCTEXT("TagNameColumnTitle", "Tag Name"));
+		Column.SetDescription(LOCTEXT("TagNameColumnDesc", "The name of the LLM tag"));
+
+		Column.SetFlags(ETableColumnFlags::ShouldBeVisible |
+						ETableColumnFlags::CanBeHidden |
+						ETableColumnFlags::CanBeFiltered);
+
+		Column.SetHorizontalAlignment(HAlign_Left);
+		Column.SetInitialWidth(120.0f);
+
+		Column.SetDataType(ETableCellDataType::CString);
+
+		class FMemTagNameValueGetter : public FAggregatedValueGetter
+		{
+		public:
+			virtual const TOptional<FTableCellValue> GetNodeValue(const FBaseTreeNode& Node) const override
+			{
+				if (Node.Is<FMemTagNode>())
+				{
+					const FMemTagNode& MemTagNode = Node.As<FMemTagNode>();
+					return FTableCellValue(MemTagNode.GetTagName());
+				}
+				return TOptional<FTableCellValue>();
+			}
+		};
+		TSharedRef<ITableCellValueGetter> Getter = MakeShared<FMemTagNameValueGetter>();
+		Column.SetValueGetter(Getter);
+
+		TSharedRef<ITableCellValueFormatter> Formatter = MakeShared<FCStringValueFormatterAsText>();
+		Column.SetValueFormatter(Formatter);
+
+		TSharedRef<ITableCellValueSorter> Sorter = MakeShared<FSorterByCStringValue>(ColumnRef);
+		Column.SetValueSorter(Sorter);
+
+		Column.SetAggregation(ETableColumnAggregation::SameValue);
+
+		AddColumn(ColumnRef);
+	}
+	//////////////////////////////////////////////////
+	// Tag Id Column
+	{
+		TSharedRef<FTableColumn> ColumnRef = MakeShared<FTableColumn>(FMemTagTableColumns::TagIdColumnId);
+		FTableColumn& Column = *ColumnRef;
+
+		Column.SetIndex(ColumnIndex++);
+
+		Column.SetShortName(LOCTEXT("TagIdColumnName", "Tag Id"));
+		Column.SetTitleName(LOCTEXT("TagIdColumnTitle", "Tag Id"));
+		Column.SetDescription(LOCTEXT("TagIdColumnDesc", "The id of the LLM tag"));
+
+		Column.SetFlags(ETableColumnFlags::CanBeHidden |
+						ETableColumnFlags::CanBeFiltered);
+
+		Column.SetHorizontalAlignment(HAlign_Left);
+		Column.SetInitialWidth(120.0f);
+
+		Column.SetDataType(ETableCellDataType::Int64);
+
+		class FMemTagIdValueGetter : public FAggregatedValueGetter
+		{
+		public:
+			virtual const TOptional<FTableCellValue> GetNodeValue(const FBaseTreeNode& Node) const override
+			{
+				if (Node.Is<FMemTagNode>())
+				{
+					const FMemTagNode& MemTagNode = Node.As<FMemTagNode>();
+					return FTableCellValue(int64(MemTagNode.GetMemTagId()));
+				}
+				return TOptional<FTableCellValue>();
+			}
+		};
+		TSharedRef<ITableCellValueGetter> Getter = MakeShared<FMemTagIdValueGetter>();
+		Column.SetValueGetter(Getter);
+
+		TSharedRef<ITableCellValueFormatter> Formatter = MakeShared<FInt64ValueFormatterAsHex64>();
+		Column.SetValueFormatter(Formatter);
+
+		TSharedRef<ITableCellValueSorter> Sorter = MakeShared<FSorterByInt64Value>(ColumnRef);
+		Column.SetValueSorter(Sorter);
+		Column.SetInitialSortMode(EColumnSortMode::Descending);
+
+		//TSharedRef<IFilterValueConverter> Converter = MakeShared<FMemoryFilterValueConverter>();
+		//Column.SetValueConverter(Converter);
+
+		//Column.SetAggregation(ETableColumnAggregation::None);
+
+		AddColumn(ColumnRef);
+	}
+	//////////////////////////////////////////////////
+	// Size (Time A) Column
+	{
+		TSharedRef<FTableColumn> ColumnRef = MakeShared<FTableColumn>(FMemTagTableColumns::SizeAColumnId);
+		FTableColumn& Column = *ColumnRef;
+
+		Column.SetIndex(ColumnIndex++);
+
+		Column.SetShortName(LOCTEXT("SizeAColumnName", "Size A"));
+		Column.SetTitleName(LOCTEXT("SizeAColumnTitle", "Size at TimeMarker A"));
+		Column.SetDescription(LOCTEXT("SizeAColumnDesc", "The memory size (bytes) of the LLM tag at the time marker A"));
+
+		Column.SetFlags(ETableColumnFlags::ShouldBeVisible |
+						ETableColumnFlags::CanBeHidden |
+						ETableColumnFlags::CanBeFiltered |
+						ETableColumnFlags::IsDynamic);
+
+		Column.SetHorizontalAlignment(HAlign_Right);
+		Column.SetInitialWidth(100.0f);
+
+		Column.SetDataType(ETableCellDataType::Int64);
+
+		class FMemTagSizeAValueGetter : public FAggregatedValueGetter
+		{
+		public:
+			virtual const TOptional<FTableCellValue> GetNodeValue(const FBaseTreeNode& Node) const override
+			{
+				if (Node.Is<FMemTagNode>())
+				{
+					const FMemTagNode& MemTagNode = Node.As<FMemTagNode>();
+					return FTableCellValue(MemTagNode.GetSizeA());
+				}
+				return TOptional<FTableCellValue>();
+			}
+		};
+		TSharedRef<ITableCellValueGetter> Getter = MakeShared<FMemTagSizeAValueGetter>();
+		Column.SetValueGetter(Getter);
+
+		TSharedRef<ITableCellValueFormatter> Formatter = MakeShared<FInt64ValueFormatterAsMemory>();
+		Column.SetValueFormatter(Formatter);
+
+		TSharedRef<ITableCellValueSorter> Sorter = MakeShared<FSorterByInt64Value>(ColumnRef);
+		Column.SetValueSorter(Sorter);
+		Column.SetInitialSortMode(EColumnSortMode::Descending);
+
+		//TSharedRef<IFilterValueConverter> Converter = MakeShared<FMemoryFilterValueConverter>();
+		//Column.SetValueConverter(Converter);
+
+		Column.SetAggregation(ETableColumnAggregation::Sum);
+
+		AddColumn(ColumnRef);
+	}
+	//////////////////////////////////////////////////
+	// Size (Time B) Column
+	{
+		TSharedRef<FTableColumn> ColumnRef = MakeShared<FTableColumn>(FMemTagTableColumns::SizeBColumnId);
+		FTableColumn& Column = *ColumnRef;
+
+		Column.SetIndex(ColumnIndex++);
+
+		Column.SetShortName(LOCTEXT("SizeBColumnName", "Size B"));
+		Column.SetTitleName(LOCTEXT("SizeBColumnTitle", "Size at TimeMarker B"));
+		Column.SetDescription(LOCTEXT("SizeBColumnDesc", "The memory size (bytes) of the LLM tag at the time marker B"));
+
+		Column.SetFlags(ETableColumnFlags::CanBeHidden |
+						ETableColumnFlags::CanBeFiltered |
+						ETableColumnFlags::IsDynamic);
+
+		Column.SetHorizontalAlignment(HAlign_Right);
+		Column.SetInitialWidth(100.0f);
+
+		Column.SetDataType(ETableCellDataType::Int64);
+
+		class FMemTagSizeAValueGetter : public FAggregatedValueGetter
+		{
+		public:
+			virtual const TOptional<FTableCellValue> GetNodeValue(const FBaseTreeNode& Node) const override
+			{
+				if (Node.Is<FMemTagNode>())
+				{
+					const FMemTagNode& MemTagNode = Node.As<FMemTagNode>();
+					return FTableCellValue(MemTagNode.GetSizeB());
+				}
+				return TOptional<FTableCellValue>();
+			}
+		};
+		TSharedRef<ITableCellValueGetter> Getter = MakeShared<FMemTagSizeAValueGetter>();
+		Column.SetValueGetter(Getter);
+
+		TSharedRef<ITableCellValueFormatter> Formatter = MakeShared<FInt64ValueFormatterAsMemory>();
+		Column.SetValueFormatter(Formatter);
+
+		TSharedRef<ITableCellValueSorter> Sorter = MakeShared<FSorterByInt64Value>(ColumnRef);
+		Column.SetValueSorter(Sorter);
+		Column.SetInitialSortMode(EColumnSortMode::Descending);
+
+		//TSharedRef<IFilterValueConverter> Converter = MakeShared<FMemoryFilterValueConverter>();
+		//Column.SetValueConverter(Converter);
+
+		Column.SetAggregation(ETableColumnAggregation::Sum);
+
+		AddColumn(ColumnRef);
+	}
+	//////////////////////////////////////////////////
+	// Size Difference (B - A) Column
+	{
+		TSharedRef<FTableColumn> ColumnRef = MakeShared<FTableColumn>(FMemTagTableColumns::SizeDiffColumnId);
+		FTableColumn& Column = *ColumnRef;
+
+		Column.SetIndex(ColumnIndex++);
+
+		Column.SetShortName(LOCTEXT("SizeDiffColumnName", "Diff (B - A)"));
+		Column.SetTitleName(LOCTEXT("SizeDiffColumnTitle", "Size Difference (B - A)"));
+		Column.SetDescription(LOCTEXT("SizeDiffColumnDesc", "The memory size variation (in bytes) of the LLM tag between time marker B and time marker A"));
+
+		Column.SetFlags(ETableColumnFlags::CanBeHidden |
+						ETableColumnFlags::CanBeFiltered |
+						ETableColumnFlags::IsDynamic);
+
+		Column.SetHorizontalAlignment(HAlign_Right);
+		Column.SetInitialWidth(100.0f);
+
+		Column.SetDataType(ETableCellDataType::Int64);
+
+		class FMemTagSizeBValueGetter : public FAggregatedValueGetter
+		{
+		public:
+			virtual const TOptional<FTableCellValue> GetNodeValue(const FBaseTreeNode& Node) const override
+			{
+				if (Node.Is<FMemTagNode>())
+				{
+					const FMemTagNode& MemTagNode = Node.As<FMemTagNode>();
+					return FTableCellValue(MemTagNode.GetSizeDiff());
+				}
+				return TOptional<FTableCellValue>();
+			}
+		};
+		TSharedRef<ITableCellValueGetter> Getter = MakeShared<FMemTagSizeBValueGetter>();
+		Column.SetValueGetter(Getter);
+
+		TSharedRef<ITableCellValueFormatter> Formatter = MakeShared<FInt64ValueFormatterAsMemory>();
+		Column.SetValueFormatter(Formatter);
+
+		TSharedRef<ITableCellValueSorter> Sorter = MakeShared<FSorterByInt64Value>(ColumnRef);
+		Column.SetValueSorter(Sorter);
+		Column.SetInitialSortMode(EColumnSortMode::Descending);
+
+		//TSharedRef<IFilterValueConverter> Converter = MakeShared<FMemoryFilterValueConverter>();
+		//Column.SetValueConverter(Converter);
+
+		Column.SetAggregation(ETableColumnAggregation::Sum);
+
+		AddColumn(ColumnRef);
+	}
+	//////////////////////////////////////////////////
+	// Instance Count Column
+	{
+		TSharedRef<FTableColumn> ColumnRef = MakeShared<FTableColumn>(FMemTagTableColumns::SampleCountColumnId);
+		FTableColumn& Column = *ColumnRef;
+
+		Column.SetIndex(ColumnIndex++);
+
+		Column.SetShortName(LOCTEXT("SampleCountColumnName", "Samples"));
+		Column.SetTitleName(LOCTEXT("SampleCountColumnTitle", "Sample Count"));
+		Column.SetDescription(LOCTEXT("SampleCountColumnDesc", "The number of snapshots the LLM tag has in the selected time range"));
+
+		Column.SetFlags(ETableColumnFlags::CanBeHidden |
+						ETableColumnFlags::CanBeFiltered |
+						ETableColumnFlags::IsDynamic);
+
+		Column.SetHorizontalAlignment(HAlign_Right);
+		Column.SetInitialWidth(100.0f);
+
+		Column.SetDataType(ETableCellDataType::Int64);
+
+		class FMemTagSampleCountValueGetter : public FAggregatedValueGetter
+		{
+		public:
+			virtual const TOptional<FTableCellValue> GetNodeValue(const FBaseTreeNode& Node) const override
+			{
+				if (Node.Is<FMemTagNode>())
+				{
+					const FMemTagNode& MemTagNode = Node.As<FMemTagNode>();
+					return FTableCellValue(MemTagNode.GetSampleCount());
+				}
+				return TOptional<FTableCellValue>();
+			}
+		};
+		TSharedRef<ITableCellValueGetter> Getter = MakeShared<FMemTagSampleCountValueGetter>();
+		Column.SetValueGetter(Getter);
+
+		TSharedRef<ITableCellValueFormatter> Formatter = MakeShared<FInt64ValueFormatterAsNumber>();
+		Column.SetValueFormatter(Formatter);
+
+		TSharedRef<ITableCellValueSorter> Sorter = MakeShared<FSorterByInt64Value>(ColumnRef);
+		Column.SetValueSorter(Sorter);
+		Column.SetInitialSortMode(EColumnSortMode::Descending);
+
+		//TSharedRef<IFilterValueConverter> Converter = MakeShared<FMemoryFilterValueConverter>();
+		//Column.SetValueConverter(Converter);
+
+		//Column.SetAggregation(ETableColumnAggregation::Sum);
+
+		AddColumn(ColumnRef);
+	}
+	//////////////////////////////////////////////////
+	// Size Min Column
+	{
+		TSharedRef<FTableColumn> ColumnRef = MakeShared<FTableColumn>(FMemTagTableColumns::SizeMinColumnId);
+		FTableColumn& Column = *ColumnRef;
+
+		Column.SetIndex(ColumnIndex++);
+
+		Column.SetShortName(LOCTEXT("SizeMinColumnName", "Min"));
+		Column.SetTitleName(LOCTEXT("SizeMinColumnTitle", "Min Size"));
+		Column.SetDescription(LOCTEXT("SizeMinColumnDesc", "The minimum size value (in bytes) the LLM tag has in the selected time range"));
+
+		Column.SetFlags(ETableColumnFlags::CanBeHidden |
+						ETableColumnFlags::CanBeFiltered |
+						ETableColumnFlags::IsDynamic);
+
+		Column.SetHorizontalAlignment(HAlign_Right);
+		Column.SetInitialWidth(100.0f);
+
+		Column.SetDataType(ETableCellDataType::Int64);
+
+		class FMemTagSizeMinValueGetter : public FAggregatedValueGetter
+		{
+		public:
+			virtual const TOptional<FTableCellValue> GetNodeValue(const FBaseTreeNode& Node) const override
+			{
+				if (Node.Is<FMemTagNode>())
+				{
+					const FMemTagNode& MemTagNode = Node.As<FMemTagNode>();
+					return FTableCellValue(MemTagNode.GetSizeMin());
+				}
+				return TOptional<FTableCellValue>();
+			}
+		};
+		TSharedRef<ITableCellValueGetter> Getter = MakeShared<FMemTagSizeMinValueGetter>();
+		Column.SetValueGetter(Getter);
+
+		TSharedRef<ITableCellValueFormatter> Formatter = MakeShared<FInt64ValueFormatterAsMemory>();
+		Column.SetValueFormatter(Formatter);
+
+		TSharedRef<ITableCellValueSorter> Sorter = MakeShared<FSorterByInt64Value>(ColumnRef);
+		Column.SetValueSorter(Sorter);
+		Column.SetInitialSortMode(EColumnSortMode::Descending);
+
+		//TSharedRef<IFilterValueConverter> Converter = MakeShared<FMemoryFilterValueConverter>();
+		//Column.SetValueConverter(Converter);
+
+		Column.SetAggregation(ETableColumnAggregation::Min);
+
+		AddColumn(ColumnRef);
+	}
+	//////////////////////////////////////////////////
+	// Size Max Column
+	{
+		TSharedRef<FTableColumn> ColumnRef = MakeShared<FTableColumn>(FMemTagTableColumns::SizeMaxColumnId);
+		FTableColumn& Column = *ColumnRef;
+
+		Column.SetIndex(ColumnIndex++);
+
+		Column.SetShortName(LOCTEXT("SizeMaxColumnName", "Max"));
+		Column.SetTitleName(LOCTEXT("SizeMaxColumnTitle", "Max Size"));
+		Column.SetDescription(LOCTEXT("SizeMaxColumnDesc", "The maximum size value (in bytes) the LLM tag has in the selected time range"));
+
+		Column.SetFlags(ETableColumnFlags::CanBeHidden |
+						ETableColumnFlags::CanBeFiltered |
+						ETableColumnFlags::IsDynamic);
+
+		Column.SetHorizontalAlignment(HAlign_Right);
+		Column.SetInitialWidth(100.0f);
+
+		Column.SetDataType(ETableCellDataType::Int64);
+
+		class FMemTagSizeMaxValueGetter : public FAggregatedValueGetter
+		{
+		public:
+			virtual const TOptional<FTableCellValue> GetNodeValue(const FBaseTreeNode& Node) const override
+			{
+				if (Node.Is<FMemTagNode>())
+				{
+					const FMemTagNode& MemTagNode = Node.As<FMemTagNode>();
+					return FTableCellValue(MemTagNode.GetSizeMax());
+				}
+				return TOptional<FTableCellValue>();
+			}
+		};
+		TSharedRef<ITableCellValueGetter> Getter = MakeShared<FMemTagSizeMaxValueGetter>();
+		Column.SetValueGetter(Getter);
+
+		TSharedRef<ITableCellValueFormatter> Formatter = MakeShared<FInt64ValueFormatterAsMemory>();
+		Column.SetValueFormatter(Formatter);
+
+		TSharedRef<ITableCellValueSorter> Sorter = MakeShared<FSorterByInt64Value>(ColumnRef);
+		Column.SetValueSorter(Sorter);
+		Column.SetInitialSortMode(EColumnSortMode::Descending);
+
+		//TSharedRef<IFilterValueConverter> Converter = MakeShared<FMemoryFilterValueConverter>();
+		//Column.SetValueConverter(Converter);
+
+		Column.SetAggregation(ETableColumnAggregation::Max);
+
+		AddColumn(ColumnRef);
+	}
+	//////////////////////////////////////////////////
+	// Size Average Column
+	{
+		TSharedRef<FTableColumn> ColumnRef = MakeShared<FTableColumn>(FMemTagTableColumns::SizeAverageColumnId);
+		FTableColumn& Column = *ColumnRef;
+
+		Column.SetIndex(ColumnIndex++);
+
+		Column.SetShortName(LOCTEXT("SizeAverageColumnName", "Avg."));
+		Column.SetTitleName(LOCTEXT("SizeAverageColumnTitle", "Average Size"));
+		Column.SetDescription(LOCTEXT("SizeAverageColumnDesc", "The average size value (in bytes) the LLM tag has in the selected time range"));
+
+		Column.SetFlags(ETableColumnFlags::CanBeHidden |
+						ETableColumnFlags::CanBeFiltered |
+						ETableColumnFlags::IsDynamic);
+
+		Column.SetHorizontalAlignment(HAlign_Right);
+		Column.SetInitialWidth(100.0f);
+
+		Column.SetDataType(ETableCellDataType::Int64);
+
+		class FMemTagSizeAverageValueGetter : public FAggregatedValueGetter
+		{
+		public:
+			virtual const TOptional<FTableCellValue> GetNodeValue(const FBaseTreeNode& Node) const override
+			{
+				if (Node.Is<FMemTagNode>())
+				{
+					const FMemTagNode& MemTagNode = Node.As<FMemTagNode>();
+					return FTableCellValue(MemTagNode.GetSizeAverage());
+				}
+				return TOptional<FTableCellValue>();
+			}
+		};
+		TSharedRef<ITableCellValueGetter> Getter = MakeShared<FMemTagSizeAverageValueGetter>();
+		Column.SetValueGetter(Getter);
+
+		TSharedRef<ITableCellValueFormatter> Formatter = MakeShared<FInt64ValueFormatterAsMemory>();
+		Column.SetValueFormatter(Formatter);
+
+		TSharedRef<ITableCellValueSorter> Sorter = MakeShared<FSorterByInt64Value>(ColumnRef);
+		Column.SetValueSorter(Sorter);
+		Column.SetInitialSortMode(EColumnSortMode::Descending);
+
+		//TSharedRef<IFilterValueConverter> Converter = MakeShared<FMemoryFilterValueConverter>();
+		//Column.SetValueConverter(Converter);
+
+		//Column.SetAggregation(ETableColumnAggregation::Average);
+
+		AddColumn(ColumnRef);
+	}
+	//////////////////////////////////////////////////
+	// Size Budget Column
+	{
+		TSharedRef<FTableColumn> ColumnRef = MakeShared<FTableColumn>(FMemTagTableColumns::SizeBudgetColumnId);
+		FTableColumn& Column = *ColumnRef;
+
+		Column.SetIndex(ColumnIndex++);
+
+		Column.SetShortName(LOCTEXT("BudgetColumnName", "Budget"));
+		Column.SetTitleName(LOCTEXT("BudgetColumnTitle", "Budget"));
+		Column.SetDescription(LOCTEXT("BudgetColumnDesc", "The budget (max size value, in bytes) of the LLM tag"));
+
+		Column.SetFlags(ETableColumnFlags::ShouldBeVisible |
+						ETableColumnFlags::CanBeHidden |
+						ETableColumnFlags::CanBeFiltered);
+
+		Column.SetHorizontalAlignment(HAlign_Right);
+		Column.SetInitialWidth(100.0f);
+
+		Column.SetDataType(ETableCellDataType::Int64);
+
+		class FMemTagBudgetValueGetter : public FAggregatedValueGetter
+		{
+		public:
+			virtual const TOptional<FTableCellValue> GetNodeValue(const FBaseTreeNode& Node) const override
+			{
+				if (Node.Is<FMemTagNode>())
+				{
+					const FMemTagNode& MemTagNode = Node.As<FMemTagNode>();
+					return FTableCellValue(MemTagNode.GetSizeBudget());
+				}
+				return TOptional<FTableCellValue>();
+			}
+		};
+		TSharedRef<ITableCellValueGetter> Getter = MakeShared<FMemTagBudgetValueGetter>();
+		Column.SetValueGetter(Getter);
+
+		TSharedRef<ITableCellValueFormatter> Formatter = MakeShared<FInt64ValueFormatterAsMemory>();
+		Column.SetValueFormatter(Formatter);
+
+		TSharedRef<ITableCellValueSorter> Sorter = MakeShared<FSorterByInt64Value>(ColumnRef);
+		Column.SetValueSorter(Sorter);
+		Column.SetInitialSortMode(EColumnSortMode::Descending);
+
+		//TSharedRef<IFilterValueConverter> Converter = MakeShared<FMemoryFilterValueConverter>();
+		//Column.SetValueConverter(Converter);
+
+		Column.SetAggregation(ETableColumnAggregation::Sum);
+
+		AddColumn(ColumnRef);
+	}
+	//////////////////////////////////////////////////
+	// LLM Tracker Column
+	{
+		TSharedRef<FTableColumn> ColumnRef = MakeShared<FTableColumn>(FMemTagTableColumns::TrackerColumnId);
+		FTableColumn& Column = *ColumnRef;
+
+		Column.SetIndex(ColumnIndex++);
+
+		Column.SetShortName(LOCTEXT("TrackerColumnName", "Tracker"));
+		Column.SetTitleName(LOCTEXT("TrackerColumnTitle", "Tracker"));
+		Column.SetDescription(LOCTEXT("TrackerColumnDesc", "The LLM tracker using the LLM tag"));
+
+		Column.SetFlags(ETableColumnFlags::CanBeHidden |
+						ETableColumnFlags::CanBeFiltered);
+
+		Column.SetHorizontalAlignment(HAlign_Left);
+		Column.SetInitialWidth(100.0f);
+
+		Column.SetDataType(ETableCellDataType::Text);
+
+		class FMemTagTrackerValueGetter : public FAggregatedValueGetter
+		{
+		public:
+			virtual const TOptional<FTableCellValue> GetNodeValue(const FBaseTreeNode& Node) const
+			{
+				if (Node.Is<FMemTagNode>())
+				{
+					const FMemTagNode& MemTagNode = Node.As<FMemTagNode>();
+					return FTableCellValue(MemTagNode.GetTrackerText());
+				}
+				return TOptional<FTableCellValue>();
+			}
+		};
+		TSharedRef<ITableCellValueGetter> Getter = MakeShared<FMemTagTrackerValueGetter>();
+		Column.SetValueGetter(Getter);
+
+		TSharedRef<ITableCellValueFormatter> Formatter = MakeShared<FTextValueFormatter>();
+		Column.SetValueFormatter(Formatter);
+
+		TSharedRef<ITableCellValueSorter> Sorter = MakeShared<FMemTagNodeSortingByTracker>(ColumnRef);
+		Column.SetValueSorter(Sorter);
+
+		Column.SetAggregation(ETableColumnAggregation::SameValue);
+
+		AddColumn(ColumnRef);
+	}
+	//////////////////////////////////////////////////
+	// LLM Tag Set Column
+	{
+		TSharedRef<FTableColumn> ColumnRef = MakeShared<FTableColumn>(FMemTagTableColumns::TagSetColumnId);
+		FTableColumn& Column = *ColumnRef;
+
+		Column.SetIndex(ColumnIndex++);
+
+		Column.SetShortName(LOCTEXT("TagSetColumnName", "Tag Set"));
+		Column.SetTitleName(LOCTEXT("TagSetColumnTitle", "Tag Set"));
+		Column.SetDescription(LOCTEXT("TagSetColumnDesc", "The LLM tag set of the LLM tag"));
+
+		Column.SetFlags(ETableColumnFlags::CanBeHidden |
+						ETableColumnFlags::CanBeFiltered);
+
+		Column.SetHorizontalAlignment(HAlign_Left);
+		Column.SetInitialWidth(100.0f);
+
+		Column.SetDataType(ETableCellDataType::CString);
+
+		class FMemTagSetNameValueGetter : public FAggregatedValueGetter
+		{
+		public:
+			virtual const TOptional<FTableCellValue> GetNodeValue(const FBaseTreeNode& Node) const
+			{
+				if (Node.Is<FMemTagNode>())
+				{
+					const FMemTagNode& MemTagNode = Node.As<FMemTagNode>();
+					return FTableCellValue(MemTagNode.GetTagSetName());
+				}
+				return TOptional<FTableCellValue>();
+			}
+		};
+		TSharedRef<ITableCellValueGetter> Getter = MakeShared<FMemTagSetNameValueGetter>();
+		Column.SetValueGetter(Getter);
+
+		TSharedRef<ITableCellValueFormatter> Formatter = MakeShared<FCStringValueFormatterAsText>();
+		Column.SetValueFormatter(Formatter);
+
+		TSharedRef<ITableCellValueSorter> Sorter = MakeShared<FSorterByCStringValue>(ColumnRef); // TODO: FMemTagNodeSortingByTagSet
+		Column.SetValueSorter(Sorter);
+
+		Column.SetAggregation(ETableColumnAggregation::SameValue);
+
+		AddColumn(ColumnRef);
+	}
+	//////////////////////////////////////////////////
+}
+
+////////////////////////////////////////////////////////////////////////////////////////////////////
+
+} // namespace UE::Insights::MemoryProfiler
+
+#undef LOCTEXT_NAMESPACE
diff --git a/Engine/Source/Developer/TraceInsights/Private/Insights/MemoryProfiler/ViewModels/MemTagTable.h b/Engine/Source/Developer/TraceInsights/Private/Insights/MemoryProfiler/ViewModels/MemTagTable.h
new file mode 100644
--- /dev/null
+++ b/Engine/Source/Developer/TraceInsights/Private/Insights/MemoryProfiler/ViewModels/MemTagTable.h
@@ -0,0 +1,59 @@
+// Copyright Epic Games, Inc. All Rights Reserved.
+
+#pragma once
+
+#include "CoreTypes.h"
+#include "UObject/NameTypes.h"
+
+// TraceInsightsCore
+#include "InsightsCore/Table/ViewModels/Table.h"
+
+namespace UE::Insights::MemoryProfiler
+{
+
+////////////////////////////////////////////////////////////////////////////////////////////////////
+
+// Column identifiers
+struct FMemTagTableColumns
+{
+	static const FName TagNameColumnId;
+	static const FName TagIdColumnId;
+	static const FName SizeAColumnId;
+	static const FName SizeBColumnId;
+	static const FName SizeDiffColumnId;
+	static const FName SizeBudgetColumnId;
+	static const FName SampleCountColumnId;
+	static const FName SizeMinColumnId;
+	static const FName SizeMaxColumnId;
+	static const FName SizeAverageColumnId;
+	static const FName TrackerColumnId;
+	static const FName TagSetColumnId;
+};
+
+////////////////////////////////////////////////////////////////////////////////////////////////////
+
+class FMemTagTable : public FTable
+{
+public:
+	FMemTagTable();
+	virtual ~FMemTagTable();
+
+	virtual void Reset();
+
+	//TArray<FMemoryTag>& GetTags() { return Tags; }
+	//const TArray<FMemoryTag>& GetTags() const { return Tags; }
+
+	//bool IsValidRowIndex(int32 InIndex) const { return InIndex >= 0 && InIndex < Tags.Num(); }
+	//const FMemoryTag* GetMemTag(int32 InIndex) const { return IsValidRowIndex(InIndex) ? &Tags[InIndex] : nullptr; }
+	//const FMemoryTag& GetMemTagChecked(int32 InIndex) const { check(IsValidRowIndex(InIndex)); return Tags[InIndex]; }
+
+private:
+	void AddDefaultColumns();
+
+private:
+	//TArray<FMemoryTag> Tags;
+};
+
+////////////////////////////////////////////////////////////////////////////////////////////////////
+
+} // namespace UE::Insights::MemoryProfiler
diff --git a/Engine/Source/Developer/TraceInsights/Private/Insights/MemoryProfiler/ViewModels/MemTagTableViewPresets.cpp b/Engine/Source/Developer/TraceInsights/Private/Insights/MemoryProfiler/ViewModels/MemTagTableViewPresets.cpp
new file mode 100644
--- /dev/null
+++ b/Engine/Source/Developer/TraceInsights/Private/Insights/MemoryProfiler/ViewModels/MemTagTableViewPresets.cpp
@@ -0,0 +1,152 @@
+// Copyright Epic Games, Inc. All Rights Reserved.
+
+#include "MemTagTableViewPresets.h"
+
+// TraceInsightsCore
+#include "InsightsCore/Table/ViewModels/TreeNodeGrouping.h"
+#include "InsightsCore/Table/Widgets/STableTreeView.h"
+
+// TraceInsights
+#include "Insights/MemoryProfiler/ViewModels/MemTagTable.h"
+#include "Insights/MemoryProfiler/Widgets/SMemTagTreeView.h"
+
+#define LOCTEXT_NAMESPACE "UE::Insights::MemoryProfiler::SMemTagTableTreeView"
+
+namespace UE::Insights::MemoryProfiler
+{
+
+////////////////////////////////////////////////////////////////////////////////////////////////////
+// Default View
+
+TSharedRef<ITableTreeViewPreset> FMemTagTableViewPresets::CreateDefaultViewPreset(SMemTagTreeView& TableTreeView)
+{
+	class FDefaultViewPreset : public ITableTreeViewPreset
+	{
+	public:
+		virtual FText GetName() const override
+		{
+			return LOCTEXT("Default_PresetName", "Default");
+		}
+		virtual FText GetToolTip() const override
+		{
+			return LOCTEXT("Default_PresetToolTip", "Default View\nConfigure the tree view to show the default LLM tags info.");
+		}
+		virtual FName GetSortColumn() const override
+		{
+			return FMemTagTableColumns::TagNameColumnId;
+		}
+		virtual EColumnSortMode::Type GetSortMode() const override
+		{
+			return EColumnSortMode::Type::Ascending;
+		}
+		virtual void SetCurrentGroupings(const TArray<TSharedPtr<FTreeNodeGrouping>>& InAvailableGroupings, TArray<TSharedPtr<FTreeNodeGrouping>>& InOutCurrentGroupings) const override
+		{
+			InOutCurrentGroupings.Reset();
+
+			check(InAvailableGroupings[0]->Is<FTreeNodeGroupingFlat>());
+			InOutCurrentGroupings.Add(InAvailableGroupings[0]);
+		}
+		virtual void GetColumnConfigSet(TArray<FTableColumnConfig>& InOutConfigSet) const override
+		{
+			InOutConfigSet.Add({ FTable::GetHierarchyColumnId(),          true, 120.0f });
+			InOutConfigSet.Add({ FMemTagTableColumns::TagNameColumnId,    true, 400.0f });
+			InOutConfigSet.Add({ FMemTagTableColumns::SizeAColumnId,      true, 100.0f });
+			InOutConfigSet.Add({ FMemTagTableColumns::SizeBudgetColumnId, true, 100.0f });
+		}
+	};
+	return MakeShared<FDefaultViewPreset>();
+}
+
+////////////////////////////////////////////////////////////////////////////////////////////////////
+// Difference View
+
+TSharedRef<ITableTreeViewPreset> FMemTagTableViewPresets::CreateDiffViewPreset(SMemTagTreeView& TableTreeView)
+{
+	class FDiffViewPreset : public ITableTreeViewPreset
+	{
+	public:
+		virtual FText GetName() const override
+		{
+			return LOCTEXT("Diff_PresetName", "Diff B-A");
+		}
+		virtual FText GetToolTip() const override
+		{
+			return LOCTEXT("Diff_PresetToolTip", "Difference View\nConfigure the tree view to investigate variation of LLM tags between two snapshots.");
+		}
+		virtual FName GetSortColumn() const override
+		{
+			return FMemTagTableColumns::TagNameColumnId;
+		}
+		virtual EColumnSortMode::Type GetSortMode() const override
+		{
+			return EColumnSortMode::Type::Ascending;
+		}
+		virtual void SetCurrentGroupings(const TArray<TSharedPtr<FTreeNodeGrouping>>& InAvailableGroupings, TArray<TSharedPtr<FTreeNodeGrouping>>& InOutCurrentGroupings) const override
+		{
+			InOutCurrentGroupings.Reset();
+
+			check(InAvailableGroupings[0]->Is<FTreeNodeGroupingFlat>());
+			InOutCurrentGroupings.Add(InAvailableGroupings[0]);
+		}
+		virtual void GetColumnConfigSet(TArray<FTableColumnConfig>& InOutConfigSet) const override
+		{
+			InOutConfigSet.Add({ FTable::GetHierarchyColumnId(),          true, 120.0f });
+			InOutConfigSet.Add({ FMemTagTableColumns::TagNameColumnId,    true, 400.0f });
+			InOutConfigSet.Add({ FMemTagTableColumns::SizeAColumnId,      true, 80.0f });
+			InOutConfigSet.Add({ FMemTagTableColumns::SizeBColumnId,      true, 80.0f });
+			InOutConfigSet.Add({ FMemTagTableColumns::SizeDiffColumnId,   true, 80.0f });
+			InOutConfigSet.Add({ FMemTagTableColumns::SizeBudgetColumnId, true, 80.0f });
+		}
+	};
+	return MakeShared<FDiffViewPreset>();
+}
+
+////////////////////////////////////////////////////////////////////////////////////////////////////
+// Min-Max View
+
+TSharedRef<ITableTreeViewPreset> FMemTagTableViewPresets::CreateMinMaxViewPreset(SMemTagTreeView& TableTreeView)
+{
+	class FMinMaxViewPreset : public ITableTreeViewPreset
+	{
+	public:
+		virtual FText GetName() const override
+		{
+			return LOCTEXT("MinMax_PresetName", "Min-Max");
+		}
+		virtual FText GetToolTip() const override
+		{
+			return LOCTEXT("MinMax_PresetToolTip", "Min-Max View\nConfigure the tree view to investigate the aggregated LLM stats for the selected time region.");
+		}
+		virtual FName GetSortColumn() const override
+		{
+			return FMemTagTableColumns::TagNameColumnId;
+		}
+		virtual EColumnSortMode::Type GetSortMode() const override
+		{
+			return EColumnSortMode::Type::Ascending;
+		}
+		virtual void SetCurrentGroupings(const TArray<TSharedPtr<FTreeNodeGrouping>>& InAvailableGroupings, TArray<TSharedPtr<FTreeNodeGrouping>>& InOutCurrentGroupings) const override
+		{
+			InOutCurrentGroupings.Reset();
+
+			check(InAvailableGroupings[0]->Is<FTreeNodeGroupingFlat>());
+			InOutCurrentGroupings.Add(InAvailableGroupings[0]);
+		}
+		virtual void GetColumnConfigSet(TArray<FTableColumnConfig>& InOutConfigSet) const override
+		{
+			InOutConfigSet.Add({ FTable::GetHierarchyColumnId(),           true, 120.0f });
+			InOutConfigSet.Add({ FMemTagTableColumns::TagNameColumnId,     true, 400.0f });
+			InOutConfigSet.Add({ FMemTagTableColumns::SizeMinColumnId,     true, 80.0f });
+			InOutConfigSet.Add({ FMemTagTableColumns::SizeMaxColumnId,     true, 80.0f });
+			InOutConfigSet.Add({ FMemTagTableColumns::SizeAverageColumnId, true, 80.0f });
+			InOutConfigSet.Add({ FMemTagTableColumns::SizeBudgetColumnId,  true, 80.0f });
+		}
+	};
+	return MakeShared<FMinMaxViewPreset>();
+}
+
+////////////////////////////////////////////////////////////////////////////////////////////////////
+
+} // namespace UE::Insights::MemoryProfiler
+
+#undef LOCTEXT_NAMESPACE
diff --git a/Engine/Source/Developer/TraceInsights/Private/Insights/MemoryProfiler/ViewModels/MemTagTableViewPresets.h b/Engine/Source/Developer/TraceInsights/Private/Insights/MemoryProfiler/ViewModels/MemTagTableViewPresets.h
new file mode 100644
--- /dev/null
+++ b/Engine/Source/Developer/TraceInsights/Private/Insights/MemoryProfiler/ViewModels/MemTagTableViewPresets.h
@@ -0,0 +1,24 @@
+// Copyright Epic Games, Inc. All Rights Reserved.
+
+#pragma once
+
+#include "CoreTypes.h"
+
+#include "Templates/SharedPointer.h"
+
+namespace UE::Insights { class ITableTreeViewPreset; }
+
+namespace UE::Insights::MemoryProfiler
+{
+
+class SMemTagTreeView;
+
+class FMemTagTableViewPresets
+{
+public:
+	static TSharedRef<ITableTreeViewPreset> CreateDefaultViewPreset(SMemTagTreeView& TableTreeView);
+	static TSharedRef<ITableTreeViewPreset> CreateDiffViewPreset(SMemTagTreeView& TableTreeView);
+	static TSharedRef<ITableTreeViewPreset> CreateMinMaxViewPreset(SMemTagTreeView& TableTreeView);
+};
+
+} // namespace UE::Insights::MemoryProfiler
diff --git a/Engine/Source/Developer/TraceInsights/Private/Insights/MemoryProfiler/ViewModels/MemTagTreeViewColumnFactory.cpp b/Engine/Source/Developer/TraceInsights/Private/Insights/MemoryProfiler/ViewModels/MemTagTreeViewColumnFactory.cpp
deleted file mode 100644
diff --git a/Engine/Source/Developer/TraceInsights/Private/Insights/MemoryProfiler/ViewModels/MemTagTreeViewColumnFactory.h b/Engine/Source/Developer/TraceInsights/Private/Insights/MemoryProfiler/ViewModels/MemTagTreeViewColumnFactory.h
deleted file mode 100644
diff --git a/Engine/Source/Developer/TraceInsights/Private/Insights/MemoryProfiler/ViewModels/MemorySharedState.h b/Engine/Source/Developer/TraceInsights/Private/Insights/MemoryProfiler/ViewModels/MemorySharedState.h
--- a/Engine/Source/Developer/TraceInsights/Private/Insights/MemoryProfiler/ViewModels/MemorySharedState.h
+++ b/Engine/Source/Developer/TraceInsights/Private/Insights/MemoryProfiler/ViewModels/MemorySharedState.h
@@ -108,6 +108,14 @@ public:
 	const TArray<TSharedPtr<FMemoryTracker>>& GetTrackers() const { return Trackers; }
 	FString TrackersToString(uint64 Flags, const TCHAR* Conjunction) const;
 	const FMemoryTracker* GetTrackerById(FMemoryTrackerId InMemTrackerId) const;
+	FMemoryTrackerId GetDefaultTrackerId() const
+	{
+		return DefaultTracker.IsValid() ? DefaultTracker->GetId() : FMemoryTracker::InvalidTrackerId;
+	}
+	FMemoryTrackerId GetPlatformTrackerId() const
+	{
+		return PlatformTracker.IsValid() ? PlatformTracker->GetId() : FMemoryTracker::InvalidTrackerId;
+	}
 
 	const TArray<TSharedPtr<FMemoryTagSet>>& GetTagSets() const { return TagSets; }
 	const FMemoryTagSet* GetTagSetById(FMemoryTagSetId InMemTagSetId) const;
diff --git a/Engine/Source/Developer/TraceInsights/Private/Insights/MemoryProfiler/Widgets/SMemTagTreeView.cpp b/Engine/Source/Developer/TraceInsights/Private/Insights/MemoryProfiler/Widgets/SMemTagTreeView.cpp
--- a/Engine/Source/Developer/TraceInsights/Private/Insights/MemoryProfiler/Widgets/SMemTagTreeView.cpp
+++ b/Engine/Source/Developer/TraceInsights/Private/Insights/MemoryProfiler/Widgets/SMemTagTreeView.cpp
@@ -7,31 +7,32 @@
 #include "SlateOptMacros.h"
 #include "Templates/UniquePtr.h"
 #include "Widgets/Colors/SColorPicker.h"
-#include "Widgets/Input/SCheckBox.h"
-#include "Widgets/Layout/SScrollBox.h"
+#include "Widgets/Images/SImage.h"
 #include "Widgets/Input/SButton.h"
+#include "Widgets/Input/SCheckBox.h"
+#include "Widgets/Input/SComboBox.h"
 #include "Widgets/Input/SSearchBox.h"
 #include "Widgets/Input/SSegmentedControl.h"
-#include "Widgets/Images/SImage.h"
+#include "Widgets/Layout/SScrollBox.h"
 #include "Widgets/SToolTip.h"
 #include "Widgets/Views/STableViewBase.h"
 
 // TraceInsightsCore
 #include "InsightsCore/Common/Stopwatch.h"
-#include "InsightsCore/Common/TimeUtils.h"
 #include "InsightsCore/Table/ViewModels/Table.h"
 #include "InsightsCore/Table/ViewModels/TableColumn.h"
 
 // TraceInsights
 #include "Insights/InsightsStyle.h"
 #include "Insights/MemoryProfiler/MemoryProfilerManager.h"
+#include "Insights/MemoryProfiler/ViewModels/MemorySharedState.h"
 #include "Insights/MemoryProfiler/ViewModels/MemoryTag.h"
 #include "Insights/MemoryProfiler/ViewModels/MemoryTracker.h"
-#include "Insights/MemoryProfiler/ViewModels/MemorySharedState.h"
-#include "Insights/MemoryProfiler/ViewModels/MemTagTreeViewColumnFactory.h"
-#include "Insights/MemoryProfiler/Widgets/SMemTagTreeViewTooltip.h"
-#include "Insights/MemoryProfiler/Widgets/SMemTagTreeViewTableRow.h"
+#include "Insights/MemoryProfiler/ViewModels/MemTagNode.h"
+#include "Insights/MemoryProfiler/ViewModels/MemTagTableViewPresets.h"
 #include "Insights/MemoryProfiler/Widgets/SMemoryProfilerWindow.h"
+#include "Insights/TimingProfiler/ViewModels/TimeMarker.h"
+#include "Insights/Widgets/STimingView.h"
 
 #define LOCTEXT_NAMESPACE "UE::Insights::MemoryProfiler::SMemTagTreeView"
 
@@ -41,141 +42,115 @@ namespace UE::Insights::MemoryProfiler
 ////////////////////////////////////////////////////////////////////////////////////////////////////
 
 SMemTagTreeView::SMemTagTreeView()
-	: ProfilerWindowWeakPtr()
-	, Table(MakeShared<FTable>())
-	, bExpansionSaved(false)
-	, bFilterOutZeroCountMemTags(false)
-	, TrackersFilter(uint64(-1))
-	, TagSetFilter(FMemoryTagSet::DefaultTagSetId)
-	, GroupingMode(EMemTagNodeGroupingMode::Flat)
-	, AvailableSorters()
-	, CurrentSorter(nullptr)
-	, ColumnBeingSorted(GetDefaultColumnBeingSorted())
-	, ColumnSortMode(GetDefaultColumnSortMode())
-	, StatsStartTime(0.0)
-	, StatsEndTime(0.0)
 {
+	bRunInAsyncMode = true;
 }
 
 ////////////////////////////////////////////////////////////////////////////////////////////////////
 
 SMemTagTreeView::~SMemTagTreeView()
 {
-	// Remove ourselves from the Insights manager.
-	if (FInsightsManager::Get().IsValid())
-	{
-		FInsightsManager::Get()->GetSessionChangedEvent().RemoveAll(this);
-	}
-
-	Session.Reset();
 }
 
 ////////////////////////////////////////////////////////////////////////////////////////////////////
-BEGIN_SLATE_FUNCTION_BUILD_OPTIMIZATION
 
 void SMemTagTreeView::Construct(const FArguments& InArgs, TSharedPtr<SMemoryProfilerWindow> InProfilerWindow)
 {
 	check(InProfilerWindow.IsValid());
 	ProfilerWindowWeakPtr = InProfilerWindow;
 
-	SAssignNew(ExternalScrollbar, SScrollBar)
-	.AlwaysShowScrollbar(true);
+	TSharedRef<FMemTagTable> TablePtr = MakeShared<FMemTagTable>();
+	TablePtr->Reset();
+	TablePtr->SetDisplayName(FText::FromString(TEXT("LLM Tags")));
 
-	ChildSlot
-	[
-		SNew(SVerticalBox)
+	ConstructWidget(TablePtr);
 
-		+ SVerticalBox::Slot()
-		.VAlign(VAlign_Center)
-		.AutoHeight()
-		.Padding(2.0f, 2.0f, 2.0f, 2.0f)
-		[
-			SNew(SVerticalBox)
+	// Apply the default preset.
+	ApplyViewPreset(*(*GetAvailableViewPresets())[0]);
+}
 
-			// Search and Filtering
-			+ SVerticalBox::Slot()
-			.VAlign(VAlign_Center)
-			.Padding(2.0f)
-			.AutoHeight()
-			[
-				ConstructTagsFilteringWidgetArea()
-			]
+////////////////////////////////////////////////////////////////////////////////////////////////////
 
-			// Tracker and Grouping
-			+ SVerticalBox::Slot()
-			.VAlign(VAlign_Center)
-			.Padding(2.0f)
-			.AutoHeight()
-			[
-				ConstructTagsGroupingWidgetArea()
-			]
+BEGIN_SLATE_FUNCTION_BUILD_OPTIMIZATION
+TSharedPtr<SWidget> SMemTagTreeView::ConstructFilterToolbar()
+{
+	FSlimHorizontalToolBarBuilder ToolbarBuilder(TSharedPtr<const FUICommandList>(), FMultiBoxCustomization::None);
+	ToolbarBuilder.SetStyle(&FInsightsStyle::Get(), "SecondaryToolbar");
 
-			// Tracks Mini Toolbar
-			+ SVerticalBox::Slot()
-			.VAlign(VAlign_Center)
-			.Padding(2.0f)
-			.AutoHeight()
+	ToolbarBuilder.BeginSection("Filters");
+	{
+		// Filter Trackers
+		ToolbarBuilder.AddComboButton(
+			FUIAction(),
+			FOnGetContent::CreateSP(this, &SMemTagTreeView::MakeTrackersMenu),
+			LOCTEXT("TrackersMenuText", "Trackers"),
+			LOCTEXT("TrackersMenuToolTip", "Filter list of LLM tags by tracker."),
+			FSlateIcon(FInsightsStyle::GetStyleSetName(), "Icons.Filter.ToolBar"),
+			false
+		);
+
+		// Select Tag Set
+		ToolbarBuilder.AddWidget(
+			SNew(SComboBox<TSharedPtr<FMemoryTagSetId>>)
+			.ToolTipText(LOCTEXT("TagSetToolTip", "Select the LLM tag set."))
+			.OptionsSource(&AvailableTagSets)
+			.OnSelectionChanged(this, &SMemTagTreeView::TagSet_OnSelectionChanged)
+			.OnGenerateWidget(this, &SMemTagTreeView::TagSet_OnGenerateWidget)
 			[
-				ConstructTracksMiniToolbar()
-			]
+				SNew(STextBlock)
+				.Text(this, &SMemTagTreeView::TagSet_GetSelectedText)
+			]);
+	}
+	ToolbarBuilder.EndSection();
+
+	return SNew(SHorizontalBox)
+
+	+ SHorizontalBox::Slot()
+		.FillWidth(1.0f)
+		.VAlign(VAlign_Center)
+		[
+			ConstructSearchBox()
 		]
 
-		// Tree view
-		+ SVerticalBox::Slot()
-		.FillHeight(1.0f)
-		.Padding(0.0f, 2.0f, 0.0f, 0.0f)
+	+ SHorizontalBox::Slot()
+		.AutoWidth()
+		.Padding(4.0f, 0.0f, 0.0f, 0.0f)
+		.VAlign(VAlign_Center)
 		[
-			SNew(SHorizontalBox)
+			ConstructFilterConfiguratorButton()
+		]
 
-			+ SHorizontalBox::Slot()
-			.FillWidth(1.0f)
-			.Padding(0.0f)
-			[
-				SAssignNew(TreeView, STreeView<FMemTagNodePtr>)
-				.ExternalScrollbar(ExternalScrollbar)
-				.SelectionMode(ESelectionMode::Multi)
-				.TreeItemsSource(&FilteredGroupNodes)
-				.OnGetChildren(this, &SMemTagTreeView::TreeView_OnGetChildren)
-				.OnGenerateRow(this, &SMemTagTreeView::TreeView_OnGenerateRow)
-				.OnSelectionChanged(this, &SMemTagTreeView::TreeView_OnSelectionChanged)
-				.OnMouseButtonDoubleClick(this, &SMemTagTreeView::TreeView_OnMouseButtonDoubleClick)
-				.OnContextMenuOpening(FOnContextMenuOpening::CreateSP(this, &SMemTagTreeView::TreeView_GetMenuContent))
-				.HeaderRow
-				(
-					SAssignNew(TreeViewHeaderRow, SHeaderRow)
-					.Visibility(EVisibility::Visible)
-				)
-			]
+	+ SHorizontalBox::Slot()
+		.AutoWidth()
+		.Padding(-2.0f, 0.0f, -6.0f, 0.0f)
+		.VAlign(VAlign_Center)
+		[
+			ToolbarBuilder.MakeWidget()
+		];
+}
+END_SLATE_FUNCTION_BUILD_OPTIMIZATION
 
-			+ SHorizontalBox::Slot()
-			.AutoWidth()
-			.Padding(0.0f)
-			[
-				SNew(SBox)
-				.WidthOverride(FOptionalSize(13.0f))
-				[
-					ExternalScrollbar.ToSharedRef()
-				]
-			]
-		]
-	];
+////////////////////////////////////////////////////////////////////////////////////////////////////
 
-	InitializeAndShowHeaderColumns();
+BEGIN_SLATE_FUNCTION_BUILD_OPTIMIZATION
+TSharedPtr<SWidget> SMemTagTreeView::ConstructToolbar()
+{
+	TSharedPtr<SHorizontalBox> Box = SNew(SHorizontalBox);
 
-	// Create the search filters: text based, type based etc.
-	TextFilter = MakeShared<FMemTagNodeTextFilter>(FMemTagNodeTextFilter::FItemToStringArray::CreateSP(this, &SMemTagTreeView::HandleItemToStringArray));
-	Filters = MakeShared<FMemTagNodeFilterCollection>();
-	Filters->Add(TextFilter);
+	ConstructViewPreset(Box);
 
-	CreateGroupByOptionsSources();
-	CreateSortings();
+	return Box;
+}
+END_SLATE_FUNCTION_BUILD_OPTIMIZATION
 
-	// Register ourselves with the Insights manager.
-	FInsightsManager::Get()->GetSessionChangedEvent().AddSP(this, &SMemTagTreeView::InsightsManager_OnSessionChanged);
+////////////////////////////////////////////////////////////////////////////////////////////////////
 
-	// Update the Session (i.e. when analysis session was already started).
-	InsightsManager_OnSessionChanged();
+BEGIN_SLATE_FUNCTION_BUILD_OPTIMIZATION
+TSharedPtr<SWidget> SMemTagTreeView::ConstructExtraToolbar()
+{
+	return ConstructTracksMiniToolbar();
 }
+END_SLATE_FUNCTION_BUILD_OPTIMIZATION
 
 ////////////////////////////////////////////////////////////////////////////////////////////////////
 
@@ -222,7 +197,7 @@ void SMemTagTreeView::CreateTrackersMenuSection(FMenuBuilder& MenuBuilder)
 void SMemTagTreeView::ToggleTracker(FMemoryTrackerId InTrackerId)
 {
 	TrackersFilter ^= FMemoryTracker::AsFlag(InTrackerId);
-	ApplyFiltering();
+	OnNodeFilteringChanged();
 }
 
 ////////////////////////////////////////////////////////////////////////////////////////////////////
@@ -239,7 +214,7 @@ void SMemTagTreeView::TagSet_OnSelectionChanged(TSharedPtr<FMemoryTagSetId> NewT
 	if (NewTagSet.IsValid())
 	{
 		TagSetFilter = *NewTagSet;
-		ApplyFiltering();
+		OnNodeFilteringChanged();
 	}
 }
 
@@ -290,133 +265,7 @@ FText SMemTagTreeView::TagSet_GetSelectedText() const
 
 ////////////////////////////////////////////////////////////////////////////////////////////////////
 
-TSharedRef<SWidget> SMemTagTreeView::ConstructTagsFilteringWidgetArea()
-{
-	FSlimHorizontalToolBarBuilder ToolbarBuilder(TSharedPtr<const FUICommandList>(), FMultiBoxCustomization::None);
-	ToolbarBuilder.SetStyle(&FInsightsStyle::Get(), "SecondaryToolbar");
-
-	ToolbarBuilder.BeginSection("Filters");
-	{
-		// Text Filter (Search Box)
-		//ToolbarBuilder.AddWidget(
-		//	SAssignNew(SearchBox, SSearchBox)
-		//	.HintText(LOCTEXT("SearchBoxHint", "Search LLM tags or groups"))
-		//	.ToolTipText(LOCTEXT("FilterSearchHint", "Type here to search LLM tags or groups"))
-		//	.OnTextChanged(this, &SMemTagTreeView::SearchBox_OnTextChanged)
-		//	.IsEnabled(this, &SMemTagTreeView::SearchBox_IsEnabled)
-		//);
-
-		// Filter Trackers
-		ToolbarBuilder.AddComboButton(
-			FUIAction(),
-			FOnGetContent::CreateSP(this, &SMemTagTreeView::MakeTrackersMenu),
-			LOCTEXT("TrackersMenuText", "Trackers"),
-			LOCTEXT("TrackersMenuToolTip", "Filter list of LLM tags by tracker."),
-			FSlateIcon(FInsightsStyle::GetStyleSetName(), "Icons.Filter.ToolBar"),
-			false
-		);
-
-		// Select Tag Set
-		ToolbarBuilder.AddWidget(
-			SAssignNew(TagSetComboBox, SComboBox<TSharedPtr<FMemoryTagSetId>>)
-				.ToolTipText(LOCTEXT("TagSetToolTip", "Select the LLM tag set."))
-				.OptionsSource(&AvailableTagSets)
-				.OnSelectionChanged(this, &SMemTagTreeView::TagSet_OnSelectionChanged)
-				.OnGenerateWidget(this, &SMemTagTreeView::TagSet_OnGenerateWidget)
-				[
-					SNew(STextBlock)
-					.Text(this, &SMemTagTreeView::TagSet_GetSelectedText)
-				]);
-	}
-	ToolbarBuilder.EndSection();
-
-	TSharedRef<SWidget> Widget = SNew(SHorizontalBox)
-
-	// Search box
-	+ SHorizontalBox::Slot()
-	.VAlign(VAlign_Center)
-	.Padding(2.0f)
-	.FillWidth(1.0f)
-	[
-		SAssignNew(SearchBox, SSearchBox)
-		.HintText(LOCTEXT("SearchBoxHint", "Search LLM tags or groups"))
-		.OnTextChanged(this, &SMemTagTreeView::SearchBox_OnTextChanged)
-		.IsEnabled(this, &SMemTagTreeView::SearchBox_IsEnabled)
-		.ToolTipText(LOCTEXT("FilterSearchHint", "Type here to search LLM tags or groups"))
-	]
-
-	// Filter out LLM tags with zero instance count.
-	//+ SHorizontalBox::Slot()
-	//.VAlign(VAlign_Center)
-	//.Padding(2.0f)
-	//.AutoWidth()
-	//[
-	//	SNew(SCheckBox)
-	//	.Style(FAppStyle::Get(), "ToggleButtonCheckbox")
-	//	.HAlign(HAlign_Center)
-	//	.Padding(3.0f)
-	//	.OnCheckStateChanged(this, &SMemTagTreeView::FilterOutZeroCountMemTags_OnCheckStateChanged)
-	//	.IsChecked(this, &SMemTagTreeView::FilterOutZeroCountMemTags_IsChecked)
-	//	.ToolTipText(LOCTEXT("FilterOutZeroCountMemTags_Tooltip", "Filter out the LLM tags having zero total instance count (aggregated stats)."))
-	//	[
-	//		SNew(SImage)
-	//		.Image(FInsightsStyle::Get().GetBrush("Icons.ZeroCountFilter"))
-	//	]
-	//]
-
-	// Trackers and Tag Sets
-	+ SHorizontalBox::Slot()
-	.VAlign(VAlign_Center)
-	.Padding(FMargin(-2.0f, -5.0f))
-	.AutoWidth()
-	[
-		ToolbarBuilder.MakeWidget()
-	];
-
-	return Widget;
-}
-
-////////////////////////////////////////////////////////////////////////////////////////////////////
-
-TSharedRef<SWidget> SMemTagTreeView::ConstructTagsGroupingWidgetArea()
-{
-	TSharedRef<SWidget> Widget = SNew(SHorizontalBox)
-
-	+ SHorizontalBox::Slot()
-	.AutoWidth()
-	.VAlign(VAlign_Center)
-	.Padding(0.0f, 0.0f, 4.0f, 0.0f)
-	[
-		SNew(STextBlock)
-		.MinDesiredWidth(60.0f)
-		.Justification(ETextJustify::Right)
-		.Text(LOCTEXT("GroupByText", "Group by"))
-	]
-
-	+ SHorizontalBox::Slot()
-	.AutoWidth()
-	.VAlign(VAlign_Center)
-	[
-		SNew(SBox)
-		.MinDesiredWidth(100.0f)
-		[
-			SAssignNew(GroupByComboBox, SComboBox<TSharedPtr<EMemTagNodeGroupingMode>>)
-			.ToolTipText(this, &SMemTagTreeView::GroupBy_GetSelectedTooltipText)
-			.OptionsSource(&GroupByOptionsSource)
-			.OnSelectionChanged(this, &SMemTagTreeView::GroupBy_OnSelectionChanged)
-			.OnGenerateWidget(this, &SMemTagTreeView::GroupBy_OnGenerateWidget)
-			[
-				SNew(STextBlock)
-				.Text(this, &SMemTagTreeView::GroupBy_GetSelectedText)
-			]
-		]
-	];
-
-	return Widget;
-}
-
-////////////////////////////////////////////////////////////////////////////////////////////////////
-
+BEGIN_SLATE_FUNCTION_BUILD_OPTIMIZATION
 TSharedRef<SWidget> SMemTagTreeView::ConstructTracksMiniToolbar()
 {
 	TSharedRef<SWidget> Widget = SNew(SHorizontalBox)
@@ -467,39 +316,6 @@ TSharedRef<SWidget> SMemTagTreeView::ConstructTracksMiniToolbar()
 		]
 	]
 
-	+ SHorizontalBox::Slot()
-	.VAlign(VAlign_Center)
-	.Padding(FMargin(2.0f, 0.0f, 2.0f, 0.0f))
-	.AutoWidth()
-	[
-		SNew(SButton)
-		.ToolTipText(LOCTEXT("LoadReportXML_ToolTip", "Load LLMReportTypes.xml"))
-		.OnClicked(this, &SMemTagTreeView::LoadReportXML_OnClicked)
-		.ButtonStyle(&FAppStyle::Get().GetWidgetStyle<FButtonStyle>("SimpleButton"))
-		.ContentPadding(FMargin(4.0f, 2.0f, 4.0f, 2.0f))
-		[
-			SNew(SHorizontalBox)
-			+ SHorizontalBox::Slot()
-			.AutoWidth()
-			.VAlign(VAlign_Center)
-			[
-				SNew(SImage)
-				.ColorAndOpacity(FSlateColor::UseForeground())
-				.Image(FAppStyle::Get().GetBrush("Icons.FolderOpen"))
-			]
-			+ SHorizontalBox::Slot()
-			.VAlign(VAlign_Center)
-			.Padding(4.0f, 0.0f, 0.0f, 0.0f)
-			.AutoWidth()
-			[
-				SNew(STextBlock)
-				.TextStyle(&FAppStyle::Get().GetWidgetStyle< FTextBlockStyle >(TEXT("ButtonText")))
-				.Justification(ETextJustify::Center)
-				.Text(LOCTEXT("LoadReportXML_Text", "Load XML..."))
-			]
-		]
-	]
-
 	+ SHorizontalBox::Slot()
 	.VAlign(VAlign_Center)
 	.HAlign(HAlign_Right)
@@ -538,1364 +354,230 @@ TSharedRef<SWidget> SMemTagTreeView::ConstructTracksMiniToolbar()
 
 	return Widget;
 }
-
 END_SLATE_FUNCTION_BUILD_OPTIMIZATION
 
 ////////////////////////////////////////////////////////////////////////////////////////////////////
 
-TSharedPtr<SWidget> SMemTagTreeView::TreeView_GetMenuContent()
+BEGIN_SLATE_FUNCTION_BUILD_OPTIMIZATION
+TSharedPtr<SWidget> SMemTagTreeView::ConstructFooter()
 {
-	const TArray<FMemTagNodePtr> SelectedNodes = TreeView->GetSelectedItems();
-	const int32 NumSelectedNodes = SelectedNodes.Num();
-	FMemTagNodePtr SelectedNode = NumSelectedNodes ? SelectedNodes[0] : nullptr;
+	return
+		SNew(SHorizontalBox)
 
-	FText SelectionStr;
-	if (NumSelectedNodes == 0)
-	{
-		SelectionStr = LOCTEXT("NothingSelected", "Nothing selected");
-	}
-	else if (NumSelectedNodes == 1)
-	{
-		FString ItemName = SelectedNode->GetName().ToString();
-		const int32 MaxStringLen = 64;
-		if (ItemName.Len() > MaxStringLen)
-		{
-			ItemName = ItemName.Left(MaxStringLen) + TEXT("...");
-		}
-		SelectionStr = FText::FromString(ItemName);
-	}
-	else
-	{
-		SelectionStr = FText::Format(LOCTEXT("MultipleSelection_Fmt", "{0} selected items"), FText::AsNumber(NumSelectedNodes));
-	}
+		+ SHorizontalBox::Slot()
+		.FillWidth(1.0f)
+		.HAlign(HAlign_Center)
+		.Padding(2.0f, 2.0f, 2.0f, 2.0f)
+		[
+			SNew(STextBlock)
+			.Text(this, &SMemTagTreeView::GetFooterCenterText)
+		];
+}
+END_SLATE_FUNCTION_BUILD_OPTIMIZATION
 
-	const bool bShouldCloseWindowAfterMenuSelection = true;
-	FMenuBuilder MenuBuilder(bShouldCloseWindowAfterMenuSelection, NULL);
+////////////////////////////////////////////////////////////////////////////////////////////////////
+
+void SMemTagTreeView::UpdateSelectionStatsText()
+{
+	TArray<FTableTreeNodePtr> SelectedNodes;
+	const int32 NumSelectedNodes = TreeView->GetSelectedItems(SelectedNodes);
 
-	// Selection menu
-	MenuBuilder.BeginSection("Selection", LOCTEXT("ContextMenu_Section_Selection", "Selection"));
+	if (NumSelectedNodes > 0)
 	{
-		struct FLocal
+		int64 TotalCount = 0;
+		int64 TotalSizeA = 0;
+		int64 TotalSizeB = 0;
+		for (const FTableTreeNodePtr& SelectedNode : SelectedNodes)
 		{
-			static bool ReturnFalse()
+			if (SelectedNode->Is<FMemTagNode>())
 			{
-				return false;
+				const FMemTagNode& MemTagNode = SelectedNode->As<FMemTagNode>();
+				++TotalCount;
+				TotalSizeA += MemTagNode.GetStats().SizeA;
+				TotalSizeB += MemTagNode.GetStats().SizeB;
 			}
-		};
+		}
 
-		FUIAction DummyUIAction;
-		DummyUIAction.CanExecuteAction = FCanExecuteAction::CreateStatic(&FLocal::ReturnFalse);
-		MenuBuilder.AddMenuEntry
-		(
-			SelectionStr,
-			LOCTEXT("ContextMenu_Selection", "Currently selected items"),
-			FSlateIcon(),
-			DummyUIAction,
-			NAME_None,
-			EUserInterfaceActionType::Button
-		);
-	}
-	MenuBuilder.EndSection();
+		FNumberFormattingOptions FormattingOptionsMem;
+		FormattingOptionsMem.MaximumFractionalDigits = 2;
 
-	MenuBuilder.BeginSection("Misc", LOCTEXT("ContextMenu_Section_Misc", "Miscellaneous"));
+		SelectionStatsText = FText::Format(LOCTEXT("SelectionStatsFmt", "{0} selected {0}|plural(one=tag,other=tags) (A: {1}, B: {2}, B-A: {3})"),
+			FText::AsNumber(TotalCount),
+			FText::AsMemory(TotalSizeA, &FormattingOptionsMem),
+			FText::AsMemory(TotalSizeB, &FormattingOptionsMem),
+			FText::AsMemory(TotalSizeB - TotalSizeA, &FormattingOptionsMem));
+	}
+	else
 	{
-		MenuBuilder.AddSubMenu
-		(
-			LOCTEXT("ContextMenu_CreateGraphTracks_SubMenu", "Create Graph Tracks"),
-			LOCTEXT("ContextMenu_CreateGraphTracks_SubMenu_Desc", "Create memory graph tracks."),
-			FNewMenuDelegate::CreateSP(this, &SMemTagTreeView::TreeView_BuildCreateGraphTracksMenu),
-			false,
-			FSlateIcon()
-		);
+		SelectionStatsText = LOCTEXT("NoSelectionStats", "No LLM tag selected");
+	}
+}
 
-		MenuBuilder.AddSubMenu
-		(
-			LOCTEXT("ContextMenu_RemoveGraphTracks_SubMenu", "Remove Graph Tracks"),
-			LOCTEXT("ContextMenu_RemoveGraphTracks_SubMenu_Desc", "Remove memory graph tracks."),
-			FNewMenuDelegate::CreateSP(this, &SMemTagTreeView::TreeView_BuildRemoveGraphTracksMenu),
-			false,
-			FSlateIcon()
-		);
+////////////////////////////////////////////////////////////////////////////////////////////////////
 
-		FUIAction Action_GenerateColorForSelectedMemTags
-		(
-			FExecuteAction::CreateSP(this, &SMemTagTreeView::GenerateColorForSelectedMemTags),
-			FCanExecuteAction::CreateSP(this, &SMemTagTreeView::CanGenerateColorForSelectedMemTags)
-		);
-		MenuBuilder.AddMenuEntry
-		(
-			LOCTEXT("ContextMenu_GenerateColorForSelectedMemTags", "Generate New Color"),
-			LOCTEXT("ContextMenu_GenerateColorForSelectedMemTags_Desc", "Generate new color for selected LLM tag(s)."),
-			FSlateIcon(),
-			Action_GenerateColorForSelectedMemTags,
-			NAME_None,
-			EUserInterfaceActionType::Button
-		);
+FText SMemTagTreeView::GetFooterCenterText() const
+{
+	return SelectionStatsText;
+}
 
-		FUIAction Action_EditColorForSelectedMemTags
-		(
-			FExecuteAction::CreateSP(this, &SMemTagTreeView::EditColorForSelectedMemTags),
-			FCanExecuteAction::CreateSP(this, &SMemTagTreeView::CanEditColorForSelectedMemTags)
-		);
-		MenuBuilder.AddMenuEntry
-		(
-			LOCTEXT("ContextMenu_EditColorForSelectedMemTags", "Edit Color..."),
-			LOCTEXT("ContextMenu_EditColorForSelectedMemTags_Desc", "Edit color for selected LLM tag(s)."),
-			FSlateIcon(),
-			Action_EditColorForSelectedMemTags,
-			NAME_None,
-			EUserInterfaceActionType::Button
-		);
-	}
-	MenuBuilder.EndSection();
+////////////////////////////////////////////////////////////////////////////////////////////////////
 
-	return MenuBuilder.MakeWidget();
+void SMemTagTreeView::TreeView_OnSelectionChanged(FTableTreeNodePtr SelectedItem, ESelectInfo::Type SelectInfo)
+{
+	UpdateSelectionStatsText();
 }
 
 ////////////////////////////////////////////////////////////////////////////////////////////////////
 
-void SMemTagTreeView::TreeView_BuildCreateGraphTracksMenu(FMenuBuilder& MenuBuilder)
+void SMemTagTreeView::TreeView_OnMouseButtonDoubleClick(FTableTreeNodePtr NodePtr)
 {
-	MenuBuilder.BeginSection("CreateGraphTracks");
+	if (NodePtr->IsGroup())
 	{
-		// Create memory graph tracks for selected LLM tag(s)
-		FUIAction Action_CreateGraphTracksForSelectedMemTags
-		(
-			FExecuteAction::CreateSP(this, &SMemTagTreeView::CreateGraphTracksForSelectedMemTags),
-			FCanExecuteAction::CreateSP(this, &SMemTagTreeView::CanCreateGraphTracksForSelectedMemTags)
-		);
-		MenuBuilder.AddMenuEntry
-		(
-			LOCTEXT("ContextMenu_CreateGraphTracksForSelectedMemTags", "Selected"),
-			LOCTEXT("ContextMenu_CreateGraphTracksForSelectedMemTags_Desc", "Create memory graph tracks for selected LLM tag(s)."),
-			FSlateIcon(),
-			Action_CreateGraphTracksForSelectedMemTags,
-			NAME_None,
-			EUserInterfaceActionType::Button
-		);
-
-		// Create memory graph tracks for filtered LLM tags
-		FUIAction Action_CreateGraphTracksForFilteredMemTags
-		(
-			FExecuteAction::CreateSP(this, &SMemTagTreeView::CreateGraphTracksForFilteredMemTags),
-			FCanExecuteAction::CreateSP(this, &SMemTagTreeView::CanCreateGraphTracksForFilteredMemTags)
-		);
-		MenuBuilder.AddMenuEntry
-		(
-			LOCTEXT("ContextMenu_CreateGraphTracksForFilteredMemTags", "Filtered"),
-			LOCTEXT("ContextMenu_CreateGraphTracksForFilteredMemTags_Desc", "Create memory graph tracks for all visible (filtered) LLM tags."),
-			FSlateIcon(),
-			Action_CreateGraphTracksForFilteredMemTags,
-			NAME_None,
-			EUserInterfaceActionType::Button
-		);
-
-		// Create memory graph tracks for all LLM tags
-		FUIAction Action_CreateAllGraphTracks
-		(
-			FExecuteAction::CreateSP(this, &SMemTagTreeView::CreateAllGraphTracks),
-			FCanExecuteAction::CreateSP(this, &SMemTagTreeView::CanCreateAllGraphTracks)
-		);
-		MenuBuilder.AddMenuEntry
-		(
-			LOCTEXT("ContextMenu_CreateAllGraphTracks", "All"),
-			LOCTEXT("ContextMenu_CreateAllGraphTracks_Desc", "Create memory graph tracks for all LLM tags."),
-			FSlateIcon(),
-			Action_CreateAllGraphTracks,
-			NAME_None,
-			EUserInterfaceActionType::Button
-		);
-	}
-	MenuBuilder.EndSection();
-}
-
-////////////////////////////////////////////////////////////////////////////////////////////////////
-
-void SMemTagTreeView::TreeView_BuildRemoveGraphTracksMenu(FMenuBuilder& MenuBuilder)
-{
-	MenuBuilder.BeginSection("RemoveGraphTracks");
-	{
-		// Remove memory graph tracks for selected LLM tag(s)
-		FUIAction Action_RemoveGraphTracksForSelectedMemTags
-		(
-			FExecuteAction::CreateSP(this, &SMemTagTreeView::RemoveGraphTracksForSelectedMemTags),
-			FCanExecuteAction::CreateSP(this, &SMemTagTreeView::CanRemoveGraphTracksForSelectedMemTags)
-		);
-		MenuBuilder.AddMenuEntry
-		(
-			LOCTEXT("ContextMenu_RemoveGraphTracksForSelectedMemTags", "Selected"),
-			LOCTEXT("ContextMenu_RemoveGraphTracksForSelectedMemTags_Desc", "Remove memory graph tracks for selected LLM tag(s)."),
-			FSlateIcon(),
-			Action_RemoveGraphTracksForSelectedMemTags,
-			NAME_None,
-			EUserInterfaceActionType::Button
-		);
-
-		// Remove memory graph tracks for all LLM tags
-		FUIAction Action_RemoveAllGraphTracks
-		(
-			FExecuteAction::CreateSP(this, &SMemTagTreeView::RemoveAllGraphTracks),
-			FCanExecuteAction::CreateSP(this, &SMemTagTreeView::CanRemoveAllGraphTracks)
-		);
-		MenuBuilder.AddMenuEntry
-		(
-			LOCTEXT("ContextMenu_RemoveAllGraphTracks", "All"),
-			LOCTEXT("ContextMenu_RemoveAllGraphTracks_Desc", "Remove memory graph tracks for all LLM tags."),
-			FSlateIcon(),
-			Action_RemoveAllGraphTracks,
-			NAME_None,
-			EUserInterfaceActionType::Button
-		);
-	}
-	MenuBuilder.EndSection();
-}
-
-////////////////////////////////////////////////////////////////////////////////////////////////////
-
-void SMemTagTreeView::TreeView_BuildSortByMenu(FMenuBuilder& MenuBuilder)
-{
-	MenuBuilder.BeginSection("SortColumn", LOCTEXT("ContextMenu_Section_SortColumn", "Sort Column"));
-
-	for (const TSharedRef<FTableColumn>& ColumnRef : Table->GetColumns())
-	{
-		const FTableColumn& Column = *ColumnRef;
-
-		if (Column.IsVisible() && Column.CanBeSorted())
-		{
-			FUIAction Action_SortByColumn
-			(
-				FExecuteAction::CreateSP(this, &SMemTagTreeView::ContextMenu_SortByColumn_Execute, Column.GetId()),
-				FCanExecuteAction::CreateSP(this, &SMemTagTreeView::ContextMenu_SortByColumn_CanExecute, Column.GetId()),
-				FIsActionChecked::CreateSP(this, &SMemTagTreeView::ContextMenu_SortByColumn_IsChecked, Column.GetId())
-			);
-			MenuBuilder.AddMenuEntry
-			(
-				Column.GetTitleName(),
-				Column.GetDescription(),
-				FSlateIcon(),
-				Action_SortByColumn,
-				NAME_None,
-				EUserInterfaceActionType::RadioButton
-			);
-		}
-	}
-
-	MenuBuilder.EndSection();
-
-	MenuBuilder.BeginSection("SortMode", LOCTEXT("ContextMenu_Section_SortMode", "Sort Mode"));
-	{
-		FUIAction Action_SortAscending
-		(
-			FExecuteAction::CreateSP(this, &SMemTagTreeView::ContextMenu_SortMode_Execute, EColumnSortMode::Ascending),
-			FCanExecuteAction::CreateSP(this, &SMemTagTreeView::ContextMenu_SortMode_CanExecute, EColumnSortMode::Ascending),
-			FIsActionChecked::CreateSP(this, &SMemTagTreeView::ContextMenu_SortMode_IsChecked, EColumnSortMode::Ascending)
-		);
-		MenuBuilder.AddMenuEntry
-		(
-			LOCTEXT("ContextMenu_SortAscending", "Sort Ascending"),
-			LOCTEXT("ContextMenu_SortAscending_Desc", "Sorts ascending."),
-			FSlateIcon(FAppStyle::Get().GetStyleSetName(), "Icons.SortUp"),
-			Action_SortAscending,
-			NAME_None,
-			EUserInterfaceActionType::RadioButton
-		);
-
-		FUIAction Action_SortDescending
-		(
-			FExecuteAction::CreateSP(this, &SMemTagTreeView::ContextMenu_SortMode_Execute, EColumnSortMode::Descending),
-			FCanExecuteAction::CreateSP(this, &SMemTagTreeView::ContextMenu_SortMode_CanExecute, EColumnSortMode::Descending),
-			FIsActionChecked::CreateSP(this, &SMemTagTreeView::ContextMenu_SortMode_IsChecked, EColumnSortMode::Descending)
-		);
-		MenuBuilder.AddMenuEntry
-		(
-			LOCTEXT("ContextMenu_SortDescending", "Sort Descending"),
-			LOCTEXT("ContextMenu_SortDescending_Desc", "Sorts descending."),
-			FSlateIcon(FAppStyle::Get().GetStyleSetName(), "Icons.SortDown"),
-			Action_SortDescending,
-			NAME_None,
-			EUserInterfaceActionType::RadioButton
-		);
-	}
-	MenuBuilder.EndSection();
-}
-
-////////////////////////////////////////////////////////////////////////////////////////////////////
-
-void SMemTagTreeView::TreeView_BuildViewColumnMenu(FMenuBuilder& MenuBuilder)
-{
-	MenuBuilder.BeginSection("Columns", LOCTEXT("ContextMenu_Section_Columns", "Columns"));
-
-	for (const TSharedRef<FTableColumn>& ColumnRef : Table->GetColumns())
-	{
-		const FTableColumn& Column = *ColumnRef;
-
-		FUIAction Action_ToggleColumn
-		(
-			FExecuteAction::CreateSP(this, &SMemTagTreeView::ToggleColumnVisibility, Column.GetId()),
-			FCanExecuteAction::CreateSP(this, &SMemTagTreeView::CanToggleColumnVisibility, Column.GetId()),
-			FIsActionChecked::CreateSP(this, &SMemTagTreeView::IsColumnVisible, Column.GetId())
-		);
-		MenuBuilder.AddMenuEntry
-		(
-			Column.GetTitleName(),
-			Column.GetDescription(),
-			FSlateIcon(),
-			Action_ToggleColumn,
-			NAME_None,
-			EUserInterfaceActionType::ToggleButton
-		);
-	}
-
-	MenuBuilder.EndSection();
-}
-
-////////////////////////////////////////////////////////////////////////////////////////////////////
-
-void SMemTagTreeView::InitializeAndShowHeaderColumns()
-{
-	// Create columns.
-	TArray<TSharedRef<FTableColumn>> Columns;
-	FMemTagTreeViewColumnFactory::CreateMemTagTreeViewColumns(Columns);
-	Table->SetColumns(Columns);
-
-	// Show columns.
-	for (const TSharedRef<FTableColumn>& ColumnRef : Table->GetColumns())
-	{
-		if (ColumnRef->ShouldBeVisible())
-		{
-			ShowColumn(ColumnRef->GetId());
-		}
-	}
-}
-
-////////////////////////////////////////////////////////////////////////////////////////////////////
-
-FText SMemTagTreeView::GetColumnHeaderText(const FName ColumnId) const
-{
-	const FTableColumn& Column = *Table->FindColumnChecked(ColumnId);
-	return Column.GetShortName();
-}
-
-////////////////////////////////////////////////////////////////////////////////////////////////////
-
-TSharedRef<SWidget> SMemTagTreeView::TreeViewHeaderRow_GenerateColumnMenu(const FTableColumn& Column)
-{
-	const bool bShouldCloseWindowAfterMenuSelection = true;
-	FMenuBuilder MenuBuilder(bShouldCloseWindowAfterMenuSelection, NULL);
-
-	MenuBuilder.BeginSection("Sorting", LOCTEXT("ContextMenu_Section_Sorting", "Sorting"));
-	{
-		if (Column.CanBeSorted())
-		{
-			FUIAction Action_SortAscending
-			(
-				FExecuteAction::CreateSP(this, &SMemTagTreeView::HeaderMenu_SortMode_Execute, Column.GetId(), EColumnSortMode::Ascending),
-				FCanExecuteAction::CreateSP(this, &SMemTagTreeView::HeaderMenu_SortMode_CanExecute, Column.GetId(), EColumnSortMode::Ascending),
-				FIsActionChecked::CreateSP(this, &SMemTagTreeView::HeaderMenu_SortMode_IsChecked, Column.GetId(), EColumnSortMode::Ascending)
-			);
-			MenuBuilder.AddMenuEntry
-			(
-				FText::Format(LOCTEXT("ContextMenu_SortAscending_Fmt", "Sort Ascending (by {0})"), Column.GetTitleName()),
-				FText::Format(LOCTEXT("ContextMenu_SortAscending_Desc_Fmt", "Sorts ascending by {0}."), Column.GetTitleName()),
-				FSlateIcon(FAppStyle::Get().GetStyleSetName(), "Icons.SortUp"),
-				Action_SortAscending,
-				NAME_None,
-				EUserInterfaceActionType::RadioButton
-			);
-
-				FUIAction Action_SortDescending
-				(
-					FExecuteAction::CreateSP(this, &SMemTagTreeView::HeaderMenu_SortMode_Execute, Column.GetId(), EColumnSortMode::Descending),
-					FCanExecuteAction::CreateSP(this, &SMemTagTreeView::HeaderMenu_SortMode_CanExecute, Column.GetId(), EColumnSortMode::Descending),
-					FIsActionChecked::CreateSP(this, &SMemTagTreeView::HeaderMenu_SortMode_IsChecked, Column.GetId(), EColumnSortMode::Descending)
-				);
-				MenuBuilder.AddMenuEntry
-				(
-				FText::Format(LOCTEXT("ContextMenu_SortDescending_Fmt", "Sort Descending (by {0})"), Column.GetTitleName()),
-				FText::Format(LOCTEXT("ContextMenu_SortDescending_Desc_Fmt", "Sorts descending by {0}."), Column.GetTitleName()),
-					FSlateIcon(FAppStyle::Get().GetStyleSetName(), "Icons.SortDown"),
-					Action_SortDescending,
-					NAME_None,
-					EUserInterfaceActionType::RadioButton
-				);
-			}
-
-			MenuBuilder.AddSubMenu
-			(
-				LOCTEXT("ContextMenu_SortBy_SubMenu", "Sort By"),
-				LOCTEXT("ContextMenu_SortBy_SubMenu_Desc", "Sorts by a column."),
-				FNewMenuDelegate::CreateSP(this, &SMemTagTreeView::TreeView_BuildSortByMenu),
-				false,
-				FSlateIcon(FInsightsStyle::GetStyleSetName(), "Icons.SortBy")
-			);
-		}
-		MenuBuilder.EndSection();
-
-	MenuBuilder.BeginSection("ColumnVisibility", LOCTEXT("ContextMenu_Section_ColumnVisibility", "Column Visibility"));
-	{
-		if (Column.CanBeHidden())
-		{
-			FUIAction Action_HideColumn
-			(
-				FExecuteAction::CreateSP(this, &SMemTagTreeView::HideColumn, Column.GetId()),
-				FCanExecuteAction::CreateSP(this, &SMemTagTreeView::CanHideColumn, Column.GetId())
-			);
-			MenuBuilder.AddMenuEntry
-			(
-				LOCTEXT("ContextMenu_HideColumn", "Hide"),
-				LOCTEXT("ContextMenu_HideColumn_Desc", "Hides the selected column."),
-				FSlateIcon(),
-				Action_HideColumn,
-				NAME_None,
-				EUserInterfaceActionType::Button
-			);
-		}
-
-		MenuBuilder.AddSubMenu
-		(
-			LOCTEXT("ContextMenu_ViewColumn_SubMenu", "View Column"),
-			LOCTEXT("ContextMenu_ViewColumn_SubMenu_Desc", "Hides or shows columns."),
-			FNewMenuDelegate::CreateSP(this, &SMemTagTreeView::TreeView_BuildViewColumnMenu),
-			false,
-			FSlateIcon(FInsightsStyle::GetStyleSetName(), "Icons.ViewColumn")
-		);
-
-		FUIAction Action_ShowAllColumns
-		(
-			FExecuteAction::CreateSP(this, &SMemTagTreeView::ContextMenu_ShowAllColumns_Execute),
-			FCanExecuteAction::CreateSP(this, &SMemTagTreeView::ContextMenu_ShowAllColumns_CanExecute)
-		);
-		MenuBuilder.AddMenuEntry
-		(
-			LOCTEXT("ContextMenu_ShowAllColumns", "Show All Columns"),
-			LOCTEXT("ContextMenu_ShowAllColumns_Desc", "Resets tree view to show all columns."),
-			FSlateIcon(FInsightsStyle::GetStyleSetName(), "Icons.ResetColumn"),
-			Action_ShowAllColumns,
-			NAME_None,
-			EUserInterfaceActionType::Button
-		);
-
-		//FUIAction Action_ShowMinMaxMedColumns
-		//(
-		//	FExecuteAction::CreateSP(this, &SMemTagTreeView::ContextMenu_ShowMinMaxMedColumns_Execute),
-		//	FCanExecuteAction::CreateSP(this, &SMemTagTreeView::ContextMenu_ShowMinMaxMedColumns_CanExecute)
-		//);
-		//MenuBuilder.AddMenuEntry
-		//(
-		//	LOCTEXT("ContextMenu_ShowMinMaxMedColumns", "Reset Columns to Min/Max/Median Preset"),
-		//	LOCTEXT("ContextMenu_ShowMinMaxMedColumns_Desc", "Resets columns to Min/Max/Median preset."),
-		//	FSlateIcon(FInsightsStyle::GetStyleSetName(), "Icons.ResetColumn"),
-		//	Action_ShowMinMaxMedColumns,
-		//	NAME_None,
-		//	EUserInterfaceActionType::Button
-		//);
-
-		FUIAction Action_ResetColumns
-		(
-			FExecuteAction::CreateSP(this, &SMemTagTreeView::ContextMenu_ResetColumns_Execute),
-			FCanExecuteAction::CreateSP(this, &SMemTagTreeView::ContextMenu_ResetColumns_CanExecute)
-		);
-		MenuBuilder.AddMenuEntry
-		(
-			LOCTEXT("ContextMenu_ResetColumns", "Reset Columns to Default"),
-			LOCTEXT("ContextMenu_ResetColumns_Desc", "Resets columns to default."),
-			FSlateIcon(FInsightsStyle::GetStyleSetName(), "Icons.ResetColumn"),
-			Action_ResetColumns,
-			NAME_None,
-			EUserInterfaceActionType::Button
-		);
-	}
-	MenuBuilder.EndSection();
-
-	return MenuBuilder.MakeWidget();
-}
-
-////////////////////////////////////////////////////////////////////////////////////////////////////
-
-void SMemTagTreeView::InsightsManager_OnSessionChanged()
-{
-	TSharedPtr<const TraceServices::IAnalysisSession> NewSession = FInsightsManager::Get()->GetSession();
-	if (NewSession != Session)
-	{
-		Session = NewSession;
-		Reset();
-	}
-	else
-	{
-		UpdateTree();
-	}
-}
-
-////////////////////////////////////////////////////////////////////////////////////////////////////
-
-void SMemTagTreeView::UpdateTree()
-{
-	FStopwatch Stopwatch;
-	Stopwatch.Start();
-
-	CreateGroups();
-
-	Stopwatch.Update();
-	const double Time1 = Stopwatch.GetAccumulatedTime();
-
-	SortTreeNodes();
-
-	Stopwatch.Update();
-	const double Time2 = Stopwatch.GetAccumulatedTime();
-
-	ApplyFiltering();
-
-	Stopwatch.Stop();
-	const double TotalTime = Stopwatch.GetAccumulatedTime();
-	if (TotalTime > 0.1)
-	{
-		UE_LOG(LogMemoryProfiler, Log, TEXT("[LLM Tags] Tree view updated in %.3fs (%d counters) --> G:%.3fs + S:%.3fs + F:%.3fs"),
-			TotalTime, MemTagNodes.Num(), Time1, Time2 - Time1, TotalTime - Time2);
-	}
-}
-
-////////////////////////////////////////////////////////////////////////////////////////////////////
-
-void SMemTagTreeView::ApplyFiltering()
-{
-	FilteredGroupNodes.Reset();
-
-	// Apply filter to all groups and its children.
-	const int32 NumGroups = GroupNodes.Num();
-	for (int32 GroupIndex = 0; GroupIndex < NumGroups; ++GroupIndex)
-	{
-		FMemTagNodePtr& GroupPtr = GroupNodes[GroupIndex];
-		GroupPtr->ClearFilteredChildren();
-
-		const bool bIsGroupVisible = Filters->PassesAllFilters(GroupPtr);
-		int32 NumVisibleChildren = 0;
-
-		const TArray<FBaseTreeNodePtr>& GroupChildren = GroupPtr->GetChildren();
-		for (const FBaseTreeNodePtr& Child : GroupChildren)
-		{
-			const FMemTagNodePtr& NodePtr = StaticCastSharedPtr<FMemTagNode, FBaseTreeNode>(Child);
-			const bool bIsChildVisible = (!bFilterOutZeroCountMemTags || NodePtr->GetAggregatedStats().InstanceCount > 0)
-									  && (!FMemoryTracker::IsValidTrackerId(NodePtr->GetMemTrackerId()) || ((FMemoryTracker::AsFlag(NodePtr->GetMemTrackerId()) & TrackersFilter) != 0))
-									  && (NodePtr->GetMemTagSetId() == TagSetFilter)
-									  && Filters->PassesAllFilters(NodePtr);
-			if (bIsChildVisible)
-			{
-				// Add a child node.
-				GroupPtr->AddFilteredChild(NodePtr);
-				NumVisibleChildren++;
-			}
-		}
-
-		if (bIsGroupVisible || NumVisibleChildren > 0)
-		{
-			// Add a group.
-			FilteredGroupNodes.Add(GroupPtr);
-			GroupPtr->SetExpansion(true);
-		}
-		else
-		{
-			GroupPtr->SetExpansion(false);
-		}
-	}
-
-	// Only expand LLM tag nodes if we have a text filter.
-	const bool bNonEmptyTextFilter = !TextFilter->GetRawFilterText().IsEmpty();
-	if (bNonEmptyTextFilter)
-	{
-		if (!bExpansionSaved)
-		{
-			ExpandedNodes.Empty();
-			TreeView->GetExpandedItems(ExpandedNodes);
-			bExpansionSaved = true;
-		}
-
-		for (int32 Fx = 0; Fx < FilteredGroupNodes.Num(); Fx++)
-		{
-			const FMemTagNodePtr& GroupPtr = FilteredGroupNodes[Fx];
-			TreeView->SetItemExpansion(GroupPtr, GroupPtr->IsExpanded());
-		}
-	}
-	else
-	{
-		if (bExpansionSaved)
-		{
-			// Restore previously expanded nodes when the text filter is disabled.
-			TreeView->ClearExpandedItems();
-			for (auto It = ExpandedNodes.CreateConstIterator(); It; ++It)
-			{
-				TreeView->SetItemExpansion(*It, true);
-			}
-			bExpansionSaved = false;
-		}
-	}
-
-	// Request tree refresh
-	TreeView->RequestTreeRefresh();
-}
-
-////////////////////////////////////////////////////////////////////////////////////////////////////
-
-void SMemTagTreeView::HandleItemToStringArray(const FMemTagNodePtr& FMemTagNodePtr, TArray<FString>& OutSearchStrings) const
-{
-	OutSearchStrings.Add(FMemTagNodePtr->GetName().GetPlainNameString());
-}
-
-////////////////////////////////////////////////////////////////////////////////////////////////////
-
-void SMemTagTreeView::FilterOutZeroCountMemTags_OnCheckStateChanged(ECheckBoxState NewRadioState)
-{
-	bFilterOutZeroCountMemTags = (NewRadioState == ECheckBoxState::Checked);
-	ApplyFiltering();
-}
-
-////////////////////////////////////////////////////////////////////////////////////////////////////
-
-ECheckBoxState SMemTagTreeView::FilterOutZeroCountMemTags_IsChecked() const
-{
-	return bFilterOutZeroCountMemTags ? ECheckBoxState::Checked : ECheckBoxState::Unchecked;
-}
-
-////////////////////////////////////////////////////////////////////////////////////////////////////
-// TreeView
-////////////////////////////////////////////////////////////////////////////////////////////////////
-
-void SMemTagTreeView::TreeView_Refresh()
-{
-	if (TreeView.IsValid())
-	{
-		TreeView->RequestTreeRefresh();
-	}
-}
-
-////////////////////////////////////////////////////////////////////////////////////////////////////
-
-void SMemTagTreeView::TreeView_OnSelectionChanged(FMemTagNodePtr SelectedItem, ESelectInfo::Type SelectInfo)
-{
-	if (SelectInfo != ESelectInfo::Direct)
-	{
-		TArray<FMemTagNodePtr> SelectedItems = TreeView->GetSelectedItems();
-		if (SelectedItems.Num() == 1 && !SelectedItems[0]->IsGroup())
-		{
-			//TODO: FMemoryProfilerManager::Get()->SetSelectedMemTag(SelectedItems[0]->GetId());
-		}
-	}
-}
-
-////////////////////////////////////////////////////////////////////////////////////////////////////
-
-void SMemTagTreeView::TreeView_OnGetChildren(FMemTagNodePtr InParent, TArray<FMemTagNodePtr>& OutChildren)
-{
-	const TArray<FBaseTreeNodePtr>& Children = InParent->GetFilteredChildren();
-	OutChildren.Reset(Children.Num());
-	for (const FBaseTreeNodePtr& Child : Children)
-	{
-		OutChildren.Add(StaticCastSharedPtr<FMemTagNode, FBaseTreeNode>(Child));
-	}
-}
-
-////////////////////////////////////////////////////////////////////////////////////////////////////
-
-void SMemTagTreeView::TreeView_OnMouseButtonDoubleClick(FMemTagNodePtr MemTagNodePtr)
-{
-	if (MemTagNodePtr->IsGroup())
-	{
-		const bool bIsGroupExpanded = TreeView->IsItemExpanded(MemTagNodePtr);
-		TreeView->SetItemExpansion(MemTagNodePtr, !bIsGroupExpanded);
-	}
-	else
-	{
-		TSharedPtr<SMemoryProfilerWindow> ProfilerWindow = GetProfilerWindow();
-		if (ProfilerWindow.IsValid())
-		{
-			FMemorySharedState& SharedState = ProfilerWindow->GetSharedState();
-			const FMemoryTrackerId MemTrackerId = MemTagNodePtr->GetMemTrackerId();
-			const FMemoryTagId MemTagId = MemTagNodePtr->GetMemTagId();
-
-			TSharedPtr<FMemoryGraphTrack> GraphTrack = SharedState.GetMemTagGraphTrack(MemTrackerId, MemTagId);
-			if (!GraphTrack.IsValid())
-			{
-				SharedState.CreateMemTagGraphTrack(MemTrackerId, MemTagId);
-			}
-			else
-			{
-				SharedState.RemoveMemTagGraphTrack(MemTrackerId, MemTagId);
-			}
-		}
-	}
-}
-
-////////////////////////////////////////////////////////////////////////////////////////////////////
-// Tree View's Table Row
-////////////////////////////////////////////////////////////////////////////////////////////////////
-
-TSharedRef<ITableRow> SMemTagTreeView::TreeView_OnGenerateRow(FMemTagNodePtr MemTagNodePtr, const TSharedRef<STableViewBase>& OwnerTable)
-{
-	TSharedRef<ITableRow> TableRow =
-		SNew(SMemTagTreeViewTableRow, OwnerTable)
-		.OnShouldBeEnabled(this, &SMemTagTreeView::TableRow_ShouldBeEnabled)
-		.OnIsColumnVisible(this, &SMemTagTreeView::IsColumnVisible)
-		.OnSetHoveredCell(this, &SMemTagTreeView::TableRow_SetHoveredCell)
-		.OnGetColumnOutlineHAlignmentDelegate(this, &SMemTagTreeView::TableRow_GetColumnOutlineHAlignment)
-		.HighlightText(this, &SMemTagTreeView::TableRow_GetHighlightText)
-		.HighlightedNodeName(this, &SMemTagTreeView::TableRow_GetHighlightedNodeName)
-		.TablePtr(Table)
-		.MemTagNodePtr(MemTagNodePtr);
-
-	return TableRow;
-}
-
-////////////////////////////////////////////////////////////////////////////////////////////////////
-
-bool SMemTagTreeView::TableRow_ShouldBeEnabled(FMemTagNodePtr NodePtr) const
-{
-	return true;
-}
-
-////////////////////////////////////////////////////////////////////////////////////////////////////
-
-void SMemTagTreeView::TableRow_SetHoveredCell(TSharedPtr<FTable> InTablePtr, TSharedPtr<FTableColumn> InColumnPtr, FMemTagNodePtr InNodePtr)
-{
-	HoveredColumnId = InColumnPtr ? InColumnPtr->GetId() : FName();
-
-	const bool bIsAnyMenusVisible = FSlateApplication::Get().AnyMenusVisible();
-	if (!HasMouseCapture() && !bIsAnyMenusVisible)
-	{
-		HoveredNodePtr = InNodePtr;
-	}
-}
-
-////////////////////////////////////////////////////////////////////////////////////////////////////
-
-EHorizontalAlignment SMemTagTreeView::TableRow_GetColumnOutlineHAlignment(const FName ColumnId) const
-{
-	const TIndirectArray<SHeaderRow::FColumn>& Columns = TreeViewHeaderRow->GetColumns();
-	const int32 LastColumnIdx = Columns.Num() - 1;
-
-	// First column
-	if (Columns[0].ColumnId == ColumnId)
-	{
-		return HAlign_Left;
-	}
-	// Last column
-	else if (Columns[LastColumnIdx].ColumnId == ColumnId)
-	{
-		return HAlign_Right;
-	}
-	// Middle columns
-	{
-		return HAlign_Center;
-	}
-}
-
-////////////////////////////////////////////////////////////////////////////////////////////////////
-
-FText SMemTagTreeView::TableRow_GetHighlightText() const
-{
-	return SearchBox->GetText();
-}
-
-////////////////////////////////////////////////////////////////////////////////////////////////////
-
-FName SMemTagTreeView::TableRow_GetHighlightedNodeName() const
-{
-	return HighlightedNodeName;
-}
-
-////////////////////////////////////////////////////////////////////////////////////////////////////
-// SearchBox
-////////////////////////////////////////////////////////////////////////////////////////////////////
-
-void SMemTagTreeView::SearchBox_OnTextChanged(const FText& InFilterText)
-{
-	TextFilter->SetRawFilterText(InFilterText);
-	SearchBox->SetError(TextFilter->GetFilterErrorText());
-	ApplyFiltering();
-}
-
-////////////////////////////////////////////////////////////////////////////////////////////////////
-
-bool SMemTagTreeView::SearchBox_IsEnabled() const
-{
-	return MemTagNodes.Num() > 0;
-}
-
-////////////////////////////////////////////////////////////////////////////////////////////////////
-// Grouping
-////////////////////////////////////////////////////////////////////////////////////////////////////
-
-void SMemTagTreeView::CreateGroups()
-{
-	if (GroupingMode == EMemTagNodeGroupingMode::Flat)
-	{
-		GroupNodes.Reset();
-
-		const FName GroupName(TEXT("All"));
-		FMemTagNodePtr GroupPtr = MakeShared<FMemTagNode>(GroupName);
-		GroupNodes.Add(GroupPtr);
-
-		for (const FMemTagNodePtr& NodePtr : MemTagNodes)
-		{
-			GroupPtr->AddChildAndSetParent(NodePtr);
-		}
-		TreeView->SetItemExpansion(GroupPtr, true);
-	}
-	// Creates one group for each stat type.
-	else if (GroupingMode == EMemTagNodeGroupingMode::ByType)
-	{
-		TMap<EMemTagNodeType, FMemTagNodePtr> GroupNodeSet;
-		for (const FMemTagNodePtr& NodePtr : MemTagNodes)
-		{
-			const EMemTagNodeType NodeType = NodePtr->GetType();
-			FMemTagNodePtr GroupPtr = GroupNodeSet.FindRef(NodeType);
-			if (!GroupPtr)
-			{
-				const FName GroupName = *MemTagNodeTypeHelper::ToText(NodeType).ToString();
-				GroupPtr = GroupNodeSet.Add(NodeType, MakeShared<FMemTagNode>(GroupName));
-			}
-			GroupPtr->AddChildAndSetParent(NodePtr);
-			TreeView->SetItemExpansion(GroupPtr, true);
-		}
-		GroupNodeSet.KeySort([](const EMemTagNodeType& A, const EMemTagNodeType& B) { return A < B; }); // sort groups by type
-		GroupNodeSet.GenerateValueArray(GroupNodes);
-	}
-	// Creates one group for one letter.
-	else if (GroupingMode == EMemTagNodeGroupingMode::ByName)
-	{
-		TMap<TCHAR, FMemTagNodePtr> GroupNodeSet;
-		for (const FMemTagNodePtr& NodePtr : MemTagNodes)
-		{
-			FString FirstLetterStr(NodePtr->GetName().GetPlainNameString().Left(1).ToUpper());
-			const TCHAR FirstLetter = FirstLetterStr[0];
-			FMemTagNodePtr GroupPtr = GroupNodeSet.FindRef(FirstLetter);
-			if (!GroupPtr)
-			{
-				const FName GroupName(FirstLetterStr);
-				GroupPtr = GroupNodeSet.Add(FirstLetter, MakeShared<FMemTagNode>(GroupName));
-			}
-			GroupPtr->AddChildAndSetParent(NodePtr);
-		}
-		GroupNodeSet.KeySort([](const TCHAR& A, const TCHAR& B) { return A < B; }); // sort groups alphabetically
-		GroupNodeSet.GenerateValueArray(GroupNodes);
-	}
-	// Groups LLM tags by tracker.
-	else if (GroupingMode == EMemTagNodeGroupingMode::ByTracker)
-	{
-		TMap<FMemoryTrackerId, FMemTagNodePtr> GroupNodeSet;
-		for (const FMemTagNodePtr& NodePtr : MemTagNodes)
-		{
-			FMemoryTrackerId TrackerId = NodePtr->GetMemTrackerId();
-			FMemTagNodePtr GroupPtr = GroupNodeSet.FindRef(TrackerId);
-			if (!GroupPtr)
-			{
-				const FName GroupName = *NodePtr->GetTrackerText().ToString();
-				GroupPtr = GroupNodeSet.Add(TrackerId, MakeShared<FMemTagNode>(GroupName));
-			}
-			GroupPtr->AddChildAndSetParent(NodePtr);
-			TreeView->SetItemExpansion(GroupPtr, true);
-		}
-		GroupNodeSet.GenerateValueArray(GroupNodes);
-	}
-	// Groups LLM tags by their hierarchy.
-	else if (GroupingMode == EMemTagNodeGroupingMode::ByParent)
-	{
-		TMap<FName, FMemTagNodePtr> GroupNodeSet;
-		for (const FMemTagNodePtr& NodePtr : MemTagNodes)
-		{
-			const FName GroupName = NodePtr->GetParentTagNode() ? NodePtr->GetParentTagNode()->GetName() : FName(TEXT("<LLM>"));
-			FMemTagNodePtr GroupPtr = GroupNodeSet.FindRef(GroupName);
-			if (!GroupPtr)
-			{
-				GroupPtr = GroupNodeSet.Add(GroupName, MakeShared<FMemTagNode>(GroupName));
-			}
-			GroupPtr->AddChildAndSetParent(NodePtr);
-			TreeView->SetItemExpansion(GroupPtr, true);
-		}
-		GroupNodeSet.GenerateValueArray(GroupNodes);
-	}
-}
-
-////////////////////////////////////////////////////////////////////////////////////////////////////
-
-void SMemTagTreeView::CreateGroupByOptionsSources()
-{
-	GroupByOptionsSource.Reset(4);
-
-	// Must be added in order of elements in the EMemTagNodeGroupingMode.
-	GroupByOptionsSource.Add(MakeShared<EMemTagNodeGroupingMode>(EMemTagNodeGroupingMode::Flat));
-	GroupByOptionsSource.Add(MakeShared<EMemTagNodeGroupingMode>(EMemTagNodeGroupingMode::ByName));
-	//GroupByOptionsSource.Add(MakeShared<EMemTagNodeGroupingMode>(EMemTagNodeGroupingMode::ByType));
-	GroupByOptionsSource.Add(MakeShared<EMemTagNodeGroupingMode>(EMemTagNodeGroupingMode::ByTracker));
-	GroupByOptionsSource.Add(MakeShared<EMemTagNodeGroupingMode>(EMemTagNodeGroupingMode::ByParent));
-
-	EMemTagNodeGroupingModePtr* GroupingModePtrPtr = GroupByOptionsSource.FindByPredicate([&](const EMemTagNodeGroupingModePtr InGroupingModePtr) { return *InGroupingModePtr == GroupingMode; });
-	if (GroupingModePtrPtr != nullptr)
-	{
-		GroupByComboBox->SetSelectedItem(*GroupingModePtrPtr);
-	}
-
-	GroupByComboBox->RefreshOptions();
-}
-
-////////////////////////////////////////////////////////////////////////////////////////////////////
-
-void SMemTagTreeView::GroupBy_OnSelectionChanged(TSharedPtr<EMemTagNodeGroupingMode> NewGroupingMode, ESelectInfo::Type SelectInfo)
-{
-	if (SelectInfo != ESelectInfo::Direct)
-	{
-		GroupingMode = *NewGroupingMode;
-
-		CreateGroups();
-		SortTreeNodes();
-		ApplyFiltering();
-	}
-}
-
-////////////////////////////////////////////////////////////////////////////////////////////////////
-
-TSharedRef<SWidget> SMemTagTreeView::GroupBy_OnGenerateWidget(TSharedPtr<EMemTagNodeGroupingMode> InGroupingMode) const
-{
-	return SNew(STextBlock)
-		.Text(MemTagNodeGroupingHelper::ToText(*InGroupingMode))
-		.ToolTipText(MemTagNodeGroupingHelper::ToDescription(*InGroupingMode));
-}
-
-////////////////////////////////////////////////////////////////////////////////////////////////////
-
-FText SMemTagTreeView::GroupBy_GetSelectedText() const
-{
-	return MemTagNodeGroupingHelper::ToText(GroupingMode);
-}
-
-////////////////////////////////////////////////////////////////////////////////////////////////////
-
-FText SMemTagTreeView::GroupBy_GetSelectedTooltipText() const
-{
-	return MemTagNodeGroupingHelper::ToDescription(GroupingMode);
-}
-
-////////////////////////////////////////////////////////////////////////////////////////////////////
-// Sorting
-////////////////////////////////////////////////////////////////////////////////////////////////////
-
-const FName SMemTagTreeView::GetDefaultColumnBeingSorted()
-{
-	return FMemTagTreeViewColumns::TrackerColumnID;
-}
-
-////////////////////////////////////////////////////////////////////////////////////////////////////
-
-const EColumnSortMode::Type SMemTagTreeView::GetDefaultColumnSortMode()
-{
-	return EColumnSortMode::Ascending;
-}
-
-////////////////////////////////////////////////////////////////////////////////////////////////////
-
-void SMemTagTreeView::CreateSortings()
-{
-	AvailableSorters.Reset();
-	CurrentSorter = nullptr;
-
-	for (const TSharedRef<FTableColumn>& ColumnRef : Table->GetColumns())
-	{
-		if (ColumnRef->CanBeSorted())
-		{
-			TSharedPtr<ITableCellValueSorter> SorterPtr = ColumnRef->GetValueSorter();
-			if (ensure(SorterPtr.IsValid()))
-			{
-				AvailableSorters.Add(SorterPtr);
-			}
-		}
-	}
-
-	UpdateCurrentSortingByColumn();
-}
-
-////////////////////////////////////////////////////////////////////////////////////////////////////
-
-void SMemTagTreeView::UpdateCurrentSortingByColumn()
-{
-	TSharedPtr<FTableColumn> ColumnPtr = Table->FindColumn(ColumnBeingSorted);
-	CurrentSorter = ColumnPtr.IsValid() ? ColumnPtr->GetValueSorter() : nullptr;
-}
-
-////////////////////////////////////////////////////////////////////////////////////////////////////
-
-void SMemTagTreeView::SortTreeNodes()
-{
-	if (CurrentSorter.IsValid())
-	{
-		// Sort groups (always by name).
-		TArray<FBaseTreeNodePtr> SortedGroupNodes;
-		for (const FMemTagNodePtr& NodePtr : GroupNodes)
-		{
-			SortedGroupNodes.Add(NodePtr);
-		}
-		TSharedPtr<ITableCellValueSorter> Sorter = CurrentSorter;
-		ESortMode SortMode = (ColumnSortMode == EColumnSortMode::Type::Descending) ? ESortMode::Descending : ESortMode::Ascending;
-		if (CurrentSorter->GetName() != FName(TEXT("ByName")))
-		{
-			Sorter = MakeShared<FSorterByName>(Table->GetColumns()[0]);
-			SortMode = ESortMode::Ascending;
-		}
-		Sorter->Sort(SortedGroupNodes, SortMode);
-		GroupNodes.Reset();
-		for (const FBaseTreeNodePtr& NodePtr : SortedGroupNodes)
-		{
-			GroupNodes.Add(StaticCastSharedPtr<FMemTagNode>(NodePtr));
-		}
-
-		// Sort nodes in each group.
-		for (FMemTagNodePtr& Root : GroupNodes)
-		{
-			SortTreeNodesRec(*Root, *CurrentSorter);
-		}
-	}
-}
-
-////////////////////////////////////////////////////////////////////////////////////////////////////
-
-void SMemTagTreeView::SortTreeNodesRec(FMemTagNode& Node, const ITableCellValueSorter& Sorter)
-{
-	ESortMode SortMode = (ColumnSortMode == EColumnSortMode::Type::Descending) ? ESortMode::Descending : ESortMode::Ascending;
-	Node.SortChildren(Sorter, SortMode);
-
-	for (FBaseTreeNodePtr ChildPtr : Node.GetChildren())
-	{
-		//if (ChildPtr->IsGroup())
-		if (ChildPtr->GetChildrenCount() > 0)
-		{
-			SortTreeNodesRec(*StaticCastSharedPtr<FMemTagNode>(ChildPtr), Sorter);
-		}
-	}
-}
-
-////////////////////////////////////////////////////////////////////////////////////////////////////
-
-EColumnSortMode::Type SMemTagTreeView::GetSortModeForColumn(const FName ColumnId) const
-{
-	if (ColumnBeingSorted != ColumnId)
-	{
-		return EColumnSortMode::None;
-	}
-
-	return ColumnSortMode;
-}
-
-////////////////////////////////////////////////////////////////////////////////////////////////////
-
-void SMemTagTreeView::SetSortModeForColumn(const FName& ColumnId, const EColumnSortMode::Type SortMode)
-{
-	ColumnBeingSorted = ColumnId;
-	ColumnSortMode = SortMode;
-	UpdateCurrentSortingByColumn();
-
-	SortTreeNodes();
-	ApplyFiltering();
-}
-
-////////////////////////////////////////////////////////////////////////////////////////////////////
-
-void SMemTagTreeView::OnSortModeChanged(const EColumnSortPriority::Type SortPriority, const FName& ColumnId, const EColumnSortMode::Type SortMode)
-{
-	SetSortModeForColumn(ColumnId, SortMode);
-	TreeView_Refresh();
-}
-
-////////////////////////////////////////////////////////////////////////////////////////////////////
-// SortMode action (HeaderMenu)
-////////////////////////////////////////////////////////////////////////////////////////////////////
-
-bool SMemTagTreeView::HeaderMenu_SortMode_IsChecked(const FName ColumnId, const EColumnSortMode::Type InSortMode)
-{
-	return ColumnBeingSorted == ColumnId && ColumnSortMode == InSortMode;
-}
-
-////////////////////////////////////////////////////////////////////////////////////////////////////
-
-bool SMemTagTreeView::HeaderMenu_SortMode_CanExecute(const FName ColumnId, const EColumnSortMode::Type InSortMode) const
-{
-	const FTableColumn& Column = *Table->FindColumnChecked(ColumnId);
-	return Column.CanBeSorted();
-}
-
-////////////////////////////////////////////////////////////////////////////////////////////////////
-
-void SMemTagTreeView::HeaderMenu_SortMode_Execute(const FName ColumnId, const EColumnSortMode::Type InSortMode)
-{
-	SetSortModeForColumn(ColumnId, InSortMode);
-	TreeView_Refresh();
-}
-
-////////////////////////////////////////////////////////////////////////////////////////////////////
-// SortMode action (ContextMenu)
-////////////////////////////////////////////////////////////////////////////////////////////////////
-
-bool SMemTagTreeView::ContextMenu_SortMode_IsChecked(const EColumnSortMode::Type InSortMode)
-{
-	return ColumnSortMode == InSortMode;
-}
-
-////////////////////////////////////////////////////////////////////////////////////////////////////
-
-bool SMemTagTreeView::ContextMenu_SortMode_CanExecute(const EColumnSortMode::Type InSortMode) const
-{
-	return true; //ColumnSortMode != InSortMode;
-}
-
-////////////////////////////////////////////////////////////////////////////////////////////////////
-
-void SMemTagTreeView::ContextMenu_SortMode_Execute(const EColumnSortMode::Type InSortMode)
-{
-	SetSortModeForColumn(ColumnBeingSorted, InSortMode);
-	TreeView_Refresh();
-}
-
-////////////////////////////////////////////////////////////////////////////////////////////////////
-// SortByColumn action (ContextMenu)
-////////////////////////////////////////////////////////////////////////////////////////////////////
-
-bool SMemTagTreeView::ContextMenu_SortByColumn_IsChecked(const FName ColumnId)
-{
-	return ColumnId == ColumnBeingSorted;
-}
-
-////////////////////////////////////////////////////////////////////////////////////////////////////
-
-bool SMemTagTreeView::ContextMenu_SortByColumn_CanExecute(const FName ColumnId) const
-{
-	return true; //ColumnId != ColumnBeingSorted;
-}
-
-////////////////////////////////////////////////////////////////////////////////////////////////////
-
-void SMemTagTreeView::ContextMenu_SortByColumn_Execute(const FName ColumnId)
-{
-	SetSortModeForColumn(ColumnId, EColumnSortMode::Descending);
-	TreeView_Refresh();
-}
-
-////////////////////////////////////////////////////////////////////////////////////////////////////
-// ShowColumn action
-////////////////////////////////////////////////////////////////////////////////////////////////////
-
-bool SMemTagTreeView::CanShowColumn(const FName ColumnId) const
-{
-	return true;
-}
-
-////////////////////////////////////////////////////////////////////////////////////////////////////
-
-void SMemTagTreeView::ShowColumn(const FName ColumnId)
-{
-	FTableColumn& Column = *Table->FindColumnChecked(ColumnId);
-	Column.Show();
-
-	SHeaderRow::FColumn::FArguments ColumnArgs;
-	ColumnArgs
-		.ColumnId(Column.GetId())
-		.DefaultLabel(Column.GetShortName())
-		.ToolTip(SMemTagTreeViewTooltip::GetColumnTooltip(Column))
-		.HAlignHeader(Column.GetHorizontalAlignment())
-		.VAlignHeader(VAlign_Center)
-		.HAlignCell(HAlign_Fill)
-		.VAlignCell(VAlign_Fill)
-		.InitialSortMode(Column.GetInitialSortMode())
-		.SortMode(this, &SMemTagTreeView::GetSortModeForColumn, Column.GetId())
-		.OnSort(this, &SMemTagTreeView::OnSortModeChanged)
-		.FillWidth(Column.GetInitialWidth())
-		//.FixedWidth(Column.IsFixedWidth() ? Column.GetInitialWidth() : TOptional<float>())
-		.HeaderContent()
-		[
-			SNew(SBox)
-			.HeightOverride(24.0f)
-			.Padding(FMargin(0.0f))
-			.VAlign(VAlign_Center)
-			[
-				SNew(STextBlock)
-				.Text(this, &SMemTagTreeView::GetColumnHeaderText, Column.GetId())
-			]
-		]
-		.MenuContent()
-		[
-			TreeViewHeaderRow_GenerateColumnMenu(Column)
-		];
-
-	int32 ColumnIndex = 0;
-	const int32 NewColumnPosition = Table->GetColumnPositionIndex(ColumnId);
-	const int32 NumColumns = TreeViewHeaderRow->GetColumns().Num();
-	for (; ColumnIndex < NumColumns; ColumnIndex++)
+		const bool bIsGroupExpanded = TreeView->IsItemExpanded(NodePtr);
+		TreeView->SetItemExpansion(NodePtr, !bIsGroupExpanded);
+	}
+	else if (NodePtr->Is<FMemTagNode>())
 	{
-		const SHeaderRow::FColumn& CurrentColumn = TreeViewHeaderRow->GetColumns()[ColumnIndex];
-		const int32 CurrentColumnPosition = Table->GetColumnPositionIndex(CurrentColumn.ColumnId);
-		if (NewColumnPosition < CurrentColumnPosition)
+		FMemTagNode& MemTagNode = NodePtr->As<FMemTagNode>();
+		if (MemTagNode.IsValidMemTag())
 		{
-			break;
+			TSharedPtr<SMemoryProfilerWindow> ProfilerWindow = GetProfilerWindow();
+			if (ProfilerWindow.IsValid())
+			{
+				FMemorySharedState& SharedState = ProfilerWindow->GetSharedState();
+				const FMemoryTrackerId MemTrackerId = MemTagNode.GetMemTrackerId();
+				const FMemoryTagId MemTagId = MemTagNode.GetMemTagId();
+				TSharedPtr<FMemoryGraphTrack> GraphTrack = SharedState.GetMemTagGraphTrack(MemTrackerId, MemTagId);
+				if (!GraphTrack.IsValid())
+				{
+					SharedState.CreateMemTagGraphTrack(MemTrackerId, MemTagId);
+				}
+				else
+				{
+					SharedState.RemoveMemTagGraphTrack(MemTrackerId, MemTagId);
+				}
+			}
 		}
 	}
-
-	TreeViewHeaderRow->InsertColumn(ColumnArgs, ColumnIndex);
-}
-
-////////////////////////////////////////////////////////////////////////////////////////////////////
-// HideColumn action
-////////////////////////////////////////////////////////////////////////////////////////////////////
-
-bool SMemTagTreeView::CanHideColumn(const FName ColumnId) const
-{
-	const FTableColumn& Column = *Table->FindColumnChecked(ColumnId);
-	return Column.CanBeHidden();
-}
-
-////////////////////////////////////////////////////////////////////////////////////////////////////
-
-void SMemTagTreeView::HideColumn(const FName ColumnId)
-{
-	FTableColumn& Column = *Table->FindColumnChecked(ColumnId);
-	Column.Hide();
-
-	TreeViewHeaderRow->RemoveColumn(ColumnId);
-}
-
-////////////////////////////////////////////////////////////////////////////////////////////////////
-// ToggleColumn action
-////////////////////////////////////////////////////////////////////////////////////////////////////
-
-bool SMemTagTreeView::IsColumnVisible(const FName ColumnId)
-{
-	const FTableColumn& Column = *Table->FindColumnChecked(ColumnId);
-	return Column.IsVisible();
-}
-
-////////////////////////////////////////////////////////////////////////////////////////////////////
-
-bool SMemTagTreeView::CanToggleColumnVisibility(const FName ColumnId) const
-{
-	const FTableColumn& Column = *Table->FindColumnChecked(ColumnId);
-	return !Column.IsVisible() || Column.CanBeHidden();
-}
-
-////////////////////////////////////////////////////////////////////////////////////////////////////
-
-void SMemTagTreeView::ToggleColumnVisibility(const FName ColumnId)
-{
-	const FTableColumn& Column = *Table->FindColumnChecked(ColumnId);
-	if (Column.IsVisible())
-	{
-		HideColumn(ColumnId);
-	}
-	else
-	{
-		ShowColumn(ColumnId);
-	}
 }
 
 ////////////////////////////////////////////////////////////////////////////////////////////////////
-// "Show All Columns" action (ContextMenu)
+// Grouping
 ////////////////////////////////////////////////////////////////////////////////////////////////////
 
-bool SMemTagTreeView::ContextMenu_ShowAllColumns_CanExecute() const
+void SMemTagTreeView::InternalCreateGroupings()
 {
-	return true;
-}
-
-////////////////////////////////////////////////////////////////////////////////////////////////////
+	STableTreeView::InternalCreateGroupings();
 
-void SMemTagTreeView::ContextMenu_ShowAllColumns_Execute()
-{
-	ColumnBeingSorted = GetDefaultColumnBeingSorted();
-	ColumnSortMode = GetDefaultColumnSortMode();
-	UpdateCurrentSortingByColumn();
+	int32 Index = 1; // after the Flat ("All") grouping
 
-	for (const TSharedRef<FTableColumn>& ColumnRef : Table->GetColumns())
+#if 0
+	// Creates one group for each tag type.
+	class FMemTagGroupingByTagType : public
 	{
-		const FTableColumn& Column = *ColumnRef;
-
-		if (!Column.IsVisible())
 		{
-			ShowColumn(Column.GetId());
+			TMap<EMemTagNodeType, FMemTagNodePtr> GroupNodeSet;
+			for (const FMemTagNodePtr& NodePtr : MemTagNodes)
+			{
+				const EMemTagNodeType NodeType = NodePtr->GetType();
+				FMemTagNodePtr GroupPtr = GroupNodeSet.FindRef(NodeType);
+				if (!GroupPtr)
+				{
+					const FName GroupName = *MemTagNodeTypeHelper::ToText(NodeType).ToString();
+					GroupPtr = GroupNodeSet.Add(NodeType, MakeShared<FMemTagNode>(GroupName));
+				}
+				GroupPtr->AddChildAndSetParent(NodePtr);
+				TreeView->SetItemExpansion(GroupPtr, true);
+			}
+			GroupNodeSet.KeySort([](const EMemTagNodeType& A, const EMemTagNodeType& B) { return A < B; }); // sort groups by type
+			GroupNodeSet.GenerateValueArray(GroupNodes);
 		}
-	}
-}
-
-////////////////////////////////////////////////////////////////////////////////////////////////////
-// "Show Min/Max/Median Columns" action (ContextMenu)
-////////////////////////////////////////////////////////////////////////////////////////////////////
-
-bool SMemTagTreeView::ContextMenu_ShowMinMaxMedColumns_CanExecute() const
-{
-	return true;
-}
-
-////////////////////////////////////////////////////////////////////////////////////////////////////
+	};
+	AvailableGroupings.Insert(MakeShared<FMemTagGroupingByTagType>(), Index++);
+#endif
 
-void SMemTagTreeView::ContextMenu_ShowMinMaxMedColumns_Execute()
-{
-	TSet<FName> Preset =
+#if 0
+	// Creates one group for one letter.
+	class FMemTagGroupingByName : public
 	{
-		FMemTagTreeViewColumns::NameColumnID,
-		//FMemTagTreeViewColumns::MetaGroupNameColumnID,
-		//FMemTagTreeViewColumns::TypeColumnID,
-		FMemTagTreeViewColumns::TrackerColumnID,
-		FMemTagTreeViewColumns::InstanceCountColumnID,
-
-		FMemTagTreeViewColumns::MaxValueColumnID,
-		FMemTagTreeViewColumns::AverageValueColumnID,
-		//FMemTagTreeViewColumns::MedianValueColumnID,
-		FMemTagTreeViewColumns::MinValueColumnID,
+		{
+			TMap<TCHAR, FMemTagNodePtr> GroupNodeSet;
+			for (const FMemTagNodePtr& NodePtr : MemTagNodes)
+			{
+				FString FirstLetterStr(NodePtr->GetName().GetPlainNameString().Left(1).ToUpper());
+				const TCHAR FirstLetter = FirstLetterStr[0];
+				FMemTagNodePtr GroupPtr = GroupNodeSet.FindRef(FirstLetter);
+				if (!GroupPtr)
+				{
+					const FName GroupName(FirstLetterStr);
+					GroupPtr = GroupNodeSet.Add(FirstLetter, MakeShared<FMemTagNode>(GroupName));
+				}
+				GroupPtr->AddChildAndSetParent(NodePtr);
+			}
+			GroupNodeSet.KeySort([](const TCHAR& A, const TCHAR& B) { return A < B; }); // sort groups alphabetically
+			GroupNodeSet.GenerateValueArray(GroupNodes);
+		}
 	};
+	AvailableGroupings.Insert(MakeShared<FMemTagGroupingByName>(), Index++);
+#endif
 
-	ColumnBeingSorted = FMemTagTreeViewColumns::MaxValueColumnID;
-	ColumnSortMode = EColumnSortMode::Descending;
-	UpdateCurrentSortingByColumn();
-
-	for (const TSharedRef<FTableColumn>& ColumnRef : Table->GetColumns())
+#if 0
+	// Groups LLM tags by tracker.
+	class FMemTagGroupingByTracker : public
 	{
-		const FTableColumn& Column = *ColumnRef;
-
-		const bool bShouldBeVisible = Preset.Contains(Column.GetId());
-
-		if (bShouldBeVisible && !Column.IsVisible())
 		{
-			ShowColumn(Column.GetId());
+			TMap<FMemoryTrackerId, FMemTagNodePtr> GroupNodeSet;
+			for (const FMemTagNodePtr& NodePtr : MemTagNodes)
+			{
+				FMemoryTrackerId TrackerId = NodePtr->GetMemTrackerId();
+				FMemTagNodePtr GroupPtr = GroupNodeSet.FindRef(TrackerId);
+				if (!GroupPtr)
+				{
+					const FName GroupName = *NodePtr->GetTrackerText().ToString();
+					GroupPtr = GroupNodeSet.Add(TrackerId, MakeShared<FMemTagNode>(GroupName));
+				}
+				GroupPtr->AddChildAndSetParent(NodePtr);
+				TreeView->SetItemExpansion(GroupPtr, true);
+			}
+			GroupNodeSet.GenerateValueArray(GroupNodes);
 		}
-		else if (!bShouldBeVisible && Column.IsVisible())
+	};
+	AvailableGroupings.Insert(MakeShared<FMemTagGroupingByTracker>(), Index++);
+#endif
+
+#if 0
+	// Groups LLM tags by their hierarchy.
+	class FMemTagGroupingByHierarchy : public
+	{
 		{
-			HideColumn(Column.GetId());
+			TMap<FName, FMemTagNodePtr> GroupNodeSet;
+			for (const FMemTagNodePtr& NodePtr : MemTagNodes)
+			{
+				const FName GroupName = NodePtr->GetParentTagNode() ? NodePtr->GetParentTagNode()->GetName() : FName(TEXT("<LLM>"));
+				FMemTagNodePtr GroupPtr = GroupNodeSet.FindRef(GroupName);
+				if (!GroupPtr)
+				{
+					GroupPtr = GroupNodeSet.Add(GroupName, MakeShared<FMemTagNode>(GroupName));
+				}
+				GroupPtr->AddChildAndSetParent(NodePtr);
+				TreeView->SetItemExpansion(GroupPtr, true);
+			}
+			GroupNodeSet.GenerateValueArray(GroupNodes);
 		}
-	}
-}
-
-////////////////////////////////////////////////////////////////////////////////////////////////////
-// ResetColumns action (ContextMenu)
-////////////////////////////////////////////////////////////////////////////////////////////////////
-
-bool SMemTagTreeView::ContextMenu_ResetColumns_CanExecute() const
-{
-	return true;
+	};
+	AvailableGroupings.Insert(MakeShared<FMemTagGroupingByHierarchy>(), Index++);
+#endif
 }
 
 ////////////////////////////////////////////////////////////////////////////////////////////////////
 
-void SMemTagTreeView::ContextMenu_ResetColumns_Execute()
+void SMemTagTreeView::InitAvailableViewPresets()
 {
-	ColumnBeingSorted = GetDefaultColumnBeingSorted();
-	ColumnSortMode = GetDefaultColumnSortMode();
-	UpdateCurrentSortingByColumn();
-
-	for (const TSharedRef<FTableColumn>& ColumnRef : Table->GetColumns())
-	{
-		const FTableColumn& Column = *ColumnRef;
+	AvailableViewPresets.Add(FMemTagTableViewPresets::CreateDefaultViewPreset(*this));
+	AvailableViewPresets.Add(FMemTagTableViewPresets::CreateDiffViewPreset(*this));
+	AvailableViewPresets.Add(FMemTagTableViewPresets::CreateMinMaxViewPreset(*this));
 
-		if (Column.ShouldBeVisible() && !Column.IsVisible())
-		{
-			ShowColumn(Column.GetId());
-		}
-		else if (!Column.ShouldBeVisible() && Column.IsVisible())
-		{
-			HideColumn(Column.GetId());
-		}
-	}
+	SelectedViewPreset = AvailableViewPresets[0];
 }
 
 ////////////////////////////////////////////////////////////////////////////////////////////////////
@@ -1905,13 +587,15 @@ void SMemTagTreeView::Reset()
 	StatsStartTime = 0.0;
 	StatsEndTime = 0.0;
 
-	RebuildTree(true);
+	STableTreeView::Reset();
 }
 
 ////////////////////////////////////////////////////////////////////////////////////////////////////
 
 void SMemTagTreeView::Tick(const FGeometry& AllottedGeometry, const double InCurrentTime, const float InDeltaTime)
 {
+	STableTreeView::Tick(AllottedGeometry, InCurrentTime, InDeltaTime);
+
 	// We need to check if the list of LLM tags has changed.
 	// But, ensure we do not check too often.
 	static uint64 NextTimestamp = 0;
@@ -1920,10 +604,12 @@ void SMemTagTreeView::Tick(const FGeometry& AllottedGeometry, const double InCur
 	{
 		RebuildTree(false);
 
-		// 1000 counters --> check each 150ms
-		// 10000 counters --> check each 600ms
-		// 100000 counters --> check each 5.1s
-		const double WaitTimeSec = 0.1 + MemTagNodes.Num() / 20000.0;
+		UpdateStats();
+
+		// 1000 tags --> check each 120ms
+		// 10000 tags --> check each 300ms
+		// 100000 tags --> check each 2.1s
+		const double WaitTimeSec = 0.1 + FilteredNodesPtr->Num() / 50000.0;
 		const uint64 WaitTime = static_cast<uint64>(WaitTimeSec / FPlatformTime::GetSecondsPerCycle64());
 		NextTimestamp = Time + WaitTime;
 	}
@@ -1943,13 +629,13 @@ void SMemTagTreeView::RebuildTree(bool bResync)
 
 	if (bResync)
 	{
+		TableRowNodes.Empty();
 		LastMemoryTagListSerialNumber = 0;
-		MemTagNodes.Empty();
 		MemTagNodesIdMap.Empty();
 		bListHasChanged = true;
 	}
 
-	const uint32 PreviousNodeCount = MemTagNodes.Num();
+	const int32 PreviousNodeCount = TableRowNodes.Num();
 
 	TSharedPtr<SMemoryProfilerWindow> ProfilerWindow = GetProfilerWindow();
 	if (ProfilerWindow.IsValid())
@@ -1964,28 +650,32 @@ void SMemTagTreeView::RebuildTree(bool bResync)
 			const TArray<FMemoryTag*>& MemTags = TagList.GetTags();
 			const int32 MemTagCount = MemTags.Num();
 
-			MemTagNodes.Empty(MemTagCount);
+			TableRowNodes.Empty(MemTagCount);
 			MemTagNodesIdMap.Empty(MemTagCount);
 			bListHasChanged = true;
 
+			TSharedPtr<FMemTagTable> ParentMemTagTable = GetMemTagTable();
+
 			for (FMemoryTag* MemTagPtr : MemTags)
 			{
-				FMemTagNodePtr MemTagNodePtr = MakeShared<FMemTagNode>(MemTagPtr);
-				MemTagNodes.Add(MemTagNodePtr);
+				check(MemTagPtr);
+				FMemTagNodePtr MemTagNodePtr = MakeShared<FMemTagNode>(ParentMemTagTable, *MemTagPtr);
+				TableRowNodes.Add(MemTagNodePtr);
 				MemTagNodesIdMap.Add(MemTagPtr->GetId(), MemTagNodePtr);
 			}
 
 			// Resolve pointers to parent tags.
-			for (FMemTagNodePtr& NodePtr : MemTagNodes)
+			for (FTableTreeNodePtr& NodePtr : TableRowNodes)
 			{
-				check(NodePtr->GetMemTag() != nullptr);
-				FMemoryTag& MemTag = *NodePtr->GetMemTag();
-
+				check(NodePtr->Is<FMemTagNode>());
+				FMemTagNode& MemTagNode = NodePtr->As<FMemTagNode>();
+				check(MemTagNode.GetMemTag() != nullptr);
+				FMemoryTag& MemTag = *MemTagNode.GetMemTag();
 				FMemTagNodePtr ParentNodePtr = MemTagNodesIdMap.FindRef(MemTag.GetParentId());
 				if (ParentNodePtr)
 				{
 					check(ParentNodePtr != NodePtr);
-					NodePtr->SetParentTagNode(ParentNodePtr);
+					MemTagNode.SetParentTagNode(ParentNodePtr);
 				}
 			}
 		}
@@ -1996,141 +686,418 @@ void SMemTagTreeView::RebuildTree(bool bResync)
 			AvailableTagSets.Reset();
 			for (const auto& TagSet : TagSets)
 			{
-				AvailableTagSets.Add(MakeShared<FMemoryTagSetId>(TagSet->GetId()));
+				AvailableTagSets.Add(MakeShared<FMemoryTagSetId>(TagSet->GetId()));
+			}
+		}
+	}
+
+	SyncStopwatch.Stop();
+
+	if (bListHasChanged)
+	{
+		// Save selection.
+		TArray<FTableTreeNodePtr> SelectedItems;
+		TreeView->GetSelectedItems(SelectedItems);
+
+		UpdateTree();
+		UpdateStats();
+
+		TreeView->RebuildList();
+
+		// Restore selection.
+		if (SelectedItems.Num() > 0)
+		{
+			TreeView->ClearSelection();
+			for (FTableTreeNodePtr& NodePtr : SelectedItems)
+			{
+				if (!NodePtr->Is<FMemTagNode>())
+				{
+					NodePtr.Reset();
+					continue;
+				}
+				NodePtr = GetMemTagNode(NodePtr->As<FMemTagNode>().GetMemTagId());
+			}
+			SelectedItems.RemoveAll([](const FTableTreeNodePtr& NodePtr) { return !NodePtr.IsValid(); });
+			if (SelectedItems.Num() > 0)
+			{
+				TreeView->SetItemSelection(SelectedItems, true);
+				TreeView->RequestScrollIntoView(SelectedItems.Last());
+			}
+		}
+	}
+
+	Stopwatch.Stop();
+	const double TotalTime = Stopwatch.GetAccumulatedTime();
+	if (TotalTime > 0.01)
+	{
+		const double SyncTime = SyncStopwatch.GetAccumulatedTime();
+		UE_LOG(LogMemoryProfiler, Log, TEXT("[LLM Tags] Tree view rebuilt in %.4fs (sync: %.4fs + update: %.4fs) --> %d tags (%d added)"),
+			TotalTime, SyncTime, TotalTime - SyncTime, TableRowNodes.Num(), TableRowNodes.Num() - PreviousNodeCount);
+	}
+}
+
+////////////////////////////////////////////////////////////////////////////////////////////////////
+
+void SMemTagTreeView::UpdateStats()
+{
+	bool bShouldUpdateStats = false;
+
+	FStopwatch AggregationStopwatch;
+	FStopwatch Stopwatch;
+	Stopwatch.Start();
+
+	TSharedPtr<SMemoryProfilerWindow> ProfilerWindow = GetProfilerWindow();
+	if (ProfilerWindow.IsValid())
+	{
+		const TSharedRef<TimingProfiler::FTimeMarker>& MarkerA = ProfilerWindow->GetCustomTimeMarker(0);
+		if (TimeMarkerA != MarkerA->GetTime())
+		{
+			TimeMarkerA = MarkerA->GetTime();
+			bShouldUpdateStats = true;
+		}
+
+		const TSharedRef<TimingProfiler::FTimeMarker>& MarkerB = ProfilerWindow->GetCustomTimeMarker(1);
+		if (TimeMarkerB != MarkerB->GetTime())
+		{
+			TimeMarkerB = MarkerB->GetTime();
+			bShouldUpdateStats = true;
+		}
+
+		TSharedPtr<TimingProfiler::STimingView> TimingView = ProfilerWindow->GetTimingView();
+		if (TimingView)
+		{
+			double SelectionStartTime = TimingView->GetSelectionStartTime();
+			double SelectionEndTime = TimingView->GetSelectionEndTime();
+			if (SelectionStartTime < SelectionEndTime)
+			{
+				if (StatsStartTime != SelectionStartTime)
+				{
+					StatsStartTime = SelectionStartTime;
+					bShouldUpdateStats = true;
+				}
+				if (StatsEndTime != SelectionEndTime)
+				{
+					StatsEndTime = SelectionEndTime;
+					bShouldUpdateStats = true;
+				}
+			}
+		}
+
+		if (bShouldUpdateStats)
+		{
+			UpdateStatsInternal();
+
+			bool bShouldUpdateTreeWithStats = false; // TODO: UI toggle (advanced setting)
+			if (bShouldUpdateTreeWithStats)
+			{
+				UpdateTree();
+			}
+		}
+	}
+
+	Stopwatch.Stop();
+	const double TotalTime = Stopwatch.GetAccumulatedTime();
+	if (TotalTime > 0.01)
+	{
+		UE_LOG(LogMemoryProfiler, Log, TEXT("[LLM Tags] Aggregated stats updated in %.4fs"), TotalTime);
+	}
+}
+
+////////////////////////////////////////////////////////////////////////////////////////////////////
+
+void SMemTagTreeView::UpdateStatsInternal()
+{
+	if (!Session)
+	{
+		return;
+	}
+
+	TraceServices::FAnalysisSessionReadScope SessionReadScope(*Session.Get());
+	const TraceServices::IMemoryProvider* MemoryProvider = TraceServices::ReadMemoryProvider(*Session.Get());
+	if (MemoryProvider)
+	{
+		for (const FTableTreeNodePtr& NodePtr : TableRowNodes)
+		{
+			if (!NodePtr->Is<FMemTagNode>())
+			{
+				continue;
+			}
+			FMemTagNode& MemTagNode = NodePtr->As<FMemTagNode>();
+			MemTagNode.ResetAggregatedStats();
+
+			FMemoryTag* MemTag = MemTagNode.GetMemTag();
+			if (!MemTag)
+			{
+				continue;
+			}
+
+			FMemTagStats CurrentStats;
+
+			CurrentStats.SizeA = std::numeric_limits<int64>::max();
+			MemoryProvider->EnumerateTagSamples(MemTag->GetTrackerId(), MemTag->GetId(), TimeMarkerA, TimeMarkerA, true,
+				[&CurrentStats](double Time, double Duration, const TraceServices::FMemoryTagSample& Sample)
+				{
+					if (CurrentStats.SizeA == std::numeric_limits<int64>::max())
+					{
+						CurrentStats.SizeA = Sample.Value;
+					}
+				});
+			if (CurrentStats.SizeA != std::numeric_limits<int64>::max())
+			{
+				MemTagNode.GetStats().SizeA = CurrentStats.SizeA;
+			}
+
+			CurrentStats.SizeB = std::numeric_limits<int64>::max();
+			MemoryProvider->EnumerateTagSamples(MemTag->GetTrackerId(), MemTag->GetId(), TimeMarkerB, TimeMarkerB, true,
+				[&CurrentStats](double Time, double Duration, const TraceServices::FMemoryTagSample& Sample)
+				{
+					if (CurrentStats.SizeB == std::numeric_limits<int64>::max())
+					{
+						CurrentStats.SizeB = Sample.Value;
+					}
+				});
+			if (CurrentStats.SizeB != std::numeric_limits<int64>::max())
+			{
+				MemTagNode.GetStats().SizeB = CurrentStats.SizeB;
+			}
+
+			CurrentStats.SampleCount = 0;
+			CurrentStats.SizeMin = std::numeric_limits<int64>::max();
+			CurrentStats.SizeMax = std::numeric_limits<int64>::min();
+			MemoryProvider->EnumerateTagSamples(MemTag->GetTrackerId(), MemTag->GetId(), StatsStartTime, StatsEndTime, false,
+				[&CurrentStats](double Time, double Duration, const TraceServices::FMemoryTagSample& Sample)
+				{
+					CurrentStats.SampleCount++;
+					if (Sample.Value < CurrentStats.SizeMin)
+					{
+						CurrentStats.SizeMin = Sample.Value;
+					}
+					if (Sample.Value > CurrentStats.SizeMax)
+					{
+						CurrentStats.SizeMax = Sample.Value;
+					}
+					CurrentStats.SizeAverage += Sample.Value;
+				});
+			if (CurrentStats.SampleCount != 0)
+			{
+				MemTagNode.GetStats().SampleCount = CurrentStats.SampleCount;
+				MemTagNode.GetStats().SizeMin = CurrentStats.SizeMin;
+				MemTagNode.GetStats().SizeMax = CurrentStats.SizeMax;
+				MemTagNode.GetStats().SizeAverage = CurrentStats.SizeAverage / CurrentStats.SampleCount;
 			}
 		}
 	}
+}
 
-	SyncStopwatch.Stop();
+////////////////////////////////////////////////////////////////////////////////////////////////////
 
-	if (bListHasChanged)
+void SMemTagTreeView::ExtendMenu(TSharedRef<FExtender> Extender)
+{
+	Extender->AddMenuExtension("Misc", EExtensionHook::Before, nullptr, FMenuExtensionDelegate::CreateSP(this, &SMemTagTreeView::ExtendMenuBeforeMisc));
+	Extender->AddMenuExtension("Misc", EExtensionHook::After, nullptr, FMenuExtensionDelegate::CreateSP(this, &SMemTagTreeView::ExtendMenuAfterMisc));
+}
+
+////////////////////////////////////////////////////////////////////////////////////////////////////
+
+void SMemTagTreeView::ExtendMenuBeforeMisc(FMenuBuilder& MenuBuilder)
+{
+	MenuBuilder.BeginSection("Memory Tag", LOCTEXT("ContextMenu_Section_MemoryTag", "Memory Tag"));
 	{
-		UpdateTree();
-		UpdateStatsInternal();
+		MenuBuilder.AddSubMenu
+		(
+			LOCTEXT("ContextMenu_CreateGraphTracks_SubMenu", "Create Graph Tracks"),
+			LOCTEXT("ContextMenu_CreateGraphTracks_SubMenu_Desc", "Create memory graph tracks."),
+			FNewMenuDelegate::CreateSP(this, &SMemTagTreeView::ExtendMenuCreateGraphTracks),
+			false,
+			FSlateIcon()
+		);
 
-		// Save selection.
-		TArray<FMemTagNodePtr> SelectedItems;
-		TreeView->GetSelectedItems(SelectedItems);
+		MenuBuilder.AddSubMenu
+		(
+			LOCTEXT("ContextMenu_RemoveGraphTracks_SubMenu", "Remove Graph Tracks"),
+			LOCTEXT("ContextMenu_RemoveGraphTracks_SubMenu_Desc", "Remove memory graph tracks."),
+			FNewMenuDelegate::CreateSP(this, &SMemTagTreeView::ExtendMenuRemoveGraphTracks),
+			false,
+			FSlateIcon()
+		);
 
-		TreeView->RebuildList();
+		FUIAction Action_GenerateColorForSelectedMemTags
+		(
+			FExecuteAction::CreateSP(this, &SMemTagTreeView::GenerateColorForSelectedMemTags),
+			FCanExecuteAction::CreateSP(this, &SMemTagTreeView::CanGenerateColorForSelectedMemTags)
+		);
+		MenuBuilder.AddMenuEntry
+		(
+			LOCTEXT("ContextMenu_GenerateColorForSelectedMemTags", "Generate New Color"),
+			LOCTEXT("ContextMenu_GenerateColorForSelectedMemTags_Desc", "Generate new color for selected LLM tag(s)."),
+			FSlateIcon(),
+			Action_GenerateColorForSelectedMemTags,
+			NAME_None,
+			EUserInterfaceActionType::Button
+		);
 
-		// Restore selection.
-		if (SelectedItems.Num() > 0)
-		{
-			TreeView->ClearSelection();
-			for (FMemTagNodePtr& NodePtr : SelectedItems)
-			{
-				NodePtr = GetMemTagNode(NodePtr->GetMemTagId());
-			}
-			SelectedItems.RemoveAll([](const FMemTagNodePtr& NodePtr) { return !NodePtr.IsValid(); });
-			if (SelectedItems.Num() > 0)
-			{
-				TreeView->SetItemSelection(SelectedItems, true);
-				TreeView->RequestScrollIntoView(SelectedItems.Last());
-			}
-		}
+		FUIAction Action_EditColorForSelectedMemTags
+		(
+			FExecuteAction::CreateSP(this, &SMemTagTreeView::EditColorForSelectedMemTags),
+			FCanExecuteAction::CreateSP(this, &SMemTagTreeView::CanEditColorForSelectedMemTags)
+		);
+		MenuBuilder.AddMenuEntry
+		(
+			LOCTEXT("ContextMenu_EditColorForSelectedMemTags", "Edit Color..."),
+			LOCTEXT("ContextMenu_EditColorForSelectedMemTags_Desc", "Edit color for selected LLM tag(s)."),
+			FSlateIcon(),
+			Action_EditColorForSelectedMemTags,
+			NAME_None,
+			EUserInterfaceActionType::Button
+		);
 	}
+	MenuBuilder.EndSection();
+}
 
-	Stopwatch.Stop();
-	const double TotalTime = Stopwatch.GetAccumulatedTime();
-	if (TotalTime > 0.01)
+////////////////////////////////////////////////////////////////////////////////////////////////////
+
+void SMemTagTreeView::ExtendMenuCreateGraphTracks(FMenuBuilder& MenuBuilder)
+{
+	MenuBuilder.BeginSection("CreateGraphTracks");
 	{
-		const double SyncTime = SyncStopwatch.GetAccumulatedTime();
-		UE_LOG(LogMemoryProfiler, Log, TEXT("[LLM Tags] Tree view rebuilt in %.4fs (sync: %.4fs + update: %.4fs) --> %d LLM tags (%d added)"),
-			TotalTime, SyncTime, TotalTime - SyncTime, MemTagNodes.Num(), MemTagNodes.Num() - PreviousNodeCount);
+		// Create memory graph tracks for selected LLM tag(s)
+		FUIAction Action_CreateGraphTracksForSelectedMemTags
+		(
+			FExecuteAction::CreateSP(this, &SMemTagTreeView::CreateGraphTracksForSelectedMemTags),
+			FCanExecuteAction::CreateSP(this, &SMemTagTreeView::CanCreateGraphTracksForSelectedMemTags)
+		);
+		MenuBuilder.AddMenuEntry
+		(
+			LOCTEXT("ContextMenu_CreateGraphTracksForSelectedMemTags", "Selected"),
+			LOCTEXT("ContextMenu_CreateGraphTracksForSelectedMemTags_Desc", "Create memory graph tracks for selected LLM tag(s)."),
+			FSlateIcon(),
+			Action_CreateGraphTracksForSelectedMemTags,
+			NAME_None,
+			EUserInterfaceActionType::Button
+		);
+
+		// Create memory graph tracks for filtered LLM tags
+		FUIAction Action_CreateGraphTracksForFilteredMemTags
+		(
+			FExecuteAction::CreateSP(this, &SMemTagTreeView::CreateGraphTracksForFilteredMemTags),
+			FCanExecuteAction::CreateSP(this, &SMemTagTreeView::CanCreateGraphTracksForFilteredMemTags)
+		);
+		MenuBuilder.AddMenuEntry
+		(
+			LOCTEXT("ContextMenu_CreateGraphTracksForFilteredMemTags", "Filtered"),
+			LOCTEXT("ContextMenu_CreateGraphTracksForFilteredMemTags_Desc", "Create memory graph tracks for all visible (filtered) LLM tags."),
+			FSlateIcon(),
+			Action_CreateGraphTracksForFilteredMemTags,
+			NAME_None,
+			EUserInterfaceActionType::Button
+		);
+
+		// Create memory graph tracks for all LLM tags
+		FUIAction Action_CreateAllGraphTracks
+		(
+			FExecuteAction::CreateSP(this, &SMemTagTreeView::CreateAllGraphTracks),
+			FCanExecuteAction::CreateSP(this, &SMemTagTreeView::CanCreateAllGraphTracks)
+		);
+		MenuBuilder.AddMenuEntry
+		(
+			LOCTEXT("ContextMenu_CreateAllGraphTracks", "All"),
+			LOCTEXT("ContextMenu_CreateAllGraphTracks_Desc", "Create memory graph tracks for all LLM tags."),
+			FSlateIcon(),
+			Action_CreateAllGraphTracks,
+			NAME_None,
+			EUserInterfaceActionType::Button
+		);
 	}
+	MenuBuilder.EndSection();
 }
 
 ////////////////////////////////////////////////////////////////////////////////////////////////////
 
-void SMemTagTreeView::ResetStats()
+void SMemTagTreeView::ExtendMenuRemoveGraphTracks(FMenuBuilder& MenuBuilder)
 {
-	StatsStartTime = 0.0;
-	StatsEndTime = 0.0;
+	MenuBuilder.BeginSection("RemoveGraphTracks");
+	{
+		// Remove memory graph tracks for selected LLM tag(s)
+		FUIAction Action_RemoveGraphTracksForSelectedMemTags
+		(
+			FExecuteAction::CreateSP(this, &SMemTagTreeView::RemoveGraphTracksForSelectedMemTags),
+			FCanExecuteAction::CreateSP(this, &SMemTagTreeView::CanRemoveGraphTracksForSelectedMemTags)
+		);
+		MenuBuilder.AddMenuEntry
+		(
+			LOCTEXT("ContextMenu_RemoveGraphTracksForSelectedMemTags", "Selected"),
+			LOCTEXT("ContextMenu_RemoveGraphTracksForSelectedMemTags_Desc", "Remove memory graph tracks for selected LLM tag(s)."),
+			FSlateIcon(),
+			Action_RemoveGraphTracksForSelectedMemTags,
+			NAME_None,
+			EUserInterfaceActionType::Button
+		);
 
-	UpdateStatsInternal();
+		// Remove memory graph tracks for all LLM tags
+		FUIAction Action_RemoveAllGraphTracks
+		(
+			FExecuteAction::CreateSP(this, &SMemTagTreeView::RemoveAllGraphTracks),
+			FCanExecuteAction::CreateSP(this, &SMemTagTreeView::CanRemoveAllGraphTracks)
+		);
+		MenuBuilder.AddMenuEntry
+		(
+			LOCTEXT("ContextMenu_RemoveAllGraphTracks", "All"),
+			LOCTEXT("ContextMenu_RemoveAllGraphTracks_Desc", "Remove memory graph tracks for all LLM tags."),
+			FSlateIcon(),
+			Action_RemoveAllGraphTracks,
+			NAME_None,
+			EUserInterfaceActionType::Button
+		);
+	}
+	MenuBuilder.EndSection();
 }
 
 ////////////////////////////////////////////////////////////////////////////////////////////////////
 
-void SMemTagTreeView::UpdateStats(double StartTime, double EndTime)
+void SMemTagTreeView::ExtendMenuAfterMisc(FMenuBuilder& MenuBuilder)
 {
-	StatsStartTime = StartTime;
-	StatsEndTime = EndTime;
-
-	UpdateStatsInternal();
+	FUIAction Action_LoadReportXML
+	(
+		FExecuteAction::CreateSP(this, &SMemTagTreeView::LoadReportXML),
+		FCanExecuteAction::CreateSP(this, &SMemTagTreeView::CanLoadReportXML)
+	);
+	MenuBuilder.AddMenuEntry
+	(
+		LOCTEXT("ContextMenu_LoadReportXML", "Load Report XML..."),
+		LOCTEXT("ContextMenu_LoadReportXML_Desc", "Load LLMReportTypes.xml"),
+		FSlateIcon(FAppStyle::Get().GetStyleSetName(), "Icons.FolderOpen"),
+		Action_LoadReportXML,
+		NAME_None,
+		EUserInterfaceActionType::Button
+	);
 }
 
 ////////////////////////////////////////////////////////////////////////////////////////////////////
 
-void SMemTagTreeView::UpdateStatsInternal()
+bool SMemTagTreeView::HasCustomNodeFilter() const
 {
-	if (StatsStartTime >= StatsEndTime)
-	{
-		// keep previous aggregated stats
-		return;
-	}
-
-	FStopwatch AggregationStopwatch;
-	FStopwatch Stopwatch;
-	Stopwatch.Start();
+	return TrackersFilter != uint64(-1) ||
+		   TagSetFilter != FMemoryTagSet::InvalidTagSetId;
+}
 
-	for (const FMemTagNodePtr& NodePtr : MemTagNodes)
-	{
-		NodePtr->ResetAggregatedStats();
-	}
+////////////////////////////////////////////////////////////////////////////////////////////////////
 
-	/*
-	if (Session.IsValid())
+bool SMemTagTreeView::FilterNodeCustom(const FTableTreeNode& InNode) const
+{
+	if (InNode.Is<FMemTagNode>())
 	{
-		TUniquePtr<TraceServices::ITable<TraceServices::FMemoryProfilerAggregatedStats>> AggregationResultTable;
-
-		AggregationStopwatch.Start();
+		const FMemTagNode& MemTagNode = InNode.As<FMemTagNode>();
+		if (FMemoryTracker::IsValidTrackerId(MemTagNode.GetMemTrackerId()) &&
+			((FMemoryTracker::AsFlag(MemTagNode.GetMemTrackerId()) & TrackersFilter) == 0))
 		{
-			TraceServices::FAnalysisSessionReadScope SessionReadScope(*Session.Get());
-			const TraceServices::IMemoryProfilerProvider& MemoryProfilerProvider = TraceServices::ReadMemoryProfilerProvider(*Session.Get());
-			AggregationResultTable.Reset(MemoryProfilerProvider.CreateAggregation(StatsStartTime, StatsEndTime));
+			return false;
 		}
-		AggregationStopwatch.Stop();
-
-		if (AggregationResultTable.IsValid())
+		if (MemTagNode.GetMemTagSetId() != TagSetFilter)
 		{
-			TUniquePtr<TraceServices::ITableReader<TraceServices::FMemoryProfilerAggregatedStats>> TableReader(AggregationResultTable->CreateReader());
-			while (TableReader->IsValid())
-			{
-				const TraceServices::FMemoryProfilerAggregatedStats* Row = TableReader->GetCurrentRow();
-				FMemTagNodePtr* MemTagNodePtrPtr = MemTagNodesIdMap.Find(static_cast<uint64>(Row->EventTypeIndex));
-				if (MemTagNodePtrPtr != nullptr)
-				{
-					FMemTagNodePtr MemTagNodePtr = *MemTagNodePtrPtr;
-					->SetAggregatedStats(*Row);
-
-					TSharedPtr<ITableRow> TableRowPtr = TreeView->WidgetFromItem(MemTagNodePtr);
-					if (TableRowPtr.IsValid())
-					{
-						TSharedPtr<SMemTagsTableRow> RowPtr = StaticCastSharedPtr<SMemTagsTableRow, ITableRow>(TableRowPtr);
-						RowPtr->InvalidateContent();
-					}
-				}
-				TableReader->NextRow();
-			}
+			return false;
 		}
 	}
-	*/
-
-	UpdateTree();
-
-	const TArray<FMemTagNodePtr> SelectedNodes = TreeView->GetSelectedItems();
-	if (SelectedNodes.Num() > 0)
-	{
-		TreeView->RequestScrollIntoView(SelectedNodes[0]);
-	}
-
-	Stopwatch.Stop();
-	const double TotalTime = Stopwatch.GetAccumulatedTime();
-	const double AggregationTime = AggregationStopwatch.GetAccumulatedTime();
-	UE_LOG(LogMemoryProfiler, Log, TEXT("[LLM Tags] Aggregated stats updated in %.4fs (%.4fs + %.4fs)"),
-		TotalTime, AggregationTime, TotalTime - AggregationTime);
+	return true;
 }
 
 ////////////////////////////////////////////////////////////////////////////////////////////////////
@@ -2149,7 +1116,14 @@ void SMemTagTreeView::SelectMemTagNode(FMemoryTagId MemTagId)
 // Load Report XML button action
 ////////////////////////////////////////////////////////////////////////////////////////////////////
 
-FReply SMemTagTreeView::LoadReportXML_OnClicked()
+bool SMemTagTreeView::CanLoadReportXML() const
+{
+	return true;
+}
+
+////////////////////////////////////////////////////////////////////////////////////////////////////
+
+void SMemTagTreeView::LoadReportXML()
 {
 	TSharedPtr<SMemoryProfilerWindow> ProfilerWindow = GetProfilerWindow();
 	if (ProfilerWindow.IsValid())
@@ -2184,8 +1158,6 @@ FReply SMemTagTreeView::LoadReportXML_OnClicked()
 			SharedState.CreateTracksFromReport(OutFiles[0]);
 		}
 	}
-
-	return FReply::Handled();
 }
 
 ////////////////////////////////////////////////////////////////////////////////////////////////////
@@ -2206,6 +1178,38 @@ FReply SMemTagTreeView::HideAllTracks_OnClicked()
 	return FReply::Handled();
 }
 
+////////////////////////////////////////////////////////////////////////////////////////////////////
+
+void SMemTagTreeView::TryCreateGraphTrackForNode(FMemorySharedState& SharedState, const FBaseTreeNode& Node) const
+{
+	if (Node.Is<FMemTagNode>())
+	{
+		const FMemTagNode& MemTagNode = Node.As<FMemTagNode>();
+		if (MemTagNode.IsValidMemTag())
+		{
+			const FMemoryTrackerId MemTrackerId = MemTagNode.GetMemTrackerId();
+			const FMemoryTagId MemTagId = MemTagNode.GetMemTagId();
+			SharedState.CreateMemTagGraphTrack(MemTrackerId, MemTagId);
+		}
+	}
+}
+
+////////////////////////////////////////////////////////////////////////////////////////////////////
+
+void SMemTagTreeView::TryRemoveGraphTrackForNode(FMemorySharedState& SharedState, const FBaseTreeNode& Node) const
+{
+	if (Node.Is<FMemTagNode>())
+	{
+		const FMemTagNode& MemTagNode = Node.As<FMemTagNode>();
+		if (MemTagNode.IsValidMemTag())
+		{
+			const FMemoryTrackerId MemTrackerId = MemTagNode.GetMemTrackerId();
+			const FMemoryTagId MemTagId = MemTagNode.GetMemTagId();
+			SharedState.RemoveMemTagGraphTrack(MemTrackerId, MemTagId);
+		}
+	}
+}
+
 ////////////////////////////////////////////////////////////////////////////////////////////////////
 // Create memory graph tracks for selected LLM tag(s)
 ////////////////////////////////////////////////////////////////////////////////////////////////////
@@ -2224,26 +1228,20 @@ void SMemTagTreeView::CreateGraphTracksForSelectedMemTags()
 	if (ProfilerWindow.IsValid())
 	{
 		FMemorySharedState& SharedState = ProfilerWindow->GetSharedState();
-		const TArray<FMemTagNodePtr> SelectedNodes = TreeView->GetSelectedItems();
-		for (const FMemTagNodePtr& SelectedMemTagNode : SelectedNodes)
+
+		const TArray<FTableTreeNodePtr> SelectedNodes = TreeView->GetSelectedItems();
+		for (const FTableTreeNodePtr& SelectedNode : SelectedNodes)
 		{
-			if (SelectedMemTagNode->IsGroup())
+			TryCreateGraphTrackForNode(SharedState, *SelectedNode);
+
+			if (SelectedNode->IsGroup())
 			{
-				const TArray<FBaseTreeNodePtr>& Children = SelectedMemTagNode->GetFilteredChildren();
+				const TArray<FBaseTreeNodePtr>& Children = SelectedNode->GetFilteredChildren();
 				for (const FBaseTreeNodePtr& Child : Children)
 				{
-					FMemTagNodePtr MemTagNode = StaticCastSharedPtr<FMemTagNode, FBaseTreeNode>(Child);
-					const FMemoryTrackerId MemTrackerId = MemTagNode->GetMemTrackerId();
-					const FMemoryTagId MemTagId = MemTagNode->GetMemTagId();
-					SharedState.CreateMemTagGraphTrack(MemTrackerId, MemTagId);
+					TryCreateGraphTrackForNode(SharedState, *Child);
 				}
 			}
-			else
-			{
-				const FMemoryTrackerId MemTrackerId = SelectedMemTagNode->GetMemTrackerId();
-				const FMemoryTagId MemTagId = SelectedMemTagNode->GetMemTagId();
-				SharedState.CreateMemTagGraphTrack(MemTrackerId, MemTagId);
-			}
 		}
 	}
 }
@@ -2254,17 +1252,7 @@ void SMemTagTreeView::CreateGraphTracksForSelectedMemTags()
 
 bool SMemTagTreeView::CanCreateGraphTracksForFilteredMemTags() const
 {
-	TSharedPtr<SMemoryProfilerWindow> ProfilerWindow = GetProfilerWindow();
-	if (!ProfilerWindow.IsValid())
-	{
-		return false;
-	}
-	int32 FilteredNodeCount = 0;
-	for (const FMemTagNodePtr& Group : FilteredGroupNodes)
-	{
-		FilteredNodeCount += Group->GetFilteredChildren().Num();
-	}
-	return FilteredNodeCount > 0;
+	return FilteredNodesPtr->Num() > 0;
 }
 
 ////////////////////////////////////////////////////////////////////////////////////////////////////
@@ -2275,16 +1263,10 @@ void SMemTagTreeView::CreateGraphTracksForFilteredMemTags()
 	if (ProfilerWindow.IsValid())
 	{
 		FMemorySharedState& SharedState = ProfilerWindow->GetSharedState();
-		for (const FMemTagNodePtr& GroupNode : FilteredGroupNodes)
+
+		for (const FTableTreeNodePtr& FilteredNode : *FilteredNodesPtr)
 		{
-			const TArray<FBaseTreeNodePtr>& Children = GroupNode->GetFilteredChildren();
-			for (const FBaseTreeNodePtr& Child : Children)
-			{
-				FMemTagNodePtr MemTagNode = StaticCastSharedPtr<FMemTagNode, FBaseTreeNode>(Child);
-				const FMemoryTrackerId MemTrackerId = MemTagNode->GetMemTrackerId();
-				const FMemoryTagId MemTagId = MemTagNode->GetMemTagId();
-				SharedState.CreateMemTagGraphTrack(MemTrackerId, MemTagId);
-			}
+			TryCreateGraphTrackForNode(SharedState, *FilteredNode);
 		}
 	}
 }
@@ -2307,11 +1289,20 @@ void SMemTagTreeView::CreateAllGraphTracks()
 	if (ProfilerWindow.IsValid())
 	{
 		FMemorySharedState& SharedState = ProfilerWindow->GetSharedState();
-		for (const FMemTagNodePtr& MemTagNode : MemTagNodes)
+
+		for (const FTableTreeNodePtr& NodePtr : TableRowNodes)
 		{
-			const FMemoryTrackerId MemTrackerId = MemTagNode->GetMemTrackerId();
-			const FMemoryTagId MemTagId = MemTagNode->GetMemTagId();
-			SharedState.CreateMemTagGraphTrack(MemTrackerId, MemTagId);
+			// Only create graph tracks for tags of the selected tag set.
+			if (NodePtr->Is<FMemTagNode>())
+			{
+				const FMemTagNode& MemTagNode = NodePtr->As<FMemTagNode>();
+				if (MemTagNode.GetMemTagSetId() != TagSetFilter)
+				{
+					continue;
+				}
+			}
+
+			TryCreateGraphTrackForNode(SharedState, *NodePtr);
 		}
 	}
 }
@@ -2334,26 +1325,20 @@ void SMemTagTreeView::RemoveGraphTracksForSelectedMemTags()
 	if (ProfilerWindow.IsValid())
 	{
 		FMemorySharedState& SharedState = ProfilerWindow->GetSharedState();
-		const TArray<FMemTagNodePtr> SelectedNodes = TreeView->GetSelectedItems();
-		for (const FMemTagNodePtr& SelectedMemTagNode : SelectedNodes)
+
+		const TArray<FTableTreeNodePtr> SelectedNodes = TreeView->GetSelectedItems();
+		for (const FTableTreeNodePtr& SelectedNode : SelectedNodes)
 		{
-			if (SelectedMemTagNode->IsGroup())
+			TryRemoveGraphTrackForNode(SharedState, *SelectedNode);
+
+			if (SelectedNode->IsGroup())
 			{
-				const TArray<FBaseTreeNodePtr>& Children = SelectedMemTagNode->GetFilteredChildren();
+				const TArray<FBaseTreeNodePtr>& Children = SelectedNode->GetFilteredChildren();
 				for (const FBaseTreeNodePtr& Child : Children)
 				{
-					FMemTagNodePtr MemTagNode = StaticCastSharedPtr<FMemTagNode, FBaseTreeNode>(Child);
-					const FMemoryTrackerId MemTrackerId = MemTagNode->GetMemTrackerId();
-					const FMemoryTagId MemTagId = MemTagNode->GetMemTagId();
-					SharedState.RemoveMemTagGraphTrack(MemTrackerId, MemTagId);
+					TryRemoveGraphTrackForNode(SharedState, *Child);
 				}
 			}
-			else
-			{
-				const FMemoryTrackerId MemTrackerId = SelectedMemTagNode->GetMemTrackerId();
-				const FMemoryTagId MemTagId = SelectedMemTagNode->GetMemTagId();
-				SharedState.RemoveMemTagGraphTrack(MemTrackerId, MemTagId);
-			}
 		}
 	}
 }
@@ -2392,36 +1377,43 @@ bool SMemTagTreeView::CanGenerateColorForSelectedMemTags() const
 
 ////////////////////////////////////////////////////////////////////////////////////////////////////
 
-void SMemTagTreeView::GenerateColorForSelectedMemTags()
+void SMemTagTreeView::GenerateColorForSelectedMemTags() const
 {
 	TSharedPtr<SMemoryProfilerWindow> ProfilerWindow = GetProfilerWindow();
 	if (ProfilerWindow.IsValid())
 	{
-		const TArray<FMemTagNodePtr> SelectedNodes = TreeView->GetSelectedItems();
-		for (const FMemTagNodePtr& SelectedMemTagNode : SelectedNodes)
+		const TArray<FTableTreeNodePtr> SelectedNodes = TreeView->GetSelectedItems();
+		for (const FTableTreeNodePtr& SelectedNode : SelectedNodes)
 		{
-			constexpr bool bSetRandomColor = true;
-			SetColorToNode(SelectedMemTagNode, FLinearColor(), bSetRandomColor);
+			if (SelectedNode->Is<FMemTagNode>())
+			{
+				FMemTagNode& MemTagNode = SelectedNode->As<FMemTagNode>();
+				constexpr bool bSetRandomColor = true;
+				SetColorToNode(MemTagNode, FLinearColor(), bSetRandomColor);
+			}
 		}
 	}
 }
 
 ////////////////////////////////////////////////////////////////////////////////////////////////////
 
-void SMemTagTreeView::SetColorToNode(const FMemTagNodePtr& MemTagNode, FLinearColor Color, bool bSetRandomColor)
+void SMemTagTreeView::SetColorToNode(FMemTagNode& MemTagNode, FLinearColor Color, bool bSetRandomColor) const
 {
-	if (MemTagNode->IsGroup())
+	if (MemTagNode.IsGroup())
 	{
-		const TArray<FBaseTreeNodePtr>& Children = MemTagNode->GetFilteredChildren();
+		const TArray<FBaseTreeNodePtr>& Children = MemTagNode.GetFilteredChildren();
 		for (const FBaseTreeNodePtr& Child : Children)
 		{
-			const FMemTagNodePtr ChildMemTagNode = StaticCastSharedPtr<FMemTagNode, FBaseTreeNode>(Child);
-			SetColorToNode(ChildMemTagNode, Color, bSetRandomColor);
+			if (Child->Is<FMemTagNode>())
+			{
+				FMemTagNode& ChildMemTagNode = Child->As<FMemTagNode>();
+				SetColorToNode(ChildMemTagNode, Color, bSetRandomColor);
+			}
 		}
 		return;
 	}
 
-	FMemoryTag* MemTag = MemTagNode->GetMemTag();
+	FMemoryTag* MemTag = MemTagNode.GetMemTag();
 	if (!MemTag)
 	{
 		return;
@@ -2439,9 +1431,9 @@ void SMemTagTreeView::SetColorToNode(const FMemTagNodePtr& MemTagNode, FLinearCo
 
 	const FLinearColor BorderColor(FMath::Min(Color.R + 0.4f, 1.0f), FMath::Min(Color.G + 0.4f, 1.0f), FMath::Min(Color.B + 0.4f, 1.0f), 1.0f);
 
-	const FMemoryTrackerId MemTrackerId = MemTagNode->GetMemTrackerId();
-	const FMemoryTagSetId MemTagSetId = MemTagNode->GetMemTagSetId();
-	const FMemoryTagId MemTagId = MemTagNode->GetMemTagId();
+	const FMemoryTrackerId MemTrackerId = MemTagNode.GetMemTrackerId();
+	const FMemoryTagSetId MemTagSetId = MemTagNode.GetMemTagSetId();
+	const FMemoryTagId MemTagId = MemTagNode.GetMemTagId();
 
 	TSharedPtr<FMemoryGraphTrack> MainGraphTrack;
 	TSharedPtr<SMemoryProfilerWindow> ProfilerWindow = GetProfilerWindow();
@@ -2486,10 +1478,15 @@ void SMemTagTreeView::EditColorForSelectedMemTags()
 	if (ProfilerWindow.IsValid())
 	{
 		EditableColorValue = FLinearColor(0.5f, 0.5f, 0.5f, 1.0f);
-		const TArray<FMemTagNodePtr> SelectedNodes = TreeView->GetSelectedItems();
-		if (SelectedNodes.Num() > 0)
+		const TArray<FTableTreeNodePtr> SelectedNodes = TreeView->GetSelectedItems();
+		for (const FTableTreeNodePtr& SelectedNode : SelectedNodes)
 		{
-			EditableColorValue = SelectedNodes[0]->GetColor();
+			if (SelectedNode->Is<FMemTagNode>())
+			{
+				FMemTagNode& MemTagNode = SelectedNode->As<FMemTagNode>();
+				EditableColorValue = MemTagNode.GetColor();
+				break;
+			}
 		}
 
 		FColorPickerArgs PickerArgs;
@@ -2525,11 +1522,15 @@ void SMemTagTreeView::SetEditableColor(FLinearColor NewColor)
 {
 	EditableColorValue = NewColor;
 
-	const TArray<FMemTagNodePtr> SelectedNodes = TreeView->GetSelectedItems();
-	for (const FMemTagNodePtr& SelectedMemTagNode : SelectedNodes)
+	const TArray<FTableTreeNodePtr> SelectedNodes = TreeView->GetSelectedItems();
+	for (const FTableTreeNodePtr& SelectedNode : SelectedNodes)
 	{
-		constexpr bool bSetRandomColor = false;
-		SetColorToNode(SelectedMemTagNode, EditableColorValue, bSetRandomColor);
+		if (SelectedNode->Is<FMemTagNode>())
+		{
+			FMemTagNode& MemTagNode = SelectedNode->As<FMemTagNode>();
+			constexpr bool bSetRandomColor = false;
+			SetColorToNode(MemTagNode, EditableColorValue, bSetRandomColor);
+		}
 	}
 }
 
diff --git a/Engine/Source/Developer/TraceInsights/Private/Insights/MemoryProfiler/Widgets/SMemTagTreeView.h b/Engine/Source/Developer/TraceInsights/Private/Insights/MemoryProfiler/Widgets/SMemTagTreeView.h
--- a/Engine/Source/Developer/TraceInsights/Private/Insights/MemoryProfiler/Widgets/SMemTagTreeView.h
+++ b/Engine/Source/Developer/TraceInsights/Private/Insights/MemoryProfiler/Widgets/SMemTagTreeView.h
@@ -2,58 +2,37 @@
 
 #pragma once
 
-#include "CoreMinimal.h"
-
-#include "Misc/FilterCollection.h"
-#include "Misc/TextFilter.h"
-#include "SlateFwd.h"
-#include "Widgets/DeclarativeSyntaxSupport.h"
-#include "Widgets/Input/SComboBox.h"
-#include "Widgets/SCompoundWidget.h"
-#include "Widgets/SWidget.h"
-#include "Widgets/Views/SHeaderRow.h"
-#include "Widgets/Views/STableRow.h"
-#include "Widgets/Views/STableViewBase.h"
-#include "Widgets/Views/STreeView.h"
+#include "CoreTypes.h"
 
 // TraceInsights
-#include "Insights/MemoryProfiler/ViewModels/MemTagNodeGroupingAndSorting.h"
-#include "Insights/MemoryProfiler/ViewModels/MemTagNode.h"
-
-class FMenuBuilder;
-
-namespace TraceServices
-{
-	class IAnalysisSession;
-}
-
-namespace UE::Insights
-{
-	class FTable;
-	class FTableColumn;
-	class ITableCellValueSorter;
-}
+#include "Insights/Table/Widgets/SSessionTableTreeView.h"
+#include "Insights/MemoryProfiler/ViewModels/MemoryTag.h"
+#include "Insights/MemoryProfiler/ViewModels/MemTagTable.h"
 
 namespace UE::Insights::MemoryProfiler
 {
 
+class FMemTagNode;
 class FMemoryGraphTrack;
+class FMemorySharedState;
 class SMemoryProfilerWindow;
 
-////////////////////////////////////////////////////////////////////////////////////////////////////
-
-/** The filter collection - used for updating the list of tree nodes. */
-typedef TFilterCollection<const FMemTagNodePtr&> FMemTagNodeFilterCollection;
-
-/** The text based filter - used for updating the list of tree nodes. */
-typedef TTextFilter<const FMemTagNodePtr&> FMemTagNodeTextFilter;
-
 ////////////////////////////////////////////////////////////////////////////////////////////////////
 /**
  * A custom widget used to display the list of LLM tags and their aggregated stats.
  */
-class SMemTagTreeView : public SCompoundWidget
+class SMemTagTreeView : public SSessionTableTreeView
 {
+public:
+	enum class EInvestigationMode
+	{
+		Default,
+		Diff,
+		MinMax,
+
+		Count
+	};
+
 public:
 	/** Default constructor. */
 	SMemTagTreeView();
@@ -79,198 +58,107 @@ public:
 	 */
 	void Construct(const FArguments& InArgs, TSharedPtr<SMemoryProfilerWindow> InProfilerWindow);
 
-	TSharedPtr<FTable> GetTable() const { return Table; }
-
-	void Reset();
+	virtual TSharedPtr<SWidget> ConstructFilterToolbar() override;
+	virtual TSharedPtr<SWidget> ConstructToolbar() override;
+	virtual TSharedPtr<SWidget> ConstructExtraToolbar() override;
+	virtual TSharedPtr<SWidget> ConstructFooter() override;
 
-	/**
-	 * Rebuilds the tree (if necessary).
-	 * @param bResync - If true, it forces a resync with list of LLM tags from Analysis, even if the list did not changed since last sync.
-	 */
-	void RebuildTree(bool bResync);
-
-	void ResetStats();
-	void UpdateStats(double StartTime, double EndTime);
-
-	FMemTagNodePtr GetMemTagNode(FMemoryTagId MemTagId) const { return MemTagNodesIdMap.FindRef(MemTagId); }
-	void SelectMemTagNode(FMemoryTagId MemTagId);
-
-private:
-	TSharedRef<SWidget> MakeTrackersMenu();
-	void CreateTrackersMenuSection(FMenuBuilder& MenuBuilder);
-	TSharedRef<SWidget> ConstructTagsFilteringWidgetArea();
-	TSharedRef<SWidget> ConstructTagsGroupingWidgetArea();
-	TSharedRef<SWidget> ConstructTracksMiniToolbar();
-
-	void UpdateTree();
+	TSharedPtr<FMemTagTable> GetMemTagTable()
+	{
+		return StaticCastSharedPtr<FMemTagTable>(GetTable());
+	}
 
-	void UpdateStatsInternal();
+	TSharedPtr<const FMemTagTable> GetMemTagTable() const
+	{
+		return StaticCastSharedPtr<const FMemTagTable>(GetTable());
+	}
 
-	/** Called when the analysis session has changed. */
-	void InsightsManager_OnSessionChanged();
+	virtual void Reset();
 
 	/**
-	 * Populates OutSearchStrings with the strings that should be used in searching.
-	 *
-	 * @param GroupOrStatNodePtr - the group and stat node to get a text description from.
-	 * @param OutSearchStrings   - an array of strings to use in searching.
+	 * Ticks this widget.  Override in derived classes, but always call the parent implementation.
 	 *
+	 * @param  AllottedGeometry The space allotted for this widget
+	 * @param  InCurrentTime  Current absolute real time
+	 * @param  InDeltaTime  Real time passed since last tick
 	 */
-	void HandleItemToStringArray(const FMemTagNodePtr& GroupOrStatNodePtr, TArray<FString>& OutSearchStrings) const;
+	virtual void Tick(const FGeometry& AllottedGeometry, const double InCurrentTime, const float InDeltaTime) override;
 
-	////////////////////////////////////////////////////////////////////////////////////////////////////
-	// Tree View - Context Menu
+	/**
+	 * Rebuilds the tree (if necessary).
+	 * @param bResync - If true, it forces a resync with list of LLM tags from Analysis, even if the list did not changed since last sync.
+	 */
+	virtual void RebuildTree(bool bResync);
 
-	TSharedPtr<SWidget> TreeView_GetMenuContent();
-	void TreeView_BuildCreateGraphTracksMenu(FMenuBuilder& MenuBuilder);
-	void TreeView_BuildRemoveGraphTracksMenu(FMenuBuilder& MenuBuilder);
-	void TreeView_BuildSortByMenu(FMenuBuilder& MenuBuilder);
-	void TreeView_BuildViewColumnMenu(FMenuBuilder& MenuBuilder);
+	void UpdateStats();
 
 	////////////////////////////////////////////////////////////////////////////////////////////////////
-	// Tree View - Columns' Header
-
-	void InitializeAndShowHeaderColumns();
 
-	FText GetColumnHeaderText(const FName ColumnId) const;
-
-	TSharedRef<SWidget> TreeViewHeaderRow_GenerateColumnMenu(const FTableColumn& Column);
+	TSharedPtr<FMemTagNode> GetMemTagNode(FMemoryTagId MemTagId) const
+	{
+		return MemTagNodesIdMap.FindRef(MemTagId);
+	}
+	void SelectMemTagNode(FMemoryTagId MemTagId);
 
 	////////////////////////////////////////////////////////////////////////////////////////////////////
-	// Tree View - Misc
 
-	void TreeView_Refresh();
+protected:
+	virtual void InternalCreateGroupings() override;
 
-	/**
-	 * Called by STreeView to retrieves the children for the specified parent item.
-	 * @param InParent    - The parent node to retrieve the children from.
-	 * @param OutChildren - List of children for the parent node.
-	 */
-	void TreeView_OnGetChildren(FMemTagNodePtr InParent, TArray<FMemTagNodePtr>& OutChildren);
-
-	/** Called by STreeView when selection has changed. */
-	void TreeView_OnSelectionChanged(FMemTagNodePtr SelectedItem, ESelectInfo::Type SelectInfo);
+	virtual void ExtendMenu(TSharedRef<FExtender> Extender) override;
 
-	/** Called by STreeView when a tree item is double clicked. */
-	void TreeView_OnMouseButtonDoubleClick(FMemTagNodePtr TreeNode);
+	virtual bool HasCustomNodeFilter() const override;
+	virtual bool FilterNodeCustom(const FTableTreeNode& InNode) const override;
 
-	////////////////////////////////////////////////////////////////////////////////////////////////////
-	// Tree View - Table Row
+	virtual void TreeView_OnMouseButtonDoubleClick(FTableTreeNodePtr NodePtr) override;
+	virtual void TreeView_OnSelectionChanged(FTableTreeNodePtr SelectedItem, ESelectInfo::Type SelectInfo) override;
 
-	/** Called by STreeView to generate a table row for the specified item. */
-	TSharedRef<ITableRow> TreeView_OnGenerateRow(FMemTagNodePtr TreeNode, const TSharedRef<STableViewBase>& OwnerTable);
+private:
+	TSharedRef<SWidget> ConstructFilteringWidgetArea();
+	TSharedRef<SWidget> ConstructTracksMiniToolbar();
 
-	bool TableRow_ShouldBeEnabled(FMemTagNodePtr NodePtr) const;
+	void UpdateSelectionStatsText();
+	FText GetFooterCenterText() const;
 
-	void TableRow_SetHoveredCell(TSharedPtr<FTable> TablePtr, TSharedPtr<FTableColumn> ColumnPtr, FMemTagNodePtr NodePtr);
-	EHorizontalAlignment TableRow_GetColumnOutlineHAlignment(const FName ColumnId) const;
+	void InitAvailableViewPresets();
 
-	FText TableRow_GetHighlightText() const;
-	FName TableRow_GetHighlightedNodeName() const;
+	void UpdateStatsInternal();
 
 	////////////////////////////////////////////////////////////////////////////////////////////////////
-	// Filtering
+	// Context Menu
 
-	/** Populates the group and stat tree with items based on the current data. */
-	void ApplyFiltering();
+	void ExtendMenuBeforeMisc(FMenuBuilder& MenuBuilder);
+	void ExtendMenuCreateGraphTracks(FMenuBuilder& MenuBuilder);
+	void ExtendMenuRemoveGraphTracks(FMenuBuilder& MenuBuilder);
+	void ExtendMenuAfterMisc(FMenuBuilder& MenuBuilder);
 
-	void FilterOutZeroCountMemTags_OnCheckStateChanged(ECheckBoxState NewRadioState);
-	ECheckBoxState FilterOutZeroCountMemTags_IsChecked() const;
+	////////////////////////////////////////////////////////////////////////////////////////////////////
+	// Trackers' Filter
 
-	bool SearchBox_IsEnabled() const;
-	void SearchBox_OnTextChanged(const FText& InFilterText);
+	TSharedRef<SWidget> MakeTrackersMenu();
+	void CreateTrackersMenuSection(FMenuBuilder& MenuBuilder);
 
 	void ToggleTracker(FMemoryTrackerId InTrackerId);
 	bool IsTrackerChecked(FMemoryTrackerId InTrackerId) const;
 
+	////////////////////////////////////////////////////////////////////////////////////////////////////
+	// Tag Set Selection
+
 	void TagSet_OnSelectionChanged(TSharedPtr<FMemoryTagSetId> NewTagSet, ESelectInfo::Type SelectInfo);
 	TSharedRef<SWidget> TagSet_OnGenerateWidget(TSharedPtr<FMemoryTagSetId> InTagSet) const;
 	FText TagSet_GetSelectedText() const;
 
 	////////////////////////////////////////////////////////////////////////////////////////////////////
-	// Grouping
-
-	void CreateGroups();
-	void CreateGroupByOptionsSources();
-
-	void GroupBy_OnSelectionChanged(TSharedPtr<EMemTagNodeGroupingMode> NewGroupingMode, ESelectInfo::Type SelectInfo);
-
-	TSharedRef<SWidget> GroupBy_OnGenerateWidget(TSharedPtr<EMemTagNodeGroupingMode> InGroupingMode) const;
-
-	FText GroupBy_GetSelectedText() const;
-
-	FText GroupBy_GetSelectedTooltipText() const;
-
-	////////////////////////////////////////////////////////////////////////////////////////////////////
-	// Sorting
-
-	static const FName GetDefaultColumnBeingSorted();
-	static const EColumnSortMode::Type GetDefaultColumnSortMode();
-
-	void CreateSortings();
-
-	void UpdateCurrentSortingByColumn();
-	void SortTreeNodes();
-	void SortTreeNodesRec(FMemTagNode& Node, const ITableCellValueSorter& Sorter);
+	// Misc
 
-	EColumnSortMode::Type GetSortModeForColumn(const FName ColumnId) const;
-	void SetSortModeForColumn(const FName& ColumnId, EColumnSortMode::Type SortMode);
-	void OnSortModeChanged(const EColumnSortPriority::Type SortPriority, const FName& ColumnId, const EColumnSortMode::Type SortMode);
-
-	////////////////////////////////////////////////////////////////////////////////////////////////////
-	// Sorting actions
-
-	// SortMode (HeaderMenu)
-	bool HeaderMenu_SortMode_IsChecked(const FName ColumnId, const EColumnSortMode::Type InSortMode);
-	bool HeaderMenu_SortMode_CanExecute(const FName ColumnId, const EColumnSortMode::Type InSortMode) const;
-	void HeaderMenu_SortMode_Execute(const FName ColumnId, const EColumnSortMode::Type InSortMode);
-
-	// SortMode (ContextMenu)
-	bool ContextMenu_SortMode_IsChecked(const EColumnSortMode::Type InSortMode);
-	bool ContextMenu_SortMode_CanExecute(const EColumnSortMode::Type InSortMode) const;
-	void ContextMenu_SortMode_Execute(const EColumnSortMode::Type InSortMode);
-
-	// SortByColumn (ContextMenu)
-	bool ContextMenu_SortByColumn_IsChecked(const FName ColumnId);
-	bool ContextMenu_SortByColumn_CanExecute(const FName ColumnId) const;
-	void ContextMenu_SortByColumn_Execute(const FName ColumnId);
-
-	////////////////////////////////////////////////////////////////////////////////////////////////////
-	// Column visibility actions
-
-	// ShowColumn
-	bool CanShowColumn(const FName ColumnId) const;
-	void ShowColumn(const FName ColumnId);
-
-	// HideColumn
-	bool CanHideColumn(const FName ColumnId) const;
-	void HideColumn(const FName ColumnId);
-
-	// ToggleColumnVisibility
-	bool IsColumnVisible(const FName ColumnId);
-	bool CanToggleColumnVisibility(const FName ColumnId) const;
-	void ToggleColumnVisibility(const FName ColumnId);
-
-	// ShowAllColumns (ContextMenu)
-	bool ContextMenu_ShowAllColumns_CanExecute() const;
-	void ContextMenu_ShowAllColumns_Execute();
-
-	// MinMaxMedColumns (ContextMenu)
-	bool ContextMenu_ShowMinMaxMedColumns_CanExecute() const;
-	void ContextMenu_ShowMinMaxMedColumns_Execute();
-
-	// ResetColumns (ContextMenu)
-	bool ContextMenu_ResetColumns_CanExecute() const;
-	void ContextMenu_ResetColumns_Execute();
-
-	////////////////////////////////////////////////////////////////////////////////////////////////////
+	bool CanLoadReportXML() const;
+	void LoadReportXML();
 
 	FReply ShowAllTracks_OnClicked();
 	FReply HideAllTracks_OnClicked();
 
-	FReply LoadReportXML_OnClicked();
-
-	////////////////////////////////////////////////////////////////////////////////////////////////////
+	void TryCreateGraphTrackForNode(FMemorySharedState& SharedState, const FBaseTreeNode& Node) const;
+	void TryRemoveGraphTrackForNode(FMemorySharedState& SharedState, const FBaseTreeNode& Node) const;
 
 	// Create memory graph tracks for selected LLM tag(s)
 	bool CanCreateGraphTracksForSelectedMemTags() const;
@@ -294,149 +182,59 @@ private:
 
 	// Generate new color for selected LLM tag(s)
 	bool CanGenerateColorForSelectedMemTags() const;
-	void GenerateColorForSelectedMemTags();
-	void SetColorToNode(const FMemTagNodePtr& MemTagNode, FLinearColor CustomColor, bool bSetRandomColor);
-	FLinearColor GetEditableColor() const;
-	void SetEditableColor(FLinearColor NewColor);
-	void ColorPickerCancelled(FLinearColor OriginalColor);
+	void GenerateColorForSelectedMemTags() const;
 
 	// Edit color for selected LLM tag(s)
 	bool CanEditColorForSelectedMemTags() const;
 	void EditColorForSelectedMemTags();
 
-	////////////////////////////////////////////////////////////////////////////////////////////////////
-
-	/**
-	 * Ticks this widget.  Override in derived classes, but always call the parent implementation.
-	 *
-	 * @param  AllottedGeometry The space allotted for this widget
-	 * @param  InCurrentTime  Current absolute real time
-	 * @param  InDeltaTime  Real time passed since last tick
-	 */
-	virtual void Tick(const FGeometry& AllottedGeometry, const double InCurrentTime, const float InDeltaTime) override;
+	void SetColorToNode(FMemTagNode& MemTagNode, FLinearColor CustomColor, bool bSetRandomColor) const;
+	FLinearColor GetEditableColor() const;
+	void SetEditableColor(FLinearColor NewColor);
+	void ColorPickerCancelled(FLinearColor OriginalColor);
 
 private:
 	/** A weak pointer to the Memory Insights window. */
 	TWeakPtr<SMemoryProfilerWindow> ProfilerWindowWeakPtr;
 
-	/** Table view model. */
-	TSharedPtr<FTable> Table;
-
-	/** The analysis session used to populate this widget. */
-	TSharedPtr<const TraceServices::IAnalysisSession> Session;
-
 	//////////////////////////////////////////////////
-	// Tree View, Columns
-
-	/** The tree widget which holds the LLM tag tree nodes. */
-	TSharedPtr<STreeView<FMemTagNodePtr>> TreeView;
-
-	/** Holds the tree view header row widget which display all columns in the tree view. */
-	TSharedPtr<SHeaderRow> TreeViewHeaderRow;
-
-	/** External scrollbar used to synchronize tree view position. */
-	TSharedPtr<SScrollBar> ExternalScrollbar;
-
-	//////////////////////////////////////////////////
-	// Hovered Column, Hovered Node
-
-	/** Name of the column currently being hovered by the mouse. */
-	FName HoveredColumnId;
-
-	/** A shared pointer to the node currently being hovered by the mouse. */
-	FMemTagNodePtr HoveredNodePtr;
-
-	/** Name of the node that should be drawn as highlighted. */
-	FName HighlightedNodeName;
-
-	//////////////////////////////////////////////////
-	// MemTag Nodes
-
-	/** An array of group nodes. */
-	TArray<FMemTagNodePtr> GroupNodes;
-
-	/** A filtered array of group nodes to be displayed in the tree widget. */
-	TArray<FMemTagNodePtr> FilteredGroupNodes;
+	// Tree Nodes
 
 	/** The serial number of the memory tag list maintained by the MemorySharedState object (updated last time we have synced MemTagNodes with it). */
-	uint32 LastMemoryTagListSerialNumber;
-
-	/** All LLM tag nodes. */
-	TSet<FMemTagNodePtr> MemTagNodes;
+	uint32 LastMemoryTagListSerialNumber = 0;
 
-	/** All LLM tag nodes, stored as NodeId -> FMemTagNodePtr. */
-	TMap<FMemoryTagId, FMemTagNodePtr> MemTagNodesIdMap;
-
-	/** Currently expanded group nodes. */
-	TSet<FMemTagNodePtr> ExpandedNodes;
-
-	/** If true, the expanded nodes have been saved before applying a text filter. */
-	bool bExpansionSaved;
+	/** All LLM tag nodes, stored as FMemoryTagId -> TSharedPtr<FMemTagNode>. */
+	TMap<FMemoryTagId, TSharedPtr<FMemTagNode>> MemTagNodesIdMap;
 
 	//////////////////////////////////////////////////
-	// Search box and filters
-
-	//bool bUseFiltering;
-
-	/** The search box widget used to filter items displayed in the stats and groups tree. */
-	TSharedPtr<SSearchBox> SearchBox;
+	// Filters
 
-	/** The text based filter. */
-	TSharedPtr<FMemTagNodeTextFilter> TextFilter;
-
-	/** The filter collection. */
-	TSharedPtr<FMemTagNodeFilterCollection> Filters;
-
-	/** Filter out the LLM tags having zero total instance count (aggregated stats). */
-	bool bFilterOutZeroCountMemTags;
+	TSharedPtr<SComboBox<TSharedPtr<FMemoryTracker>>> TrackerComboBox;
 
 	/** Filter the LLM tags by tracker. */
-	uint64 TrackersFilter;
+	uint64 TrackersFilter = uint64(-1);
 
 	/** Filter the LLM tags by tag set. */
-	FMemoryTagSetId TagSetFilter;
+	FMemoryTagSetId TagSetFilter = FMemoryTagSet::DefaultTagSetId;
 	TArray<TSharedPtr<FMemoryTagSetId>> AvailableTagSets;
-	TSharedPtr<SComboBox<TSharedPtr<FMemoryTagSetId>>> TagSetComboBox;
-
-	//////////////////////////////////////////////////
-	// Grouping
-
-	//bool bUseGrouping;
-
-	TArray<TSharedPtr<EMemTagNodeGroupingMode>> GroupByOptionsSource;
-
-	TSharedPtr<SComboBox<TSharedPtr<EMemTagNodeGroupingMode>>> GroupByComboBox;
-
-	/** How we group the LLM tag nodes? */
-	EMemTagNodeGroupingMode GroupingMode;
 
 	//////////////////////////////////////////////////
-	// Sorting
 
-	//bool bUseSorting;
+	FText SelectionStatsText;
 
-	/** All available sorters. */
-	TArray<TSharedPtr<ITableCellValueSorter>> AvailableSorters;
-
-	/** Current sorter. It is nullptr if sorting is disabled. */
-	TSharedPtr<ITableCellValueSorter> CurrentSorter;
-
-	/** Name of the column currently being sorted. Can be NAME_None if sorting is disabled (CurrentSorting == nullptr) or if a complex sorting is used (CurrentSorting != nullptr). */
-	FName ColumnBeingSorted;
-
-	/** How we sort the nodes? Ascending or Descending. */
-	EColumnSortMode::Type ColumnSortMode;
+	FLinearColor EditableColorValue;
 
 	//////////////////////////////////////////////////
 
-	TSharedPtr<SComboBox<TSharedPtr<FMemoryTracker>>> TrackerComboBox;
-
-	FLinearColor EditableColorValue;
+	EInvestigationMode Mode = EInvestigationMode::Default;
 
-	//////////////////////////////////////////////////
+	// A, B, B-A
+	double TimeMarkerA = 0.0;
+	double TimeMarkerB = 0.0;
 
-	double StatsStartTime;
-	double StatsEndTime;
+	// Sample Count, Min, Max, Average
+	double StatsStartTime = 0.0;
+	double StatsEndTime = 0.0;
 };
 
 ////////////////////////////////////////////////////////////////////////////////////////////////////
diff --git a/Engine/Source/Developer/TraceInsights/Private/Insights/MemoryProfiler/Widgets/SMemTagTreeViewTableCell.cpp b/Engine/Source/Developer/TraceInsights/Private/Insights/MemoryProfiler/Widgets/SMemTagTreeViewTableCell.cpp
deleted file mode 100644
diff --git a/Engine/Source/Developer/TraceInsights/Private/Insights/MemoryProfiler/Widgets/SMemTagTreeViewTableCell.h b/Engine/Source/Developer/TraceInsights/Private/Insights/MemoryProfiler/Widgets/SMemTagTreeViewTableCell.h
deleted file mode 100644
diff --git a/Engine/Source/Developer/TraceInsights/Private/Insights/MemoryProfiler/Widgets/SMemTagTreeViewTableRow.cpp b/Engine/Source/Developer/TraceInsights/Private/Insights/MemoryProfiler/Widgets/SMemTagTreeViewTableRow.cpp
deleted file mode 100644
diff --git a/Engine/Source/Developer/TraceInsights/Private/Insights/MemoryProfiler/Widgets/SMemTagTreeViewTableRow.h b/Engine/Source/Developer/TraceInsights/Private/Insights/MemoryProfiler/Widgets/SMemTagTreeViewTableRow.h
deleted file mode 100644
diff --git a/Engine/Source/Developer/TraceInsights/Private/Insights/MemoryProfiler/Widgets/SMemTagTreeViewTooltip.cpp b/Engine/Source/Developer/TraceInsights/Private/Insights/MemoryProfiler/Widgets/SMemTagTreeViewTooltip.cpp
deleted file mode 100644
diff --git a/Engine/Source/Developer/TraceInsights/Private/Insights/MemoryProfiler/Widgets/SMemTagTreeViewTooltip.h b/Engine/Source/Developer/TraceInsights/Private/Insights/MemoryProfiler/Widgets/SMemTagTreeViewTooltip.h
deleted file mode 100644
diff --git a/Engine/Source/Developer/TraceInsights/Private/Insights/MemoryProfiler/Widgets/SMemoryProfilerWindow.cpp b/Engine/Source/Developer/TraceInsights/Private/Insights/MemoryProfiler/Widgets/SMemoryProfilerWindow.cpp
--- a/Engine/Source/Developer/TraceInsights/Private/Insights/MemoryProfiler/Widgets/SMemoryProfilerWindow.cpp
+++ b/Engine/Source/Developer/TraceInsights/Private/Insights/MemoryProfiler/Widgets/SMemoryProfilerWindow.cpp
@@ -12,6 +12,7 @@
 // TraceInsights
 #include "Insights/InsightsStyle.h"
 #include "Insights/MemoryProfiler/MemoryProfilerManager.h"
+#include "Insights/MemoryProfiler/ViewModels/MemAllocTable.h"
 #include "Insights/MemoryProfiler/ViewModels/MemorySharedState.h"
 #include "Insights/MemoryProfiler/Widgets/SMemAllocTableTreeView.h"
 #include "Insights/MemoryProfiler/Widgets/SMemInvestigationView.h"
