[Insights]
- LLM Tags: Improved appearance of status bar text for selected tags to use different colors for text and values.
- LLM Tags: Changed the table to use tag name in hierarchy column. The Tag Name column is now hidden by default.

#rnx

diff --git a/Engine/Source/Developer/TraceInsights/Private/Insights/MemoryProfiler/ViewModels/MemTagNode.cpp b/Engine/Source/Developer/TraceInsights/Private/Insights/MemoryProfiler/ViewModels/MemTagNode.cpp
--- a/Engine/Source/Developer/TraceInsights/Private/Insights/MemoryProfiler/ViewModels/MemTagNode.cpp
+++ b/Engine/Source/Developer/TraceInsights/Private/Insights/MemoryProfiler/ViewModels/MemTagNode.cpp
@@ -115,7 +115,7 @@ const TCHAR* FMemTagNode::GetTagSetName() const
 
 const FText FMemTagNode::GetDisplayName() const
 {
-	return FText::FromString(TEXT("tag"));
+	return GetTagText();
 }
 
 ////////////////////////////////////////////////////////////////////////////////////////////////////
diff --git a/Engine/Source/Developer/TraceInsights/Private/Insights/MemoryProfiler/ViewModels/MemTagTable.cpp b/Engine/Source/Developer/TraceInsights/Private/Insights/MemoryProfiler/ViewModels/MemTagTable.cpp
--- a/Engine/Source/Developer/TraceInsights/Private/Insights/MemoryProfiler/ViewModels/MemTagTable.cpp
+++ b/Engine/Source/Developer/TraceInsights/Private/Insights/MemoryProfiler/ViewModels/MemTagTable.cpp
@@ -96,7 +96,7 @@ void FMemTagTable::AddDefaultColumns()
 		AddHierarchyColumn(HierarchyColumnIndex, HierarchyColumnName);
 
 		FTableColumn& HierarchyColumn = *GetColumns()[0];
-		HierarchyColumn.SetInitialWidth(60.0f);
+		HierarchyColumn.SetInitialWidth(200.0f);
 		HierarchyColumn.SetShortName(LOCTEXT("HierarchyColumnName", "Hierarchy"));
 		HierarchyColumn.SetTitleName(LOCTEXT("HierarchyColumnTitle", "LLM Tag Hierarchy"));
 		HierarchyColumn.SetDescription(LOCTEXT("HierarchyColumnDesc", "Hierarchy of the LLM tag's tree"));
@@ -116,8 +116,7 @@ void FMemTagTable::AddDefaultColumns()
 		Column.SetTitleName(LOCTEXT("TagNameColumnTitle", "Tag Name"));
 		Column.SetDescription(LOCTEXT("TagNameColumnDesc", "The name of the LLM tag"));
 
-		Column.SetFlags(ETableColumnFlags::ShouldBeVisible |
-						ETableColumnFlags::CanBeHidden |
+		Column.SetFlags(ETableColumnFlags::CanBeHidden |
 						ETableColumnFlags::CanBeFiltered);
 
 		Column.SetHorizontalAlignment(HAlign_Left);
diff --git a/Engine/Source/Developer/TraceInsights/Private/Insights/MemoryProfiler/ViewModels/MemTagTableViewPresets.cpp b/Engine/Source/Developer/TraceInsights/Private/Insights/MemoryProfiler/ViewModels/MemTagTableViewPresets.cpp
--- a/Engine/Source/Developer/TraceInsights/Private/Insights/MemoryProfiler/ViewModels/MemTagTableViewPresets.cpp
+++ b/Engine/Source/Developer/TraceInsights/Private/Insights/MemoryProfiler/ViewModels/MemTagTableViewPresets.cpp
@@ -48,8 +48,7 @@ TSharedRef<ITableTreeViewPreset> FMemTagTableViewPresets::CreateDefaultViewPrese
 		}
 		virtual void GetColumnConfigSet(TArray<FTableColumnConfig>& InOutConfigSet) const override
 		{
-			InOutConfigSet.Add({ FTable::GetHierarchyColumnId(),          true, 120.0f });
-			InOutConfigSet.Add({ FMemTagTableColumns::TagNameColumnId,    true, 400.0f });
+			InOutConfigSet.Add({ FTable::GetHierarchyColumnId(),          true, 400.0f });
 			InOutConfigSet.Add({ FMemTagTableColumns::SizeAColumnId,      true, 100.0f });
 			InOutConfigSet.Add({ FMemTagTableColumns::SizeBudgetColumnId, true, 100.0f });
 		}
@@ -90,8 +89,7 @@ TSharedRef<ITableTreeViewPreset> FMemTagTableViewPresets::CreateDiffViewPreset(S
 		}
 		virtual void GetColumnConfigSet(TArray<FTableColumnConfig>& InOutConfigSet) const override
 		{
-			InOutConfigSet.Add({ FTable::GetHierarchyColumnId(),          true, 120.0f });
-			InOutConfigSet.Add({ FMemTagTableColumns::TagNameColumnId,    true, 400.0f });
+			InOutConfigSet.Add({ FTable::GetHierarchyColumnId(),          true, 400.0f });
 			InOutConfigSet.Add({ FMemTagTableColumns::SizeAColumnId,      true, 80.0f });
 			InOutConfigSet.Add({ FMemTagTableColumns::SizeBColumnId,      true, 80.0f });
 			InOutConfigSet.Add({ FMemTagTableColumns::SizeDiffColumnId,   true, 80.0f });
@@ -134,8 +132,7 @@ TSharedRef<ITableTreeViewPreset> FMemTagTableViewPresets::CreateMinMaxViewPreset
 		}
 		virtual void GetColumnConfigSet(TArray<FTableColumnConfig>& InOutConfigSet) const override
 		{
-			InOutConfigSet.Add({ FTable::GetHierarchyColumnId(),           true, 120.0f });
-			InOutConfigSet.Add({ FMemTagTableColumns::TagNameColumnId,     true, 400.0f });
+			InOutConfigSet.Add({ FTable::GetHierarchyColumnId(),           true, 400.0f });
 			InOutConfigSet.Add({ FMemTagTableColumns::SizeMinColumnId,     true, 80.0f });
 			InOutConfigSet.Add({ FMemTagTableColumns::SizeMaxColumnId,     true, 80.0f });
 			InOutConfigSet.Add({ FMemTagTableColumns::SizeAverageColumnId, true, 80.0f });
diff --git a/Engine/Source/Developer/TraceInsights/Private/Insights/MemoryProfiler/Widgets/SMemTagTreeView.cpp b/Engine/Source/Developer/TraceInsights/Private/Insights/MemoryProfiler/Widgets/SMemTagTreeView.cpp
--- a/Engine/Source/Developer/TraceInsights/Private/Insights/MemoryProfiler/Widgets/SMemTagTreeView.cpp
+++ b/Engine/Source/Developer/TraceInsights/Private/Insights/MemoryProfiler/Widgets/SMemTagTreeView.cpp
@@ -5,22 +5,16 @@
 #include "DesktopPlatformModule.h"
 #include "Framework/MultiBox/MultiBoxBuilder.h"
 #include "SlateOptMacros.h"
-#include "Templates/UniquePtr.h"
+#include "Styling/StyleColors.h"
 #include "Widgets/Colors/SColorPicker.h"
 #include "Widgets/Images/SImage.h"
 #include "Widgets/Input/SButton.h"
-#include "Widgets/Input/SCheckBox.h"
 #include "Widgets/Input/SComboBox.h"
-#include "Widgets/Input/SSearchBox.h"
 #include "Widgets/Input/SSegmentedControl.h"
-#include "Widgets/Layout/SScrollBox.h"
-#include "Widgets/SToolTip.h"
-#include "Widgets/Views/STableViewBase.h"
 
 // TraceInsightsCore
 #include "InsightsCore/Common/Stopwatch.h"
 #include "InsightsCore/Table/ViewModels/Table.h"
-#include "InsightsCore/Table/ViewModels/TableColumn.h"
 
 // TraceInsights
 #include "Insights/InsightsStyle.h"
@@ -67,6 +61,8 @@ void SMemTagTreeView::Construct(const FArguments& InArgs, TSharedPtr<SMemoryProf
 
 	// Apply the default preset.
 	ApplyViewPreset(*(*GetAvailableViewPresets())[0]);
+
+	UpdateSelectionStatsText();
 }
 
 ////////////////////////////////////////////////////////////////////////////////////////////////////
@@ -366,12 +362,90 @@ TSharedPtr<SWidget> SMemTagTreeView::ConstructFooter()
 
 		+ SHorizontalBox::Slot()
 		.FillWidth(1.0f)
-		.HAlign(HAlign_Center)
-		.Padding(2.0f, 2.0f, 2.0f, 2.0f)
+
+		+ SHorizontalBox::Slot()
+		.AutoWidth()
+		.Padding(0.0f, 2.0f, 0.0f, 2.0f)
 		[
 			SNew(STextBlock)
-			.Text(this, &SMemTagTreeView::GetFooterCenterText)
-		];
+			.Text(this, &SMemTagTreeView::GetNumSelectedTagsText)
+		]
+
+		+ SHorizontalBox::Slot()
+		.AutoWidth()
+		.Padding(4.0f, 2.0f, 0.0f, 2.0f)
+		[
+			SNew(STextBlock)
+			.Text(this, &SMemTagTreeView::GetSelectedTagsText)
+			.ColorAndOpacity(FSlateColor(EStyleColor::White25))
+		]
+
+		+ SHorizontalBox::Slot()
+		.AutoWidth()
+		.Padding(0.0f, 2.0f, 0.0f, 2.0f)
+		[
+			SNew(STextBlock)
+			.Text(FText::FromString(TEXT(" (A: ")))
+			.ColorAndOpacity(FSlateColor(EStyleColor::White25))
+			.Visibility_Lambda([this]() { return !NumSelectedTagsText.IsEmpty() ? EVisibility::Visible : EVisibility::Collapsed; })
+		]
+
+		+ SHorizontalBox::Slot()
+		.AutoWidth()
+		.Padding(0.0f, 2.0f, 0.0f, 2.0f)
+		[
+			SNew(STextBlock)
+			.Text(this, &SMemTagTreeView::GetSelectionSizeAText)
+		]
+
+		+ SHorizontalBox::Slot()
+		.AutoWidth()
+		.Padding(0.0f, 2.0f, 0.0f, 2.0f)
+		[
+			SNew(STextBlock)
+			.Text(FText::FromString(TEXT(", B: ")))
+			.ColorAndOpacity(FSlateColor(EStyleColor::White25))
+			.Visibility_Lambda([this]() { return !NumSelectedTagsText.IsEmpty() ? EVisibility::Visible : EVisibility::Collapsed; })
+		]
+
+		+ SHorizontalBox::Slot()
+		.AutoWidth()
+		.Padding(0.0f, 2.0f, 0.0f, 2.0f)
+		[
+			SNew(STextBlock)
+			.Text(this, &SMemTagTreeView::GetSelectionSizeBText)
+		]
+
+		+ SHorizontalBox::Slot()
+		.AutoWidth()
+		.Padding(0.0f, 2.0f, 0.0f, 2.0f)
+		[
+			SNew(STextBlock)
+			.Text(FText::FromString(TEXT(", B-A: ")))
+			.ColorAndOpacity(FSlateColor(EStyleColor::White25))
+			.Visibility_Lambda([this]() { return !NumSelectedTagsText.IsEmpty() ? EVisibility::Visible : EVisibility::Collapsed; })
+		]
+
+		+ SHorizontalBox::Slot()
+		.AutoWidth()
+		.Padding(0.0f, 2.0f, 0.0f, 2.0f)
+		[
+			SNew(STextBlock)
+			.Text(this, &SMemTagTreeView::GetSelectionDiffText)
+		]
+
+		+ SHorizontalBox::Slot()
+		.AutoWidth()
+		.Padding(0.0f, 2.0f, 0.0f, 2.0f)
+		[
+			SNew(STextBlock)
+			.Text(FText::FromString(TEXT(")")))
+			.ColorAndOpacity(FSlateColor(EStyleColor::White25))
+			.Visibility_Lambda([this]() { return !NumSelectedTagsText.IsEmpty() ? EVisibility::Visible : EVisibility::Collapsed; })
+		]
+
+		+ SHorizontalBox::Slot()
+		.FillWidth(1.0f);
 }
 END_SLATE_FUNCTION_BUILD_OPTIMIZATION
 
@@ -401,27 +475,24 @@ void SMemTagTreeView::UpdateSelectionStatsText()
 		FNumberFormattingOptions FormattingOptionsMem;
 		FormattingOptionsMem.MaximumFractionalDigits = 2;
 
-		SelectionStatsText = FText::Format(LOCTEXT("SelectionStatsFmt", "{0} selected {0}|plural(one=tag,other=tags) (A: {1}, B: {2}, B-A: {3})"),
-			FText::AsNumber(TotalCount),
-			FText::AsMemory(TotalSizeA, &FormattingOptionsMem),
-			FText::AsMemory(TotalSizeB, &FormattingOptionsMem),
-			FText::AsMemory(TotalSizeB - TotalSizeA, &FormattingOptionsMem));
+		NumSelectedTagsText = FText::AsNumber(TotalCount);
+		SelectedTagsText = FText::Format(LOCTEXT("SelectionStatsFmt", "selected {0}|plural(one=tag,other=tags)"), FText::AsNumber(TotalCount));
+		SelectionSizeAText = TotalSizeA == 0 ? FText::FromString(TEXT("0")) : FText::AsMemory(TotalSizeA, &FormattingOptionsMem);
+		SelectionSizeBText = TotalSizeB == 0 ? FText::FromString(TEXT("0")) : FText::AsMemory(TotalSizeB, &FormattingOptionsMem);
+		SelectionDiffText = TotalSizeA == TotalSizeB ? FText::FromString(TEXT("0")) : FText::AsMemory(TotalSizeB - TotalSizeA, &FormattingOptionsMem);
 	}
 	else
 	{
-		SelectionStatsText = LOCTEXT("NoSelectionStats", "No LLM tag selected");
+		NumSelectedTagsText = FText::GetEmpty();
+		SelectedTagsText = LOCTEXT("NoSelectionStats", "No LLM tag selected");
+		SelectionSizeAText = FText::GetEmpty();
+		SelectionSizeBText = FText::GetEmpty();
+		SelectionDiffText = FText::GetEmpty();
 	}
 }
 
 ////////////////////////////////////////////////////////////////////////////////////////////////////
 
-FText SMemTagTreeView::GetFooterCenterText() const
-{
-	return SelectionStatsText;
-}
-
-////////////////////////////////////////////////////////////////////////////////////////////////////
-
 void SMemTagTreeView::TreeView_OnSelectionChanged(FTableTreeNodePtr SelectedItem, ESelectInfo::Type SelectInfo)
 {
 	UpdateSelectionStatsText();
diff --git a/Engine/Source/Developer/TraceInsights/Private/Insights/MemoryProfiler/Widgets/SMemTagTreeView.h b/Engine/Source/Developer/TraceInsights/Private/Insights/MemoryProfiler/Widgets/SMemTagTreeView.h
--- a/Engine/Source/Developer/TraceInsights/Private/Insights/MemoryProfiler/Widgets/SMemTagTreeView.h
+++ b/Engine/Source/Developer/TraceInsights/Private/Insights/MemoryProfiler/Widgets/SMemTagTreeView.h
@@ -118,7 +118,11 @@ private:
 	TSharedRef<SWidget> ConstructTracksMiniToolbar();
 
 	void UpdateSelectionStatsText();
-	FText GetFooterCenterText() const;
+	FText GetNumSelectedTagsText() const { return NumSelectedTagsText; }
+	FText GetSelectedTagsText() const { return SelectedTagsText; }
+	FText GetSelectionSizeAText() const { return SelectionSizeAText; }
+	FText GetSelectionSizeBText() const { return SelectionSizeBText; }
+	FText GetSelectionDiffText() const { return SelectionDiffText; }
 
 	void InitAvailableViewPresets();
 
@@ -220,7 +224,11 @@ private:
 
 	//////////////////////////////////////////////////
 
-	FText SelectionStatsText;
+	FText NumSelectedTagsText;
+	FText SelectedTagsText;
+	FText SelectionSizeAText;
+	FText SelectionSizeBText;
+	FText SelectionDiffText;
 
 	FLinearColor EditableColorValue;
 
