Unreal Insights | Add support for multiple child tracks per parent track

- Introduce a new set of methods that expose an array of child tracks
- Deprecate the existing API which assumes a track may only have one child track
- Rework existing usages of the old API to instead work with the new one

Documentation of the new API can be found in Insights/ViewModels/BaseTimingTrack.h. In particular, it was designed to be extensible yet resilient toward misuse (prevent breaking upstream assumptions by downstream consumers.) Some of the design goals are explained in comments.

diff --git a/Engine/Source/Developer/TraceInsights/Private/Insights/ContextSwitches/ViewModels/ContextSwitchesSharedState.cpp b/Engine/Source/Developer/TraceInsights/Private/Insights/ContextSwitches/ViewModels/ContextSwitchesSharedState.cpp
--- a/Engine/Source/Developer/TraceInsights/Private/Insights/ContextSwitches/ViewModels/ContextSwitchesSharedState.cpp
+++ b/Engine/Source/Developer/TraceInsights/Private/Insights/ContextSwitches/ViewModels/ContextSwitchesSharedState.cpp
@@ -487,11 +487,11 @@ void FContextSwitchesSharedState::AddContextSwitchesChildTracks()
 	for (const TPair<uint32, TSharedPtr<TimingProfiler::FCpuTimingTrack>>& MapEntry : CpuTracks)
 	{
 		const TSharedPtr<TimingProfiler::FCpuTimingTrack>& CpuTrack = MapEntry.Value;
-		if (CpuTrack.IsValid() && !CpuTrack->GetChildTrack().IsValid())
+		if (CpuTrack.IsValid() && !CpuTrack->FindChildTrackOfType<FContextSwitchesTimingTrack>().IsValid())
 		{
-			TSharedPtr<FContextSwitchesTimingTrack> ContextSwitchesTrack = MakeShared<FContextSwitchesTimingTrack>(*this, TEXT("Context Switches"), CpuTrack->GetTimelineIndex(), CpuTrack->GetThreadId());
+			TSharedRef<FContextSwitchesTimingTrack> ContextSwitchesTrack = MakeShared<FContextSwitchesTimingTrack>(*this, TEXT("Context Switches"), CpuTrack->GetTimelineIndex(), CpuTrack->GetThreadId());
 			ContextSwitchesTrack->SetParentTrack(CpuTrack);
-			CpuTrack->SetChildTrack(ContextSwitchesTrack);
+			CpuTrack->AddChildTrack(ContextSwitchesTrack);
 		}
 	}
 }
@@ -519,9 +519,12 @@ void FContextSwitchesSharedState::RemoveContextSwitchesChildTracks()
 	for (const TPair<uint32, TSharedPtr<TimingProfiler::FCpuTimingTrack>>& MapEntry : CpuTracks)
 	{
 		const TSharedPtr<TimingProfiler::FCpuTimingTrack>& CpuTrack = MapEntry.Value;
-		if (CpuTrack.IsValid() && CpuTrack->GetChildTrack().IsValid() && CpuTrack->GetChildTrack()->Is<FContextSwitchesTimingTrack>())
+		if (CpuTrack.IsValid())
 		{
-			CpuTrack->SetChildTrack(nullptr);
+			if (TSharedPtr<FContextSwitchesTimingTrack> ChildTrack = CpuTrack->FindChildTrackOfType<FContextSwitchesTimingTrack>())
+			{
+				CpuTrack->RemoveChildTrack(ChildTrack.ToSharedRef());
+			}
 		}
 	}
 }
@@ -660,19 +663,14 @@ void FContextSwitchesSharedState::Command_NavigateToCpuThreadEvent_Execute()
 				if (ThreadTimingTrack.IsValid() && ThreadTimingTrack->IsVisible())
 				{
 					TimingView->SelectTimingTrack(ThreadTimingTrack, true);
-					TSharedPtr<FBaseTimingTrack> ChildTrack = ThreadTimingTrack->GetChildTrack();
-					if (ChildTrack && ChildTrack->IsKindOf("FContextSwitchesTimingTrack"))
+					if (TSharedPtr<FContextSwitchesTimingTrack> ContextSwitchesTrack = ThreadTimingTrack->FindChildTrackOfType<FContextSwitchesTimingTrack>())
 					{
-						TSharedPtr<FContextSwitchesTimingTrack> ContextSwitchesTrack = StaticCastSharedPtr<FContextSwitchesTimingTrack>(ChildTrack);
-						if (ContextSwitchesTrack.IsValid())
-						{
-							TSharedPtr<const ITimingEvent> FoundEvent = ContextSwitchesTrack->SearchEvent(
-								FTimingEventSearchParameters(CpuCoreEvent.GetStartTime(), CpuCoreEvent.GetEndTime(), ETimingEventSearchFlags::StopAtFirstMatch));
+						TSharedPtr<const ITimingEvent> FoundEvent = ContextSwitchesTrack->SearchEvent(
+							FTimingEventSearchParameters(CpuCoreEvent.GetStartTime(), CpuCoreEvent.GetEndTime(), ETimingEventSearchFlags::StopAtFirstMatch));
 
-							if (FoundEvent.IsValid())
-							{
-								TimingView->SelectTimingEvent(FoundEvent, true);
-							}
+						if (FoundEvent.IsValid())
+						{
+							TimingView->SelectTimingEvent(FoundEvent, true);
 						}
 					}
 				}
diff --git a/Engine/Source/Developer/TraceInsights/Private/Insights/ContextSwitches/ViewModels/ContextSwitchesTimingTrack.cpp b/Engine/Source/Developer/TraceInsights/Private/Insights/ContextSwitches/ViewModels/ContextSwitchesTimingTrack.cpp
--- a/Engine/Source/Developer/TraceInsights/Private/Insights/ContextSwitches/ViewModels/ContextSwitchesTimingTrack.cpp
+++ b/Engine/Source/Developer/TraceInsights/Private/Insights/ContextSwitches/ViewModels/ContextSwitchesTimingTrack.cpp
@@ -234,8 +234,6 @@ const TSharedPtr<const ITimingEvent> FContextSwitchesTimingTrack::GetEvent(doubl
 
 void FContextSwitchesTimingTrack::InitTooltip(FTooltipDrawState& InOutTooltip, const ITimingEvent& InTooltipEvent) const
 {
-	InOutTooltip.ResetContent();
-
 	if (!InTooltipEvent.CheckTrack(this) || !InTooltipEvent.Is<FContextSwitchTimingEvent>())
 	{
 		return;
diff --git a/Engine/Source/Developer/TraceInsights/Private/Insights/ContextSwitches/ViewModels/CpuCoreTimingTrack.cpp b/Engine/Source/Developer/TraceInsights/Private/Insights/ContextSwitches/ViewModels/CpuCoreTimingTrack.cpp
--- a/Engine/Source/Developer/TraceInsights/Private/Insights/ContextSwitches/ViewModels/CpuCoreTimingTrack.cpp
+++ b/Engine/Source/Developer/TraceInsights/Private/Insights/ContextSwitches/ViewModels/CpuCoreTimingTrack.cpp
@@ -238,9 +238,9 @@ void FCpuCoreTimingTrack::Draw(const ITimingTrackDrawContext& Context) const
 
 void FCpuCoreTimingTrack::PostDraw(const ITimingTrackDrawContext& Context) const
 {
-	if (ChildTrack.IsValid())
+	for (const TSharedRef<FBaseTimingTrack>& Track : GetChildTracks())
 	{
-		ChildTrack->PostDraw(Context);
+		Track->PostDraw(Context);
 	}
 
 	const TSharedPtr<const ITimingEvent> SelectedEventPtr = Context.GetSelectedEvent();
diff --git a/Engine/Source/Developer/TraceInsights/Private/Insights/TaskGraphProfiler/ViewModels/TaskGraphRelation.cpp b/Engine/Source/Developer/TraceInsights/Private/Insights/TaskGraphProfiler/ViewModels/TaskGraphRelation.cpp
--- a/Engine/Source/Developer/TraceInsights/Private/Insights/TaskGraphProfiler/ViewModels/TaskGraphRelation.cpp
+++ b/Engine/Source/Developer/TraceInsights/Private/Insights/TaskGraphProfiler/ViewModels/TaskGraphRelation.cpp
@@ -86,10 +86,7 @@ void FTaskGraphRelation::Draw(const FDrawContext& DrawContext, const FTimingTrac
 		int32 ActualSourceDepth = FMath::Min(SourceDepth, MaxEventDepth);
 		Y1 = SourceTrackShared->GetPosY();
 		Y1 += Viewport.GetLayout().GetLaneY(ActualSourceDepth) + Viewport.GetLayout().EventH / 2.0f;
-		if (SourceTrackShared->GetChildTrack() && SourceTrackShared->GetChildTrack()->GetHeight() > 0.0f)
-		{
-			Y1 += SourceTrackShared->GetChildTrack()->GetHeight() + Viewport.GetLayout().ChildTimelineDY;
-		}
+		Y1 += SourceTrackShared->GetChildTracksTopHeight(Viewport.GetLayout());
 	}
 	else
 	{
@@ -101,10 +98,7 @@ void FTaskGraphRelation::Draw(const FDrawContext& DrawContext, const FTimingTrac
 		int32 ActualTargetDepth = FMath::Min(TargetDepth, MaxEventDepth);
 		Y2 = TargetTrackShared->GetPosY();
 		Y2 += Viewport.GetLayout().GetLaneY(ActualTargetDepth) + Viewport.GetLayout().EventH / 2.0f;
-		if (TargetTrackShared->GetChildTrack() && TargetTrackShared->GetChildTrack()->GetHeight() > 0.0f)
-		{
-			Y2 += TargetTrackShared->GetChildTrack()->GetHeight() + Viewport.GetLayout().ChildTimelineDY;
-		}
+		Y2 += TargetTrackShared->GetChildTracksTopHeight(Viewport.GetLayout());
 	}
 	else
 	{
diff --git a/Engine/Source/Developer/TraceInsights/Private/Insights/TimingProfiler/Tracks/ThreadTimingTrack.cpp b/Engine/Source/Developer/TraceInsights/Private/Insights/TimingProfiler/Tracks/ThreadTimingTrack.cpp
--- a/Engine/Source/Developer/TraceInsights/Private/Insights/TimingProfiler/Tracks/ThreadTimingTrack.cpp
+++ b/Engine/Source/Developer/TraceInsights/Private/Insights/TimingProfiler/Tracks/ThreadTimingTrack.cpp
@@ -619,9 +619,9 @@ void FThreadTimingTrackImpl::BuildFilteredDrawState(ITimingEventsTrackDrawStateB
 
 void FThreadTimingTrackImpl::PostDraw(const ITimingTrackDrawContext& Context) const
 {
-	if (ChildTrack.IsValid())
+	for (const TSharedRef<FBaseTimingTrack>& Track : GetChildTracks())
 	{
-		ChildTrack->PostDraw(Context);
+		Track->PostDraw(Context);
 	}
 
 	const TSharedPtr<const ITimingEvent> SelectedEventPtr = Context.GetSelectedEvent();
@@ -674,7 +674,10 @@ void FThreadTimingTrackImpl::PostDraw(const ITimingTrackDrawContext& Context) co
 
 void FThreadTimingTrackImpl::InitTooltip(FTooltipDrawState& InOutTooltip, const ITimingEvent& InTooltipEvent) const
 {
-	InOutTooltip.ResetContent();
+	if (!IsChildTrack())
+	{
+		InOutTooltip.ResetContent();
+	}
 
 	if (InTooltipEvent.CheckTrack(this) && InTooltipEvent.Is<FThreadTrackEvent>())
 	{
@@ -748,9 +751,12 @@ void FThreadTimingTrackImpl::InitTooltip(FTooltipDrawState& InOutTooltip, const
 
 		PostInitTooltip(InOutTooltip, TooltipEvent, *Session.Get(), TimerName);
 	}
-	else if (ChildTrack.IsValid())
+	else
 	{
-		ChildTrack->InitTooltip(InOutTooltip, InTooltipEvent);
+		for (const TSharedRef<FBaseTimingTrack>& Track : GetChildTracks())
+		{
+			Track->InitTooltip(InOutTooltip, InTooltipEvent);
+		}
 	}
 
 	InOutTooltip.UpdateLayout();
@@ -800,30 +806,20 @@ const TSharedPtr<const ITimingEvent> FThreadTimingTrackImpl::GetEvent(float InPo
 
 	const FTimingViewLayout& Layout = Viewport.GetLayout();
 
-	float TopLaneY = 0.0f;
-	float TrackLanesHeight = 0.0f;
-	if (ChildTrack.IsValid() && ChildTrack->GetHeight() > 0.0f)
-	{
-		const float HeaderDY = InPosY - ChildTrack->GetPosY();
-		if (HeaderDY >= 0.0f && HeaderDY < ChildTrack->GetHeight())
-		{
-			return ChildTrack->GetEvent(InPosX, InPosY, Viewport);
-		}
-		TopLaneY = GetPosY() + 1.0f + Layout.TimelineDY + ChildTrack->GetHeight() + Layout.ChildTimelineDY;
-		TrackLanesHeight = GetHeight() - ChildTrack->GetHeight() - 1.0f - 2.0f * Layout.TimelineDY - Layout.ChildTimelineDY;
-	}
-	else
+	float TopLaneY = GetPosY() + Layout.TimelineDY;
+	float TrackLanesHeight = GetHeight();
+
+	for (const TSharedRef<FBaseTimingTrack>& Track : GetChildTracks())
 	{
-		if (IsChildTrack())
+		float HeaderDY = InPosY - Track->GetPosY();
+		float TrackHeightWithPadding = Track->GetHeight() + Layout.ChildTimelineDY;
+		if (HeaderDY >= 0.0f && HeaderDY < TrackHeightWithPadding)
 		{
-			TopLaneY = GetPosY();
-			TrackLanesHeight = GetHeight();
-		}
-		else
-		{
-			TopLaneY = GetPosY() + 1.0f + Layout.TimelineDY;
-			TrackLanesHeight = GetHeight() - 1.0f - 2.0f * Layout.TimelineDY;
+			return Track->GetEvent(InPosX, InPosY, Viewport);
 		}
+
+		TopLaneY += TrackHeightWithPadding;
+		TrackLanesHeight -= TrackHeightWithPadding;
 	}
 
 	const float DY = InPosY - TopLaneY;
@@ -1004,9 +1000,9 @@ void FThreadTimingTrackImpl::BuildContextMenu(FMenuBuilder& MenuBuilder)
 		MenuBuilder.EndSection();
 	}
 
-	if (ChildTrack.IsValid())
+	for (const TSharedRef<FBaseTimingTrack>& Track : GetChildTracks())
 	{
-		ChildTrack->BuildContextMenu(MenuBuilder);
+		Track->BuildContextMenu(MenuBuilder);
 	}
 }
 
diff --git a/Engine/Source/Developer/TraceInsights/Private/Insights/TimingProfiler/ViewModels/ThreadTimingSharedState.cpp b/Engine/Source/Developer/TraceInsights/Private/Insights/TimingProfiler/ViewModels/ThreadTimingSharedState.cpp
--- a/Engine/Source/Developer/TraceInsights/Private/Insights/TimingProfiler/ViewModels/ThreadTimingSharedState.cpp
+++ b/Engine/Source/Developer/TraceInsights/Private/Insights/TimingProfiler/ViewModels/ThreadTimingSharedState.cpp
@@ -377,9 +377,9 @@ void FThreadTimingSharedState::Tick(Timing::ITimingViewSession& InSession, const
 					{
 						// Create the GPU Work track and attach it to the GPU queue track.
 						const FString WorkTrackName = FString::Printf(TEXT("GPU%u - %s %u - WORK"), QueueInfo.GPU, QueueInfo.Name, QueueInfo.Index);
-						TSharedPtr<FGpuQueueWorkTimingTrack> WorkTrack = MakeShared<FGpuQueueWorkTimingTrack>(*this, WorkTrackName, QueueInfo.WorkTimelineIndex, QueueInfo.Id);
-						Track->SetChildTrack(WorkTrack);
+						TSharedRef<FGpuQueueWorkTimingTrack> WorkTrack = MakeShared<FGpuQueueWorkTimingTrack>(*this, WorkTrackName, QueueInfo.WorkTimelineIndex, QueueInfo.Id);
 						WorkTrack->SetParentTrack(Track);
+						Track->AddChildTrack(WorkTrack);
 					}
 				}
 				else
@@ -824,8 +824,8 @@ void FThreadTimingSharedState::AddGpuWorkChildTracks()
 		{
 			// Create the GPU Work track and attach it to the GPU queue track.
 			const FString WorkTrackName = FString::Printf(TEXT("GPU%u - %s %u - WORK"), QueueInfo.GPU, QueueInfo.Name, QueueInfo.Index);
-			TSharedPtr<FGpuQueueWorkTimingTrack> WorkTrack = MakeShared<FGpuQueueWorkTimingTrack>(*this, WorkTrackName, QueueInfo.WorkTimelineIndex, QueueInfo.Id);
-			(*TrackPtrPtr)->SetChildTrack(WorkTrack);
+			TSharedRef<FGpuQueueWorkTimingTrack> WorkTrack = MakeShared<FGpuQueueWorkTimingTrack>(*this, WorkTrackName, QueueInfo.WorkTimelineIndex, QueueInfo.Id);
+			(*TrackPtrPtr)->AddChildTrack(WorkTrack);
 			WorkTrack->SetParentTrack(*TrackPtrPtr);
 		}
 	});
@@ -837,8 +837,11 @@ void FThreadTimingSharedState::RemoveGpuWorkChildTracks()
 {
 	for (auto& Pair : GpuTracks)
 	{
-		TSharedPtr<FGpuQueueTimingTrack> Track = Pair.Value;
-		Track->SetChildTrack(nullptr);
+		const TSharedPtr<FGpuQueueTimingTrack>& Track = Pair.Value;
+		if (TSharedPtr<FGpuQueueWorkTimingTrack> WorkTrack = Track->FindChildTrackOfType<FGpuQueueWorkTimingTrack>())
+		{
+			Track->RemoveChildTrack(WorkTrack.ToSharedRef());
+		}
 	}
 }
 
diff --git a/Engine/Source/Developer/TraceInsights/Private/Insights/ViewModels/BaseTimingTrack.cpp b/Engine/Source/Developer/TraceInsights/Private/Insights/ViewModels/BaseTimingTrack.cpp
--- a/Engine/Source/Developer/TraceInsights/Private/Insights/ViewModels/BaseTimingTrack.cpp
+++ b/Engine/Source/Developer/TraceInsights/Private/Insights/ViewModels/BaseTimingTrack.cpp
@@ -3,6 +3,7 @@
 #include "Insights/ViewModels/BaseTimingTrack.h"
 
 #include "Framework/MultiBox/MultiBoxBuilder.h"
+#include "Insights/ViewModels/TimingViewLayout.h"
 
 #define LOCTEXT_NAMESPACE "BaseTimingTrack"
 
@@ -11,8 +12,81 @@ INSIGHTS_IMPLEMENT_RTTI(FBaseTimingTrack)
 // start auto generated ids from a big number (MSB set to 1) to avoid collisions with ids for GPU/CPU tracks based on 32bit timeline index
 uint64 FBaseTimingTrack::IdGenerator = (1ULL << 63);
 
+// The default destructor generates deprecation warnings for any member variables marked as
+// deprecated, so we need to suppress that.
+PRAGMA_DISABLE_DEPRECATION_WARNINGS
+FBaseTimingTrack::~FBaseTimingTrack() = default;
+PRAGMA_ENABLE_DEPRECATION_WARNINGS
+
 void FBaseTimingTrack::BuildContextMenu(FMenuBuilder& MenuBuilder)
 {
 }
 
+void FBaseTimingTrack::AddChildTrack(TSharedRef<FBaseTimingTrack> Track)
+{
+	ChildTracks.AddUnique(Track);
+}
+
+void FBaseTimingTrack::RemoveChildTrack(TSharedRef<FBaseTimingTrack> Track)
+{
+	ChildTracks.Remove(Track);
+}
+
+TSharedPtr<FBaseTimingTrack> FBaseTimingTrack::FindChildTrackOfType(FName TrackTypeName)
+{
+	const TSharedRef<FBaseTimingTrack>* FoundTrack = ChildTracks.FindByPredicate(
+		[&TrackTypeName](const TSharedRef<FBaseTimingTrack>& Track)
+		{
+			return Track->GetTypeName() == TrackTypeName;
+		});
+
+	return FoundTrack ? *FoundTrack : TSharedPtr<FBaseTimingTrack>{};
+}
+
+float FBaseTimingTrack::GetChildTracksTopHeight(const FTimingViewLayout& Layout) const
+{
+	float TotalHeight = 0.0f;
+	for (const TSharedRef<FBaseTimingTrack>& Track : GetChildTracks())
+	{
+		float TrackHeight = Track->GetHeight();
+		TotalHeight += TrackHeight;
+		if (TrackHeight > 0.0f)
+		{
+			TotalHeight += Layout.ChildTimelineDY;
+		}
+	}
+	return TotalHeight;
+}
+
+void FBaseTimingTrack::UpdateChildTracksPosY(const FTimingViewLayout& Layout)
+{
+	float RelativeChildTrackY = 0.0f;
+	for (const TSharedRef<FBaseTimingTrack>& Track : GetChildTracks())
+	{
+		Track->SetPosY(GetPosY() + RelativeChildTrackY);
+		Track->UpdateChildTracksPosY(Layout);
+		RelativeChildTrackY += Track->GetHeight() + Layout.ChildTimelineDY;
+	}
+}
+
+PRAGMA_DISABLE_DEPRECATION_WARNINGS
+TSharedPtr<FBaseTimingTrack> FBaseTimingTrack::GetChildTrack() const
+{
+	return ChildTrack;
+}
+
+void FBaseTimingTrack::SetChildTrack(TSharedPtr<FBaseTimingTrack> InTrack)
+{
+	if (ChildTrack.IsValid())
+	{
+		RemoveChildTrack(ChildTrack.ToSharedRef());
+	}
+	ChildTrack = InTrack;
+	if (InTrack.IsValid())
+	{
+		AddChildTrack(InTrack.ToSharedRef());
+	}
+}
+PRAGMA_ENABLE_DEPRECATION_WARNINGS
+
 #undef LOCTEXT_NAMESPACE
\ No newline at end of file
diff --git a/Engine/Source/Developer/TraceInsights/Private/Insights/ViewModels/TimingEventsTrack.cpp b/Engine/Source/Developer/TraceInsights/Private/Insights/ViewModels/TimingEventsTrack.cpp
--- a/Engine/Source/Developer/TraceInsights/Private/Insights/ViewModels/TimingEventsTrack.cpp
+++ b/Engine/Source/Developer/TraceInsights/Private/Insights/ViewModels/TimingEventsTrack.cpp
@@ -65,9 +65,9 @@ void FTimingEventsTrack::Reset()
 
 void FTimingEventsTrack::PreUpdate(const ITimingTrackUpdateContext& Context)
 {
-	if (ChildTrack.IsValid())
+	for (const TSharedRef<FBaseTimingTrack>& Track : GetChildTracks())
 	{
-		ChildTrack->PreUpdate(Context);
+		Track->PreUpdate(Context);
 	}
 
 	if (IsDirty() || Context.GetViewport().IsHorizontalViewportDirty())
@@ -177,14 +177,7 @@ void FTimingEventsTrack::UpdateTrackHeight(const ITimingTrackUpdateContext& Cont
 		DesiredTrackHeight = Viewport.GetLayout().ComputeTrackHeight(NumLanes);
 	}
 
-	if (ChildTrack.IsValid() && ChildTrack->GetHeight() > 0.0f)
-	{
-		if (DesiredTrackHeight == 0.0f)
-		{
-			DesiredTrackHeight += 1.0f + 2.0f * Viewport.GetLayout().TimelineDY;
-		}
-		DesiredTrackHeight += ChildTrack->GetHeight() + Viewport.GetLayout().ChildTimelineDY;
-	}
+	DesiredTrackHeight += GetChildTracksTopHeight(Viewport.GetLayout());
 
 	if (CurrentTrackHeight < DesiredTrackHeight)
 	{
@@ -218,9 +211,9 @@ void FTimingEventsTrack::UpdateTrackHeight(const ITimingTrackUpdateContext& Cont
 
 void FTimingEventsTrack::PostUpdate(const ITimingTrackUpdateContext& Context)
 {
-	if (ChildTrack.IsValid())
+	for (const TSharedRef<FBaseTimingTrack>& Track : GetChildTracks())
 	{
-		ChildTrack->PostUpdate(Context);
+		Track->PostUpdate(Context);
 	}
 
 	FBaseTimingTrack::PostUpdate(Context);
@@ -246,9 +239,9 @@ void FTimingEventsTrack::PostUpdate(const ITimingTrackUpdateContext& Context)
 
 void FTimingEventsTrack::Draw(const ITimingTrackDrawContext& Context) const
 {
-	if (ChildTrack.IsValid())
+	for (const TSharedRef<FBaseTimingTrack>& Track : GetChildTracks())
 	{
-		ChildTrack->Draw(Context);
+		Track->Draw(Context);
 	}
 
 	DrawEvents(Context, 1.0f);
@@ -263,9 +256,9 @@ void FTimingEventsTrack::Draw(const ITimingTrackDrawContext& Context) const
 
 void FTimingEventsTrack::PostDraw(const ITimingTrackDrawContext& Context) const
 {
-	if (ChildTrack.IsValid())
+	for (const TSharedRef<FBaseTimingTrack>& Track : GetChildTracks())
 	{
-		ChildTrack->PostDraw(Context);
+		Track->PostDraw(Context);
 	}
 }
 
@@ -339,19 +332,21 @@ void FTimingEventsTrack::DrawHeader(const ITimingTrackDrawContext& Context) cons
 
 void FTimingEventsTrack::DrawEvent(const ITimingTrackDrawContext& Context, const ITimingEvent& InTimingEvent, EDrawEventMode InDrawMode) const
 {
-	if (ChildTrack.IsValid() && InTimingEvent.CheckTrack(ChildTrack.Get()))
+	for (const TSharedRef<FBaseTimingTrack>& Track : GetChildTracks())
 	{
-		ChildTrack->DrawEvent(Context, InTimingEvent, InDrawMode);
+		if (InTimingEvent.CheckTrack(&Track.Get()))
+		{
+			Track->DrawEvent(Context, InTimingEvent, InDrawMode);
+			return;
+		}
 	}
-	else if (InTimingEvent.CheckTrack(this) && InTimingEvent.Is<FTimingEvent>())
+
+	if (InTimingEvent.CheckTrack(this) && InTimingEvent.Is<FTimingEvent>())
 	{
 		const FTimingEvent& TrackEvent = InTimingEvent.As<FTimingEvent>();
 		const FTimingViewLayout& Layout = Context.GetViewport().GetLayout();
 		float Y = TrackEvent.GetTrack()->GetPosY();
-		if (ChildTrack.IsValid() && ChildTrack->GetHeight() > 0.0f)
-		{
-			Y += ChildTrack->GetHeight() + Layout.ChildTimelineDY;
-		}
+		Y += GetChildTracksTopHeight(Layout);
 		Y += Layout.GetLaneY(TrackEvent.GetDepth());
 
 		const FTimingViewDrawHelper& Helper = *static_cast<const FTimingViewDrawHelper*>(&Context.GetHelper());
@@ -365,17 +360,20 @@ const TSharedPtr<const ITimingEvent> FTimingEventsTrack::GetEvent(float InPosX,
 {
 	const FTimingViewLayout& Layout = Viewport.GetLayout();
 
-	float TopLaneY = GetPosY() + 1.0f + Layout.TimelineDY;
+	float TopLaneY = GetPosY() + Layout.TimelineDY;
 	float TrackLanesHeight = GetHeight();
-	if (ChildTrack.IsValid() && ChildTrack->GetHeight() > 0.0f)
+
+	for (const TSharedRef<FBaseTimingTrack>& Track : GetChildTracks())
 	{
-		const float HeaderDY = InPosY - ChildTrack->GetPosY();
-		if (HeaderDY >= 0.0f && HeaderDY < (ChildTrack->GetHeight() + Layout.ChildTimelineDY))
+		float HeaderDY = InPosY - Track->GetPosY();
+		float TrackHeightWithPadding = Track->GetHeight() + Layout.ChildTimelineDY;
+		if (HeaderDY >= 0.0f && HeaderDY < TrackHeightWithPadding)
 		{
-			return ChildTrack->GetEvent(InPosX, InPosY, Viewport);
+			return Track->GetEvent(InPosX, InPosY, Viewport);
 		}
-		TopLaneY += ChildTrack->GetHeight() + Layout.ChildTimelineDY;
-		TrackLanesHeight -= ChildTrack->GetHeight();
+
+		TopLaneY += TrackHeightWithPadding;
+		TrackLanesHeight -= TrackHeightWithPadding;
 	}
 
 	const float DY = InPosY - TopLaneY;
diff --git a/Engine/Source/Developer/TraceInsights/Private/Insights/ViewModels/TimingViewDrawHelper.cpp b/Engine/Source/Developer/TraceInsights/Private/Insights/ViewModels/TimingViewDrawHelper.cpp
--- a/Engine/Source/Developer/TraceInsights/Private/Insights/ViewModels/TimingViewDrawHelper.cpp
+++ b/Engine/Source/Developer/TraceInsights/Private/Insights/ViewModels/TimingViewDrawHelper.cpp
@@ -328,10 +328,7 @@ void FTimingViewDrawHelper::DrawEvents(const FTimingEventsTrackDrawState& DrawSt
 
 		float TopLaneY = Track.GetPosY();
 		TopLaneY += OffsetY + Layout.TimelineDY;
-		if (Track.GetChildTrack().IsValid() && Track.ChildTrack->GetHeight() > 0.0f)
-		{
-			TopLaneY += Track.ChildTrack->GetHeight() + Layout.ChildTimelineDY;
-		}
+		TopLaneY += Track.GetChildTracksTopHeight(Layout);
 
 		// Draw filled boxes (merged borders).
 		//if (Layout.EventH > 0.0f)
@@ -415,14 +412,8 @@ void FTimingViewDrawHelper::DrawFadedEvents(const FTimingEventsTrackDrawState& D
 		NumMergedBoxes += DrawState.GetNumMergedBoxes();
 
 		float TopLaneY = Track.GetPosY();
-		if (!Track.IsChildTrack())
-		{
-			TopLaneY += OffsetY + Layout.TimelineDY;
-		}
-		if (Track.GetChildTrack().IsValid() && Track.ChildTrack->GetHeight() > 0.0f)
-		{
-			TopLaneY += Track.ChildTrack->GetHeight() + Layout.ChildTimelineDY;
-		}
+		TopLaneY += OffsetY + Layout.TimelineDY;
+		TopLaneY += Track.GetChildTracksTopHeight(Layout);
 
 		// Draw filled boxes (merged borders).
 		//if (Layout.EventH > 0.0f)
@@ -508,14 +499,8 @@ void FTimingViewDrawHelper::DrawLineEvents(const FTimingEventsTrackDrawState& Dr
 		NumMergedBoxes += DrawState.GetNumMergedBoxes();
 
 		float TopLaneY = Track.GetPosY();
-		if (!Track.IsChildTrack())
-		{
-			TopLaneY += OffsetY + Layout.TimelineDY;
-		}
-		if (Track.GetChildTrack().IsValid() && Track.ChildTrack->GetHeight() > 0.0f)
-		{
-			TopLaneY += Track.ChildTrack->GetHeight() + Layout.ChildTimelineDY;
-		}
+		TopLaneY += OffsetY + Layout.TimelineDY;
+		TopLaneY += Track.GetChildTracksTopHeight(Layout);
 
 		constexpr float LineH = 3.0f;
 		constexpr float LineEndsDH = 1.0f;
@@ -525,7 +510,7 @@ void FTimingViewDrawHelper::DrawLineEvents(const FTimingEventsTrackDrawState& Dr
 		//if (Layout.EventH > 0.0f)
 		{
 			const int32 EventFillLayerId = ReservedLayerId + ToInt32(EDrawLayer::EventFill);
-			const float Y = TopLaneY + Layout.EventH - LineEndsH + LineEndsDH + 1.0f;
+			const float Y = TopLaneY + Layout.EventH - LineEndsH + LineEndsDH;
 			for (const FTimingEventsTrackDrawState::FBoxPrimitive& Box : DrawState.Boxes)
 			{
 				DrawContext.DrawBox(EventFillLayerId, Box.X, Y, Box.W, LineEndsH, WhiteBrush, Box.Color);
@@ -537,7 +522,7 @@ void FTimingViewDrawHelper::DrawLineEvents(const FTimingEventsTrackDrawState& Dr
 		if (Layout.EventH > 2.0f)
 		{
 			const int32 EventFillLayerId = ReservedLayerId + ToInt32(EDrawLayer::EventFill);
-			const float Y = TopLaneY + Layout.EventH - LineEndsH + LineEndsDH + 2.0f;
+			const float Y = TopLaneY + Layout.EventH - LineEndsH + LineEndsDH + 1.0f;
 			for (const FTimingEventsTrackDrawState::FBoxPrimitive& Box : DrawState.InsideBoxes)
 			{
 				DrawContext.DrawBox(EventFillLayerId, Box.X, Y, Box.W, LineH, WhiteBrush, Box.Color);
@@ -549,7 +534,7 @@ void FTimingViewDrawHelper::DrawLineEvents(const FTimingEventsTrackDrawState& Dr
 		//if (Layout.EventH > 0.0f)
 		{
 			const int32 EventBorderLayerId = ReservedLayerId + ToInt32(EDrawLayer::EventBorder);
-			const float Y = TopLaneY + Layout.EventH - LineEndsH + LineEndsDH + 1.0f;
+			const float Y = TopLaneY + Layout.EventH - LineEndsH + LineEndsDH;
 			for (const FTimingEventsTrackDrawState::FBoxPrimitive& Box : DrawState.Borders)
 			{
 				DrawContext.DrawBox(EventBorderLayerId, Box.X, Y, 1.0f, LineEndsH, WhiteBrush, Box.Color);
@@ -568,7 +553,7 @@ void FTimingViewDrawHelper::DrawLineEvents(const FTimingEventsTrackDrawState& Dr
 				TextOpacity *= (Layout.EventH - TextMinEventH + 1.0f) / (FTimingViewLayout::NormalLayoutEventH - TextMinEventH + 1.0f);
 			}
 
-			const float TextY = TopLaneY - (FTimingViewLayout::NormalLayoutEventH - Layout.EventH) / 2.0f;
+			const float TextY = TopLaneY - (FTimingViewLayout::NormalLayoutEventH - Layout.EventH) / 2.0f - 1.0f;
 
 			const int32 EventTextLayerId = ReservedLayerId + ToInt32(EDrawLayer::EventText);
 			for (const FTimingEventsTrackDrawState::FTextPrimitive& Text : DrawState.Texts)
diff --git a/Engine/Source/Developer/TraceInsights/Private/Insights/Widgets/STimingView.cpp b/Engine/Source/Developer/TraceInsights/Private/Insights/Widgets/STimingView.cpp
--- a/Engine/Source/Developer/TraceInsights/Private/Insights/Widgets/STimingView.cpp
+++ b/Engine/Source/Developer/TraceInsights/Private/Insights/Widgets/STimingView.cpp
@@ -1203,10 +1203,7 @@ void STimingView::UpdatePositionForScrollableTracks()
 void STimingView::SetTrackPosY(FBaseTimingTrack& Track, float TrackPosY) const
 {
 	Track.SetPosY(TrackPosY);
-	if (Track.GetChildTrack().IsValid())
-	{
-		SetTrackPosY(*Track.GetChildTrack(), TrackPosY);
-	}
+	Track.UpdateChildTracksPosY(Viewport.GetLayout());
 }
 
 ////////////////////////////////////////////////////////////////////////////////////////////////////
@@ -3867,16 +3864,20 @@ const TSharedPtr<FBaseTimingTrack> STimingView::GetTrackAt(float InPosX, float I
 
 static bool EventBelongsToTrack(const ITimingEvent& Event, const FBaseTimingTrack* Track)
 {
-	const FBaseTimingTrack* Current = Track;
-	while (Current != nullptr)
+	if (Track == nullptr)
 	{
-		if (&Event.GetTrack().Get() == Current)
+		return false;
+	}
+
+	for (const TSharedRef<FBaseTimingTrack>& ChildTrack : Track->GetChildTracks())
+	{
+		if (EventBelongsToTrack(Event, &*ChildTrack))
 		{
 			return true;
 		}
-		Current = Current->GetChildTrack().Get();
 	}
-	return false;
+
+	return &Event.GetTrack().Get() == Track;
 }
 
 ////////////////////////////////////////////////////////////////////////////////////////////////////
diff --git a/Engine/Source/Developer/TraceInsights/Public/Insights/ViewModels/BaseTimingTrack.h b/Engine/Source/Developer/TraceInsights/Public/Insights/ViewModels/BaseTimingTrack.h
--- a/Engine/Source/Developer/TraceInsights/Public/Insights/ViewModels/BaseTimingTrack.h
+++ b/Engine/Source/Developer/TraceInsights/Public/Insights/ViewModels/BaseTimingTrack.h
@@ -35,6 +35,7 @@ class FTooltipDrawState;
 class ITimingEvent;
 class ITimingEventRelation;
 class ITimingEventFilter;
+struct FTimingViewLayout;
 
 ////////////////////////////////////////////////////////////////////////////////////////////////////
 
@@ -147,8 +148,7 @@ protected:
 		, Flags(ETimingTrackFlags::IsVisible | ETimingTrackFlags::IsDirty)
 	{}
 
-	virtual ~FBaseTimingTrack()
-	{}
+	virtual ~FBaseTimingTrack();
 
 public:
 	virtual void Reset()
@@ -262,19 +262,90 @@ public:
 	// Called back from the timing view when an event is selected.
 	virtual void OnEventSelected(const ITimingEvent& InSelectedEvent) const {}
 
-	// Called to initialize the tooltip's content with info from a timing event.
+	/**
+	 * Called to initialize the tooltip's content with info from a timing event.
+	 *
+	 * @note In most cases, this should begin by calling `InOutTooltip.ResetContent()` before adding
+	 *       any content to the tooltip. However, in case the track may be used as a child track,
+	 *       the tooltip content should not be reset, and instead should only be appended to
+	 *       (because the parent track will reset the tooltip's content.)
+	 *       This is because other child tracks may want to append to the tooltip's content,
+	 *       and their appended content should not be reset.
+	 */
 	virtual void InitTooltip(FTooltipDrawState& InOutTooltip, const ITimingEvent& InTooltipEvent) const {}
 
 	// Called back from the timing view when an event is copied to the clipboard with Ctrl+C.
 	virtual void OnClipboardCopyEvent(const ITimingEvent& InSelectedEvent) const {}
 
-	TSharedPtr<FBaseTimingTrack> GetChildTrack() const { return ChildTrack; }
-	void SetChildTrack(TSharedPtr<FBaseTimingTrack> InTrack) { ChildTrack = InTrack; };
+	// Adding children to tracks is a two-step process:
+	//
+	//  1. add the child track into the parent track using `AddChildTrack`
+	//  2. set the child track's parent track using `SetParentTrack`
+	//
+	// Both steps are needed to ensure the track renders properly.
+	// If a track is a child track, it must not also be added as a docked or scrollable track
+	// into the view.
+	//
+	// Note that child tracks do not respect the order established with `SetOrder`.
+	// Instead, you can reorder the elements individually in `GetChildTracks()`.
+
+	/** Gets a view into the array of child tracks, for enumeration and reordering. */
+	TArrayView<const TSharedRef<FBaseTimingTrack>> GetChildTracks() const { return ChildTracks; }
+	TArrayView<TSharedRef<FBaseTimingTrack>> GetChildTracks() { return ChildTracks; }
+
+	/**
+	 * Adds a child track to this track.
+	 * Note that this operation is idempotent. Calling this with the same track twice will only
+	 * add the track once.
+	 *
+	 * @note This invalidates the result of `GetChildTracks`.
+	 */
+	void AddChildTrack(TSharedRef<FBaseTimingTrack> Track);
+
+	/**
+	 * Removes the provided track from this track's child tracks.
+	 *
+	 * @note This invalidates the result of `GetChildTracks`.
+	 */
+	void RemoveChildTrack(TSharedRef<FBaseTimingTrack> Track);
+
+	// The reason why FindChildTrackOfType does not consider the inheritance hierarchy is because
+	// if it did, downstream inheritors of specific public track types could break plugins defining
+	// them. Certain parts of Insights assume that there exists at least one instance of a child
+	// track of a given type (e.g. there must be at most one FContextSwitchesTimingTrack per
+	// FCpuCoreTimingTrack.)
+
+	/**
+	 * Returns the first child track with the same type name.
+	 * Note that the type name is matched exactly (without considering inheritance.)
+	 * You should generally prefer the templated version over this.
+	 */
+	TSharedPtr<FBaseTimingTrack> FindChildTrackOfType(FName TrackTypeName);
+
+	/**
+	 * Returns the first child track of the given type.
+	 * Note that the type is matched exactly (without considering inheritance.)
+	 * `T` must belong to the SimpleRTTI hierarchy.
+	 */
+	template<typename T>
+	TSharedPtr<T> FindChildTrackOfType()
+	{
+		return StaticCastSharedPtr<T>(FindChildTrackOfType(T::GetStaticTypeName()));
+	}
 
 	void SetParentTrack(TWeakPtr<FBaseTimingTrack> InTrack) { ParentTrack = InTrack; }
 	TWeakPtr<FBaseTimingTrack> GetParentTrack() const { return ParentTrack; }
 	bool IsChildTrack() const { return ParentTrack.IsValid(); }
 
+	float GetChildTracksTopHeight(const FTimingViewLayout& Layout) const;
+	void UpdateChildTracksPosY(const FTimingViewLayout& Layout);
+
+	// Legacy API supporting a single track. Do not use in new code.
+	UE_DEPRECATED(5.6, "Loop over GetChildTracks instead")
+	TSharedPtr<FBaseTimingTrack> GetChildTrack() const;
+	UE_DEPRECATED(5.6, "Use AddChildTrack/RemoveChildTrack instead")
+	void SetChildTrack(TSharedPtr<FBaseTimingTrack> InTrack);
+
 	virtual void SetFilterConfigurator(TSharedPtr<UE::Insights::FFilterConfigurator> InFilterConfigurator) {}
 
 	// Returns number of text lines needed to display the debug string.
@@ -288,7 +359,10 @@ protected:
 	static uint64 GenerateId() { return IdGenerator++; }
 
 protected:
+	UE_DEPRECATED(5.6, "Use the ChildTracks array instead")
 	TSharedPtr<FBaseTimingTrack> ChildTrack;
+
+	TArray<TSharedRef<FBaseTimingTrack>> ChildTracks;
 	TWeakPtr<FBaseTimingTrack> ParentTrack;
 
 private:
diff --git a/Engine/Source/Developer/TraceInsights/Public/Insights/ViewModels/TimingViewLayout.h b/Engine/Source/Developer/TraceInsights/Public/Insights/ViewModels/TimingViewLayout.h
--- a/Engine/Source/Developer/TraceInsights/Public/Insights/ViewModels/TimingViewLayout.h
+++ b/Engine/Source/Developer/TraceInsights/Public/Insights/ViewModels/TimingViewLayout.h
@@ -86,14 +86,7 @@ struct TRACEINSIGHTS_API FTimingViewLayout
 		}
 		else //if (NumLanes > 0)
 		{
-			const float TrackHeight = (EventH + EventDY) * (float)NumLanes - EventDY;
-
-			if (TrackHeight < FTimingViewLayout::RealMinTimelineH)
-			{
-				return FTimingViewLayout::RealMinTimelineH;
-			}
-
-			return TrackHeight;
+			return (EventH + EventDY) * (float)NumLanes - EventDY;
 		}
 	}
 
