Fixed deprecation warning when compiling with Clang.

#rnx

diff --git a/Engine/Source/Developer/TraceInsights/Public/Insights/ViewModels/BaseTimingTrack.h b/Engine/Source/Developer/TraceInsights/Public/Insights/ViewModels/BaseTimingTrack.h
--- a/Engine/Source/Developer/TraceInsights/Public/Insights/ViewModels/BaseTimingTrack.h
+++ b/Engine/Source/Developer/TraceInsights/Public/Insights/ViewModels/BaseTimingTrack.h
@@ -137,16 +137,18 @@ class TRACEINSIGHTS_API FBaseTimingTrack : public TSharedFromThis<FBaseTimingTra
 	INSIGHTS_DECLARE_RTTI_BASE(FBaseTimingTrack)
 
 protected:
-	explicit FBaseTimingTrack(const FString& InName = FString())
+	FBaseTimingTrack()
+		: Id(GenerateId())
+	{
+	}
+
+	explicit FBaseTimingTrack(const FString& InName)
 		: Id(GenerateId())
 		, Name(InName)
-		, ValidLocations(ETimingTrackLocation::Scrollable)
-		, Location(ETimingTrackLocation::None)
-		, Order(0)
-		, PosY(0.0f)
-		, Height(0.0f)
-		, Flags(ETimingTrackFlags::IsVisible | ETimingTrackFlags::IsDirty)
-	{}
+	{
+	}
+
+	FBaseTimingTrack(const FBaseTimingTrack&) = delete;
 
 	virtual ~FBaseTimingTrack();
 
@@ -368,12 +370,12 @@ protected:
 private:
 	const uint64 Id;
 	FString Name;
-	ETimingTrackLocation ValidLocations;
-	ETimingTrackLocation Location;
-	int32 Order;
-	float PosY; // y position, in Slate units
-	float Height; // height, in Slate units
-	ETimingTrackFlags Flags;
+	ETimingTrackLocation ValidLocations = ETimingTrackLocation::Scrollable;
+	ETimingTrackLocation Location = ETimingTrackLocation::None;
+	int32 Order = 0;
+	float PosY = 0.0f; // y position, in Slate units
+	float Height = 0.0f; // height, in Slate units
+	ETimingTrackFlags Flags = ETimingTrackFlags::IsVisible | ETimingTrackFlags::IsDirty;
 
 	static uint64 IdGenerator;
 };
