Use a single console variable to enable memory insights in editor.
Improvements to the Start Memory Trace menu item.

diff --git a/Engine/Source/Developer/TraceInsights/Private/Insights/MemoryProfiler/MemoryProfilerManager.cpp b/Engine/Source/Developer/TraceInsights/Private/Insights/MemoryProfiler/MemoryProfilerManager.cpp
--- a/Engine/Source/Developer/TraceInsights/Private/Insights/MemoryProfiler/MemoryProfilerManager.cpp
+++ b/Engine/Source/Developer/TraceInsights/Private/Insights/MemoryProfiler/MemoryProfilerManager.cpp
@@ -2,6 +2,7 @@
 
 #include "MemoryProfilerManager.h"
 
+#include "HAL/IConsoleManager.h"
 #include "MessageLogModule.h"
 #include "Modules/ModuleManager.h"
 #include "Widgets/Docking/SDockTab.h"
@@ -26,7 +27,14 @@
 
 #define LOCTEXT_NAMESPACE "UE::Insights::MemoryProfiler"
 
-#define UE_ENABLE_MEMORYINSIGHTS_IN_EDITOR 0
+// For Memory Insights in editor preview.
+static bool GEnableMemoryInsightsInEditor = false;
+static FAutoConsoleVariableRef CVAREnableMemoryInsights(
+	TEXT("r.EnableMemoryInsights"),
+	GEnableMemoryInsightsInEditor,
+	TEXT("Enables the Memory Insights feature in the Editor."),
+	ECVF_Default
+);
 
 namespace UE::Insights::MemoryProfiler
 {
@@ -150,20 +158,29 @@ void FMemoryProfilerManager::BindCommands()
 
 void FMemoryProfilerManager::RegisterMajorTabs(IUnrealInsightsModule& InsightsModule)
 {
-#if !WITH_EDITOR || UE_ENABLE_MEMORYINSIGHTS_IN_EDITOR
-	const FInsightsMajorTabConfig& Config = InsightsModule.FindMajorTabConfig(FInsightsManagerTabs::MemoryProfilerTabId);
-
-	if (Config.bIsAvailable)
+#if WITH_EDITOR
+	if (GEnableMemoryInsightsInEditor)
 	{
-		// Register tab spawner for the Memory Insights.
-		FTabSpawnerEntry& TabSpawnerEntry = FGlobalTabmanager::Get()->RegisterNomadTabSpawner(FInsightsManagerTabs::MemoryProfilerTabId,
-			FOnSpawnTab::CreateRaw(this, &FMemoryProfilerManager::SpawnTab), FCanSpawnTab::CreateRaw(this, &FMemoryProfilerManager::CanSpawnTab))
-			.SetDisplayName(Config.TabLabel.IsSet() ? Config.TabLabel.GetValue() : LOCTEXT("MemoryProfilerTabTitle", "Memory Insights"))
-			.SetTooltipText(Config.TabTooltip.IsSet() ? Config.TabTooltip.GetValue() : LOCTEXT("MemoryProfilerTooltipText", "Open the Memory Insights tab."))
-			.SetIcon(Config.TabIcon.IsSet() ? Config.TabIcon.GetValue() : FSlateIcon(FInsightsStyle::GetStyleSetName(), "Icons.MemoryProfiler"));
-
-		TSharedRef<FWorkspaceItem> Group = Config.WorkspaceGroup.IsValid() ? Config.WorkspaceGroup.ToSharedRef() : FInsightsManager::Get()->GetInsightsMenuBuilder()->GetInsightsToolsGroup();
-		TabSpawnerEntry.SetGroup(Group);
+#endif
+		if (!bIsTabRegistered)
+		{
+			bIsTabRegistered = true;
+			const FInsightsMajorTabConfig& Config = InsightsModule.FindMajorTabConfig(FInsightsManagerTabs::MemoryProfilerTabId);
+
+			if (Config.bIsAvailable)
+			{
+				// Register tab spawner for the Memory Insights.
+				FTabSpawnerEntry& TabSpawnerEntry = FGlobalTabmanager::Get()->RegisterNomadTabSpawner(FInsightsManagerTabs::MemoryProfilerTabId,
+					FOnSpawnTab::CreateRaw(this, &FMemoryProfilerManager::SpawnTab), FCanSpawnTab::CreateRaw(this, &FMemoryProfilerManager::CanSpawnTab))
+					.SetDisplayName(Config.TabLabel.IsSet() ? Config.TabLabel.GetValue() : LOCTEXT("MemoryProfilerTabTitle", "Memory Insights"))
+					.SetTooltipText(Config.TabTooltip.IsSet() ? Config.TabTooltip.GetValue() : LOCTEXT("MemoryProfilerTooltipText", "Open the Memory Insights tab."))
+					.SetIcon(Config.TabIcon.IsSet() ? Config.TabIcon.GetValue() : FSlateIcon(FInsightsStyle::GetStyleSetName(), "Icons.MemoryProfiler"));
+
+				TSharedRef<FWorkspaceItem> Group = Config.WorkspaceGroup.IsValid() ? Config.WorkspaceGroup.ToSharedRef() : FInsightsManager::Get()->GetInsightsMenuBuilder()->GetInsightsToolsGroup();
+				TabSpawnerEntry.SetGroup(Group);
+			}
+		}
+#if WITH_EDITOR
 	}
 #endif
 }
@@ -198,7 +215,11 @@ TSharedRef<SDockTab> FMemoryProfilerManager::SpawnTab(const FSpawnTabArgs& Args)
 
 bool FMemoryProfilerManager::CanSpawnTab(const FSpawnTabArgs& Args) const
 {
+#if WITH_EDITOR
+	return true;
+#else
 	return bIsAvailable;
+#endif
 }
 
 ////////////////////////////////////////////////////////////////////////////////////////////////////
@@ -252,6 +273,13 @@ FMemorySharedState* FMemoryProfilerManager::GetSharedState()
 
 bool FMemoryProfilerManager::Tick(float DeltaTime)
 {
+#if WITH_EDITOR
+	if (!bIsTabRegistered && GEnableMemoryInsightsInEditor)
+	{
+		IUnrealInsightsModule& TraceInsightsModule = FModuleManager::LoadModuleChecked<IUnrealInsightsModule>("TraceInsights");
+		RegisterMajorTabs(TraceInsightsModule);
+	}
+#endif
 	// Check if session has Memory events (to spawn the tab), but not too often.
 	if (!bIsAvailable && AvailabilityCheck.Tick())
 	{
diff --git a/Engine/Source/Developer/TraceInsights/Private/Insights/MemoryProfiler/MemoryProfilerManager.h b/Engine/Source/Developer/TraceInsights/Private/Insights/MemoryProfiler/MemoryProfilerManager.h
--- a/Engine/Source/Developer/TraceInsights/Private/Insights/MemoryProfiler/MemoryProfilerManager.h
+++ b/Engine/Source/Developer/TraceInsights/Private/Insights/MemoryProfiler/MemoryProfilerManager.h
@@ -166,6 +166,9 @@ private:
 
 	/** A shared pointer to the global instance of the Memory Profiler manager. */
 	static TSharedPtr<FMemoryProfilerManager> Instance;
+
+	/** If the tab is registered in the tab manager. */
+	bool bIsTabRegistered = false;
 };
 
 } // namespace UE::Insights::MemoryProfiler
