[Insights]
- Memory Insights / LLM Tags: Added support for memory budgets.
* Budgets can be loaded from an xml file.
* DefaultMemoryBudgets.xml was added in the Programs/UnrealInsights/Config folder.
* The list of budget files can be specified in the ini file ("[Insights.MemoryProfiler]" section, "BudgetFilePath" key).
* Multiple budgets can be specified in a single file (modes).
* Budgets can be defined for any LLM tag (system, asset or asset classes) by specifying the full stat name of the tag.
* Overrides can be defined by platform.
* The LLM tags with a budget defined have the "Size" column values highlighted in red (if over budget) or green (if under budget).
- Memory Insights / LLM Tags: Fixed aggregation and status bar stats when the time markers or selection changes the values of Size columns.
- Memory Insights / LLM Tags: Changed the size columns to display memory values with 2 decimal digits.
- TraceInsightsCore/Tables: Added GenerateCustomWidget() API to ITableCellValueFormatter to allow customization of the table cell content.

diff --git a/Engine/Programs/UnrealInsights/Config/DefaultMemoryBudgets.xml b/Engine/Programs/UnrealInsights/Config/DefaultMemoryBudgets.xml
new file mode 100644
--- /dev/null
+++ b/Engine/Programs/UnrealInsights/Config/DefaultMemoryBudgets.xml
@@ -0,0 +1,73 @@
+<MemoryBudget>
+
+	<!-- 1 KiB = 2^10 Bytes; 1 KB = 1000 Bytes -->
+	<!-- 1 MiB = 2^20 Bytes; 1 MB = 10^6 Bytes -->
+	<!-- 1 GiB = 2^30 Bytes; 1 GB = 10^9 Bytes -->
+
+	<budget name="Default">
+
+		<tag name="Animation"                     mem-max="50 MiB"/>
+		<tag name="AsyncLoading"                  mem-max="225 MiB"/>
+		<tag name="Audio"                         mem-max="180 MiB"/>
+		<tag name="FName"                         mem-max="6 MiB"/>
+		<tag name="GC"                            mem-max="20 MiB"/>
+		<tag name="Materials"                     mem-max="30 MiB"/>
+		<tag name="Meshes"                        mem-max="150 MiB"/>
+		<tag name="Meshes/StaticMesh"             mem-max="20 MiB"/>
+		<tag name="Niagara"                       mem-max="100 MiB"/>
+		<tag name="Particles"                     mem-max="100 MiB"/>
+		<tag name="Physics"                       mem-max="500 MiB"/>
+		<tag name="RHIMisc"                       mem-max="50 MiB"/>
+		<tag name="RenderTargets"                 mem-max="750 MiB"/>
+		<tag name="SceneRender"                   mem-max="50 MiB"/>
+		<tag name="Shaders"                       mem-max="100 MiB"/>
+		<tag name="Stats"                         mem-max="5 MiB"/>
+		<tag name="TextureMetaData"               mem-max="5 MiB"/>
+		<tag name="Textures"                      mem-max="2 GiB"/>
+		<tag name="UI"                            mem-max="100 MiB"/>
+		<tag name="UObject"                       mem-max="250 MiB"/>
+
+		<platform name="Platform 1">
+			<tag name="Animation"                     mem-max="100 MiB"/>
+			<tag name="Audio"                         mem-max="200 MiB"/>
+			<tag name="Meshes"                        mem-max="300 MiB"/>
+			<tag name="Textures"                      mem-max="3 GiB"/>
+		</platform>
+
+		<platform name="Platform 2">
+			<tag name="Animation"                     mem-max="10 MiB"/>
+			<tag name="Audio"                         mem-max="50 MiB"/>
+			<tag name="Meshes"                        mem-max="100 MiB"/>
+			<tag name="Textures"                      mem-max="1 GiB"/>
+		</platform>
+	</budget>
+
+	<budget name="Empty">
+	</budget>
+
+	<budget name="Test1">
+		<tag name="Animation"                     mem-max="0"/>
+		<tag name="Audio"                         mem-max="123"/>
+		<tag name="Meshes"                        mem-max="2000"/>
+		<tag name="Textures"                      mem-max="1 GiB"/>
+		<platform name="Platform 1">
+			<tag name="Animation"                     mem-max="100"/>
+			<tag name="Audio"                         mem-max="1 KB"/>
+			<tag name="Meshes"                        mem-max="3 KB"/>
+			<tag name="Textures"                      mem-max="2 GB"/>
+		</platform>
+	</budget>
+
+	<budget name="Test2">
+		<tag name="Audio"                         mem-max="1 MiB"/>
+		<platform name="Platform 2">
+			<tag name="Animation"                     mem-max="2 KiB"/>
+		</platform>
+		<platform name="Platform 3">
+			<tag name="Animation"                     mem-max="3 KiB"/>
+			<tag name="Audio"                         mem-max="3 MiB"/>
+			<tag name="Textures"                      mem-max="3 GiB"/>
+		</platform>
+	</budget>
+
+</MemoryBudget>
diff --git a/Engine/Programs/UnrealInsights/Config/DefaultUnrealInsightsSettings.ini b/Engine/Programs/UnrealInsights/Config/DefaultUnrealInsightsSettings.ini
--- a/Engine/Programs/UnrealInsights/Config/DefaultUnrealInsightsSettings.ini
+++ b/Engine/Programs/UnrealInsights/Config/DefaultUnrealInsightsSettings.ini
@@ -2,87 +2,123 @@
 bCanSaveAllSections=true
 
 [Insights.TimingProfiler]
+
 ; The default (initial) zoom level of the Timing view.
-; This option can be changed from the "Tracks" drop down menu.
 ; It is defined as time between major tick marks, in seconds.
 ; Valid values: [0.00000001 .. 3600] i.e. [10ns to 1h]
 DefaultZoomLevel=5.0
+
 ; Auto hide empty tracks (ex.: ones without timing events in the current viewport).
-; This option can be changed from the "Tracks" drop down menu.
+; This option can be changed from the "View Mode" drop down menu.
 bAutoHideEmptyTracks=True
+
 ; If enabled, the panning is allowed to continue when the mouse cursor reaches the edges of the screen.
+; This option can be changed from the "View Mode" drop down menu.
 bAllowPanningOnScreenEdges=False
+
 ; The default state for Auto-Scroll, in the Timing view.
 bAutoScroll=False
+
 ; Frame Alignment. Controls if auto-scroll should align center of the viewport with start of a frame or not.
 ; Valid options: 'none' to disable frame alignment, 'game' / 'rendering' to align with Game / Rendering frames
+; This option can be changed from the "Auto-Scroll Options" menu.
 AutoScrollFrameAlignment=game
+
 ; Viewport offset while auto-scrolling, as percent of viewport width.
 ; If positive, it offsets the viewport forward, allowing an empty space at the right side of the viewport (i.e. after end of session).
 ; If negative, it offsets the viewport backward (i.e. end of session will be outside viewport).
+; This option can be changed from the "Auto-Scroll Options" menu.
 AutoScrollViewportOffsetPercent=0.1
+
 ; Minimum time between two auto-scroll updates, in [seconds].
+; This option can be changed from the "Auto-Scroll Options" menu.
 AutoScrollMinDelay=0.3
 
 [Insights.TimingProfiler.FramesView]
+; These options can be changed from the context menu of the Frames view.
+
 ; If enabled, the upper threshold line is visible. The frame coloring by threshold is enabled regardless of this setting.
 bShowUpperThresholdLine=False
+
 ; If enabled, the lower threshold line is visible. The frame coloring by threshold is enabled regardless of this setting.
 bShowLowerThresholdLine=False
+
 ; The upper threshold for frames. Can be specified as a frame duration ([0.001 .. 1.0] seconds; ex.: "0.010" for 10 ms) or as a framerate ([1 fps .. 1000 fps]; ex: "100 fps").
 UpperThreshold=30 fps
+
 ; The lower threshold for frames. Can be specified as a frame duration ([0.001 .. 1.0] seconds; ex.: "0.010" for 10 ms) or as a framerate ([1 fps .. 1000 fps]; ex: "100 fps").
 LowerThreshold=60 fps
-; If enabled, the Timing view will also be zoomed when a frame is selected in the Frames track.
-; This option can be changed from the context menu of the Frames track.
+
+; If enabled, the Timing view will also be zoomed when a frame is selected in the Frames view.
 bAutoZoomOnFrameSelection=False
 
 [Insights.TimingProfiler.MainGraph]
+; These options can be changed from the context menu of the Main Graph track (Timing Insights profiler / Timing view).
+
 ; If enabled, values will be displayed as points in the Main Graph Track in Timing Insights.
 ShowPoints=False
+
 ; If enabled, values will be displayed as points with border in the Main Graph Track in Timing Insights.
 ShowPointsWithBorder=True
+
 ; If enabled, values will be displayed as connected lines in the Main Graph Track in Timing Insights.
 ShowConnectedLines=True
+
 ; If enabled, values will be displayed as polygons in the Main Graph Track in Timing Insights.
 ShowPolygons=True
+
 ; If enabled, uses duration of timing events for connected lines and polygons in the Main Graph Track in Timing Insights.
 ShowEventDuration=True
+
 ; If enabled, shows bars corresponding to the duration of the timing events in the Main Graph Track in Timing Insights.
 ShowBars=False
+
 ; If enabled, shows game frames in the Main Graph Track in Timing Insights.
 ShowGameFrames=True
+
 ; If enabled, shows rendering frames in the Main Graph Track in Timing Insights.
 ShowRenderingFrame=True
 
 [Insights.TimingProfiler.TimersView]
+
 ; Names of visible columns in the Timers Panel in the Game mode.
 +GameFrameColumns=MaxInclTime
 +GameFrameColumns=AverageInclTime
 +GameFrameColumns=MedianInclTime
 +GameFrameColumns=MinInclTime
+
 ; Names of visible columns in the Timers Panel in the Rendering mode.
 +RenderingFrameColumns=MaxInclTime
 +RenderingFrameColumns=AverageInclTime
 +RenderingFrameColumns=MedianInclTime
 +RenderingFrameColumns=MinInclTime
+
 ; Names of visible columns in the Timers Panel in the Instance mode.
 +InstanceColumns=Count
 +InstanceColumns=TotalInclTime
 +InstanceColumns=TotalExclTime
+
 ; Mode for the timers panel. 0 : Game Frame; 1 : Rendering Frame; 2 : Instance
 Mode=2
+
 ; Grouping Mode for the timers panel. 0 : Flat; 1 : ByName; 3 : ByType; 4 : ByInstanceCount
 GroupingMode=3
+
 ; Controls the CPU filter button from Timers View.
 ShowCpuTimers=1
+
 ; Controls the GPU filter button from Timers View.
 ShowGpuTimers=1
+
 ; Controls the "Filter Out Zero Count Timers" button from Timers View.
 ShowZeroCountTimers=1
 
 [Insights.MemoryProfiler]
+
 ; Add symbol search paths by adding to this array. Order is important.
-;+SymbolSearchPaths=c:\symbols
+;+SymbolSearchPaths=C:\Symbols
 ; Symbol server format is also accepted, such as the official Microsoft symbol servers (see https://docs.microsoft.com/en-us/windows-hardware/drivers/debugger/advanced-symsrv-use for details)
-;+SymbolSearchPaths=srv*C:\Symbols;https://msdl.microsoft.com/download/symbols
\ No newline at end of file
+;+SymbolSearchPaths=srv*C:\Symbols;https://msdl.microsoft.com/download/symbols
+
+; Memory Budgets (used by the LLM Tags view)
++BudgetFilePath=(Label="Default",Path="Engine/Programs/UnrealInsights/Config/DefaultMemoryBudgets.xml")
diff --git a/Engine/Source/Developer/TraceInsights/Private/Insights/MemoryProfiler/ViewModels/MemTagBudget.cpp b/Engine/Source/Developer/TraceInsights/Private/Insights/MemoryProfiler/ViewModels/MemTagBudget.cpp
new file mode 100644
--- /dev/null
+++ b/Engine/Source/Developer/TraceInsights/Private/Insights/MemoryProfiler/ViewModels/MemTagBudget.cpp
@@ -0,0 +1,116 @@
+// Copyright Epic Games, Inc. All Rights Reserved.
+
+#include "MemTagBudget.h"
+
+#include "Math/UnitConversion.h"
+#include "Misc/Optional.h"
+#include "XmlParser.h"
+
+// TraceInsights
+#include "Insights/MemoryProfiler/MemoryProfilerManager.h"
+
+namespace UE::Insights::MemoryProfiler
+{
+
+////////////////////////////////////////////////////////////////////////////////////////////////////
+
+bool FMemTagBudget::LoadFromFile(const FString& InFilePath)
+{
+	Reset();
+	FilePath = InFilePath;
+	return ReLoadFromFile();
+}
+
+////////////////////////////////////////////////////////////////////////////////////////////////////
+
+bool FMemTagBudget::ReLoadFromFile()
+{
+	static const FString STR_Budget(TEXT("budget"));
+	static const FString STR_Platform(TEXT("platform"));
+	//static const FString STR_Set(TEXT("set"));
+	//static const FString STR_Tracker(TEXT("tracker"));
+	static const FString STR_Name(TEXT("name"));
+	//static const FString STR_NameRegEx(TEXT("name_regex"));
+	static const FString STR_Tag(TEXT("tag"));
+	//static const FString STR_Tag(TEXT("group"));
+	static const FString STR_TagMemMax(TEXT("mem-max"));
+
+	FXmlFile XmlFile;
+	if (!XmlFile.LoadFile(FilePath, EConstructMethod::ConstructFromFile))
+	{
+		UE_LOG(LogMemoryProfiler, Warning, TEXT("Could not load memory budget file: \"%s\""), *FilePath);
+		return false;
+	}
+
+	FXmlNode* RootNode = XmlFile.GetRootNode();
+	if (RootNode == nullptr)
+	{
+		UE_LOG(LogMemoryProfiler, Warning, TEXT("Could not parse memory budget file as an XML file, make sure it has only one root node"));
+		return false;
+	}
+
+	const TArray<FXmlNode*>& BudgetNodes = RootNode->GetChildrenNodes();
+	for (FXmlNode* BudgetNode : BudgetNodes)
+	{
+		if (BudgetNode->GetTag() == STR_Budget)
+		{
+			FString ModeName = BudgetNode->GetAttribute(STR_Name);
+			FMemTagBudgetMode& BudgetMode = GetOrAddMode(ModeName);
+
+			auto ProcessTagNode = [&](FXmlNode* GroupNode, FMemTagBudgetValueSet& Set)
+				{
+					FString TagName = GroupNode->GetAttribute(STR_Name);
+					if (Set.FindValue(TagName) != nullptr)
+					{
+						UE_LOG(LogMemoryProfiler, Warning, TEXT("Ignoring budget value for tag \"%s\" (defined multiple times)."), *TagName);
+						return;
+					}
+					FString Value = GroupNode->GetAttribute(STR_TagMemMax);
+
+					PRAGMA_DISABLE_UNSAFE_TYPECAST_WARNINGS
+					TValueOrError<FNumericUnit<double>, FText> Result = FNumericUnit<double>::TryParseExpression(*Value, EUnit::Bytes, 0);
+					if (Result.IsValid())
+					{
+						auto MemoryValueInBytes = Result.GetValue().ConvertTo(EUnit::Bytes);
+						if (MemoryValueInBytes.IsSet())
+						{
+							Set.AddValue(TagName, (int64)MemoryValueInBytes->Value);
+						}
+					}
+					else
+					{
+						UE_LOG(LogMemoryProfiler, Warning, TEXT("%s"), *(Result.GetError().ToString()));
+					}
+					PRAGMA_RESTORE_UNSAFE_TYPECAST_WARNINGS
+				};
+
+			const TArray<FXmlNode*>& BudgetChildNodes = BudgetNode->GetChildrenNodes();
+			for (FXmlNode* BudgetChildNode : BudgetChildNodes)
+			{
+				if (BudgetChildNode->GetTag() == STR_Platform)
+				{
+					FString PlatformName = BudgetChildNode->GetAttribute(STR_Name);
+					FMemTagBudgetValueSet& BudgetPlatform = *BudgetMode.GetOrAddPlatform(PlatformName);
+					const TArray<FXmlNode*>& PlatformChildNodes = BudgetChildNode->GetChildrenNodes();
+					for (FXmlNode* PlatformChildNode : PlatformChildNodes)
+					{
+						if (PlatformChildNode->GetTag() == STR_Tag)
+						{
+							ProcessTagNode(PlatformChildNode, BudgetPlatform);
+						}
+					}
+				}
+				else if (BudgetChildNode->GetTag() == STR_Tag)
+				{
+					ProcessTagNode(BudgetChildNode, BudgetMode.GetDefaultValues());
+				}
+			}
+		}
+	}
+
+	return true;
+}
+
+////////////////////////////////////////////////////////////////////////////////////////////////////
+
+} // namespace UE::Insights::MemoryProfiler
diff --git a/Engine/Source/Developer/TraceInsights/Private/Insights/MemoryProfiler/ViewModels/MemTagBudget.h b/Engine/Source/Developer/TraceInsights/Private/Insights/MemoryProfiler/ViewModels/MemTagBudget.h
new file mode 100644
--- /dev/null
+++ b/Engine/Source/Developer/TraceInsights/Private/Insights/MemoryProfiler/ViewModels/MemTagBudget.h
@@ -0,0 +1,161 @@
+// Copyright Epic Games, Inc. All Rights Reserved.
+
+#pragma once
+
+#include "CoreTypes.h"
+
+#include "Containers/Array.h"
+#include "Containers/Map.h"
+#include "Containers/UnrealString.h"
+#include "Templates/SharedPointer.h"
+
+namespace UE::Insights::MemoryProfiler
+{
+
+////////////////////////////////////////////////////////////////////////////////////////////////////
+
+class FMemTagBudgetValueSet
+{
+public:
+	FMemTagBudgetValueSet() {}
+
+	const int64* FindValue(const FString& Name) const
+	{
+		return Values.Find(Name);
+	}
+
+	void AddValue(const FString& Name, int64 Value)
+	{
+		Values.Add(Name, Value);
+	}
+
+private:
+	TMap<FString, int64> Values;
+};
+
+////////////////////////////////////////////////////////////////////////////////////////////////////
+
+class FMemTagBudgetMode
+{
+public:
+	FMemTagBudgetMode(const FString& InName)
+		: Name(InName)
+		, DefaultValues(MakeShared<FMemTagBudgetValueSet>())
+		, PlatformOverrides()
+	{
+	}
+
+	const FString& GetName() const
+	{
+		return Name;
+	}
+
+	const FMemTagBudgetValueSet& GetDefaultValues() const
+	{
+		return *DefaultValues;
+	}
+
+	FMemTagBudgetValueSet& GetDefaultValues()
+	{
+		return *DefaultValues;
+	}
+
+	void GetPlatforms(TArray<FString> &OutPlatforms)
+	{
+		PlatformOverrides.GetKeys(OutPlatforms);
+	}
+
+	void EnumeratePlatforms(TFunctionRef<void(const FString& Platform, const FMemTagBudgetValueSet& Set)> Callback) const
+	{
+		for (const auto& KV : PlatformOverrides)
+		{
+			Callback(KV.Key, *KV.Value);
+		}
+	}
+
+	const FMemTagBudgetValueSet* FindPlatform(const FString& Platform) const
+	{
+		const TSharedRef<FMemTagBudgetValueSet>* Found = PlatformOverrides.Find(Platform);
+		return Found ? &**Found : nullptr;
+	}
+
+	TSharedRef<FMemTagBudgetValueSet>& GetOrAddPlatform(const FString& Platform)
+	{
+		TSharedRef<FMemTagBudgetValueSet>* Found = PlatformOverrides.Find(Platform);
+		if (Found)
+		{
+			return *Found;
+		}
+		return PlatformOverrides.Add(Platform, MakeShared<FMemTagBudgetValueSet>());
+	}
+
+private:
+	FString Name;
+	TSharedRef<FMemTagBudgetValueSet> DefaultValues;
+	TMap<FString, TSharedRef<FMemTagBudgetValueSet>> PlatformOverrides;
+};
+
+////////////////////////////////////////////////////////////////////////////////////////////////////
+
+class FMemTagBudget
+{
+public:
+	FMemTagBudget(const FString& InName)
+		: Name(InName)
+	{
+	}
+
+	void Reset()
+	{
+		Name.Reset();
+		Modes.Reset();
+	}
+
+	const FString& GetName() const
+	{
+		return Name;
+	}
+
+	const FString& GetFilePath() const
+	{
+		return FilePath;
+	}
+
+	TArray<TSharedRef<FMemTagBudgetMode>>& GetModes() { return Modes; }
+
+	const FMemTagBudgetMode* FindMode(const FString& InModeName) const
+	{
+		const TSharedRef<FMemTagBudgetMode>* FoundBudgetMode = Modes.FindByPredicate(
+			[&InModeName](const TSharedRef<const FMemTagBudgetMode>& BudgetMode)
+			{
+				return InModeName.Equals(BudgetMode->GetName(), ESearchCase::CaseSensitive);
+			});
+		return (FoundBudgetMode) ? &**FoundBudgetMode : nullptr;
+	}
+
+	FMemTagBudgetMode& GetOrAddMode(const FString& InModeName)
+	{
+		TSharedRef<FMemTagBudgetMode>* FoundBudgetMode = Modes.FindByPredicate(
+			[&InModeName](const TSharedRef<const FMemTagBudgetMode>& BudgetMode)
+			{
+				return InModeName.Equals(BudgetMode->GetName(), ESearchCase::CaseSensitive);
+			});
+		if (FoundBudgetMode)
+		{
+			return **FoundBudgetMode;
+		}
+		return *Modes.Add_GetRef(MakeShared<FMemTagBudgetMode>(InModeName));
+	}
+
+	bool LoadFromFile(const FString& InFilePath);
+	bool ReLoadFromFile();
+
+private:
+	FString Name;
+	FString FilePath;
+	TArray<TSharedRef<FMemTagBudgetMode>> Modes;
+};
+
+////////////////////////////////////////////////////////////////////////////////////////////////////
+
+} // namespace UE::Insights::MemoryProfiler
diff --git a/Engine/Source/Developer/TraceInsights/Private/Insights/MemoryProfiler/ViewModels/MemTagTable.cpp b/Engine/Source/Developer/TraceInsights/Private/Insights/MemoryProfiler/ViewModels/MemTagTable.cpp
--- a/Engine/Source/Developer/TraceInsights/Private/Insights/MemoryProfiler/ViewModels/MemTagTable.cpp
+++ b/Engine/Source/Developer/TraceInsights/Private/Insights/MemoryProfiler/ViewModels/MemTagTable.cpp
@@ -3,6 +3,7 @@
 #include "MemTagTable.h"
 
 #include "Styling/StyleColors.h"
+#include "Widgets/Layout/SBox.h"
 #include "Widgets/SBoxPanel.h"
 #include "Widgets/SToolTip.h"
 #include "Widgets/Text/STextBlock.h"
@@ -16,7 +17,6 @@
 // TraceInsights
 #include "Insights/MemoryProfiler/ViewModels/MemTagNode.h"
 #include "Insights/MemoryProfiler/ViewModels/MemTagNodeGroupingAndSorting.h"
-#include "Insights/MemoryProfiler/ViewModels/MemTagTable.h"
 
 #define LOCTEXT_NAMESPACE "UE::Insights::MemoryProfiler::FMemTagTable"
 
@@ -88,6 +88,60 @@ void FMemTagTable::AddDefaultColumns()
 		}
 	};
 
+	class FBudgetedMemoryValueFormatter : public FInt64ValueFormatterAsMemory
+	{
+	public:
+		FBudgetedMemoryValueFormatter()
+		{
+			GetFormattingOptions().MaximumFractionalDigits = 2;
+		}
+		virtual ~FBudgetedMemoryValueFormatter() {}
+
+		virtual TSharedPtr<SWidget> GenerateCustomWidget(const FTableColumn& Column, const FBaseTreeNode& Node) const override
+		{
+			return SNew(SBox)
+				.ToolTip(GetCustomTooltip(Column, Node))
+				.HAlign(Column.GetHorizontalAlignment())
+				.VAlign(VAlign_Center)
+				.Padding(FMargin(2.0f, 0.0f, 2.0f, 0.0f))
+				[
+					SNew(STextBlock)
+					.Text_Lambda([this, &Column, &Node]()
+						{
+							return FormatValue(Column.GetValue(Node));
+						})
+					.ColorAndOpacity_Lambda([&Column, &Node]
+						{
+							if (Node.Is<FMemTagNode>())
+							{
+								TOptional<FTableCellValue> Value = Column.GetValue(Node);
+								if (Value.IsSet())
+								{
+									const FMemTagNode& MemTagNode = Node.As<FMemTagNode>();
+									if (MemTagNode.GetSizeBudget() != 0)
+									{
+										int64 SizeValue = Value.GetValue().AsInt64();
+										if (SizeValue > MemTagNode.GetSizeBudget())
+										{
+											return FLinearColor(1.0f, 0.3f, 0.3f, 1.0f);
+										}
+										else
+										{
+											return FLinearColor(0.1f, 0.5f, 0.1f, 1.0f);
+										}
+									}
+									else
+									{
+										return FLinearColor(0.5f, 0.5f, 0.5f, 1.0f);
+									}
+								}
+							}
+							return FLinearColor(1.0f, 0.7f, 0.3f, 1.0f);
+						})
+				];
+		}
+	};
+
 	//////////////////////////////////////////////////
 	// Hierarchy Column
 	{
@@ -238,7 +292,7 @@ void FMemTagTable::AddDefaultColumns()
 		TSharedRef<ITableCellValueGetter> Getter = MakeShared<FMemTagSizeAValueGetter>();
 		Column.SetValueGetter(Getter);
 
-		TSharedRef<ITableCellValueFormatter> Formatter = MakeShared<FInt64ValueFormatterAsMemory>();
+		TSharedRef<ITableCellValueFormatter> Formatter = MakeShared<FBudgetedMemoryValueFormatter>();
 		Column.SetValueFormatter(Formatter);
 
 		TSharedRef<ITableCellValueSorter> Sorter = MakeShared<FSorterByInt64Value>(ColumnRef);
@@ -289,7 +343,7 @@ void FMemTagTable::AddDefaultColumns()
 		TSharedRef<ITableCellValueGetter> Getter = MakeShared<FMemTagSizeAValueGetter>();
 		Column.SetValueGetter(Getter);
 
-		TSharedRef<ITableCellValueFormatter> Formatter = MakeShared<FInt64ValueFormatterAsMemory>();
+		TSharedRef<ITableCellValueFormatter> Formatter = MakeShared<FBudgetedMemoryValueFormatter>();
 		Column.SetValueFormatter(Formatter);
 
 		TSharedRef<ITableCellValueSorter> Sorter = MakeShared<FSorterByInt64Value>(ColumnRef);
@@ -324,7 +378,7 @@ void FMemTagTable::AddDefaultColumns()
 
 		Column.SetDataType(ETableCellDataType::Int64);
 
-		class FMemTagSizeBValueGetter : public FAggregatedValueGetter
+		class FMemTagSizeDiffValueGetter : public FAggregatedValueGetter
 		{
 		public:
 			virtual const TOptional<FTableCellValue> GetNodeValue(const FBaseTreeNode& Node) const override
@@ -337,10 +391,11 @@ void FMemTagTable::AddDefaultColumns()
 				return TOptional<FTableCellValue>();
 			}
 		};
-		TSharedRef<ITableCellValueGetter> Getter = MakeShared<FMemTagSizeBValueGetter>();
+		TSharedRef<ITableCellValueGetter> Getter = MakeShared<FMemTagSizeDiffValueGetter>();
 		Column.SetValueGetter(Getter);
 
-		TSharedRef<ITableCellValueFormatter> Formatter = MakeShared<FInt64ValueFormatterAsMemory>();
+		TSharedRef<FInt64ValueFormatterAsMemory> Formatter = MakeShared<FInt64ValueFormatterAsMemory>();
+		Formatter->GetFormattingOptions().MaximumFractionalDigits = 2;
 		Column.SetValueFormatter(Formatter);
 
 		TSharedRef<ITableCellValueSorter> Sorter = MakeShared<FSorterByInt64Value>(ColumnRef);
@@ -442,7 +497,7 @@ void FMemTagTable::AddDefaultColumns()
 		TSharedRef<ITableCellValueGetter> Getter = MakeShared<FMemTagSizeMinValueGetter>();
 		Column.SetValueGetter(Getter);
 
-		TSharedRef<ITableCellValueFormatter> Formatter = MakeShared<FInt64ValueFormatterAsMemory>();
+		TSharedRef<ITableCellValueFormatter> Formatter = MakeShared<FBudgetedMemoryValueFormatter>();
 		Column.SetValueFormatter(Formatter);
 
 		TSharedRef<ITableCellValueSorter> Sorter = MakeShared<FSorterByInt64Value>(ColumnRef);
@@ -493,7 +548,7 @@ void FMemTagTable::AddDefaultColumns()
 		TSharedRef<ITableCellValueGetter> Getter = MakeShared<FMemTagSizeMaxValueGetter>();
 		Column.SetValueGetter(Getter);
 
-		TSharedRef<ITableCellValueFormatter> Formatter = MakeShared<FInt64ValueFormatterAsMemory>();
+		TSharedRef<ITableCellValueFormatter> Formatter = MakeShared<FBudgetedMemoryValueFormatter>();
 		Column.SetValueFormatter(Formatter);
 
 		TSharedRef<ITableCellValueSorter> Sorter = MakeShared<FSorterByInt64Value>(ColumnRef);
@@ -544,7 +599,7 @@ void FMemTagTable::AddDefaultColumns()
 		TSharedRef<ITableCellValueGetter> Getter = MakeShared<FMemTagSizeAverageValueGetter>();
 		Column.SetValueGetter(Getter);
 
-		TSharedRef<ITableCellValueFormatter> Formatter = MakeShared<FInt64ValueFormatterAsMemory>();
+		TSharedRef<ITableCellValueFormatter> Formatter = MakeShared<FBudgetedMemoryValueFormatter>();
 		Column.SetValueFormatter(Formatter);
 
 		TSharedRef<ITableCellValueSorter> Sorter = MakeShared<FSorterByInt64Value>(ColumnRef);
@@ -572,7 +627,8 @@ void FMemTagTable::AddDefaultColumns()
 
 		Column.SetFlags(ETableColumnFlags::ShouldBeVisible |
 						ETableColumnFlags::CanBeHidden |
-						ETableColumnFlags::CanBeFiltered);
+						ETableColumnFlags::CanBeFiltered |
+						ETableColumnFlags::IsDynamic);
 
 		Column.SetHorizontalAlignment(HAlign_Right);
 		Column.SetInitialWidth(100.0f);
@@ -587,7 +643,10 @@ void FMemTagTable::AddDefaultColumns()
 				if (Node.Is<FMemTagNode>())
 				{
 					const FMemTagNode& MemTagNode = Node.As<FMemTagNode>();
-					return FTableCellValue(MemTagNode.GetSizeBudget());
+					if (MemTagNode.GetSizeBudget() != 0)
+					{
+						return FTableCellValue(MemTagNode.GetSizeBudget());
+					}
 				}
 				return TOptional<FTableCellValue>();
 			}
@@ -595,7 +654,8 @@ void FMemTagTable::AddDefaultColumns()
 		TSharedRef<ITableCellValueGetter> Getter = MakeShared<FMemTagBudgetValueGetter>();
 		Column.SetValueGetter(Getter);
 
-		TSharedRef<ITableCellValueFormatter> Formatter = MakeShared<FInt64ValueFormatterAsMemory>();
+		TSharedRef<FInt64ValueFormatterAsMemory> Formatter = MakeShared<FInt64ValueFormatterAsMemory>();
+		Formatter->GetFormattingOptions().MaximumFractionalDigits = 2;
 		Column.SetValueFormatter(Formatter);
 
 		TSharedRef<ITableCellValueSorter> Sorter = MakeShared<FSorterByInt64Value>(ColumnRef);
@@ -605,7 +665,7 @@ void FMemTagTable::AddDefaultColumns()
 		//TSharedRef<IFilterValueConverter> Converter = MakeShared<FMemoryFilterValueConverter>();
 		//Column.SetValueConverter(Converter);
 
-		Column.SetAggregation(ETableColumnAggregation::Sum);
+		//Column.SetAggregation(ETableColumnAggregation::Sum);
 
 		AddColumn(ColumnRef);
 	}
diff --git a/Engine/Source/Developer/TraceInsights/Private/Insights/MemoryProfiler/Widgets/SMemTagTreeView.cpp b/Engine/Source/Developer/TraceInsights/Private/Insights/MemoryProfiler/Widgets/SMemTagTreeView.cpp
--- a/Engine/Source/Developer/TraceInsights/Private/Insights/MemoryProfiler/Widgets/SMemTagTreeView.cpp
+++ b/Engine/Source/Developer/TraceInsights/Private/Insights/MemoryProfiler/Widgets/SMemTagTreeView.cpp
@@ -4,6 +4,7 @@
 
 #include "DesktopPlatformModule.h"
 #include "Framework/MultiBox/MultiBoxBuilder.h"
+#include "Misc/ConfigCacheIni.h"
 #include "SlateOptMacros.h"
 #include "Styling/StyleColors.h"
 #include "Widgets/Colors/SColorPicker.h"
@@ -22,6 +23,7 @@
 #include "Insights/MemoryProfiler/ViewModels/MemorySharedState.h"
 #include "Insights/MemoryProfiler/ViewModels/MemoryTag.h"
 #include "Insights/MemoryProfiler/ViewModels/MemoryTracker.h"
+#include "Insights/MemoryProfiler/ViewModels/MemTagBudget.h"
 #include "Insights/MemoryProfiler/ViewModels/MemTagNode.h"
 #include "Insights/MemoryProfiler/ViewModels/MemTagTableViewPresets.h"
 #include "Insights/MemoryProfiler/Widgets/SMemoryProfilerWindow.h"
@@ -63,6 +65,7 @@ void SMemTagTreeView::Construct(const FArguments& InArgs, TSharedPtr<SMemoryProf
 	ApplyViewPreset(*(*GetAvailableViewPresets())[0]);
 
 	UpdateSelectionStatsText();
+	InitBudgetOptions();
 }
 
 ////////////////////////////////////////////////////////////////////////////////////////////////////
@@ -144,7 +147,22 @@ END_SLATE_FUNCTION_BUILD_OPTIMIZATION
 BEGIN_SLATE_FUNCTION_BUILD_OPTIMIZATION
 TSharedPtr<SWidget> SMemTagTreeView::ConstructExtraToolbar()
 {
-	return ConstructTracksMiniToolbar();
+	TSharedPtr<SVerticalBox> Box = SNew(SVerticalBox)
+
+	+ SVerticalBox::Slot()
+	.AutoHeight()
+	[
+		ConstructTracksMiniToolbar()
+	]
+
+	+ SVerticalBox::Slot()
+	.AutoHeight()
+	.Padding(FMargin(0.0f, 4.0f, 0.0f, 0.0f))
+	[
+		ConstructBudgetSettingsToolbar()
+	];
+
+	return Box;
 }
 END_SLATE_FUNCTION_BUILD_OPTIMIZATION
 
@@ -264,7 +282,7 @@ FText SMemTagTreeView::TagSet_GetSelectedText() const
 BEGIN_SLATE_FUNCTION_BUILD_OPTIMIZATION
 TSharedRef<SWidget> SMemTagTreeView::ConstructTracksMiniToolbar()
 {
-	TSharedRef<SWidget> Widget = SNew(SHorizontalBox)
+	TSharedRef<SHorizontalBox> Box = SNew(SHorizontalBox)
 
 	+ SHorizontalBox::Slot()
 	.VAlign(VAlign_Center)
@@ -348,7 +366,95 @@ TSharedRef<SWidget> SMemTagTreeView::ConstructTracksMiniToolbar()
 		.ToolTip(LOCTEXT("LargeHeight_ToolTip", "Change height of LLM Tag Graph tracks to Large."))
 	];
 
-	return Widget;
+	return Box;
+}
+END_SLATE_FUNCTION_BUILD_OPTIMIZATION
+
+////////////////////////////////////////////////////////////////////////////////////////////////////
+
+BEGIN_SLATE_FUNCTION_BUILD_OPTIMIZATION
+TSharedRef<SWidget> SMemTagTreeView::ConstructBudgetSettingsToolbar()
+{
+	float constexpr TopPadding = 0.0f;
+
+	TSharedRef<SHorizontalBox> Box = SNew(SHorizontalBox)
+
+	+ SHorizontalBox::Slot()
+	.AutoWidth()
+	.VAlign(VAlign_Center)
+	.Padding(0.0f, TopPadding, 4.0f, 0.0f)
+	[
+		SNew(STextBlock)
+		.Text(LOCTEXT("BudgetFileText", "Budget File:"))
+	]
+
+	+ SHorizontalBox::Slot()
+	.FillWidth(1.0f)
+	.VAlign(VAlign_Center)
+	.Padding(0.0f, TopPadding, 0.0f, 0.0f)
+	[
+		SAssignNew(BudgetFileComboBox, SComboBox<TSharedPtr<FMemTagBudgetFileDesc>>)
+		.ToolTipText(this, &SMemTagTreeView::BudgetFile_GetTooltipText)
+		.OptionsSource(&AvailableBudgetFiles)
+		.OnSelectionChanged(this, &SMemTagTreeView::BudgetFile_OnSelectionChanged)
+		.OnGenerateWidget(this, &SMemTagTreeView::BudgetFile_OnGenerateWidget)
+		[
+			SNew(STextBlock)
+			.Text(this, &SMemTagTreeView::BudgetFile_GetSelectedText)
+		]
+	]
+
+	+ SHorizontalBox::Slot()
+	.AutoWidth()
+	.VAlign(VAlign_Center)
+	.Padding(6.0f, TopPadding, 4.0f, 0.0f)
+	[
+		SNew(STextBlock)
+		.Text(LOCTEXT("BudgetModeText", "Mode:"))
+	]
+
+	+ SHorizontalBox::Slot()
+	.AutoWidth()
+	.VAlign(VAlign_Center)
+	.Padding(0.0f, TopPadding, 0.0f, 0.0f)
+	[
+		SAssignNew(BudgetModeComboBox, SComboBox<TSharedPtr<FMemTagBudgetModeDesc>>)
+		.ToolTipText(this, &SMemTagTreeView::BudgetMode_GetTooltipText)
+		.OptionsSource(&AvailableBudgetModes)
+		.OnSelectionChanged(this, &SMemTagTreeView::BudgetMode_OnSelectionChanged)
+		.OnGenerateWidget(this, &SMemTagTreeView::BudgetMode_OnGenerateWidget)
+		[
+			SNew(STextBlock)
+			.Text(this, &SMemTagTreeView::BudgetMode_GetSelectedText)
+		]
+	]
+
+	+ SHorizontalBox::Slot()
+	.AutoWidth()
+	.VAlign(VAlign_Center)
+	.Padding(6.0f, TopPadding, 4.0f, 0.0f)
+	[
+		SNew(STextBlock)
+		.Text(LOCTEXT("PlatformModeText", "Platform:"))
+	]
+
+	+ SHorizontalBox::Slot()
+	.AutoWidth()
+	.VAlign(VAlign_Center)
+	.Padding(0.0f, TopPadding, 0.0f, 0.0f)
+	[
+		SAssignNew(BudgetPlatformComboBox, SComboBox<TSharedPtr<FMemTagBudgetPlatformDesc>>)
+		.ToolTipText(this, &SMemTagTreeView::BudgetPlatform_GetTooltipText)
+		.OptionsSource(&AvailableBudgetPlatforms)
+		.OnSelectionChanged(this, &SMemTagTreeView::BudgetPlatform_OnSelectionChanged)
+		.OnGenerateWidget(this, &SMemTagTreeView::BudgetPlatform_OnGenerateWidget)
+		[
+			SNew(STextBlock)
+			.Text(this, &SMemTagTreeView::BudgetPlatform_GetSelectedText)
+		]
+	];
+
+	return Box;
 }
 END_SLATE_FUNCTION_BUILD_OPTIMIZATION
 
@@ -749,6 +855,8 @@ void SMemTagTreeView::RebuildTree(bool bResync)
 					MemTagNode.SetParentTagNode(ParentNodePtr);
 				}
 			}
+
+			ApplyBudgetToNodes();
 		}
 
 		const auto& TagSets = SharedState.GetTagSets();
@@ -894,11 +1002,11 @@ void SMemTagTreeView::UpdateStatsInternal()
 				continue;
 			}
 			FMemTagNode& MemTagNode = NodePtr->As<FMemTagNode>();
-			MemTagNode.ResetAggregatedStats();
 
 			FMemoryTag* MemTag = MemTagNode.GetMemTag();
 			if (!MemTag)
 			{
+				MemTagNode.ResetAggregatedStats();
 				continue;
 			}
 
@@ -917,6 +1025,10 @@ void SMemTagTreeView::UpdateStatsInternal()
 			{
 				MemTagNode.GetStats().SizeA = CurrentStats.SizeA;
 			}
+			else
+			{
+				MemTagNode.GetStats().SizeA = 0;
+			}
 
 			CurrentStats.SizeB = std::numeric_limits<int64>::max();
 			MemoryProvider->EnumerateTagSamples(MemTag->GetTrackerId(), MemTag->GetId(), TimeMarkerB, TimeMarkerB, true,
@@ -931,6 +1043,10 @@ void SMemTagTreeView::UpdateStatsInternal()
 			{
 				MemTagNode.GetStats().SizeB = CurrentStats.SizeB;
 			}
+			else
+			{
+				MemTagNode.GetStats().SizeB = 0;
+			}
 
 			CurrentStats.SampleCount = 0;
 			CurrentStats.SizeMin = std::numeric_limits<int64>::max();
@@ -956,7 +1072,17 @@ void SMemTagTreeView::UpdateStatsInternal()
 				MemTagNode.GetStats().SizeMax = CurrentStats.SizeMax;
 				MemTagNode.GetStats().SizeAverage = CurrentStats.SizeAverage / CurrentStats.SampleCount;
 			}
+			else
+			{
+				MemTagNode.GetStats().SampleCount = 0;
+				MemTagNode.GetStats().SizeMin = 0;
+				MemTagNode.GetStats().SizeMax = 0;
+				MemTagNode.GetStats().SizeAverage = 0;
+			}
 		}
+
+		UpdateSelectionStatsText();
+		UpdateAggregatedValuesRec(*Root);
 	}
 }
 
@@ -1199,7 +1325,7 @@ void SMemTagTreeView::LoadReportXML()
 	TSharedPtr<SMemoryProfilerWindow> ProfilerWindow = GetProfilerWindow();
 	if (ProfilerWindow.IsValid())
 	{
-		TArray<FString> OutFiles;
+		TArray<FString> Files;
 		bool bOpened = false;
 
 		IDesktopPlatform* DesktopPlatform = FDesktopPlatformModule::Get();
@@ -1218,15 +1344,15 @@ void SMemTagTreeView::LoadReportXML()
 				DefaultFile, // Not actually used. See FDesktopPlatformWindows::FileDialogShared implementation. :(
 				LOCTEXT("LoadReportXML_FileFilter", "XML files (*.xml)|*.xml|All files (*.*)|*.*").ToString(),
 				EFileDialogFlags::None,
-				OutFiles
+				Files
 			);
 		}
 
-		if (bOpened == true && OutFiles.Num() == 1)
+		if (bOpened && Files.Num() == 1)
 		{
 			FMemorySharedState& SharedState = ProfilerWindow->GetSharedState();
 			SharedState.RemoveAllMemTagGraphTracks();
-			SharedState.CreateTracksFromReport(OutFiles[0]);
+			SharedState.CreateTracksFromReport(Files[0]);
 		}
 	}
 }
@@ -1612,6 +1738,461 @@ void SMemTagTreeView::ColorPickerCancelled(FLinearColor OriginalColor)
 	SetEditableColor(OriginalColor);
 }
 
+////////////////////////////////////////////////////////////////////////////////////////////////////
+// Budget File
+////////////////////////////////////////////////////////////////////////////////////////////////////
+
+void SMemTagTreeView::BudgetFile_OnSelectionChanged(TSharedPtr<FMemTagBudgetFileDesc> InBudgetFile, ESelectInfo::Type SelectInfo)
+{
+	if (InBudgetFile.IsValid() &&
+		InBudgetFile->Type == FMemTagBudgetFileDesc::EType::OpenFileDialog)
+	{
+		TArray<FString> Files;
+		bool bOpened = false;
+
+		IDesktopPlatform* DesktopPlatform = FDesktopPlatformModule::Get();
+		if (DesktopPlatform != nullptr)
+		{
+			FSlateApplication::Get().CloseToolTip();
+
+			const FString DefaultPath(FPaths::RootDir() / TEXT("Engine/Programs/UnrealInsights/Config"));
+
+			bOpened = DesktopPlatform->OpenFileDialog
+			(
+				FSlateApplication::Get().FindBestParentWindowHandleForDialogs(nullptr),
+				LOCTEXT("LoadBudgetFile_FileDesc", "Open the Memory Budget .xml file...").ToString(),
+				DefaultPath,
+				FString(),
+				LOCTEXT("LoadBudgetFile_FileFilter", "XML files (*.xml)|*.xml|All files (*.*)|*.*").ToString(),
+				EFileDialogFlags::None,
+				Files
+			);
+		}
+
+		if (bOpened && Files.Num() == 1)
+		{
+			InBudgetFile = MakeShared<FMemTagBudgetFileDesc>(FPaths::GetCleanFilename(Files[0]), Files[0]);
+			bool bIsAlreadyAvailable = false;
+			for (int32 BudgetFileIndex = 0; BudgetFileIndex < AvailableBudgetFiles.Num() - 1; ++BudgetFileIndex)
+			{
+				if (*InBudgetFile == *AvailableBudgetFiles[BudgetFileIndex])
+				{
+					InBudgetFile = AvailableBudgetFiles[BudgetFileIndex];
+					bIsAlreadyAvailable = true;
+				}
+			}
+			if (!bIsAlreadyAvailable)
+			{
+				AvailableBudgetFiles.Insert(InBudgetFile, AvailableBudgetFiles.Num() - 1);
+			}
+			BudgetFileComboBox->SetSelectedItem(InBudgetFile);
+		}
+		else
+		{
+			// No file opened, revert combo box selection
+			InBudgetFile = SelectedBudgetFile;
+			BudgetFileComboBox->SetSelectedItem(SelectedBudgetFile);
+		}
+	}
+
+	if (InBudgetFile.IsValid())
+	{
+		if (!SelectedBudgetFile.IsValid() || *SelectedBudgetFile != *InBudgetFile)
+		{
+			SelectedBudgetFile = InBudgetFile;
+			CurrentBudget = MakeShared<FMemTagBudget>(InBudgetFile->Name);
+			if (!CurrentBudget->LoadFromFile(InBudgetFile->FileName))
+			{
+				// error
+			}
+			OnBudgetChanged();
+		}
+	}
+	else if (SelectedBudgetFile.IsValid())
+	{
+		SelectedBudgetFile = nullptr;
+		CurrentBudget = nullptr;
+		OnBudgetChanged();
+	}
+}
+
+////////////////////////////////////////////////////////////////////////////////////////////////////
+
+TSharedRef<SWidget> SMemTagTreeView::BudgetFile_OnGenerateWidget(TSharedPtr<FMemTagBudgetFileDesc> InBudgetFile)
+{
+	if (InBudgetFile.IsValid())
+	{
+		if (InBudgetFile->Type == FMemTagBudgetFileDesc::EType::NamedBudgetFile)
+		{
+			return SNew(STextBlock)
+				.Text(FText::FromString(FString::Printf(TEXT("%s (%s)"), *InBudgetFile->Name, *InBudgetFile->FileName)));
+				//.Text(FText::FromString(InBudgetFile->Name))
+				//.ToolTip(FSlateApplicationBase::Get().MakeToolTip(FText::FromString(InBudgetFile->FileName)));
+		}
+		else if (InBudgetFile->Type == FMemTagBudgetFileDesc::EType::OpenFileDialog)
+		{
+			return SNew(STextBlock)
+				.Text(LOCTEXT("BudgetFileComboBox_OpenFileOption", "Open file..."));
+		}
+	}
+	return SNew(SBox);
+}
+
+////////////////////////////////////////////////////////////////////////////////////////////////////
+
+FText SMemTagTreeView::BudgetFile_GetSelectedText() const
+{
+	if (SelectedBudgetFile.IsValid())
+	{
+		if (SelectedBudgetFile->Type == FMemTagBudgetFileDesc::EType::NamedBudgetFile)
+		{
+			return FText::FromString(*SelectedBudgetFile->Name);
+		}
+		else if (SelectedBudgetFile->Type == FMemTagBudgetFileDesc::EType::OpenFileDialog)
+		{
+			return LOCTEXT("BudgetFileComboBox_OpenFileOption2", "Choose budget file...");
+		}
+	}
+	return FText::GetEmpty();
+}
+
+////////////////////////////////////////////////////////////////////////////////////////////////////
+
+FText SMemTagTreeView::BudgetFile_GetTooltipText() const
+{
+	return LOCTEXT("BudgetFileComboBox_Tooltip", "Choose the current memory budget file");
+}
+
+////////////////////////////////////////////////////////////////////////////////////////////////////
+// Budget Mode
+////////////////////////////////////////////////////////////////////////////////////////////////////
+
+void SMemTagTreeView::BudgetMode_OnSelectionChanged(TSharedPtr<FMemTagBudgetModeDesc> InBudgetMode, ESelectInfo::Type SelectInfo)
+{
+	if (InBudgetMode.IsValid())
+	{
+		if (!SelectedBudgetMode.IsValid() || !SelectedBudgetMode->Name.Equals(InBudgetMode->Name, ESearchCase::IgnoreCase))
+		{
+			SelectedBudgetMode = InBudgetMode;
+			ApplyBudgetToNodes();
+		}
+	}
+}
+
+////////////////////////////////////////////////////////////////////////////////////////////////////
+
+TSharedRef<SWidget> SMemTagTreeView::BudgetMode_OnGenerateWidget(TSharedPtr<FMemTagBudgetModeDesc> InBudgetMode)
+{
+	return SNew(STextBlock)
+		.Text(FText::FromString(InBudgetMode->Name));
+}
+
+////////////////////////////////////////////////////////////////////////////////////////////////////
+
+FText SMemTagTreeView::BudgetMode_GetSelectedText() const
+{
+	if (SelectedBudgetMode.IsValid())
+	{
+		return FText::FromString(SelectedBudgetMode->Name);
+	}
+	return FText::GetEmpty();
+}
+
+////////////////////////////////////////////////////////////////////////////////////////////////////
+
+FText SMemTagTreeView::BudgetMode_GetTooltipText() const
+{
+	return LOCTEXT("BudgetModeComboBox_Tooltip", "Choose the current memory budget mode");
+}
+
+////////////////////////////////////////////////////////////////////////////////////////////////////
+// Budget Platform
+////////////////////////////////////////////////////////////////////////////////////////////////////
+
+void SMemTagTreeView::BudgetPlatform_OnSelectionChanged(TSharedPtr<FMemTagBudgetPlatformDesc> InBudgetPlatform, ESelectInfo::Type SelectInfo)
+{
+	if (InBudgetPlatform.IsValid())
+	{
+		if (!SelectedBudgetPlatform.IsValid() || !SelectedBudgetPlatform->Name.Equals(InBudgetPlatform->Name, ESearchCase::IgnoreCase))
+		{
+			SelectedBudgetPlatform = InBudgetPlatform;
+			ApplyBudgetToNodes();
+		}
+	}
+}
+
+////////////////////////////////////////////////////////////////////////////////////////////////////
+
+TSharedRef<SWidget> SMemTagTreeView::BudgetPlatform_OnGenerateWidget(TSharedPtr<FMemTagBudgetPlatformDesc> InBudgetPlatform)
+{
+	return SNew(STextBlock)
+		.Text(FText::FromString(InBudgetPlatform->Name));
+}
+
+////////////////////////////////////////////////////////////////////////////////////////////////////
+
+FText SMemTagTreeView::BudgetPlatform_GetSelectedText() const
+{
+	if (SelectedBudgetPlatform.IsValid())
+	{
+		return FText::FromString(SelectedBudgetPlatform->Name);
+	}
+	return FText::GetEmpty();
+}
+
+////////////////////////////////////////////////////////////////////////////////////////////////////
+
+FText SMemTagTreeView::BudgetPlatform_GetTooltipText() const
+{
+	return LOCTEXT("BudgetPlatformComboBox_Tooltip", "Choose the current memory budget platform");
+}
+
+////////////////////////////////////////////////////////////////////////////////////////////////////
+// Budget
+////////////////////////////////////////////////////////////////////////////////////////////////////
+
+void SMemTagTreeView::InitBudgetOptions()
+{
+	CurrentBudget.Reset();
+
+	AvailableBudgetFiles.Reset();
+
+	FConfigFile* Settings = GConfig->FindConfigFileWithBaseName(TEXT("UnrealInsightsSettings"));
+	if (Settings != nullptr)
+	{
+		const FConfigSection* Section = Settings->FindSection(TEXT("Insights.MemoryProfiler"));
+		if (Section != nullptr)
+		{
+			FName KeyName(TEXT("BudgetFilePath"));
+			for (FConfigSection::TConstKeyIterator It(*Section, KeyName); It; ++It)
+			{
+				const FString& ValueString = It.Value().GetValue();
+
+				FString Label;
+				FParse::Value(*ValueString, TEXT("Label="), Label);
+
+				FString Path;
+				FParse::Value(*ValueString, TEXT("Path="), Path);
+
+				FString FullPath = FPaths::Combine(*FPaths::RootDir(), Path);
+				if (FPaths::FileExists(FullPath))
+				{
+					AvailableBudgetFiles.Insert(MakeShared<FMemTagBudgetFileDesc>(Label, FullPath), 0);
+					UE_LOG(LogMemoryProfiler, Log, TEXT("Budget file found (Label=\"%s\", Path=\"%s\")"), *Label, *FullPath);
+				}
+				else
+				{
+					UE_LOG(LogMemoryProfiler, Warning, TEXT("Budget file not found (Label=\"%s\", Path=\"%s\")"), *Label, *FullPath);
+				}
+			}
+		}
+		else
+		{
+			UE_LOG(LogMemoryProfiler, Log, TEXT("Couldn't find Insights.MemoryProfiler config section"));
+		}
+	}
+	else
+	{
+		UE_LOG(LogMemoryProfiler, Log, TEXT("Couldn't find UnrealInsightsSettings config"));
+	}
+
+	AvailableBudgetFiles.Add(MakeShared<FMemTagBudgetFileDesc>()); // open file dialog
+
+	AvailableBudgetModes.Reset();
+	AvailableBudgetPlatforms.Reset();
+	AvailablePlatforms.Reset();
+
+	if (AvailableBudgetFiles.Num() > 1)
+	{
+		BudgetFileComboBox->SetSelectedItem(AvailableBudgetFiles[0]);
+	}
+}
+
+////////////////////////////////////////////////////////////////////////////////////////////////////
+
+void SMemTagTreeView::OnBudgetChanged()
+{
+	bIsLoadingBudget = true;
+
+	// Save the selected Budget Mode.
+	TSharedPtr<FMemTagBudgetModeDesc> PrevSelectedBudgetMode = SelectedBudgetMode;
+
+	// Save the selected Budget Platform.
+	TSharedPtr<FMemTagBudgetPlatformDesc> PrevSelectedBudgetPlatform = SelectedBudgetPlatform;
+
+	AvailableBudgetModes.Reset();
+	AvailableBudgetPlatforms.Reset();
+	AvailablePlatforms.Reset();
+
+	if (CurrentBudget.IsValid())
+	{
+		// The "default" budget platform is always the first one.
+		AvailableBudgetPlatforms.Add(MakeShared<FMemTagBudgetPlatformDesc>(TEXT("Default")));
+		AvailablePlatforms.Add(AvailableBudgetPlatforms[0]->Name);
+
+		// Add modes and platforms available in the current budget.
+		for (const auto& BudgetMode : CurrentBudget->GetModes())
+		{
+			const FString& ModeName = BudgetMode->GetName();
+			AvailableBudgetModes.Add(MakeShared<FMemTagBudgetModeDesc>(ModeName));
+
+			TArray<FString> Platforms;
+			BudgetMode->GetPlatforms(Platforms);
+			for (const FString& Platform : Platforms)
+			{
+				bool bIsAlreadyInSet = false;
+				AvailablePlatforms.Add(Platform, &bIsAlreadyInSet);
+				if (!bIsAlreadyInSet)
+				{
+					AvailableBudgetPlatforms.Add(MakeShared<FMemTagBudgetPlatformDesc>(Platform));
+				}
+			}
+		}
+	}
+
+	// Restore selection for Budget Mode.
+	if (PrevSelectedBudgetMode.IsValid())
+	{
+		bool bFound = false;
+		for (const auto& BudgetMode : AvailableBudgetModes)
+		{
+			if (PrevSelectedBudgetMode->Name.Equals(BudgetMode->Name, ESearchCase::CaseSensitive))
+			{
+				BudgetModeComboBox->SetSelectedItem(BudgetMode);
+				bFound = true;
+				break;
+			}
+		}
+		if (!bFound)
+		{
+			if (AvailableBudgetModes.Num() > 0)
+			{
+				// Select first available mode.
+				BudgetModeComboBox->SetSelectedItem(AvailableBudgetModes[0]);
+			}
+			else
+			{
+				BudgetModeComboBox->ClearSelection();
+			}
+		}
+	}
+	if (!SelectedBudgetMode.IsValid() && AvailableBudgetModes.Num() > 0)
+	{
+		// Select the first available budget mode.
+		BudgetModeComboBox->SetSelectedItem(AvailableBudgetModes[0]);
+	}
+
+	// Restore selection for Budget Platform.
+	if (PrevSelectedBudgetPlatform.IsValid())
+	{
+		bool bFound = false;
+		for (const auto& BudgetPlatform : AvailableBudgetPlatforms)
+		{
+			if (PrevSelectedBudgetPlatform->Name.Equals(BudgetPlatform->Name, ESearchCase::CaseSensitive))
+			{
+				BudgetPlatformComboBox->SetSelectedItem(BudgetPlatform);
+				bFound = true;
+				break;
+			}
+		}
+		if (!bFound)
+		{
+			if (AvailableBudgetPlatforms.Num() > 0)
+			{
+				// Select the first available platform.
+				BudgetPlatformComboBox->SetSelectedItem(AvailableBudgetPlatforms[0]);
+			}
+			else
+			{
+				BudgetPlatformComboBox->ClearSelection();
+			}
+		}
+	}
+	if (!SelectedBudgetPlatform.IsValid() && AvailableBudgetPlatforms.Num() > 0)
+	{
+		// Select the first available platform.
+		BudgetPlatformComboBox->SetSelectedItem(AvailableBudgetPlatforms[0]);
+	}
+
+	bIsLoadingBudget = false;
+	ApplyBudgetToNodes();
+}
+
+////////////////////////////////////////////////////////////////////////////////////////////////////
+
+void SMemTagTreeView::ApplyBudgetToNodes()
+{
+	if (bIsLoadingBudget)
+	{
+		return;
+	}
+
+	if (!CurrentBudget.IsValid())
+	{
+		ResetBudgetForAllNodes();
+		return;
+	}
+
+	const FMemTagBudgetMode* FoundBudgetMode = SelectedBudgetMode.IsValid() ? CurrentBudget->FindMode(SelectedBudgetMode->Name) : nullptr;
+	if (!FoundBudgetMode)
+	{
+		ResetBudgetForAllNodes();
+		return;
+	}
+	const FMemTagBudgetMode& BudgetMode = *FoundBudgetMode;
+
+	const FMemTagBudgetValueSet& DefaultValues = BudgetMode.GetDefaultValues();
+	const FMemTagBudgetValueSet* PlatformOverrides = SelectedBudgetPlatform.IsValid() ? BudgetMode.FindPlatform(SelectedBudgetPlatform->Name) : nullptr;
+
+	for (const FTableTreeNodePtr& NodePtr : TableRowNodes)
+	{
+		if (!NodePtr->Is<FMemTagNode>())
+		{
+			continue;
+		}
+		FMemTagNode& MemTagNode = NodePtr->As<FMemTagNode>();
+		FMemoryTag* MemoryTag = MemTagNode.GetMemTag();
+		if (!MemoryTag)
+		{
+			continue;
+		}
+		if (PlatformOverrides)
+		{
+			const int64* ValuePtr = PlatformOverrides->FindValue(MemoryTag->GetStatFullName());
+			if (ValuePtr)
+			{
+				MemTagNode.GetStats().SizeBudget = *ValuePtr;
+				continue;
+			}
+		}
+		// DefaultValues
+		{
+			const int64* ValuePtr = DefaultValues.FindValue(MemoryTag->GetStatFullName());
+			if (ValuePtr)
+			{
+				MemTagNode.GetStats().SizeBudget = *ValuePtr;
+				continue;
+			}
+		}
+		MemTagNode.GetStats().SizeBudget = 0;
+	}
+}
+
+////////////////////////////////////////////////////////////////////////////////////////////////////
+
+void SMemTagTreeView::ResetBudgetForAllNodes()
+{
+	for (const FTableTreeNodePtr& NodePtr : TableRowNodes)
+	{
+		if (!NodePtr->Is<FMemTagNode>())
+		{
+			continue;
+		}
+		FMemTagNode& MemTagNode = NodePtr->As<FMemTagNode>();
+		MemTagNode.GetStats().SizeBudget = 0;
+	}
+}
+
 ////////////////////////////////////////////////////////////////////////////////////////////////////
 
 } // namespace UE::Insights::MemoryProfiler
diff --git a/Engine/Source/Developer/TraceInsights/Private/Insights/MemoryProfiler/Widgets/SMemTagTreeView.h b/Engine/Source/Developer/TraceInsights/Private/Insights/MemoryProfiler/Widgets/SMemTagTreeView.h
--- a/Engine/Source/Developer/TraceInsights/Private/Insights/MemoryProfiler/Widgets/SMemTagTreeView.h
+++ b/Engine/Source/Developer/TraceInsights/Private/Insights/MemoryProfiler/Widgets/SMemTagTreeView.h
@@ -12,11 +12,65 @@
 namespace UE::Insights::MemoryProfiler
 {
 
+class FMemTagBudget;
 class FMemTagNode;
 class FMemoryGraphTrack;
 class FMemorySharedState;
 class SMemoryProfilerWindow;
 
+////////////////////////////////////////////////////////////////////////////////////////////////////
+
+struct FMemTagBudgetFileDesc
+{
+	enum class EType
+	{
+		NamedBudgetFile,
+		OpenFileDialog
+	};
+
+	FMemTagBudgetFileDesc()
+		: Type(EType::OpenFileDialog)
+	{
+	}
+
+	FMemTagBudgetFileDesc(const FString& InName, const FString& InFileName)
+		: Type(EType::NamedBudgetFile)
+		, Name(InName)
+		, FileName(InFileName)
+	{
+	}
+
+	bool operator==(const FMemTagBudgetFileDesc& Other) const
+	{
+		return (Type == Other.Type)
+			&& Name.Equals(Other.Name, ESearchCase::IgnoreCase)
+			&& FileName.Equals(Other.FileName, ESearchCase::IgnoreCase);
+	}
+
+	bool operator!=(const FMemTagBudgetFileDesc& Other) const
+	{
+		return !(*this == Other);
+	}
+
+	EType Type;
+	FString Name;
+	FString FileName;
+};
+
+////////////////////////////////////////////////////////////////////////////////////////////////////
+
+struct FMemTagBudgetModeDesc
+{
+	FString Name;
+};
+
+////////////////////////////////////////////////////////////////////////////////////////////////////
+
+struct FMemTagBudgetPlatformDesc
+{
+	FString Name;
+};
+
 ////////////////////////////////////////////////////////////////////////////////////////////////////
 /**
  * A custom widget used to display the list of LLM tags and their aggregated stats.
@@ -114,8 +168,8 @@ protected:
 	virtual void TreeView_OnSelectionChanged(FTableTreeNodePtr SelectedItem, ESelectInfo::Type SelectInfo) override;
 
 private:
-	TSharedRef<SWidget> ConstructFilteringWidgetArea();
 	TSharedRef<SWidget> ConstructTracksMiniToolbar();
+	TSharedRef<SWidget> ConstructBudgetSettingsToolbar();
 
 	void UpdateSelectionStatsText();
 	FText GetNumSelectedTagsText() const { return NumSelectedTagsText; }
@@ -197,6 +251,29 @@ private:
 	void SetEditableColor(FLinearColor NewColor);
 	void ColorPickerCancelled(FLinearColor OriginalColor);
 
+	////////////////////////////////////////////////////////////////////////////////////////////////////
+	// Budgets
+
+	void BudgetFile_OnSelectionChanged(TSharedPtr<FMemTagBudgetFileDesc> InBudgetFile, ESelectInfo::Type SelectInfo);
+	TSharedRef<SWidget> BudgetFile_OnGenerateWidget(TSharedPtr<FMemTagBudgetFileDesc> InBudgetFile);
+	FText BudgetFile_GetSelectedText() const;
+	FText BudgetFile_GetTooltipText() const;
+
+	void BudgetMode_OnSelectionChanged(TSharedPtr<FMemTagBudgetModeDesc> InBudgetMode, ESelectInfo::Type SelectInfo);
+	TSharedRef<SWidget> BudgetMode_OnGenerateWidget(TSharedPtr<FMemTagBudgetModeDesc> InBudgetMode);
+	FText BudgetMode_GetSelectedText() const;
+	FText BudgetMode_GetTooltipText() const;
+
+	void BudgetPlatform_OnSelectionChanged(TSharedPtr<FMemTagBudgetPlatformDesc> InPlatform, ESelectInfo::Type SelectInfo);
+	TSharedRef<SWidget> BudgetPlatform_OnGenerateWidget(TSharedPtr<FMemTagBudgetPlatformDesc> InPlatform);
+	FText BudgetPlatform_GetSelectedText() const;
+	FText BudgetPlatform_GetTooltipText() const;
+
+	void InitBudgetOptions();
+	void OnBudgetChanged();
+	void ApplyBudgetToNodes();
+	void ResetBudgetForAllNodes();
+
 private:
 	/** A weak pointer to the Memory Insights window. */
 	TWeakPtr<SMemoryProfilerWindow> ProfilerWindowWeakPtr;
@@ -234,6 +311,25 @@ private:
 
 	//////////////////////////////////////////////////
 
+	TSharedPtr<FMemTagBudget> CurrentBudget;
+
+	TSharedPtr<SComboBox<TSharedPtr<FMemTagBudgetFileDesc>>> BudgetFileComboBox;
+	TSharedPtr<SComboBox<TSharedPtr<FMemTagBudgetModeDesc>>> BudgetModeComboBox;
+	TSharedPtr<SComboBox<TSharedPtr<FMemTagBudgetPlatformDesc>>> BudgetPlatformComboBox;
+
+	TArray<TSharedPtr<FMemTagBudgetFileDesc>> AvailableBudgetFiles;
+	TArray<TSharedPtr<FMemTagBudgetModeDesc>> AvailableBudgetModes;
+	TArray<TSharedPtr<FMemTagBudgetPlatformDesc>> AvailableBudgetPlatforms;
+	TSet<FString> AvailablePlatforms;
+
+	TSharedPtr<FMemTagBudgetFileDesc> SelectedBudgetFile;
+	TSharedPtr<FMemTagBudgetModeDesc> SelectedBudgetMode;
+	TSharedPtr<FMemTagBudgetPlatformDesc> SelectedBudgetPlatform;
+
+	bool bIsLoadingBudget = false;
+
+	//////////////////////////////////////////////////
+
 	EInvestigationMode Mode = EInvestigationMode::Default;
 
 	// A, B, B-A
diff --git a/Engine/Source/Developer/TraceInsightsCore/Private/InsightsCore/Table/ViewModels/TableCellValueFormatter.cpp b/Engine/Source/Developer/TraceInsightsCore/Private/InsightsCore/Table/ViewModels/TableCellValueFormatter.cpp
--- a/Engine/Source/Developer/TraceInsightsCore/Private/InsightsCore/Table/ViewModels/TableCellValueFormatter.cpp
+++ b/Engine/Source/Developer/TraceInsightsCore/Private/InsightsCore/Table/ViewModels/TableCellValueFormatter.cpp
@@ -128,7 +128,7 @@ FText FInt64ValueFormatterAsUInt32InfinteNumber::FormatValue(const TOptional<FTa
 		const int64 Value = InValue.GetValue().Int64;
 		if (Value == int64(uint32(~0)))
 		{
-			return LOCTEXT("AsUInt32InfinteNumber_Inf", "∞");
+			return LOCTEXT("AsUInt32InfiniteNumber_Inf", "∞");
 		}
 		return FText::AsNumber(Value);
 	}
@@ -144,7 +144,7 @@ FText FInt64ValueFormatterAsUInt32InfinteNumber::FormatValueForTooltip(const TOp
 		const int64 Value = InValue.GetValue().Int64;
 		if (Value == int64(uint32(~0)))
 		{
-			return FText::Format(LOCTEXT("AsUInt32InfinteNumber_Inf_Fmt", "{0} (∞)"), FText::AsNumber(Value));
+			return FText::Format(LOCTEXT("AsUInt32InfiniteNumber_Inf_Fmt", "{0} (∞)"), FText::AsNumber(Value));
 		}
 		return FText::AsNumber(Value);
 	}
diff --git a/Engine/Source/Developer/TraceInsightsCore/Private/InsightsCore/Table/Widgets/STableTreeViewCell.cpp b/Engine/Source/Developer/TraceInsightsCore/Private/InsightsCore/Table/Widgets/STableTreeViewCell.cpp
--- a/Engine/Source/Developer/TraceInsightsCore/Private/InsightsCore/Table/Widgets/STableTreeViewCell.cpp
+++ b/Engine/Source/Developer/TraceInsightsCore/Private/InsightsCore/Table/Widgets/STableTreeViewCell.cpp
@@ -50,7 +50,7 @@ void STableTreeViewCell::Construct(const FArguments& InArgs, const TSharedRef<IT
 
 TSharedRef<SWidget> STableTreeViewCell::GenerateWidgetForColumn(const FArguments& InArgs)
 {
-	if (InArgs._IsNameColumn)
+	if (InArgs._ColumnPtr.IsValid() && InArgs._ColumnPtr->IsHierarchy())
 	{
 		return GenerateWidgetForNameColumn(InArgs);
 	}
@@ -133,6 +133,12 @@ FText STableTreeViewCell::GetValueAsText() const
 
 TSharedRef<SWidget> STableTreeViewCell::GenerateWidgetForTableColumn(const FArguments& InArgs)
 {
+	TSharedPtr<SWidget> CustomWidget = ColumnPtr->GetValueFormatter()->GenerateCustomWidget(*ColumnPtr, *TableTreeNodePtr);
+	if (CustomWidget.IsValid())
+	{
+		return CustomWidget.ToSharedRef();
+	}
+
 	TSharedRef<STextBlock> TextBox = SNew(STextBlock)
 		.TextStyle(FInsightsCoreStyle::Get(), TEXT("TreeTable.NormalText"))
 		.ColorAndOpacity(this, &STableTreeViewCell::GetNormalTextColorAndOpacity)
@@ -150,12 +156,8 @@ TSharedRef<SWidget> STableTreeViewCell::GenerateWidgetForTableColumn(const FArgu
 
 	TSharedPtr<IToolTip> ColumnToolTip = ColumnPtr->GetValueFormatter()->GetCustomTooltip(*ColumnPtr, *TableTreeNodePtr);
 
-	return
-		SNew(SHorizontalBox)
+	return SNew(SBox)
 		.ToolTip(ColumnToolTip)
-
-		+ SHorizontalBox::Slot()
-		.FillWidth(1.0f)
 		.HAlign(ColumnPtr->GetHorizontalAlignment())
 		.VAlign(VAlign_Center)
 		.Padding(FMargin(2.0f, 0.0f, 2.0f, 0.0f))
diff --git a/Engine/Source/Developer/TraceInsightsCore/Private/InsightsCore/Table/Widgets/STableTreeViewRow.cpp b/Engine/Source/Developer/TraceInsightsCore/Private/InsightsCore/Table/Widgets/STableTreeViewRow.cpp
--- a/Engine/Source/Developer/TraceInsightsCore/Private/InsightsCore/Table/Widgets/STableTreeViewRow.cpp
+++ b/Engine/Source/Developer/TraceInsightsCore/Private/InsightsCore/Table/Widgets/STableTreeViewRow.cpp
@@ -20,8 +20,6 @@ namespace UE::Insights
 
 ////////////////////////////////////////////////////////////////////////////////////////////////////
 
-BEGIN_SLATE_FUNCTION_BUILD_OPTIMIZATION
-
 void STableTreeViewRow::Construct(const FArguments& InArgs, const TSharedRef<STableViewBase>& InOwnerTableView)
 {
 	OnShouldBeEnabled = InArgs._OnShouldBeEnabled;
@@ -44,15 +42,13 @@ void STableTreeViewRow::Construct(const FArguments& InArgs, const TSharedRef<STa
 
 ////////////////////////////////////////////////////////////////////////////////////////////////////
 
+BEGIN_SLATE_FUNCTION_BUILD_OPTIMIZATION
 TSharedRef<SWidget> STableTreeViewRow::GenerateWidgetForColumn(const FName& ColumnId)
 {
-	TSharedPtr<FTableColumn> ColumnPtr = TablePtr->FindColumnChecked(ColumnId);
-
-	return
-		SNew(SOverlay)
+	return SNew(SOverlay)
 		.Visibility(EVisibility::SelfHitTestInvisible)
 
-		+SOverlay::Slot()
+		+ SOverlay::Slot()
 		.Padding(0.0f)
 		[
 			SNew(SImage)
@@ -60,7 +56,7 @@ TSharedRef<SWidget> STableTreeViewRow::GenerateWidgetForColumn(const FName& Colu
 			.ColorAndOpacity(this, &STableTreeViewRow::GetBackgroundColorAndOpacity)
 		]
 
-		+SOverlay::Slot()
+		+ SOverlay::Slot()
 		.Padding(0.0f)
 		[
 			SNew(SImage)
@@ -68,19 +64,28 @@ TSharedRef<SWidget> STableTreeViewRow::GenerateWidgetForColumn(const FName& Colu
 			.ColorAndOpacity(this, &STableTreeViewRow::GetOutlineColorAndOpacity)
 		]
 
-		+SOverlay::Slot()
+		+ SOverlay::Slot()
 		[
-			SNew(STableTreeViewCell, SharedThis(this))
-			.Visibility(this, &STableTreeViewRow::IsColumnVisible, ColumnId)
-			.TablePtr(TablePtr)
-			.ColumnPtr(ColumnPtr)
-			.TableTreeNodePtr(TableTreeNodePtr)
-			.HighlightText(HighlightText)
-			.IsNameColumn(ColumnPtr->IsHierarchy())
-			.OnSetHoveredCell(this, &STableTreeViewRow::OnSetHoveredCell)
+			CreateCellWidget(ColumnId)
 		];
 }
+END_SLATE_FUNCTION_BUILD_OPTIMIZATION
 
+////////////////////////////////////////////////////////////////////////////////////////////////////
+
+BEGIN_SLATE_FUNCTION_BUILD_OPTIMIZATION
+TSharedRef<SWidget> STableTreeViewRow::CreateCellWidget(FName ColumnId)
+{
+	TSharedRef<FTableColumn> ColumnPtr = TablePtr->FindColumnChecked(ColumnId);
+
+	return SNew(STableTreeViewCell, SharedThis(this))
+		.Visibility(this, &STableTreeViewRow::IsColumnVisible, ColumnId)
+		.TablePtr(TablePtr)
+		.ColumnPtr(ColumnPtr.ToSharedPtr())
+		.TableTreeNodePtr(TableTreeNodePtr)
+		.HighlightText(HighlightText)
+		.OnSetHoveredCell(this, &STableTreeViewRow::OnSetHoveredCell);
+}
 END_SLATE_FUNCTION_BUILD_OPTIMIZATION
 
 ////////////////////////////////////////////////////////////////////////////////////////////////////
diff --git a/Engine/Source/Developer/TraceInsightsCore/Public/InsightsCore/Table/ViewModels/TableCellValueFormatter.h b/Engine/Source/Developer/TraceInsightsCore/Public/InsightsCore/Table/ViewModels/TableCellValueFormatter.h
--- a/Engine/Source/Developer/TraceInsightsCore/Public/InsightsCore/Table/ViewModels/TableCellValueFormatter.h
+++ b/Engine/Source/Developer/TraceInsightsCore/Public/InsightsCore/Table/ViewModels/TableCellValueFormatter.h
@@ -5,10 +5,12 @@
 #include "CoreTypes.h"
 
 #include "Layout/Visibility.h"
+#include "Misc/Optional.h"
 
 #include "InsightsCore/Table/ViewModels/TableCellValue.h"
 
 class IToolTip;
+class SWidget;
 
 namespace UE::Insights
 {
@@ -31,6 +33,7 @@ public:
 	virtual FText CopyValue(const FTableColumn& Column, const FBaseTreeNode& Node) const = 0;
 	virtual FText CopyTooltip(const FTableColumn& Column, const FBaseTreeNode& Node) const = 0;
 
+	virtual TSharedPtr<SWidget> GenerateCustomWidget(const FTableColumn& Column, const FBaseTreeNode& Node) const = 0;
 	virtual TSharedPtr<IToolTip> GetCustomTooltip(const FTableColumn& Column, const FBaseTreeNode& Node) const = 0;
 };
 
@@ -52,6 +55,7 @@ public:
 	virtual FText CopyValue(const FTableColumn& Column, const FBaseTreeNode& Node) const override; // { return FormatValue(Column, Node); }
 	virtual FText CopyTooltip(const FTableColumn& Column, const FBaseTreeNode& Node) const override; // { return FormatValueForTooltip(Column, Node); }
 
+	virtual TSharedPtr<SWidget> GenerateCustomWidget(const FTableColumn& Column, const FBaseTreeNode& Node) const override { return nullptr; }
 	virtual TSharedPtr<IToolTip> GetCustomTooltip(const FTableColumn& Column, const FBaseTreeNode& Node) const override;
 
 	static EVisibility GetTooltipVisibility();
diff --git a/Engine/Source/Developer/TraceInsightsCore/Public/InsightsCore/Table/ViewModels/TableCellValueGetter.h b/Engine/Source/Developer/TraceInsightsCore/Public/InsightsCore/Table/ViewModels/TableCellValueGetter.h
--- a/Engine/Source/Developer/TraceInsightsCore/Public/InsightsCore/Table/ViewModels/TableCellValueGetter.h
+++ b/Engine/Source/Developer/TraceInsightsCore/Public/InsightsCore/Table/ViewModels/TableCellValueGetter.h
@@ -4,6 +4,8 @@
 
 #include "CoreTypes.h"
 
+#include "Misc/Optional.h"
+
 #include "InsightsCore/Table/ViewModels/TableCellValue.h"
 
 namespace UE::Insights
diff --git a/Engine/Source/Developer/TraceInsightsCore/Public/InsightsCore/Table/Widgets/STableTreeViewCell.h b/Engine/Source/Developer/TraceInsightsCore/Public/InsightsCore/Table/Widgets/STableTreeViewCell.h
--- a/Engine/Source/Developer/TraceInsightsCore/Public/InsightsCore/Table/Widgets/STableTreeViewCell.h
+++ b/Engine/Source/Developer/TraceInsightsCore/Public/InsightsCore/Table/Widgets/STableTreeViewCell.h
@@ -30,7 +30,6 @@ public:
 		SLATE_ARGUMENT(TSharedPtr<FTable>, TablePtr)
 		SLATE_ARGUMENT(TSharedPtr<FTableColumn>, ColumnPtr)
 		SLATE_ARGUMENT(FTableTreeNodePtr, TableTreeNodePtr)
-		SLATE_ARGUMENT(bool, IsNameColumn)
 	SLATE_END_ARGS()
 
 	void Construct(const FArguments& InArgs, const TSharedRef<ITableRow>& InTableRow);
diff --git a/Engine/Source/Developer/TraceInsightsCore/Public/InsightsCore/Table/Widgets/STableTreeViewRow.h b/Engine/Source/Developer/TraceInsightsCore/Public/InsightsCore/Table/Widgets/STableTreeViewRow.h
--- a/Engine/Source/Developer/TraceInsightsCore/Public/InsightsCore/Table/Widgets/STableTreeViewRow.h
+++ b/Engine/Source/Developer/TraceInsightsCore/Public/InsightsCore/Table/Widgets/STableTreeViewRow.h
@@ -66,6 +66,7 @@ public:
 	void InvalidateContent();
 
 protected:
+	TSharedRef<SWidget> CreateCellWidget(FName ColumnId);
 	FSlateColor GetBackgroundColorAndOpacity() const;
 	FSlateColor GetBackgroundColorAndOpacity(double Time) const;
 	FSlateColor GetOutlineColorAndOpacity() const;
