Fix for some compilers that does not auto generate a valid constructor for a struct.

diff --git a/Engine/Source/Developer/TraceInsights/Private/Insights/MemoryProfiler/Widgets/SMemTagTreeView.h b/Engine/Source/Developer/TraceInsights/Private/Insights/MemoryProfiler/Widgets/SMemTagTreeView.h
--- a/Engine/Source/Developer/TraceInsights/Private/Insights/MemoryProfiler/Widgets/SMemTagTreeView.h
+++ b/Engine/Source/Developer/TraceInsights/Private/Insights/MemoryProfiler/Widgets/SMemTagTreeView.h
@@ -61,6 +61,19 @@ struct FMemTagBudgetFileDesc
 
 struct FMemTagBudgetModeDesc
 {
+	FMemTagBudgetModeDesc() {}
+	FMemTagBudgetModeDesc(const FString& InName) : Name(InName) {}
+
+	bool operator==(const FMemTagBudgetModeDesc& Other) const
+	{
+		return Name.Equals(Other.Name, ESearchCase::IgnoreCase);
+	}
+
+	bool operator!=(const FMemTagBudgetModeDesc& Other) const
+	{
+		return !(*this == Other);
+	}
+
 	FString Name;
 };
 
@@ -68,6 +81,19 @@ struct FMemTagBudgetModeDesc
 
 struct FMemTagBudgetPlatformDesc
 {
+	FMemTagBudgetPlatformDesc() {}
+	FMemTagBudgetPlatformDesc(const FString& InName) : Name(InName) {}
+
+	bool operator==(const FMemTagBudgetPlatformDesc& Other) const
+	{
+		return Name.Equals(Other.Name, ESearchCase::IgnoreCase);
+	}
+
+	bool operator!=(const FMemTagBudgetPlatformDesc& Other) const
+	{
+		return !(*this == Other);
+	}
+
 	FString Name;
 };
 
