[Backout] - CL39185852
#rnx
Original CL Desc
-----------------------------------------------------------------
Insights Automation
- Viewer and hub test cases to be run nightly on CIS
- Automate generate traces, can run multiple clients at the same time and multiple platforms for each client as long as they have the same build platform
- Small adjustments to tests so that they can be run both non-editor and editor
#rnx

diff --git a/Engine/Build/InsightsTests.xml b/Engine/Build/InsightsTests.xml
deleted file mode 100644
diff --git a/Engine/Source/Developer/TraceInsights/Private/Insights/Tests/FunctionalTests/ExportCommandsTests.cpp b/Engine/Source/Developer/TraceInsights/Private/Insights/Tests/FunctionalTests/ExportCommandsTests.cpp
--- a/Engine/Source/Developer/TraceInsights/Private/Insights/Tests/FunctionalTests/ExportCommandsTests.cpp
+++ b/Engine/Source/Developer/TraceInsights/Private/Insights/Tests/FunctionalTests/ExportCommandsTests.cpp
@@ -66,12 +66,12 @@ bool FCommandsExportThreadsDataTest::RunTest(const FString& Parameters)
 {
 	IPlatformFile& PlatformFile = FPlatformFileManager::Get().GetPlatformFile();
 	const FString SourceTracePath = FPaths::RootDir() / TEXT("EngineTest/SourceAssets/Utrace/CommandsExportTest_5.4.utrace");
-	const FString StoreTracePath = TEXT("CommandsExportTest_5.4.utrace");
-	const FString TestResultsDirPath = TEXT("TestResults");
+	const FString StoreTracePath = TEXT("/tmp/CommandsExportTest_5.4.utrace");
+	const FString TestResultsDirPath = TEXT("/tmp/TestResults");
 	const FString CmdThreadLogPath = TestResultsDirPath / TEXT("Logs/cmd_treads.log");
 	const FString ExportThreadsReportPath = TestResultsDirPath / TEXT("SingleCommand/Threads.csv");
 	const FString ExportThreadsTask = TEXT("TimingInsights.ExportThreads ") + ExportThreadsReportPath;
-	double Timeout = 30.0;
+	double Timeout = 10.0;
 
 	UTEST_TRUE(TEXT("Trace in project exists"), PlatformFile.FileExists(*SourceTracePath));
 	PlatformFile.CopyFile(*StoreTracePath, *SourceTracePath);
@@ -97,12 +97,12 @@ bool FCommandsExportTimersDataTest::RunTest(const FString& Parameters)
 {
 	IPlatformFile& PlatformFile = FPlatformFileManager::Get().GetPlatformFile();
 	const FString SourceTracePath = FPaths::RootDir() / TEXT("EngineTest/SourceAssets/Utrace/CommandsExportTest_5.4.utrace");
-	const FString StoreTracePath = TEXT("CommandsExportTest_5.4.utrace");
-	const FString TestResultsDirPath = TEXT("TestResults");
+	const FString StoreTracePath = TEXT("/tmp/CommandsExportTest_5.4.utrace");
+	const FString TestResultsDirPath = TEXT("/tmp/TestResults");
 	const FString CmdTimersLogPath = TestResultsDirPath / TEXT("Logs/cmd_timers.log");
 	const FString ExportTimersReportPath = TestResultsDirPath / TEXT("SingleCommand/Timers.csv");
 	const FString ExportTimersTask = TEXT("TimingInsights.ExportTimers ") + ExportTimersReportPath;
-	double Timeout = 30.0;
+	double Timeout = 10.0;
 
 	UTEST_TRUE(TEXT("Trace in project exists"), PlatformFile.FileExists(*SourceTracePath));
 	PlatformFile.CopyFile(*StoreTracePath, *SourceTracePath);
@@ -128,12 +128,12 @@ bool FCommandsExportTimingEventsDataTest::RunTest(const FString& Parameters)
 {
 	IPlatformFile& PlatformFile = FPlatformFileManager::Get().GetPlatformFile();
 	const FString SourceTracePath = FPaths::RootDir() / TEXT("EngineTest/SourceAssets/Utrace/CommandsExportTest_5.4.utrace");
-	const FString StoreTracePath = TEXT("CommandsExportTest_5.4.utrace");
-	const FString TestResultsDirPath = TEXT("TestResults");
+	const FString StoreTracePath = TEXT("/tmp/CommandsExportTest_5.4.utrace");
+	const FString TestResultsDirPath = TEXT("/tmp/TestResults");
 	const FString CmdTimingEventsLogPath = TestResultsDirPath / TEXT("Logs/cmd_timing_events.log");
 	const FString ExportTimingEventsReportPath = TestResultsDirPath / TEXT("SingleCommand/TimingEvents.txt");
 	const FString ExportTimingEventsTask = TEXT("TimingInsights.ExportTimingEvents " + ExportTimingEventsReportPath);
-	double Timeout = 30.0;
+	double Timeout = 10.0;
 
 	UTEST_TRUE(TEXT("Trace in project exists"), PlatformFile.FileExists(*SourceTracePath));
 	PlatformFile.CopyFile(*StoreTracePath, *SourceTracePath);
@@ -159,12 +159,12 @@ bool FCommandsExportFilteredTimingEventsDataTest::RunTest(const FString& Paramet
 {
 	IPlatformFile& PlatformFile = FPlatformFileManager::Get().GetPlatformFile();
 	const FString SourceTracePath = FPaths::RootDir() / TEXT("EngineTest/SourceAssets/Utrace/CommandsExportTest_5.4.utrace");
-	const FString StoreTracePath = TEXT("CommandsExportTest_5.4.utrace");
-	const FString TestResultsDirPath = TEXT("TestResults");
+	const FString StoreTracePath = TEXT("/tmp/CommandsExportTest_5.4.utrace");
+	const FString TestResultsDirPath = TEXT("/tmp/TestResults");
 	const FString CmdTimingEventsNonDefaultLogPath = TestResultsDirPath / TEXT("Logs/cmd_timing_non_default.log");
 	const FString ExportTimingEventsNonDefaultReportPath = TestResultsDirPath / TEXT("SingleCommand/TimingEventsNonDefault.csv");
 	const FString ExportTimingEventsBorderedTask = TEXT("TimingInsights.ExportTimingEvents " + ExportTimingEventsNonDefaultReportPath + " -columns=ThreadId,ThreadName,TimerId,TimerName,StartTime,EndTime,Duration,Depth -threads=GameThread -timers=* -startTime=10 -endTime=20");
-	double Timeout = 30.0;
+	double Timeout = 10.0;
 
 	UTEST_TRUE(TEXT("Trace in project exists"), PlatformFile.FileExists(*SourceTracePath));
 	PlatformFile.CopyFile(*StoreTracePath, *SourceTracePath);
@@ -212,11 +212,11 @@ bool FCommandsExportMultipleExportCommandsTest::RunTest(const FString& Parameter
 	IPlatformFile& PlatformFile = FPlatformFileManager::Get().GetPlatformFile();
 	const FString SourceTracePath = FPaths::RootDir() / TEXT("EngineTest/SourceAssets/Utrace/CommandsExportTest_5.4.utrace");
 	const FString SourceExportPath = FPaths::RootDir() / TEXT("EngineTest/SourceAssets/Rsp/export.rsp");
-	const FString StoreTracePath = TEXT("CommandsExportTest_5.4.utrace");
-	const FString LogResultExportPath = TEXT("export.rsp");
-	const FString TestResultsDirPath = TEXT("TestResults");
+	const FString StoreTracePath = TEXT("/tmp/CommandsExportTest_5.4.utrace");
+	const FString LogResultExportPath = TEXT("/tmp/export.rsp");
+	const FString TestResultsDirPath = TEXT("/tmp/TestResults");
 	const FString CmdExportLogPath = TestResultsDirPath / TEXT("Logs/cmd_export.log");
-	double Timeout = 30.0;
+	double Timeout = 10.0;
 
 	UTEST_TRUE(TEXT("Trace in project exists"), PlatformFile.FileExists(*SourceTracePath));
 	UTEST_TRUE(TEXT("Export in project exists"), PlatformFile.FileExists(*SourceExportPath));
@@ -236,21 +236,21 @@ bool FCommandsExportMultipleExportCommandsTest::RunTest(const FString& Parameter
 
 	const TArray<FStringView> ExpectedThreadsElementsRsp =
 	{
-		TEXTVIEW("/TestResults/RSPtest/CSV/Threads_rsp.csv"),
-		TEXTVIEW("/TestResults/RSPtest/TSV/Threads_rsp.tsv"),
-		TEXTVIEW("/TestResults/RSPtest/TXT/Threads_rsp.txt"),
+		TEXTVIEW("/tmp/TestResults/RSPtest/CSV/Threads_rsp.csv"),
+		TEXTVIEW("/tmp/TestResults/RSPtest/TSV/Threads_rsp.tsv"),
+		TEXTVIEW("/tmp/TestResults/RSPtest/TXT/Threads_rsp.txt"),
 	};
 	const TArray<FStringView> ExpectedTimersElementsRsp =
 	{
-		TEXTVIEW("/TestResults/RSPtest/CSV/Timers_rsp.csv"),
-		TEXTVIEW("/TestResults/RSPtest/TSV/Timers_rsp.tsv"),
-		TEXTVIEW("/TestResults/RSPtest/TXT/Timers_rsp.txt"),
+		TEXTVIEW("/tmp/TestResults/RSPtest/CSV/Timers_rsp.csv"),
+		TEXTVIEW("/tmp/TestResults/RSPtest/TSV/Timers_rsp.tsv"),
+		TEXTVIEW("/tmp/TestResults/RSPtest/TXT/Timers_rsp.txt"),
 	};
 	const TArray<FStringView> ExpectedTimingEventsRsp =
 	{
-		TEXTVIEW("/TestResults/RSPtest/CSV/TimingEvents_rsp.csv"),
-		TEXTVIEW("/TestResults/RSPtest/TSV/TimingEvents_rsp.tsv"),
-		TEXTVIEW("/TestResults/RSPtest/TXT/TimingEvents_rsp.txt"),
+		TEXTVIEW("/tmp/TestResults/RSPtest/CSV/TimingEvents_rsp.csv"),
+		TEXTVIEW("/tmp/TestResults/RSPtest/TSV/TimingEvents_rsp.tsv"),
+		TEXTVIEW("/tmp/TestResults/RSPtest/TXT/TimingEvents_rsp.txt"),
 	};
 	const TMap<FStringView, TArray<FStringView>> ExpectedElementsRspMap =
 	{
diff --git a/Engine/Source/Developer/TraceInsights/Private/Insights/Tests/FunctionalTests/MemoryInsightsTests.cpp b/Engine/Source/Developer/TraceInsights/Private/Insights/Tests/FunctionalTests/MemoryInsightsTests.cpp
--- a/Engine/Source/Developer/TraceInsights/Private/Insights/Tests/FunctionalTests/MemoryInsightsTests.cpp
+++ b/Engine/Source/Developer/TraceInsights/Private/Insights/Tests/FunctionalTests/MemoryInsightsTests.cpp
@@ -33,39 +33,7 @@ DECLARE_LOG_CATEGORY_EXTERN(MemoryInsightsTests, Log, All);
 
 #if WITH_AUTOMATION_TESTS
 
-class FMemoryInsightsTestBase : public FAutomationTestBase
-{
-public:
-	FMemoryInsightsTestBase(const FString& InName, const bool bInComplexTask)
-		: FAutomationTestBase(InName, bInComplexTask)
-	{
-	}
-
-	virtual bool CanRunInEnvironment(const FString& TestParams, FString* OutReason, bool* OutWarn) const override
-	{
-		using namespace UE::Insights::MemoryProfiler;
-
-		FMemorySharedState* SharedState = FMemoryProfilerManager::Get()->GetSharedState();
-		if (!SharedState)
-		{
-			if (OutReason)
-			{
-				*OutReason = TEXT("ProfilerWindow should be valid. Please, run this test through Insights Session automation tab.");
-			}
-
-			if (OutWarn)
-			{
-				*OutWarn = true;
-			}
-
-			return false;
-		}
-
-		return true;
-	}
-};
-
-IMPLEMENT_CUSTOM_SIMPLE_AUTOMATION_TEST(FMemoryInsightsUploadLLMXMLReportsTraceTest, FMemoryInsightsTestBase, "System.Insights.Trace.Analysis.MemoryInsights.UploadMemoryInsightsLLMXMLReportsTrace", EAutomationTestFlags::ProgramContext | EAutomationTestFlags::EngineFilter)
+IMPLEMENT_SIMPLE_AUTOMATION_TEST(FMemoryInsightsUploadLLMXMLReportsTraceTest, "System.Insights.Trace.Analysis.MemoryInsights.UploadMemoryInsightsLLMXMLReportsTrace", EAutomationTestFlags::ProgramContext | EAutomationTestFlags::EngineFilter)
 bool FMemoryInsightsUploadLLMXMLReportsTraceTest::RunTest(const FString& Parameters)
 {
 	using namespace UE::Insights::MemoryProfiler;
@@ -74,6 +42,12 @@ bool FMemoryInsightsUploadLLMXMLReportsTraceTest::RunTest(const FString& Paramet
 	const FString LLMReportTypesXMLPath = FPaths::RootDir() / TEXT("EngineTest/SourceAssets/Utrace/LLMReportTypes.xml");
 
 	FMemorySharedState* SharedState = FMemoryProfilerManager::Get()->GetSharedState();
+	if (!SharedState)
+	{
+		AddError("ProfilerWindow should be valid. Please, run this test through Insights Session automation tab");
+		return false;
+	}
+
 	SharedState->RemoveAllMemTagGraphTracks();
 	const int DefaultTracksAmount = SharedState->GetTimingView()->GetAllTracks().Num();
 	SharedState->RemoveAllMemTagGraphTracks();
@@ -245,7 +219,7 @@ bool MemoryInsightsAllocationsQueryTableTest(const FString& Parameters, const TM
 	return true;
 }
 
-IMPLEMENT_COMPLEX_AUTOMATION_TEST(FMemoryInsightsAllocationsQueryTableEditorPackageTest, "System.Insights.Trace.Analysis.MemoryInsights.AllocationsQueryTable.Editor.Package", EAutomationTestFlags::ProgramContext | EAutomationTestFlags::EngineFilter)
+IMPLEMENT_COMPLEX_AUTOMATION_TEST(FMemoryInsightsAllocationsQueryTableEditorPackageTest, "System.Insights.Trace.Analysis.MemoryInsights.AllocationsQueryTable.Editor/Package", EAutomationTestFlags::ProgramContext | EAutomationTestFlags::EngineFilter)
 bool FMemoryInsightsAllocationsQueryTableEditorPackageTest::RunTest(const FString& Parameters)
 {
 	bool bSuccess = MemoryInsightsAllocationsQueryTableTest(Parameters, AllocsTimeMarkerEditorPackageGetterMap, this);
diff --git a/Engine/Source/Developer/TraceInsights/Private/Insights/Tests/FunctionalTests/TimingProfilerTests.cpp b/Engine/Source/Developer/TraceInsights/Private/Insights/Tests/FunctionalTests/TimingProfilerTests.cpp
--- a/Engine/Source/Developer/TraceInsights/Private/Insights/Tests/FunctionalTests/TimingProfilerTests.cpp
+++ b/Engine/Source/Developer/TraceInsights/Private/Insights/Tests/FunctionalTests/TimingProfilerTests.cpp
@@ -143,22 +143,22 @@ bool EnumerateEventsToFile::RunTest(const FString& Parameters)
 		const TraceServices::ITimingProfilerTimerReader* TimerReader;
 		TimingProfilerProvider.ReadTimers([&TimerReader](const TraceServices::ITimingProfilerTimerReader& Out) { TimerReader = &Out; });
 
-		FString FilePath = FPaths::ProjectSavedDir() + TEXT("/EnumerateEventsToFile.txt");
-		TUniquePtr<FArchive> Ar = TUniquePtr<FArchive>(IFileManager::Get().CreateFileWriter(*FilePath));
+		FString EnumeratedEventsStr;
 		TimingProfilerProvider.ReadTimeline(TimelineIndex,
-			[Session, &Ar, TimerReader](const TraceServices::ITimingProfilerProvider::Timeline& Timeline)
+			[Session, &EnumeratedEventsStr, TimerReader](const TraceServices::ITimingProfilerProvider::Timeline& Timeline)
 			{
 				Timeline.EnumerateEvents(-1.0, Session->GetDurationSeconds() + 1.0,
-					[&Ar, TimerReader](double EventStartTime, double EventEndTime, uint32 EventDepth, const TraceServices::FTimingProfilerEvent& Event)
+					[&EnumeratedEventsStr, TimerReader](double EventStartTime, double EventEndTime, uint32 EventDepth, const TraceServices::FTimingProfilerEvent& Event)
 					{
 						const TraceServices::FTimingProfilerTimer* Timer = TimerReader->GetTimer(Event.TimerIndex);
 
-						(*Ar) << FString::Printf(TEXT("%s %f %f %d\n"), Timer->Name, EventStartTime, EventEndTime, EventDepth).GetCharArray();
+						EnumeratedEventsStr.Appendf(TEXT("%s %f %f %d\n"), Timer->Name, EventStartTime, EventEndTime, EventDepth);
 						return TraceServices::EEventEnumerate::Continue;
 					});
 			});
 
-		Ar->Close();
+		FString FilePath = FPaths::ProjectSavedDir() + TEXT("/EnumerateEventsToFile.txt");
+		FFileHelper::SaveStringToFile(EnumeratedEventsStr, *FilePath);
 	}
 
 	return !HasAnyErrors();
@@ -184,22 +184,22 @@ bool EnumerateScopesToFile::RunTest(const FString& Parameters)
 		const TraceServices::ITimingProfilerTimerReader* TimerReader;
 		TimingProfilerProvider.ReadTimers([&TimerReader](const TraceServices::ITimingProfilerTimerReader& Out) { TimerReader = &Out; });
 
-		FString FilePath = FPaths::ProjectSavedDir() + TEXT("/EnumerateScopesToFile.txt");
-		TUniquePtr<FArchive> Ar = TUniquePtr<FArchive>(IFileManager::Get().CreateFileWriter(*FilePath));
+		FString EnumeratedEventsStr;
 		TimingProfilerProvider.ReadTimeline(TimelineIndex,
-			[Session, &Ar, TimerReader](const TraceServices::ITimingProfilerProvider::Timeline& Timeline)
+			[Session, &EnumeratedEventsStr, TimerReader](const TraceServices::ITimingProfilerProvider::Timeline& Timeline)
 			{
 				Timeline.EnumerateEvents(-1.0, Session->GetDurationSeconds() + 1.0,
-					[&Ar, TimerReader](bool bStart, double Time, const TraceServices::FTimingProfilerEvent& Event)
+					[&EnumeratedEventsStr, TimerReader](bool bStart, double Time, const TraceServices::FTimingProfilerEvent& Event)
 					{
 						const TraceServices::FTimingProfilerTimer* Timer = TimerReader->GetTimer(Event.TimerIndex);
 
-						(*Ar) << FString::Printf(TEXT("%s %d %f\n"), Timer->Name, (int32)bStart, Time).GetCharArray();
+						EnumeratedEventsStr.Appendf(TEXT("%s %d %f\n"), Timer->Name, (int32)bStart, Time);
 						return TraceServices::EEventEnumerate::Continue;
 					});
 			});
 
-		Ar->Close();
+		FString FilePath = FPaths::ProjectSavedDir() + TEXT("/EnumerateScopesToFile.txt");
+		FFileHelper::SaveStringToFile(EnumeratedEventsStr, *FilePath);
 	}
 
 	return !HasAnyErrors();
diff --git a/Engine/Source/Developer/TraceInsightsFrontend/Private/InsightsFrontend/Tests/FunctionalTests/MemoryInsightsTests.cpp b/Engine/Source/Developer/TraceInsightsFrontend/Private/InsightsFrontend/Tests/FunctionalTests/MemoryInsightsTests.cpp
--- a/Engine/Source/Developer/TraceInsightsFrontend/Private/InsightsFrontend/Tests/FunctionalTests/MemoryInsightsTests.cpp
+++ b/Engine/Source/Developer/TraceInsightsFrontend/Private/InsightsFrontend/Tests/FunctionalTests/MemoryInsightsTests.cpp
@@ -33,7 +33,7 @@ void FAutomationDriverUnrealInsightsHubMemoryInsightsTest::Define()
 	{
 		AutomationWindow = FSlateApplication::Get().GetActiveTopLevelWindow();
 		const FString AutomationWindowName = TEXT("Automation");
-		if (AutomationWindow && AutomationWindow->GetTitle().ToString().Contains(AutomationWindowName))
+		if (AutomationWindow->GetTitle().ToString().Contains(AutomationWindowName))
 		{
 			AutomationWindow->Minimize();
 		}
@@ -104,7 +104,7 @@ void FAutomationDriverUnrealInsightsHubMemoryInsightsTest::Define()
 				return false;
 			};
 
-			if (!Driver->Wait(Until::Condition(TraceWaiter, FWaitTimeout::InSeconds(30))))
+			if (!Driver->Wait(Until::Condition(TraceWaiter, FWaitTimeout::InSeconds(10))))
 			{
 				AddError("Live trace should appear");
 				FPlatformProcess::TerminateProc(EditorHandle);
@@ -125,7 +125,7 @@ void FAutomationDriverUnrealInsightsHubMemoryInsightsTest::Define()
 			// Test live trace
 			FString TraceParameters = FString::Printf(TEXT("-InsightsTest -ABSLOG=\"%s\" -AutoQuit -ExecOnAnalysisCompleteCmd=\"Automation RunTests System.Insights.Trace.Analysis.MemoryInsights.UploadMemoryInsightsLLMXMLReportsTrace\" -OpenTraceFile=\"%s\""), *TestLogPath, *StoreTracePath);
 			UE::Insights::FMiscUtils::OpenUnrealInsights(*TraceParameters);
-			bool bLineFound = Utils.FileContainsString(TestLogPath, SuccessTestResult, 120.0f);
+			bool bLineFound = Utils.FileContainsString(TestLogPath, SuccessTestResult, 60.0f);
 			TestTrue("Test for live trace should pass", bLineFound);
 
 			IFileManager::Get().DeleteDirectory(*LogDirPath, false, true);
@@ -133,7 +133,7 @@ void FAutomationDriverUnrealInsightsHubMemoryInsightsTest::Define()
 
 			// Test stopped trace 
 			UE::Insights::FMiscUtils::OpenUnrealInsights(*TraceParameters);
-			bLineFound = Utils.FileContainsString(TestLogPath, SuccessTestResult, 120.0f);
+			bLineFound = Utils.FileContainsString(TestLogPath, SuccessTestResult, 60.0f);
 			TestTrue("Test for stopped trace should pass", bLineFound);
 
 			IFileManager::Get().DeleteDirectory(*LogDirPath, false, true);
@@ -145,11 +145,8 @@ void FAutomationDriverUnrealInsightsHubMemoryInsightsTest::Define()
 	{
 		Driver.Reset();
 		IAutomationDriverModule::Get().Disable();
-		if (AutomationWindow)
-		{
-			AutomationWindow->Restore();
-			AutomationWindow.Reset();
-		}
+		AutomationWindow->Restore();
+		AutomationWindow.Reset();
 	});
 }
 
diff --git a/Engine/Source/Developer/TraceInsightsFrontend/Private/InsightsFrontend/Tests/FunctionalTests/TraceStoreWindowTests.cpp b/Engine/Source/Developer/TraceInsightsFrontend/Private/InsightsFrontend/Tests/FunctionalTests/TraceStoreWindowTests.cpp
--- a/Engine/Source/Developer/TraceInsightsFrontend/Private/InsightsFrontend/Tests/FunctionalTests/TraceStoreWindowTests.cpp
+++ b/Engine/Source/Developer/TraceInsightsFrontend/Private/InsightsFrontend/Tests/FunctionalTests/TraceStoreWindowTests.cpp
@@ -163,11 +163,8 @@ void FAutomationDriverUnrealInsightsSessionBrowserTest::Define()
 	{
 		Driver.Reset();
 		IAutomationDriverModule::Get().Disable();
-		if (AutomationWindow)
-		{
-			AutomationWindow->Restore();
-			AutomationWindow.Reset();
-		}
+		AutomationWindow->Restore();
+		AutomationWindow.Reset();
 	});
 }
 
diff --git a/Engine/Source/Developer/TraceInsightsFrontend/Private/InsightsFrontend/Tests/FunctionalTests/UnrealTraceServerTests.cpp b/Engine/Source/Developer/TraceInsightsFrontend/Private/InsightsFrontend/Tests/FunctionalTests/UnrealTraceServerTests.cpp
--- a/Engine/Source/Developer/TraceInsightsFrontend/Private/InsightsFrontend/Tests/FunctionalTests/UnrealTraceServerTests.cpp
+++ b/Engine/Source/Developer/TraceInsightsFrontend/Private/InsightsFrontend/Tests/FunctionalTests/UnrealTraceServerTests.cpp
@@ -77,8 +77,8 @@ bool FUnrealTraceServerStartingTest::RunTest(const FString& Parameters)
 	bIsUnrealTraceServerRunning = WaitForUTSProcess(UnrealTraceServerName, Timeout, false);
 	TestFalse("UnrealTraceServer should not be in processes", bIsUnrealTraceServerRunning);
 
-	bConnectionStatus = WaitForConnectionIconStatus(TraceStoreConnection, Timeout, true);
-	TestTrue("Connection status should be true", bConnectionStatus);
+	bConnectionStatus = WaitForConnectionIconStatus(TraceStoreConnection, Timeout, false);
+	TestFalse("Connection status should be false", bConnectionStatus);
 
 	FString UEPath = FPlatformProcess::GenerateApplicationPath("UnrealEditor", EBuildConfiguration::Development);
 	FString UEParameters = TEXT("");
diff --git a/Engine/Source/Developer/TraceInsightsFrontend/Private/InsightsFrontend/Tests/TestRunner.cpp b/Engine/Source/Developer/TraceInsightsFrontend/Private/InsightsFrontend/Tests/TestRunner.cpp
--- a/Engine/Source/Developer/TraceInsightsFrontend/Private/InsightsFrontend/Tests/TestRunner.cpp
+++ b/Engine/Source/Developer/TraceInsightsFrontend/Private/InsightsFrontend/Tests/TestRunner.cpp
@@ -4,14 +4,8 @@
 
 #if !UE_BUILD_SHIPPING && !WITH_EDITOR
 
-#include "ISessionManager.h"
-#include "ISessionServicesModule.h"
-#include "Misc/App.h"
-
 #include "IAutomationControllerModule.h"
 #include "IAutomationControllerManager.h"
-#include "IAutomationWindowModule.h"
-#include "IAutomationWorkerModule.h"
 #include "UObject/UObjectGlobals.h" // for StaticExec()
 
 ////////////////////////////////////////////////////////////////////////////////////////////////////
@@ -23,16 +17,14 @@ DEFINE_LOG_CATEGORY(InsightsTestRunner);
 namespace UE::Insights
 {
 
-TSharedPtr<FTestRunner> FTestRunner::Instance;
-
 ////////////////////////////////////////////////////////////////////////////////////////////////////
 
-void FTestRunner::Run()
+void FTestRunner::Run(const FString& InCmd)
 {
-	FString ActualCmd = ScheduledCommand.Replace(TEXT("\""), TEXT(""));
+	FString ActualCmd = InCmd.Replace(TEXT("\""), TEXT(""));
 	if (!ActualCmd.StartsWith(TEXT("Automation RunTests")))
 	{
-		UE_LOG(InsightsTestRunner, Warning, TEXT("[InsightsTestRunner] Command %s does not start with Automation RunTests. Command will be ignored."), *ScheduledCommand);
+		UE_LOG(InsightsTestRunner, Warning, TEXT("[InsightsTestRunner] Command %s does not start with Automation RunTests. Command will be ignored."), *InCmd);
 		return;
 	}
 
@@ -42,62 +34,11 @@ void FTestRunner::Run()
 	AutomationControllerManager->OnTestsComplete().AddLambda([this]() { bIsRunningTests = false; });
 	bIsRunningTests = true;
 
-	OnTick = FTickerDelegate::CreateSP(this, &FTestRunner::Tick);
-	OnTickHandle = FTSTicker::GetCoreTicker().AddTicker(OnTick, 1.0f);
-
 	::StaticExec(NULL, *ActualCmd);
 }
 
 ////////////////////////////////////////////////////////////////////////////////////////////////////
 
-TSharedPtr<FTestRunner> FTestRunner::CreateInstance()
-{
-	ensure(!FTestRunner::Instance.IsValid());
-	if (FTestRunner::Instance.IsValid())
-	{
-		FTestRunner::Instance.Reset();
-	}
-
-	FTestRunner::Instance = MakeShared<FTestRunner>();
-
-	return FTestRunner::Instance;
-}
-
-////////////////////////////////////////////////////////////////////////////////////////////////////
-
-TSharedPtr<FTestRunner> FTestRunner::Get()
-{
-	return FTestRunner::Instance;
-}
-
-////////////////////////////////////////////////////////////////////////////////////////////////////
-
-bool FTestRunner::Tick(float DeltaTime)
-{
-	IAutomationWorkerModule& AutomationWorkerModule = FModuleManager::LoadModuleChecked<IAutomationWorkerModule>("AutomationWorker");
-	IAutomationControllerModule& AutomationControllerModule = FModuleManager::LoadModuleChecked<IAutomationControllerModule>(TEXT("AutomationController"));
-
-	AutomationControllerModule.Tick();
-	AutomationWorkerModule.Tick();
-
-	return true;
-}
-
-////////////////////////////////////////////////////////////////////////////////////////////////////
-
-void FTestRunner::Initialize()
-{
-	ISessionServicesModule& SessionServicesModule = FModuleManager::LoadModuleChecked<ISessionServicesModule>("SessionServices");
-
-	FApp::SetSessionName(TEXT("UnrealInsights"));
-	auto SessionService = SessionServicesModule.GetSessionService();
-	SessionService->Start();
-
-	auto SessionManager = SessionServicesModule.GetSessionManager();
-
-	FModuleManager::Get().LoadModule("AutomationWorker");
-}
-
 } // namespace UE::Insights
 
 #undef LOCTEXT_NAMESPACE
diff --git a/Engine/Source/Developer/TraceInsightsFrontend/Private/InsightsFrontend/Tests/TestRunner.h b/Engine/Source/Developer/TraceInsightsFrontend/Private/InsightsFrontend/Tests/TestRunner.h
--- a/Engine/Source/Developer/TraceInsightsFrontend/Private/InsightsFrontend/Tests/TestRunner.h
+++ b/Engine/Source/Developer/TraceInsightsFrontend/Private/InsightsFrontend/Tests/TestRunner.h
@@ -6,10 +6,7 @@
 
 #if !UE_BUILD_SHIPPING && !WITH_EDITOR
 
-#include "CoreMinimal.h"
-#include "Containers/Ticker.h"
 #include "Logging/LogMacros.h"
-#include "Misc/CoreMisc.h"
 
 #include "InsightsFrontend/ITraceInsightsFrontendModule.h"
 
@@ -18,30 +15,17 @@ DECLARE_LOG_CATEGORY_EXTERN(InsightsTestRunner, Log, All);
 namespace UE::Insights
 {
 
-class FTestRunner : public TSharedFromThis<FTestRunner>
+class FTestRunner
 {
 public:
 	virtual ~FTestRunner() {}
 
-	static TSharedPtr<FTestRunner> CreateInstance();
-	static TSharedPtr<FTestRunner> Get();
+	void Run(const FString& InCmd);
 
-	void Initialize();
-	void ScheduleCommand(const FString& InCmd) { ScheduledCommand = InCmd; }
-	void Run();
-
-	bool Tick(float DeltaTime);
 private:
-	FString ScheduledCommand;
-	FTickerDelegate OnTick;
-
-	/** Handle to the registered OnTick. */
-	FTSTicker::FDelegateHandle OnTickHandle;
-
 	bool bIsRunningTests = false;
-	static TSharedPtr<FTestRunner> Instance;
 };
 
 } // namespace UE::Insights
 
-#endif //!UE_BUILD_SHIPPING && !WITH_EDITOR
+#endif //UE_BUILD_SHIPPING && !WITH_EDITOR
diff --git a/Engine/Source/Programs/AutomationTool/Gauntlet/Framework/Devices/Gauntlet.DevicePool.cs b/Engine/Source/Programs/AutomationTool/Gauntlet/Framework/Devices/Gauntlet.DevicePool.cs
--- a/Engine/Source/Programs/AutomationTool/Gauntlet/Framework/Devices/Gauntlet.DevicePool.cs
+++ b/Engine/Source/Programs/AutomationTool/Gauntlet/Framework/Devices/Gauntlet.DevicePool.cs
@@ -904,7 +904,8 @@ namespace Gauntlet
 				}
 			}
 
-			return true;
+			Log.Info(KnownLogEvents.Gauntlet_DeviceEvent, "Unable to reserve service devices:");
+			return false;
 		}
 
 		/// <summary>
diff --git a/Engine/Source/Programs/AutomationTool/Gauntlet/Platform/Mac/Gauntlet.TargetDeviceMac.cs b/Engine/Source/Programs/AutomationTool/Gauntlet/Platform/Mac/Gauntlet.TargetDeviceMac.cs
--- a/Engine/Source/Programs/AutomationTool/Gauntlet/Platform/Mac/Gauntlet.TargetDeviceMac.cs
+++ b/Engine/Source/Programs/AutomationTool/Gauntlet/Platform/Mac/Gauntlet.TargetDeviceMac.cs
@@ -50,7 +50,7 @@ namespace Gauntlet
 					// If the build already exists on this drive, we can skip installation
 					string BuildVolume = GetVolumeName(BuildDir);
 					string LocalRoot = GetVolumeName(Environment.CurrentDirectory);
-					if (!string.IsNullOrEmpty(BuildVolume) && BuildVolume.Equals(LocalRoot, StringComparison.OrdinalIgnoreCase))
+					if (BuildVolume.Equals(LocalRoot, StringComparison.OrdinalIgnoreCase))
 					{
 						Log.Info("Build exists on desired volume, skipping installation");
 						return;
@@ -241,7 +241,7 @@ namespace Gauntlet
 			string LocalRoot = GetVolumeName(Environment.CurrentDirectory);
 
 			// Must be on our volume to run
-			if (!string.IsNullOrEmpty(BuildVolume) && BuildVolume.Equals(LocalRoot, StringComparison.OrdinalIgnoreCase) == false)
+			if (BuildVolume.Equals(LocalRoot, StringComparison.OrdinalIgnoreCase) == false)
 			{
 				string SubDir = string.IsNullOrEmpty(AppConfig.Sandbox) ? AppConfig.ProjectName : AppConfig.Sandbox;
 				string InstallDir = Path.Combine(InstallRoot, SubDir, AppConfig.ProcessType.ToString());
@@ -273,7 +273,7 @@ namespace Gauntlet
 
 			string BuildVolume = GetVolumeName(TargetDir);
 			string LocalRoot = GetVolumeName(Environment.CurrentDirectory);
-			if (!string.IsNullOrEmpty(BuildVolume) && BuildVolume.Equals(LocalRoot, StringComparison.OrdinalIgnoreCase))
+			if (BuildVolume.Equals(LocalRoot, StringComparison.OrdinalIgnoreCase))
 			{
 				string SubDir = string.IsNullOrEmpty(Sandbox) ? Project : Sandbox;
 				string InstallDir = Path.Combine(InstallRoot, SubDir, ProcessType);
diff --git a/Engine/Source/Programs/AutomationTool/Gauntlet/Unreal/Base/Gauntlet.UnrealAppConfig.cs b/Engine/Source/Programs/AutomationTool/Gauntlet/Unreal/Base/Gauntlet.UnrealAppConfig.cs
--- a/Engine/Source/Programs/AutomationTool/Gauntlet/Unreal/Base/Gauntlet.UnrealAppConfig.cs
+++ b/Engine/Source/Programs/AutomationTool/Gauntlet/Unreal/Base/Gauntlet.UnrealAppConfig.cs
@@ -79,7 +79,7 @@ namespace Gauntlet
 				{
 					CommandLineParams = new GauntletCommandLine();
 				}
-				return CommandLineParams.GenerateFullCommandLine(CanAlterCommandArgs);
+				return CommandLineParams.GenerateFullCommandLine();
 			}
 			set
 			{
diff --git a/Engine/Source/Programs/AutomationTool/Gauntlet/Unreal/Base/Gauntlet.UnrealTestConfiguration.cs b/Engine/Source/Programs/AutomationTool/Gauntlet/Unreal/Base/Gauntlet.UnrealTestConfiguration.cs
--- a/Engine/Source/Programs/AutomationTool/Gauntlet/Unreal/Base/Gauntlet.UnrealTestConfiguration.cs
+++ b/Engine/Source/Programs/AutomationTool/Gauntlet/Unreal/Base/Gauntlet.UnrealTestConfiguration.cs
@@ -285,22 +285,19 @@ namespace Gauntlet
 		/// in quotes etc.
 		/// This is the function we use, as well, when accessing the Role's commandline itself.
 		/// </summary>
-		/// <param name="CanAlterCommandArgs">Some tests require that the command is not altered by any Gauntlet logic and only the provided args are used.</param>
 		/// <returns></returns>
-		public string GenerateFullCommandLine(bool CanAlterCommandArgs = true)
+		public string GenerateFullCommandLine()
 		{
 			string FinalCommandline = "";
-			if (CanAlterCommandArgs)
+			if (!string.IsNullOrEmpty(Project))
 			{
-				if (!string.IsNullOrEmpty(Project))
-				{
-					FinalCommandline += string.Format("\"{0}\" ", Project);
-				}
-				if (!string.IsNullOrEmpty(GameMap))
-				{
-					FinalCommandline += string.Format("\"{0}\" ", GameMap);
-				}
+				FinalCommandline += string.Format("\"{0}\" ", Project);
 			}
+			if (!string.IsNullOrEmpty(GameMap))
+			{
+				FinalCommandline += string.Format("\"{0}\" ", GameMap);
+			}
+
 			foreach (string Key in Params.Keys)
 			{
 				string CurrentArgument;
@@ -901,7 +898,6 @@ namespace Gauntlet
 				UnrealTargetRole.EditorGame,
 				UnrealTargetRole.Server,
 				UnrealTargetRole.EditorServer,
-				UnrealTargetRole.Host,
 				UnrealTargetRole.Editor,
 				UnrealTargetRole.CookedEditor
 			};
diff --git a/Engine/Source/Programs/AutomationTool/Gauntlet/Unreal/Utils/Gauntlet.UnrealHelpers.cs b/Engine/Source/Programs/AutomationTool/Gauntlet/Unreal/Utils/Gauntlet.UnrealHelpers.cs
--- a/Engine/Source/Programs/AutomationTool/Gauntlet/Unreal/Utils/Gauntlet.UnrealHelpers.cs
+++ b/Engine/Source/Programs/AutomationTool/Gauntlet/Unreal/Utils/Gauntlet.UnrealHelpers.cs
@@ -27,7 +27,6 @@ namespace Gauntlet
 		EditorServer,
 		Client,
 		Server,
-		Host,
 		CookedEditor,
 	};
 
@@ -81,11 +80,6 @@ namespace Gauntlet
 			return Type == UnrealTargetRole.EditorGame || Type == UnrealTargetRole.Client;
 		}
 
-		public static bool IsHost(this UnrealTargetRole Type)
-		{
-			return Type == UnrealTargetRole.Host;
-		}
-
 		public static bool IsEditor(this UnrealTargetRole Type)
 		{
 			return Type == UnrealTargetRole.Editor || Type == UnrealTargetRole.CookedEditor;
diff --git a/Engine/Source/Programs/AutomationTool/Insights/Insights.Automation.csproj b/Engine/Source/Programs/AutomationTool/Insights/Insights.Automation.csproj
deleted file mode 100644
diff --git a/Engine/Source/Programs/AutomationTool/Insights/RunInsightsTests.cs b/Engine/Source/Programs/AutomationTool/Insights/RunInsightsTests.cs
deleted file mode 100644
diff --git a/Engine/Source/Programs/AutomationTool/Insights/Tests/Gauntlet.InsightsTestNode.cs b/Engine/Source/Programs/AutomationTool/Insights/Tests/Gauntlet.InsightsTestNode.cs
deleted file mode 100644
diff --git a/Engine/Source/Programs/AutomationTool/Scripts/BuildCookRun.Automation.cs b/Engine/Source/Programs/AutomationTool/Scripts/BuildCookRun.Automation.cs
--- a/Engine/Source/Programs/AutomationTool/Scripts/BuildCookRun.Automation.cs
+++ b/Engine/Source/Programs/AutomationTool/Scripts/BuildCookRun.Automation.cs
@@ -139,12 +139,6 @@ public class BuildCookRun : BuildCommand, IProjectParamsHelpers
 			Params.AdditionalBuildOptions += " -EnablePlugins=\"" + string.Join(",", PluginsToEnable) + "\"";
 		}
 
-		bool WithUBA = ParseParamBool("UBA");
-		if (WithUBA)
-		{
-			Params.AdditionalBuildOptions += " -UBA";
-		}	
-
 		Params.ValidateAndLog();
 		return Params;
 	}
diff --git a/Engine/Source/Programs/UnrealInsights/Private/UserInterfaceCommand.cpp b/Engine/Source/Programs/UnrealInsights/Private/UserInterfaceCommand.cpp
--- a/Engine/Source/Programs/UnrealInsights/Private/UserInterfaceCommand.cpp
+++ b/Engine/Source/Programs/UnrealInsights/Private/UserInterfaceCommand.cpp
@@ -9,7 +9,6 @@
 #include "Framework/Application/SlateApplication.h"
 #include "Framework/Docking/LayoutService.h"
 #include "Framework/Docking/TabManager.h"
-#include "HAL/ExceptionHandling.h"
 #include "HAL/PlatformApplicationMisc.h"
 #include "HAL/PlatformProcess.h"
 #include "Interfaces/IPluginManager.h"
@@ -221,7 +220,6 @@ void FUserInterfaceCommand::InitializeSlateApplication(bool bFrontendMode, const
 		const bool bInitializeTesting = FParse::Param(FCommandLine::Get(), TEXT("InsightsTest"));
 		if (bInitializeTesting)
 		{
-			GUseCrashReportClient = false;
 			const bool bInitAutomationModules = true;
 			TraceInsightsModule.InitializeTesting(bInitAutomationModules, bAutoQuit);
 		}
@@ -302,7 +300,6 @@ void FUserInterfaceCommand::InitializeSlateApplication(bool bFrontendMode, const
 			TraceInsightsFrontendModule.RunAutomationTests(AutomationTests);
 		}
 	}
-	UE_LOG(LogSlate, Display, TEXT("Insights slate application initialized successfully."));
 }
 
 ////////////////////////////////////////////////////////////////////////////////////////////////////
