[Insights]
- Session Info: Fixed the analysis duration to update also when session duration does not change.
- Session Info: Changed analysis status texts (previously hardcoded) to be localized.

diff --git a/Engine/Source/Developer/TraceInsights/Private/Insights/InsightsManager.cpp b/Engine/Source/Developer/TraceInsights/Private/Insights/InsightsManager.cpp
--- a/Engine/Source/Developer/TraceInsights/Private/Insights/InsightsManager.cpp
+++ b/Engine/Source/Developer/TraceInsights/Private/Insights/InsightsManager.cpp
@@ -353,20 +353,24 @@ void FInsightsManager::UpdateSessionDuration()
 
 	if (Session.IsValid())
 	{
+		if (!bIsAnalysisComplete)
+		{
+			AnalysisStopwatch.Update();
+			AnalysisDuration = AnalysisStopwatch.GetAccumulatedTime();
+			AnalysisSpeedFactor = SessionDuration / AnalysisDuration;
+		}
+
 		bool bLocalIsAnalysisComplete = false;
 		double LocalSessionDuration = 0.0;
 		{
 			TraceServices::FAnalysisSessionReadScope SessionReadScope(*Session.Get());
 			bLocalIsAnalysisComplete = Session->IsAnalysisComplete();
 			LocalSessionDuration = Session->GetDurationSeconds();
-
 		}
 
 		if (LocalSessionDuration != SessionDuration)
 		{
 			SessionDuration = LocalSessionDuration;
-			AnalysisStopwatch.Update();
-			AnalysisDuration = AnalysisStopwatch.GetAccumulatedTime();
 			AnalysisSpeedFactor = SessionDuration / AnalysisDuration;
 			if (bIsAnalysisComplete)
 			{
@@ -378,6 +382,9 @@ void FInsightsManager::UpdateSessionDuration()
 		if (bLocalIsAnalysisComplete && !bIsAnalysisComplete)
 		{
 			bIsAnalysisComplete = true;
+			AnalysisStopwatch.Update();
+			AnalysisDuration = AnalysisStopwatch.GetAccumulatedTime();
+			AnalysisSpeedFactor = SessionDuration / AnalysisDuration;
 			SessionAnalysisCompletedEvent.Broadcast();
 
 			UE_LOG(TraceInsights, Log, TEXT("Analysis has completed in %s (%.1fX speed; session duration: %s)."),
diff --git a/Engine/Source/Developer/TraceInsights/Private/Insights/Widgets/SSessionInfoWindow.cpp b/Engine/Source/Developer/TraceInsights/Private/Insights/Widgets/SSessionInfoWindow.cpp
--- a/Engine/Source/Developer/TraceInsights/Private/Insights/Widgets/SSessionInfoWindow.cpp
+++ b/Engine/Source/Developer/TraceInsights/Private/Insights/Widgets/SSessionInfoWindow.cpp
@@ -575,7 +575,9 @@ FText SSessionInfoWindow::GetStatusText() const
 
 	using namespace UE::Insights;
 	FText Status = FText::Format(LOCTEXT("StatusFmt", "{0}\nSession Duration: {1}\nAnalyzed in {2} at {3}X speed."),
-		InsightsManager->IsAnalysisComplete() ? FText::FromString(FString(TEXT("ANALYSIS COMPLETED."))) : FText::FromString(FString(TEXT("ANALYZING..."))),
+		InsightsManager->IsAnalysisComplete() ?
+			LOCTEXT("Status_AnalysisCompleted", "ANALYSIS COMPLETED.") :
+			LOCTEXT("Status_Analyzing", "ANALYZING..."),
 		FText::FromString(FormatTimeAuto(InsightsManager->GetSessionDuration(), NumDigits)),
 		FText::FromString(FormatTimeAuto(InsightsManager->GetAnalysisDuration(), NumDigits)),
 		FText::AsNumber(InsightsManager->GetAnalysisSpeedFactor(), &FormattingOptions));
