Disable Memory Insights in Editor in Shipping

diff --git a/Engine/Source/Developer/TraceInsights/Private/Insights/MemoryProfiler/MemoryProfilerManager.cpp b/Engine/Source/Developer/TraceInsights/Private/Insights/MemoryProfiler/MemoryProfilerManager.cpp
--- a/Engine/Source/Developer/TraceInsights/Private/Insights/MemoryProfiler/MemoryProfilerManager.cpp
+++ b/Engine/Source/Developer/TraceInsights/Private/Insights/MemoryProfiler/MemoryProfilerManager.cpp
@@ -158,31 +158,34 @@ void FMemoryProfilerManager::BindCommands()
 
 void FMemoryProfilerManager::RegisterMajorTabs(IUnrealInsightsModule& InsightsModule)
 {
-#if WITH_EDITOR
-	if (GEnableMemoryInsightsInEditor)
-	{
+	bool bShouldRegister = true;
+
+#if WITH_EDITOR 
+	#if UE_BUILD_SHIPPING
+		bShouldRegister = false;
+	#else
+		bShouldRegister = GEnableMemoryInsightsInEditor;
+	#endif
 #endif
-		if (!bIsTabRegistered)
-		{
-			bIsTabRegistered = true;
-			const FInsightsMajorTabConfig& Config = InsightsModule.FindMajorTabConfig(FInsightsManagerTabs::MemoryProfilerTabId);
 
-			if (Config.bIsAvailable)
-			{
-				// Register tab spawner for the Memory Insights.
-				FTabSpawnerEntry& TabSpawnerEntry = FGlobalTabmanager::Get()->RegisterNomadTabSpawner(FInsightsManagerTabs::MemoryProfilerTabId,
-					FOnSpawnTab::CreateRaw(this, &FMemoryProfilerManager::SpawnTab), FCanSpawnTab::CreateRaw(this, &FMemoryProfilerManager::CanSpawnTab))
-					.SetDisplayName(Config.TabLabel.IsSet() ? Config.TabLabel.GetValue() : LOCTEXT("MemoryProfilerTabTitle", "Memory Insights"))
-					.SetTooltipText(Config.TabTooltip.IsSet() ? Config.TabTooltip.GetValue() : LOCTEXT("MemoryProfilerTooltipText", "Open the Memory Insights tab."))
-					.SetIcon(Config.TabIcon.IsSet() ? Config.TabIcon.GetValue() : FSlateIcon(FInsightsStyle::GetStyleSetName(), "Icons.MemoryProfiler"));
-
-				TSharedRef<FWorkspaceItem> Group = Config.WorkspaceGroup.IsValid() ? Config.WorkspaceGroup.ToSharedRef() : FInsightsManager::Get()->GetInsightsMenuBuilder()->GetInsightsToolsGroup();
-				TabSpawnerEntry.SetGroup(Group);
-			}
+	if (!bIsTabRegistered && bShouldRegister)
+	{
+		bIsTabRegistered = true;
+		const FInsightsMajorTabConfig& Config = InsightsModule.FindMajorTabConfig(FInsightsManagerTabs::MemoryProfilerTabId);
+
+		if (Config.bIsAvailable)
+		{
+			// Register tab spawner for the Memory Insights.
+			FTabSpawnerEntry& TabSpawnerEntry = FGlobalTabmanager::Get()->RegisterNomadTabSpawner(FInsightsManagerTabs::MemoryProfilerTabId,
+				FOnSpawnTab::CreateRaw(this, &FMemoryProfilerManager::SpawnTab), FCanSpawnTab::CreateRaw(this, &FMemoryProfilerManager::CanSpawnTab))
+				.SetDisplayName(Config.TabLabel.IsSet() ? Config.TabLabel.GetValue() : LOCTEXT("MemoryProfilerTabTitle", "Memory Insights"))
+				.SetTooltipText(Config.TabTooltip.IsSet() ? Config.TabTooltip.GetValue() : LOCTEXT("MemoryProfilerTooltipText", "Open the Memory Insights tab."))
+				.SetIcon(Config.TabIcon.IsSet() ? Config.TabIcon.GetValue() : FSlateIcon(FInsightsStyle::GetStyleSetName(), "Icons.MemoryProfiler"));
+
+			TSharedRef<FWorkspaceItem> Group = Config.WorkspaceGroup.IsValid() ? Config.WorkspaceGroup.ToSharedRef() : FInsightsManager::Get()->GetInsightsMenuBuilder()->GetInsightsToolsGroup();
+			TabSpawnerEntry.SetGroup(Group);
 		}
-#if WITH_EDITOR
 	}
-#endif
 }
 
 ////////////////////////////////////////////////////////////////////////////////////////////////////
