[Insights] Session Info: Fixed fickering in the text showing analysis status (i.e. when the the speed factor has changed too often).

diff --git a/Engine/Source/Developer/TraceInsights/Private/Insights/Widgets/SSessionInfoWindow.cpp b/Engine/Source/Developer/TraceInsights/Private/Insights/Widgets/SSessionInfoWindow.cpp
--- a/Engine/Source/Developer/TraceInsights/Private/Insights/Widgets/SSessionInfoWindow.cpp
+++ b/Engine/Source/Developer/TraceInsights/Private/Insights/Widgets/SSessionInfoWindow.cpp
@@ -569,6 +569,7 @@ FText SSessionInfoWindow::GetStatusText() const
 	InsightsManager->UpdateSessionDuration();
 
 	FNumberFormattingOptions FormattingOptions;
+	FormattingOptions.MinimumFractionalDigits = 1;
 	FormattingOptions.MaximumFractionalDigits = 1;
 
 	const int32 NumDigits = InsightsManager->IsAnalysisComplete() ? 2 : 0;
