[Insights]
- Memory Insights: Added budget line in the graph tracks for LLM tag series.
- Memory Insights: Added Graph track options to LLM tag tracks to allow "auto zoom" to include baseline (value == 0) and threshold values (i.e. budget value).

diff --git a/Engine/Source/Developer/TraceInsights/Private/Insights/MemoryProfiler/Tracks/AllocationsGraphSeries.cpp b/Engine/Source/Developer/TraceInsights/Private/Insights/MemoryProfiler/Tracks/AllocationsGraphSeries.cpp
--- a/Engine/Source/Developer/TraceInsights/Private/Insights/MemoryProfiler/Tracks/AllocationsGraphSeries.cpp
+++ b/Engine/Source/Developer/TraceInsights/Private/Insights/MemoryProfiler/Tracks/AllocationsGraphSeries.cpp
@@ -321,111 +321,93 @@ void FAllocationsGraphSeries::Update(FGraphTrack& GraphTrack, const FTimingTrack
 {
 	FGraphTrackBuilder Builder(GraphTrack, *this, Viewport);
 
+	//////////////////////////////////////////////////
+	// Update auto-zoom (vertical scaling)...
+
+	const float TopY = 4.0f;
+	const float BottomY = GraphTrack.GetHeight() - 4.0f;
+
+	if (IsAutoZoomEnabled() && TopY < BottomY)
+	{
+		UpdateAutoZoom(TopY, BottomY, GetMinValue(), GetMaxValue());
+	}
+
+	//////////////////////////////////////////////////
+
 	TSharedPtr<const TraceServices::IAnalysisSession> Session = FInsightsManager::Get()->GetSession();
 	if (Session.IsValid())
 	{
 		const TraceServices::IAllocationsProvider* AllocationsProvider = TraceServices::ReadAllocationsProvider(*Session.Get());
 		if (AllocationsProvider)
 		{
-			struct FTimelineEvent
-			{
-				FTimelineEvent(double InTime, double InDuration, double InValue) :
-					Time(InTime),
-					Duration(InDuration),
-					Value(InValue)
-				{
-
-				}
-				double Time;
-				double Duration;
-				double Value;
-			};
+			TraceServices::FProviderReadScopeLock ProviderReadScope(*AllocationsProvider);
 
 			int32 StartIndex = -1;
 			int32 EndIndex = -1;
 
-			TArray<FTimelineEvent> TimelineEvents;
-			{
-				TraceServices::FProviderReadScopeLock ProviderReadScope(*AllocationsProvider);
-				AllocationsProvider->GetTimelineIndexRange(Viewport.GetStartTime(), Viewport.GetEndTime(), StartIndex, EndIndex);
-
-				if (EndIndex >= 0)
-				{
-					--StartIndex; // include one more point on the left side
-					++EndIndex; // include one more point on the right side
-				}
-
-				TimelineEvents.Reserve((EndIndex - StartIndex) + 1);
-
-				const float TopY = 4.0f;
-				const float BottomY = GraphTrack.GetHeight() - 4.0f;
-
-				if (IsAutoZoomEnabled() && TopY < BottomY)
-				{
-					UpdateAutoZoom(TopY, BottomY, GetMinValue(), GetMaxValue());
-				}
-
-				auto Callback64 = [&TimelineEvents](double Time, double Duration, uint64 Value)
-				{
-					TimelineEvents.Emplace(Time, Duration, static_cast<double>(Value));
-				};
-				auto Callback32 = [&TimelineEvents](double Time, double Duration, uint32 Value)
-				{
-					TimelineEvents.Emplace(Time, Duration, static_cast<double>(Value));
-				};
-				auto Callback32Negative = [&TimelineEvents](double Time, double Duration, uint32 Value)
-				{
-					// Shows FreeEvents as negative values in order to be displayed on same graph as AllocEvents.
-					TimelineEvents.Emplace(Time, Duration, -static_cast<double>(Value));
-				};
+			AllocationsProvider->GetTimelineIndexRange(Viewport.GetStartTime(), Viewport.GetEndTime(), StartIndex, EndIndex);
 
-				switch (Timeline)
-				{
-				case ETimeline::MinTotalMem:
-					AllocationsProvider->EnumerateMinTotalAllocatedMemoryTimeline(StartIndex, EndIndex, Callback64);
-					break;
-				case ETimeline::MaxTotalMem:
-					AllocationsProvider->EnumerateMaxTotalAllocatedMemoryTimeline(StartIndex, EndIndex, Callback64);
-					break;
-				case ETimeline::MinLiveAllocs:
-					AllocationsProvider->EnumerateMinLiveAllocationsTimeline(StartIndex, EndIndex, Callback32);
-					break;
-				case ETimeline::MaxLiveAllocs:
-					AllocationsProvider->EnumerateMaxLiveAllocationsTimeline(StartIndex, EndIndex, Callback32);
-					break;
-				case ETimeline::MinSwapMem:
-					AllocationsProvider->EnumerateMinTotalSwapMemoryTimeline(StartIndex, EndIndex, Callback64);
-					break;
-				case ETimeline::MaxSwapMem:
-					AllocationsProvider->EnumerateMaxTotalSwapMemoryTimeline(StartIndex, EndIndex, Callback64);
-					break;
-				case ETimeline::MinCompressedSwapMem:
-					AllocationsProvider->EnumerateMinTotalCompressedSwapMemoryTimeline(StartIndex, EndIndex, Callback64);
-					break;
-				case ETimeline::MaxCompressedSwapMem:
-					AllocationsProvider->EnumerateMaxTotalCompressedSwapMemoryTimeline(StartIndex, EndIndex, Callback64);
-					break;
-				case ETimeline::AllocEvents:
-					AllocationsProvider->EnumerateAllocEventsTimeline(StartIndex, EndIndex, Callback32);
-					break;
-				case ETimeline::FreeEvents:
-					AllocationsProvider->EnumerateFreeEventsTimeline(StartIndex, EndIndex, Callback32Negative);
-					break;
-				case ETimeline::PageInEvents:
-					AllocationsProvider->EnumeratePageInEventsTimeline(StartIndex, EndIndex, Callback32);
-					break;
-				case ETimeline::PageOutEvents:
-					AllocationsProvider->EnumeratePageOutEventsTimeline(StartIndex, EndIndex, Callback32Negative);
-					break;
-				case ETimeline::SwapFreeEvents:
-					AllocationsProvider->EnumerateSwapFreeEventsTimeline(StartIndex, EndIndex, Callback32);
-					break;
-				}
+			if (EndIndex >= 0)
+			{
+				--StartIndex; // include one more point on the left side
+				++EndIndex; // include one more point on the right side
 			}
 
-			for (const FTimelineEvent& TimelineEvent : TimelineEvents)
+			auto Callback64 = [&Builder](double Time, double Duration, uint64 Value)
+			{
+				Builder.AddEvent(Time, Duration, static_cast<double>(Value));
+			};
+			auto Callback32 = [&Builder](double Time, double Duration, uint32 Value)
+			{
+				Builder.AddEvent(Time, Duration, static_cast<double>(Value));
+			};
+			auto Callback32Negative = [&Builder](double Time, double Duration, uint32 Value)
+			{
+				// Shows FreeEvents as negative values in order to be displayed on same graph as AllocEvents.
+				Builder.AddEvent(Time, Duration, -static_cast<double>(Value));
+			};
+
+			switch (Timeline)
 			{
-				Builder.AddEvent(TimelineEvent.Time, TimelineEvent.Duration, TimelineEvent.Value);
+			case ETimeline::MinTotalMem:
+				AllocationsProvider->EnumerateMinTotalAllocatedMemoryTimeline(StartIndex, EndIndex, Callback64);
+				break;
+			case ETimeline::MaxTotalMem:
+				AllocationsProvider->EnumerateMaxTotalAllocatedMemoryTimeline(StartIndex, EndIndex, Callback64);
+				break;
+			case ETimeline::MinLiveAllocs:
+				AllocationsProvider->EnumerateMinLiveAllocationsTimeline(StartIndex, EndIndex, Callback32);
+				break;
+			case ETimeline::MaxLiveAllocs:
+				AllocationsProvider->EnumerateMaxLiveAllocationsTimeline(StartIndex, EndIndex, Callback32);
+				break;
+			case ETimeline::MinSwapMem:
+				AllocationsProvider->EnumerateMinTotalSwapMemoryTimeline(StartIndex, EndIndex, Callback64);
+				break;
+			case ETimeline::MaxSwapMem:
+				AllocationsProvider->EnumerateMaxTotalSwapMemoryTimeline(StartIndex, EndIndex, Callback64);
+				break;
+			case ETimeline::MinCompressedSwapMem:
+				AllocationsProvider->EnumerateMinTotalCompressedSwapMemoryTimeline(StartIndex, EndIndex, Callback64);
+				break;
+			case ETimeline::MaxCompressedSwapMem:
+				AllocationsProvider->EnumerateMaxTotalCompressedSwapMemoryTimeline(StartIndex, EndIndex, Callback64);
+				break;
+			case ETimeline::AllocEvents:
+				AllocationsProvider->EnumerateAllocEventsTimeline(StartIndex, EndIndex, Callback32);
+				break;
+			case ETimeline::FreeEvents:
+				AllocationsProvider->EnumerateFreeEventsTimeline(StartIndex, EndIndex, Callback32Negative);
+				break;
+			case ETimeline::PageInEvents:
+				AllocationsProvider->EnumeratePageInEventsTimeline(StartIndex, EndIndex, Callback32);
+				break;
+			case ETimeline::PageOutEvents:
+				AllocationsProvider->EnumeratePageOutEventsTimeline(StartIndex, EndIndex, Callback32Negative);
+				break;
+			case ETimeline::SwapFreeEvents:
+				AllocationsProvider->EnumerateSwapFreeEventsTimeline(StartIndex, EndIndex, Callback32);
+				break;
 			}
 		}
 	}
diff --git a/Engine/Source/Developer/TraceInsights/Private/Insights/MemoryProfiler/Tracks/MemTagGraphSeries.cpp b/Engine/Source/Developer/TraceInsights/Private/Insights/MemoryProfiler/Tracks/MemTagGraphSeries.cpp
--- a/Engine/Source/Developer/TraceInsights/Private/Insights/MemoryProfiler/Tracks/MemTagGraphSeries.cpp
+++ b/Engine/Source/Developer/TraceInsights/Private/Insights/MemoryProfiler/Tracks/MemTagGraphSeries.cpp
@@ -11,6 +11,7 @@
 #include "Insights/ViewModels/GraphTrackBuilder.h"
 #include "Insights/ViewModels/TimingTrackViewport.h"
 
+#include <cmath>
 #include <limits>
 
 #define LOCTEXT_NAMESPACE "UE::Insights::MemoryProfiler::FMemoryGraphTrack"
@@ -43,6 +44,66 @@ FMemTagGraphSeries::~FMemTagGraphSeries()
 
 ////////////////////////////////////////////////////////////////////////////////////////////////////
 
+bool FMemTagGraphSeries::HasHighThresholdValue() const
+{
+	return std::isfinite(HighThresholdValue);
+}
+
+////////////////////////////////////////////////////////////////////////////////////////////////////
+
+double FMemTagGraphSeries::GetHighThresholdValue() const
+{
+	return HighThresholdValue;
+}
+
+////////////////////////////////////////////////////////////////////////////////////////////////////
+
+void FMemTagGraphSeries::SetHighThresholdValue(double InValue)
+{
+	HighThresholdValue = InValue;
+	SetAutoZoomDirty();
+}
+
+////////////////////////////////////////////////////////////////////////////////////////////////////
+
+void FMemTagGraphSeries::ResetHighThresholdValue()
+{
+	HighThresholdValue = +std::numeric_limits<double>::infinity();
+	SetAutoZoomDirty();
+}
+
+////////////////////////////////////////////////////////////////////////////////////////////////////
+
+bool FMemTagGraphSeries::HasLowThresholdValue() const
+{
+	return std::isfinite(LowThresholdValue);
+}
+
+////////////////////////////////////////////////////////////////////////////////////////////////////
+
+double FMemTagGraphSeries::GetLowThresholdValue() const
+{
+	return LowThresholdValue;
+}
+
+////////////////////////////////////////////////////////////////////////////////////////////////////
+
+void FMemTagGraphSeries::SetLowThresholdValue(double InValue)
+{
+	LowThresholdValue = InValue;
+	SetAutoZoomDirty();
+}
+
+////////////////////////////////////////////////////////////////////////////////////////////////////
+
+void FMemTagGraphSeries::ResetLowThresholdValue()
+{
+	LowThresholdValue = -std::numeric_limits<double>::infinity();
+	SetAutoZoomDirty();
+}
+
+////////////////////////////////////////////////////////////////////////////////////////////////////
+
 FString FMemTagGraphSeries::FormatValue(double Value) const
 {
 	const int64 ValueInt64 = static_cast<int64>(Value);
@@ -73,38 +134,32 @@ FString FMemTagGraphSeries::FormatValue(double Value) const
 
 void FMemTagGraphSeries::PreUpdate(FGraphTrack& GraphTrack, const FTimingTrackViewport& Viewport)
 {
+	double LocalMinValue = +std::numeric_limits<double>::infinity();
+	double LocalMaxValue = -std::numeric_limits<double>::infinity();
+
 	TSharedPtr<const TraceServices::IAnalysisSession> Session = FInsightsManager::Get()->GetSession();
 	if (Session.IsValid())
 	{
-		TraceServices::FAnalysisSessionReadScope SessionReadScope(*Session.Get());
 		const TraceServices::IMemoryProvider* MemoryProvider = TraceServices::ReadMemoryProvider(*Session.Get());
 		if (MemoryProvider)
 		{
+			TraceServices::FAnalysisSessionReadScope SessionReadScope(*Session.Get());
+
 			const uint64 TotalSampleCount = MemoryProvider->GetTagSampleCount(TrackerId, TagId);
 			if (TotalSampleCount > 0)
 			{
-				double LocalMinValue = +std::numeric_limits<double>::infinity();
-				double LocalMaxValue = -std::numeric_limits<double>::infinity();
-
 				// Compute Min/Max values.
 				MemoryProvider->EnumerateTagSamples(TrackerId, TagId, Viewport.GetStartTime(), Viewport.GetEndTime(), true,
 					[&LocalMinValue, &LocalMaxValue](double Time, double Duration, const TraceServices::FMemoryTagSample& Sample)
 					{
 						const double Value = static_cast<double>(Sample.Value);
-						if (Value < LocalMinValue)
-						{
-							LocalMinValue = Value;
-						}
-						if (Value > LocalMaxValue)
-						{
-							LocalMaxValue = Value;
-						}
+						ExpandRange(LocalMinValue, LocalMaxValue, Value);
 					});
-
-				SetValueRange(LocalMinValue, LocalMaxValue);
 			}
 		}
 	}
+
+	SetValueRange(LocalMinValue, LocalMaxValue);
 }
 
 ////////////////////////////////////////////////////////////////////////////////////////////////////
@@ -113,25 +168,52 @@ void FMemTagGraphSeries::Update(FGraphTrack& GraphTrack, const FTimingTrackViewp
 {
 	FGraphTrackBuilder Builder(GraphTrack, *this, Viewport);
 
+	//////////////////////////////////////////////////
+	// Update auto-zoom (vertical scaling)...
+
+	const float TopY = 4.0f;
+	const float BottomY = GraphTrack.GetHeight() - 4.0f;
+
+	if (IsAutoZoomEnabled() && TopY < BottomY)
+	{
+		double ZoomMinValue = GetMinValue();
+		double ZoomMaxValue = GetMaxValue();
+
+		bool bAutoZoomIncludesBaseline = GraphTrack.IsAnyOptionEnabled(EGraphOptions::AutoZoomIncludesBaseline);
+		if (bAutoZoomIncludesBaseline)
+		{
+			ExpandRange(ZoomMinValue, ZoomMaxValue, 0.0);
+		}
+
+		bool bAutoZoomIncludesThresholds = GraphTrack.IsAnyOptionEnabled(EGraphOptions::AutoZoomIncludesThresholds);
+		if (bAutoZoomIncludesThresholds)
+		{
+			if (HasHighThresholdValue())
+			{
+				ExpandRange(ZoomMinValue, ZoomMaxValue, GetHighThresholdValue());
+			}
+			if (HasLowThresholdValue())
+			{
+				ExpandRange(ZoomMinValue, ZoomMaxValue, GetLowThresholdValue());
+			}
+		}
+
+		UpdateAutoZoom(TopY, BottomY, ZoomMinValue, ZoomMaxValue);
+	}
+
+	//////////////////////////////////////////////////
+
 	TSharedPtr<const TraceServices::IAnalysisSession> Session = FInsightsManager::Get()->GetSession();
 	if (Session.IsValid())
 	{
-		TraceServices::FAnalysisSessionReadScope SessionReadScope(*Session.Get());
-
 		const TraceServices::IMemoryProvider* MemoryProvider = TraceServices::ReadMemoryProvider(*Session.Get());
 		if (MemoryProvider)
 		{
+			TraceServices::FAnalysisSessionReadScope SessionReadScope(*Session.Get());
+
 			const uint64 TotalSampleCount = MemoryProvider->GetTagSampleCount(TrackerId, TagId);
 			if (TotalSampleCount > 0)
 			{
-				const float TopY = 4.0f;
-				const float BottomY = GraphTrack.GetHeight() - 4.0f;
-
-				if (IsAutoZoomEnabled() && TopY < BottomY)
-				{
-					UpdateAutoZoom(TopY, BottomY, GetMinValue(), GetMaxValue());
-				}
-
 				MemoryProvider->EnumerateTagSamples(TrackerId, TagId, Viewport.GetStartTime(), Viewport.GetEndTime(), true,
 					[this, &Builder](double Time, double Duration, const TraceServices::FMemoryTagSample& Sample)
 					{
diff --git a/Engine/Source/Developer/TraceInsights/Private/Insights/MemoryProfiler/Tracks/MemTagGraphSeries.h b/Engine/Source/Developer/TraceInsights/Private/Insights/MemoryProfiler/Tracks/MemTagGraphSeries.h
--- a/Engine/Source/Developer/TraceInsights/Private/Insights/MemoryProfiler/Tracks/MemTagGraphSeries.h
+++ b/Engine/Source/Developer/TraceInsights/Private/Insights/MemoryProfiler/Tracks/MemTagGraphSeries.h
@@ -9,6 +9,8 @@
 #include "Insights/MemoryProfiler/ViewModels/MemoryTag.h"
 #include "Insights/MemoryProfiler/ViewModels/MemoryTracker.h"
 
+#include <limits>
+
 namespace UE::Insights::MemoryProfiler
 {
 
@@ -20,6 +22,16 @@ public:
 	FMemTagGraphSeries(FMemoryTrackerId InTrackerId, FMemoryTagSetId InTagSetId, FMemoryTagId InTagId);
 	virtual ~FMemTagGraphSeries();
 
+	virtual bool HasHighThresholdValue() const override;
+	virtual double GetHighThresholdValue() const override;
+	virtual void SetHighThresholdValue(double InValue) override;
+	virtual void ResetHighThresholdValue() override;
+
+	virtual bool HasLowThresholdValue() const override;
+	virtual double GetLowThresholdValue() const override;
+	virtual void SetLowThresholdValue(double InValue) override;
+	virtual void ResetLowThresholdValue() override;
+
 	virtual FString FormatValue(double Value) const override;
 	virtual void PreUpdate(FGraphTrack& GraphTrack, const FTimingTrackViewport& Viewport) override;
 	virtual void Update(FGraphTrack& GraphTrack, const FTimingTrackViewport& Viewport) override;
@@ -32,6 +44,9 @@ private:
 	FMemoryTrackerId TrackerId = FMemoryTracker::InvalidTrackerId; // LLM tracker id
 	FMemoryTagSetId TagSetId = FMemoryTagSet::InvalidTagSetId; // LLM tag set id
 	FMemoryTagId TagId = FMemoryTag::InvalidTagId; // LLM tag id
+
+	double HighThresholdValue = +std::numeric_limits<double>::infinity();
+	double LowThresholdValue  = -std::numeric_limits<double>::infinity();
 };
 
 } // namespace UE::Insights::MemoryProfiler
diff --git a/Engine/Source/Developer/TraceInsights/Private/Insights/MemoryProfiler/Tracks/MemoryGraphSeries.cpp b/Engine/Source/Developer/TraceInsights/Private/Insights/MemoryProfiler/Tracks/MemoryGraphSeries.cpp
--- a/Engine/Source/Developer/TraceInsights/Private/Insights/MemoryProfiler/Tracks/MemoryGraphSeries.cpp
+++ b/Engine/Source/Developer/TraceInsights/Private/Insights/MemoryProfiler/Tracks/MemoryGraphSeries.cpp
@@ -26,4 +26,18 @@ FString FMemoryGraphSeries::FormatValue(double Value) const
 
 ////////////////////////////////////////////////////////////////////////////////////////////////////
 
+void FMemoryGraphSeries::ExpandRange(double& InOutMinValue, double& InOutMaxValue, double InValue)
+{
+	if (InValue < InOutMinValue)
+	{
+		InOutMinValue = InValue;
+	}
+	if (InValue > InOutMaxValue)
+	{
+		InOutMaxValue = InValue;
+	}
+}
+
+////////////////////////////////////////////////////////////////////////////////////////////////////
+
 } // namespace UE::Insights::MemoryProfiler
diff --git a/Engine/Source/Developer/TraceInsights/Private/Insights/MemoryProfiler/Tracks/MemoryGraphSeries.h b/Engine/Source/Developer/TraceInsights/Private/Insights/MemoryProfiler/Tracks/MemoryGraphSeries.h
--- a/Engine/Source/Developer/TraceInsights/Private/Insights/MemoryProfiler/Tracks/MemoryGraphSeries.h
+++ b/Engine/Source/Developer/TraceInsights/Private/Insights/MemoryProfiler/Tracks/MemoryGraphSeries.h
@@ -31,6 +31,9 @@ public:
 	double GetMaxValue() const { return MaxValue; }
 	void SetValueRange(double Min, double Max) { MinValue = Min; MaxValue = Max; }
 
+protected:
+	static void ExpandRange(double& InOutMinValue, double& InOutMaxValue, double InValue);
+
 private:
 	double MinValue = 0.0;
 	double MaxValue = 0.0;
diff --git a/Engine/Source/Developer/TraceInsights/Private/Insights/MemoryProfiler/Tracks/MemoryGraphTrack.cpp b/Engine/Source/Developer/TraceInsights/Private/Insights/MemoryProfiler/Tracks/MemoryGraphTrack.cpp
--- a/Engine/Source/Developer/TraceInsights/Private/Insights/MemoryProfiler/Tracks/MemoryGraphTrack.cpp
+++ b/Engine/Source/Developer/TraceInsights/Private/Insights/MemoryProfiler/Tracks/MemoryGraphTrack.cpp
@@ -20,6 +20,7 @@
 #include "Insights/ViewModels/TooltipDrawState.h"
 
 #include <limits>
+#include <cmath>
 
 #define LOCTEXT_NAMESPACE "UE::Insights::MemoryProfiler::FMemoryGraphTrack"
 
@@ -37,14 +38,6 @@ INSIGHTS_IMPLEMENT_RTTI(FMemoryGraphTrack)
 FMemoryGraphTrack::FMemoryGraphTrack(FMemorySharedState& InSharedState)
 	: FGraphTrack()
 	, SharedState(InSharedState)
-	, LabelUnit(EGraphTrackLabelUnit::Auto)
-	, LabelDecimalDigitCount(2)
-	, DefaultMinValue(+std::numeric_limits<double>::infinity())
-	, DefaultMaxValue(-std::numeric_limits<double>::infinity())
-	, AllSeriesMinValue(0.0)
-	, AllSeriesMaxValue(0.0)
-	, bAutoZoom(false)
-	, bIsStacked(false)
 {
 	EnabledOptions = //EGraphOptions::ShowDebugInfo |
 					 //EGraphOptions::ShowPoints |
@@ -54,9 +47,12 @@ FMemoryGraphTrack::FMemoryGraphTrack(FMemorySharedState& InSharedState)
 					 EGraphOptions::UseEventDuration |
 					 //EGraphOptions::ShowBars |
 					 //EGraphOptions::ShowBaseline |
+					 //EGraphOptions::ShowThresholds |
 					 EGraphOptions::ShowVerticalAxisGrid |
 					 EGraphOptions::ShowHeader |
 					 EGraphOptions::None;
+	//VisibleOptions |= EGraphOptions::AutoZoomIncludesBaseline | EGraphOptions::AutoZoomIncludesThresholds;
+	//EditableOptions |= EGraphOptions::AutoZoomIncludesBaseline | EGraphOptions::AutoZoomIncludesThresholds;
 
 	for (uint32 Mode = 0; Mode < static_cast<uint32>(EMemoryTrackHeightMode::Count); ++Mode)
 	{
@@ -81,6 +77,20 @@ FMemoryGraphTrack::~FMemoryGraphTrack()
 
 ////////////////////////////////////////////////////////////////////////////////////////////////////
 
+static void ExpandRange(double& InOutMinValue, double& InOutMaxValue, double InValue)
+{
+	if (InValue < InOutMinValue)
+	{
+		InOutMinValue = InValue;
+	}
+	if (InValue > InOutMaxValue)
+	{
+		InOutMaxValue = InValue;
+	}
+}
+
+////////////////////////////////////////////////////////////////////////////////////////////////////
+
 void FMemoryGraphTrack::Update(const ITimingTrackUpdateContext& Context)
 {
 	FGraphTrack::Update(Context);
@@ -112,6 +122,11 @@ void FMemoryGraphTrack::Update(const ITimingTrackUpdateContext& Context)
 		AllSeriesMinValue = DefaultMinValue;
 		AllSeriesMaxValue = DefaultMaxValue;
 
+		if (IsAnyOptionEnabled(EGraphOptions::AutoZoomIncludesBaseline))
+		{
+			ExpandRange(AllSeriesMinValue, AllSeriesMaxValue, 0.0);
+		}
+
 		for (TSharedPtr<FGraphSeries>& Series : AllSeries)
 		{
 			if (Series->IsVisible() &&
@@ -124,15 +139,24 @@ void FMemoryGraphTrack::Update(const ITimingTrackUpdateContext& Context)
 					MemorySeries.PreUpdate(*this, Viewport);
 				}
 
-				const double CurrentSeriesMinValue = MemorySeries.GetMinValue();
-				const double CurrentSeriesMaxValue = MemorySeries.GetMaxValue();
-				if (CurrentSeriesMinValue < AllSeriesMinValue)
+				const double SeriesMinValue = MemorySeries.GetMinValue();
+				const double SeriesMaxValue = MemorySeries.GetMaxValue();
+				if (SeriesMinValue <= SeriesMaxValue && std::isfinite(SeriesMinValue))
 				{
-					AllSeriesMinValue = CurrentSeriesMinValue;
+					ExpandRange(AllSeriesMinValue, AllSeriesMaxValue, SeriesMinValue);
+					ExpandRange(AllSeriesMinValue, AllSeriesMaxValue, SeriesMaxValue);
 				}
-				if (CurrentSeriesMaxValue > AllSeriesMaxValue)
+
+				if (IsAnyOptionEnabled(EGraphOptions::AutoZoomIncludesThresholds))
 				{
-					AllSeriesMaxValue = CurrentSeriesMaxValue;
+					if (Series->HasHighThresholdValue())
+					{
+						ExpandRange(AllSeriesMinValue, AllSeriesMaxValue, Series->GetHighThresholdValue());
+					}
+					if (Series->HasLowThresholdValue())
+					{
+						ExpandRange(AllSeriesMinValue, AllSeriesMaxValue, Series->GetLowThresholdValue());
+					}
 				}
 			}
 		}
@@ -356,7 +380,6 @@ void FMemoryGraphTrack::DrawVerticalAxisGrid(const ITimingTrackDrawContext& Cont
 
 	FDrawContext& DrawContext = Context.GetDrawContext();
 	const FSlateBrush* Brush = Context.GetHelper().GetWhiteBrush();
-	//const FSlateFontInfo& Font = Context.GetHelper().GetEventFont();
 	const TSharedRef<FSlateFontMeasure> FontMeasureService = FSlateApplication::Get().GetRenderer()->GetFontMeasureService();
 
 	const double TopValue = ViewportY.GetValueAtOffset(RoundedViewHeight);
@@ -408,7 +431,7 @@ void FMemoryGraphTrack::DrawVerticalAxisGrid(const ITimingTrackDrawContext& Cont
 
 		const double StartValue = FMath::GridSnap(BottomValue, Grid);
 
-		TDrawHorizontalAxisLabelParams Params(DrawContext, Brush, FontMeasureService);
+		TDrawLabelParams Params(DrawContext, Brush, FontMeasureService);
 		Params.TextBgColor = FLinearColor(0.05f, 0.05f, 0.05f, 1.0f);
 		Params.TextColor = FLinearColor(1.0f, 1.0f, 1.0f, 1.0f);
 		Params.X = X0;
@@ -434,13 +457,13 @@ void FMemoryGraphTrack::DrawVerticalAxisGrid(const ITimingTrackDrawContext& Cont
 			// Draw label.
 			Params.Y = LabelY;
 			Params.Value = Value;
-			DrawHorizontalAxisLabel(Params);
+			DrawLabel(Params);
 		}
 	}
 
 	if (bHasMinMax && GetHeight() >= TextH)
 	{
-		TDrawHorizontalAxisLabelParams Params(DrawContext, Brush, FontMeasureService);
+		TDrawLabelParams Params(DrawContext, Brush, FontMeasureService);
 
 		if (MainSeries.IsValid() || AllSeries.Num() == 1)
 		{
@@ -458,30 +481,31 @@ void FMemoryGraphTrack::DrawVerticalAxisGrid(const ITimingTrackDrawContext& Cont
 
 		int32 MinMaxAxis = 0;
 
-		// Draw horizontal axis at the max value.
+		// Draw label for the max value.
 		if (MaxValueY >= Y0 && MaxValueY <= Y0 + RoundedViewHeight)
 		{
 			Params.Y = MaxValueY;
 			Params.Value = AllSeriesMaxValue;
-			DrawHorizontalAxisLabel(Params);
+			DrawLabel(Params);
 			++MinMaxAxis;
 		}
 
-		// Draw horizontal axis at the min value.
+		// Draw label for the min value.
 		if (MinValueY >= Y0 && MinValueY <= Y0 + RoundedViewHeight && FMath::Abs(MaxValueY - MinValueY) > TextH)
 		{
 			Params.Y = MinValueY;
 			Params.Value = AllSeriesMinValue;
-			DrawHorizontalAxisLabel(Params);
+			DrawLabel(Params);
 			++MinMaxAxis;
 		}
 
+		// If mouse hovers close to labels area, show also the Max-Min delta value.
 		if (MinMaxAxis == 2)
 		{
 			const float MX = static_cast<float>(Context.GetMousePosition().X);
 			const float MY = static_cast<float>(Context.GetMousePosition().Y);
 
-			constexpr float MX2 = 120.0f; // completely faded out
+			constexpr float MX2 = 120.0f; // width of the hover area
 
 			if (MX > ViewWidth - MX2 && MY >= MaxValueY && MY < MinValueY + TextH)
 			{
@@ -497,7 +521,7 @@ void FMemoryGraphTrack::DrawVerticalAxisGrid(const ITimingTrackDrawContext& Cont
 				Params.Value = AllSeriesMaxValue - AllSeriesMinValue;
 				Params.Precision = -Precision * 0.1; // format with detailed text and increased precision
 				Params.Prefix = TEXT("\u0394=");
-				DrawHorizontalAxisLabel(Params);
+				DrawLabel(Params);
 			}
 		}
 	}
@@ -507,7 +531,7 @@ void FMemoryGraphTrack::DrawVerticalAxisGrid(const ITimingTrackDrawContext& Cont
 
 ////////////////////////////////////////////////////////////////////////////////////////////////////
 
-void FMemoryGraphTrack::DrawHorizontalAxisLabel(const TDrawHorizontalAxisLabelParams& Params) const
+void FMemoryGraphTrack::DrawLabel(const TDrawLabelParams& Params) const
 {
 	FString LabelText = Params.Prefix;
 
diff --git a/Engine/Source/Developer/TraceInsights/Private/Insights/MemoryProfiler/Tracks/MemoryGraphTrack.h b/Engine/Source/Developer/TraceInsights/Private/Insights/MemoryProfiler/Tracks/MemoryGraphTrack.h
--- a/Engine/Source/Developer/TraceInsights/Private/Insights/MemoryProfiler/Tracks/MemoryGraphTrack.h
+++ b/Engine/Source/Developer/TraceInsights/Private/Insights/MemoryProfiler/Tracks/MemoryGraphTrack.h
@@ -13,6 +13,8 @@
 #include "Insights/MemoryProfiler/ViewModels/MemoryTracker.h"
 #include "Insights/ViewModels/GraphTrack.h"
 
+#include <limits>
+
 class FSlateFontMeasure;
 struct FSlateBrush;
 
@@ -94,9 +96,9 @@ public:
 protected:
 	virtual void DrawVerticalAxisGrid(const ITimingTrackDrawContext& Context) const override;
 
-	struct TDrawHorizontalAxisLabelParams
+	struct TDrawLabelParams
 	{
-		TDrawHorizontalAxisLabelParams(FDrawContext& InDrawContext, const FSlateBrush* InBrush, const TSharedRef<FSlateFontMeasure>& InFontMeasureService)
+		TDrawLabelParams(FDrawContext& InDrawContext, const FSlateBrush* InBrush, const TSharedRef<FSlateFontMeasure>& InFontMeasureService)
 		: DrawContext(InDrawContext), Brush(InBrush), FontMeasureService(InFontMeasureService)
 		{
 		}
@@ -106,19 +108,18 @@ protected:
 		const TSharedRef<FSlateFontMeasure>& FontMeasureService;
 		FLinearColor TextBgColor;
 		FLinearColor TextColor;
-		float X;
-		float Y;
-		double Value;
-		double Precision;
-		bool bShowTextDetail;
+		float X = 0.0f;
+		float Y = 0.0f;
+		double Value = 0.0;
+		double Precision = 0.0; // if Precision < 0, formats the value with detailed text
 		FString Prefix;
 	};
-	void DrawHorizontalAxisLabel(const TDrawHorizontalAxisLabelParams& Params) const;
+	void DrawLabel(const TDrawLabelParams& Params) const;
 
 protected:
 	FMemorySharedState& SharedState;
 
-	EGraphTrackLabelUnit LabelUnit;
+	EGraphTrackLabelUnit LabelUnit = EGraphTrackLabelUnit::Auto;
 
 	/**
 	 * Number of decimal digits for labels.
@@ -126,18 +127,18 @@ protected:
 	 * If negative, the formatting will use maximum the number of decimal digits specified (trims trailing 0s),
 	 * otherwise, it will use exactly the number of decimal digits specified.
 	 */
-	int32 LabelDecimalDigitCount;
+	int32 LabelDecimalDigitCount = 2;
 
-	double DefaultMinValue;
-	double DefaultMaxValue;
-	double AllSeriesMinValue;
-	double AllSeriesMaxValue;
+	double DefaultMinValue = +std::numeric_limits<double>::infinity();
+	double DefaultMaxValue = -std::numeric_limits<double>::infinity();
+	double AllSeriesMinValue = 0.0;
+	double AllSeriesMaxValue = 0.0;
 
-	bool bAutoZoom; // all series will share same scale
+	bool bAutoZoom = false; // all series will share same scale
+	bool bIsStacked = false;
 
 	float AvailableTrackHeights[static_cast<uint32>(EMemoryTrackHeightMode::Count)];
 
-	bool bIsStacked;
 	TSharedPtr<FMemoryGraphSeries> MainSeries;
 };
 
diff --git a/Engine/Source/Developer/TraceInsights/Private/Insights/MemoryProfiler/ViewModels/MemTagNode.h b/Engine/Source/Developer/TraceInsights/Private/Insights/MemoryProfiler/ViewModels/MemTagNode.h
--- a/Engine/Source/Developer/TraceInsights/Private/Insights/MemoryProfiler/ViewModels/MemTagNode.h
+++ b/Engine/Source/Developer/TraceInsights/Private/Insights/MemoryProfiler/ViewModels/MemTagNode.h
@@ -127,6 +127,9 @@ public:
 	int64 GetSizeMax() const { return Stats.SizeMax; }
 	int64 GetSizeAverage() const { return Stats.SizeAverage; }
 
+	bool HasSizeBudget() const { return Stats.SizeBudget != 0; }
+	void ResetSizeBudget() { Stats.SizeBudget = 0; }
+	void SetSizeBudget(int64 InSizeBudget) { Stats.SizeBudget = InSizeBudget; }
 	void ResetAggregatedStats();
 
 private:
diff --git a/Engine/Source/Developer/TraceInsights/Private/Insights/MemoryProfiler/ViewModels/MemTagTable.cpp b/Engine/Source/Developer/TraceInsights/Private/Insights/MemoryProfiler/ViewModels/MemTagTable.cpp
--- a/Engine/Source/Developer/TraceInsights/Private/Insights/MemoryProfiler/ViewModels/MemTagTable.cpp
+++ b/Engine/Source/Developer/TraceInsights/Private/Insights/MemoryProfiler/ViewModels/MemTagTable.cpp
@@ -118,7 +118,7 @@ void FMemTagTable::AddDefaultColumns()
 								if (Value.IsSet())
 								{
 									const FMemTagNode& MemTagNode = Node.As<FMemTagNode>();
-									if (MemTagNode.GetSizeBudget() != 0)
+									if (MemTagNode.HasSizeBudget())
 									{
 										int64 SizeValue = Value.GetValue().AsInt64();
 										if (SizeValue > MemTagNode.GetSizeBudget())
@@ -643,7 +643,7 @@ void FMemTagTable::AddDefaultColumns()
 				if (Node.Is<FMemTagNode>())
 				{
 					const FMemTagNode& MemTagNode = Node.As<FMemTagNode>();
-					if (MemTagNode.GetSizeBudget() != 0)
+					if (MemTagNode.HasSizeBudget())
 					{
 						return FTableCellValue(MemTagNode.GetSizeBudget());
 					}
diff --git a/Engine/Source/Developer/TraceInsights/Private/Insights/MemoryProfiler/ViewModels/MemorySharedState.cpp b/Engine/Source/Developer/TraceInsights/Private/Insights/MemoryProfiler/ViewModels/MemorySharedState.cpp
--- a/Engine/Source/Developer/TraceInsights/Private/Insights/MemoryProfiler/ViewModels/MemorySharedState.cpp
+++ b/Engine/Source/Developer/TraceInsights/Private/Insights/MemoryProfiler/ViewModels/MemorySharedState.cpp
@@ -707,49 +707,54 @@ TSharedPtr<FMemoryGraphTrack> FMemorySharedState::CreateMemTagGraphTrack(FMemory
 		TimingView->HandleTrackVisibilityChanged();
 	}
 
-	TSharedPtr<FMemoryGraphTrack> GraphTrack = GetMemTagGraphTrack(InMemTrackerId, InMemTagId);
+	TSharedPtr<FMemoryGraphTrack> GraphTrackPtr = GetMemTagGraphTrack(InMemTrackerId, InMemTagId);
 
-	if (!GraphTrack.IsValid())
+	if (!GraphTrackPtr.IsValid())
 	{
-		// Create new Graph track.
-		GraphTrack = MakeShared<FMemoryGraphTrack>(*this);
+		// Create a new Graph track.
+		GraphTrackPtr = MakeShared<FMemoryGraphTrack>(*this);
+		FMemoryGraphTrack& GraphTrack = *GraphTrackPtr;
+
+		GraphTrack.SetVisibleOptions(GraphTrack.GetVisibleOptions() | EGraphOptions::AutoZoomIncludesBaseline | EGraphOptions::AutoZoomIncludesThresholds);
+		GraphTrack.SetEditableOptions(GraphTrack.GetEditableOptions() | EGraphOptions::AutoZoomIncludesBaseline | EGraphOptions::AutoZoomIncludesThresholds);
+		GraphTrack.SetEnabledOptions(GraphTrack.GetEnabledOptions() | EGraphOptions::ShowThresholds);
 
 		const int32 Order = GetNextMemoryGraphTrackOrder();
-		GraphTrack->SetOrder(Order);
-		GraphTrack->SetName(SeriesName);
-		//GraphTrack->SetVisibilityFlag(bShowHideAllMemoryTracks);
-		GraphTrack->Show();
+		GraphTrack.SetOrder(Order);
+		GraphTrack.SetName(SeriesName);
+		GraphTrack.Show();
 
-		GraphTrack->SetAvailableTrackHeight(EMemoryTrackHeightMode::Small, 32.0f);
-		GraphTrack->SetAvailableTrackHeight(EMemoryTrackHeightMode::Medium, 100.0f);
-		GraphTrack->SetAvailableTrackHeight(EMemoryTrackHeightMode::Large, 200.0f);
-		GraphTrack->SetCurrentTrackHeight(TrackHeightMode);
+		GraphTrack.SetAvailableTrackHeight(EMemoryTrackHeightMode::Small, 32.0f);
+		GraphTrack.SetAvailableTrackHeight(EMemoryTrackHeightMode::Medium, 100.0f);
+		GraphTrack.SetAvailableTrackHeight(EMemoryTrackHeightMode::Large, 200.0f);
+		GraphTrack.SetCurrentTrackHeight(TrackHeightMode);
 
-		GraphTrack->EnableAutoZoom();
+		GraphTrack.EnableAutoZoom();
 
-		// Create series.
-		TSharedPtr<FMemTagGraphSeries> Series = GraphTrack->AddMemTagSeries(InMemTrackerId, MemTagSetId, InMemTagId);
+		// Create a new MemTag graph series.
+		TSharedPtr<FMemTagGraphSeries> Series = GraphTrack.AddMemTagSeries(InMemTrackerId, MemTagSetId, InMemTagId);
 		Series->SetName(SeriesName);
 		Series->SetColor(Color, BorderColor);
-		Series->SetBaselineY(GraphTrack->GetHeight() - 1.0f);
+		Series->SetBaselineY(GraphTrack.GetHeight() - 1.0f);
 		Series->EnableAutoZoom();
 
 		if (TagPtr)
 		{
-			TagPtr->AddTrack(GraphTrack);
+			TagPtr->AddTrack(GraphTrackPtr);
 		}
 
-		// Add the new Graph in scrollable tracks.
-		TimingView->AddScrollableTrack(GraphTrack);
-		AllTracks.Add(GraphTrack);
+		// Add the new Graph to the TimingView as a scrollable track.
+		TimingView->AddScrollableTrack(GraphTrackPtr);
+
+		AllTracks.Add(GraphTrackPtr);
 	}
 	else
 	{
-		GraphTrack->Show();
+		GraphTrackPtr->Show();
 		TimingView->HandleTrackVisibilityChanged();
 	}
 
-	return GraphTrack;
+	return GraphTrackPtr;
 }
 
 ////////////////////////////////////////////////////////////////////////////////////////////////////
diff --git a/Engine/Source/Developer/TraceInsights/Private/Insights/MemoryProfiler/Widgets/SMemTagTreeView.cpp b/Engine/Source/Developer/TraceInsights/Private/Insights/MemoryProfiler/Widgets/SMemTagTreeView.cpp
--- a/Engine/Source/Developer/TraceInsights/Private/Insights/MemoryProfiler/Widgets/SMemTagTreeView.cpp
+++ b/Engine/Source/Developer/TraceInsights/Private/Insights/MemoryProfiler/Widgets/SMemTagTreeView.cpp
@@ -627,7 +627,11 @@ void SMemTagTreeView::TreeView_OnMouseButtonDoubleClick(FTableTreeNodePtr NodePt
 				TSharedPtr<FMemoryGraphTrack> GraphTrack = SharedState.GetMemTagGraphTrack(MemTrackerId, MemTagId);
 				if (!GraphTrack.IsValid())
 				{
-					SharedState.CreateMemTagGraphTrack(MemTrackerId, MemTagId);
+					GraphTrack = SharedState.CreateMemTagGraphTrack(MemTrackerId, MemTagId);
+					if (GraphTrack.IsValid())
+					{
+						UpdateHighThreshold(MemTagNode, *GraphTrack);
+					}
 				}
 				else
 				{
@@ -2144,37 +2148,51 @@ void SMemTagTreeView::ApplyBudgetToNodes()
 	const FMemTagBudgetValueSet& DefaultValues = BudgetMode.GetDefaultValues();
 	const FMemTagBudgetValueSet* PlatformOverrides = SelectedBudgetPlatform.IsValid() ? BudgetMode.FindPlatform(SelectedBudgetPlatform->Name) : nullptr;
 
+	TSharedPtr<SMemoryProfilerWindow> ProfilerWindow = GetProfilerWindow();
+	if (!ProfilerWindow.IsValid())
+	{
+		return;
+	}
+	FMemorySharedState& SharedState = ProfilerWindow->GetSharedState();
+
 	for (const FTableTreeNodePtr& NodePtr : TableRowNodes)
 	{
 		if (!NodePtr->Is<FMemTagNode>())
 		{
 			continue;
 		}
+
 		FMemTagNode& MemTagNode = NodePtr->As<FMemTagNode>();
 		FMemoryTag* MemoryTag = MemTagNode.GetMemTag();
 		if (!MemoryTag)
 		{
 			continue;
 		}
+
 		if (PlatformOverrides)
 		{
 			const int64* ValuePtr = PlatformOverrides->FindValue(MemoryTag->GetStatFullName());
 			if (ValuePtr)
 			{
-				MemTagNode.GetStats().SizeBudget = *ValuePtr;
+				MemTagNode.SetSizeBudget(*ValuePtr);
+				UpdateHighThreshold(MemTagNode, SharedState);
 				continue;
 			}
 		}
+
 		// DefaultValues
 		{
 			const int64* ValuePtr = DefaultValues.FindValue(MemoryTag->GetStatFullName());
 			if (ValuePtr)
 			{
-				MemTagNode.GetStats().SizeBudget = *ValuePtr;
+				MemTagNode.SetSizeBudget(*ValuePtr);
+				UpdateHighThreshold(MemTagNode, SharedState);
 				continue;
 			}
 		}
-		MemTagNode.GetStats().SizeBudget = 0;
+
+		MemTagNode.ResetSizeBudget();
+		UpdateHighThreshold(MemTagNode, SharedState);
 	}
 }
 
@@ -2182,6 +2200,13 @@ void SMemTagTreeView::ApplyBudgetToNodes()
 
 void SMemTagTreeView::ResetBudgetForAllNodes()
 {
+	TSharedPtr<SMemoryProfilerWindow> ProfilerWindow = GetProfilerWindow();
+	if (!ProfilerWindow.IsValid())
+	{
+		return;
+	}
+	FMemorySharedState& SharedState = ProfilerWindow->GetSharedState();
+
 	for (const FTableTreeNodePtr& NodePtr : TableRowNodes)
 	{
 		if (!NodePtr->Is<FMemTagNode>())
@@ -2189,7 +2214,44 @@ void SMemTagTreeView::ResetBudgetForAllNodes()
 			continue;
 		}
 		FMemTagNode& MemTagNode = NodePtr->As<FMemTagNode>();
-		MemTagNode.GetStats().SizeBudget = 0;
+		MemTagNode.ResetSizeBudget();
+		UpdateHighThreshold(MemTagNode, SharedState);
+	}
+}
+
+////////////////////////////////////////////////////////////////////////////////////////////////////
+
+void SMemTagTreeView::UpdateHighThreshold(const FMemTagNode& MemTagNode, FMemorySharedState& SharedState)
+{
+	TSharedPtr<FMemoryGraphTrack> GraphTrack = SharedState.GetMemTagGraphTrack(MemTagNode.GetMemTrackerId(), MemTagNode.GetMemTagId());
+	if (GraphTrack.IsValid())
+	{
+		UpdateHighThreshold(MemTagNode, *GraphTrack);
+	}
+}
+
+////////////////////////////////////////////////////////////////////////////////////////////////////
+
+void SMemTagTreeView::UpdateHighThreshold(const FMemTagNode& MemTagNode, FMemoryGraphTrack& GraphTrack)
+{
+	for (TSharedPtr<FGraphSeries>& Series : GraphTrack.GetSeries())
+	{
+		if (Series->Is<FMemTagGraphSeries>())
+		{
+			FMemTagGraphSeries& MemTagGraphSeries = Series->As<FMemTagGraphSeries>();
+			if (MemTagGraphSeries.GetTrackerId() == MemTagNode.GetMemTrackerId() &&
+				MemTagGraphSeries.GetTagId() == MemTagNode.GetMemTagId())
+			{
+				if (MemTagNode.HasSizeBudget())
+				{
+					MemTagGraphSeries.SetHighThresholdValue((double)MemTagNode.GetSizeBudget());
+				}
+				else
+				{
+					MemTagGraphSeries.ResetHighThresholdValue();
+				}
+			}
+		}
 	}
 }
 
diff --git a/Engine/Source/Developer/TraceInsights/Private/Insights/MemoryProfiler/Widgets/SMemTagTreeView.h b/Engine/Source/Developer/TraceInsights/Private/Insights/MemoryProfiler/Widgets/SMemTagTreeView.h
--- a/Engine/Source/Developer/TraceInsights/Private/Insights/MemoryProfiler/Widgets/SMemTagTreeView.h
+++ b/Engine/Source/Developer/TraceInsights/Private/Insights/MemoryProfiler/Widgets/SMemTagTreeView.h
@@ -299,6 +299,8 @@ private:
 	void OnBudgetChanged();
 	void ApplyBudgetToNodes();
 	void ResetBudgetForAllNodes();
+	void UpdateHighThreshold(const FMemTagNode& MemTagNode, FMemorySharedState& SharedState);
+	void UpdateHighThreshold(const FMemTagNode& MemTagNode, FMemoryGraphTrack& GraphTrack);
 
 private:
 	/** A weak pointer to the Memory Insights window. */
diff --git a/Engine/Source/Developer/TraceInsights/Private/Insights/ViewModels/GraphSeries.cpp b/Engine/Source/Developer/TraceInsights/Private/Insights/ViewModels/GraphSeries.cpp
--- a/Engine/Source/Developer/TraceInsights/Private/Insights/ViewModels/GraphSeries.cpp
+++ b/Engine/Source/Developer/TraceInsights/Private/Insights/ViewModels/GraphSeries.cpp
@@ -17,22 +17,6 @@ INSIGHTS_IMPLEMENT_RTTI(FGraphSeries)
 ////////////////////////////////////////////////////////////////////////////////////////////////////
 
 FGraphSeries::FGraphSeries()
-	: Name()
-	, Description()
-	, bIsVisible(true)
-	, bIsDirty(false)
-	, bHasEventDuration(true)
-	, bAutoZoom(false)
-	, bIsAutoZoomDirty(false)
-	, bUseSharedViewport(false)
-	, ValueViewport()
-	, Color(0.0f, 0.5f, 1.0f, 1.0f)
-	, FillColor(0.0f, 0.5f, 1.0f, 1.0f)
-	, BorderColor(0.3f, 0.8f, 1.0f, 1.0f)
-	//, Events()
-	//, Points()
-	//, LinePoints()
-	//, Boxes()
 {
 }
 
diff --git a/Engine/Source/Developer/TraceInsights/Private/Insights/ViewModels/GraphTrack.cpp b/Engine/Source/Developer/TraceInsights/Private/Insights/ViewModels/GraphTrack.cpp
--- a/Engine/Source/Developer/TraceInsights/Private/Insights/ViewModels/GraphTrack.cpp
+++ b/Engine/Source/Developer/TraceInsights/Private/Insights/ViewModels/GraphTrack.cpp
@@ -35,19 +35,10 @@ INSIGHTS_IMPLEMENT_RTTI(FGraphTrack)
 
 FGraphTrack::FGraphTrack()
 	: FBaseTimingTrack()
-	//, AllSeries()
 	, WhiteBrush(FInsightsStyle::Get().GetBrush("WhiteBrush"))
 	, PointBrush(FInsightsStyle::GetBrush("Graph.Point"))
 	, BorderBrush(FInsightsStyle::Get().GetBrush("SingleBorder"))
 	, Font(FAppStyle::Get().GetFontStyle("SmallFont"))
-	, EnabledOptions(EGraphOptions::DefaultEnabledOptions)
-	, VisibleOptions(EGraphOptions::DefaultVisibleOptions)
-	, EditableOptions(EGraphOptions::DefaultEditableOptions)
-	, SharedValueViewport()
-	, NumAddedEvents(0)
-	, NumDrawPoints(0)
-	, NumDrawLines(0)
-	, NumDrawBoxes(0)
 {
 	SetValidLocations(ETimingTrackLocation::Scrollable | ETimingTrackLocation::TopDocked | ETimingTrackLocation::BottomDocked);
 }
@@ -56,20 +47,10 @@ FGraphTrack::FGraphTrack()
 
 FGraphTrack::FGraphTrack(const FString& InName)
 	: FBaseTimingTrack(InName)
-	//, AllSeries()
 	, WhiteBrush(FInsightsStyle::Get().GetBrush("WhiteBrush"))
 	, PointBrush(FInsightsStyle::GetBrush("Graph.Point"))
 	, BorderBrush(FInsightsStyle::Get().GetBrush("SingleBorder"))
 	, Font(FAppStyle::Get().GetFontStyle("SmallFont"))
-	, EnabledOptions(EGraphOptions::DefaultEnabledOptions)
-	, VisibleOptions(EGraphOptions::DefaultVisibleOptions)
-	, EditableOptions(EGraphOptions::DefaultEditableOptions)
-	, SharedValueViewport()
-	, TimeScaleX(1.0)
-	, NumAddedEvents(0)
-	, NumDrawPoints(0)
-	, NumDrawLines(0)
-	, NumDrawBoxes(0)
 {
 	SetValidLocations(ETimingTrackLocation::Scrollable | ETimingTrackLocation::TopDocked | ETimingTrackLocation::BottomDocked);
 }
@@ -149,6 +130,9 @@ void FGraphTrack::Draw(const ITimingTrackDrawContext& Context) const
 	UE::Insights::FDrawContext& DrawContext = Context.GetDrawContext();
 	const FTimingTrackViewport& Viewport = Context.GetViewport();
 
+	const float AlignedPosY = FMath::RoundToFloat(GetPosY());
+	const float AlignedHeight = FMath::RoundToFloat(GetHeight());
+
 	// Set clipping area.
 	{
 		const FVector2f AbsPos = FVector2f(DrawContext.Geometry.GetAbsolutePosition());
@@ -164,29 +148,62 @@ void FGraphTrack::Draw(const ITimingTrackDrawContext& Context) const
 	// Draw top line.
 	//DrawContext.DrawBox(0.0f, GetPosY(), Viewport.Width, 1.0f, WhiteBrush, FLinearColor(0.05f, 0.05f, 0.05f, 1.0f));
 
-	bool bDrawnBaseline = false;
-	const bool bDrawBaseline = IsAnyOptionEnabled(EGraphOptions::ShowBaseline);
+	bool bDrawBaseline = IsAnyOptionEnabled(EGraphOptions::ShowBaseline);
 	for (const TSharedPtr<FGraphSeries>& Series : AllSeries)
 	{
 		if (Series->IsVisible())
 		{
-			// Draw baseline (Value == 0), for the first visible series only.
-			if (bDrawBaseline && !bDrawnBaseline)
+			// Draw baseline (Value == 0)...
+			if (bDrawBaseline)
 			{
+				bDrawBaseline = false; // only for the first visible series
+
 				const float BaselineY = FMath::RoundToFloat(static_cast<float>(Series->GetBaselineY()));
-				if (BaselineY >= 0.0f && BaselineY < GetHeight())
+				if (BaselineY >= 0.0f && BaselineY < AlignedHeight)
 				{
-					DrawContext.DrawBox(0.0f, GetPosY() + BaselineY, Viewport.GetWidth(), 1.0f, WhiteBrush, FLinearColor(0.05f, 0.05f, 0.05f, 1.0f));
+					DrawContext.DrawBox(0.0f, AlignedPosY + BaselineY, Viewport.GetWidth(), 1.0f, WhiteBrush, FLinearColor(0.05f, 0.05f, 0.05f, 1.0f));
 					DrawContext.LayerId++;
 				}
-
-				bDrawnBaseline = true;
 			}
 
 			DrawSeries(*Series, DrawContext, Viewport);
 		}
 	}
 
+	if (IsAnyOptionEnabled(EGraphOptions::ShowThresholds))
+	{
+		const FLinearColor HighThresholdLineColor(1.0f, 0.0f, 0.0f, 0.75f);
+		const FLinearColor LowThresholdLineColor(1.0f, 1.0f, 0.0f, 0.75f);
+
+		for (const TSharedPtr<FGraphSeries>& Series : AllSeries)
+		{
+			if (Series->IsVisible() && (Series->HasHighThresholdValue() || Series->HasLowThresholdValue()))
+			{
+				if (Series->HasHighThresholdValue())
+				{
+					const float ThresholdLineY = Series->ValueViewport.GetRoundedYForValue(Series->GetHighThresholdValue());
+					if (ThresholdLineY>= 0.0f && ThresholdLineY < AlignedHeight)
+					{
+						DrawContext.DrawBox(0.0f, AlignedPosY + ThresholdLineY, Viewport.GetWidth(), 1.0f, WhiteBrush, HighThresholdLineColor);
+						DrawContext.LayerId++;
+					}
+				}
+
+				if (Series->HasLowThresholdValue())
+				{
+					const float ThresholdLineY = Series->ValueViewport.GetRoundedYForValue(Series->GetLowThresholdValue());
+					if (ThresholdLineY >= 0.0f && ThresholdLineY < AlignedHeight)
+					{
+						DrawContext.DrawBox(0.0f, AlignedPosY + ThresholdLineY, Viewport.GetWidth(), 1.0f, WhiteBrush, LowThresholdLineColor);
+						DrawContext.LayerId++;
+					}
+				}
+
+				break; // only for the first visible series with valid thresholds
+			}
+		}
+	}
+
 	if (IsAnyOptionEnabled(EGraphOptions::ShowVerticalAxisGrid))
 	{
 		DrawVerticalAxisGrid(Context);
@@ -811,6 +828,44 @@ void FGraphTrack::BuildContextMenu(FMenuBuilder& MenuBuilder)
 				EUserInterfaceActionType::ToggleButton
 			);
 		}
+
+		if (EnumHasAnyFlags(VisibleOptions, EGraphOptions::AutoZoomIncludesBaseline))
+		{
+			FUIAction Action_AutoZoomIncludesBaseline
+			(
+				FExecuteAction::CreateSP(this, &FGraphTrack::ContextMenu_ToggleOption_Execute, EGraphOptions::AutoZoomIncludesBaseline),
+				FCanExecuteAction::CreateSP(this, &FGraphTrack::ContextMenu_ToggleOption_CanExecute, EGraphOptions::AutoZoomIncludesBaseline),
+				FIsActionChecked::CreateSP(this, &FGraphTrack::ContextMenu_ToggleOption_IsChecked, EGraphOptions::AutoZoomIncludesBaseline)
+			);
+			MenuBuilder.AddMenuEntry
+			(
+				LOCTEXT("ContextMenu_AutoZoomIncludesBaseline", "Auto Zoom Includes Baseline"),
+				LOCTEXT("ContextMenu_AutoZoomIncludesBaseline_Desc", "If enabled, the auto zoom will also include the baseline (value == 0)."),
+				FSlateIcon(),
+				Action_AutoZoomIncludesBaseline,
+				NAME_None,
+				EUserInterfaceActionType::ToggleButton
+			);
+		}
+
+		if (EnumHasAnyFlags(VisibleOptions, EGraphOptions::AutoZoomIncludesThresholds))
+		{
+			FUIAction Action_AutoZoomIncludesThresholds
+			(
+				FExecuteAction::CreateSP(this, &FGraphTrack::ContextMenu_ToggleOption_Execute, EGraphOptions::AutoZoomIncludesThresholds),
+				FCanExecuteAction::CreateSP(this, &FGraphTrack::ContextMenu_ToggleOption_CanExecute, EGraphOptions::AutoZoomIncludesThresholds),
+				FIsActionChecked::CreateSP(this, &FGraphTrack::ContextMenu_ToggleOption_IsChecked, EGraphOptions::AutoZoomIncludesThresholds)
+			);
+			MenuBuilder.AddMenuEntry
+			(
+				LOCTEXT("ContextMenu_AutoZoomIncludesThresholds", "Auto Zoom Includes Thresholds"),
+				LOCTEXT("ContextMenu_AutoZoomIncludesThresholds_Desc", "If enabled, the auto zoom will also include the high and low thresholds, if available."),
+				FSlateIcon(),
+				Action_AutoZoomIncludesThresholds,
+				NAME_None,
+				EUserInterfaceActionType::ToggleButton
+			);
+		}
 	}
 	MenuBuilder.EndSection();
 
@@ -885,6 +940,7 @@ FRandomGraphTrack::FRandomGraphTrack()
 					 //EGraphOptions::UseEventDuration |
 					 //EGraphOptions::ShowBars |
 					 //EGraphOptions::ShowBaseline |
+					 //EGraphOptions::ShowThresholds |
 					 //EGraphOptions::ShowVerticalAxisGrid |
 					 EGraphOptions::ShowHeader |
 					 EGraphOptions::None;
diff --git a/Engine/Source/Developer/TraceInsights/Private/Insights/ViewModels/TimingGraphTrack.cpp b/Engine/Source/Developer/TraceInsights/Private/Insights/ViewModels/TimingGraphTrack.cpp
--- a/Engine/Source/Developer/TraceInsights/Private/Insights/ViewModels/TimingGraphTrack.cpp
+++ b/Engine/Source/Developer/TraceInsights/Private/Insights/ViewModels/TimingGraphTrack.cpp
@@ -141,12 +141,15 @@ INSIGHTS_IMPLEMENT_RTTI(FTimingGraphTrack)
 FTimingGraphTrack::FTimingGraphTrack(TSharedPtr<UE::Insights::TimingProfiler::STimingView> InTimingView)
 	: FGraphTrack()
 	, TimingView(InTimingView)
-	//, SharedValueViewport()
 {
 	LoadDefaultSettings();
-	
+
 	// Add non editable options.
-	EnabledOptions = EnabledOptions | EGraphOptions::ShowBaseline | EGraphOptions::ShowVerticalAxisGrid | EGraphOptions::ShowHeader;
+	EnabledOptions |=
+		EGraphOptions::ShowBaseline |
+		EGraphOptions::ShowThresholds |
+		EGraphOptions::ShowVerticalAxisGrid |
+		EGraphOptions::ShowHeader;
 
 	bNotifyTimersOnDestruction = InTimingView->GetName() == FInsightsManagerTabs::TimingProfilerTabId;
 }
@@ -1025,8 +1028,8 @@ uint32 FTimingGraphTrack::GetNumSeriesForTimer(uint32 TimerId)
 	for(const TSharedPtr<FGraphSeries>& Series : AllSeries)
 	{
 		const TSharedPtr<FTimingGraphSeries> TimingSeries = StaticCastSharedPtr<FTimingGraphSeries>(Series);
-		if (TimingSeries.IsValid() && 
-			(TimingSeries->Type == FTimingGraphSeries::ESeriesType::Timer || TimingSeries->Type == FTimingGraphSeries::ESeriesType::FrameStatsTimer) && 
+		if (TimingSeries.IsValid() &&
+			(TimingSeries->Type == FTimingGraphSeries::ESeriesType::Timer || TimingSeries->Type == FTimingGraphSeries::ESeriesType::FrameStatsTimer) &&
 			TimingSeries->TimerId == TimerId)
 		{
 			++NumSeries;
diff --git a/Engine/Source/Developer/TraceInsights/Public/Insights/ViewModels/GraphSeries.h b/Engine/Source/Developer/TraceInsights/Public/Insights/ViewModels/GraphSeries.h
--- a/Engine/Source/Developer/TraceInsights/Public/Insights/ViewModels/GraphSeries.h
+++ b/Engine/Source/Developer/TraceInsights/Public/Insights/ViewModels/GraphSeries.h
@@ -16,8 +16,11 @@
 #include "InsightsCore/Common/SimpleRtti.h"
 
 // TraceInsights
+#include "Insights/ViewModels/BaseTimingTrack.h"
 #include "Insights/ViewModels/GraphTrackEvent.h"
 
+#include <limits>
+
 class FTimingTrackViewport;
 struct FGraphSeriesEvent;
 
@@ -141,10 +144,22 @@ public:
 	void DisableAutoZoom() { bAutoZoom = false; }
 
 	bool IsAutoZoomDirty() const { return bIsAutoZoomDirty; }
+	void SetAutoZoomDirty() { bIsAutoZoomDirty = true; }
+	void ResetAutoZoomDirty() { bIsAutoZoomDirty = false; }
 
 	bool IsUsingSharedViewport() const { return bUseSharedViewport; }
 	void EnableSharedViewport() { bUseSharedViewport = true; }
 
+	virtual bool HasHighThresholdValue() const { return false; }
+	virtual double GetHighThresholdValue() const { return +std::numeric_limits<double>::infinity(); }
+	virtual void SetHighThresholdValue(double InValue) {}
+	virtual void ResetHighThresholdValue() {}
+
+	virtual bool HasLowThresholdValue() const { return false; }
+	virtual double GetLowThresholdValue() const { return -std::numeric_limits<double>::infinity(); }
+	virtual void SetLowThresholdValue(double InValue) {}
+	virtual void ResetLowThresholdValue() {}
+
 	//////////////////////////////////////////////////
 
 	/**
@@ -214,20 +229,20 @@ private:
 	FText Name;
 	FText Description;
 
-	bool bIsVisible;
-	bool bIsDirty;
+	bool bIsVisible = true;
+	bool bIsDirty = false;
 
-	bool bHasEventDuration;
+	bool bHasEventDuration = true;
 
-	bool bAutoZoom;
-	bool bIsAutoZoomDirty;
+	bool bAutoZoom = false;
+	bool bIsAutoZoomDirty = false;
 
-	bool bUseSharedViewport;
+	bool bUseSharedViewport = false;
 	FGraphValueViewport ValueViewport;
 
-	FLinearColor Color;
-	FLinearColor FillColor;
-	FLinearColor BorderColor;
+	FLinearColor Color       = FLinearColor(0.0f, 0.5f, 1.0f, 1.0f);
+	FLinearColor FillColor   = FLinearColor(0.0f, 0.5f, 1.0f, 1.0f);
+	FLinearColor BorderColor = FLinearColor(0.3f, 0.8f, 1.0f, 1.0f);
 
 protected:
 	TArray<FGraphSeriesEvent> Events; // reduced list of events; used to identify an event at a certain screen position (ex.: the event hovered by mouse)
diff --git a/Engine/Source/Developer/TraceInsights/Public/Insights/ViewModels/GraphTrack.h b/Engine/Source/Developer/TraceInsights/Public/Insights/ViewModels/GraphTrack.h
--- a/Engine/Source/Developer/TraceInsights/Public/Insights/ViewModels/GraphTrack.h
+++ b/Engine/Source/Developer/TraceInsights/Public/Insights/ViewModels/GraphTrack.h
@@ -5,6 +5,7 @@
 #include "CoreTypes.h"
 
 #include "Fonts/SlateFontInfo.h"
+#include "Misc/EnumClassFlags.h"
 #include "Templates/SharedPointer.h"
 
 // TraceInsights
@@ -21,20 +22,24 @@ class FTimingTrackViewport;
 // Various available options for display
 enum class EGraphOptions
 {
-	None					= 0,
+	None = 0,
 
-	ShowDebugInfo			= (1 << 0),
-	ShowPoints				= (1 << 1),
-	ShowPointsWithBorder	= (1 << 2),
-	ShowLines				= (1 << 3),
-	ShowPolygon				= (1 << 4),
-	UseEventDuration		= (1 << 5),
-	ShowBars				= (1 << 6),
-	ShowBaseline			= (1 << 7),
-	ShowVerticalAxisGrid	= (1 << 8),
-	ShowHeader				= (1 << 9),
+	ShowDebugInfo				= (1 << 0),
+	ShowPoints					= (1 << 1),
+	ShowPointsWithBorder		= (1 << 2),
+	ShowLines					= (1 << 3),
+	ShowPolygon					= (1 << 4),
+	UseEventDuration			= (1 << 5),
+	ShowBars					= (1 << 6),
+	ShowBaseline				= (1 << 7),
+	ShowThresholds				= (1 << 8),
+	ShowVerticalAxisGrid		= (1 << 9),
+	ShowHeader					= (1 << 10),
 
-	FirstCustomOption		= (1 << 10),
+	AutoZoomIncludesBaseline	= (1 << 11),
+	AutoZoomIncludesThresholds	= (1 << 12),
+
+	FirstCustomOption			= (1 << 13),
 
 	DefaultEnabledOptions	= None,
 	DefaultVisibleOptions	= ShowPoints | ShowPointsWithBorder | ShowLines | ShowPolygon | UseEventDuration | ShowBars,
@@ -137,25 +142,26 @@ protected:
 	TArray<TSharedPtr<FGraphSeries>> AllSeries;
 
 	// Slate resources
-	const FSlateBrush* WhiteBrush;
-	const FSlateBrush* PointBrush;
-	const FSlateBrush* BorderBrush;
+	const FSlateBrush* WhiteBrush = nullptr;
+	const FSlateBrush* PointBrush = nullptr;
+	const FSlateBrush* BorderBrush = nullptr;
 	const FSlateFontInfo Font;
 
 	// Flags controlling various Graph options
-	EGraphOptions EnabledOptions; // currently enabled options
-	EGraphOptions VisibleOptions; // if the option is visible in the context menu
-	EGraphOptions EditableOptions; // if the option is editable from the context menu (if false, the option can be readonly)
+	EGraphOptions EnabledOptions = EGraphOptions::DefaultEnabledOptions; // currently enabled options
+	EGraphOptions VisibleOptions = EGraphOptions::DefaultVisibleOptions; // if the option is visible in the context menu
+	EGraphOptions EditableOptions = EGraphOptions::DefaultEditableOptions; // if the option is editable from the context menu (if false, the option can be readonly)
 
 	FGraphValueViewport SharedValueViewport;
 
-	double TimeScaleX;
+	// Time scale (horizontal axis) saved from "graph update" and used later in setting up tooltips.
+	double TimeScaleX = 1.0;
 
 	// Stats
-	int32 NumAddedEvents; // total event count
-	int32 NumDrawPoints;
-	int32 NumDrawLines;
-	int32 NumDrawBoxes;
+	int32 NumAddedEvents = 0; // total event count
+	int32 NumDrawPoints = 0;
+	int32 NumDrawLines = 0;
+	int32 NumDrawBoxes = 0;
 };
 
 ////////////////////////////////////////////////////////////////////////////////////////////////////
