[Insights] MemoryInsights: Changed "Live Allocation Count" and "Alloc/Free Event Count" tracks to not appear if trace was not recorded with detailed mem trace events (i.e. "memalloc" channel).

diff --git a/Engine/Source/Developer/TraceInsights/Private/Insights/MemoryProfiler/ViewModels/MemorySharedState.cpp b/Engine/Source/Developer/TraceInsights/Private/Insights/MemoryProfiler/ViewModels/MemorySharedState.cpp
--- a/Engine/Source/Developer/TraceInsights/Private/Insights/MemoryProfiler/ViewModels/MemorySharedState.cpp
+++ b/Engine/Source/Developer/TraceInsights/Private/Insights/MemoryProfiler/ViewModels/MemorySharedState.cpp
@@ -8,6 +8,8 @@
 #include "Logging/MessageLog.h"
 
 // TraceServices
+#include "Common/ProviderLock.h"
+#include "TraceServices/Model/AllocationsProvider.h"
 #include "TraceServices/Model/Memory.h"
 
 // TraceInsightsCore
@@ -172,6 +174,18 @@ void FMemorySharedState::Tick(Timing::ITimingViewSession& InSession, const Trace
 		return;
 	}
 
+	bool bIsDetailedMemoryTrace = false;
+	TSharedPtr<const TraceServices::IAnalysisSession> Session = FInsightsManager::Get()->GetSession();
+	if (Session.IsValid())
+	{
+		const TraceServices::IAllocationsProvider* AllocationsProvider = TraceServices::ReadAllocationsProvider(*Session.Get());
+		if (AllocationsProvider)
+		{
+			TraceServices::FProviderReadScopeLock ProviderReadScope(*AllocationsProvider);
+			bIsDetailedMemoryTrace = AllocationsProvider->IsInitialized() && AllocationsProvider->GetTimelineNumPoints() > 0;
+		}
+	}
+
 	constexpr int32 GraphTrackOrder = FTimingTrackOrder::First + FTimingTrackOrder::GroupRange / 2;
 
 	if (!MainGraphTrack.IsValid())
@@ -195,7 +209,7 @@ void FMemorySharedState::Tick(Timing::ITimingViewSession& InSession, const Trace
 		TimingView->InvalidateScrollableTracksOrder();
 	}
 
-	if (!LiveAllocsGraphTrack.IsValid())
+	if (!LiveAllocsGraphTrack.IsValid() && bIsDetailedMemoryTrace)
 	{
 		LiveAllocsGraphTrack = CreateMemoryGraphTrack();
 		check(LiveAllocsGraphTrack);
@@ -217,7 +231,7 @@ void FMemorySharedState::Tick(Timing::ITimingViewSession& InSession, const Trace
 		TimingView->InvalidateScrollableTracksOrder();
 	}
 
-	if (!AllocFreeGraphTrack.IsValid())
+	if (!AllocFreeGraphTrack.IsValid() && bIsDetailedMemoryTrace)
 	{
 		AllocFreeGraphTrack = CreateMemoryGraphTrack();
 		check(AllocFreeGraphTrack);
@@ -239,7 +253,7 @@ void FMemorySharedState::Tick(Timing::ITimingViewSession& InSession, const Trace
 		TimingView->InvalidateScrollableTracksOrder();
 	}
 
-	if (!SwapMemoryGraphTrack.IsValid())
+	if (!SwapMemoryGraphTrack.IsValid() && bIsDetailedMemoryTrace)
 	{
 		SwapMemoryGraphTrack = CreateMemoryGraphTrack();
 		check(SwapMemoryGraphTrack);
@@ -263,7 +277,7 @@ void FMemorySharedState::Tick(Timing::ITimingViewSession& InSession, const Trace
 		TimingView->InvalidateScrollableTracksOrder();
 	}
 
-	if (!PageSwapGraphTrack.IsValid())
+	if (!PageSwapGraphTrack.IsValid() && bIsDetailedMemoryTrace)
 	{
 		PageSwapGraphTrack = CreateMemoryGraphTrack();
 		check(PageSwapGraphTrack);
