[Insights]
- Memory Insights: Added label for low and high threshold values in LLM tag graph tracks (i.e. for budget lines).
- Memory Insights: Improved the way min, max, grid and budget values are displayed on the right side of the LLM tag graphs (i.e. how auto precision is computed).
- Memory Insights: Added faded display for the name of track (for LLM tag series) in the lower right side of the track. This is to avoid switching attention between left and right sides of the track when looking at grid values.

diff --git a/Engine/Source/Developer/TraceInsights/Private/Insights/MemoryProfiler/Tracks/MemoryGraphTrack.cpp b/Engine/Source/Developer/TraceInsights/Private/Insights/MemoryProfiler/Tracks/MemoryGraphTrack.cpp
--- a/Engine/Source/Developer/TraceInsights/Private/Insights/MemoryProfiler/Tracks/MemoryGraphTrack.cpp
+++ b/Engine/Source/Developer/TraceInsights/Private/Insights/MemoryProfiler/Tracks/MemoryGraphTrack.cpp
@@ -378,10 +378,54 @@ void FMemoryGraphTrack::DrawVerticalAxisGrid(const ITimingTrackDrawContext& Cont
 		MaxValueY = FMath::Min(MaxLabelY, FMath::Max(MinLabelY, ActualMaxValueY - TextH / 2));
 	}
 
+	// Label for the High Threshold value.
+	float HighThresholdLabelY = 0.0;
+	bool bShowHighThresholdLabel = IsAnyOptionEnabled(EGraphOptions::ShowThresholds) && Series->HasHighThresholdValue();
+	if (bShowHighThresholdLabel)
+	{
+		HighThresholdLabelY = Y0 + RoundedViewHeight - FMath::RoundToFloat(ViewportY.GetOffsetForValue(Series->GetHighThresholdValue())) - TextH / 2;
+		if (HighThresholdLabelY < MinLabelY ||
+			HighThresholdLabelY > MaxLabelY ||
+			(bHasMinMax && FMath::Abs(MinValueY - HighThresholdLabelY) < TextH) ||
+			(bHasMinMax && FMath::Abs(MaxValueY - HighThresholdLabelY) < TextH))
+		{
+			bShowHighThresholdLabel = false;
+		}
+	}
+
+	// Label for the Low Threshold value.
+	float LowThresholdLabelY = 0.0;
+	bool bShowLowThresholdLabel = IsAnyOptionEnabled(EGraphOptions::ShowThresholds) && Series->HasLowThresholdValue();
+	if (bShowLowThresholdLabel)
+	{
+		LowThresholdLabelY = Y0 + RoundedViewHeight - FMath::RoundToFloat(ViewportY.GetOffsetForValue(Series->GetLowThresholdValue())) - TextH / 2;
+		if (LowThresholdLabelY < MinLabelY ||
+			LowThresholdLabelY > MaxLabelY ||
+			(bHasMinMax && FMath::Abs(MinValueY - LowThresholdLabelY) < TextH) ||
+			(bHasMinMax && FMath::Abs(MaxValueY - LowThresholdLabelY) < TextH))
+		{
+			bShowLowThresholdLabel = false;
+		}
+	}
+
 	FDrawContext& DrawContext = Context.GetDrawContext();
 	const FSlateBrush* Brush = Context.GetHelper().GetWhiteBrush();
 	const TSharedRef<FSlateFontMeasure> FontMeasureService = FSlateApplication::Get().GetRenderer()->GetFontMeasureService();
 
+	if (AllSeries.Num() == 1 && Series->Is<FMemTagGraphSeries>())
+	{
+		// Show name of track (for LLM tag series) in the lower right side of the track.
+		// (to avoid switching attention between left and right sides of the track when looking at grid values)
+		TDrawLabelParams Params(DrawContext, Brush, FontMeasureService);
+		Params.TextBgColor = (Series->GetColor() * 0.05f).CopyWithNewOpacity(0.2f);
+		Params.TextColor = Series->GetBorderColor().CopyWithNewOpacity(0.2f);
+		Params.X = ViewWidth - 150.0f;
+		Params.Y = Y0 + RoundedViewHeight - 20.0f;
+		Params.Precision = std::numeric_limits<double>::quiet_NaN();
+		Params.Prefix = GetName();
+		DrawLabel(Params);
+	}
+
 	const double TopValue = ViewportY.GetValueAtOffset(RoundedViewHeight);
 	const double GridValue = ViewportY.GetValueAtOffset(MinDY);
 	const double BottomValue = ViewportY.GetValueAtOffset(0.0f);
@@ -392,6 +436,7 @@ void FMemoryGraphTrack::DrawVerticalAxisGrid(const ITimingTrackDrawContext& Cont
 	{
 		Precision *= 10.0;
 	}
+	Precision = FMath::Min(Precision, TopValue);
 
 	if (Delta > 0.0)
 	{
@@ -437,7 +482,7 @@ void FMemoryGraphTrack::DrawVerticalAxisGrid(const ITimingTrackDrawContext& Cont
 		Params.X = X0;
 		Params.Precision = Precision;
 
-		const FLinearColor GridColor(0.0f, 0.0f, 0.0f, 0.1f);
+		const FLinearColor GridLineColor(0.0f, 0.0f, 0.0f, 0.1f);
 
 		for (double Value = StartValue; Value < TopValue; Value += Grid)
 		{
@@ -450,9 +495,19 @@ void FMemoryGraphTrack::DrawVerticalAxisGrid(const ITimingTrackDrawContext& Cont
 			{
 				continue;
 			}
+			// Do not overlap with the High Threshold value.
+			if (bShowHighThresholdLabel && FMath::Abs(LabelY - HighThresholdLabelY) < TextH)
+			{
+				continue;
+			}
+			// Do not overlap with the Low Threshold value.
+			if (bShowLowThresholdLabel && FMath::Abs(LabelY - LowThresholdLabelY) < TextH)
+			{
+				continue;
+			}
 
 			// Draw horizontal grid line.
-			DrawContext.DrawBox(0, Y, ViewWidth, 1, Brush, GridColor);
+			DrawContext.DrawBox(0, Y, ViewWidth, 1, Brush, GridLineColor);
 
 			// Draw label.
 			Params.Y = LabelY;
@@ -461,7 +516,35 @@ void FMemoryGraphTrack::DrawVerticalAxisGrid(const ITimingTrackDrawContext& Cont
 		}
 	}
 
-	if (bHasMinMax && GetHeight() >= TextH)
+	const bool bIsMinHeight = (GetHeight() >= TextH);
+
+	// Draw label for the High Threshold value.
+	if (bShowHighThresholdLabel && bIsMinHeight)
+	{
+		TDrawLabelParams Params(DrawContext, Brush, FontMeasureService);
+		Params.TextBgColor = FLinearColor(0.1f, 0.05f, 0.05f, 1.0f);
+		Params.TextColor = FLinearColor(1.0f, 0.3f, 0.3f, 1.0f);
+		Params.X = X0;
+		Params.Y = HighThresholdLabelY;
+		Params.Value = Series->GetHighThresholdValue();
+		Params.Precision = -Precision;
+		DrawLabel(Params);
+	}
+
+	// Draw label for the Low Threshold value.
+	if (bShowLowThresholdLabel && bIsMinHeight)
+	{
+		TDrawLabelParams Params(DrawContext, Brush, FontMeasureService);
+		Params.TextBgColor = FLinearColor(0.1f, 0.1f, 0.05f, 1.0f);
+		Params.TextColor = FLinearColor(1.0f, 1.0f, 0.3f, 1.0f);
+		Params.X = X0;
+		Params.Y = LowThresholdLabelY;
+		Params.Value = Series->GetLowThresholdValue();
+		Params.Precision = -Precision;
+		DrawLabel(Params);
+	}
+
+	if (bHasMinMax && bIsMinHeight)
 	{
 		TDrawLabelParams Params(DrawContext, Brush, FontMeasureService);
 
@@ -519,7 +602,7 @@ void FMemoryGraphTrack::DrawVerticalAxisGrid(const ITimingTrackDrawContext& Cont
 				Params.X = MX;
 				Params.Y = MY - TextH / 2;
 				Params.Value = AllSeriesMaxValue - AllSeriesMinValue;
-				Params.Precision = -Precision * 0.1; // format with detailed text and increased precision
+				Params.Precision = -Precision; // format with detailed text
 				Params.Prefix = TEXT("\u0394=");
 				DrawLabel(Params);
 			}
@@ -535,7 +618,11 @@ void FMemoryGraphTrack::DrawLabel(const TDrawLabelParams& Params) const
 {
 	FString LabelText = Params.Prefix;
 
-	if (FMath::IsNearlyZero(Params.Value, 0.5))
+	if (std::isnan(Params.Precision))
+	{
+		// Draw only the Prefix text.
+	}
+	else if (FMath::IsNearlyZero(Params.Value, 0.5))
 	{
 		LabelText += TEXT("0");
 	}
