[Insights]
* Memory Insights / LLM Tags: Added editable Time Markers in toolbar.
* Memory Insights / LLM Tags: Update the layout of the entire toolbar for LLM Tags panel.
* Added reusabe widget to facilitate the editing of Time Markers (STimeMarkerEditBlock).

diff --git a/Engine/Source/Developer/TraceInsights/Private/Insights/MemoryProfiler/Widgets/SMemInvestigationView.cpp b/Engine/Source/Developer/TraceInsights/Private/Insights/MemoryProfiler/Widgets/SMemInvestigationView.cpp
--- a/Engine/Source/Developer/TraceInsights/Private/Insights/MemoryProfiler/Widgets/SMemInvestigationView.cpp
+++ b/Engine/Source/Developer/TraceInsights/Private/Insights/MemoryProfiler/Widgets/SMemInvestigationView.cpp
@@ -21,6 +21,7 @@
 #include "Insights/MemoryProfiler/Widgets/SMemAllocTableTreeView.h"
 #include "Insights/MemoryProfiler/Widgets/SMemoryProfilerWindow.h"
 #include "Insights/TimingProfiler/ViewModels/TimeMarker.h"
+#include "Insights/TimingProfiler/Widgets/STimeMarkerEditBlock.h"
 #include "Insights/Widgets/STimingView.h"
 
 #define LOCTEXT_NAMESPACE "UE::Insights::MemoryProfiler::SMemInvestigationView"
@@ -298,7 +299,6 @@ TSharedRef<SWidget> SMemInvestigationView::ConstructTimeMarkerWidget(uint32 Time
 		return SNew(SBox);
 	}
 
-	FMemorySharedState& SharedState = ProfilerWindow->GetSharedState();
 	const uint32 NumTimeMarkers = ProfilerWindow->GetNumCustomTimeMarkers();
 	if (TimeMarkerIndex >= NumTimeMarkers)
 	{
@@ -306,6 +306,7 @@ TSharedRef<SWidget> SMemInvestigationView::ConstructTimeMarkerWidget(uint32 Time
 	}
 
 	using FTimeMarker = TimingProfiler::FTimeMarker;
+
 	const TSharedRef<FTimeMarker>& TimeMarker = ProfilerWindow->GetCustomTimeMarker(TimeMarkerIndex);
 	TSharedPtr<FTimeMarker> PreviousTimeMarker;
 	if (TimeMarkerIndex > 0)
@@ -313,66 +314,45 @@ TSharedRef<SWidget> SMemInvestigationView::ConstructTimeMarkerWidget(uint32 Time
 		PreviousTimeMarker = ProfilerWindow->GetCustomTimeMarker(TimeMarkerIndex - 1);
 	}
 
-	TSharedRef<SWidget> Widget = SNew(SHorizontalBox)
-		.Visibility_Lambda([TimeMarker]()
+	TSharedRef<SWidget> Widget =
+		SNew(TimingProfiler::STimeMarkerEditBlock, TimeMarker)
+		.Padding(FMargin(20.0f, 0.0f, 0.0f, 0.0f))
+		.Visibility_Lambda(
+			[TimeMarkerIndex, ProfilerWindowWeak = ProfilerWindowWeakPtr]()
 			{
-				return TimeMarker->IsVisible() ? EVisibility::Visible : EVisibility::Collapsed;
-			})
-
-		+ SHorizontalBox::Slot()
-		.AutoWidth()
-		.VAlign(VAlign_Center)
-		.Padding(20.0f, 1.0f, 4.0f, 1.0f)
-		[
-			SNew(STextBlock)
-			.Text(FText::FromString(TimeMarker->GetName()))
-			.ColorAndOpacity(FSlateColor(TimeMarker->GetColor()))
-			.OnDoubleClicked(this, &SMemInvestigationView::OnTimeMarkerLabelDoubleClicked, TimeMarkerIndex)
-		]
-
-		+ SHorizontalBox::Slot()
-		.AutoWidth()
-		.VAlign(VAlign_Center)
-		.Padding(0.0f, 1.0f, 0.0f, 1.0f)
-		[
-			SNew(SEditableTextBox)
-			.MinDesiredWidth(110.0f)
-			.Text_Lambda([TimeMarker]()
+				TSharedPtr<SMemoryProfilerWindow> ProfilerWindow = ProfilerWindowWeak.Pin();
+				if (ProfilerWindow.IsValid())
 				{
-					return FText::FromString(FString::Printf(TEXT("%.9f"), TimeMarker->GetTime()));
-				})
-			.OnTextCommitted_Lambda([this, TimeMarker](const FText& InText, ETextCommit::Type InCommitType)
-				{
-					const double Time = FCString::Atod(*InText.ToString());
-					TimeMarker->SetTime(Time);
-					TSharedPtr<SMemoryProfilerWindow> ProfilerWindow = GetProfilerWindow();
-					if (ProfilerWindow.IsValid())
+					FMemorySharedState& SharedState = ProfilerWindow->GetSharedState();
+					TSharedPtr<FMemoryRuleSpec> Rule = SharedState.GetCurrentMemoryRule();
+					if (Rule && TimeMarkerIndex < Rule->GetNumTimeMarkers())
 					{
-						ProfilerWindow->OnTimeMarkerChanged(Timing::ETimeChangedFlags::None, TimeMarker);
+						return EVisibility::Visible;
 					}
-				})
-		]
-
-		+ SHorizontalBox::Slot()
-		.AutoWidth()
-		.VAlign(VAlign_Center)
-		.Padding(4.0f, 1.0f, 0.0f, 1.0f)
-		[
-			SNew(STextBlock)
-			.Text_Lambda([TimeMarker, PreviousTimeMarker]()
+				}
+				return EVisibility::Collapsed;
+			})
+		.PreviousTimeMarker(PreviousTimeMarker)
+		.OnGetTimingView_Lambda(
+			[ProfilerWindowWeak = ProfilerWindowWeakPtr](TSharedRef<FTimeMarker> TimeMarker) -> TSharedPtr<TimingProfiler::STimingView>
+			{
+				TSharedPtr<SMemoryProfilerWindow> ProfilerWindow = ProfilerWindowWeak.Pin();
+				if (ProfilerWindow.IsValid())
 				{
-					if (PreviousTimeMarker.IsValid())
-					{
-						return FText::FromString(FString::Printf(TEXT("%s (+%s)"),
-							*FormatTime(TimeMarker->GetTime(), 0.1),
-							*FormatTime(TimeMarker->GetTime() - PreviousTimeMarker->GetTime(), 0.1)));
-					}
-					else
-					{
-						return FText::FromString(FormatTime(TimeMarker->GetTime(), 0.1));
-					}
-				})
-		];
+					return ProfilerWindow->GetTimingView();
+				}
+				return nullptr;
+			})
+		.OnTimeMarkerChanged_Lambda(
+			[ProfilerWindowWeak = ProfilerWindowWeakPtr](TSharedRef<FTimeMarker> TimeMarker)
+			{
+				TSharedPtr<SMemoryProfilerWindow> ProfilerWindow = ProfilerWindowWeak.Pin();
+				if (ProfilerWindow.IsValid())
+				{
+					ProfilerWindow->OnTimeMarkerChanged(Timing::ETimeChangedFlags::None, TimeMarker);
+				}
+			})
+		;
 
 	return Widget;
 }
diff --git a/Engine/Source/Developer/TraceInsights/Private/Insights/MemoryProfiler/Widgets/SMemTagTreeView.cpp b/Engine/Source/Developer/TraceInsights/Private/Insights/MemoryProfiler/Widgets/SMemTagTreeView.cpp
--- a/Engine/Source/Developer/TraceInsights/Private/Insights/MemoryProfiler/Widgets/SMemTagTreeView.cpp
+++ b/Engine/Source/Developer/TraceInsights/Private/Insights/MemoryProfiler/Widgets/SMemTagTreeView.cpp
@@ -28,6 +28,7 @@
 #include "Insights/MemoryProfiler/ViewModels/MemTagTableViewPresets.h"
 #include "Insights/MemoryProfiler/Widgets/SMemoryProfilerWindow.h"
 #include "Insights/TimingProfiler/ViewModels/TimeMarker.h"
+#include "Insights/TimingProfiler/Widgets/STimeMarkerEditBlock.h"
 #include "Insights/Widgets/STimingView.h"
 
 #define LOCTEXT_NAMESPACE "UE::Insights::MemoryProfiler::SMemTagTreeView"
@@ -73,9 +74,69 @@ void SMemTagTreeView::Construct(const FArguments& InArgs, TSharedPtr<SMemoryProf
 BEGIN_SLATE_FUNCTION_BUILD_OPTIMIZATION
 TSharedPtr<SWidget> SMemTagTreeView::ConstructFilterToolbar()
 {
+	TSharedRef<SHorizontalBox> TopLine = SNew(SHorizontalBox)
+
+	+ SHorizontalBox::Slot()
+	.AutoWidth()
+	.Padding(FMargin(0.0f, 0.0f, 0.0f, 0.0f))
+	.VAlign(VAlign_Center)
+	[
+		ConstructHideAllButton()
+	]
+
+	+ SHorizontalBox::Slot()
+	.AutoWidth()
+	.Padding(FMargin(0.0f, 0.0f, 0.0f, 0.0f))
+	.VAlign(VAlign_Center)
+	[
+		ConstructShowAllButton()
+	]
+
+	+ SHorizontalBox::Slot()
+	.AutoWidth()
+	.Padding(FMargin(4.0f, 0.0f, 0.0f, 0.0f))
+	.VAlign(VAlign_Center)
+	[
+		ConstructTrackHeightControls()
+	]
+
+	+ SHorizontalBox::Slot()
+	.FillWidth(1.0f)
+	[
+		SNew(SBox)
+	]
+
+	+ SHorizontalBox::Slot()
+	.AutoWidth()
+	.Padding(4.0f, 0.0f, 0.0f, 0.0f)
+	.VAlign(VAlign_Center)
+	[
+		SNew(STextBlock)
+		.Text(LOCTEXT("TagSetText", "Tag Set:"))
+	]
+
+	+ SHorizontalBox::Slot()
+	.AutoWidth()
+	.Padding(4.0f, 0.0f, 6.0f, 0.0f)
+	.VAlign(VAlign_Center)
+	[
+		SNew(SComboBox<TSharedPtr<FMemoryTagSetId>>)
+		.ToolTipText(LOCTEXT("TagSetToolTip", "Select the LLM tag set."))
+		.OptionsSource(&AvailableTagSets)
+		.OnSelectionChanged(this, &SMemTagTreeView::TagSet_OnSelectionChanged)
+		.OnGenerateWidget(this, &SMemTagTreeView::TagSet_OnGenerateWidget)
+		[
+			SNew(STextBlock)
+			.Text(this, &SMemTagTreeView::TagSet_GetSelectedText)
+		]
+	];
+
+	ConstructViewPreset(TopLine);
+
+	//////////////////////////////////////////////////
+
 	FSlimHorizontalToolBarBuilder ToolbarBuilder(TSharedPtr<const FUICommandList>(), FMultiBoxCustomization::None);
 	ToolbarBuilder.SetStyle(&FInsightsStyle::Get(), "SecondaryToolbar");
-
 	ToolbarBuilder.BeginSection("Filters");
 	{
 		// Filter Trackers
@@ -85,47 +146,56 @@ TSharedPtr<SWidget> SMemTagTreeView::ConstructFilterToolbar()
 			LOCTEXT("TrackersMenuText", "Trackers"),
 			LOCTEXT("TrackersMenuToolTip", "Filter list of LLM tags by tracker."),
 			FSlateIcon(FInsightsStyle::GetStyleSetName(), "Icons.Filter.ToolBar"),
-			false
-		);
-
-		// Select Tag Set
-		ToolbarBuilder.AddWidget(
-			SNew(SComboBox<TSharedPtr<FMemoryTagSetId>>)
-			.ToolTipText(LOCTEXT("TagSetToolTip", "Select the LLM tag set."))
-			.OptionsSource(&AvailableTagSets)
-			.OnSelectionChanged(this, &SMemTagTreeView::TagSet_OnSelectionChanged)
-			.OnGenerateWidget(this, &SMemTagTreeView::TagSet_OnGenerateWidget)
-			[
-				SNew(STextBlock)
-				.Text(this, &SMemTagTreeView::TagSet_GetSelectedText)
-			]);
+			false);
 	}
 	ToolbarBuilder.EndSection();
 
-	return SNew(SHorizontalBox)
+	TSharedRef<SHorizontalBox> FilterLine = SNew(SHorizontalBox)
 
 	+ SHorizontalBox::Slot()
-		.FillWidth(1.0f)
-		.VAlign(VAlign_Center)
-		[
-			ConstructSearchBox()
-		]
+	.FillWidth(1.0f)
+	.VAlign(VAlign_Center)
+	[
+		ConstructSearchBox()
+	]
 
 	+ SHorizontalBox::Slot()
-		.AutoWidth()
-		.Padding(4.0f, 0.0f, 0.0f, 0.0f)
-		.VAlign(VAlign_Center)
-		[
-			ConstructFilterConfiguratorButton()
-		]
+	.AutoWidth()
+	.Padding(4.0f, 0.0f, 0.0f, 0.0f)
+	.VAlign(VAlign_Center)
+	[
+		ConstructFilterConfiguratorButton()
+	]
 
 	+ SHorizontalBox::Slot()
-		.AutoWidth()
-		.Padding(-2.0f, 0.0f, -6.0f, 0.0f)
-		.VAlign(VAlign_Center)
-		[
-			ToolbarBuilder.MakeWidget()
-		];
+	.AutoWidth()
+	.MinSize(10.0f)
+	.Padding(0.0f, 0.0f, -6.0f, 0.0f)
+	.VAlign(VAlign_Center)
+	[
+		ToolbarBuilder.MakeWidget()
+	];
+
+	//////////////////////////////////////////////////
+
+	return SNew(SVerticalBox)
+
+	+ SVerticalBox::Slot()
+	[
+		TopLine
+	]
+
+	+ SVerticalBox::Slot()
+	.AutoHeight()
+	.Padding(FMargin(0.0f, 2.0f, 0.0f, 2.0f))
+	[
+		ConstructBudgetSettingsToolbar()
+	]
+
+	+ SVerticalBox::Slot()
+	[
+		FilterLine
+	];
 }
 END_SLATE_FUNCTION_BUILD_OPTIMIZATION
 
@@ -134,11 +204,7 @@ END_SLATE_FUNCTION_BUILD_OPTIMIZATION
 BEGIN_SLATE_FUNCTION_BUILD_OPTIMIZATION
 TSharedPtr<SWidget> SMemTagTreeView::ConstructToolbar()
 {
-	TSharedPtr<SHorizontalBox> Box = SNew(SHorizontalBox);
-
-	ConstructViewPreset(Box);
-
-	return Box;
+	return nullptr;
 }
 END_SLATE_FUNCTION_BUILD_OPTIMIZATION
 
@@ -147,22 +213,97 @@ END_SLATE_FUNCTION_BUILD_OPTIMIZATION
 BEGIN_SLATE_FUNCTION_BUILD_OPTIMIZATION
 TSharedPtr<SWidget> SMemTagTreeView::ConstructExtraToolbar()
 {
-	TSharedPtr<SVerticalBox> Box = SNew(SVerticalBox)
+	return SNew(SVerticalBox)
 
 	+ SVerticalBox::Slot()
 	.AutoHeight()
 	[
-		ConstructTracksMiniToolbar()
+		ConstructTimeMarkerA()
 	]
 
 	+ SVerticalBox::Slot()
 	.AutoHeight()
-	.Padding(FMargin(0.0f, 4.0f, 0.0f, 0.0f))
 	[
-		ConstructBudgetSettingsToolbar()
+		ConstructTimeMarkerB()
 	];
+}
+END_SLATE_FUNCTION_BUILD_OPTIMIZATION
 
-	return Box;
+////////////////////////////////////////////////////////////////////////////////////////////////////
+
+BEGIN_SLATE_FUNCTION_BUILD_OPTIMIZATION
+TSharedRef<SWidget> SMemTagTreeView::ConstructTimeMarkerA()
+{
+	TSharedPtr<SMemoryProfilerWindow> ProfilerWindow = GetProfilerWindow();
+	if (!ProfilerWindow.IsValid() || ProfilerWindow->GetNumCustomTimeMarkers() < 2)
+	{
+		return SNew(SBox);
+	}
+
+	using FTimeMarker = TimingProfiler::FTimeMarker;
+	TSharedRef<FTimeMarker> TimeMarkerA = ProfilerWindow->GetCustomTimeMarker(0);
+
+	return SNew(TimingProfiler::STimeMarkerEditBlock, TimeMarkerA)
+		.Padding(FMargin(20.0f, 0.0f, 0.0f, 0.0f))
+		.OnGetTimingView_Lambda(
+			[ProfilerWindowWeak = ProfilerWindowWeakPtr](TSharedRef<FTimeMarker> TimeMarker) -> TSharedPtr<TimingProfiler::STimingView>
+			{
+				TSharedPtr<SMemoryProfilerWindow> ProfilerWindow = ProfilerWindowWeak.Pin();
+				if (ProfilerWindow.IsValid())
+				{
+					return ProfilerWindow->GetTimingView();
+				}
+				return nullptr;
+			})
+		.OnTimeMarkerChanged_Lambda(
+			[ProfilerWindowWeak = ProfilerWindowWeakPtr](TSharedRef<FTimeMarker> TimeMarker)
+			{
+				TSharedPtr<SMemoryProfilerWindow> ProfilerWindow = ProfilerWindowWeak.Pin();
+				if (ProfilerWindow.IsValid())
+				{
+					ProfilerWindow->OnTimeMarkerChanged(Timing::ETimeChangedFlags::None, TimeMarker);
+				}
+			});
+}
+END_SLATE_FUNCTION_BUILD_OPTIMIZATION
+
+////////////////////////////////////////////////////////////////////////////////////////////////////
+
+BEGIN_SLATE_FUNCTION_BUILD_OPTIMIZATION
+TSharedRef<SWidget> SMemTagTreeView::ConstructTimeMarkerB()
+{
+	TSharedPtr<SMemoryProfilerWindow> ProfilerWindow = GetProfilerWindow();
+	if (!ProfilerWindow.IsValid() || ProfilerWindow->GetNumCustomTimeMarkers() < 2)
+	{
+		return SNew(SBox);
+	}
+
+	using FTimeMarker = TimingProfiler::FTimeMarker;
+	TSharedRef<FTimeMarker> TimeMarkerA = ProfilerWindow->GetCustomTimeMarker(0);
+	TSharedRef<FTimeMarker> TimeMarkerB = ProfilerWindow->GetCustomTimeMarker(1);
+
+	return SNew(TimingProfiler::STimeMarkerEditBlock, TimeMarkerB)
+		.Padding(FMargin(20.0f, 0.0f, 0.0f, 0.0f))
+		.PreviousTimeMarker(TimeMarkerA)
+		.OnGetTimingView_Lambda(
+			[ProfilerWindowWeak = ProfilerWindowWeakPtr](TSharedRef<FTimeMarker> TimeMarker) -> TSharedPtr<TimingProfiler::STimingView>
+			{
+				TSharedPtr<SMemoryProfilerWindow> ProfilerWindow = ProfilerWindowWeak.Pin();
+				if (ProfilerWindow.IsValid())
+				{
+					return ProfilerWindow->GetTimingView();
+				}
+				return nullptr;
+			})
+		.OnTimeMarkerChanged_Lambda(
+			[ProfilerWindowWeak = ProfilerWindowWeakPtr](TSharedRef<FTimeMarker> TimeMarker)
+			{
+				TSharedPtr<SMemoryProfilerWindow> ProfilerWindow = ProfilerWindowWeak.Pin();
+				if (ProfilerWindow.IsValid())
+				{
+					ProfilerWindow->OnTimeMarkerChanged(Timing::ETimeChangedFlags::None, TimeMarker);
+				}
+			});
 }
 END_SLATE_FUNCTION_BUILD_OPTIMIZATION
 
@@ -280,16 +421,9 @@ FText SMemTagTreeView::TagSet_GetSelectedText() const
 ////////////////////////////////////////////////////////////////////////////////////////////////////
 
 BEGIN_SLATE_FUNCTION_BUILD_OPTIMIZATION
-TSharedRef<SWidget> SMemTagTreeView::ConstructTracksMiniToolbar()
+TSharedRef<SButton> SMemTagTreeView::ConstructHideAllButton()
 {
-	TSharedRef<SHorizontalBox> Box = SNew(SHorizontalBox)
-
-	+ SHorizontalBox::Slot()
-	.VAlign(VAlign_Center)
-	.Padding(FMargin(0.0f, 0.0f, 0.0f, 0.0f))
-	.AutoWidth()
-	[
-		SNew(SButton)
+	return SNew(SButton)
 		.ToolTipText(LOCTEXT("HideAll_ToolTip", "Remove memory graph tracks for all LLM tags."))
 		.OnClicked(this, &SMemTagTreeView::HideAllTracks_OnClicked)
 		.ButtonStyle(&FAppStyle::Get().GetWidgetStyle<FButtonStyle>("SimpleButton"))
@@ -304,15 +438,16 @@ TSharedRef<SWidget> SMemTagTreeView::ConstructTracksMiniToolbar()
 				.ColorAndOpacity(FSlateColor(FLinearColor(1.0f, 0.5f, 0.5f, 1.0f)))
 				.Image(FInsightsStyle::Get().GetBrush("Icons.RemoveAllMemTagGraphs"))
 			]
-		]
-	]
+		];
+}
+END_SLATE_FUNCTION_BUILD_OPTIMIZATION
 
-	+ SHorizontalBox::Slot()
-	.VAlign(VAlign_Center)
-	.Padding(FMargin(0.0f, 0.0f, 0.0f, 0.0f))
-	.AutoWidth()
-	[
-		SNew(SButton)
+////////////////////////////////////////////////////////////////////////////////////////////////////
+
+BEGIN_SLATE_FUNCTION_BUILD_OPTIMIZATION
+TSharedRef<SButton> SMemTagTreeView::ConstructShowAllButton()
+{
+	return SNew(SButton)
 		.ToolTipText(LOCTEXT("ShowAll_ToolTip", "Create memory graph tracks for all visible (filtered) LLM tags."))
 		.OnClicked(this, &SMemTagTreeView::ShowAllTracks_OnClicked)
 		.ButtonStyle(&FAppStyle::Get().GetWidgetStyle<FButtonStyle>("SimpleButton"))
@@ -327,15 +462,16 @@ TSharedRef<SWidget> SMemTagTreeView::ConstructTracksMiniToolbar()
 				.ColorAndOpacity(FSlateColor(FLinearColor(0.5f, 1.0f, 0.5f, 1.0f)))
 				.Image(FInsightsStyle::Get().GetBrush("Icons.AddAllMemTagGraphs"))
 			]
-		]
-	]
+		];
+}
+END_SLATE_FUNCTION_BUILD_OPTIMIZATION
 
-	+ SHorizontalBox::Slot()
-	.VAlign(VAlign_Center)
-	.HAlign(HAlign_Right)
-	.Padding(FMargin(0.0f, 0.0f, 0.0f, 0.0f))
-	[
-		SNew(SSegmentedControl<uint32>)
+////////////////////////////////////////////////////////////////////////////////////////////////////
+
+BEGIN_SLATE_FUNCTION_BUILD_OPTIMIZATION
+TSharedRef<SWidget> SMemTagTreeView::ConstructTrackHeightControls()
+{
+	return SNew(SSegmentedControl<uint32>)
 		.OnValueChanged_Lambda([this](uint32 InValue)
 			{
 				TSharedPtr<SMemoryProfilerWindow> ProfilerWindow = GetProfilerWindow();
@@ -355,18 +491,18 @@ TSharedRef<SWidget> SMemTagTreeView::ConstructTracksMiniToolbar()
 				}
 				return (uint32)EMemoryTrackHeightMode::Medium;
 			})
+
 		+ SSegmentedControl<uint32>::Slot((uint32)EMemoryTrackHeightMode::Small)
 		.Icon(FInsightsStyle::GetBrush("Icons.SizeSmall"))
 		.ToolTip(LOCTEXT("SmallHeight_ToolTip", "Change height of LLM Tag Graph tracks to Small."))
+
 		+ SSegmentedControl<uint32>::Slot((uint32)EMemoryTrackHeightMode::Medium)
 		.Icon(FInsightsStyle::GetBrush("Icons.SizeMedium"))
 		.ToolTip(LOCTEXT("MediumHeight_ToolTip", "Change height of LLM Tag Graph tracks to Medium."))
+
 		+ SSegmentedControl<uint32>::Slot((uint32)EMemoryTrackHeightMode::Large)
 		.Icon(FInsightsStyle::GetBrush("Icons.SizeLarge"))
-		.ToolTip(LOCTEXT("LargeHeight_ToolTip", "Change height of LLM Tag Graph tracks to Large."))
-	];
-
-	return Box;
+		.ToolTip(LOCTEXT("LargeHeight_ToolTip", "Change height of LLM Tag Graph tracks to Large."));
 }
 END_SLATE_FUNCTION_BUILD_OPTIMIZATION
 
@@ -933,16 +1069,16 @@ void SMemTagTreeView::UpdateStats()
 	if (ProfilerWindow.IsValid())
 	{
 		const TSharedRef<TimingProfiler::FTimeMarker>& MarkerA = ProfilerWindow->GetCustomTimeMarker(0);
-		if (TimeMarkerA != MarkerA->GetTime())
+		if (StatsTimeA != MarkerA->GetTime())
 		{
-			TimeMarkerA = MarkerA->GetTime();
+			StatsTimeA = MarkerA->GetTime();
 			bShouldUpdateStats = true;
 		}
 
 		const TSharedRef<TimingProfiler::FTimeMarker>& MarkerB = ProfilerWindow->GetCustomTimeMarker(1);
-		if (TimeMarkerB != MarkerB->GetTime())
+		if (StatsTimeB != MarkerB->GetTime())
 		{
-			TimeMarkerB = MarkerB->GetTime();
+			StatsTimeB = MarkerB->GetTime();
 			bShouldUpdateStats = true;
 		}
 
@@ -1017,7 +1153,7 @@ void SMemTagTreeView::UpdateStatsInternal()
 			FMemTagStats CurrentStats;
 
 			CurrentStats.SizeA = std::numeric_limits<int64>::max();
-			MemoryProvider->EnumerateTagSamples(MemTag->GetTrackerId(), MemTag->GetId(), TimeMarkerA, TimeMarkerA, true,
+			MemoryProvider->EnumerateTagSamples(MemTag->GetTrackerId(), MemTag->GetId(), StatsTimeA, StatsTimeA, true,
 				[&CurrentStats](double Time, double Duration, const TraceServices::FMemoryTagSample& Sample)
 				{
 					if (CurrentStats.SizeA == std::numeric_limits<int64>::max())
@@ -1035,7 +1171,7 @@ void SMemTagTreeView::UpdateStatsInternal()
 			}
 
 			CurrentStats.SizeB = std::numeric_limits<int64>::max();
-			MemoryProvider->EnumerateTagSamples(MemTag->GetTrackerId(), MemTag->GetId(), TimeMarkerB, TimeMarkerB, true,
+			MemoryProvider->EnumerateTagSamples(MemTag->GetTrackerId(), MemTag->GetId(), StatsTimeB, StatsTimeB, true,
 				[&CurrentStats](double Time, double Duration, const TraceServices::FMemoryTagSample& Sample)
 				{
 					if (CurrentStats.SizeB == std::numeric_limits<int64>::max())
diff --git a/Engine/Source/Developer/TraceInsights/Private/Insights/MemoryProfiler/Widgets/SMemTagTreeView.h b/Engine/Source/Developer/TraceInsights/Private/Insights/MemoryProfiler/Widgets/SMemTagTreeView.h
--- a/Engine/Source/Developer/TraceInsights/Private/Insights/MemoryProfiler/Widgets/SMemTagTreeView.h
+++ b/Engine/Source/Developer/TraceInsights/Private/Insights/MemoryProfiler/Widgets/SMemTagTreeView.h
@@ -194,8 +194,12 @@ protected:
 	virtual void TreeView_OnSelectionChanged(FTableTreeNodePtr SelectedItem, ESelectInfo::Type SelectInfo) override;
 
 private:
-	TSharedRef<SWidget> ConstructTracksMiniToolbar();
+	TSharedRef<SButton> ConstructHideAllButton();
+	TSharedRef<SButton> ConstructShowAllButton();
+	TSharedRef<SWidget> ConstructTrackHeightControls();
 	TSharedRef<SWidget> ConstructBudgetSettingsToolbar();
+	TSharedRef<SWidget> ConstructTimeMarkerA();
+	TSharedRef<SWidget> ConstructTimeMarkerB();
 
 	void UpdateSelectionStatsText();
 	FText GetNumSelectedTagsText() const { return NumSelectedTagsText; }
@@ -361,8 +365,8 @@ private:
 	EInvestigationMode Mode = EInvestigationMode::Default;
 
 	// A, B, B-A
-	double TimeMarkerA = 0.0;
-	double TimeMarkerB = 0.0;
+	double StatsTimeA = 0.0;
+	double StatsTimeB = 0.0;
 
 	// Sample Count, Min, Max, Average
 	double StatsStartTime = 0.0;
diff --git a/Engine/Source/Developer/TraceInsights/Private/Insights/TimingProfiler/Widgets/STimeMarkerEditBlock.cpp b/Engine/Source/Developer/TraceInsights/Private/Insights/TimingProfiler/Widgets/STimeMarkerEditBlock.cpp
new file mode 100644
--- /dev/null
+++ b/Engine/Source/Developer/TraceInsights/Private/Insights/TimingProfiler/Widgets/STimeMarkerEditBlock.cpp
@@ -0,0 +1,154 @@
+// Copyright Epic Games, Inc. All Rights Reserved.
+
+#include "STimeMarkerEditBlock.h"
+
+#include "SlateOptMacros.h"
+#include "Widgets/Input/SCheckBox.h"
+#include "Widgets/Input/SEditableTextBox.h"
+#include "Widgets/Layout/SBox.h"
+#include "Widgets/SBoxPanel.h"
+#include "Widgets/Text/STextBlock.h"
+
+// TraceInsightsCore
+#include "InsightsCore/Common/TimeUtils.h"
+
+// TraceInsights
+#include "Insights/TimingProfiler/ViewModels/TimeMarker.h"
+#include "Insights/Widgets/STimingView.h"
+
+#define LOCTEXT_NAMESPACE "UE::Insights::TimingProfiler::STimeMarkerEditBlock"
+
+namespace UE::Insights::TimingProfiler
+{
+
+////////////////////////////////////////////////////////////////////////////////////////////////////
+
+BEGIN_SLATE_FUNCTION_BUILD_OPTIMIZATION
+void STimeMarkerEditBlock::Construct(const FArguments& InArgs, TSharedRef<FTimeMarker> InTimeMarker)
+{
+	TimeMarker = InTimeMarker;
+
+	Padding = InArgs._Padding;
+	PreviousTimeMarker = InArgs._PreviousTimeMarker;
+	OnGetTimingViewCallback = InArgs._OnGetTimingView;
+	OnTimeMarkerChangedCallback = InArgs._OnTimeMarkerChanged;
+
+	const FText TimeMarkerName = FText::FromString(TimeMarker->GetName());
+
+	ChildSlot
+	[
+		SNew(SHorizontalBox)
+
+		+ SHorizontalBox::Slot()
+		.AutoWidth()
+		.VAlign(VAlign_Center)
+		.Padding(Padding.Left, Padding.Top + 1.0f, 4.0f, Padding.Bottom + 1.0f)
+		[
+			SNew(SCheckBox)
+			.ToolTipText(LOCTEXT("VisibilityTooltip", "Time Marker Visibility\nShows the time marker in the Timing View."))
+			.IsChecked_Lambda([this]()
+				{
+					return TimeMarker->IsVisible() ? ECheckBoxState::Checked : ECheckBoxState::Unchecked;
+				})
+			.OnCheckStateChanged_Lambda([this](ECheckBoxState NewState)
+				{
+					TimeMarker->SetVisibility(NewState == ECheckBoxState::Checked);
+				})
+		]
+
+		+ SHorizontalBox::Slot()
+		.AutoWidth()
+		.VAlign(VAlign_Center)
+		.Padding(0.0f, 1.0f, 0.0f, 1.0f)
+		[
+			SNew(STextBlock)
+			.Text(TimeMarkerName)
+			.ToolTipText(FText::Format(
+				LOCTEXT("NameTooltipFmt", "Time Marker '{0}'\nDouble click the name to move the time marker to the center of the Timing View."),
+				TimeMarkerName))
+			.ColorAndOpacity(FSlateColor(TimeMarker->GetColor()))
+			.OnDoubleClicked(this, &STimeMarkerEditBlock::OnLabelDoubleClicked)
+		]
+
+		+ SHorizontalBox::Slot()
+		.AutoWidth()
+		.VAlign(VAlign_Center)
+		.Padding(4.0f, 1.0f, 0.0f, 1.0f)
+		[
+			SNew(SEditableTextBox)
+			.MinDesiredWidth(110.0f)
+			.Text_Lambda([this]()
+				{
+					return FText::FromString(FString::Printf(TEXT("%.9f"), TimeMarker->GetTime()));
+				})
+			.OnTextCommitted_Lambda([this](const FText& InText, ETextCommit::Type InCommitType)
+				{
+					const double Time = FCString::Atod(*InText.ToString());
+					TimeMarker->SetTime(Time);
+					OnTimeMarkerChanged();
+				})
+		]
+
+		+ SHorizontalBox::Slot()
+		.AutoWidth()
+		.VAlign(VAlign_Center)
+		.Padding(4.0f, 1.0f, Padding.Right, 1.0f)
+		[
+			SNew(STextBlock)
+			.Text_Lambda([this]()
+				{
+					double Time = TimeMarker->GetTime();
+					if (PreviousTimeMarker.IsSet() && PreviousTimeMarker.Get().IsValid())
+					{
+						double PrevTime = PreviousTimeMarker.Get()->GetTime();
+						return FText::FromString(FString::Printf(TEXT("%s (+%s)"),
+							*FormatTime(Time, 0.1),
+							*FormatTime(Time - PrevTime, 0.1)));
+					}
+					else
+					{
+						return FText::FromString(FormatTime(Time, 0.1));
+					}
+				})
+		]
+	];
+}
+END_SLATE_FUNCTION_BUILD_OPTIMIZATION
+
+////////////////////////////////////////////////////////////////////////////////////////////////////
+
+TSharedPtr<STimingView> STimeMarkerEditBlock::GetTimingView() const
+{
+	if (OnGetTimingViewCallback.IsBound())
+	{
+		return OnGetTimingViewCallback.Execute(TimeMarker.ToSharedRef());
+	}
+	return nullptr;
+}
+
+////////////////////////////////////////////////////////////////////////////////////////////////////
+
+void STimeMarkerEditBlock::OnTimeMarkerChanged()
+{
+	OnTimeMarkerChangedCallback.ExecuteIfBound(TimeMarker.ToSharedRef());
+}
+
+////////////////////////////////////////////////////////////////////////////////////////////////////
+
+FReply STimeMarkerEditBlock::OnLabelDoubleClicked(const FGeometry& MyGeometry, const FPointerEvent& PointerEvent)
+{
+	TSharedPtr<STimingView> TimingView = GetTimingView();
+	if (TimingView.IsValid())
+	{
+		// Move timer to the center of the timing view.
+		const double Time = (TimingView->GetViewport().GetStartTime() + TimingView->GetViewport().GetEndTime()) / 2.0;
+		TimeMarker->SetTime(Time);
+	}
+	return FReply::Handled();
+}
+
+////////////////////////////////////////////////////////////////////////////////////////////////////
+
+} // namespace UE::Insights::TimingProfiler
+
+#undef LOCTEXT_NAMESPACE
diff --git a/Engine/Source/Developer/TraceInsights/Private/Insights/TimingProfiler/Widgets/STimeMarkerEditBlock.h b/Engine/Source/Developer/TraceInsights/Private/Insights/TimingProfiler/Widgets/STimeMarkerEditBlock.h
new file mode 100644
--- /dev/null
+++ b/Engine/Source/Developer/TraceInsights/Private/Insights/TimingProfiler/Widgets/STimeMarkerEditBlock.h
@@ -0,0 +1,54 @@
+// Copyright Epic Games, Inc. All Rights Reserved.
+
+#pragma once
+
+#include "CoreTypes.h"
+
+#include "Input/Events.h"
+#include "Input/Reply.h"
+#include "Layout/Geometry.h"
+#include "Templates/SharedPointer.h"
+#include "Widgets/SCompoundWidget.h"
+
+namespace UE::Insights::TimingProfiler
+{
+
+class FTimeMarker;
+class STimingView;
+
+FUNC_DECLARE_DELEGATE(FOnGetTimingView, TSharedPtr<STimingView>, TSharedRef<FTimeMarker>);
+FUNC_DECLARE_DELEGATE(FOnTimeMarkerChanged, void, TSharedRef<FTimeMarker>);
+
+class STimeMarkerEditBlock : public SCompoundWidget
+{
+public:
+	SLATE_BEGIN_ARGS(STimeMarkerEditBlock) {}
+		SLATE_ARGUMENT(FMargin, Padding)
+		SLATE_ATTRIBUTE(TSharedPtr<FTimeMarker>, PreviousTimeMarker)
+		SLATE_EVENT(FOnGetTimingView, OnGetTimingView)
+		SLATE_EVENT(FOnTimeMarkerChanged, OnTimeMarkerChanged)
+	SLATE_END_ARGS()
+
+	void Construct(const FArguments& InArgs, TSharedRef<FTimeMarker> InTimeMarker);
+
+private:
+	TSharedPtr<STimingView> GetTimingView() const;
+	void OnTimeMarkerChanged();
+	FReply OnLabelDoubleClicked(const FGeometry& MyGeometry, const FPointerEvent& PointerEvent);
+
+private:
+	TSharedPtr<FTimeMarker> TimeMarker;
+
+	FMargin Padding;
+
+	/** The previous TimeMarker; used to display delta time. */
+	TAttribute<TSharedPtr<FTimeMarker>> PreviousTimeMarker;
+
+	/** Called whenever this widget needs to know the TimingView widget where the TimeMarker is located. */
+	FOnGetTimingView OnGetTimingViewCallback;
+
+	/** Called whenever the time marker is changed by this widget. */
+	FOnTimeMarkerChanged OnTimeMarkerChangedCallback;
+};
+
+} // namespace UE::Insights::TimingProfiler
