Gpu Fence visualization on child track

diff --git a/Engine/Source/Developer/TraceInsights/Private/Insights/TimingProfiler/Tracks/GpuTimingTrack.cpp b/Engine/Source/Developer/TraceInsights/Private/Insights/TimingProfiler/Tracks/GpuTimingTrack.cpp
--- a/Engine/Source/Developer/TraceInsights/Private/Insights/TimingProfiler/Tracks/GpuTimingTrack.cpp
+++ b/Engine/Source/Developer/TraceInsights/Private/Insights/TimingProfiler/Tracks/GpuTimingTrack.cpp
@@ -2,7 +2,13 @@
 
 #include "GpuTimingTrack.h"
 
+#include "Fonts/FontMeasure.h"
+
+// TraceInsightsCore
+#include "InsightsCore/Common/PaintUtils.h"
+
 // TraceInsights
+#include "Insights/InsightsManager.h"
 #include "Insights/TimingProfiler/ViewModels/ThreadTimingSharedState.h"
 #include "Insights/ViewModels/TimingTrackViewport.h"
 #include "Insights/ViewModels/TimingViewDrawHelper.h"
@@ -13,11 +19,12 @@ namespace UE::Insights::TimingProfiler
 {
 
 ////////////////////////////////////////////////////////////////////////////////////////////////////
-// FGpuTimingTrack, FGpuQueueTimingTrack
+// FGpuTimingTrack, FGpuQueueTimingTrack, FGpuFencesTimingTrack
 ////////////////////////////////////////////////////////////////////////////////////////////////////
 
 INSIGHTS_IMPLEMENT_RTTI(FGpuTimingTrack)
 INSIGHTS_IMPLEMENT_RTTI(FGpuQueueTimingTrack)
+INSIGHTS_IMPLEMENT_RTTI(FGpuFencesTimingTrack)
 
 ////////////////////////////////////////////////////////////////////////////////////////////////////
 // FGpuQueueWorkTimingTrack
@@ -139,6 +146,359 @@ void FGpuQueueWorkTimingTrack::PostDraw(const ITimingTrackDrawContext& Context)
 	}
 }
 
+////////////////////////////////////////////////////////////////////////////////////////////////////
+// FGpuFencesTrackBuilder
+////////////////////////////////////////////////////////////////////////////////////////////////////
+
+FGpuFencesTrackBuilder::FGpuFencesTrackBuilder(FGpuFencesTimingTrack& InTrack, const FTimingTrackViewport& InViewport, float InFontScale)
+	: Track(InTrack)
+	, Viewport(InViewport)
+	, FontMeasureService(FSlateApplication::Get().GetRenderer()->GetFontMeasureService())
+	, Font(FAppStyle::Get().GetFontStyle("SmallFont"))
+	, FontScale(InFontScale)
+{
+	Track.ResetCache();
+}
+
+////////////////////////////////////////////////////////////////////////////////////////////////////
+
+void FGpuFencesTrackBuilder::AddFence(double Timestamp, TraceServices::EGpuFenceType Type, const TCHAR* Text)
+{
+	float X = Viewport.TimeToSlateUnitsRounded(Timestamp);
+
+	AddTimeMarker(X, Type, Text);
+}
+
+////////////////////////////////////////////////////////////////////////////////////////////////////
+
+void FGpuFencesTrackBuilder::Flush()
+{
+	if (LastMessage.IsEmpty())
+	{
+		return;
+	}
+
+	Flush(Viewport.GetWidth() - LastX2);
+}
+
+////////////////////////////////////////////////////////////////////////////////////////////////////
+
+void FGpuFencesTrackBuilder::Flush(float AvailableTextW)
+{
+	if (LastMessage.IsEmpty())
+	{
+		return;
+	}
+
+	auto ToLiniarColorNoAlpha = [](uint32 Value)
+		{
+			const float R = static_cast<float>((Value & 0xFF000000) >> 24) / 255.0f;
+			const float G = static_cast<float>((Value & 0x00FF0000) >> 16) / 255.0f;
+			const float B = static_cast<float>((Value & 0x0000FF00) >> 8) / 255.0f;
+			return FLinearColor(R, G, B);
+		};
+
+	constexpr FLinearColor SingnalFenceColor = ToLiniarColorNoAlpha(0x3A9C7bFF);
+	constexpr FLinearColor WaitFenceColor = ToLiniarColorNoAlpha(0xFFDC1AFF);
+
+	const FLinearColor Color = LastType == TraceServices::EGpuFenceType::SignalFence ? SingnalFenceColor : WaitFenceColor;
+
+	bool bAddNewBox = true;
+	if (Track.TimeMarkerBoxes.Num() > 0)
+	{
+		FGpuFenceTimeMarkerBoxInfo& PrevBox = Track.TimeMarkerBoxes.Last();
+		if (PrevBox.X + PrevBox.W == LastX1 &&
+			PrevBox.Color.R == Color.R &&
+			PrevBox.Color.G == Color.G &&
+			PrevBox.Color.B == Color.B)
+		{
+			// Extend previous box instead.
+			PrevBox.W += LastX2 - LastX1;
+			bAddNewBox = false;
+		}
+	}
+
+	if (bAddNewBox)
+	{
+		// Add new Box info to cache.
+		FGpuFenceTimeMarkerBoxInfo& Box = Track.TimeMarkerBoxes.AddDefaulted_GetRef();
+		Box.X = LastX1;
+		Box.W = LastX2 - LastX1;
+		Box.Color = Color;
+		Box.Color.A = 0.25f;
+	}
+
+	if (AvailableTextW > 6.0f)
+	{
+		const int32 HorizontalOffset = FMath::RoundToInt((AvailableTextW - 2.0f) * FontScale);
+		const int32 LastWholeCharacterIndexCategory = FontMeasureService->FindLastWholeCharacterIndexBeforeOffset(LastMessage, Font, HorizontalOffset, FontScale);
+
+		if (LastWholeCharacterIndexCategory >= 0)
+		{
+			// Add new Text info to cache.
+			FGpuFenceTextInfo& TextInfo = Track.TimeMarkerTexts.AddDefaulted_GetRef();
+			TextInfo.X = LastX2 + 2.0f;
+			TextInfo.Color = Color;
+			if (LastWholeCharacterIndexCategory >= 0)
+			{
+				TextInfo.Text.AppendChars(*LastMessage, LastWholeCharacterIndexCategory + 1);
+			}
+		}
+	}
+}
+
+////////////////////////////////////////////////////////////////////////////////////////////////////
+
+void FGpuFencesTrackBuilder::AddTimeMarker(const float X, TraceServices::EGpuFenceType Type, const TCHAR* Message)
+{
+	const float W = X - LastX2;
+	constexpr float BoxWidth = 1.0f;
+
+	if (W > 0.0f) // There is at least 1px from previous box?
+	{
+		// Flush previous marker (if any).
+		Flush(W);
+
+		// Begin new marker info.
+		LastX1 = X;
+		LastX2 = X + BoxWidth;
+	}
+	else if (W == 0.0f) // Adjacent to previous box?
+	{
+		// Same color as previous marker?
+		if (Type == LastType)
+		{
+			// Extend previous box.
+			LastX2++;
+		}
+		else
+		{
+			// Flush previous marker (if any).
+			Flush(0.0f);
+
+			// Begin new box.
+			LastX1 = X;
+			LastX2 = X + BoxWidth;
+		}
+	}
+	else // Overlaps previous box?
+	{
+		// Same color as previous marker?
+		if (Type == LastType)
+		{
+			// Keep previous box.
+		}
+		else
+		{
+			// Shrink previous box.
+			LastX2--;
+
+			if (LastX2 > LastX1)
+			{
+				// Flush previous marker (if any).
+				Flush(0.0f);
+			}
+
+			// Begin new box.
+			LastX1 = X;
+			LastX2 = X + BoxWidth;
+		}
+	}
+
+	// Save marker.
+	LastType = Type;
+	LastMessage = Message;
+}
+
+////////////////////////////////////////////////////////////////////////////////////////////////////
+// FGpuFencesTimingTrack
+////////////////////////////////////////////////////////////////////////////////////////////////////
+
+FGpuFencesTimingTrack::FGpuFencesTimingTrack(FThreadTimingSharedState& InSharedState, const FString& InName, uint32 InTimelineIndex, uint32 InQueueId)
+	: FThreadTimingTrackImpl(InSharedState, InName, nullptr, InTimelineIndex, InQueueId)
+	, WhiteBrush(FAppStyle::Get().GetBrush("WhiteBrush"))
+	, Font(FAppStyle::Get().GetFontStyle("SmallFont"))
+{
+}
+
+void FGpuFencesTimingTrack::Reset()
+{
+	FThreadTimingTrackImpl::Reset();
+
+	ResetCache();
+}
+
+void FGpuFencesTimingTrack::ResetCache()
+{
+	TimeMarkerBoxes.Empty();
+	TimeMarkerTexts.Empty();
+}
+
+void FGpuFencesTimingTrack::PreUpdate(const ITimingTrackUpdateContext& Context)
+{
+	const FTimingTrackViewport& Viewport = Context.GetViewport();
+	if (IsDirty() || Viewport.IsHorizontalViewportDirty())
+	{
+		ClearDirtyFlag();
+
+		using namespace TraceServices;
+		TSharedPtr<const IAnalysisSession> Session = FInsightsManager::Get()->GetSession();
+		if (Session.IsValid() && ReadTimingProfilerProvider(*Session.Get()))
+		{
+			FAnalysisSessionReadScope SessionReadScope(*Session.Get());
+
+			const TraceServices::ITimingProfilerProvider& TimingProfilerProvider = *ReadTimingProfilerProvider(*Session.Get());
+
+			FGpuFencesTrackBuilder Builder(*this, Viewport, Context.GetGeometry().Scale);
+
+			TimingProfilerProvider.EnumerateGpuFences(GetThreadId(), Viewport.GetStartTime(), Viewport.GetEndTime(), [&Builder](const FGpuFenceWrapper& Fence)
+				{
+					if (Fence.FenceType == EGpuFenceType::SignalFence)
+					{
+						const FGpuSignalFence* SignalFence = Fence.Fence.Get<const FGpuSignalFence*>();
+						TCHAR Name[32];
+						FCString::Sprintf(Name, TEXT("%ld"), SignalFence->Value);
+						Builder.AddFence(SignalFence->Timestamp, EGpuFenceType::SignalFence, Name);
+					}
+					else if (Fence.FenceType == EGpuFenceType::WaitFence)
+					{
+						const FGpuWaitFence* WaitFence = Fence.Fence.Get<const FGpuWaitFence*>();
+						TCHAR Name[32];
+						FCString::Sprintf(Name, TEXT("%ld"), WaitFence->Value);
+						Builder.AddFence(WaitFence->Timestamp, EGpuFenceType::WaitFence, Name);
+					}
+
+					return TraceServices::EEnumerateResult::Continue;
+				});
+
+			Builder.Flush();
+		}
+
+		SetNumLanes(1);
+	}
+
+	UpdateTrackHeight(Context);
+}
+
+////////////////////////////////////////////////////////////////////////////////////////////////////
+
+void FGpuFencesTimingTrack::Draw(const ITimingTrackDrawContext& Context) const
+{
+	//////////////////////////////////////////////////
+	// Draw texts (strings are already truncated).
+	FDrawContext& DrawContext = Context.GetDrawContext();
+	const FTimingTrackViewport& Viewport = Context.GetViewport();
+
+	const FTimingViewDrawHelper& Helper = *static_cast<const FTimingViewDrawHelper*>(&Context.GetHelper());
+	const FTimingViewLayout& Layout = Viewport.GetLayout();
+	const int32 EventTextLayerId = Helper.GetTextLayerId();
+	const float TextY = GetPosY() - (FTimingViewLayout::NormalLayoutEventH - Layout.EventH) / 2.0f + 3.0f;
+
+	const int32 NumTexts = TimeMarkerTexts.Num();
+
+	constexpr float TextMinEventH = 7.0f;
+	if (Layout.EventH > TextMinEventH)
+	{
+		float TextOpacity = 1.0f;
+		if (Layout.EventH < FTimingViewLayout::NormalLayoutEventH)
+		{
+			TextOpacity *= (Layout.EventH - TextMinEventH + 1.0f) / (FTimingViewLayout::NormalLayoutEventH - TextMinEventH + 1.0f);
+		}
+		for (int32 TextIndex = 0; TextIndex < NumTexts; TextIndex++)
+		{
+			const FGpuFenceTextInfo& TextInfo = TimeMarkerTexts[TextIndex];
+
+			if (TextInfo.Text.Len() > 0)
+			{
+				DrawContext.DrawText(EventTextLayerId, TextInfo.X, TextY, TextInfo.Text, Font, TextInfo.Color.CopyWithNewOpacity(TextInfo.Color.A * TextOpacity));
+			}
+		}
+	}
+}
+
+////////////////////////////////////////////////////////////////////////////////////////////////////
+
+void FGpuFencesTimingTrack::PostDraw(const ITimingTrackDrawContext& Context) const
+{
+	FDrawContext& DrawContext = Context.GetDrawContext();
+	const FTimingTrackViewport& Viewport = Context.GetViewport();
+
+	//////////////////////////////////////////////////
+	// Draw vertical lines.
+	// Multiple adjacent vertical lines with same color are merged into a single box.
+
+	DrawContext.LayerId++;
+
+	bool bExtentFenceLines = Viewport.GetDuration() < 3.0f;
+	float BoxY1 = GetPosY();
+	float BoxY2 = GetPosY() + GetHeight();
+	ETimingTrackLocation LocalLocation = GetLocation();
+
+	auto ParentTrackShared = GetParentTrack().Pin();
+	if (ParentTrack.IsValid() && bExtentFenceLines)
+	{
+		// Substract the size of the child tracks above this track.
+		float TopY = GetPosY() - ParentTrackShared->GetPosY();
+		BoxY2 = GetPosY() + ParentTrackShared->GetHeight() - TopY;
+		LocalLocation = ParentTrackShared->GetLocation();
+	}
+
+	switch (LocalLocation)
+	{
+	case ETimingTrackLocation::Scrollable:
+	{
+		const float TopY = Viewport.GetPosY() + Viewport.GetTopOffset();
+		if (BoxY1 < TopY)
+		{
+			BoxY1 = TopY;
+		}
+		const float BottomY = Viewport.GetPosY() + Viewport.GetHeight() - Viewport.GetBottomOffset();
+		if (BoxY2 > BottomY)
+		{
+			BoxY2 = BottomY;
+		}
+		break;
+	}
+	case ETimingTrackLocation::TopDocked:
+	{
+		const float TopY = Viewport.GetPosY();
+		if (BoxY1 < TopY)
+		{
+			BoxY1 = TopY;
+		}
+		const float BottomY = Viewport.GetPosY() + Viewport.GetTopOffset();
+		if (BoxY2 > BottomY)
+		{
+			BoxY2 = BottomY;
+		}
+		break;
+	}
+	case ETimingTrackLocation::BottomDocked:
+	{
+		const float TopY = Viewport.GetPosY() + Viewport.GetHeight() - Viewport.GetBottomOffset();
+		if (BoxY1 < TopY)
+		{
+			BoxY1 = TopY;
+		}
+		const float BottomY = Viewport.GetPosY() + Viewport.GetHeight();
+		if (BoxY2 > BottomY)
+		{
+			BoxY2 = BottomY;
+		}
+		break;
+	}
+	}
+
+	float BoxH = BoxY2 - BoxY1;
+	const int32 NumBoxes = TimeMarkerBoxes.Num();
+	for (int32 BoxIndex = 0; BoxIndex < NumBoxes; BoxIndex++)
+	{
+		const FGpuFenceTimeMarkerBoxInfo& Box = TimeMarkerBoxes[BoxIndex];
+		DrawContext.DrawBox(Box.X, BoxY1, Box.W, BoxH, WhiteBrush, Box.Color);
+	}
+	DrawContext.LayerId++;
+}
+
 ////////////////////////////////////////////////////////////////////////////////////////////////////
 
 } // namespace UE::Insights::TimingProfiler
diff --git a/Engine/Source/Developer/TraceInsights/Private/Insights/TimingProfiler/Tracks/GpuTimingTrack.h b/Engine/Source/Developer/TraceInsights/Private/Insights/TimingProfiler/Tracks/GpuTimingTrack.h
--- a/Engine/Source/Developer/TraceInsights/Private/Insights/TimingProfiler/Tracks/GpuTimingTrack.h
+++ b/Engine/Source/Developer/TraceInsights/Private/Insights/TimingProfiler/Tracks/GpuTimingTrack.h
@@ -7,6 +7,8 @@
 // TraceInsights
 #include "Insights/TimingProfiler/Tracks/ThreadTimingTrackPrivate.h"
 
+class FSlateFontMeasure;
+
 namespace UE::Insights::TimingProfiler
 {
 
@@ -65,4 +67,88 @@ private:
 
 ////////////////////////////////////////////////////////////////////////////////////////////////////
 
+struct FGpuFenceTimeMarkerBoxInfo
+{
+	float X;
+	float W;
+	FLinearColor Color;
+};
+
+////////////////////////////////////////////////////////////////////////////////////////////////////
+
+struct FGpuFenceTextInfo
+{
+	float X;
+	FLinearColor Color;
+	FString Text;
+};
+
+////////////////////////////////////////////////////////////////////////////////////////////////////
+
+class FGpuFencesTimingTrack : public FThreadTimingTrackImpl
+{
+	friend class FGpuFencesTrackBuilder;
+	INSIGHTS_DECLARE_RTTI(FGpuFencesTimingTrack, FThreadTimingTrackImpl)
+
+public:
+	explicit FGpuFencesTimingTrack(FThreadTimingSharedState& InSharedState, const FString& InName, uint32 InTimelineIndex, uint32 InQueueId);
+
+	uint32 GetQueueId() const { return GetThreadId(); }
+
+	virtual void Reset() override;
+	virtual void PreUpdate(const ITimingTrackUpdateContext& Context) override;
+
+	virtual void Draw(const ITimingTrackDrawContext& Context) const override;
+	virtual void PostDraw(const ITimingTrackDrawContext& Context) const override;
+	virtual void BuildDrawState(ITimingEventsTrackDrawStateBuilder& Builder, const ITimingTrackUpdateContext& Context) override {};
+	virtual void BuildFilteredDrawState(ITimingEventsTrackDrawStateBuilder& Builder, const ITimingTrackUpdateContext& Context) override {}
+	virtual const TSharedPtr<const ITimingEvent> GetEvent(float InPosX, float InPosY, const FTimingTrackViewport& Viewport) const override { return nullptr; };
+
+private:
+	void ResetCache();
+
+private:
+	TArray<FGpuFenceTimeMarkerBoxInfo> TimeMarkerBoxes;
+	TArray<FGpuFenceTextInfo> TimeMarkerTexts;
+
+	const FSlateBrush* WhiteBrush;
+	const FSlateFontInfo Font;
+};
+
+////////////////////////////////////////////////////////////////////////////////////////////////////
+
+class FGpuFencesTrackBuilder
+{
+public:
+	explicit FGpuFencesTrackBuilder(FGpuFencesTimingTrack& InTrack, const FTimingTrackViewport& InViewport, float InFontScale);
+
+	/**
+	 * Non-copyable
+	 */
+	FGpuFencesTrackBuilder(const FGpuFencesTrackBuilder&) = delete;
+	FGpuFencesTrackBuilder& operator=(const FGpuFencesTrackBuilder&) = delete;
+
+	const FTimingTrackViewport& GetViewport() { return Viewport; }
+
+	void AddFence(double Timestamp, TraceServices::EGpuFenceType Type, const TCHAR* Text);
+	void Flush();
+
+private:
+	void Flush(float AvailableTextW);
+	void AddTimeMarker(const float X, TraceServices::EGpuFenceType Type, const TCHAR* Message);
+
+private:
+	FGpuFencesTimingTrack& Track;
+	const FTimingTrackViewport& Viewport;
+
+	const TSharedRef<FSlateFontMeasure> FontMeasureService;
+	const FSlateFontInfo Font;
+	float FontScale;
+
+	float LastX1;
+	float LastX2;
+	TraceServices::EGpuFenceType LastType;
+	FString LastMessage;
+};
+
 } // namespace UE::Insights::TimingProfiler
diff --git a/Engine/Source/Developer/TraceInsights/Private/Insights/TimingProfiler/ViewModels/ThreadTimingSharedState.cpp b/Engine/Source/Developer/TraceInsights/Private/Insights/TimingProfiler/ViewModels/ThreadTimingSharedState.cpp
--- a/Engine/Source/Developer/TraceInsights/Private/Insights/TimingProfiler/ViewModels/ThreadTimingSharedState.cpp
+++ b/Engine/Source/Developer/TraceInsights/Private/Insights/TimingProfiler/ViewModels/ThreadTimingSharedState.cpp
@@ -385,6 +385,11 @@ void FThreadTimingSharedState::Tick(Timing::ITimingViewSession& InSession, const
 						TSharedRef<FGpuQueueWorkTimingTrack> WorkTrack = MakeShared<FGpuQueueWorkTimingTrack>(*this, WorkTrackName, QueueInfo.WorkTimelineIndex, QueueInfo.Id);
 						WorkTrack->SetParentTrack(Track);
 						Track->AddChildTrack(WorkTrack);
+
+						const FString FencesTrackName = FString::Printf(TEXT("GPU%u - %s %u - Fences"), QueueInfo.GPU, QueueInfo.Name, QueueInfo.Index);
+						TSharedRef<FGpuFencesTimingTrack> FencesTrack = MakeShared<FGpuFencesTimingTrack>(*this, WorkTrackName, QueueInfo.WorkTimelineIndex, QueueInfo.Id);
+						FencesTrack->SetParentTrack(Track);
+						Track->AddChildTrack(FencesTrack);
 					}
 				}
 				else
diff --git a/Engine/Source/Developer/TraceInsights/Private/Insights/ViewModels/TimingViewDrawHelper.h b/Engine/Source/Developer/TraceInsights/Private/Insights/ViewModels/TimingViewDrawHelper.h
--- a/Engine/Source/Developer/TraceInsights/Private/Insights/ViewModels/TimingViewDrawHelper.h
+++ b/Engine/Source/Developer/TraceInsights/Private/Insights/ViewModels/TimingViewDrawHelper.h
@@ -187,6 +187,7 @@ public:
 	virtual int32 GetHeaderBackgroundLayerId() const override { return ReservedLayerId + ToInt32(EDrawLayer::HeaderBackground); }
 	virtual int32 GetHeaderTextLayerId() const override { return ReservedLayerId + ToInt32(EDrawLayer::HeaderText); }
 	virtual int32 GetRelationLayerId() const override {return ReservedLayerId + ToInt32(EDrawLayer::Relation); };
+	virtual int32 GetTextLayerId() const override {return ReservedLayerId + ToInt32(EDrawLayer::EventText); };
 	virtual int32 GetFirstLayerId() const override { return ReservedLayerId; }
 	virtual int32 GetNumLayerIds() const override { return ToInt32(EDrawLayer::Count); }
 
diff --git a/Engine/Source/Developer/TraceInsights/Public/Insights/ViewModels/ITimingViewDrawHelper.h b/Engine/Source/Developer/TraceInsights/Public/Insights/ViewModels/ITimingViewDrawHelper.h
--- a/Engine/Source/Developer/TraceInsights/Public/Insights/ViewModels/ITimingViewDrawHelper.h
+++ b/Engine/Source/Developer/TraceInsights/Public/Insights/ViewModels/ITimingViewDrawHelper.h
@@ -24,6 +24,7 @@ public:
 	virtual int32 GetHeaderBackgroundLayerId() const = 0;
 	virtual int32 GetHeaderTextLayerId() const = 0;
 	virtual int32 GetRelationLayerId() const = 0;
+	virtual int32 GetTextLayerId() const = 0;
 	virtual int32 GetFirstLayerId() const = 0;
 	virtual int32 GetNumLayerIds() const = 0;
 };
\ No newline at end of file
diff --git a/Engine/Source/Developer/TraceServices/Private/Model/TimingProfiler.cpp b/Engine/Source/Developer/TraceServices/Private/Model/TimingProfiler.cpp
--- a/Engine/Source/Developer/TraceServices/Private/Model/TimingProfiler.cpp
+++ b/Engine/Source/Developer/TraceServices/Private/Model/TimingProfiler.cpp
@@ -305,6 +305,79 @@ void FTimingProfilerProvider::EnumerateGpuWaitFences(uint32 QueueId, double Star
 	}
 }
 
+void FTimingProfilerProvider::EnumerateGpuFences(uint32 QueueId, double StartTime, double EndTime, EnumerateGpuFencesCallback Callback) const
+{
+	Session.ReadAccessCheck();
+
+	const uint32* FoundQueueIndex = GpuQueueIdToQueueIndexMap.Find(QueueId);
+	if (FoundQueueIndex == nullptr)
+	{
+		return;
+	}
+
+	const FGpuQueueData& Data = GpuQueueData[*FoundQueueIndex];
+	const TArray<FGpuSignalFence>& SignalArray = Data.SignalFenceArray;
+	const TArray<FGpuWaitFence>& WaitArray = Data.WaitFenceArray;
+
+	int32 SignalFenceIndex = Algo::LowerBoundBy(SignalArray, StartTime,
+		[](const FGpuSignalFence& Fence) { return Fence.Timestamp; });
+
+	int32 WaitFenceIndex = Algo::LowerBoundBy(WaitArray, StartTime,
+		[](const FGpuWaitFence& Fence) { return Fence.Timestamp; });
+
+	FGpuFenceWrapper Wrapper;
+
+	while (true)
+	{
+		const FGpuSignalFence* SignalFence = nullptr;
+		const FGpuWaitFence* WaitFence = nullptr;
+
+		if (SignalFenceIndex < SignalArray.Num() && SignalArray[SignalFenceIndex].Timestamp <= EndTime)
+		{
+			SignalFence = &SignalArray[SignalFenceIndex];
+		}
+
+		if (WaitFenceIndex < WaitArray.Num() && WaitArray[WaitFenceIndex].Timestamp <= EndTime)
+		{
+			WaitFence = &WaitArray[WaitFenceIndex];
+		}
+
+		if (SignalFence && WaitFence)
+		{
+			if (SignalFence->Timestamp <= WaitFence->Timestamp)
+			{
+				Wrapper.FenceType = EGpuFenceType::SignalFence;
+				Wrapper.Fence.Set<const FGpuSignalFence*>(SignalFence);
+				++SignalFenceIndex;
+			}
+			else
+			{
+				Wrapper.FenceType = EGpuFenceType::WaitFence;
+				Wrapper.Fence.Set<const FGpuWaitFence*>(WaitFence);
+				++WaitFenceIndex;
+			}
+		}
+		else if (SignalFence)
+		{
+			Wrapper.FenceType = EGpuFenceType::SignalFence;
+			Wrapper.Fence.Set<const FGpuSignalFence*>(SignalFence);
+			++SignalFenceIndex;
+		}
+		else if (WaitFence)
+		{
+			Wrapper.FenceType = EGpuFenceType::WaitFence;
+			Wrapper.Fence.Set<const FGpuWaitFence*>(WaitFence);
+			++WaitFenceIndex;
+		}
+		else
+		{
+			break;
+		}
+
+		Callback(Wrapper);
+	}
+}
+
 void FTimingProfilerProvider::EnumerateResolvedGpuFences(uint32 QueueId, double StartTime, double EndTime, EnumerateResolvedGpuFencesCallback Callback) const
 {
 	Session.ReadAccessCheck();
diff --git a/Engine/Source/Developer/TraceServices/Private/Model/TimingProfilerPrivate.h b/Engine/Source/Developer/TraceServices/Private/Model/TimingProfilerPrivate.h
--- a/Engine/Source/Developer/TraceServices/Private/Model/TimingProfilerPrivate.h
+++ b/Engine/Source/Developer/TraceServices/Private/Model/TimingProfilerPrivate.h
@@ -42,6 +42,7 @@ public:
 
 	virtual void EnumerateGpuSignalFences(uint32 QueueId, double StartTime, double EndTime, EnumerateGpuSignalFencesCallback Callback) const override;
 	virtual void EnumerateGpuWaitFences(uint32 QueueId, double StartTime, double EndTime, EnumerateGpuWaitFencesCallback Callback) const override;
+	virtual void EnumerateGpuFences(uint32 QueueId, double StartTime, double EndTime, EnumerateGpuFencesCallback Callback) const override;
 	virtual void EnumerateResolvedGpuFences(uint32 QueueId, double StartTime, double EndTime, EnumerateResolvedGpuFencesCallback Callback) const override;
 
 	virtual bool ReadTimeline(uint32 Index, TFunctionRef<void(const Timeline&)> Callback) const override;
diff --git a/Engine/Source/Developer/TraceServices/Public/TraceServices/Model/TimingProfiler.h b/Engine/Source/Developer/TraceServices/Public/TraceServices/Model/TimingProfiler.h
--- a/Engine/Source/Developer/TraceServices/Public/TraceServices/Model/TimingProfiler.h
+++ b/Engine/Source/Developer/TraceServices/Public/TraceServices/Model/TimingProfiler.h
@@ -143,6 +143,18 @@ struct FGpuWaitFence
 	uint32 QueueToWaitForId = 0;
 };
 
+enum EGpuFenceType : uint8
+{
+	SignalFence = 0,
+	WaitFence = 1,
+};
+
+struct FGpuFenceWrapper
+{
+	EGpuFenceType FenceType;
+	TVariant<const FGpuSignalFence*, const FGpuWaitFence*> Fence;
+};
+
 struct FGpuQueueInfo
 {
 	uint32 Id = 0;
@@ -167,6 +179,7 @@ public:
 	typedef ITimeline<FTimingProfilerEvent> Timeline;
 	typedef TFunctionRef<EEnumerateResult(const FGpuSignalFence&)> EnumerateGpuSignalFencesCallback;
 	typedef TFunctionRef<EEnumerateResult(const FGpuWaitFence&)> EnumerateGpuWaitFencesCallback;
+	typedef TFunctionRef<EEnumerateResult(const FGpuFenceWrapper&)> EnumerateGpuFencesCallback;
 	typedef TFunctionRef<EEnumerateResult(uint32 /*SignalFenceQueueId*/, const FGpuSignalFence& /*SignalFence*/, uint32 /*WaitFenceQueueId*/, const FGpuWaitFence& /*WaitFence*/)> EnumerateResolvedGpuFencesCallback;
 
 	virtual ~ITimingProfilerProvider() = default;
@@ -179,6 +192,7 @@ public:
 
 	virtual void EnumerateGpuSignalFences(uint32 QueueId, double StartTime, double EndTime, EnumerateGpuSignalFencesCallback Callback) const {}
 	virtual void EnumerateGpuWaitFences(uint32 QueueId, double StartTime, double EndTime, EnumerateGpuWaitFencesCallback Callback) const {}
+	virtual void EnumerateGpuFences(uint32 QueueId, double StartTime, double EndTime, EnumerateGpuFencesCallback Callback) const {}
 	virtual void EnumerateResolvedGpuFences(uint32 QueueId, double StartTime, double EndTime, EnumerateResolvedGpuFencesCallback Callback) const {}
 
 	virtual bool ReadTimeline(uint32 Index, TFunctionRef<void(const Timeline&)> Callback) const = 0;
