[Insights]
- Exposed a new TraceServices::IStringStore interface. The main string store can be accessed now from IAnalysisSession (using GetStringStore()).
- The AnalyssiSession's string store now supports also Find() in addition to Store().
- Memory Insights: Added support for specifying tag set and tracker when identifing a tag in the budget file.

diff --git a/Engine/Source/Developer/TraceInsights/Private/Insights/MemoryProfiler/ViewModels/MemTagBudget.cpp b/Engine/Source/Developer/TraceInsights/Private/Insights/MemoryProfiler/ViewModels/MemTagBudget.cpp
--- a/Engine/Source/Developer/TraceInsights/Private/Insights/MemoryProfiler/ViewModels/MemTagBudget.cpp
+++ b/Engine/Source/Developer/TraceInsights/Private/Insights/MemoryProfiler/ViewModels/MemTagBudget.cpp
@@ -12,6 +12,244 @@
 namespace UE::Insights::MemoryProfiler
 {
 
+////////////////////////////////////////////////////////////////////////////////////////////////////
+// FMemTagBudgetTagSet
+////////////////////////////////////////////////////////////////////////////////////////////////////
+
+void FMemTagBudgetTagSet::Reset()
+{
+	Values.Reset();
+}
+
+////////////////////////////////////////////////////////////////////////////////////////////////////
+// FMemTagBudgetTracker
+////////////////////////////////////////////////////////////////////////////////////////////////////
+
+void FMemTagBudgetTracker::Reset()
+{
+	for (auto& KV : TagSets)
+	{
+		delete KV.Value;
+	}
+	TagSets.Reset();
+}
+
+////////////////////////////////////////////////////////////////////////////////////////////////////
+
+void FMemTagBudgetTracker::EnumerateTagSets(TFunctionRef<void(const TCHAR* InCachedTagSetName, const FMemTagBudgetTagSet& InTagSet)> Callback) const
+{
+	for (const auto& KV : TagSets)
+	{
+		Callback((const TCHAR*)KV.Key, *KV.Value);
+	}
+}
+
+////////////////////////////////////////////////////////////////////////////////////////////////////
+
+const FMemTagBudgetTagSet* FMemTagBudgetTracker::FindTagSet(const TCHAR* InCachedTagSetName) const
+{
+	FMemTagBudgetTagSet*const* Found = TagSets.Find(InCachedTagSetName);
+	return Found ? *Found : nullptr;
+}
+
+////////////////////////////////////////////////////////////////////////////////////////////////////
+
+FMemTagBudgetTagSet& FMemTagBudgetTracker::GetOrAddTagSet(const TCHAR* InCachedTagSetName)
+{
+	FMemTagBudgetTagSet** Found = TagSets.Find(InCachedTagSetName);
+	if (Found)
+	{
+		return **Found;
+	}
+	FMemTagBudgetTagSet* TagSet = new FMemTagBudgetTagSet(InCachedTagSetName);
+	TagSets.Add(InCachedTagSetName, TagSet);
+	return *TagSet;
+}
+
+////////////////////////////////////////////////////////////////////////////////////////////////////
+// FMemTagBudgetPlatform
+////////////////////////////////////////////////////////////////////////////////////////////////////
+
+void FMemTagBudgetPlatform::Reset()
+{
+	for (auto& KV : Trackers)
+	{
+		delete KV.Value;
+	}
+	Trackers.Reset();
+}
+
+////////////////////////////////////////////////////////////////////////////////////////////////////
+
+void FMemTagBudgetPlatform::EnumerateTrackers(TFunctionRef<void(const TCHAR* InCachedTrackerName, const FMemTagBudgetTracker& InTracker)> Callback) const
+{
+	for (const auto& KV : Trackers)
+	{
+		Callback((const TCHAR*)KV.Key, *KV.Value);
+	}
+}
+
+////////////////////////////////////////////////////////////////////////////////////////////////////
+
+const FMemTagBudgetTracker* FMemTagBudgetPlatform::FindTracker(const TCHAR* InCachedTrackerName) const
+{
+	FMemTagBudgetTracker*const* Found = Trackers.Find(InCachedTrackerName);
+	return Found ? *Found : nullptr;
+}
+
+////////////////////////////////////////////////////////////////////////////////////////////////////
+
+FMemTagBudgetTracker& FMemTagBudgetPlatform::GetOrAddTracker(const TCHAR* InCachedTrackerName)
+{
+	FMemTagBudgetTracker** Found = Trackers.Find(InCachedTrackerName);
+	if (Found)
+	{
+		return **Found;
+	}
+	FMemTagBudgetTracker* Tracker = new FMemTagBudgetTracker(InCachedTrackerName);
+	Trackers.Add(InCachedTrackerName, Tracker);
+	return *Tracker;
+}
+
+////////////////////////////////////////////////////////////////////////////////////////////////////
+
+const FMemTagBudgetTagSet* FMemTagBudgetPlatform::FindTagSet(const TCHAR* InCachedTrackerName, const TCHAR* InCachedTagSetName) const
+{
+	const FMemTagBudgetTracker* Tracker = FindTracker(InCachedTrackerName);
+	return Tracker ? Tracker->FindTagSet(InCachedTagSetName) : nullptr;
+}
+
+////////////////////////////////////////////////////////////////////////////////////////////////////
+// FMemTagBudgetMode
+////////////////////////////////////////////////////////////////////////////////////////////////////
+
+void FMemTagBudgetMode::Reset()
+{
+	DefaultPlatform.Reset();
+	for (auto& KV : PlatformOverrides)
+	{
+		delete KV.Value;
+	}
+	PlatformOverrides.Reset();
+}
+
+////////////////////////////////////////////////////////////////////////////////////////////////////
+
+void FMemTagBudgetMode::EnumeratePlatforms(TFunctionRef<void(const TCHAR* InCachedPlatformName, const FMemTagBudgetPlatform& InPlatform)> Callback) const
+{
+	Callback(DefaultPlatform.GetName(), DefaultPlatform);
+
+	for (const auto& KV : PlatformOverrides)
+	{
+		Callback((const TCHAR*)KV.Key, *KV.Value);
+	}
+}
+
+////////////////////////////////////////////////////////////////////////////////////////////////////
+
+void FMemTagBudgetMode::EnumeratePlatformOverrides(TFunctionRef<void(const TCHAR* InCachedPlatformName, const FMemTagBudgetPlatform& InPlatform)> Callback) const
+{
+	for (const auto& KV : PlatformOverrides)
+	{
+		Callback((const TCHAR*)KV.Key, *KV.Value);
+	}
+}
+
+////////////////////////////////////////////////////////////////////////////////////////////////////
+
+const FMemTagBudgetPlatform* FMemTagBudgetMode::FindPlatformOverride(const TCHAR* InCachedPlatformName) const
+{
+	FMemTagBudgetPlatform*const* Found = PlatformOverrides.Find(InCachedPlatformName);
+	return Found ? *Found : nullptr;
+}
+
+////////////////////////////////////////////////////////////////////////////////////////////////////
+
+FMemTagBudgetPlatform& FMemTagBudgetMode::GetOrAddPlatformOverride(const TCHAR* InCachedPlatformName)
+{
+	FMemTagBudgetPlatform** Found = PlatformOverrides.Find(InCachedPlatformName);
+	if (Found)
+	{
+		return **Found;
+	}
+	FMemTagBudgetPlatform* Platform = new FMemTagBudgetPlatform(InCachedPlatformName);
+	PlatformOverrides.Add(InCachedPlatformName, Platform);
+	return *Platform;
+}
+
+////////////////////////////////////////////////////////////////////////////////////////////////////
+// FMemTagBudget
+////////////////////////////////////////////////////////////////////////////////////////////////////
+
+void FMemTagBudget::Reset()
+{
+	for (const auto& KV : Modes)
+	{
+		delete KV.Value;
+	}
+	Modes.Reset();
+}
+
+////////////////////////////////////////////////////////////////////////////////////////////////////
+
+const TCHAR* FMemTagBudget::FindString(const FStringView& InString) const
+{
+	return StringStore ? StringStore->Find(InString) : nullptr;
+}
+
+////////////////////////////////////////////////////////////////////////////////////////////////////
+
+const TCHAR* FMemTagBudget::StoreString(const FStringView& InString)
+{
+	return StringStore ? StringStore->Store(InString) : nullptr;
+}
+
+////////////////////////////////////////////////////////////////////////////////////////////////////
+
+void FMemTagBudget::EnumerateModes(TFunctionRef<void(const TCHAR* InCachedModeName, const FMemTagBudgetMode& InMode)> Callback) const
+{
+	for (const auto& KV : Modes)
+	{
+		Callback((const TCHAR*)KV.Key, *KV.Value);
+	}
+}
+
+////////////////////////////////////////////////////////////////////////////////////////////////////
+
+const FMemTagBudgetMode* FMemTagBudget::FindMode(const TCHAR* InCachedModeName) const
+{
+	FMemTagBudgetMode*const* Found = Modes.Find(InCachedModeName);
+	return Found ? *Found : nullptr;
+}
+
+////////////////////////////////////////////////////////////////////////////////////////////////////
+
+FMemTagBudgetMode& FMemTagBudget::GetOrAddMode(const TCHAR* InCachedModeName)
+{
+	FMemTagBudgetMode** Found = Modes.Find(InCachedModeName);
+	if (Found)
+	{
+		return **Found;
+	}
+	FMemTagBudgetMode* Mode = new FMemTagBudgetMode(InCachedModeName);
+	Modes.Add(InCachedModeName, Mode);
+	return *Mode;
+}
+
+////////////////////////////////////////////////////////////////////////////////////////////////////
+
+const FMemTagBudgetMode* FMemTagBudget::FindMode(const FString& InModeName) const
+{
+	return FindMode(FindString(InModeName));
+}
+
+////////////////////////////////////////////////////////////////////////////////////////////////////
+
+FMemTagBudgetMode& FMemTagBudget::GetOrAddMode(const FString& InModeName)
+{
+	return GetOrAddMode(StoreString(InModeName));
+}
+
 ////////////////////////////////////////////////////////////////////////////////////////////////////
 
 bool FMemTagBudget::LoadFromFile(const FString& InFilePath)
@@ -25,10 +263,19 @@ bool FMemTagBudget::LoadFromFile(const FString& InFilePath)
 
 bool FMemTagBudget::ReLoadFromFile()
 {
+	const TCHAR* DefaultCachedPlatformName = StoreString(FString(TEXT("Default")));
+
+	const TCHAR* DefaultCachedTrackerName = StoreString(FString(TEXT("Default")));
+	const TCHAR* PlatformCachedTrackerName = StoreString(FString(TEXT("Platform")));
+
+	const TCHAR* SystemCachedTagSetName = StoreString(FString(TEXT("System")));
+	const TCHAR* AssetCachedTagSetName = StoreString(FString(TEXT("Assets")));
+	const TCHAR* AssetClassCachedTagSetName = StoreString(FString(TEXT("AssetsClasses")));
+
 	static const FString STR_Budget(TEXT("budget"));
 	static const FString STR_Platform(TEXT("platform"));
-	//static const FString STR_Set(TEXT("set"));
-	//static const FString STR_Tracker(TEXT("tracker"));
+	static const FString STR_Set(TEXT("set"));
+	static const FString STR_Tracker(TEXT("tracker"));
 	static const FString STR_Name(TEXT("name"));
 	//static const FString STR_NameRegEx(TEXT("name_regex"));
 	static const FString STR_Tag(TEXT("tag"));
@@ -56,16 +303,74 @@ bool FMemTagBudget::ReLoadFromFile()
 		{
 			FString ModeName = BudgetNode->GetAttribute(STR_Name);
 			FMemTagBudgetMode& BudgetMode = GetOrAddMode(ModeName);
+			BudgetMode.GetDefaultPlatform().SetName(DefaultCachedPlatformName);
 
-			auto ProcessTagNode = [&](FXmlNode* GroupNode, FMemTagBudgetValueSet& Set)
+			auto ProcessTagNode = [&](FXmlNode* TagNode, FMemTagBudgetPlatform& BudgetPlatform)
 				{
-					FString TagName = GroupNode->GetAttribute(STR_Name);
-					if (Set.FindValue(TagName) != nullptr)
+					FString TagName = TagNode->GetAttribute(STR_Name);
+					const TCHAR* CachedTagName = StoreString(TagName);
+					if (CachedTagName == nullptr)
+					{
+						UE_LOG(LogMemoryProfiler, Warning, TEXT("Ignoring budget value for tag \"%s\" (invalid name)."), *TagName);
+						return;
+					}
+
+					//////////////////////////////////////////////////
+
+					const TCHAR* CachedTrackerName = DefaultCachedTrackerName;
+					FString TrackerName = TagNode->GetAttribute(STR_Tracker);
+					if (!TrackerName.IsEmpty())
+					{
+						if (!FCString::Stricmp(*TrackerName, TEXT("Platform")))
+						{
+							CachedTrackerName = PlatformCachedTrackerName;
+						}
+						else if (!FCString::Stricmp(*TrackerName, TEXT("Default")))
+						{
+							CachedTrackerName = DefaultCachedTrackerName;
+						}
+						else
+						{
+							UE_LOG(LogMemoryProfiler, Warning, TEXT("Ignoring unknown tracker \"%s\"."), *TrackerName);
+						}
+						CachedTrackerName = StoreString(TrackerName);
+					}
+					FMemTagBudgetTracker& BudgetTracker = BudgetPlatform.GetOrAddTracker(CachedTrackerName);
+
+					//////////////////////////////////////////////////
+
+					const TCHAR* CachedTagSetName = SystemCachedTagSetName;
+					FString TagSetName = TagNode->GetAttribute(STR_Set);
+					if (!TagSetName.IsEmpty())
+					{
+						if (!FCString::Stricmp(*TagSetName, TEXT("Asset")) || !FCString::Stricmp(*TagSetName, TEXT("Assets")))
+						{
+							CachedTagSetName = AssetCachedTagSetName;
+						}
+						else if (!FCString::Stricmp(*TagSetName, TEXT("AssetClass")) || !FCString::Stricmp(*TagSetName, TEXT("AssetClasses")))
+						{
+							CachedTagSetName = AssetClassCachedTagSetName;
+						}
+						else if (!FCString::Stricmp(*TagSetName, TEXT("System")) || !FCString::Stricmp(*TagSetName, TEXT("Systems")))
+						{
+							CachedTagSetName = SystemCachedTagSetName;
+						}
+						else
+						{
+							UE_LOG(LogMemoryProfiler, Warning, TEXT("Ignoring unknown tag set \"%s\"."), *TagSetName);
+						}
+					}
+					FMemTagBudgetTagSet& BudgetTagSet = BudgetTracker.GetOrAddTagSet(CachedTagSetName);
+
+					//////////////////////////////////////////////////
+
+					if (BudgetTagSet.FindValue(CachedTagName) != nullptr)
 					{
 						UE_LOG(LogMemoryProfiler, Warning, TEXT("Ignoring budget value for tag \"%s\" (defined multiple times)."), *TagName);
 						return;
 					}
-					FString Value = GroupNode->GetAttribute(STR_TagMemMax);
+
+					FString Value = TagNode->GetAttribute(STR_TagMemMax);
 
 					PRAGMA_DISABLE_UNSAFE_TYPECAST_WARNINGS
 					TValueOrError<FNumericUnit<double>, FText> Result = FNumericUnit<double>::TryParseExpression(*Value, EUnit::Bytes, 0);
@@ -74,7 +379,7 @@ bool FMemTagBudget::ReLoadFromFile()
 						auto MemoryValueInBytes = Result.GetValue().ConvertTo(EUnit::Bytes);
 						if (MemoryValueInBytes.IsSet())
 						{
-							Set.AddValue(TagName, (int64)MemoryValueInBytes->Value);
+							BudgetTagSet.AddValue(CachedTagName, (int64)MemoryValueInBytes->Value);
 						}
 					}
 					else
@@ -90,7 +395,9 @@ bool FMemTagBudget::ReLoadFromFile()
 				if (BudgetChildNode->GetTag() == STR_Platform)
 				{
 					FString PlatformName = BudgetChildNode->GetAttribute(STR_Name);
-					FMemTagBudgetValueSet& BudgetPlatform = *BudgetMode.GetOrAddPlatform(PlatformName);
+					const TCHAR* CachedPlatformName = StoreString(PlatformName);
+
+					FMemTagBudgetPlatform& BudgetPlatform = BudgetMode.GetOrAddPlatformOverride(CachedPlatformName);
 					const TArray<FXmlNode*>& PlatformChildNodes = BudgetChildNode->GetChildrenNodes();
 					for (FXmlNode* PlatformChildNode : PlatformChildNodes)
 					{
@@ -102,7 +409,7 @@ bool FMemTagBudget::ReLoadFromFile()
 				}
 				else if (BudgetChildNode->GetTag() == STR_Tag)
 				{
-					ProcessTagNode(BudgetChildNode, BudgetMode.GetDefaultValues());
+					ProcessTagNode(BudgetChildNode, BudgetMode.GetDefaultPlatform());
 				}
 			}
 		}
diff --git a/Engine/Source/Developer/TraceInsights/Private/Insights/MemoryProfiler/ViewModels/MemTagBudget.h b/Engine/Source/Developer/TraceInsights/Private/Insights/MemoryProfiler/ViewModels/MemTagBudget.h
--- a/Engine/Source/Developer/TraceInsights/Private/Insights/MemoryProfiler/ViewModels/MemTagBudget.h
+++ b/Engine/Source/Developer/TraceInsights/Private/Insights/MemoryProfiler/ViewModels/MemTagBudget.h
@@ -4,6 +4,9 @@
 
 #include "CoreTypes.h"
 
+// TraceServices
+#include "TraceServices/Containers/Allocators.h" // for IStringStore
+
 #include "Containers/Array.h"
 #include "Containers/Map.h"
 #include "Containers/UnrealString.h"
@@ -14,146 +17,169 @@ namespace UE::Insights::MemoryProfiler
 
 ////////////////////////////////////////////////////////////////////////////////////////////////////
 
-class FMemTagBudgetValueSet
+class FMemTagBudgetTagSet
 {
 public:
-	FMemTagBudgetValueSet() {}
+	FMemTagBudgetTagSet(const TCHAR* InCachedTagSetName)
+	: CachedTagSetName(InCachedTagSetName)
+	{
+	}
+	~FMemTagBudgetTagSet()
+	{
+		Reset();
+	}
+
+	void Reset();
 
-	const int64* FindValue(const FString& Name) const
+	const TCHAR* GetName() const { return CachedTagSetName; }
+	void SetName(const TCHAR* InCachedTagSetName) { CachedTagSetName = InCachedTagSetName; }
+
+	const int64* FindValue(const TCHAR* InCachedTagName) const
 	{
-		return Values.Find(Name);
+		return Values.Find(InCachedTagName);
 	}
 
-	void AddValue(const FString& Name, int64 Value)
+	void AddValue(const TCHAR* InCachedTagName, int64 Value)
 	{
-		Values.Add(Name, Value);
+		Values.Add(InCachedTagName, Value);
 	}
 
 private:
-	TMap<FString, int64> Values;
+	const TCHAR* CachedTagSetName = nullptr;
+	TMap<const void*, int64> Values;
 };
 
 ////////////////////////////////////////////////////////////////////////////////////////////////////
 
-class FMemTagBudgetMode
+class FMemTagBudgetTracker
 {
 public:
-	FMemTagBudgetMode(const FString& InName)
-		: Name(InName)
-		, DefaultValues(MakeShared<FMemTagBudgetValueSet>())
-		, PlatformOverrides()
+	FMemTagBudgetTracker(const TCHAR* InCachedTrackerName)
+	: CachedTrackerName(InCachedTrackerName)
 	{
 	}
-
-	const FString& GetName() const
+	~FMemTagBudgetTracker()
 	{
-		return Name;
+		Reset();
 	}
 
-	const FMemTagBudgetValueSet& GetDefaultValues() const
-	{
-		return *DefaultValues;
-	}
+	void Reset();
 
-	FMemTagBudgetValueSet& GetDefaultValues()
-	{
-		return *DefaultValues;
-	}
+	const TCHAR* GetName() const { return CachedTrackerName; }
+	void SetName(const TCHAR* InCachedTrackerName) { CachedTrackerName = InCachedTrackerName; }
 
-	void GetPlatforms(TArray<FString> &OutPlatforms)
-	{
-		PlatformOverrides.GetKeys(OutPlatforms);
-	}
+	void EnumerateTagSets(TFunctionRef<void(const TCHAR* InCachedTagSetName, const FMemTagBudgetTagSet& InTagSet)> Callback) const;
+	const FMemTagBudgetTagSet* FindTagSet(const TCHAR* InCachedTagSetName) const;
+	FMemTagBudgetTagSet& GetOrAddTagSet(const TCHAR* InCachedTagSetName);
 
-	void EnumeratePlatforms(TFunctionRef<void(const FString& Platform, const FMemTagBudgetValueSet& Set)> Callback) const
-	{
-		for (const auto& KV : PlatformOverrides)
-		{
-			Callback(KV.Key, *KV.Value);
-		}
-	}
+private:
+	const TCHAR* CachedTrackerName = nullptr;
+	TMap<const void*, FMemTagBudgetTagSet*> TagSets;
+};
+
+////////////////////////////////////////////////////////////////////////////////////////////////////
 
-	const FMemTagBudgetValueSet* FindPlatform(const FString& Platform) const
+class FMemTagBudgetPlatform
+{
+public:
+	FMemTagBudgetPlatform(const TCHAR* InCachedPlatformName)
+		: CachedPlatformName(InCachedPlatformName)
 	{
-		const TSharedRef<FMemTagBudgetValueSet>* Found = PlatformOverrides.Find(Platform);
-		return Found ? &**Found : nullptr;
 	}
 
-	TSharedRef<FMemTagBudgetValueSet>& GetOrAddPlatform(const FString& Platform)
+	~FMemTagBudgetPlatform()
 	{
-		TSharedRef<FMemTagBudgetValueSet>* Found = PlatformOverrides.Find(Platform);
-		if (Found)
-		{
-			return *Found;
-		}
-		return PlatformOverrides.Add(Platform, MakeShared<FMemTagBudgetValueSet>());
+		Reset();
 	}
 
+	void Reset();
+
+	const TCHAR* GetName() const { return CachedPlatformName; }
+	void SetName(const TCHAR* InCachedPlatformName) { CachedPlatformName = InCachedPlatformName; }
+
+	void EnumerateTrackers(TFunctionRef<void(const TCHAR* InCachedTrackerName, const FMemTagBudgetTracker& InTracker)> Callback) const;
+	const FMemTagBudgetTracker* FindTracker(const TCHAR* InCachedTrackerName) const;
+	FMemTagBudgetTracker& GetOrAddTracker(const TCHAR* InCachedTrackerName);
+
+	const FMemTagBudgetTagSet* FindTagSet(const TCHAR* InCachedTrackerName, const TCHAR* InCachedTagSetName) const;
+
 private:
-	FString Name;
-	TSharedRef<FMemTagBudgetValueSet> DefaultValues;
-	TMap<FString, TSharedRef<FMemTagBudgetValueSet>> PlatformOverrides;
+	const TCHAR* CachedPlatformName = nullptr;
+	TMap<const void*, FMemTagBudgetTracker*> Trackers;
 };
 
 ////////////////////////////////////////////////////////////////////////////////////////////////////
 
-class FMemTagBudget
+class FMemTagBudgetMode
 {
 public:
-	FMemTagBudget(const FString& InName)
-		: Name(InName)
+	FMemTagBudgetMode(const TCHAR* InCachedModeName)
+		: CachedModeName(InCachedModeName)
+		, DefaultPlatform(nullptr)
 	{
 	}
 
-	void Reset()
+	~FMemTagBudgetMode()
 	{
-		Name.Reset();
-		Modes.Reset();
+		Reset();
 	}
 
-	const FString& GetName() const
-	{
-		return Name;
-	}
+	void Reset();
 
-	const FString& GetFilePath() const
-	{
-		return FilePath;
-	}
+	const TCHAR* GetName() const { return CachedModeName; }
+	void SetName(const TCHAR* InCachedModeName) { CachedModeName = InCachedModeName; }
 
-	TArray<TSharedRef<FMemTagBudgetMode>>& GetModes() { return Modes; }
+	void EnumeratePlatforms(TFunctionRef<void(const TCHAR* InCachedPlatformName, const FMemTagBudgetPlatform& InPlatform)> Callback) const;
+	const FMemTagBudgetPlatform& GetDefaultPlatform() const { return DefaultPlatform; }
+	void EnumeratePlatformOverrides(TFunctionRef<void(const TCHAR* InCachedPlatformName, const FMemTagBudgetPlatform& InPlatform)> Callback) const;
+	const FMemTagBudgetPlatform* FindPlatformOverride(const TCHAR* InCachedPlatformName) const;
+	FMemTagBudgetPlatform& GetDefaultPlatform() { return DefaultPlatform; }
+	FMemTagBudgetPlatform& GetOrAddPlatformOverride(const TCHAR* InCachedPlatformName);
 
-	const FMemTagBudgetMode* FindMode(const FString& InModeName) const
+private:
+	const TCHAR* CachedModeName = nullptr;
+	FMemTagBudgetPlatform DefaultPlatform;
+	TMap<const void*, FMemTagBudgetPlatform*> PlatformOverrides;
+};
+////////////////////////////////////////////////////////////////////////////////////////////////////
+
+class FMemTagBudget
+{
+public:
+	FMemTagBudget(TraceServices::IStringStore* InStringStore)
+		: StringStore(InStringStore)
 	{
-		const TSharedRef<FMemTagBudgetMode>* FoundBudgetMode = Modes.FindByPredicate(
-			[&InModeName](const TSharedRef<const FMemTagBudgetMode>& BudgetMode)
-			{
-				return InModeName.Equals(BudgetMode->GetName(), ESearchCase::CaseSensitive);
-			});
-		return (FoundBudgetMode) ? &**FoundBudgetMode : nullptr;
 	}
 
-	FMemTagBudgetMode& GetOrAddMode(const FString& InModeName)
+	~FMemTagBudget()
 	{
-		TSharedRef<FMemTagBudgetMode>* FoundBudgetMode = Modes.FindByPredicate(
-			[&InModeName](const TSharedRef<const FMemTagBudgetMode>& BudgetMode)
-			{
-				return InModeName.Equals(BudgetMode->GetName(), ESearchCase::CaseSensitive);
-			});
-		if (FoundBudgetMode)
-		{
-			return **FoundBudgetMode;
-		}
-		return *Modes.Add_GetRef(MakeShared<FMemTagBudgetMode>(InModeName));
+		Reset();
 	}
 
+	void Reset();
+
+	const TCHAR* FindString(const FStringView& InString) const;
+	const TCHAR* StoreString(const FStringView& InString);
+
+	void EnumerateModes(TFunctionRef<void(const TCHAR* InCachedModeName, const FMemTagBudgetMode& InMode)> Callback) const;
+	const FMemTagBudgetMode* FindMode(const TCHAR* InCachedModeName) const;
+	FMemTagBudgetMode& GetOrAddMode(const TCHAR* InCachedModeName);
+
+	const FMemTagBudgetMode* FindMode(const FString& InModeName) const;
+	FMemTagBudgetMode& GetOrAddMode(const FString& InModeName);
+
+	const FString& GetName() const { return Name; }
+	void SetName(const FString& InName) { Name = InName; }
+
+	const FString& GetFilePath() const { return FilePath; }
 	bool LoadFromFile(const FString& InFilePath);
 	bool ReLoadFromFile();
 
 private:
+	TraceServices::IStringStore* StringStore = nullptr;
 	FString Name;
 	FString FilePath;
-	TArray<TSharedRef<FMemTagBudgetMode>> Modes;
+	TMap<const void*, FMemTagBudgetMode*> Modes;
 };
 
 ////////////////////////////////////////////////////////////////////////////////////////////////////
diff --git a/Engine/Source/Developer/TraceInsights/Private/Insights/MemoryProfiler/Widgets/SMemTagTreeView.cpp b/Engine/Source/Developer/TraceInsights/Private/Insights/MemoryProfiler/Widgets/SMemTagTreeView.cpp
--- a/Engine/Source/Developer/TraceInsights/Private/Insights/MemoryProfiler/Widgets/SMemTagTreeView.cpp
+++ b/Engine/Source/Developer/TraceInsights/Private/Insights/MemoryProfiler/Widgets/SMemTagTreeView.cpp
@@ -102,13 +102,8 @@ TSharedPtr<SWidget> SMemTagTreeView::ConstructFilterToolbar()
 
 	+ SHorizontalBox::Slot()
 	.FillWidth(1.0f)
-	[
-		SNew(SBox)
-	]
-
-	+ SHorizontalBox::Slot()
-	.AutoWidth()
-	.Padding(4.0f, 0.0f, 0.0f, 0.0f)
+	.Padding(0.0f, 0.0f, 0.0f, 0.0f)
+	.HAlign(HAlign_Right)
 	.VAlign(VAlign_Center)
 	[
 		SNew(STextBlock)
@@ -1940,10 +1935,16 @@ void SMemTagTreeView::BudgetFile_OnSelectionChanged(TSharedPtr<FMemTagBudgetFile
 		if (!SelectedBudgetFile.IsValid() || *SelectedBudgetFile != *InBudgetFile)
 		{
 			SelectedBudgetFile = InBudgetFile;
-			CurrentBudget = MakeShared<FMemTagBudget>(InBudgetFile->Name);
-			if (!CurrentBudget->LoadFromFile(InBudgetFile->FileName))
+			CurrentBudget = nullptr;
+			if (Session.IsValid())
 			{
-				// error
+				TraceServices::IStringStore* StringStore = &const_cast<TraceServices::IAnalysisSession*>(Session.Get())->GetStringStore();
+				CurrentBudget = MakeShared<FMemTagBudget>(StringStore);
+				CurrentBudget->SetName(InBudgetFile->Name);
+				if (!CurrentBudget->LoadFromFile(InBudgetFile->FileName))
+				{
+					// error
+				}
 			}
 			OnBudgetChanged();
 		}
@@ -2171,23 +2172,22 @@ void SMemTagTreeView::OnBudgetChanged()
 		AvailablePlatforms.Add(AvailableBudgetPlatforms[0]->Name);
 
 		// Add modes and platforms available in the current budget.
-		for (const auto& BudgetMode : CurrentBudget->GetModes())
-		{
-			const FString& ModeName = BudgetMode->GetName();
-			AvailableBudgetModes.Add(MakeShared<FMemTagBudgetModeDesc>(ModeName));
-
-			TArray<FString> Platforms;
-			BudgetMode->GetPlatforms(Platforms);
-			for (const FString& Platform : Platforms)
+		CurrentBudget->EnumerateModes([this](const TCHAR* CachedModeName, const FMemTagBudgetMode& BudgetMode)
 			{
-				bool bIsAlreadyInSet = false;
-				AvailablePlatforms.Add(Platform, &bIsAlreadyInSet);
-				if (!bIsAlreadyInSet)
-				{
-					AvailableBudgetPlatforms.Add(MakeShared<FMemTagBudgetPlatformDesc>(Platform));
-				}
-			}
-		}
+				FString ModeName(CachedModeName);
+				AvailableBudgetModes.Add(MakeShared<FMemTagBudgetModeDesc>(ModeName));
+
+				BudgetMode.EnumeratePlatforms([this](const TCHAR* CachedPlatformName, const FMemTagBudgetPlatform& Platform)
+					{
+						FString PlatformName(CachedPlatformName);
+						bool bIsAlreadyInSet = false;
+						AvailablePlatforms.Add(PlatformName, &bIsAlreadyInSet);
+						if (!bIsAlreadyInSet)
+						{
+							AvailableBudgetPlatforms.Add(MakeShared<FMemTagBudgetPlatformDesc>(PlatformName));
+						}
+					});
+			});
 	}
 
 	// Restore selection for Budget Mode.
@@ -2267,6 +2267,13 @@ void SMemTagTreeView::ApplyBudgetToNodes()
 		return;
 	}
 
+	TSharedPtr<SMemoryProfilerWindow> ProfilerWindow = GetProfilerWindow();
+	if (!ProfilerWindow.IsValid())
+	{
+		return;
+	}
+	FMemorySharedState& SharedState = ProfilerWindow->GetSharedState();
+
 	if (!CurrentBudget.IsValid())
 	{
 		ResetBudgetForAllNodes();
@@ -2280,16 +2287,38 @@ void SMemTagTreeView::ApplyBudgetToNodes()
 		return;
 	}
 	const FMemTagBudgetMode& BudgetMode = *FoundBudgetMode;
+	
+	const TCHAR* CachedTagSetName = nullptr;
+	const FMemoryTagSet* TagSet = SharedState.GetTagSetById(TagSetFilter);
+	if (TagSet)
+	{
+		CachedTagSetName = CurrentBudget->FindString(TagSet->GetName());
+	}
 
-	const FMemTagBudgetValueSet& DefaultValues = BudgetMode.GetDefaultValues();
-	const FMemTagBudgetValueSet* PlatformOverrides = SelectedBudgetPlatform.IsValid() ? BudgetMode.FindPlatform(SelectedBudgetPlatform->Name) : nullptr;
+	const TCHAR* CachedDefaultTrackerName = CurrentBudget->FindString(TEXTVIEW("Default"));
+	const TCHAR* CachedPlatformTrackerName = CurrentBudget->FindString(TEXTVIEW("Platform"));
 
-	TSharedPtr<SMemoryProfilerWindow> ProfilerWindow = GetProfilerWindow();
-	if (!ProfilerWindow.IsValid())
+	const FMemTagBudgetPlatform& DefaultPlatform = BudgetMode.GetDefaultPlatform();
+	const FMemTagBudgetTagSet* DefaultTagSet = DefaultPlatform.FindTagSet(CachedDefaultTrackerName, CachedTagSetName);
+	const FMemTagBudgetTagSet* PlatformTagSet = DefaultPlatform.FindTagSet(CachedPlatformTrackerName, CachedTagSetName);
+	const FMemTagBudgetTagSet* DefaultTagSetOverride = nullptr;
+	const FMemTagBudgetTagSet* PlatformTagSetOverride = nullptr;
+	if (SelectedBudgetPlatform.IsValid())
 	{
+		const TCHAR* CachedPlatformName = CurrentBudget->FindString(SelectedBudgetPlatform->Name);
+		const FMemTagBudgetPlatform* PlatformOverride = BudgetMode.FindPlatformOverride(CachedPlatformName);
+		if (PlatformOverride)
+		{
+			DefaultTagSetOverride = PlatformOverride->FindTagSet(CachedDefaultTrackerName, CachedTagSetName);
+			PlatformTagSetOverride = PlatformOverride->FindTagSet(CachedPlatformTrackerName, CachedTagSetName);
+		}
+	}
+
+	if (!DefaultTagSet && !PlatformTagSet && !DefaultTagSetOverride && !PlatformTagSetOverride)
+	{
+		ResetBudgetForAllNodes();
 		return;
 	}
-	FMemorySharedState& SharedState = ProfilerWindow->GetSharedState();
 
 	for (const FTableTreeNodePtr& NodePtr : TableRowNodes)
 	{
@@ -2305,9 +2334,14 @@ void SMemTagTreeView::ApplyBudgetToNodes()
 			continue;
 		}
 
-		if (PlatformOverrides)
+		const TCHAR* CachedTagName = CurrentBudget->FindString(MemoryTag->GetStatFullName());
+
+		bool bIsDefaultTracker = MemoryTag->GetTrackerId() == SharedState.GetDefaultTrackerId();
+
+		const FMemTagBudgetTagSet* BudgetTagSetOverride = bIsDefaultTracker ? DefaultTagSetOverride : PlatformTagSetOverride;
+		if (BudgetTagSetOverride)
 		{
-			const int64* ValuePtr = PlatformOverrides->FindValue(MemoryTag->GetStatFullName());
+			const int64* ValuePtr = BudgetTagSetOverride->FindValue(CachedTagName);
 			if (ValuePtr)
 			{
 				MemTagNode.SetSizeBudget(*ValuePtr);
@@ -2316,9 +2350,10 @@ void SMemTagTreeView::ApplyBudgetToNodes()
 			}
 		}
 
-		// DefaultValues
+		const FMemTagBudgetTagSet* BudgetTagSet = bIsDefaultTracker ? DefaultTagSet : PlatformTagSet;
+		if (BudgetTagSet)
 		{
-			const int64* ValuePtr = DefaultValues.FindValue(MemoryTag->GetStatFullName());
+			const int64* ValuePtr = BudgetTagSet->FindValue(CachedTagName);
 			if (ValuePtr)
 			{
 				MemTagNode.SetSizeBudget(*ValuePtr);
diff --git a/Engine/Source/Developer/TraceInsights/Private/Insights/MemoryProfiler/Widgets/SMemTagTreeView.h b/Engine/Source/Developer/TraceInsights/Private/Insights/MemoryProfiler/Widgets/SMemTagTreeView.h
--- a/Engine/Source/Developer/TraceInsights/Private/Insights/MemoryProfiler/Widgets/SMemTagTreeView.h
+++ b/Engine/Source/Developer/TraceInsights/Private/Insights/MemoryProfiler/Widgets/SMemTagTreeView.h
@@ -322,8 +322,6 @@ private:
 	//////////////////////////////////////////////////
 	// Filters
 
-	TSharedPtr<SComboBox<TSharedPtr<FMemoryTracker>>> TrackerComboBox;
-
 	/** Filter the LLM tags by tracker. */
 	uint64 TrackersFilter = uint64(-1);
 
diff --git a/Engine/Source/Developer/TraceServices/Private/AnalysisServicePrivate.h b/Engine/Source/Developer/TraceServices/Private/AnalysisServicePrivate.h
--- a/Engine/Source/Developer/TraceServices/Private/AnalysisServicePrivate.h
+++ b/Engine/Source/Developer/TraceServices/Private/AnalysisServicePrivate.h
@@ -64,6 +64,7 @@ public:
 	virtual TArray<FAnalysisMessage> DrainPendingMessages() override;
 
 	virtual ILinearAllocator& GetLinearAllocator() override { return Allocator; }
+	virtual IStringStore& GetStringStore() override { return StringStore; }
 	virtual const TCHAR* StoreString(const TCHAR* String) override { return StringStore.Store(String); }
 	virtual const TCHAR* StoreString(const FStringView& String) override { return StringStore.Store(String); }
 
diff --git a/Engine/Source/Developer/TraceServices/Private/Common/StringStore.cpp b/Engine/Source/Developer/TraceServices/Private/Common/StringStore.cpp
--- a/Engine/Source/Developer/TraceServices/Private/Common/StringStore.cpp
+++ b/Engine/Source/Developer/TraceServices/Private/Common/StringStore.cpp
@@ -11,6 +11,34 @@ FStringStore::FStringStore(FSlabAllocator& InAllocator)
 {
 }
 
+FStringStore::~FStringStore()
+{
+}
+
+const TCHAR* FStringStore::Find(const TCHAR* String) const
+{
+	return Find(FStringView(String));
+}
+
+const TCHAR* FStringStore::Find(const FStringView& String) const
+{
+	FScopeLock _(&Cs);
+
+	uint32 Hash = GetTypeHash(String);
+
+	FindStoredStrings.Reset();
+	StoredStrings.MultiFind(Hash, FindStoredStrings);
+	for (const TCHAR* FoundStoredString : FindStoredStrings)
+	{
+		if (!String.Compare(FStringView(FoundStoredString)))
+		{
+			return FoundStoredString;
+		}
+	}
+
+	return nullptr;
+}
+
 const TCHAR* FStringStore::Store(const TCHAR* String)
 {
 	return Store(FStringView(String));
diff --git a/Engine/Source/Developer/TraceServices/Private/Common/StringStore.h b/Engine/Source/Developer/TraceServices/Private/Common/StringStore.h
--- a/Engine/Source/Developer/TraceServices/Private/Common/StringStore.h
+++ b/Engine/Source/Developer/TraceServices/Private/Common/StringStore.h
@@ -13,22 +13,27 @@
 namespace TraceServices
 {
 
-class FStringStore
+class FStringStore : public IStringStore
 {
 public:
 	FStringStore(FSlabAllocator& Allocator);
-	const TCHAR* Store(const TCHAR* String);
-	const TCHAR* Store(const FStringView& String);
+	virtual ~FStringStore();
+
+	virtual const TCHAR* Find(const TCHAR* String) const override;
+	virtual const TCHAR* Find(const FStringView& String) const override;
+
+	virtual const TCHAR* Store(const TCHAR* String) override;
+	virtual const TCHAR* Store(const FStringView& String) override;
 
 private:
 	enum
 	{
 		BlockSize = 4 << 20
 	};
-	FCriticalSection Cs;
+	mutable FCriticalSection Cs;
 	FSlabAllocator& Allocator;
 	TMultiMap<uint32, const TCHAR*> StoredStrings;
-	TArray<const TCHAR*> FindStoredStrings;
+	mutable TArray<const TCHAR*> FindStoredStrings;
 	TCHAR* BufferPtr = nullptr;
 	uint64 BufferLeft = 0;
 	uint64 BlockCount = 0;
diff --git a/Engine/Source/Developer/TraceServices/Public/TraceServices/Containers/Allocators.h b/Engine/Source/Developer/TraceServices/Public/TraceServices/Containers/Allocators.h
--- a/Engine/Source/Developer/TraceServices/Public/TraceServices/Containers/Allocators.h
+++ b/Engine/Source/Developer/TraceServices/Public/TraceServices/Containers/Allocators.h
@@ -3,6 +3,7 @@
 #pragma once
 
 #include "HAL/Platform.h"
+#include "Containers/StringFwd.h"
 
 namespace TraceServices
 {
@@ -14,4 +15,14 @@ public:
 	virtual void* Allocate(uint64 Size) = 0;
 };
 
+class IStringStore
+{
+public:
+	virtual ~IStringStore() = default;
+	virtual const TCHAR* Find(const TCHAR* String) const = 0;
+	virtual const TCHAR* Find(const FStringView& String) const = 0;
+	virtual const TCHAR* Store(const TCHAR* String) = 0;
+	virtual const TCHAR* Store(const FStringView& String) = 0;
+};
+
 } // namespace TraceServices
diff --git a/Engine/Source/Developer/TraceServices/Public/TraceServices/Model/AnalysisSession.h b/Engine/Source/Developer/TraceServices/Public/TraceServices/Model/AnalysisSession.h
--- a/Engine/Source/Developer/TraceServices/Public/TraceServices/Model/AnalysisSession.h
+++ b/Engine/Source/Developer/TraceServices/Public/TraceServices/Model/AnalysisSession.h
@@ -22,8 +22,9 @@ class IAnalyzer;
 namespace TraceServices
 {
 
-class ILinearAllocator;
 class IAnalysisCache;
+class ILinearAllocator;
+class IStringStore;
 
 class IProvider
 {
@@ -111,6 +112,7 @@ public:
 	virtual void WriteAccessCheck() = 0;
 
 	virtual ILinearAllocator& GetLinearAllocator() = 0;
+	virtual IStringStore& GetStringStore() = 0;
 	virtual const TCHAR* StoreString(const TCHAR* String) = 0;
 	virtual const TCHAR* StoreString(const FStringView& String) = 0;
 
