[Insights] Memory Insights: Fixed crash when analyzing a trace session that has AssetClasses tag set, but not the Assets tag set.

diff --git a/Engine/Source/Developer/TraceInsights/Private/Insights/MemoryProfiler/ViewModels/MemorySharedState.cpp b/Engine/Source/Developer/TraceInsights/Private/Insights/MemoryProfiler/ViewModels/MemorySharedState.cpp
--- a/Engine/Source/Developer/TraceInsights/Private/Insights/MemoryProfiler/ViewModels/MemorySharedState.cpp
+++ b/Engine/Source/Developer/TraceInsights/Private/Insights/MemoryProfiler/ViewModels/MemorySharedState.cpp
@@ -82,18 +82,6 @@ UE_ENABLE_OPTIMIZATION_SHIP
 ////////////////////////////////////////////////////////////////////////////////////////////////////
 
 FMemorySharedState::FMemorySharedState()
-	: TimingView(nullptr)
-	, DefaultTracker(nullptr)
-	, PlatformTracker(nullptr)
-	, MainGraphTrack(nullptr)
-	, LiveAllocsGraphTrack(nullptr)
-	, AllocFreeGraphTrack(nullptr)
-	, SwapMemoryGraphTrack(nullptr)
-	, PageSwapGraphTrack(nullptr)
-	, TrackHeightMode(EMemoryTrackHeightMode::Medium)
-	, bShowHideAllMemoryTracks(false)
-	, CreatedDefaultTracks()
-	, CurrentMemoryRule(nullptr)
 {
 	InitMemoryRules();
 
@@ -123,6 +111,8 @@ void FMemorySharedState::OnBeginSession(Timing::ITimingViewSession& InSession)
 	PlatformTracker = nullptr;
 
 	TagSets.Reset();
+	NumSyncedTagSets = 0;
+	NumValidTagSets = 0;
 
 	MainGraphTrack = nullptr;
 	LiveAllocsGraphTrack = nullptr;
@@ -152,6 +142,8 @@ void FMemorySharedState::OnEndSession(Timing::ITimingViewSession& InSession)
 	PlatformTracker = nullptr;
 
 	TagSets.Reset();
+	NumSyncedTagSets = 0;
+	NumValidTagSets = 0;
 
 	MainGraphTrack = nullptr;
 	LiveAllocsGraphTrack = nullptr;
@@ -460,8 +452,10 @@ void FMemorySharedState::UpdateTagSets()
 		const TraceServices::IMemoryProvider* MemoryProvider = TraceServices::ReadMemoryProvider(*Session.Get());
 		if (MemoryProvider)
 		{
-			if ((int32)MemoryProvider->GetTagSetCount() != TagSets.Num())
+			const uint32 NumProviderTagSets = MemoryProvider->GetTagSetCount();
+			if (NumProviderTagSets != NumSyncedTagSets)
 			{
+				NumSyncedTagSets = NumProviderTagSets;
 				MemoryProvider->EnumerateTagSets([this](const TraceServices::FMemoryTagSetInfo& TagSet)
 				{
 					int32 Index = (int32)TagSet.Id;
@@ -471,6 +465,10 @@ void FMemorySharedState::UpdateTagSets()
 						{
 							TagSets.AddDefaulted(Index - TagSets.Num() + 1);
 						}
+						if (!TagSets[Index].IsValid())
+						{
+							++NumValidTagSets;
+						}
 						TagSets[Index] = MakeShared<FMemoryTagSet>(static_cast<FMemoryTagSetId>(TagSet.Id), TagSet.Name);
 					}
 				});
@@ -480,6 +478,21 @@ void FMemorySharedState::UpdateTagSets()
 	else
 	{
 		TagSets.Reset();
+		NumSyncedTagSets = 0;
+		NumValidTagSets = 0;
+	}
+}
+
+////////////////////////////////////////////////////////////////////////////////////////////////////
+
+void FMemorySharedState::EnumerateTagSets(TFunctionRef<void(const FMemoryTagSet& InTagSet)> Callback) const
+{
+	for (const TSharedPtr<FMemoryTagSet>& TagSet : TagSets)
+	{
+		if (TagSet.IsValid())
+		{
+			Callback(*TagSet);
+		}
 	}
 }
 
diff --git a/Engine/Source/Developer/TraceInsights/Private/Insights/MemoryProfiler/ViewModels/MemorySharedState.h b/Engine/Source/Developer/TraceInsights/Private/Insights/MemoryProfiler/ViewModels/MemorySharedState.h
--- a/Engine/Source/Developer/TraceInsights/Private/Insights/MemoryProfiler/ViewModels/MemorySharedState.h
+++ b/Engine/Source/Developer/TraceInsights/Private/Insights/MemoryProfiler/ViewModels/MemorySharedState.h
@@ -117,7 +117,8 @@ public:
 		return PlatformTracker.IsValid() ? PlatformTracker->GetId() : FMemoryTracker::InvalidTrackerId;
 	}
 
-	const TArray<TSharedPtr<FMemoryTagSet>>& GetTagSets() const { return TagSets; }
+	uint32 GetNumTagSets() const { return NumValidTagSets; }
+	void EnumerateTagSets(TFunctionRef<void(const FMemoryTagSet& InTagSet)> Callback) const;
 	const FMemoryTagSet* GetTagSetById(FMemoryTagSetId InMemTagSetId) const;
 
 	TSharedPtr<FMemoryGraphTrack> GetMainGraphTrack() const { return MainGraphTrack; }
@@ -190,6 +191,8 @@ private:
 	TSharedPtr<FMemoryTracker> PlatformTracker;
 
 	TArray<TSharedPtr<FMemoryTagSet>> TagSets;
+	uint32 NumSyncedTagSets = 0;
+	uint32 NumValidTagSets = 0;
 
 	TSharedPtr<FMemoryGraphTrack> MainGraphTrack; // the Main Memory Graph track; also hosts the Total Allocated Memory series
 	TSharedPtr<FMemoryGraphTrack> LiveAllocsGraphTrack; // the graph track for the Live Allocation Count series
@@ -198,9 +201,9 @@ private:
 	TSharedPtr<FMemoryGraphTrack> PageSwapGraphTrack; // the graph track for the Page In Event Count and the Page Out Event Count series
 	TSet<TSharedPtr<FMemoryGraphTrack>> AllTracks;
 
-	EMemoryTrackHeightMode TrackHeightMode;
+	EMemoryTrackHeightMode TrackHeightMode = EMemoryTrackHeightMode::Medium;
 
-	bool bShowHideAllMemoryTracks;
+	bool bShowHideAllMemoryTracks = false;
 
 	TBitArray<> CreatedDefaultTracks;
 
diff --git a/Engine/Source/Developer/TraceInsights/Private/Insights/MemoryProfiler/Widgets/SMemTagTreeView.cpp b/Engine/Source/Developer/TraceInsights/Private/Insights/MemoryProfiler/Widgets/SMemTagTreeView.cpp
--- a/Engine/Source/Developer/TraceInsights/Private/Insights/MemoryProfiler/Widgets/SMemTagTreeView.cpp
+++ b/Engine/Source/Developer/TraceInsights/Private/Insights/MemoryProfiler/Widgets/SMemTagTreeView.cpp
@@ -994,14 +994,13 @@ void SMemTagTreeView::RebuildTree(bool bResync)
 			ApplyBudgetToNodes();
 		}
 
-		const auto& TagSets = SharedState.GetTagSets();
-		if (AvailableTagSets.Num() != TagSets.Num())
+		if (AvailableTagSets.Num() != SharedState.GetNumTagSets())
 		{
 			AvailableTagSets.Reset();
-			for (const auto& TagSet : TagSets)
-			{
-				AvailableTagSets.Add(MakeShared<FMemoryTagSetId>(TagSet->GetId()));
-			}
+			SharedState.EnumerateTagSets([this](const FMemoryTagSet& InTagSet)
+				{
+					AvailableTagSets.Add(MakeShared<FMemoryTagSetId>(InTagSet.GetId()));
+				});
 		}
 	}
 
