Add Gpu Fences menu options.
Make Gpu Track settings persistent in Timing Profiler.

diff --git a/Engine/Source/Developer/TraceInsights/Private/Insights/InsightsSettings.cpp b/Engine/Source/Developer/TraceInsights/Private/Insights/InsightsSettings.cpp
--- a/Engine/Source/Developer/TraceInsights/Private/Insights/InsightsSettings.cpp
+++ b/Engine/Source/Developer/TraceInsights/Private/Insights/InsightsSettings.cpp
@@ -99,6 +99,17 @@ void FInsightsSettings::LoadFromConfig()
 	GConfig->GetDouble(TEXT("Insights.TimingProfiler"), TEXT("AutoScrollViewportOffsetPercent"), AutoScrollViewportOffsetPercent, SettingsIni);
 	GConfig->GetDouble(TEXT("Insights.TimingProfiler"), TEXT("AutoScrollMinDelay"), AutoScrollMinDelay, SettingsIni);
 
+	//////////////////////////////////////////////////
+	// [Insights.TimingProfiler.TimingView]
+
+	GConfig->GetBool(TEXT("Insights.TimingProfiler.TimingView"), TEXT("ShowGpuWorkTracks"), bTimingViewShowGpuWorkTracks, SettingsIni);
+	GConfig->GetBool(TEXT("Insights.TimingProfiler.TimingView"), TEXT("ShowGpuWorkOverlays"), bTimingViewShowGpuWorkOverlays, SettingsIni);
+	GConfig->GetBool(TEXT("Insights.TimingProfiler.TimingView"), TEXT("ShowGpuWorkExtendedLines"), bTimingViewShowGpuWorkExtendedLines, SettingsIni);
+
+	GConfig->GetBool(TEXT("Insights.TimingProfiler.TimingView"), TEXT("ShowGpuFencesTracks"), bTimingViewShowGpuFencesTracks, SettingsIni);
+	GConfig->GetBool(TEXT("Insights.TimingProfiler.TimingView"), TEXT("ShowGpuFencesExtendedLines"), bTimingViewShowGpuFencesExtendedLines, SettingsIni);
+	GConfig->GetBool(TEXT("Insights.TimingProfiler.TimingView"), TEXT("ShowGpuFencesRelations"), bTimingViewShowGpuFencesRelations, SettingsIni);
+
 	//////////////////////////////////////////////////
 	// [Insights.TimingProfiler.FramesView]
 
@@ -202,6 +213,17 @@ void FInsightsSettings::SaveToConfig()
 	GConfig->SetDouble(TEXT("Insights.TimingProfiler"), TEXT("AutoScrollViewportOffsetPercent"), AutoScrollViewportOffsetPercent, SettingsIni);
 	GConfig->SetDouble(TEXT("Insights.TimingProfiler"), TEXT("AutoScrollMinDelay"), AutoScrollMinDelay, SettingsIni);
 
+	//////////////////////////////////////////////////
+	// [Insights.TimingProfiler.TimingView]
+
+	GConfig->SetBool(TEXT("Insights.TimingProfiler.TimingView"), TEXT("ShowGpuWorkTracks"), bTimingViewShowGpuWorkTracks, SettingsIni);
+	GConfig->SetBool(TEXT("Insights.TimingProfiler.TimingView"), TEXT("ShowGpuWorkOverlays"), bTimingViewShowGpuWorkOverlays, SettingsIni);
+	GConfig->SetBool(TEXT("Insights.TimingProfiler.TimingView"), TEXT("ShowGpuWorkExtendedLines"), bTimingViewShowGpuWorkExtendedLines, SettingsIni);
+
+	GConfig->SetBool(TEXT("Insights.TimingProfiler.TimingView"), TEXT("ShowGpuFencesTracks"), bTimingViewShowGpuFencesTracks, SettingsIni);
+	GConfig->SetBool(TEXT("Insights.TimingProfiler.TimingView"), TEXT("ShowGpuFencesExtendedLines"), bTimingViewShowGpuFencesExtendedLines, SettingsIni);
+	GConfig->SetBool(TEXT("Insights.TimingProfiler.TimingView"), TEXT("ShowGpuFencesRelations"), bTimingViewShowGpuFencesRelations, SettingsIni);
+
 	//////////////////////////////////////////////////
 	// [Insights.TimingProfiler.FramesView]
 
diff --git a/Engine/Source/Developer/TraceInsights/Private/Insights/InsightsSettings.h b/Engine/Source/Developer/TraceInsights/Private/Insights/InsightsSettings.h
--- a/Engine/Source/Developer/TraceInsights/Private/Insights/InsightsSettings.h
+++ b/Engine/Source/Developer/TraceInsights/Private/Insights/InsightsSettings.h
@@ -72,6 +72,33 @@ public:
 	void SetAutoScrollMinDelay(double Delay) { AutoScrollMinDelay = Delay; }
 	void SetAndSaveAutoScrollMinDelay(double Delay) { SET_AND_SAVE(AutoScrollMinDelay, Delay); }
 
+	//////////////////////////////////////////////////
+	// [Insights.TimingProfiler.TimingView]
+
+	bool GetTimingViewShowGpuWorkTracks() const { return bTimingViewShowGpuWorkTracks; }
+	void SetTimingViewShowGpuWorkTracks(bool InValue) { bTimingViewShowGpuWorkTracks = InValue; }
+	void SetAndSaveTimingViewShowGpuWorkTracks(bool InValue) { SET_AND_SAVE(bTimingViewShowGpuWorkTracks, InValue); }
+
+	bool GetTimingViewShowGpuWorkOverlays() const { return bTimingViewShowGpuWorkOverlays; }
+	void SetTimingViewShowGpuWorkOverlays(bool InValue) { bTimingViewShowGpuWorkOverlays = InValue; }
+	void SetAndSaveTimingViewShowGpuWorkOverlays(bool InValue) { SET_AND_SAVE(bTimingViewShowGpuWorkOverlays, InValue); }
+
+	bool GetTimingViewShowGpuWorkExtendedLines() const { return bTimingViewShowGpuWorkExtendedLines; }
+	void SetTimingViewShowGpuWorkExtendedLines(bool InValue) { bTimingViewShowGpuWorkExtendedLines = InValue; }
+	void SetAndSaveTimingViewShowGpuWorkExtendedLines(bool InValue) { SET_AND_SAVE(bTimingViewShowGpuWorkExtendedLines, InValue); }
+
+	bool GetTimingViewShowGpuFencesTracks() const { return bTimingViewShowGpuFencesTracks; }
+	void SetTimingViewShowGpuFencesTracks(bool InValue) { bTimingViewShowGpuFencesTracks = InValue; }
+	void SetAndSaveTimingViewShowGpuFencesTracks(bool InValue) { SET_AND_SAVE(bTimingViewShowGpuFencesTracks, InValue); }
+
+	bool GetTimingViewShowGpuFencesExtendedLines() const { return bTimingViewShowGpuFencesExtendedLines; }
+	void SetTimingViewShowGpuFencesExtendedLines(bool InValue) { bTimingViewShowGpuFencesExtendedLines = InValue; }
+	void SetAndSaveTimingViewShowGpuFencesExtendedLines(bool InValue) { SET_AND_SAVE(bTimingViewShowGpuFencesExtendedLines, InValue); }
+
+	bool GetTimingViewShowGpuFencesRelations() const { return bTimingViewShowGpuFencesRelations; }
+	void SetTimingViewShowGpuFencesRelations(bool InValue) { bTimingViewShowGpuFencesRelations = InValue; }
+	void SetAndSaveTimingViewShowGpuFencesRelations(bool InValue) { SET_AND_SAVE(bTimingViewShowGpuFencesRelations, InValue); }
+
 	//////////////////////////////////////////////////
 	// [Insights.TimingProfiler.FramesView]
 
@@ -254,6 +281,27 @@ private:
 	/** Minimum time between two auto-scroll updates, in [seconds]. */
 	double AutoScrollMinDelay = 0.3; // [seconds]
 
+	//////////////////////////////////////////////////
+	// [Insights.TimingProfiler.TimingView]
+
+	/** Toggles visibility for GPU work header tracks. */
+	bool bTimingViewShowGpuWorkTracks = true;
+
+	/** Extends the visualization of GPU work events over the GPU timing tracks. */
+	bool bTimingViewShowGpuWorkOverlays = true;
+
+	/** Shows/hides the extended vertical lines at the edges of each GPU work event. */
+	bool bTimingViewShowGpuWorkExtendedLines = true;
+
+	/** Shows/hides the gpu fences child track. */
+	bool bTimingViewShowGpuFencesTracks = true;
+
+	/** Shows/hides the extended vertical lines at the location of gpu fences. */
+	bool bTimingViewShowGpuFencesExtendedLines = true;
+
+	/** If enabled, relations between Signal and Wait fences will be displayed when selecting a Timing Event in a Gpu Queue Track. */
+	bool bTimingViewShowGpuFencesRelations = true;
+
 	//////////////////////////////////////////////////
 	// [Insights.TimingProfiler.FramesView]
 
diff --git a/Engine/Source/Developer/TraceInsights/Private/Insights/TimingProfiler/Tracks/GpuTimingTrack.cpp b/Engine/Source/Developer/TraceInsights/Private/Insights/TimingProfiler/Tracks/GpuTimingTrack.cpp
--- a/Engine/Source/Developer/TraceInsights/Private/Insights/TimingProfiler/Tracks/GpuTimingTrack.cpp
+++ b/Engine/Source/Developer/TraceInsights/Private/Insights/TimingProfiler/Tracks/GpuTimingTrack.cpp
@@ -314,8 +314,8 @@ void FGpuFencesTrackBuilder::AddTimeMarker(const float X, TraceServices::EGpuFen
 // FGpuFencesTimingTrack
 ////////////////////////////////////////////////////////////////////////////////////////////////////
 
-FGpuFencesTimingTrack::FGpuFencesTimingTrack(FThreadTimingSharedState& InSharedState, const FString& InName, uint32 InTimelineIndex, uint32 InQueueId)
-	: FThreadTimingTrackImpl(InSharedState, InName, nullptr, InTimelineIndex, InQueueId)
+FGpuFencesTimingTrack::FGpuFencesTimingTrack(FThreadTimingSharedState& InSharedState, const FString& InName, uint32 InQueueId)
+	: FThreadTimingTrackImpl(InSharedState, InName, nullptr, 0, InQueueId)
 	, WhiteBrush(FAppStyle::Get().GetBrush("WhiteBrush"))
 	, Font(FAppStyle::Get().GetFontStyle("SmallFont"))
 {
@@ -429,7 +429,9 @@ void FGpuFencesTimingTrack::PostDraw(const ITimingTrackDrawContext& Context) con
 
 	DrawContext.LayerId++;
 
-	bool bExtentFenceLines = Viewport.GetDuration() < 3.0f;
+	constexpr double MaxViewportSizeForExtendedLines = 3.0;
+	const bool bExtentFenceLines = GetSharedState().AreGpuFencesExtendedLinesVisible() && Viewport.GetDuration() < MaxViewportSizeForExtendedLines;
+
 	float BoxY1 = GetPosY();
 	float BoxY2 = GetPosY() + GetHeight();
 	ETimingTrackLocation LocalLocation = GetLocation();
diff --git a/Engine/Source/Developer/TraceInsights/Private/Insights/TimingProfiler/Tracks/GpuTimingTrack.h b/Engine/Source/Developer/TraceInsights/Private/Insights/TimingProfiler/Tracks/GpuTimingTrack.h
--- a/Engine/Source/Developer/TraceInsights/Private/Insights/TimingProfiler/Tracks/GpuTimingTrack.h
+++ b/Engine/Source/Developer/TraceInsights/Private/Insights/TimingProfiler/Tracks/GpuTimingTrack.h
@@ -91,7 +91,7 @@ class FGpuFencesTimingTrack : public FThreadTimingTrackImpl
 	INSIGHTS_DECLARE_RTTI(FGpuFencesTimingTrack, FThreadTimingTrackImpl)
 
 public:
-	explicit FGpuFencesTimingTrack(FThreadTimingSharedState& InSharedState, const FString& InName, uint32 InTimelineIndex, uint32 InQueueId);
+	explicit FGpuFencesTimingTrack(FThreadTimingSharedState& InSharedState, const FString& InName, uint32 InQueueId);
 
 	uint32 GetQueueId() const { return GetThreadId(); }
 
diff --git a/Engine/Source/Developer/TraceInsights/Private/Insights/TimingProfiler/ViewModels/ThreadTimingSharedState.cpp b/Engine/Source/Developer/TraceInsights/Private/Insights/TimingProfiler/ViewModels/ThreadTimingSharedState.cpp
--- a/Engine/Source/Developer/TraceInsights/Private/Insights/TimingProfiler/ViewModels/ThreadTimingSharedState.cpp
+++ b/Engine/Source/Developer/TraceInsights/Private/Insights/TimingProfiler/ViewModels/ThreadTimingSharedState.cpp
@@ -17,7 +17,6 @@
 #include "InsightsCore/Common/Log.h"
 
 // TraceInsights
-#include "Insights/InsightsManager.h"
 #include "Insights/InsightsStyle.h"
 #include "Insights/IUnrealInsightsModule.h"
 #include "Insights/ITimingViewSession.h"
@@ -50,10 +49,19 @@ public:
 	TSharedPtr<FUICommandInfo> Command_ShowWorkTracks;
 
 	/** Extends the visualization of GPU work events over the GPU timing tracks. */
-	TSharedPtr<FUICommandInfo> Command_ShowOverlays;
+	TSharedPtr<FUICommandInfo> Command_ShowGpuWorkOverlays;
 
 	/** Shows/hides the extended vertical lines at the edges of each GPU work event. */
-	TSharedPtr<FUICommandInfo> Command_ShowExtendedLines;
+	TSharedPtr<FUICommandInfo> Command_ShowGpuWorkExtendedLines;
+
+	/** If enabled, relations between Signal and Wait fences will be displayed when selecting a Timing Event in a Gpu Queue Track. */
+	TSharedPtr<FUICommandInfo> Command_ShowGpuFenceRelations;
+
+	/** Shows/hides the gpu fences child track. */
+	TSharedPtr<FUICommandInfo> Command_ShowGpuFencesTrack;
+
+	/** Shows/hides the extended vertical lines at the location of gpu fences. */
+	TSharedPtr<FUICommandInfo> Command_ShowGpuFencesExtendedLines;
 
 	/** Toggles visibility for all CPU thread tracks at once. */
 	TSharedPtr<FUICommandInfo> ShowHideAllCpuTracks;
@@ -88,18 +96,36 @@ void FThreadTimingViewCommands::RegisterCommands()
 		EUserInterfaceActionType::ToggleButton,
 		FInputChord());
 
-	UI_COMMAND(Command_ShowOverlays,
-		"Show Overlays",
+	UI_COMMAND(Command_ShowGpuWorkOverlays,
+		"Show GPU Work Overlays",
 		"Extends the visualization of GPU work events over the GPU timing tracks.",
 		EUserInterfaceActionType::ToggleButton,
 		FInputChord());
 
-	UI_COMMAND(Command_ShowExtendedLines,
-		"Show Extended Lines",
+	UI_COMMAND(Command_ShowGpuWorkExtendedLines,
+		"Show GPU Work Extended Lines",
 		"Shows/hides the extended vertical lines at the edges of each GPU work event.",
 		EUserInterfaceActionType::ToggleButton,
 		FInputChord());
 
+	UI_COMMAND(Command_ShowGpuFenceRelations,
+		"Show GPU Fences Relations",
+		"If enabled, relations between signal and wait fences will be displayed when selecting a timing event in a GPU timing track.",
+		EUserInterfaceActionType::ToggleButton,
+		FInputChord());
+
+	UI_COMMAND(Command_ShowGpuFencesTrack,
+		"Show GPU Fences Track(s)",
+		"Shows/hides the GPU fences header tracks.",
+		EUserInterfaceActionType::ToggleButton,
+		FInputChord());
+
+	UI_COMMAND(Command_ShowGpuFencesExtendedLines,
+		"Show GPU Fences Extended Lines",
+		"Shows/hides the extended vertical lines at the location of GPU fences.",
+		EUserInterfaceActionType::ToggleButton,
+		FInputChord());
+
 	UI_COMMAND(ShowHideAllCpuTracks,
 		"CPU Thread Tracks",
 		"Shows/hides all CPU tracks (and all CPU thread groups).",
@@ -114,6 +140,7 @@ UE_ENABLE_OPTIMIZATION_SHIP
 
 FThreadTimingSharedState::FThreadTimingSharedState(STimingView* InTimingView)
 	: TimingView(InTimingView)
+	, Settings(MakeShared<FThreadSharedStateLocalSettings>())
 {
 	check(TimingView != nullptr);
 }
@@ -255,21 +282,17 @@ void FThreadTimingSharedState::OnBeginSession(Timing::ITimingViewSession& InSess
 		return;
 	}
 
-	if (TimingView->GetName() == FInsightsManagerTabs::LoadingProfilerTabId)
+	if (TimingView->GetName() == FInsightsManagerTabs::TimingProfilerTabId)
 	{
-		bShowHideAllGpuTracks = false;
-		bShowGpuWorkTracks = false;
-		bShowOverlaysInGpuQueueTracks = false;
-		bShowExtendedLinesInGpuQueueTracks = false;
-		bShowHideAllCpuTracks = false;
+		bShowHideAllGpuTracks = true;
+		bShowHideAllCpuTracks = true;
+
+		Settings = MakeShared<FThreadSharedStatePersistentSettings>();
 	}
 	else
 	{
-		bShowHideAllGpuTracks = true;
-		bShowGpuWorkTracks = true;
-		bShowOverlaysInGpuQueueTracks = true;
-		bShowExtendedLinesInGpuQueueTracks = true;
-		bShowHideAllCpuTracks = true;
+		bShowHideAllGpuTracks = false;
+		bShowHideAllCpuTracks = false;
 	}
 
 	OldGpu1Track = nullptr;
@@ -378,16 +401,19 @@ void FThreadTimingSharedState::Tick(Timing::ITimingViewSession& InSession, const
 					GpuTracks.Add(QueueInfo.Id, Track);
 					InSession.AddScrollableTrack(Track);
 
-					if (bShowGpuWorkTracks)
+					if (AreGpuWorkTracksVisible())
 					{
 						// Create the GPU Work track and attach it to the GPU queue track.
 						const FString WorkTrackName = FString::Printf(TEXT("GPU%u - %s %u - WORK"), QueueInfo.GPU, QueueInfo.Name, QueueInfo.Index);
 						TSharedRef<FGpuQueueWorkTimingTrack> WorkTrack = MakeShared<FGpuQueueWorkTimingTrack>(*this, WorkTrackName, QueueInfo.WorkTimelineIndex, QueueInfo.Id);
 						WorkTrack->SetParentTrack(Track);
 						Track->AddChildTrack(WorkTrack);
+					}
 
+					if (AreGpuFencesTracksVisible())
+					{
 						const FString FencesTrackName = FString::Printf(TEXT("GPU%u - %s %u - Fences"), QueueInfo.GPU, QueueInfo.Name, QueueInfo.Index);
-						TSharedRef<FGpuFencesTimingTrack> FencesTrack = MakeShared<FGpuFencesTimingTrack>(*this, WorkTrackName, QueueInfo.WorkTimelineIndex, QueueInfo.Id);
+						TSharedRef<FGpuFencesTimingTrack> FencesTrack = MakeShared<FGpuFencesTimingTrack>(*this, FencesTrackName, QueueInfo.Id);
 						FencesTrack->SetParentTrack(Track);
 						Track->AddChildTrack(FencesTrack);
 					}
@@ -506,8 +532,11 @@ void FThreadTimingSharedState::ExtendGpuTracksFilterMenu(Timing::ITimingViewSess
 	{
 		InOutMenuBuilder.AddMenuEntry(FThreadTimingViewCommands::Get().ShowHideAllGpuTracks);
 		InOutMenuBuilder.AddMenuEntry(FThreadTimingViewCommands::Get().Command_ShowWorkTracks);
-		InOutMenuBuilder.AddMenuEntry(FThreadTimingViewCommands::Get().Command_ShowOverlays);
-		InOutMenuBuilder.AddMenuEntry(FThreadTimingViewCommands::Get().Command_ShowExtendedLines);
+		InOutMenuBuilder.AddMenuEntry(FThreadTimingViewCommands::Get().Command_ShowGpuWorkOverlays);
+		InOutMenuBuilder.AddMenuEntry(FThreadTimingViewCommands::Get().Command_ShowGpuWorkExtendedLines);
+		InOutMenuBuilder.AddMenuEntry(FThreadTimingViewCommands::Get().Command_ShowGpuFenceRelations);
+		InOutMenuBuilder.AddMenuEntry(FThreadTimingViewCommands::Get().Command_ShowGpuFencesTrack);
+		InOutMenuBuilder.AddMenuEntry(FThreadTimingViewCommands::Get().Command_ShowGpuFencesExtendedLines);
 	}
 	InOutMenuBuilder.EndSection();
 }
@@ -554,16 +583,34 @@ void FThreadTimingSharedState::BindCommands()
 		FIsActionChecked::CreateSP(this, &FThreadTimingSharedState::Command_ShowGpuWorkTracks_IsChecked));
 
 	CommandList->MapAction(
-		FThreadTimingViewCommands::Get().Command_ShowOverlays,
-		FExecuteAction::CreateSP(this, &FThreadTimingSharedState::Command_ShowOverlays_Execute),
-		FCanExecuteAction::CreateSP(this, &FThreadTimingSharedState::Command_ShowOverlays_CanExecute),
-		FIsActionChecked::CreateSP(this, &FThreadTimingSharedState::Command_ShowOverlays_IsChecked));
+		FThreadTimingViewCommands::Get().Command_ShowGpuWorkOverlays,
+		FExecuteAction::CreateSP(this, &FThreadTimingSharedState::Command_ShowGpuWorkOverlays_Execute),
+		FCanExecuteAction::CreateSP(this, &FThreadTimingSharedState::Command_ShowGpuWorkOverlays_CanExecute),
+		FIsActionChecked::CreateSP(this, &FThreadTimingSharedState::Command_ShowGpuWorkOverlays_IsChecked));
+
+	CommandList->MapAction(
+		FThreadTimingViewCommands::Get().Command_ShowGpuWorkExtendedLines,
+		FExecuteAction::CreateSP(this, &FThreadTimingSharedState::Command_ShowGpuWorkExtendedLines_Execute),
+		FCanExecuteAction::CreateSP(this, &FThreadTimingSharedState::Command_ShowGpuWorkExtendedLines_CanExecute),
+		FIsActionChecked::CreateSP(this, &FThreadTimingSharedState::Command_ShowGpuWorkExtendedLines_IsChecked));
+
+	CommandList->MapAction(
+		FThreadTimingViewCommands::Get().Command_ShowGpuFencesTrack,
+		FExecuteAction::CreateSP(this, &FThreadTimingSharedState::Command_ShowGpuFencesTracks_Execute),
+		FCanExecuteAction::CreateSP(this, &FThreadTimingSharedState::Command_ShowGpuFencesTracks_CanExecute),
+		FIsActionChecked::CreateSP(this, &FThreadTimingSharedState::Command_ShowGpuFencesTracks_IsChecked));
 
 	CommandList->MapAction(
-		FThreadTimingViewCommands::Get().Command_ShowExtendedLines,
-		FExecuteAction::CreateSP(this, &FThreadTimingSharedState::Command_ShowExtendedLines_Execute),
-		FCanExecuteAction::CreateSP(this, &FThreadTimingSharedState::Command_ShowExtendedLines_CanExecute),
-		FIsActionChecked::CreateSP(this, &FThreadTimingSharedState::Command_ShowExtendedLines_IsChecked));
+		FThreadTimingViewCommands::Get().Command_ShowGpuFencesExtendedLines,
+		FExecuteAction::CreateSP(this, &FThreadTimingSharedState::Command_ShowGpuFencesExtendedLines_Execute),
+		FCanExecuteAction::CreateSP(this, &FThreadTimingSharedState::Command_ShowGpuFencesExtendedLines_CanExecute),
+		FIsActionChecked::CreateSP(this, &FThreadTimingSharedState::Command_ShowGpuFencesExtendedLines_IsChecked));
+
+	CommandList->MapAction(
+		FThreadTimingViewCommands::Get().Command_ShowGpuFenceRelations,
+		FExecuteAction::CreateSP(this, &FThreadTimingSharedState::Command_ShowGpuFencesRelations_Execute),
+		FCanExecuteAction::CreateSP(this, &FThreadTimingSharedState::Command_ShowGpuFencesRelations_CanExecute),
+		FIsActionChecked::CreateSP(this, &FThreadTimingSharedState::Command_ShowGpuFencesRelations_IsChecked));
 
 	CommandList->MapAction(
 		FThreadTimingViewCommands::Get().ShowHideAllCpuTracks,
@@ -796,11 +843,11 @@ TSharedPtr<const ITimingEvent> FThreadTimingSharedState::FindMinEventInstance(ui
 
 void FThreadTimingSharedState::SetGpuWorkTracksVisibility(bool bOnOff)
 {
-	if (bShowGpuWorkTracks != bOnOff)
+	if (Settings->GetTimingViewShowGpuWorkTracks() != bOnOff)
 	{
-		bShowGpuWorkTracks = bOnOff;
+		Settings->SetTimingViewShowGpuWorkTracks(bOnOff);
 
-		if (bShowGpuWorkTracks)
+		if (bOnOff)
 		{
 			AddGpuWorkChildTracks();
 		}
@@ -835,7 +882,7 @@ void FThreadTimingSharedState::AddGpuWorkChildTracks()
 			// Create the GPU Work track and attach it to the GPU queue track.
 			const FString WorkTrackName = FString::Printf(TEXT("GPU%u - %s %u - WORK"), QueueInfo.GPU, QueueInfo.Name, QueueInfo.Index);
 			TSharedRef<FGpuQueueWorkTimingTrack> WorkTrack = MakeShared<FGpuQueueWorkTimingTrack>(*this, WorkTrackName, QueueInfo.WorkTimelineIndex, QueueInfo.Id);
-			(*TrackPtrPtr)->AddChildTrack(WorkTrack);
+			(*TrackPtrPtr)->AddChildTrack(WorkTrack, 0);
 			WorkTrack->SetParentTrack(*TrackPtrPtr);
 		}
 	});
@@ -857,8 +904,93 @@ void FThreadTimingSharedState::RemoveGpuWorkChildTracks()
 
 ////////////////////////////////////////////////////////////////////////////////////////////////////
 
+
+void FThreadTimingSharedState::AddGpuFencesChildTracks()
+{
+	TSharedPtr<const TraceServices::IAnalysisSession> Session = FInsightsManager::Get()->GetSession();
+	if (!Session.IsValid())
+	{
+		return;
+	}
+
+	const TraceServices::ITimingProfilerProvider* TimingProfilerProvider = TraceServices::ReadTimingProfilerProvider(*Session);
+
+	TraceServices::FAnalysisSessionReadScope SessionReadScope(*Session);
+
+	TimingProfilerProvider->EnumerateGpuQueues([this](const TraceServices::FGpuQueueInfo& QueueInfo)
+	{
+		// Check if there is an available GPU track for this queue.
+		TSharedPtr<FGpuQueueTimingTrack>* TrackPtrPtr = GpuTracks.Find(QueueInfo.Id);
+
+		if (TrackPtrPtr != nullptr)
+		{
+			// Create the GPU Work track and attach it to the GPU queue track.
+			const FString TrackName = FString::Printf(TEXT("GPU%u - %s %u - Fences"), QueueInfo.GPU, QueueInfo.Name, QueueInfo.Index);
+			TSharedRef<FGpuFencesTimingTrack> Track = MakeShared<FGpuFencesTimingTrack>(*this, TrackName, QueueInfo.Id);
+			(*TrackPtrPtr)->AddChildTrack(Track, (*TrackPtrPtr)->GetChildTracks().Num());
+			Track->SetParentTrack(*TrackPtrPtr);
+		}
+	});
+}
+
+////////////////////////////////////////////////////////////////////////////////////////////////////
+
+void FThreadTimingSharedState::RemoveGpuFencesChildTracks()
+{
+	for (auto& Pair : GpuTracks)
+	{
+		const TSharedPtr<FGpuQueueTimingTrack>& Track = Pair.Value;
+		if (TSharedPtr<FGpuFencesTimingTrack> FencesTrack = Track->FindChildTrackOfType<FGpuFencesTimingTrack>())
+		{
+			Track->RemoveChildTrack(FencesTrack.ToSharedRef());
+		}
+	}
+}
+
+////////////////////////////////////////////////////////////////////////////////////////////////////
+
+void FThreadTimingSharedState::SetGpuFencesTracksVisibility(bool bOnOff)
+{
+	if (Settings->GetTimingViewShowGpuFencesTracks() != bOnOff)
+	{
+		Settings->SetTimingViewShowGpuFencesTracks(bOnOff);
+
+		if (bOnOff)
+		{
+			AddGpuFencesChildTracks();
+		}
+		else
+		{
+			RemoveGpuFencesChildTracks();
+		}
+	}
+}
+
+////////////////////////////////////////////////////////////////////////////////////////////////////
+
+void FThreadTimingSharedState::Command_ShowGpuFencesRelations_Execute()
+{
+	bool NewValue = !Settings->GetTimingViewShowGpuFencesRelations();
+	Settings->SetTimingViewShowGpuFencesRelations(NewValue);
+
+	if (NewValue == false && TimingView)
+	{
+		TimingView->EditCurrentRelations().RemoveAll([](TUniquePtr<ITimingEventRelation>& Item)
+			{
+				return Item->Is<FGpuFenceRelation>();
+			});
+	}
+}
+
+////////////////////////////////////////////////////////////////////////////////////////////////////
+
 void FThreadTimingSharedState::OnTimingEventSelected(TSharedPtr<const ITimingEvent> InSelectedEvent)
 {
+	if (!AreGpuFenceRelationsVisible())
+	{
+		return;
+	}
+
 	if (TimingView)
 	{
 		TimingView->EditCurrentRelations().RemoveAll([](TUniquePtr<ITimingEventRelation>& Item)
diff --git a/Engine/Source/Developer/TraceInsights/Private/Insights/TimingProfiler/ViewModels/ThreadTimingSharedState.h b/Engine/Source/Developer/TraceInsights/Private/Insights/TimingProfiler/ViewModels/ThreadTimingSharedState.h
--- a/Engine/Source/Developer/TraceInsights/Private/Insights/TimingProfiler/ViewModels/ThreadTimingSharedState.h
+++ b/Engine/Source/Developer/TraceInsights/Private/Insights/TimingProfiler/ViewModels/ThreadTimingSharedState.h
@@ -10,6 +10,7 @@
 #include "Templates/SharedPointer.h"
 
 // TraceInsights
+#include "Insights/InsightsManager.h"
 #include "Insights/ITimingViewExtender.h"
 
 namespace TraceServices
@@ -41,6 +42,86 @@ private:
 		int32 GetOrder() const { return Order; }
 	};
 
+	class IThreadSharedStateSetting
+	{
+	public:
+		virtual ~IThreadSharedStateSetting() {}
+
+		virtual bool GetTimingViewShowGpuWorkTracks() const = 0;
+		virtual void SetTimingViewShowGpuWorkTracks(bool InValue) = 0;
+
+		virtual bool GetTimingViewShowGpuWorkOverlays() const = 0;
+		virtual void SetTimingViewShowGpuWorkOverlays(bool InValue) = 0;
+
+		virtual bool GetTimingViewShowGpuWorkExtendedLines() const = 0;
+		virtual void SetTimingViewShowGpuWorkExtendedLines(bool InValue) = 0;
+
+		virtual bool GetTimingViewShowGpuFencesTracks() const = 0;
+		virtual void SetTimingViewShowGpuFencesTracks(bool InValue) = 0;
+
+		virtual bool GetTimingViewShowGpuFencesExtendedLines() const = 0;
+		virtual void SetTimingViewShowGpuFencesExtendedLines(bool InValue) = 0;
+
+		virtual bool GetTimingViewShowGpuFencesRelations() const = 0;
+		virtual void SetTimingViewShowGpuFencesRelations(bool InValue) = 0;
+	};
+
+	class FThreadSharedStatePersistentSettings : public IThreadSharedStateSetting
+	{
+	public:
+		virtual bool GetTimingViewShowGpuWorkTracks() const override { return GetInsightsSettings().GetTimingViewShowGpuWorkTracks(); }
+		virtual void SetTimingViewShowGpuWorkTracks(bool InValue) override { GetInsightsSettings().SetAndSaveTimingViewShowGpuWorkTracks(InValue); }
+
+		virtual bool GetTimingViewShowGpuWorkOverlays() const override { return GetInsightsSettings().GetTimingViewShowGpuWorkOverlays(); }
+		virtual void SetTimingViewShowGpuWorkOverlays(bool InValue) override { GetInsightsSettings().SetAndSaveTimingViewShowGpuWorkOverlays(InValue); }
+
+		virtual bool GetTimingViewShowGpuWorkExtendedLines() const override { return GetInsightsSettings().GetTimingViewShowGpuWorkExtendedLines(); }
+		virtual void SetTimingViewShowGpuWorkExtendedLines(bool InValue) override { GetInsightsSettings().SetAndSaveTimingViewShowGpuWorkExtendedLines(InValue); }
+
+		virtual bool GetTimingViewShowGpuFencesTracks() const override { return  GetInsightsSettings().GetTimingViewShowGpuFencesTracks(); }
+		virtual void SetTimingViewShowGpuFencesTracks(bool InValue) override { GetInsightsSettings().SetAndSaveTimingViewShowGpuFencesTracks(InValue); }
+
+		virtual bool GetTimingViewShowGpuFencesExtendedLines() const override { return GetInsightsSettings().GetTimingViewShowGpuFencesExtendedLines(); }
+		virtual void SetTimingViewShowGpuFencesExtendedLines(bool InValue) override { GetInsightsSettings().SetAndSaveTimingViewShowGpuFencesExtendedLines(InValue); }
+
+		virtual bool GetTimingViewShowGpuFencesRelations() const override { return GetInsightsSettings().GetTimingViewShowGpuFencesRelations(); }
+		virtual void SetTimingViewShowGpuFencesRelations(bool InValue) override { GetInsightsSettings().SetAndSaveTimingViewShowGpuFencesRelations(InValue); };
+
+		FInsightsSettings& GetInsightsSettings() { return UE::Insights::FInsightsManager::Get()->GetSettings(); }
+		const FInsightsSettings& GetInsightsSettings() const { return UE::Insights::FInsightsManager::Get()->GetSettings(); }
+	};
+
+	class FThreadSharedStateLocalSettings : public IThreadSharedStateSetting
+	{
+	public:
+		virtual bool GetTimingViewShowGpuWorkTracks() const override { return bTimingViewShowGpuWorkTracks; }
+		virtual void SetTimingViewShowGpuWorkTracks(bool InValue) override { bTimingViewShowGpuWorkTracks = InValue; }
+
+		virtual bool GetTimingViewShowGpuWorkOverlays() const override { return bTimingViewShowGpuWorkOverlays; }
+		virtual void SetTimingViewShowGpuWorkOverlays(bool InValue) override { bTimingViewShowGpuWorkOverlays = InValue; }
+
+		virtual bool GetTimingViewShowGpuWorkExtendedLines() const override { return bTimingViewShowGpuWorkExtendedLines; }
+		virtual void SetTimingViewShowGpuWorkExtendedLines(bool InValue) override { bTimingViewShowGpuWorkExtendedLines = InValue; }
+
+		virtual bool GetTimingViewShowGpuFencesTracks() const override { return bTimingViewShowGpuFencesTracks; }
+		virtual void SetTimingViewShowGpuFencesTracks(bool InValue) override { bTimingViewShowGpuFencesTracks = InValue; }
+
+		virtual bool GetTimingViewShowGpuFencesExtendedLines() const override { return bTimingViewShowGpuFencesExtendedLines; }
+		virtual void SetTimingViewShowGpuFencesExtendedLines(bool InValue) override { bTimingViewShowGpuFencesExtendedLines = InValue; }
+
+		virtual bool GetTimingViewShowGpuFencesRelations() const override { return bTimingViewShowGpuFencesRelations; }
+		virtual void SetTimingViewShowGpuFencesRelations(bool InValue) override { bTimingViewShowGpuFencesRelations = InValue; };
+
+	private:
+		bool bTimingViewShowGpuWorkTracks = true;
+		bool bTimingViewShowGpuWorkOverlays = true;
+		bool bTimingViewShowGpuWorkExtendedLines = true;
+
+		bool bTimingViewShowGpuFencesTracks = true;
+		bool bTimingViewShowGpuFencesExtendedLines = true;
+		bool bTimingViewShowGpuFencesRelations = true;
+	};
+
 public:
 	explicit FThreadTimingSharedState(STimingView* InTimingView);
 	virtual ~FThreadTimingSharedState() override = default;
@@ -86,12 +167,18 @@ public:
 	void HideAllCpuTracks() { SetAllCpuTracksToggle(false); }
 	void ShowHideAllCpuTracks() { SetAllCpuTracksToggle(!IsAllCpuTracksToggleOn()); }
 
-	bool AreOverlaysVisibleInGpuQueueTracks() { return bShowOverlaysInGpuQueueTracks; }
-	bool AreExtendedLinesVisibleInGpuQueueTracks() { return bShowExtendedLinesInGpuQueueTracks; }
+	bool AreOverlaysVisibleInGpuQueueTracks() const { return Settings->GetTimingViewShowGpuWorkOverlays(); }
+	bool AreExtendedLinesVisibleInGpuQueueTracks() const { return Settings->GetTimingViewShowGpuWorkExtendedLines(); }
 
-	bool AreGpuWorkTracksVisible() const { return bShowGpuWorkTracks; }
+	bool AreGpuWorkTracksVisible() const { return Settings->GetTimingViewShowGpuWorkTracks(); }
 	void SetGpuWorkTracksVisibility(bool bOnOff);
 
+	bool AreGpuFencesTracksVisible() const { return Settings->GetTimingViewShowGpuFencesTracks(); }
+	void SetGpuFencesTracksVisibility(bool bOnOff);
+
+	bool AreGpuFencesExtendedLinesVisible() const { return Settings->GetTimingViewShowGpuFencesExtendedLines(); }
+	bool AreGpuFenceRelationsVisible() const { return Settings->GetTimingViewShowGpuFencesRelations(); }
+
 	TSharedPtr<const ITimingEvent> FindMaxEventInstance(uint32 TimerId, double StartTime, double EndTime);
 	TSharedPtr<const ITimingEvent> FindMinEventInstance(uint32 TimerId, double StartTime, double EndTime);
 
@@ -101,30 +188,42 @@ private:
 	bool ToggleTrackVisibilityByGroup_IsChecked(const TCHAR* InGroupName) const;
 	void ToggleTrackVisibilityByGroup_Execute(const TCHAR* InGroupName);
 
-	void Command_ShowGpuWorkTracks_Execute() { SetGpuWorkTracksVisibility(!bShowGpuWorkTracks); }
+	void Command_ShowGpuWorkTracks_Execute() { SetGpuWorkTracksVisibility(!Settings->GetTimingViewShowGpuWorkTracks()); }
 	bool Command_ShowGpuWorkTracks_CanExecute() { return GpuTracks.Num() > 0; }
-	bool Command_ShowGpuWorkTracks_IsChecked() { return bShowGpuWorkTracks; }
+	bool Command_ShowGpuWorkTracks_IsChecked() { return AreGpuWorkTracksVisible(); }
+
+	void Command_ShowGpuWorkOverlays_Execute() { Settings->SetTimingViewShowGpuWorkOverlays(!Settings->GetTimingViewShowGpuWorkOverlays()); }
+	bool Command_ShowGpuWorkOverlays_CanExecute() { return AreGpuWorkTracksVisible() && GpuTracks.Num() > 0; }
+	bool Command_ShowGpuWorkOverlays_IsChecked() { return AreOverlaysVisibleInGpuQueueTracks(); }
 
-	void Command_ShowOverlays_Execute() { bShowOverlaysInGpuQueueTracks = !bShowOverlaysInGpuQueueTracks; }
-	bool Command_ShowOverlays_CanExecute() { return bShowGpuWorkTracks && GpuTracks.Num() > 0; }
-	bool Command_ShowOverlays_IsChecked() { return bShowOverlaysInGpuQueueTracks; }
+	void Command_ShowGpuWorkExtendedLines_Execute() { Settings->SetTimingViewShowGpuWorkExtendedLines(!Settings->GetTimingViewShowGpuWorkExtendedLines()); }
+	bool Command_ShowGpuWorkExtendedLines_CanExecute() { return AreGpuWorkTracksVisible() && GpuTracks.Num() > 0; }
+	bool Command_ShowGpuWorkExtendedLines_IsChecked() { return AreExtendedLinesVisibleInGpuQueueTracks(); }
 
-	void Command_ShowExtendedLines_Execute() { bShowExtendedLinesInGpuQueueTracks = !bShowExtendedLinesInGpuQueueTracks; }
-	bool Command_ShowExtendedLines_CanExecute() { return bShowGpuWorkTracks && GpuTracks.Num() > 0; }
-	bool Command_ShowExtendedLines_IsChecked() { return bShowExtendedLinesInGpuQueueTracks; }
+	void Command_ShowGpuFencesExtendedLines_Execute() { Settings->SetTimingViewShowGpuFencesExtendedLines(!Settings->GetTimingViewShowGpuFencesExtendedLines()); }
+	bool Command_ShowGpuFencesExtendedLines_CanExecute() { return AreGpuFencesTracksVisible() && GpuTracks.Num() > 0; }
+	bool Command_ShowGpuFencesExtendedLines_IsChecked() { return AreGpuFencesExtendedLinesVisible(); }
+
+	void Command_ShowGpuFencesRelations_Execute();
+	bool Command_ShowGpuFencesRelations_CanExecute() { return GpuTracks.Num() > 0; }
+	bool Command_ShowGpuFencesRelations_IsChecked() { return AreGpuFenceRelationsVisible(); }
+
+	void Command_ShowGpuFencesTracks_Execute() { SetGpuFencesTracksVisibility(!Settings->GetTimingViewShowGpuFencesTracks()); }
+	bool Command_ShowGpuFencesTracks_CanExecute() { return GpuTracks.Num() > 0; }
+	bool Command_ShowGpuFencesTracks_IsChecked() { return AreGpuFencesTracksVisible(); }
 
 	void AddGpuWorkChildTracks();
 	void RemoveGpuWorkChildTracks();
 
+	void AddGpuFencesChildTracks();
+	void RemoveGpuFencesChildTracks();
+
 	void OnTimingEventSelected(TSharedPtr<const ITimingEvent> InSelectedEvent);
 
 private:
 	STimingView* TimingView = nullptr;
 
 	bool bShowHideAllGpuTracks = false;
-	bool bShowGpuWorkTracks = true;
-	bool bShowOverlaysInGpuQueueTracks = true;
-	bool bShowExtendedLinesInGpuQueueTracks = true;
 	bool bShowHideAllCpuTracks = false;
 
 	TSharedPtr<FGpuTimingTrack> OldGpu1Track;
@@ -141,6 +240,8 @@ private:
 
 	uint64 TimingProfilerTimelineCount = 0;
 	uint64 LoadTimeProfilerTimelineCount = 0;
+
+	TSharedRef<IThreadSharedStateSetting> Settings;
 };
 
 ////////////////////////////////////////////////////////////////////////////////////////////////////
diff --git a/Engine/Source/Developer/TraceInsights/Private/Insights/ViewModels/BaseTimingTrack.cpp b/Engine/Source/Developer/TraceInsights/Private/Insights/ViewModels/BaseTimingTrack.cpp
--- a/Engine/Source/Developer/TraceInsights/Private/Insights/ViewModels/BaseTimingTrack.cpp
+++ b/Engine/Source/Developer/TraceInsights/Private/Insights/ViewModels/BaseTimingTrack.cpp
@@ -27,6 +27,11 @@ void FBaseTimingTrack::AddChildTrack(TSharedRef<FBaseTimingTrack> Track)
 	ChildTracks.AddUnique(Track);
 }
 
+void FBaseTimingTrack::AddChildTrack(TSharedRef<FBaseTimingTrack> Track, int32 Index)
+{
+	ChildTracks.Insert(Track, Index);
+}
+
 void FBaseTimingTrack::RemoveChildTrack(TSharedRef<FBaseTimingTrack> Track)
 {
 	ChildTracks.Remove(Track);
diff --git a/Engine/Source/Developer/TraceInsights/Public/Insights/ViewModels/BaseTimingTrack.h b/Engine/Source/Developer/TraceInsights/Public/Insights/ViewModels/BaseTimingTrack.h
--- a/Engine/Source/Developer/TraceInsights/Public/Insights/ViewModels/BaseTimingTrack.h
+++ b/Engine/Source/Developer/TraceInsights/Public/Insights/ViewModels/BaseTimingTrack.h
@@ -304,6 +304,12 @@ public:
 	 */
 	void AddChildTrack(TSharedRef<FBaseTimingTrack> Track);
 
+	/**
+	 * Inserts a child track to this track at the specified index.
+	 * @note This invalidates the result of `GetChildTracks`.
+	 */
+	void AddChildTrack(TSharedRef<FBaseTimingTrack> Track, int32 Index);
+
 	/**
 	 * Removes the provided track from this track's child tracks.
 	 *
