[Insights] * Memory Insights / Investigation: Updated tooltip for "Include Heap Allocs" toggle.

diff --git a/Engine/Source/Developer/TraceInsights/Private/Insights/MemoryProfiler/Widgets/SMemInvestigationView.cpp b/Engine/Source/Developer/TraceInsights/Private/Insights/MemoryProfiler/Widgets/SMemInvestigationView.cpp
--- a/Engine/Source/Developer/TraceInsights/Private/Insights/MemoryProfiler/Widgets/SMemInvestigationView.cpp
+++ b/Engine/Source/Developer/TraceInsights/Private/Insights/MemoryProfiler/Widgets/SMemInvestigationView.cpp
@@ -165,9 +165,9 @@ TSharedRef<SWidget> SMemInvestigationView::ConstructInvestigationWidgetArea()
 			.Content()
 			[
 				SNew(STextBlock)
-				.Text(LOCTEXT("IncludeHeapAllocsText", "Include Heap Allocs"))
+				.Text(LOCTEXT("IncludeHeapAllocsText", "Include Heap Allocs (Advanced)"))
 			]
-			.ToolTipText(LOCTEXT("IncludeHeapAllocsToolTipText", "Include heap allocs."))
+			.ToolTipText(LOCTEXT("IncludeHeapAllocsToolTipText", "Include Heap Allocs (Advanced Feature)\n\nIf enabled, the table will also list the heap allocs (\"heap_XXXX\").\nThe heap allocs are memory blocks that can have child allocations.\n\nWarning:\tTotal aggregated sizes may be irrelevant when a heap is added with its child allocs.\n\t\t\t\t\tSome memory would be double counted."))
 		]
 
 		+ SVerticalBox::Slot()
@@ -232,7 +232,7 @@ TSharedRef<SWidget> SMemInvestigationView::ConstructInvestigationWidgetArea()
 			.ToolTipText(LOCTEXT("RunQueryBtnToolTipText", "Run Memory Query.\nThe resulting list of allocations will be available in a tree view."))
 			.OnClicked(this, &SMemInvestigationView::RunQuery)
 		]
-	
+
 		+ SVerticalBox::Slot()
 		.AutoHeight()
 		.Padding(FMargin(0.0f, 10.0f, 0.0f, 0.0f))
