[Insights] * Memory Insights / MemAlloc table: Changed the initial view preset to be "Heap" when "Include Heap Allocs" is on.

diff --git a/Engine/Source/Developer/TraceInsights/Private/Insights/MemoryProfiler/Widgets/SMemAllocTableTreeView.cpp b/Engine/Source/Developer/TraceInsights/Private/Insights/MemoryProfiler/Widgets/SMemAllocTableTreeView.cpp
--- a/Engine/Source/Developer/TraceInsights/Private/Insights/MemoryProfiler/Widgets/SMemAllocTableTreeView.cpp
+++ b/Engine/Source/Developer/TraceInsights/Private/Insights/MemoryProfiler/Widgets/SMemAllocTableTreeView.cpp
@@ -10,6 +10,7 @@
 #include "ISourceCodeAccessModule.h"
 #include "ISourceCodeAccessor.h"
 #include "Logging/MessageLog.h"
+#include "Misc/Optional.h"
 #include "Modules/ModuleManager.h"
 #include "SlateOptMacros.h"
 #include "Styling/AppStyle.h"
@@ -31,6 +32,7 @@
 // TraceInsightsCore
 #include "InsightsCore/Filter/ViewModels/FilterConfigurator.h"
 #include "InsightsCore/Filter/ViewModels/Filters.h"
+#include "InsightsCore/Table/ViewModels/TableCellValue.h"
 #include "InsightsCore/Table/ViewModels/TableColumn.h"
 
 // TraceInsights
@@ -162,6 +164,7 @@ TSharedPtr<SWidget> SMemAllocTableTreeView::ConstructFooter()
 END_SLATE_FUNCTION_BUILD_OPTIMIZATION
 
 ////////////////////////////////////////////////////////////////////////////////////////////////////
+
 void SMemAllocTableTreeView::InternalCreateGroupings()
 {
 	STableTreeView::InternalCreateGroupings();
@@ -399,6 +402,15 @@ void SMemAllocTableTreeView::OnQueryInvalidated()
 
 void SMemAllocTableTreeView::ResetAndStartQuery()
 {
+	// Apply the "Default" or the "Heap" preset.
+	const TArray<TSharedRef<ITableTreeViewPreset>>* ViewPresets = GetAvailableViewPresets();
+	if (ViewPresets && ViewPresets->Num() > 3)
+	{
+		constexpr int32 DefaultPresetIndex = 0;
+		constexpr int32 HeapPresetIndex = 2;
+		SelectViewPreset((*ViewPresets)[HasHeapAllocs() ? HeapPresetIndex : DefaultPresetIndex].ToSharedPtr());
+	}
+
 	TableRowNodes.Reset();
 
 	TSharedPtr<FMemAllocTable> MemAllocTable = GetMemAllocTable();
@@ -825,13 +837,13 @@ void SMemAllocTableTreeView::UpdateSelectionStatsText()
 
 		for (const FTableTreeNodePtr& Node : SelectedNodes)
 		{
-			TOptional<FTableCellValue> CountValue = CountColumn->GetValue(*Node.Get());
+			const TOptional<FTableCellValue> CountValue = CountColumn->GetValue(*Node.Get());
 			if (CountValue.IsSet())
 			{
 				TotalAllocCount += CountValue.GetValue().AsInt64();
 			}
 
-			TOptional<FTableCellValue> SizeValue = SizeColumn->GetValue(*Node.Get());
+			const TOptional<FTableCellValue> SizeValue = SizeColumn->GetValue(*Node.Get());
 			if (SizeValue.IsSet())
 			{
 				TotalAllocSize += SizeValue.GetValue().AsInt64();
diff --git a/Engine/Source/Developer/TraceInsights/Private/Insights/MemoryProfiler/Widgets/SMemAllocTableTreeView.h b/Engine/Source/Developer/TraceInsights/Private/Insights/MemoryProfiler/Widgets/SMemAllocTableTreeView.h
--- a/Engine/Source/Developer/TraceInsights/Private/Insights/MemoryProfiler/Widgets/SMemAllocTableTreeView.h
+++ b/Engine/Source/Developer/TraceInsights/Private/Insights/MemoryProfiler/Widgets/SMemAllocTableTreeView.h
@@ -76,6 +76,9 @@ public:
 	 */
 	virtual void RebuildTree(bool bResync);
 
+	bool HasHeapAllocs() const { return bIncludeHeapAllocs; }
+	bool HasSwapAllocs() const { return bIncludeSwapAllocs; }
+
 	struct FQueryParams
 	{
 		TSharedPtr<FMemoryRuleSpec> Rule;
diff --git a/Engine/Source/Developer/TraceInsightsCore/Private/InsightsCore/Table/Widgets/STableTreeView.cpp b/Engine/Source/Developer/TraceInsightsCore/Private/InsightsCore/Table/Widgets/STableTreeView.cpp
--- a/Engine/Source/Developer/TraceInsightsCore/Private/InsightsCore/Table/Widgets/STableTreeView.cpp
+++ b/Engine/Source/Developer/TraceInsightsCore/Private/InsightsCore/Table/Widgets/STableTreeView.cpp
@@ -3898,10 +3898,19 @@ void STableTreeView::UpdateBannerText()
 
 ////////////////////////////////////////////////////////////////////////////////////////////////////
 
-FReply STableTreeView::OnApplyViewPreset(const ITableTreeViewPreset* InPreset)
+TSharedPtr<ITableTreeViewPreset> STableTreeView::GetSelectedViewPreset() const
 {
-	ApplyViewPreset(*InPreset);
-	return FReply::Handled();
+	return PresetComboBox->GetSelectedItem();
+}
+
+////////////////////////////////////////////////////////////////////////////////////////////////////
+
+void STableTreeView::SelectViewPreset(TSharedPtr<ITableTreeViewPreset> InPreset)
+{
+	if (PresetComboBox)
+	{
+		PresetComboBox->SetSelectedItem(InPreset);
+	}
 }
 
 ////////////////////////////////////////////////////////////////////////////////////////////////////
diff --git a/Engine/Source/Developer/TraceInsightsCore/Public/InsightsCore/Table/Widgets/STableTreeView.h b/Engine/Source/Developer/TraceInsightsCore/Public/InsightsCore/Table/Widgets/STableTreeView.h
--- a/Engine/Source/Developer/TraceInsightsCore/Public/InsightsCore/Table/Widgets/STableTreeView.h
+++ b/Engine/Source/Developer/TraceInsightsCore/Public/InsightsCore/Table/Widgets/STableTreeView.h
@@ -465,9 +465,10 @@ protected:
 	////////////////////////////////////////////////////////////////////////////////////////////////////
 	// Presets
 
-	virtual void InitAvailableViewPresets() {};
+	virtual void InitAvailableViewPresets() {}
 	const TArray<TSharedRef<ITableTreeViewPreset>>* GetAvailableViewPresets() const { return &AvailableViewPresets; }
-	FReply OnApplyViewPreset(const ITableTreeViewPreset* InPreset);
+	TSharedPtr<ITableTreeViewPreset> GetSelectedViewPreset() const;
+	void SelectViewPreset(TSharedPtr<ITableTreeViewPreset> InPreset);
 	void ApplyViewPreset(const ITableTreeViewPreset& InPreset);
 	void ApplyColumnConfig(const TArrayView<FTableColumnConfig>& InTableConfig);
 	void ViewPreset_OnSelectionChanged(TSharedPtr<ITableTreeViewPreset> InPreset, ESelectInfo::Type SelectInfo);
