Minor Gpu Profiler fixes

diff --git a/Engine/Source/Developer/TraceInsights/Private/Insights/TimingProfiler/Tracks/GpuTimingTrack.cpp b/Engine/Source/Developer/TraceInsights/Private/Insights/TimingProfiler/Tracks/GpuTimingTrack.cpp
--- a/Engine/Source/Developer/TraceInsights/Private/Insights/TimingProfiler/Tracks/GpuTimingTrack.cpp
+++ b/Engine/Source/Developer/TraceInsights/Private/Insights/TimingProfiler/Tracks/GpuTimingTrack.cpp
@@ -198,10 +198,10 @@ void FGpuFencesTrackBuilder::Flush(float AvailableTextW)
 			return FLinearColor(R, G, B);
 		};
 
-	constexpr FLinearColor SingnalFenceColor = ToLiniarColorNoAlpha(0x3A9C7bFF);
-	constexpr FLinearColor WaitFenceColor = ToLiniarColorNoAlpha(0xFFDC1AFF);
+	constexpr FLinearColor SignalFenceColor = ToLiniarColorNoAlpha(0x3A9C7bFF); // Light Blue
+	constexpr FLinearColor WaitFenceColor = ToLiniarColorNoAlpha(0xFFDC1AFF); // Yellow
 
-	const FLinearColor Color = LastType == TraceServices::EGpuFenceType::SignalFence ? SingnalFenceColor : WaitFenceColor;
+	const FLinearColor Color = LastType == TraceServices::EGpuFenceType::SignalFence ? SignalFenceColor : WaitFenceColor;
 
 	bool bAddNewBox = true;
 	if (Track.TimeMarkerBoxes.Num() > 0)
@@ -374,7 +374,7 @@ void FGpuFencesTimingTrack::PreUpdate(const ITimingTrackUpdateContext& Context)
 			Builder.Flush();
 		}
 
-		SetNumLanes(1);
+		SetNumLanes(TimeMarkerBoxes.Num() ? 1 : 0);
 	}
 
 	UpdateTrackHeight(Context);
