[Insights]
* TraceServices: Added provider specific edit/read lock for the LLM MemoryProvider.
* MemoryInsights/LLMTags: Improved performance of updating stats and budgets for larger tag sets or when still analyzing.
* MemoryInsights: Fixed tag set override for budgets.
* TimingInsights: Added IsDragging(), StartDragging() and StopDragging() to FTimeMarker.

diff --git a/Engine/Source/Developer/TraceInsights/Private/Insights/MemoryProfiler/MemoryProfilerManager.cpp b/Engine/Source/Developer/TraceInsights/Private/Insights/MemoryProfiler/MemoryProfilerManager.cpp
--- a/Engine/Source/Developer/TraceInsights/Private/Insights/MemoryProfiler/MemoryProfilerManager.cpp
+++ b/Engine/Source/Developer/TraceInsights/Private/Insights/MemoryProfiler/MemoryProfilerManager.cpp
@@ -160,7 +160,7 @@ void FMemoryProfilerManager::RegisterMajorTabs(IUnrealInsightsModule& InsightsMo
 {
 	bool bShouldRegister = true;
 
-#if WITH_EDITOR 
+#if WITH_EDITOR
 	#if UE_BUILD_SHIPPING
 		bShouldRegister = false;
 	#else
@@ -291,17 +291,19 @@ bool FMemoryProfilerManager::Tick(float DeltaTime)
 		TSharedPtr<const TraceServices::IAnalysisSession> Session = FInsightsManager::Get()->GetSession();
 		if (Session.IsValid())
 		{
-			TraceServices::FAnalysisSessionReadScope SessionReadScope(*Session.Get());
-
-			if (Session->IsAnalysisComplete())
 			{
-				// Never check again during this session.
-				AvailabilityCheck.Disable();
+				TraceServices::FAnalysisSessionReadScope SessionReadScope(*Session.Get());
+				if (Session->IsAnalysisComplete())
+				{
+					// Never check again during this session.
+					AvailabilityCheck.Disable();
+				}
 			}
 
 			const TraceServices::IMemoryProvider* MemoryProvider = TraceServices::ReadMemoryProvider(*Session.Get());
 			if (MemoryProvider)
 			{
+				TraceServices::FProviderReadScopeLock _(*MemoryProvider);
 				uint32 TagCount = MemoryProvider->GetTagCount();
 				if (TagCount > 0)
 				{
@@ -426,7 +428,7 @@ void FMemoryProfilerManager::ShowHideModulesView(const bool bIsVisible)
 
 void FMemoryProfilerManager::OnWindowClosedEvent()
 {
-	// Need to close MemAlloc window to prevent it being saved in the layout and spawning as an "Unregister Tab" on the next application start. 
+	// Need to close MemAlloc window to prevent it being saved in the layout and spawning as an "Unregister Tab" on the next application start.
 	TSharedPtr<SMemoryProfilerWindow> Wnd = GetProfilerWindow();
 	if (Wnd.IsValid())
 	{
diff --git a/Engine/Source/Developer/TraceInsights/Private/Insights/MemoryProfiler/Tracks/MemTagGraphSeries.cpp b/Engine/Source/Developer/TraceInsights/Private/Insights/MemoryProfiler/Tracks/MemTagGraphSeries.cpp
--- a/Engine/Source/Developer/TraceInsights/Private/Insights/MemoryProfiler/Tracks/MemTagGraphSeries.cpp
+++ b/Engine/Source/Developer/TraceInsights/Private/Insights/MemoryProfiler/Tracks/MemTagGraphSeries.cpp
@@ -3,6 +3,7 @@
 #include "MemTagGraphSeries.h"
 
 // TraceServices
+#include "Common/ProviderLock.h"
 #include "TraceServices/Model/Memory.h"
 
 // TraceInsights
@@ -143,7 +144,7 @@ void FMemTagGraphSeries::PreUpdate(FGraphTrack& GraphTrack, const FTimingTrackVi
 		const TraceServices::IMemoryProvider* MemoryProvider = TraceServices::ReadMemoryProvider(*Session.Get());
 		if (MemoryProvider)
 		{
-			TraceServices::FAnalysisSessionReadScope SessionReadScope(*Session.Get());
+			TraceServices::FProviderReadScopeLock _(*MemoryProvider);
 
 			const uint64 TotalSampleCount = MemoryProvider->GetTagSampleCount(TrackerId, TagId);
 			if (TotalSampleCount > 0)
@@ -209,7 +210,7 @@ void FMemTagGraphSeries::Update(FGraphTrack& GraphTrack, const FTimingTrackViewp
 		const TraceServices::IMemoryProvider* MemoryProvider = TraceServices::ReadMemoryProvider(*Session.Get());
 		if (MemoryProvider)
 		{
-			TraceServices::FAnalysisSessionReadScope SessionReadScope(*Session.Get());
+			TraceServices::FProviderReadScopeLock _(*MemoryProvider);
 
 			const uint64 TotalSampleCount = MemoryProvider->GetTagSampleCount(TrackerId, TagId);
 			if (TotalSampleCount > 0)
diff --git a/Engine/Source/Developer/TraceInsights/Private/Insights/MemoryProfiler/ViewModels/MemAllocGroupingByTag.cpp b/Engine/Source/Developer/TraceInsights/Private/Insights/MemoryProfiler/ViewModels/MemAllocGroupingByTag.cpp
--- a/Engine/Source/Developer/TraceInsights/Private/Insights/MemoryProfiler/ViewModels/MemAllocGroupingByTag.cpp
+++ b/Engine/Source/Developer/TraceInsights/Private/Insights/MemoryProfiler/ViewModels/MemAllocGroupingByTag.cpp
@@ -324,31 +324,27 @@ void FMemTagTableTreeNode::UpdateLLMSize() const
 	}
 
 	const TraceServices::IMemoryProvider* MemoryProvider = TraceServices::ReadMemoryProvider(*Session.Get());
-	if (!MemoryProvider)
+	if (MemoryProvider)
 	{
-		return;
-	}
-
-	TraceServices::FAnalysisSessionReadScope SessionReadScope(*Session.Get());
-
-	const uint64 TotalSampleCount = MemoryProvider->GetTagSampleCount(TrackerId, TagId);
-	if (TotalSampleCount == 0)
-	{
-		return;
-	}
+		TraceServices::FProviderReadScopeLock _(*MemoryProvider);
 
-	int64 LocalLLMSize = INT64_MAX;
-	MemoryProvider->EnumerateTagSamples(TrackerId, TagId, Time, Time, true,
-		[&LocalLLMSize](double Time, double Duration, const TraceServices::FMemoryTagSample& Sample)
+		const uint64 TotalSampleCount = MemoryProvider->GetTagSampleCount(TrackerId, TagId);
+		if (TotalSampleCount > 0)
 		{
-			if (LocalLLMSize == INT64_MAX)
+			int64 LocalLLMSize = INT64_MAX;
+			MemoryProvider->EnumerateTagSamples(TrackerId, TagId, Time, Time, true,
+				[&LocalLLMSize](double Time, double Duration, const TraceServices::FMemoryTagSample& Sample)
+				{
+					if (LocalLLMSize == INT64_MAX)
+					{
+						LocalLLMSize = Sample.Value;
+					}
+				});
+			if (LocalLLMSize != INT64_MAX)
 			{
-				LocalLLMSize = Sample.Value;
+				LLMSize = LocalLLMSize;
 			}
-		});
-	if (LocalLLMSize != INT64_MAX)
-	{
-		LLMSize = LocalLLMSize;
+		}
 	}
 }
 
diff --git a/Engine/Source/Developer/TraceInsights/Private/Insights/MemoryProfiler/ViewModels/MemTagBudget.cpp b/Engine/Source/Developer/TraceInsights/Private/Insights/MemoryProfiler/ViewModels/MemTagBudget.cpp
--- a/Engine/Source/Developer/TraceInsights/Private/Insights/MemoryProfiler/ViewModels/MemTagBudget.cpp
+++ b/Engine/Source/Developer/TraceInsights/Private/Insights/MemoryProfiler/ViewModels/MemTagBudget.cpp
@@ -268,9 +268,9 @@ bool FMemTagBudget::ReLoadFromFile()
 	const TCHAR* DefaultCachedTrackerName = StoreString(FString(TEXT("Default")));
 	const TCHAR* PlatformCachedTrackerName = StoreString(FString(TEXT("Platform")));
 
-	const TCHAR* SystemCachedTagSetName = StoreString(FString(TEXT("System")));
+	const TCHAR* SystemCachedTagSetName = StoreString(FString(TEXT("Systems")));
 	const TCHAR* AssetCachedTagSetName = StoreString(FString(TEXT("Assets")));
-	const TCHAR* AssetClassCachedTagSetName = StoreString(FString(TEXT("AssetsClasses")));
+	const TCHAR* AssetClassCachedTagSetName = StoreString(FString(TEXT("AssetClasses")));
 
 	static const FString STR_Budget(TEXT("budget"));
 	static const FString STR_Platform(TEXT("platform"));
@@ -343,15 +343,20 @@ bool FMemTagBudget::ReLoadFromFile()
 					FString TagSetName = TagNode->GetAttribute(STR_Set);
 					if (!TagSetName.IsEmpty())
 					{
-						if (!FCString::Stricmp(*TagSetName, TEXT("Asset")) || !FCString::Stricmp(*TagSetName, TEXT("Assets")))
+						if (!FCString::Stricmp(*TagSetName, TEXT("Asset")) ||
+							!FCString::Stricmp(*TagSetName, TEXT("Assets")))
 						{
 							CachedTagSetName = AssetCachedTagSetName;
 						}
-						else if (!FCString::Stricmp(*TagSetName, TEXT("AssetClass")) || !FCString::Stricmp(*TagSetName, TEXT("AssetClasses")))
+						else
+						if (!FCString::Stricmp(*TagSetName, TEXT("AssetClass")) ||
+							!FCString::Stricmp(*TagSetName, TEXT("AssetClasses")))
 						{
 							CachedTagSetName = AssetClassCachedTagSetName;
 						}
-						else if (!FCString::Stricmp(*TagSetName, TEXT("System")) || !FCString::Stricmp(*TagSetName, TEXT("Systems")))
+						else
+						if (!FCString::Stricmp(*TagSetName, TEXT("System")) ||
+							!FCString::Stricmp(*TagSetName, TEXT("Systems")))
 						{
 							CachedTagSetName = SystemCachedTagSetName;
 						}
diff --git a/Engine/Source/Developer/TraceInsights/Private/Insights/MemoryProfiler/ViewModels/MemorySharedState.cpp b/Engine/Source/Developer/TraceInsights/Private/Insights/MemoryProfiler/ViewModels/MemorySharedState.cpp
--- a/Engine/Source/Developer/TraceInsights/Private/Insights/MemoryProfiler/ViewModels/MemorySharedState.cpp
+++ b/Engine/Source/Developer/TraceInsights/Private/Insights/MemoryProfiler/ViewModels/MemorySharedState.cpp
@@ -412,14 +412,16 @@ void FMemorySharedState::SyncTrackers()
 	TSharedPtr<const TraceServices::IAnalysisSession> Session = FInsightsManager::Get()->GetSession();
 	if (Session.IsValid())
 	{
-		TraceServices::FAnalysisSessionReadScope SessionReadScope(*Session.Get());
 		const TraceServices::IMemoryProvider* MemoryProvider = TraceServices::ReadMemoryProvider(*Session.Get());
 		if (MemoryProvider)
 		{
-			MemoryProvider->EnumerateTrackers([this](const TraceServices::FMemoryTrackerInfo& Tracker)
 			{
-				Trackers.Add(MakeShared<FMemoryTracker>(Tracker.Id, Tracker.Name));
-			});
+				TraceServices::FProviderReadScopeLock _(*MemoryProvider);
+				MemoryProvider->EnumerateTrackers([this](const TraceServices::FMemoryTrackerInfo& Tracker)
+				{
+					Trackers.Add(MakeShared<FMemoryTracker>(Tracker.Id, Tracker.Name));
+				});
+			}
 
 			Trackers.Sort([](const TSharedPtr<FMemoryTracker>& A, const TSharedPtr<FMemoryTracker>& B) { return A->GetId() < B->GetId(); });
 		}
@@ -452,6 +454,8 @@ void FMemorySharedState::UpdateTagSets()
 		const TraceServices::IMemoryProvider* MemoryProvider = TraceServices::ReadMemoryProvider(*Session.Get());
 		if (MemoryProvider)
 		{
+			TraceServices::FProviderReadScopeLock _(*MemoryProvider);
+
 			const uint32 NumProviderTagSets = MemoryProvider->GetTagSetCount();
 			if (NumProviderTagSets != NumSyncedTagSets)
 			{
diff --git a/Engine/Source/Developer/TraceInsights/Private/Insights/MemoryProfiler/ViewModels/MemoryTag.cpp b/Engine/Source/Developer/TraceInsights/Private/Insights/MemoryProfiler/ViewModels/MemoryTag.cpp
--- a/Engine/Source/Developer/TraceInsights/Private/Insights/MemoryProfiler/ViewModels/MemoryTag.cpp
+++ b/Engine/Source/Developer/TraceInsights/Private/Insights/MemoryProfiler/ViewModels/MemoryTag.cpp
@@ -3,6 +3,7 @@
 #include "MemoryTag.h"
 
 // TraceServices
+#include "Common/ProviderLock.h"
 #include "TraceServices/Model/Memory.h"
 
 // TraceInsights
@@ -141,10 +142,11 @@ void FMemoryTagList::UpdateInternal()
 	TSharedPtr<const TraceServices::IAnalysisSession> Session = FInsightsManager::Get()->GetSession();
 	if (Session.IsValid())
 	{
-		TraceServices::FAnalysisSessionReadScope SessionReadScope(*Session.Get());
 		const TraceServices::IMemoryProvider* MemoryProvider = TraceServices::ReadMemoryProvider(*Session.Get());
 		if (MemoryProvider)
 		{
+			TraceServices::FProviderReadScopeLock _(*MemoryProvider);
+
 			const int32 TraceSerialNumber = MemoryProvider->GetTagSerial();
 			if (LastTraceSerialNumber != TraceSerialNumber)
 			{
diff --git a/Engine/Source/Developer/TraceInsights/Private/Insights/MemoryProfiler/Widgets/SMemTagTreeView.cpp b/Engine/Source/Developer/TraceInsights/Private/Insights/MemoryProfiler/Widgets/SMemTagTreeView.cpp
--- a/Engine/Source/Developer/TraceInsights/Private/Insights/MemoryProfiler/Widgets/SMemTagTreeView.cpp
+++ b/Engine/Source/Developer/TraceInsights/Private/Insights/MemoryProfiler/Widgets/SMemTagTreeView.cpp
@@ -13,6 +13,10 @@
 #include "Widgets/Input/SComboBox.h"
 #include "Widgets/Input/SSegmentedControl.h"
 
+// TraceServices
+#include "Common/ProviderLock.h"
+#include "TraceServices/Model/Memory.h"
+
 // TraceInsightsCore
 #include "InsightsCore/Common/Stopwatch.h"
 #include "InsightsCore/Table/ViewModels/Table.h"
@@ -365,6 +369,8 @@ void SMemTagTreeView::TagSet_OnSelectionChanged(TSharedPtr<FMemoryTagSetId> NewT
 	{
 		TagSetFilter = *NewTagSet;
 		OnNodeFilteringChanged();
+		bShouldUpdateStats = true;
+		bShouldUpdateBudgets = true;
 	}
 }
 
@@ -916,15 +922,27 @@ void SMemTagTreeView::Tick(const FGeometry& AllottedGeometry, const double InCur
 	{
 		RebuildTree(false);
 
-		UpdateStats();
-
-		// 1000 tags --> check each 120ms
-		// 10000 tags --> check each 300ms
-		// 100000 tags --> check each 2.1s
-		const double WaitTimeSec = 0.1 + FilteredNodesPtr->Num() / 50000.0;
+		const int32 NumTags = GetTableRowNodes().Num();
+		// 1000 tags --> check each 200ms
+		// 10000 tags --> check each 1.1s
+		// 100000 tags --> check each 10.1s
+		const double WaitTimeSec = 0.1 + (double)NumTags / 10000.0;
 		const uint64 WaitTime = static_cast<uint64>(WaitTimeSec / FPlatformTime::GetSecondsPerCycle64());
 		NextTimestamp = Time + WaitTime;
 	}
+
+	bShouldUpdateStats = CheckIfShouldUpdateStats();
+	if (bShouldUpdateStats)
+	{
+		bShouldUpdateStats = false;
+		UpdateStats();
+	}
+
+	if (bShouldUpdateBudgets)
+	{
+		bShouldUpdateBudgets = false;
+		ApplyBudgetToNodes();
+	}
 }
 
 ////////////////////////////////////////////////////////////////////////////////////////////////////
@@ -991,7 +1009,7 @@ void SMemTagTreeView::RebuildTree(bool bResync)
 				}
 			}
 
-			ApplyBudgetToNodes();
+			bShouldUpdateBudgets = true;
 		}
 
 		if (AvailableTagSets.Num() != SharedState.GetNumTagSets())
@@ -1051,30 +1069,82 @@ void SMemTagTreeView::RebuildTree(bool bResync)
 
 ////////////////////////////////////////////////////////////////////////////////////////////////////
 
-void SMemTagTreeView::UpdateStats()
+bool SMemTagTreeView::CheckIfShouldUpdateStats() const
 {
-	bool bShouldUpdateStats = false;
-
-	FStopwatch AggregationStopwatch;
-	FStopwatch Stopwatch;
-	Stopwatch.Start();
+	if (bShouldUpdateStats)
+	{
+		return true;
+	}
 
 	TSharedPtr<SMemoryProfilerWindow> ProfilerWindow = GetProfilerWindow();
-	if (ProfilerWindow.IsValid())
+	if (!ProfilerWindow.IsValid())
 	{
-		const TSharedRef<TimingProfiler::FTimeMarker>& MarkerA = ProfilerWindow->GetCustomTimeMarker(0);
-		if (StatsTimeA != MarkerA->GetTime())
+		return false;
+	}
+	FMemorySharedState& SharedState = ProfilerWindow->GetSharedState();
+
+	bool bUpdateWhileDragging = true;
+
+	// Do not update while dragging the A or B time markers and analysis is still active (ex.: during live session).
+	if (Session.IsValid())
+	{
+		const TraceServices::IMemoryProvider* MemoryProvider = TraceServices::ReadMemoryProvider(*Session.Get());
+		if (MemoryProvider && !MemoryProvider->IsCompleted())
 		{
-			StatsTimeA = MarkerA->GetTime();
-			bShouldUpdateStats = true;
+			bUpdateWhileDragging = false;
 		}
+	}
 
-		const TSharedRef<TimingProfiler::FTimeMarker>& MarkerB = ProfilerWindow->GetCustomTimeMarker(1);
-		if (StatsTimeB != MarkerB->GetTime())
+	// Do not update while dragging the A or B time markers and current tag set is other than Systems.
+	if (TagSetFilter != FMemoryTagSet::DefaultTagSetId)
+	{
+		bUpdateWhileDragging = false;
+	}
+
+	const TSharedRef<TimingProfiler::FTimeMarker>& MarkerA = ProfilerWindow->GetCustomTimeMarker(0);
+	const TSharedRef<TimingProfiler::FTimeMarker>& MarkerB = ProfilerWindow->GetCustomTimeMarker(1);
+	const bool bIsDragging = MarkerA->IsDragging() || MarkerB->IsDragging();
+
+	if (bUpdateWhileDragging || !bIsDragging)
+	{
+		if (StatsTimeA != MarkerA->GetTime() ||
+			StatsTimeB != MarkerB->GetTime())
 		{
-			StatsTimeB = MarkerB->GetTime();
-			bShouldUpdateStats = true;
+			return true;
 		}
+	}
+
+	TSharedPtr<TimingProfiler::STimingView> TimingView = ProfilerWindow->GetTimingView();
+	if (TimingView &&
+		(bUpdateWhileDragging || !TimingView->IsSelecting()))
+	{
+		double SelectionStartTime = TimingView->GetSelectionStartTime();
+		double SelectionEndTime = TimingView->GetSelectionEndTime();
+		if (SelectionStartTime < SelectionEndTime)
+		{
+			if (StatsStartTime != SelectionStartTime ||
+				StatsEndTime != SelectionEndTime)
+			{
+				return true;
+			}
+		}
+	}
+
+	return false;
+}
+
+////////////////////////////////////////////////////////////////////////////////////////////////////
+
+void SMemTagTreeView::UpdateStats()
+{
+	TSharedPtr<SMemoryProfilerWindow> ProfilerWindow = GetProfilerWindow();
+	if (ProfilerWindow.IsValid())
+	{
+		const TSharedRef<TimingProfiler::FTimeMarker>& MarkerA = ProfilerWindow->GetCustomTimeMarker(0);
+		StatsTimeA = MarkerA->GetTime();
+
+		const TSharedRef<TimingProfiler::FTimeMarker>& MarkerB = ProfilerWindow->GetCustomTimeMarker(1);
+		StatsTimeB = MarkerB->GetTime();
 
 		TSharedPtr<TimingProfiler::STimingView> TimingView = ProfilerWindow->GetTimingView();
 		if (TimingView)
@@ -1083,29 +1153,21 @@ void SMemTagTreeView::UpdateStats()
 			double SelectionEndTime = TimingView->GetSelectionEndTime();
 			if (SelectionStartTime < SelectionEndTime)
 			{
-				if (StatsStartTime != SelectionStartTime)
-				{
-					StatsStartTime = SelectionStartTime;
-					bShouldUpdateStats = true;
-				}
-				if (StatsEndTime != SelectionEndTime)
-				{
-					StatsEndTime = SelectionEndTime;
-					bShouldUpdateStats = true;
-				}
+				StatsStartTime = SelectionStartTime;
+				StatsEndTime = SelectionEndTime;
 			}
 		}
+	}
 
-		if (bShouldUpdateStats)
-		{
-			UpdateStatsInternal();
+	FStopwatch Stopwatch;
+	Stopwatch.Start();
 
-			bool bShouldUpdateTreeWithStats = false; // TODO: UI toggle (advanced setting)
-			if (bShouldUpdateTreeWithStats)
-			{
-				UpdateTree();
-			}
-		}
+	UpdateStatsInternal();
+
+	bool bShouldUpdateTreeWithStats = false; // TODO: UI toggle (advanced setting)
+	if (bShouldUpdateTreeWithStats)
+	{
+		UpdateTree();
 	}
 
 	Stopwatch.Stop();
@@ -1125,10 +1187,11 @@ void SMemTagTreeView::UpdateStatsInternal()
 		return;
 	}
 
-	TraceServices::FAnalysisSessionReadScope SessionReadScope(*Session.Get());
 	const TraceServices::IMemoryProvider* MemoryProvider = TraceServices::ReadMemoryProvider(*Session.Get());
 	if (MemoryProvider)
 	{
+		TraceServices::FProviderReadScopeLock _(*MemoryProvider);
+
 		for (const FTableTreeNodePtr& NodePtr : TableRowNodes)
 		{
 			if (!NodePtr->Is<FMemTagNode>())
@@ -1144,6 +1207,11 @@ void SMemTagTreeView::UpdateStatsInternal()
 				continue;
 			}
 
+			if (MemTagNode.GetMemTagSetId() != TagSetFilter)
+			{
+				continue;
+			}
+
 			FMemTagStats CurrentStats;
 
 			CurrentStats.SizeA = std::numeric_limits<int64>::max();
@@ -1214,10 +1282,10 @@ void SMemTagTreeView::UpdateStatsInternal()
 				MemTagNode.GetStats().SizeAverage = 0;
 			}
 		}
-
-		UpdateSelectionStatsText();
-		UpdateAggregatedValuesRec(*Root);
 	}
+
+	UpdateSelectionStatsText();
+	UpdateAggregatedValuesRec(*Root);
 }
 
 ////////////////////////////////////////////////////////////////////////////////////////////////////
@@ -2014,7 +2082,7 @@ void SMemTagTreeView::BudgetMode_OnSelectionChanged(TSharedPtr<FMemTagBudgetMode
 		if (!SelectedBudgetMode.IsValid() || !SelectedBudgetMode->Name.Equals(InBudgetMode->Name, ESearchCase::IgnoreCase))
 		{
 			SelectedBudgetMode = InBudgetMode;
-			ApplyBudgetToNodes();
+			bShouldUpdateBudgets = true;
 		}
 	}
 }
@@ -2056,7 +2124,7 @@ void SMemTagTreeView::BudgetPlatform_OnSelectionChanged(TSharedPtr<FMemTagBudget
 		if (!SelectedBudgetPlatform.IsValid() || !SelectedBudgetPlatform->Name.Equals(InBudgetPlatform->Name, ESearchCase::IgnoreCase))
 		{
 			SelectedBudgetPlatform = InBudgetPlatform;
-			ApplyBudgetToNodes();
+			bShouldUpdateBudgets = true;
 		}
 	}
 }
@@ -2254,7 +2322,7 @@ void SMemTagTreeView::OnBudgetChanged()
 	}
 
 	bIsLoadingBudget = false;
-	ApplyBudgetToNodes();
+	bShouldUpdateBudgets = true;
 }
 
 ////////////////////////////////////////////////////////////////////////////////////////////////////
@@ -2266,6 +2334,9 @@ void SMemTagTreeView::ApplyBudgetToNodes()
 		return;
 	}
 
+	FStopwatch Stopwatch;
+	Stopwatch.Start();
+
 	TSharedPtr<SMemoryProfilerWindow> ProfilerWindow = GetProfilerWindow();
 	if (!ProfilerWindow.IsValid())
 	{
@@ -2286,34 +2357,51 @@ void SMemTagTreeView::ApplyBudgetToNodes()
 		return;
 	}
 	const FMemTagBudgetMode& BudgetMode = *FoundBudgetMode;
-	
-	const TCHAR* CachedTagSetName = nullptr;
-	const FMemoryTagSet* TagSet = SharedState.GetTagSetById(TagSetFilter);
-	if (TagSet)
-	{
-		CachedTagSetName = CurrentBudget->FindString(TagSet->GetName());
-	}
+
+	const TCHAR* CachedTagSetNames[3];
+	CachedTagSetNames[0] = CurrentBudget->FindString(TEXTVIEW("Systems"));
+	CachedTagSetNames[1] = CurrentBudget->FindString(TEXTVIEW("Assets"));
+	CachedTagSetNames[2] = CurrentBudget->FindString(TEXTVIEW("AssetClasses"));
 
 	const TCHAR* CachedDefaultTrackerName = CurrentBudget->FindString(TEXTVIEW("Default"));
 	const TCHAR* CachedPlatformTrackerName = CurrentBudget->FindString(TEXTVIEW("Platform"));
 
 	const FMemTagBudgetPlatform& DefaultPlatform = BudgetMode.GetDefaultPlatform();
-	const FMemTagBudgetTagSet* DefaultTagSet = DefaultPlatform.FindTagSet(CachedDefaultTrackerName, CachedTagSetName);
-	const FMemTagBudgetTagSet* PlatformTagSet = DefaultPlatform.FindTagSet(CachedPlatformTrackerName, CachedTagSetName);
-	const FMemTagBudgetTagSet* DefaultTagSetOverride = nullptr;
-	const FMemTagBudgetTagSet* PlatformTagSetOverride = nullptr;
-	if (SelectedBudgetPlatform.IsValid())
+
+	const FMemTagBudgetTagSet* DefaultTagSet[3];
+	const FMemTagBudgetTagSet* PlatformTagSet[3];
+	const FMemTagBudgetTagSet* DefaultTagSetOverride[3];
+	const FMemTagBudgetTagSet* PlatformTagSetOverride[3];
+
+	bool bHasAnyBudget = false;
+
+	for (int32 TagSetIndex = 0; TagSetIndex < 3; ++TagSetIndex)
 	{
-		const TCHAR* CachedPlatformName = CurrentBudget->FindString(SelectedBudgetPlatform->Name);
-		const FMemTagBudgetPlatform* PlatformOverride = BudgetMode.FindPlatformOverride(CachedPlatformName);
-		if (PlatformOverride)
+		const TCHAR* CachedTagSetName = CachedTagSetNames[TagSetIndex];
+
+		DefaultTagSet[TagSetIndex] = DefaultPlatform.FindTagSet(CachedDefaultTrackerName, CachedTagSetName);
+		PlatformTagSet[TagSetIndex] = DefaultPlatform.FindTagSet(CachedPlatformTrackerName, CachedTagSetName);
+
+		DefaultTagSetOverride[TagSetIndex] = nullptr;
+		PlatformTagSetOverride[TagSetIndex] = nullptr;
+		if (SelectedBudgetPlatform.IsValid())
 		{
-			DefaultTagSetOverride = PlatformOverride->FindTagSet(CachedDefaultTrackerName, CachedTagSetName);
-			PlatformTagSetOverride = PlatformOverride->FindTagSet(CachedPlatformTrackerName, CachedTagSetName);
+			const TCHAR* CachedPlatformName = CurrentBudget->FindString(SelectedBudgetPlatform->Name);
+			const FMemTagBudgetPlatform* PlatformOverride = BudgetMode.FindPlatformOverride(CachedPlatformName);
+			if (PlatformOverride)
+			{
+				DefaultTagSetOverride[TagSetIndex] = PlatformOverride->FindTagSet(CachedDefaultTrackerName, CachedTagSetName);
+				PlatformTagSetOverride[TagSetIndex] = PlatformOverride->FindTagSet(CachedPlatformTrackerName, CachedTagSetName);
+			}
 		}
+
+		bHasAnyBudget = bHasAnyBudget || (DefaultTagSet[TagSetIndex] != nullptr);
+		bHasAnyBudget = bHasAnyBudget || (PlatformTagSet[TagSetIndex] != nullptr);
+		bHasAnyBudget = bHasAnyBudget || (DefaultTagSetOverride[TagSetIndex] != nullptr);
+		bHasAnyBudget = bHasAnyBudget || (PlatformTagSetOverride[TagSetIndex] != nullptr);
 	}
 
-	if (!DefaultTagSet && !PlatformTagSet && !DefaultTagSetOverride && !PlatformTagSetOverride)
+	if (!bHasAnyBudget)
 	{
 		ResetBudgetForAllNodes();
 		return;
@@ -2333,11 +2421,19 @@ void SMemTagTreeView::ApplyBudgetToNodes()
 			continue;
 		}
 
+		if (MemTagNode.GetMemTagSetId() != TagSetFilter)
+		{
+			continue;
+		}
+
+		int32 TagSetIndex = (int32)MemoryTag->GetTagSetId();
+		check(TagSetIndex >= 0 && TagSetIndex < 3);
+
 		const TCHAR* CachedTagName = CurrentBudget->FindString(MemoryTag->GetStatFullName());
 
 		bool bIsDefaultTracker = MemoryTag->GetTrackerId() == SharedState.GetDefaultTrackerId();
 
-		const FMemTagBudgetTagSet* BudgetTagSetOverride = bIsDefaultTracker ? DefaultTagSetOverride : PlatformTagSetOverride;
+		const FMemTagBudgetTagSet* BudgetTagSetOverride = bIsDefaultTracker ? DefaultTagSetOverride[TagSetIndex] : PlatformTagSetOverride[TagSetIndex];
 		if (BudgetTagSetOverride)
 		{
 			const int64* ValuePtr = BudgetTagSetOverride->FindValue(CachedTagName);
@@ -2349,7 +2445,7 @@ void SMemTagTreeView::ApplyBudgetToNodes()
 			}
 		}
 
-		const FMemTagBudgetTagSet* BudgetTagSet = bIsDefaultTracker ? DefaultTagSet : PlatformTagSet;
+		const FMemTagBudgetTagSet* BudgetTagSet = bIsDefaultTracker ? DefaultTagSet[TagSetIndex] : PlatformTagSet[TagSetIndex];
 		if (BudgetTagSet)
 		{
 			const int64* ValuePtr = BudgetTagSet->FindValue(CachedTagName);
@@ -2364,6 +2460,13 @@ void SMemTagTreeView::ApplyBudgetToNodes()
 		MemTagNode.ResetSizeBudget();
 		UpdateHighThreshold(MemTagNode, SharedState);
 	}
+
+	Stopwatch.Stop();
+	const double TotalTime = Stopwatch.GetAccumulatedTime();
+	if (TotalTime > 0.01)
+	{
+		UE_LOG(LogMemoryProfiler, Log, TEXT("[LLM Tags] Budgets updated in %.4fs"), TotalTime);
+	}
 }
 
 ////////////////////////////////////////////////////////////////////////////////////////////////////
diff --git a/Engine/Source/Developer/TraceInsights/Private/Insights/MemoryProfiler/Widgets/SMemTagTreeView.h b/Engine/Source/Developer/TraceInsights/Private/Insights/MemoryProfiler/Widgets/SMemTagTreeView.h
--- a/Engine/Source/Developer/TraceInsights/Private/Insights/MemoryProfiler/Widgets/SMemTagTreeView.h
+++ b/Engine/Source/Developer/TraceInsights/Private/Insights/MemoryProfiler/Widgets/SMemTagTreeView.h
@@ -170,6 +170,7 @@ public:
 	 */
 	virtual void RebuildTree(bool bResync);
 
+	bool CheckIfShouldUpdateStats() const;
 	void UpdateStats();
 
 	////////////////////////////////////////////////////////////////////////////////////////////////////
@@ -369,6 +370,9 @@ private:
 	// Sample Count, Min, Max, Average
 	double StatsStartTime = 0.0;
 	double StatsEndTime = 0.0;
+
+	bool bShouldUpdateStats = false;
+	bool bShouldUpdateBudgets = false;
 };
 
 ////////////////////////////////////////////////////////////////////////////////////////////////////
diff --git a/Engine/Source/Developer/TraceInsights/Private/Insights/TimingProfiler/Tracks/TimeRulerTrack.cpp b/Engine/Source/Developer/TraceInsights/Private/Insights/TimingProfiler/Tracks/TimeRulerTrack.cpp
--- a/Engine/Source/Developer/TraceInsights/Private/Insights/TimingProfiler/Tracks/TimeRulerTrack.cpp
+++ b/Engine/Source/Developer/TraceInsights/Private/Insights/TimingProfiler/Tracks/TimeRulerTrack.cpp
@@ -151,6 +151,9 @@ void FTimeRulerTrack::StartScrubbing(TSharedRef<FTimeMarker> InTimeMarker)
 	TimeMarkers.Remove(InTimeMarker);
 	TimeMarkers.Add(InTimeMarker);
 
+	ScrubbingTimeMarker = InTimeMarker;
+	ScrubbingTimeMarker->StartDragging();
+
 	bIsScrubbing = true;
 }
 
@@ -160,6 +163,7 @@ void FTimeRulerTrack::StopScrubbing()
 {
 	if (bIsScrubbing)
 	{
+		ScrubbingTimeMarker->StopDragging();
 		bIsScrubbing = false;
 	}
 }
diff --git a/Engine/Source/Developer/TraceInsights/Private/Insights/TimingProfiler/Tracks/TimeRulerTrack.h b/Engine/Source/Developer/TraceInsights/Private/Insights/TimingProfiler/Tracks/TimeRulerTrack.h
--- a/Engine/Source/Developer/TraceInsights/Private/Insights/TimingProfiler/Tracks/TimeRulerTrack.h
+++ b/Engine/Source/Developer/TraceInsights/Private/Insights/TimingProfiler/Tracks/TimeRulerTrack.h
@@ -80,6 +80,8 @@ private:
 	 */
 	TArray<TSharedRef<FTimeMarker>> TimeMarkers;
 
+	TSharedPtr<FTimeMarker> ScrubbingTimeMarker;
+
 	/** True if the user is currently dragging a time marker. */
 	bool bIsScrubbing;
 };
diff --git a/Engine/Source/Developer/TraceInsights/Private/Insights/TimingProfiler/ViewModels/TimeMarker.h b/Engine/Source/Developer/TraceInsights/Private/Insights/TimingProfiler/ViewModels/TimeMarker.h
--- a/Engine/Source/Developer/TraceInsights/Private/Insights/TimingProfiler/ViewModels/TimeMarker.h
+++ b/Engine/Source/Developer/TraceInsights/Private/Insights/TimingProfiler/ViewModels/TimeMarker.h
@@ -39,6 +39,10 @@ public:
 	bool IsHighlighted() const { return bIsHighlighted; }
 	void SetHighlighted(bool bOnOff) { bIsHighlighted = bOnOff; }
 
+	bool IsDragging() const { return bIsDragging; }
+	void StartDragging() { bIsDragging = true; }
+	void StopDragging() { bIsDragging = false; }
+
 	float GetCrtTextWidth() const { return CrtTextWidth; }
 	void SetCrtTextWidthAnimated(const float InTextWidth) const { CrtTextWidth = CrtTextWidth * 0.6f + InTextWidth * 0.4f; }
 
@@ -48,6 +52,7 @@ private:
 	FLinearColor Color = FLinearColor(1.0f, 1.0f, 1.0f, 0.5f);
 	bool bIsVisible = true;
 	bool bIsHighlighted = false;
+	bool bIsDragging = false;
 
 	// Smoothed time marker text width to avoid flickering
 	mutable float CrtTextWidth = 0.0f;
diff --git a/Engine/Source/Developer/TraceServices/Private/Analyzers/MemoryAnalysis.cpp b/Engine/Source/Developer/TraceServices/Private/Analyzers/MemoryAnalysis.cpp
--- a/Engine/Source/Developer/TraceServices/Private/Analyzers/MemoryAnalysis.cpp
+++ b/Engine/Source/Developer/TraceServices/Private/Analyzers/MemoryAnalysis.cpp
@@ -30,6 +30,15 @@ void FMemoryAnalyzer::OnAnalysisBegin(const FOnAnalysisContext& Context)
 	Builder.RouteEvent(RouteId_TagValue, "LLM", "TagValue");
 }
 
+void FMemoryAnalyzer::OnAnalysisEnd()
+{
+	if (Provider)
+	{
+		FProviderEditScopeLock _(*Provider);
+		Provider->OnAnalysisCompleted();
+	}
+}
+
 bool FMemoryAnalyzer::OnEvent(uint16 RouteId, EStyle Style, const FOnEventContext& Context)
 {
 	if (!Provider)
@@ -39,8 +48,6 @@ bool FMemoryAnalyzer::OnEvent(uint16 RouteId, EStyle Style, const FOnEventContex
 
 	LLM_SCOPE_BYNAME(TEXT("Insights/FMemoryAnalyzer"));
 
-	FAnalysisSessionEditScope _(Session);
-
 	const auto& EventData = Context.EventData;
 	switch (RouteId)
 	{
@@ -50,6 +57,8 @@ bool FMemoryAnalyzer::OnEvent(uint16 RouteId, EStyle Style, const FOnEventContex
 		const int64 ParentId = EventData.GetValue<int64>("ParentId");
 		const uint8 TagSetId = EventData.GetValue<uint8>("TagSetId"); // added in UE 5.6
 		FString Name = FTraceAnalyzerUtils::LegacyAttachmentString<TCHAR>("Name", Context);
+
+		FProviderEditScopeLock _(*Provider);
 		Provider->AddTagSpec(TagId, Name, ParentId, (FMemoryTagSetId)TagSetId);
 	}
 	break;
@@ -57,14 +66,18 @@ bool FMemoryAnalyzer::OnEvent(uint16 RouteId, EStyle Style, const FOnEventContex
 	{
 		uint8 TrackerId = EventData.GetValue<uint8>("TrackerId");
 		FString Name = FTraceAnalyzerUtils::LegacyAttachmentString<ANSICHAR>("Name", Context);
+
+		FProviderEditScopeLock _(*Provider);
 		Provider->AddTrackerSpec((FMemoryTrackerId)TrackerId, Name);
 	}
 	break;
-	case RouteId_TagSetSpec:
+	case RouteId_TagSetSpec: // added in UE 5.6
 	{
 		const uint8 TagSetId = EventData.GetValue<uint8>("TagSetId");
 		FString Name;
 		EventData.GetString("Name", Name);
+
+		FProviderEditScopeLock _(*Provider);
 		Provider->AddTagSetSpec((FMemoryTagSetId)TagSetId, Name);
 	}
 	break;
@@ -108,10 +121,17 @@ bool FMemoryAnalyzer::OnEvent(uint16 RouteId, EStyle Style, const FOnEventContex
 			Values.Push(FMemoryTagSample{ SamplesData[TagIndex] });
 		}
 
-		Provider->AddTagSnapshot(TrackerId, Time, TagsData, Values.GetData(), TagsCount);
 		Sample++;
 
-		Session.UpdateDurationSeconds(Time);
+		{
+			FProviderEditScopeLock _(*Provider);
+			Provider->AddTagSnapshot(TrackerId, Time, TagsData, Values.GetData(), TagsCount);
+		}
+
+		{
+			FAnalysisSessionEditScope _(Session);
+			Session.UpdateDurationSeconds(Time);
+		}
 	}
 	break;
 	}
diff --git a/Engine/Source/Developer/TraceServices/Private/Analyzers/MemoryAnalysis.h b/Engine/Source/Developer/TraceServices/Private/Analyzers/MemoryAnalysis.h
--- a/Engine/Source/Developer/TraceServices/Private/Analyzers/MemoryAnalysis.h
+++ b/Engine/Source/Developer/TraceServices/Private/Analyzers/MemoryAnalysis.h
@@ -20,6 +20,7 @@ public:
 	FMemoryAnalyzer(IAnalysisSession& Session, FMemoryProvider* InProvider);
 	~FMemoryAnalyzer();
 	virtual void OnAnalysisBegin(const FOnAnalysisContext& Context) override;
+	virtual void OnAnalysisEnd() override;
 	virtual bool OnEvent(uint16 RouteId, EStyle Style, const FOnEventContext& Context) override;
 
 private:
diff --git a/Engine/Source/Developer/TraceServices/Private/Model/Memory.cpp b/Engine/Source/Developer/TraceServices/Private/Model/Memory.cpp
--- a/Engine/Source/Developer/TraceServices/Private/Model/Memory.cpp
+++ b/Engine/Source/Developer/TraceServices/Private/Model/Memory.cpp
@@ -3,18 +3,25 @@
 #include "TraceServices/Model/Memory.h"
 #include "Model/MemoryPrivate.h"
 
+#include "Common/Utils.h"
+
 namespace TraceServices
 {
 
+thread_local FProviderLock::FThreadLocalState GMemoryProviderLockState;
+
 FMemoryProvider::FMemoryProvider(IAnalysisSession& InSession)
 	: Session(InSession)
 	, SnapshotTimes(Session.GetLinearAllocator(), 8 * 1024)
 {
-	AddTagSetSpec((FMemoryTagSetId)0, TEXT("Default"));
+	bIsInitialized = true;
+	InternalAddTagSetSpec((FMemoryTagSetId)0, TEXT("Default"));
 }
 
 FMemoryProvider::~FMemoryProvider()
 {
+	bIsInitialized = false;
+
 	for (FTrackerData* Tracker : AvailableTrackers)
 	{
 		if (Tracker)
@@ -46,6 +53,12 @@ FMemoryProvider::~FMemoryProvider()
 }
 
 void FMemoryProvider::AddTrackerSpec(FMemoryTrackerId TrackerId, const FString& Name)
+{
+	EditAccessCheck();
+	InternalAddTrackerSpec(TrackerId, Name);
+}
+
+void FMemoryProvider::InternalAddTrackerSpec(FMemoryTrackerId TrackerId, const FString& Name)
 {
 	if (TrackerId < 0 || TrackerId >= FMemoryTrackerInfo::MaxTrackers)
 	{
@@ -90,11 +103,17 @@ FMemoryProvider::FTrackerData* FMemoryProvider::GetOrAddTracker(FMemoryTrackerId
 	{
 		return Tracker;
 	}
-	AddTagSetSpec(TrackerId, TEXT("<unknown>"));
+	InternalAddTrackerSpec(TrackerId, TEXT("<unknown>"));
 	return AvailableTrackers[TrackerId];
 }
 
 void FMemoryProvider::AddTagSetSpec(FMemoryTagSetId TagSetId, const FString& Name)
+{
+	EditAccessCheck();
+	InternalAddTagSetSpec(TagSetId, Name);
+}
+
+void FMemoryProvider::InternalAddTagSetSpec(FMemoryTagSetId TagSetId, const FString& Name)
 {
 	if (TagSetId < 0 || TagSetId >= FMemoryTagSetInfo::MaxTagSets)
 	{
@@ -140,11 +159,17 @@ FMemoryProvider::FTagSetData* FMemoryProvider::GetOrAddTagSet(FMemoryTagSetId Ta
 	{
 		return TagSet;
 	}
-	AddTagSetSpec(TagSetId, TEXT("<unknown>"));
+	InternalAddTagSetSpec(TagSetId, TEXT("<unknown>"));
 	return AvailableTagSets[TagSetId];
 }
 
 void FMemoryProvider::AddTagSpec(FMemoryTagId TagId, const FString& Name, FMemoryTagId ParentTagId, FMemoryTagSetId TagSetId)
+{
+	EditAccessCheck();
+	InternalAddTagSpec(TagId, Name, ParentTagId, TagSetId);
+}
+
+void FMemoryProvider::InternalAddTagSpec(FMemoryTagId TagId, const FString& Name, FMemoryTagId ParentTagId, FMemoryTagSetId TagSetId)
 {
 	if (TagId == FMemoryTagInfo::InvalidTagId || TagId == -1)
 	{
@@ -187,6 +212,8 @@ void FMemoryProvider::AddTagSpec(FMemoryTagId TagId, const FString& Name, FMemor
 
 void FMemoryProvider::AddTagSnapshot(FMemoryTrackerId TrackerId, double Time, const int64* Tags, const FMemoryTagSample* Values, uint32 TagCount)
 {
+	EditAccessCheck();
+
 	if (TrackerId < 0 || TrackerId >= FMemoryTrackerInfo::MaxTrackers)
 	{
 		// invalid tracker id
@@ -206,8 +233,8 @@ void FMemoryProvider::AddTagSnapshot(FMemoryTrackerId TrackerId, double Time, co
 		FTagSampleData* TagSamples = Tracker->Samples.FindRef(TagId);
 		if (!TagSamples)
 		{
-			FTagSampleData* TagSampleData = new FTagSampleData(Session.GetLinearAllocator());
-			TagSamples = Tracker->Samples.Add(TagId, TagSampleData);
+			TagSamples = new FTagSampleData(Session.GetLinearAllocator());
+			Tracker->Samples.Add(TagId, TagSamples);
 		}
 
 		auto& TagValues = TagSamples->Values;
@@ -248,18 +275,31 @@ void FMemoryProvider::AddTagSnapshot(FMemoryTrackerId TrackerId, double Time, co
 	}
 }
 
+void FMemoryProvider::OnAnalysisCompleted()
+{
+	EditAccessCheck();
+
+	bIsCompleted = true;
+
+	UE_LOG(LogTraceServices, Log, TEXT("[MemTags] Analysis completed (%u trackers, %u tag sets, %d tags)."), NumTrackers, NumTagSets, AvailableTags.Num());
+}
+
 uint32 FMemoryProvider::GetTagSerial() const
 {
+	ReadAccessCheck();
 	return TagsSerial;
 }
 
 uint32 FMemoryProvider::GetTagCount() const
 {
+	ReadAccessCheck();
 	return AvailableTags.Num();
 }
 
 void FMemoryProvider::EnumerateTags(TFunctionRef<void(const FMemoryTagInfo&)> Callback) const
 {
+	ReadAccessCheck();
+
 	for (const FMemoryTagInfo* Tag : AvailableTags)
 	{
 		Callback(*Tag);
@@ -268,6 +308,8 @@ void FMemoryProvider::EnumerateTags(TFunctionRef<void(const FMemoryTagInfo&)> Ca
 
 void FMemoryProvider::EnumerateTags(FMemoryTagSetId TagSetId, TFunctionRef<void(const FMemoryTagInfo&)> Callback) const
 {
+	ReadAccessCheck();
+
 	FTagSetData* TagSet = GetTagSet(TagSetId);
 	if (!TagSet)
 	{
@@ -282,16 +324,20 @@ void FMemoryProvider::EnumerateTags(FMemoryTagSetId TagSetId, TFunctionRef<void(
 
 const FMemoryTagInfo* FMemoryProvider::GetTag(FMemoryTagId TagId) const
 {
+	ReadAccessCheck();
 	return TagMap.FindRef(TagId);
 }
 
 uint32 FMemoryProvider::GetTrackerCount() const
 {
+	ReadAccessCheck();
 	return NumTrackers;
 }
 
 void FMemoryProvider::EnumerateTrackers(TFunctionRef<void(const FMemoryTrackerInfo&)> Callback) const
 {
+	ReadAccessCheck();
+
 	for (const FTrackerData* Tracker : AvailableTrackers)
 	{
 		if (Tracker)
@@ -303,11 +349,14 @@ void FMemoryProvider::EnumerateTrackers(TFunctionRef<void(const FMemoryTrackerIn
 
 uint32 FMemoryProvider::GetTagSetCount() const
 {
+	ReadAccessCheck();
 	return NumTagSets;
 }
 
 void FMemoryProvider::EnumerateTagSets(TFunctionRef<void(const FMemoryTagSetInfo&)> Callback) const
 {
+	ReadAccessCheck();
+
 	for (const FTagSetData* TagSet : AvailableTagSets)
 	{
 		if (TagSet)
@@ -319,6 +368,8 @@ void FMemoryProvider::EnumerateTagSets(TFunctionRef<void(const FMemoryTagSetInfo
 
 uint64 FMemoryProvider::GetTagSampleCount(FMemoryTrackerId TrackerId, FMemoryTagId TagId) const
 {
+	ReadAccessCheck();
+
 	FTrackerData* Tracker = GetTracker(TrackerId);
 	if (!Tracker)
 	{
@@ -338,6 +389,8 @@ uint64 FMemoryProvider::GetTagSampleCount(FMemoryTrackerId TrackerId, FMemoryTag
 
 void FMemoryProvider::EnumerateTagSamples(FMemoryTrackerId TrackerId, FMemoryTagId TagId, double StartTime, double EndTime, bool bIncludeRangeNeighbours, TFunctionRef<void(double Time, double Duration, const FMemoryTagSample&)> Callback) const
 {
+	ReadAccessCheck();
+
 	FTrackerData* Tracker = GetTracker(TrackerId);
 	if (!Tracker)
 	{
diff --git a/Engine/Source/Developer/TraceServices/Private/Model/MemoryPrivate.h b/Engine/Source/Developer/TraceServices/Private/Model/MemoryPrivate.h
--- a/Engine/Source/Developer/TraceServices/Private/Model/MemoryPrivate.h
+++ b/Engine/Source/Developer/TraceServices/Private/Model/MemoryPrivate.h
@@ -5,22 +5,29 @@
 #include "UObject/NameTypes.h"
 
 #include "Common/PagedArray.h"
+#include "Common/ProviderLock.h"
 #include "TraceServices/Model/Memory.h"
 
 namespace TraceServices
 {
 
-class FMemoryProvider : public IMemoryProvider
+extern thread_local FProviderLock::FThreadLocalState GMemoryProviderLockState;
+
+class FMemoryProvider : public IMemoryProvider, public IEditableProvider
 {
 public:
 	explicit FMemoryProvider(IAnalysisSession& Session);
 	virtual ~FMemoryProvider();
 
-	void AddTrackerSpec(FMemoryTrackerId TrackerId, const FString& Name);
-	void AddTagSetSpec(FMemoryTagSetId TagSetId, const FString& Name);
-	void AddTagSpec(FMemoryTagId Tag, const FString& Name, FMemoryTagId ParentTag, FMemoryTagSetId TagSet);
+	//////////////////////////////////////////////////
+	// Read operations
 
-	void AddTagSnapshot(FMemoryTrackerId TrackerId, double Time, const int64* Tags, const FMemoryTagSample* Values, uint32 TagCount);
+	virtual void BeginRead() const override { Lock.BeginRead(GMemoryProviderLockState); }
+	virtual void EndRead() const override { Lock.EndRead(GMemoryProviderLockState); }
+	virtual void ReadAccessCheck() const override { Lock.ReadAccessCheck(GMemoryProviderLockState); }
+
+	virtual bool IsInitialized() const override { return bIsInitialized; }
+	virtual bool IsCompleted() const override { return bIsCompleted; }
 
 	virtual uint32 GetTagSerial() const override;
 	virtual uint32 GetTagCount() const override;
@@ -44,6 +51,23 @@ public:
 		bool bIncludeRangeNeighbours,
 		TFunctionRef<void(double Time, double Duration, const FMemoryTagSample&)> Callback) const override;
 
+	//////////////////////////////////////////////////
+	// Edit operations
+
+	virtual void BeginEdit() const override { Lock.BeginWrite(GMemoryProviderLockState); }
+	virtual void EndEdit() const override { Lock.EndWrite(GMemoryProviderLockState); }
+	virtual void EditAccessCheck() const override { Lock.WriteAccessCheck(GMemoryProviderLockState); }
+
+	void AddTrackerSpec(FMemoryTrackerId TrackerId, const FString& Name);
+	void AddTagSetSpec(FMemoryTagSetId TagSetId, const FString& Name);
+	void AddTagSpec(FMemoryTagId Tag, const FString& Name, FMemoryTagId ParentTag, FMemoryTagSetId TagSet);
+
+	void AddTagSnapshot(FMemoryTrackerId TrackerId, double Time, const int64* Tags, const FMemoryTagSample* Values, uint32 TagCount);
+
+	void OnAnalysisCompleted();
+
+	//////////////////////////////////////////////////
+
 private:
 	struct FTagSampleData
 	{
@@ -89,12 +113,17 @@ private:
 	};
 
 private:
+	void InternalAddTrackerSpec(FMemoryTrackerId TrackerId, const FString& Name);
+	void InternalAddTagSetSpec(FMemoryTagSetId TagSetId, const FString& Name);
+	void InternalAddTagSpec(FMemoryTagId Tag, const FString& Name, FMemoryTagId ParentTag, FMemoryTagSetId TagSet);
 	FTrackerData* GetTracker(FMemoryTrackerId TrackerId) const;
 	FTrackerData* GetOrAddTracker(FMemoryTrackerId TrackerId);
 	FTagSetData* GetTagSet(FMemoryTagSetId TagSetId) const;
 	FTagSetData* GetOrAddTagSet(FMemoryTagSetId TagSetId);
 
 private:
+	mutable FProviderLock Lock;
+
 	IAnalysisSession& Session;
 
 	TArray<FTrackerData*> AvailableTrackers;
@@ -110,6 +139,9 @@ private:
 	TPagedArray<double> SnapshotTimes;
 
 	uint32 TagsSerial = 0;
+
+	std::atomic<bool> bIsInitialized = false;
+	std::atomic<bool> bIsCompleted = false;
 };
 
 } // namespace TraceServices
diff --git a/Engine/Source/Developer/TraceServices/Public/TraceServices/Model/Memory.h b/Engine/Source/Developer/TraceServices/Public/TraceServices/Model/Memory.h
--- a/Engine/Source/Developer/TraceServices/Public/TraceServices/Model/Memory.h
+++ b/Engine/Source/Developer/TraceServices/Public/TraceServices/Model/Memory.h
@@ -72,6 +72,21 @@ class IMemoryProvider : public IProvider
 public:
 	virtual ~IMemoryProvider() = default;
 
+	virtual void BeginRead() const = 0;
+	virtual void EndRead() const = 0;
+	virtual void ReadAccessCheck() const = 0;
+
+	/**
+	 * @returns Whether the provider is initialized or not.
+	 */
+	virtual bool IsInitialized() const = 0;
+
+	/**
+	 * When this provider is completed it cannot be further modified.
+	 * @returns Whether this provider is completed or not.
+	 */
+	virtual bool IsCompleted() const = 0;
+
 	/**
 	 * Unique serial index that changes when new tags are registered or when the Trackers flags is updated for a tag.
 	 * @return Unique index
