Fix for gpu header track location not set.

diff --git a/Engine/Source/Developer/TraceInsights/Private/Insights/TimingProfiler/ViewModels/ThreadTimingSharedState.cpp b/Engine/Source/Developer/TraceInsights/Private/Insights/TimingProfiler/ViewModels/ThreadTimingSharedState.cpp
--- a/Engine/Source/Developer/TraceInsights/Private/Insights/TimingProfiler/ViewModels/ThreadTimingSharedState.cpp
+++ b/Engine/Source/Developer/TraceInsights/Private/Insights/TimingProfiler/ViewModels/ThreadTimingSharedState.cpp
@@ -406,6 +406,7 @@ void FThreadTimingSharedState::Tick(Timing::ITimingViewSession& InSession, const
 						// Create the GPU Work track and attach it to the GPU queue track.
 						const FString WorkTrackName = FString::Printf(TEXT("GPU%u - %s %u - WORK"), QueueInfo.GPU, QueueInfo.Name, QueueInfo.Index);
 						TSharedRef<FGpuQueueWorkTimingTrack> WorkTrack = MakeShared<FGpuQueueWorkTimingTrack>(*this, WorkTrackName, QueueInfo.WorkTimelineIndex, QueueInfo.Id);
+						WorkTrack->SetLocation(Track->GetLocation());
 						WorkTrack->SetParentTrack(Track);
 						Track->AddChildTrack(WorkTrack);
 					}
@@ -414,6 +415,7 @@ void FThreadTimingSharedState::Tick(Timing::ITimingViewSession& InSession, const
 					{
 						const FString FencesTrackName = FString::Printf(TEXT("GPU%u - %s %u - Fences"), QueueInfo.GPU, QueueInfo.Name, QueueInfo.Index);
 						TSharedRef<FGpuFencesTimingTrack> FencesTrack = MakeShared<FGpuFencesTimingTrack>(*this, FencesTrackName, QueueInfo.Id);
+						FencesTrack->SetLocation(Track->GetLocation());
 						FencesTrack->SetParentTrack(Track);
 						Track->AddChildTrack(FencesTrack);
 					}
@@ -882,6 +884,7 @@ void FThreadTimingSharedState::AddGpuWorkChildTracks()
 			// Create the GPU Work track and attach it to the GPU queue track.
 			const FString WorkTrackName = FString::Printf(TEXT("GPU%u - %s %u - WORK"), QueueInfo.GPU, QueueInfo.Name, QueueInfo.Index);
 			TSharedRef<FGpuQueueWorkTimingTrack> WorkTrack = MakeShared<FGpuQueueWorkTimingTrack>(*this, WorkTrackName, QueueInfo.WorkTimelineIndex, QueueInfo.Id);
+			WorkTrack->SetLocation((*TrackPtrPtr)->GetLocation());
 			(*TrackPtrPtr)->AddChildTrack(WorkTrack, 0);
 			WorkTrack->SetParentTrack(*TrackPtrPtr);
 		}
@@ -927,6 +930,7 @@ void FThreadTimingSharedState::AddGpuFencesChildTracks()
 			// Create the GPU Work track and attach it to the GPU queue track.
 			const FString TrackName = FString::Printf(TEXT("GPU%u - %s %u - Fences"), QueueInfo.GPU, QueueInfo.Name, QueueInfo.Index);
 			TSharedRef<FGpuFencesTimingTrack> Track = MakeShared<FGpuFencesTimingTrack>(*this, TrackName, QueueInfo.Id);
+			Track->SetLocation((*TrackPtrPtr)->GetLocation());
 			(*TrackPtrPtr)->AddChildTrack(Track, (*TrackPtrPtr)->GetChildTracks().Num());
 			Track->SetParentTrack(*TrackPtrPtr);
 		}
diff --git a/Engine/Source/Developer/TraceInsights/Private/Insights/ViewModels/BaseTimingTrack.cpp b/Engine/Source/Developer/TraceInsights/Private/Insights/ViewModels/BaseTimingTrack.cpp
--- a/Engine/Source/Developer/TraceInsights/Private/Insights/ViewModels/BaseTimingTrack.cpp
+++ b/Engine/Source/Developer/TraceInsights/Private/Insights/ViewModels/BaseTimingTrack.cpp
@@ -74,6 +74,17 @@ void FBaseTimingTrack::UpdateChildTracksPosY(const FTimingViewLayout& Layout)
 	}
 }
 
+void FBaseTimingTrack::SetLocation(ETimingTrackLocation InLocation)
+{ 
+	for (const TSharedRef<FBaseTimingTrack>& Track : GetChildTracks())
+	{
+		Track->SetLocation(InLocation);
+	}
+
+	Location = InLocation;
+	OnLocationChanged();
+}
+
 PRAGMA_DISABLE_DEPRECATION_WARNINGS
 TSharedPtr<FBaseTimingTrack> FBaseTimingTrack::GetChildTrack() const
 {
diff --git a/Engine/Source/Developer/TraceInsights/Public/Insights/ViewModels/BaseTimingTrack.h b/Engine/Source/Developer/TraceInsights/Public/Insights/ViewModels/BaseTimingTrack.h
--- a/Engine/Source/Developer/TraceInsights/Public/Insights/ViewModels/BaseTimingTrack.h
+++ b/Engine/Source/Developer/TraceInsights/Public/Insights/ViewModels/BaseTimingTrack.h
@@ -167,7 +167,7 @@ public:
 
 	ETimingTrackLocation GetValidLocations() const { return ValidLocations; }
 	ETimingTrackLocation GetLocation() const { return Location; }
-	void SetLocation(ETimingTrackLocation InLocation) { Location = InLocation; OnLocationChanged(); }
+	void SetLocation(ETimingTrackLocation InLocation);
 	virtual void OnLocationChanged() { SetDirtyFlag(); }
 
 	int32 GetOrder() const { return Order; }
