[Insights]
* Memory Insights: Added support for defining tag grouping in the xml budget file, for each tag set. Groups are defined using include/exclude regex rules. Groups can also have a budget value.
* Memory Insights: Added support for FMemTagBudgetGroupNode (group nodes with support for budget) in the LLM Tag tree/table. These have budget value and colors the Size columns like LLM tag nodes, with regard to budget value.
* Memory Insights: Changed budget xml spec to define tags inside <set> and <tracker> xml elements.
* Memory Insights: Updated the DefaultMemoryBudgets.xml to include example of the supported functionality.

diff --git a/Engine/Programs/UnrealInsights/Config/DefaultMemoryBudgets.xml b/Engine/Programs/UnrealInsights/Config/DefaultMemoryBudgets.xml
--- a/Engine/Programs/UnrealInsights/Config/DefaultMemoryBudgets.xml
+++ b/Engine/Programs/UnrealInsights/Config/DefaultMemoryBudgets.xml
@@ -1,73 +1,216 @@
 <MemoryBudget>
 
-	<!-- 1 KiB = 2^10 Bytes; 1 KB = 1000 Bytes -->
-	<!-- 1 MiB = 2^20 Bytes; 1 MB = 10^6 Bytes -->
-	<!-- 1 GiB = 2^30 Bytes; 1 GB = 10^9 Bytes -->
+	<!--
+		1 KB = 1000 Bytes
+		1 MB = 10^6 Bytes = 1,000,000 Bytes
+		1 GB = 10^9 Bytes = 1,000,000,000 Bytes
+		1 KiB = 2^10 Bytes = 1024 Bytes
+		1 MiB = 2^20 Bytes = 1,048,576 Bytes
+		1 GiB = 2^30 Bytes = 1,073,741,824 Bytes
+	-->
+
+	<!--
+		<budget name=""> ...sets...platforms... </budget>
+		where:
+			name = the budget mode name
+			sets = <set/> nodes
+			platforms = <platform/> nodes
+
+		<set name=""> ...trackers... </set>
+		where:
+			name = "System", "Asset" or "AssetClass"
+			trackers = <tracker/> nodes
+			
+		<tracker name=""> ...tags... </tracker>
+		where:
+			name = "Default" or "Platform"
+			tags = <tag/> nodes
+			groups = <group/> nodes
+
+		<grouping> ...groups... </grouping>
+		where:
+			groups = <group/> nodes
+
+		<tag name="" mem-max="0"/>
+		where:
+			name = the LLM tag name
+			mem-max = the memory budget value (high threshold value); in [bytes], if no unit is specified
+
+		<group name="" include="" exclude="" min-max=""/>
+		where:
+			name = the group name
+			include = the regex for tag names to be included
+			exclude = the regex for tag names to be excluded
+			mem-max = the memory budget value (high threshold value); in [bytes], if no unit is specified
+
+		<platform name=""> ...sets... </platform>
+		where:
+			name = the platform name
+			sets = <set/> nodes
+	-->
 
 	<budget name="Default">
 
-		<tag name="Animation"                     mem-max="50 MiB"/>
-		<tag name="AsyncLoading"                  mem-max="225 MiB"/>
-		<tag name="Audio"                         mem-max="180 MiB"/>
-		<tag name="FName"                         mem-max="6 MiB"/>
-		<tag name="GC"                            mem-max="20 MiB"/>
-		<tag name="Materials"                     mem-max="30 MiB"/>
-		<tag name="Meshes"                        mem-max="150 MiB"/>
-		<tag name="Meshes/StaticMesh"             mem-max="20 MiB"/>
-		<tag name="Niagara"                       mem-max="100 MiB"/>
-		<tag name="Particles"                     mem-max="100 MiB"/>
-		<tag name="Physics"                       mem-max="500 MiB"/>
-		<tag name="RHIMisc"                       mem-max="50 MiB"/>
-		<tag name="RenderTargets"                 mem-max="750 MiB"/>
-		<tag name="SceneRender"                   mem-max="50 MiB"/>
-		<tag name="Shaders"                       mem-max="100 MiB"/>
-		<tag name="Stats"                         mem-max="5 MiB"/>
-		<tag name="TextureMetaData"               mem-max="5 MiB"/>
-		<tag name="Textures"                      mem-max="2 GiB"/>
-		<tag name="UI"                            mem-max="100 MiB"/>
-		<tag name="UObject"                       mem-max="250 MiB"/>
+		<!-- System tag set -->
+		<set name="System">
+			<tracker name="Default">
+				<tag name="Animation"                       mem-max="90 MiB"/>
+				<tag name="AssetManager"                    mem-max="30 MiB"/>
+				<tag name="AssetRegistry"                   mem-max="200 MiB"/>
+				<tag name="AsyncLoading"                    mem-max="225 MiB"/>
+				<tag name="Audio"                           mem-max="180 MiB"/>
+				<tag name="EngineInit"                      mem-max="16 MiB"/>
+				<tag name="EnginePreInit"                   mem-max="100 MiB"/>
+				<tag name="EngineMisc"                      mem-max="50 MiB"/>
+				<tag name="FileSystem"                      mem-max="100 MiB"/>
+				<tag name="FName"                           mem-max="6 MiB"/>
+				<tag name="GC"                              mem-max="20 MiB"/>
+				<tag name="Localization"                    mem-max="50 MiB"/>
+				<tag name="Materials"                       mem-max="30 MiB"/>
+				<tag name="Meshes"                          mem-max="150 MiB"/>
+				<tag name="Meshes/StaticMesh"               mem-max="20 MiB"/>
+				<tag name="Meshes/StaticMesh/Serialize"     mem-max="10 MiB"/>
+				<tag name="Meshes/SkeletalMesh"             mem-max="50 MiB"/>
+				<tag name="Meshes/SkeletalMesh/Serialize"   mem-max="50 MiB"/>
+				<tag name="Niagara"                         mem-max="100 MiB"/>
+				<tag name="Particles"                       mem-max="100 MiB"/>
+				<tag name="Physics"                         mem-max="500 MiB"/>
+				<tag name="RenderTargets"                   mem-max="750 MiB"/>
+				<tag name="RHIMisc"                         mem-max="50 MiB"/>
+				<tag name="SceneRender"                     mem-max="50 MiB"/>
+				<tag name="Shaders"                         mem-max="100 MiB"/>
+				<tag name="Stats"                           mem-max="5 MiB"/>
+				<tag name="TextureMetaData"                 mem-max="5 MiB"/>
+				<tag name="Textures"                        mem-max="2 GiB"/>
+				<tag name="UI"                              mem-max="100 MiB"/>
+				<tag name="UObject"                         mem-max="250 MiB"/>
+				<tag name="ThreadStack"                     mem-max="300 MiB"/>
+				<tag name="Untagged"                        mem-max="1 MiB"/>
+				<tag name="Total"                           mem-max="2.11 GiB"/>
+				<tag name="TrackedTotal"                    mem-max="2.1 GiB"/>
+				<tag name="Untracked"                       mem-max="10 MiB"/>
+			</tracker>
+			<tracker name="Platform">
+				<tag name="ThreadStack"                     mem-max="300 MiB"/>
+				<tag name="Untagged"                        mem-max="2 MiB"/>
+				<tag name="Graphics"                        mem-max="1 GiB"/>
+				<tag name="FMalloc"                         mem-max="2.5 GiB"/>
+				<tag name="LLMOverhead"                     mem-max="500 MiB"/>
+				<tag name="Total"                           mem-max="2.21 GiB"/>
+				<tag name="TrackedTotal"                    mem-max="2.2 GiB"/>
+				<tag name="Untracked"                       mem-max="20 MiB"/>
+			</tracker>
+		</set>
 
-		<platform name="Platform 1">
-			<tag name="Animation"                     mem-max="100 MiB"/>
-			<tag name="Audio"                         mem-max="200 MiB"/>
-			<tag name="Meshes"                        mem-max="300 MiB"/>
-			<tag name="Textures"                      mem-max="3 GiB"/>
-		</platform>
+		<!-- Asset tag set -->
+		<set name="Asset">
+			<tracker name="Default">
+				<tag name="UntaggedAsset" mem-max="2 GiB"/>
+			</tracker>
+			<tracker name="Platform">
+				<tag name="UntaggedAsset" mem-max="3 GiB"/>
+			</tracker>
+			<grouping>
+				<!-- Asset groups -->
+				<group name="Engine"  include="^/Engine/.*"  mem-max="1 GiB"/>
+				<group name="Engine2" include="^\.\./\.\./\.\./Engine/.*"  mem-max="1 GiB"/> <!-- assets stating with "../../../Engine/" -->
+				<group name="Game"    include="^/Game/.*"    mem-max="2 GiB"/>
+				<group name="Script"  include="^/Script/.*"  mem-max="120 MiB"/>
+				<!-- A group for all other assets -->
+				<group name="Other" include=".*" exclude="UntaggedAsset" mem-max="100 MiB"/>
+			</grouping>
+		</set>
 
-		<platform name="Platform 2">
-			<tag name="Animation"                     mem-max="10 MiB"/>
-			<tag name="Audio"                         mem-max="50 MiB"/>
-			<tag name="Meshes"                        mem-max="100 MiB"/>
-			<tag name="Textures"                      mem-max="1 GiB"/>
-		</platform>
-	</budget>
+		<!-- AssetClass tag set -->
+		<set name="AssetClass">
+			<tracker name="Default">
+				<tag name="UntaggedAssetClass" mem-max="2 GiB"/>
+			</tracker>
+			<tracker name="Platform">
+				<tag name="UntaggedAssetClass" mem-max="1 GiB"/>
+			</tracker>
+			<grouping>
+				<!-- AssetClass groups -->
+				<group name="Engine"  include="^/Engine/.*"  mem-max="1.2 MiB"/>
+				<group name="Game"    include="^/Game/.*"    mem-max="2.2 GiB"/>
+				<group name="Script"  include="^/Script/.*"  mem-max="120 MiB"/>
+				<!-- A group for all other asset classes -->
+				<group name="Other" include=".*" exclude="UntaggedAssetClass" mem-max="100 MiB"/>
+			</grouping>
+		</set>
 
-	<budget name="Empty">
-	</budget>
+		<platform name="Gen 5">
+			<set name="System">
+				<tracker name="Default">
+					<tag name="Animation"                     mem-max="100 MiB"/>
+					<tag name="Audio"                         mem-max="200 MiB"/>
+					<tag name="Meshes"                        mem-max="300 MiB"/>
+					<tag name="Textures"                      mem-max="3 GiB"/>
+				</tracker>
+			</set>
+			<set name="Asset">
+				<tracker name="Default">
+					<tag name="UntaggedAsset"                 mem-max="2.1 GiB"/>
+				</tracker>
+				<tracker name="Platform">
+					<tag name="UntaggedAsset"                 mem-max="3.1 GiB"/>
+				</tracker>
+				<grouping>
+					<group name="Engine"                      mem-max="1.1 GiB"/>
+					<group name="Game"                        mem-max="2.1 GiB"/>
+				</grouping>
+			</set>
+			<set name="AssetClass">
+				<tracker name="Default">
+					<tag name="UntaggedAssetClass"            mem-max="1.1 GiB"/>
+				</tracker>
+				<tracker name="Platform">
+					<tag name="UntaggedAssetClass"            mem-max="2.1 GiB"/>
+				</tracker>
+				<grouping>
+					<group name="Engine"                      mem-max="1.1 GiB"/>
+					<group name="Game"                        mem-max="2.1 GiB"/>
+				</grouping>
+			</set>
+		</platform>
 
-	<budget name="Test1">
-		<tag name="Animation"                     mem-max="0"/>
-		<tag name="Audio"                         mem-max="123"/>
-		<tag name="Meshes"                        mem-max="2000"/>
-		<tag name="Textures"                      mem-max="1 GiB"/>
-		<platform name="Platform 1">
-			<tag name="Animation"                     mem-max="100"/>
-			<tag name="Audio"                         mem-max="1 KB"/>
-			<tag name="Meshes"                        mem-max="3 KB"/>
-			<tag name="Textures"                      mem-max="2 GB"/>
+		<platform name="Mobile">
+			<set name="System">
+				<tracker name="Default">
+					<tag name="Animation"                     mem-max="10 MiB"/>
+					<tag name="Audio"                         mem-max="50 MiB"/>
+					<tag name="Meshes"                        mem-max="100 MiB"/>
+					<tag name="Textures"                      mem-max="1 GiB"/>
+				</tracker>
+			</set>
+			<set name="Asset">
+				<tracker name="Default">
+					<tag name="UntaggedAsset"                 mem-max="0.8 GiB"/>
+				</tracker>
+				<tracker name="Platform">
+					<tag name="UntaggedAsset"                 mem-max="1.0 GiB"/>
+				</tracker>
+				<grouping>
+					<group name="Engine"                      mem-max="0.8 GiB"/>
+					<group name="Game"                        mem-max="1.8 GiB"/>
+				</grouping>
+			</set>
+			<set name="AssetClass">
+				<tracker name="Default">
+					<tag name="UntaggedAssetClass"            mem-max="0.7 GiB"/>
+				</tracker>
+				<tracker name="Platform">
+					<tag name="UntaggedAssetClass"            mem-max="1.7 GiB"/>
+				</tracker>
+				<grouping>
+					<group name="Engine"                      mem-max="0.7 MiB"/>
+					<group name="Game"                        mem-max="1.7 GiB"/>
+				</grouping>
+			</set>
 		</platform>
 	</budget>
 
-	<budget name="Test2">
-		<tag name="Audio"                         mem-max="1 MiB"/>
-		<platform name="Platform 2">
-			<tag name="Animation"                     mem-max="2 KiB"/>
-		</platform>
-		<platform name="Platform 3">
-			<tag name="Animation"                     mem-max="3 KiB"/>
-			<tag name="Audio"                         mem-max="3 MiB"/>
-			<tag name="Textures"                      mem-max="3 GiB"/>
-		</platform>
+	<budget name="Empty">
 	</budget>
 
 </MemoryBudget>
diff --git a/Engine/Source/Developer/TraceInsights/Private/Insights/MemoryProfiler/ViewModels/MemTagBudget.cpp b/Engine/Source/Developer/TraceInsights/Private/Insights/MemoryProfiler/ViewModels/MemTagBudget.cpp
--- a/Engine/Source/Developer/TraceInsights/Private/Insights/MemoryProfiler/ViewModels/MemTagBudget.cpp
+++ b/Engine/Source/Developer/TraceInsights/Private/Insights/MemoryProfiler/ViewModels/MemTagBudget.cpp
@@ -12,13 +12,70 @@
 namespace UE::Insights::MemoryProfiler
 {
 
+const FString FMemTagBudget::STR_Name(TEXT("name"));
+const FString FMemTagBudget::STR_Id(TEXT("id"));
+const FString FMemTagBudget::STR_Budget(TEXT("budget"));
+const FString FMemTagBudget::STR_Set(TEXT("set"));
+const FString FMemTagBudget::STR_Tag(TEXT("tag"));
+const FString FMemTagBudget::STR_TagMemMax(TEXT("mem-max"));
+const FString FMemTagBudget::STR_Tracker(TEXT("tracker"));
+const FString FMemTagBudget::STR_Platform(TEXT("platform"));
+const FString FMemTagBudget::STR_Grouping(TEXT("grouping"));
+const FString FMemTagBudget::STR_Group(TEXT("group"));
+const FString FMemTagBudget::STR_Include(TEXT("include"));
+const FString FMemTagBudget::STR_Exclude(TEXT("exclude"));
+
 ////////////////////////////////////////////////////////////////////////////////////////////////////
-// FMemTagBudgetTagSet
+// FMemTagBudgetGrouping
 ////////////////////////////////////////////////////////////////////////////////////////////////////
 
-void FMemTagBudgetTagSet::Reset()
+void FMemTagBudgetGrouping::Reset()
 {
-	Values.Reset();
+	for (FMemTagBudgetGroup* Group : Groups)
+	{
+		delete Group;
+	}
+	Groups.Reset();
+}
+
+////////////////////////////////////////////////////////////////////////////////////////////////////
+
+void FMemTagBudgetGrouping::EnumerateGroups(TFunctionRef<void(const TCHAR* InCachedGroupName, const FMemTagBudgetGroup& InGroup)> Callback) const
+{
+	for (const FMemTagBudgetGroup* Group : Groups)
+	{
+		Callback(Group->GetName(), *Group);
+	}
+}
+
+////////////////////////////////////////////////////////////////////////////////////////////////////
+
+const FMemTagBudgetGroup* FMemTagBudgetGrouping::FindGroup(const TCHAR* InCachedGroupName) const
+{
+	for (const FMemTagBudgetGroup* Group : Groups)
+	{
+		if (Group->GetName() == InCachedGroupName)
+		{
+			return Group;
+		}
+	}
+	return nullptr;
+}
+
+////////////////////////////////////////////////////////////////////////////////////////////////////
+
+FMemTagBudgetGroup& FMemTagBudgetGrouping::GetOrAddGroup(const TCHAR* InCachedGroupName)
+{
+	for (FMemTagBudgetGroup* Group : Groups)
+	{
+		if (Group->GetName() == InCachedGroupName)
+		{
+			return *Group;
+		}
+	}
+	FMemTagBudgetGroup* NewGroup = new FMemTagBudgetGroup(InCachedGroupName);
+	Groups.Add(NewGroup);
+	return *NewGroup;
 }
 
 ////////////////////////////////////////////////////////////////////////////////////////////////////
@@ -27,18 +84,33 @@ void FMemTagBudgetTagSet::Reset()
 
 void FMemTagBudgetTracker::Reset()
 {
-	for (auto& KV : TagSets)
+	Values.Reset();
+}
+
+////////////////////////////////////////////////////////////////////////////////////////////////////
+// FMemTagBudgetTagSet
+////////////////////////////////////////////////////////////////////////////////////////////////////
+
+void FMemTagBudgetTagSet::Reset()
+{
+	for (auto& KV : Trackers)
 	{
 		delete KV.Value;
 	}
-	TagSets.Reset();
+	Trackers.Reset();
+
+	if (Grouping)
+	{
+		delete Grouping;
+		Grouping = nullptr;
+	}
 }
 
 ////////////////////////////////////////////////////////////////////////////////////////////////////
 
-void FMemTagBudgetTracker::EnumerateTagSets(TFunctionRef<void(const TCHAR* InCachedTagSetName, const FMemTagBudgetTagSet& InTagSet)> Callback) const
+void FMemTagBudgetTagSet::EnumerateTrackers(TFunctionRef<void(const TCHAR* InCachedTrackerName, const FMemTagBudgetTracker& InTracker)> Callback) const
 {
-	for (const auto& KV : TagSets)
+	for (const auto& KV : Trackers)
 	{
 		Callback((const TCHAR*)KV.Key, *KV.Value);
 	}
@@ -46,24 +118,35 @@ void FMemTagBudgetTracker::EnumerateTagSets(TFunctionRef<void(const TCHAR* InCac
 
 ////////////////////////////////////////////////////////////////////////////////////////////////////
 
-const FMemTagBudgetTagSet* FMemTagBudgetTracker::FindTagSet(const TCHAR* InCachedTagSetName) const
+const FMemTagBudgetTracker* FMemTagBudgetTagSet::FindTracker(const TCHAR* InCachedTrackerName) const
 {
-	FMemTagBudgetTagSet*const* Found = TagSets.Find(InCachedTagSetName);
+	FMemTagBudgetTracker* const* Found = Trackers.Find(InCachedTrackerName);
 	return Found ? *Found : nullptr;
 }
 
 ////////////////////////////////////////////////////////////////////////////////////////////////////
 
-FMemTagBudgetTagSet& FMemTagBudgetTracker::GetOrAddTagSet(const TCHAR* InCachedTagSetName)
+FMemTagBudgetTracker& FMemTagBudgetTagSet::GetOrAddTracker(const TCHAR* InCachedTrackerName)
 {
-	FMemTagBudgetTagSet** Found = TagSets.Find(InCachedTagSetName);
+	FMemTagBudgetTracker** Found = Trackers.Find(InCachedTrackerName);
 	if (Found)
 	{
 		return **Found;
 	}
-	FMemTagBudgetTagSet* TagSet = new FMemTagBudgetTagSet(InCachedTagSetName);
-	TagSets.Add(InCachedTagSetName, TagSet);
-	return *TagSet;
+	FMemTagBudgetTracker* Tracker = new FMemTagBudgetTracker(InCachedTrackerName, this);
+	Trackers.Add(InCachedTrackerName, Tracker);
+	return *Tracker;
+}
+
+////////////////////////////////////////////////////////////////////////////////////////////////////
+
+FMemTagBudgetGrouping& FMemTagBudgetTagSet::GetOrCreateGrouping()
+{
+	if (!Grouping)
+	{
+		Grouping = new FMemTagBudgetGrouping();
+	}
+	return *Grouping;
 }
 
 ////////////////////////////////////////////////////////////////////////////////////////////////////
@@ -72,18 +155,18 @@ FMemTagBudgetTagSet& FMemTagBudgetTracker::GetOrAddTagSet(const TCHAR* InCachedT
 
 void FMemTagBudgetPlatform::Reset()
 {
-	for (auto& KV : Trackers)
+	for (auto& KV : TagSets)
 	{
 		delete KV.Value;
 	}
-	Trackers.Reset();
+	TagSets.Reset();
 }
 
 ////////////////////////////////////////////////////////////////////////////////////////////////////
 
-void FMemTagBudgetPlatform::EnumerateTrackers(TFunctionRef<void(const TCHAR* InCachedTrackerName, const FMemTagBudgetTracker& InTracker)> Callback) const
+void FMemTagBudgetPlatform::EnumerateTagSets(TFunctionRef<void(const TCHAR* InCachedTagSetName, const FMemTagBudgetTagSet& InTagSet)> Callback) const
 {
-	for (const auto& KV : Trackers)
+	for (const auto& KV : TagSets)
 	{
 		Callback((const TCHAR*)KV.Key, *KV.Value);
 	}
@@ -91,32 +174,24 @@ void FMemTagBudgetPlatform::EnumerateTrackers(TFunctionRef<void(const TCHAR* InC
 
 ////////////////////////////////////////////////////////////////////////////////////////////////////
 
-const FMemTagBudgetTracker* FMemTagBudgetPlatform::FindTracker(const TCHAR* InCachedTrackerName) const
+const FMemTagBudgetTagSet* FMemTagBudgetPlatform::FindTagSet(const TCHAR* InCachedTagSetName) const
 {
-	FMemTagBudgetTracker*const* Found = Trackers.Find(InCachedTrackerName);
+	FMemTagBudgetTagSet* const* Found = TagSets.Find(InCachedTagSetName);
 	return Found ? *Found : nullptr;
 }
 
 ////////////////////////////////////////////////////////////////////////////////////////////////////
 
-FMemTagBudgetTracker& FMemTagBudgetPlatform::GetOrAddTracker(const TCHAR* InCachedTrackerName)
+FMemTagBudgetTagSet& FMemTagBudgetPlatform::GetOrAddTagSet(const TCHAR* InCachedTagSetName)
 {
-	FMemTagBudgetTracker** Found = Trackers.Find(InCachedTrackerName);
+	FMemTagBudgetTagSet** Found = TagSets.Find(InCachedTagSetName);
 	if (Found)
 	{
 		return **Found;
 	}
-	FMemTagBudgetTracker* Tracker = new FMemTagBudgetTracker(InCachedTrackerName);
-	Trackers.Add(InCachedTrackerName, Tracker);
-	return *Tracker;
-}
-
-////////////////////////////////////////////////////////////////////////////////////////////////////
-
-const FMemTagBudgetTagSet* FMemTagBudgetPlatform::FindTagSet(const TCHAR* InCachedTrackerName, const TCHAR* InCachedTagSetName) const
-{
-	const FMemTagBudgetTracker* Tracker = FindTracker(InCachedTrackerName);
-	return Tracker ? Tracker->FindTagSet(InCachedTagSetName) : nullptr;
+	FMemTagBudgetTagSet* TagSet = new FMemTagBudgetTagSet(InCachedTagSetName);
+	TagSets.Add(InCachedTagSetName, TagSet);
+	return *TagSet;
 }
 
 ////////////////////////////////////////////////////////////////////////////////////////////////////
@@ -263,165 +338,326 @@ bool FMemTagBudget::LoadFromFile(const FString& InFilePath)
 
 bool FMemTagBudget::ReLoadFromFile()
 {
-	const TCHAR* DefaultCachedPlatformName = StoreString(FString(TEXT("Default")));
-
-	const TCHAR* DefaultCachedTrackerName = StoreString(FString(TEXT("Default")));
-	const TCHAR* PlatformCachedTrackerName = StoreString(FString(TEXT("Platform")));
+	DefaultCachedPlatformName = StoreString(FString(TEXT("Default")));
 
-	const TCHAR* SystemCachedTagSetName = StoreString(FString(TEXT("Systems")));
-	const TCHAR* AssetCachedTagSetName = StoreString(FString(TEXT("Assets")));
-	const TCHAR* AssetClassCachedTagSetName = StoreString(FString(TEXT("AssetClasses")));
+	DefaultCachedTrackerName = StoreString(FString(TEXT("Default")));
+	PlatformCachedTrackerName = StoreString(FString(TEXT("Platform")));
 
-	static const FString STR_Budget(TEXT("budget"));
-	static const FString STR_Platform(TEXT("platform"));
-	static const FString STR_Set(TEXT("set"));
-	static const FString STR_Tracker(TEXT("tracker"));
-	static const FString STR_Name(TEXT("name"));
-	//static const FString STR_NameRegEx(TEXT("name_regex"));
-	static const FString STR_Tag(TEXT("tag"));
-	//static const FString STR_Tag(TEXT("group"));
-	static const FString STR_TagMemMax(TEXT("mem-max"));
+	SystemCachedTagSetName = StoreString(FString(TEXT("Systems")));
+	AssetCachedTagSetName = StoreString(FString(TEXT("Assets")));
+	AssetClassCachedTagSetName = StoreString(FString(TEXT("AssetClasses")));
 
 	FXmlFile XmlFile;
 	if (!XmlFile.LoadFile(FilePath, EConstructMethod::ConstructFromFile))
 	{
-		UE_LOG(LogMemoryProfiler, Warning, TEXT("Could not load memory budget file: \"%s\""), *FilePath);
+		UE_LOG(LogMemoryProfiler, Warning, TEXT("Could not load memory budget file \"%s\"."), *FilePath);
 		return false;
 	}
 
 	FXmlNode* RootNode = XmlFile.GetRootNode();
 	if (RootNode == nullptr)
 	{
-		UE_LOG(LogMemoryProfiler, Warning, TEXT("Could not parse memory budget file as an XML file, make sure it has only one root node"));
+		UE_LOG(LogMemoryProfiler, Warning, TEXT("Could not parse memory budget file as an XML file; make sure it has only one root node."));
 		return false;
 	}
 
-	const TArray<FXmlNode*>& BudgetNodes = RootNode->GetChildrenNodes();
-	for (FXmlNode* BudgetNode : BudgetNodes)
+	const TArray<FXmlNode*>& ChildNodes = RootNode->GetChildrenNodes();
+	for (FXmlNode* ChildNode : ChildNodes)
+	{
+		if (ChildNode->GetTag() == STR_Budget)
+		{
+			ProcessBudgetNode(ChildNode);
+		}
+		else
+		{
+			UE_LOG(LogMemoryProfiler, Warning, TEXT("Ignoring unknown XML child node <%s>."), *ChildNode->GetTag());
+		}
+	}
+
+	return true;
+}
+
+////////////////////////////////////////////////////////////////////////////////////////////////////
+
+bool FMemTagBudget::ProcessBudgetNode(FXmlNode* BudgetNode)
+{
+	FString ModeName = BudgetNode->GetAttribute(STR_Name);
+	FMemTagBudgetMode& BudgetMode = GetOrAddMode(ModeName);
+
+	BudgetMode.GetDefaultPlatform().SetName(DefaultCachedPlatformName);
+
+	const TArray<FXmlNode*>& ChildNodes = BudgetNode->GetChildrenNodes();
+	for (FXmlNode* ChildNode : ChildNodes)
+	{
+		if (ChildNode->GetTag() == STR_Set)
+		{
+			ProcessTagSetNode(ChildNode, BudgetMode.GetDefaultPlatform());
+		}
+		else if (ChildNode->GetTag() == STR_Platform)
+		{
+			ProcessPlatformNode(ChildNode, BudgetMode);
+		}
+		else
+		{
+			UE_LOG(LogMemoryProfiler, Warning, TEXT("Ignoring unknown XML child node <%s> in <%s>."), *ChildNode->GetTag(), *STR_Budget);
+		}
+	}
+
+	return true;
+}
+
+////////////////////////////////////////////////////////////////////////////////////////////////////
+
+bool FMemTagBudget::ProcessPlatformNode(FXmlNode* PlatformNode, FMemTagBudgetMode& BudgetMode)
+{
+	FString PlatformName = PlatformNode->GetAttribute(STR_Name);
+	const TCHAR* CachedPlatformName = StoreString(PlatformName);
+
+	FMemTagBudgetPlatform& BudgetPlatform = BudgetMode.GetOrAddPlatformOverride(CachedPlatformName);
+
+	const TArray<FXmlNode*>& ChildNodes = PlatformNode->GetChildrenNodes();
+	for (FXmlNode* ChildNode : ChildNodes)
 	{
-		if (BudgetNode->GetTag() == STR_Budget)
+		if (ChildNode->GetTag() == STR_Set)
 		{
-			FString ModeName = BudgetNode->GetAttribute(STR_Name);
-			FMemTagBudgetMode& BudgetMode = GetOrAddMode(ModeName);
-			BudgetMode.GetDefaultPlatform().SetName(DefaultCachedPlatformName);
-
-			auto ProcessTagNode = [&](FXmlNode* TagNode, FMemTagBudgetPlatform& BudgetPlatform)
-				{
-					FString TagName = TagNode->GetAttribute(STR_Name);
-					const TCHAR* CachedTagName = StoreString(TagName);
-					if (CachedTagName == nullptr)
-					{
-						UE_LOG(LogMemoryProfiler, Warning, TEXT("Ignoring budget value for tag \"%s\" (invalid name)."), *TagName);
-						return;
-					}
-
-					//////////////////////////////////////////////////
-
-					const TCHAR* CachedTrackerName = DefaultCachedTrackerName;
-					FString TrackerName = TagNode->GetAttribute(STR_Tracker);
-					if (!TrackerName.IsEmpty())
-					{
-						if (!FCString::Stricmp(*TrackerName, TEXT("Platform")))
-						{
-							CachedTrackerName = PlatformCachedTrackerName;
-						}
-						else if (!FCString::Stricmp(*TrackerName, TEXT("Default")))
-						{
-							CachedTrackerName = DefaultCachedTrackerName;
-						}
-						else
-						{
-							UE_LOG(LogMemoryProfiler, Warning, TEXT("Ignoring unknown tracker \"%s\"."), *TrackerName);
-						}
-						CachedTrackerName = StoreString(TrackerName);
-					}
-					FMemTagBudgetTracker& BudgetTracker = BudgetPlatform.GetOrAddTracker(CachedTrackerName);
-
-					//////////////////////////////////////////////////
-
-					const TCHAR* CachedTagSetName = SystemCachedTagSetName;
-					FString TagSetName = TagNode->GetAttribute(STR_Set);
-					if (!TagSetName.IsEmpty())
-					{
-						if (!FCString::Stricmp(*TagSetName, TEXT("Asset")) ||
-							!FCString::Stricmp(*TagSetName, TEXT("Assets")))
-						{
-							CachedTagSetName = AssetCachedTagSetName;
-						}
-						else
-						if (!FCString::Stricmp(*TagSetName, TEXT("AssetClass")) ||
-							!FCString::Stricmp(*TagSetName, TEXT("AssetClasses")))
-						{
-							CachedTagSetName = AssetClassCachedTagSetName;
-						}
-						else
-						if (!FCString::Stricmp(*TagSetName, TEXT("System")) ||
-							!FCString::Stricmp(*TagSetName, TEXT("Systems")))
-						{
-							CachedTagSetName = SystemCachedTagSetName;
-						}
-						else
-						{
-							UE_LOG(LogMemoryProfiler, Warning, TEXT("Ignoring unknown tag set \"%s\"."), *TagSetName);
-						}
-					}
-					FMemTagBudgetTagSet& BudgetTagSet = BudgetTracker.GetOrAddTagSet(CachedTagSetName);
-
-					//////////////////////////////////////////////////
-
-					if (BudgetTagSet.FindValue(CachedTagName) != nullptr)
-					{
-						UE_LOG(LogMemoryProfiler, Warning, TEXT("Ignoring budget value for tag \"%s\" (defined multiple times)."), *TagName);
-						return;
-					}
-
-					FString Value = TagNode->GetAttribute(STR_TagMemMax);
-
-					PRAGMA_DISABLE_UNSAFE_TYPECAST_WARNINGS
-					TValueOrError<FNumericUnit<double>, FText> Result = FNumericUnit<double>::TryParseExpression(*Value, EUnit::Bytes, 0);
-					if (Result.IsValid())
-					{
-						auto MemoryValueInBytes = Result.GetValue().ConvertTo(EUnit::Bytes);
-						if (MemoryValueInBytes.IsSet())
-						{
-							BudgetTagSet.AddValue(CachedTagName, (int64)MemoryValueInBytes->Value);
-						}
-					}
-					else
-					{
-						UE_LOG(LogMemoryProfiler, Warning, TEXT("%s"), *(Result.GetError().ToString()));
-					}
-					PRAGMA_RESTORE_UNSAFE_TYPECAST_WARNINGS
-				};
-
-			const TArray<FXmlNode*>& BudgetChildNodes = BudgetNode->GetChildrenNodes();
-			for (FXmlNode* BudgetChildNode : BudgetChildNodes)
-			{
-				if (BudgetChildNode->GetTag() == STR_Platform)
-				{
-					FString PlatformName = BudgetChildNode->GetAttribute(STR_Name);
-					const TCHAR* CachedPlatformName = StoreString(PlatformName);
-
-					FMemTagBudgetPlatform& BudgetPlatform = BudgetMode.GetOrAddPlatformOverride(CachedPlatformName);
-					const TArray<FXmlNode*>& PlatformChildNodes = BudgetChildNode->GetChildrenNodes();
-					for (FXmlNode* PlatformChildNode : PlatformChildNodes)
-					{
-						if (PlatformChildNode->GetTag() == STR_Tag)
-						{
-							ProcessTagNode(PlatformChildNode, BudgetPlatform);
-						}
-					}
-				}
-				else if (BudgetChildNode->GetTag() == STR_Tag)
-				{
-					ProcessTagNode(BudgetChildNode, BudgetMode.GetDefaultPlatform());
-				}
-			}
+			ProcessTagSetNode(ChildNode, BudgetPlatform);
 		}
+		else
+		{
+			UE_LOG(LogMemoryProfiler, Warning, TEXT("Ignoring unknown XML child node <%s> in <%s>."), *ChildNode->GetTag(), *STR_Platform);
+		}
+	}
+
+	return true;
+}
+
+////////////////////////////////////////////////////////////////////////////////////////////////////
+
+bool FMemTagBudget::ProcessTagSetNode(FXmlNode* TagSetNode, FMemTagBudgetPlatform& BudgetPlatform)
+{
+	FString TagSetName = TagSetNode->GetAttribute(STR_Name);
+	const TCHAR* CachedTagSetName = nullptr;
+	if (!TagSetName.IsEmpty())
+	{
+		if (!FCString::Stricmp(*TagSetName, TEXT("Asset")) ||
+			!FCString::Stricmp(*TagSetName, TEXT("Assets")))
+		{
+			CachedTagSetName = AssetCachedTagSetName;
+		}
+		else
+		if (!FCString::Stricmp(*TagSetName, TEXT("AssetClass")) ||
+			!FCString::Stricmp(*TagSetName, TEXT("AssetClasses")))
+		{
+			CachedTagSetName = AssetClassCachedTagSetName;
+		}
+		else
+		if (!FCString::Stricmp(*TagSetName, TEXT("System")) ||
+			!FCString::Stricmp(*TagSetName, TEXT("Systems")))
+		{
+			CachedTagSetName = SystemCachedTagSetName;
+		}
+#if 0
+		else
+		{
+			CachedTagSetName = StoreString(TagSetName);
+			UE_LOG(LogMemoryProfiler, Warning, TEXT("Unknown tag set \"%s\"."), *TagSetName);
+		}
+#endif
+	}
+	if (!CachedTagSetName)
+	{
+		UE_LOG(LogMemoryProfiler, Warning, TEXT("Ignoring unknown tag set \"%s\"."), *TagSetName);
+		return false;
+	}
+
+	FMemTagBudgetTagSet& BudgetTagSet = BudgetPlatform.GetOrAddTagSet(CachedTagSetName);
+
+	const TArray<FXmlNode*>& ChildNodes = TagSetNode->GetChildrenNodes();
+	for (FXmlNode* ChildNode : ChildNodes)
+	{
+		if (ChildNode->GetTag() == STR_Tracker)
+		{
+			ProcessTrackerNode(ChildNode, BudgetTagSet);
+		}
+		else if (ChildNode->GetTag() == STR_Grouping)
+		{
+			ProcessGroupingNode(ChildNode, BudgetTagSet);
+		}
+		else
+		{
+			UE_LOG(LogMemoryProfiler, Warning, TEXT("Ignoring unknown XML child node <%s> in <%s>."), *ChildNode->GetTag(), *STR_Set);
+		}
+	}
+
+	return true;
+}
+
+////////////////////////////////////////////////////////////////////////////////////////////////////
+
+bool FMemTagBudget::ProcessTrackerNode(FXmlNode* TrackerNode, FMemTagBudgetTagSet& BudgetTagSet)
+{
+	FString TrackerName = TrackerNode->GetAttribute(STR_Name);
+	const TCHAR* CachedTrackerName = nullptr;
+	if (!TrackerName.IsEmpty())
+	{
+		if (!FCString::Stricmp(*TrackerName, DefaultCachedTrackerName))
+		{
+			CachedTrackerName = DefaultCachedTrackerName;
+		}
+		else if (!FCString::Stricmp(*TrackerName, PlatformCachedTrackerName))
+		{
+			CachedTrackerName = PlatformCachedTrackerName;
+		}
+#if 0
+		else
+		{
+			CachedTrackerName = StoreString(TrackerName);
+			UE_LOG(LogMemoryProfiler, Warning, TEXT("Unknown tracker \"%s\" in tag set \"%s\"."), *TrackerName, BudgetTagSet.GetName());
+		}
+#endif
+	}
+	if (!CachedTrackerName)
+	{
+		UE_LOG(LogMemoryProfiler, Warning, TEXT("Ignoring unknown tracker \"%s\" in tag set \"%s\"."), *TrackerName, BudgetTagSet.GetName());
+		return false;
+	}
+
+	FMemTagBudgetTracker& BudgetTracker = BudgetTagSet.GetOrAddTracker(CachedTrackerName);
+
+	const TArray<FXmlNode*>& ChildNodes = TrackerNode->GetChildrenNodes();
+	for (FXmlNode* ChildNode : ChildNodes)
+	{
+		if (ChildNode->GetTag() == STR_Tag)
+		{
+			ProcessTagNode(ChildNode, BudgetTracker);
+		}
+		else
+		{
+			UE_LOG(LogMemoryProfiler, Warning, TEXT("Ignoring unknown XML child node <%s> in <%s>."), *ChildNode->GetTag(), *STR_Tracker);
+		}
+	}
+
+	return true;
+}
+
+////////////////////////////////////////////////////////////////////////////////////////////////////
+
+bool FMemTagBudget::ProcessTagNode(FXmlNode* TagNode, FMemTagBudgetTracker& BudgetTracker)
+{
+	FString TagName = TagNode->GetAttribute(STR_Name);
+	if (TagName.IsEmpty())
+	{
+		UE_LOG(LogMemoryProfiler, Warning, TEXT("Ignoring budget value for empty tag."));
+		return false;
+	}
+	const TCHAR* CachedTagName = StoreString(TagName);
+	check(CachedTagName != nullptr);
+
+	//////////////////////////////////////////////////
+	// Value
+
+	if (BudgetTracker.FindValue(CachedTagName) != nullptr)
+	{
+		UE_LOG(LogMemoryProfiler, Warning, TEXT("Ignoring budget value for tag \"%s\" (defined multiple times)."), *TagName);
+		return false;
+	}
+
+	int64 Value;
+	if (!ReadMemValue(TagNode, STR_TagMemMax, Value))
+	{
+		UE_LOG(LogMemoryProfiler, Warning, TEXT("Cannot parse value for tag \"%s\"."), *TagName);
+		return false;
+	}
+
+	BudgetTracker.AddValue(CachedTagName, Value);
+	return true;
+}
+
+////////////////////////////////////////////////////////////////////////////////////////////////////
+
+bool FMemTagBudget::ProcessGroupingNode(FXmlNode* TrackerNode, FMemTagBudgetTagSet& BudgetTagSet)
+{
+	const TArray<FXmlNode*>& ChildNodes = TrackerNode->GetChildrenNodes();
+	for (FXmlNode* ChildNode : ChildNodes)
+	{
+		if (ChildNode->GetTag() == STR_Group)
+		{
+			ProcessGroupNode(ChildNode, BudgetTagSet);
+		}
+		else
+		{
+			UE_LOG(LogMemoryProfiler, Warning, TEXT("Ignoring unknown XML child node <%s> in <%s>."), *ChildNode->GetTag(), *STR_Grouping);
+		}
+	}
+
+	return true;
+}
+
+////////////////////////////////////////////////////////////////////////////////////////////////////
+
+bool FMemTagBudget::ProcessGroupNode(FXmlNode* GroupNode, FMemTagBudgetTagSet& BudgetTagSet)
+{
+	FString GroupName = GroupNode->GetAttribute(STR_Name);
+	if (GroupName.IsEmpty())
+	{
+		UE_LOG(LogMemoryProfiler, Warning, TEXT("Ignoring unnamed group."));
+		return false;
+	}
+	const TCHAR* CachedGroupName = StoreString(GroupName);
+	check(CachedGroupName != nullptr);
+
+	FString Include = GroupNode->GetAttribute(STR_Include);
+	FString Exclude = GroupNode->GetAttribute(STR_Exclude);
+
+	int64 Value;
+	if (!ReadMemValue(GroupNode, STR_TagMemMax, Value))
+	{
+		UE_LOG(LogMemoryProfiler, Warning, TEXT("Cannot parse value for group \"%s\"."), *GroupName);
+		return false;
+	}
+
+	FMemTagBudgetGrouping& Grouping = BudgetTagSet.GetOrCreateGrouping();
+	FMemTagBudgetGroup& Group = Grouping.GetOrAddGroup(CachedGroupName);
+
+	if (!Include.IsEmpty())
+	{
+		Group.SetInclude(Include);
+	}
+
+	if (!Exclude.IsEmpty())
+	{
+		Group.SetExclude(Exclude);
+	}
+
+	Group.SetMemMax(Value);
+	return true;
+}
+
+////////////////////////////////////////////////////////////////////////////////////////////////////
+
+PRAGMA_DISABLE_UNSAFE_TYPECAST_WARNINGS
+bool FMemTagBudget::ReadMemValue(FXmlNode* TagNode, const FString& Attribute, int64& OutValue)
+{
+	FString Value = TagNode->GetAttribute(Attribute);
+
+	TValueOrError<FNumericUnit<double>, FText> Result = FNumericUnit<double>::TryParseExpression(*Value, EUnit::Bytes, 0);
+	if (!Result.IsValid())
+	{
+		UE_LOG(LogMemoryProfiler, Warning, TEXT("Failed to parse value. %s"), *(Result.GetError().ToString()));
+		return false;
+	}
+
+	TOptional<FNumericUnit<double>> MemoryValueInBytes = Result.GetValue().ConvertTo(EUnit::Bytes);
+	if (!MemoryValueInBytes.IsSet())
+	{
+		UE_LOG(LogMemoryProfiler, Warning, TEXT("Failed to parse value."));
+		return false;
 	}
 
+	OutValue = (int64)MemoryValueInBytes.GetValue().Value;
 	return true;
 }
+PRAGMA_RESTORE_UNSAFE_TYPECAST_WARNINGS
 
 ////////////////////////////////////////////////////////////////////////////////////////////////////
 
diff --git a/Engine/Source/Developer/TraceInsights/Private/Insights/MemoryProfiler/ViewModels/MemTagBudget.h b/Engine/Source/Developer/TraceInsights/Private/Insights/MemoryProfiler/ViewModels/MemTagBudget.h
--- a/Engine/Source/Developer/TraceInsights/Private/Insights/MemoryProfiler/ViewModels/MemTagBudget.h
+++ b/Engine/Source/Developer/TraceInsights/Private/Insights/MemoryProfiler/ViewModels/MemTagBudget.h
@@ -12,27 +12,89 @@
 #include "Containers/UnrealString.h"
 #include "Templates/SharedPointer.h"
 
+class FXmlNode;
+
 namespace UE::Insights::MemoryProfiler
 {
 
 ////////////////////////////////////////////////////////////////////////////////////////////////////
 
-class FMemTagBudgetTagSet
+class FMemTagBudgetGroup
 {
 public:
-	FMemTagBudgetTagSet(const TCHAR* InCachedTagSetName)
-	: CachedTagSetName(InCachedTagSetName)
+	FMemTagBudgetGroup(const TCHAR* InCachedGroupName)
+		: CachedGroupName(InCachedGroupName)
 	{
 	}
-	~FMemTagBudgetTagSet()
+
+	~FMemTagBudgetGroup()
+	{
+	}
+
+	const TCHAR* GetName() const { return CachedGroupName; }
+	void SetName(const TCHAR* InCachedGroupName) { CachedGroupName = InCachedGroupName; }
+
+	const FString& GetInclude() const { return Include; }
+	void SetInclude(const FString& InInclude) { Include = InInclude; }
+
+	const FString& GetExclude() const { return Exclude; }
+	void SetExclude(const FString& InExclude) { Exclude = InExclude; }
+
+	int64 GetMemMax() const { return MemMax; }
+	void SetMemMax(int64 InMemMax) { MemMax = InMemMax; }
+
+private:
+	const TCHAR* CachedGroupName = nullptr;
+	FString Include;
+	FString Exclude;
+	int64 MemMax = 0;
+};
+
+////////////////////////////////////////////////////////////////////////////////////////////////////
+
+class FMemTagBudgetGrouping
+{
+public:
+	FMemTagBudgetGrouping()
+	{
+	}
+	~FMemTagBudgetGrouping()
 	{
 		Reset();
 	}
 
 	void Reset();
 
-	const TCHAR* GetName() const { return CachedTagSetName; }
-	void SetName(const TCHAR* InCachedTagSetName) { CachedTagSetName = InCachedTagSetName; }
+	int32 GetNumGroups() const { return Groups.Num(); }
+	void EnumerateGroups(TFunctionRef<void(const TCHAR* InCachedGroupName, const FMemTagBudgetGroup& InGroup)> Callback) const;
+	const FMemTagBudgetGroup* FindGroup(const TCHAR* InCachedGroupName) const;
+	FMemTagBudgetGroup& GetOrAddGroup(const TCHAR* InCachedGroupName);
+
+private:
+	TArray<FMemTagBudgetGroup*> Groups;
+};
+
+////////////////////////////////////////////////////////////////////////////////////////////////////
+
+class FMemTagBudgetTagSet;
+
+class FMemTagBudgetTracker
+{
+public:
+	FMemTagBudgetTracker(const TCHAR* InCachedTrackerName, FMemTagBudgetTagSet* InParentTagSet)
+		: CachedTrackerName(InCachedTrackerName)
+		, ParentTagSet(InParentTagSet)
+	{
+	}
+	~FMemTagBudgetTracker()
+	{
+		Reset();
+	}
+
+	void Reset();
+
+	const TCHAR* GetName() const { return CachedTrackerName; }
+	FMemTagBudgetTagSet* GetParentTagSet() const { return ParentTagSet; }
 
 	const int64* FindValue(const TCHAR* InCachedTagName) const
 	{
@@ -45,36 +107,41 @@ public:
 	}
 
 private:
-	const TCHAR* CachedTagSetName = nullptr;
+	const TCHAR* CachedTrackerName = nullptr;
+	FMemTagBudgetTagSet* ParentTagSet = nullptr;
 	TMap<const void*, int64> Values;
 };
 
 ////////////////////////////////////////////////////////////////////////////////////////////////////
 
-class FMemTagBudgetTracker
+class FMemTagBudgetTagSet
 {
 public:
-	FMemTagBudgetTracker(const TCHAR* InCachedTrackerName)
-	: CachedTrackerName(InCachedTrackerName)
+	FMemTagBudgetTagSet(const TCHAR* InCachedTagSetName)
+		: CachedTagSetName(InCachedTagSetName)
 	{
 	}
-	~FMemTagBudgetTracker()
+	~FMemTagBudgetTagSet()
 	{
 		Reset();
 	}
 
 	void Reset();
 
-	const TCHAR* GetName() const { return CachedTrackerName; }
-	void SetName(const TCHAR* InCachedTrackerName) { CachedTrackerName = InCachedTrackerName; }
+	const TCHAR* GetName() const { return CachedTagSetName; }
+	void SetName(const TCHAR* InCachedTagSetName) { CachedTagSetName = InCachedTagSetName; }
 
-	void EnumerateTagSets(TFunctionRef<void(const TCHAR* InCachedTagSetName, const FMemTagBudgetTagSet& InTagSet)> Callback) const;
-	const FMemTagBudgetTagSet* FindTagSet(const TCHAR* InCachedTagSetName) const;
-	FMemTagBudgetTagSet& GetOrAddTagSet(const TCHAR* InCachedTagSetName);
+	void EnumerateTrackers(TFunctionRef<void(const TCHAR* InCachedTrackerName, const FMemTagBudgetTracker& InTracker)> Callback) const;
+	const FMemTagBudgetTracker* FindTracker(const TCHAR* InCachedTrackerName) const;
+	FMemTagBudgetTracker& GetOrAddTracker(const TCHAR* InCachedTrackerName);
+
+	const FMemTagBudgetGrouping* GetGrouping() const { return Grouping; }
+	FMemTagBudgetGrouping& GetOrCreateGrouping();
 
 private:
-	const TCHAR* CachedTrackerName = nullptr;
-	TMap<const void*, FMemTagBudgetTagSet*> TagSets;
+	const TCHAR* CachedTagSetName = nullptr;
+	TMap<const void*, FMemTagBudgetTracker*> Trackers;
+	FMemTagBudgetGrouping* Grouping = nullptr;
 };
 
 ////////////////////////////////////////////////////////////////////////////////////////////////////
@@ -97,15 +164,13 @@ public:
 	const TCHAR* GetName() const { return CachedPlatformName; }
 	void SetName(const TCHAR* InCachedPlatformName) { CachedPlatformName = InCachedPlatformName; }
 
-	void EnumerateTrackers(TFunctionRef<void(const TCHAR* InCachedTrackerName, const FMemTagBudgetTracker& InTracker)> Callback) const;
-	const FMemTagBudgetTracker* FindTracker(const TCHAR* InCachedTrackerName) const;
-	FMemTagBudgetTracker& GetOrAddTracker(const TCHAR* InCachedTrackerName);
-
-	const FMemTagBudgetTagSet* FindTagSet(const TCHAR* InCachedTrackerName, const TCHAR* InCachedTagSetName) const;
+	void EnumerateTagSets(TFunctionRef<void(const TCHAR* InCachedTagSetName, const FMemTagBudgetTagSet& InTagSet)> Callback) const;
+	const FMemTagBudgetTagSet* FindTagSet(const TCHAR* InCachedTagSetName) const;
+	FMemTagBudgetTagSet& GetOrAddTagSet(const TCHAR* InCachedTagSetName);
 
 private:
 	const TCHAR* CachedPlatformName = nullptr;
-	TMap<const void*, FMemTagBudgetTracker*> Trackers;
+	TMap<const void*, FMemTagBudgetTagSet*> TagSets;
 };
 
 ////////////////////////////////////////////////////////////////////////////////////////////////////
@@ -175,11 +240,41 @@ public:
 	bool LoadFromFile(const FString& InFilePath);
 	bool ReLoadFromFile();
 
+private:
+	bool ProcessBudgetNode(FXmlNode* BudgetNode);
+	bool ProcessPlatformNode(FXmlNode* PlatformNode, FMemTagBudgetMode& BudgetMode);
+	bool ProcessTagSetNode(FXmlNode* TagSetNode, FMemTagBudgetPlatform& BudgetPlatform);
+	bool ProcessTrackerNode(FXmlNode* TrackerNode, FMemTagBudgetTagSet& BudgetTagSet);
+	bool ProcessTagNode(FXmlNode* TagNode, FMemTagBudgetTracker& BudgetTracker);
+	bool ProcessGroupingNode(FXmlNode* TrackerNode, FMemTagBudgetTagSet& BudgetTagSet);
+	bool ProcessGroupNode(FXmlNode* GroupNode, FMemTagBudgetTagSet& BudgetTagSet);
+	bool ReadMemValue(FXmlNode* TagNode, const FString& Attribute, int64& OutValue);
+
 private:
 	TraceServices::IStringStore* StringStore = nullptr;
 	FString Name;
 	FString FilePath;
 	TMap<const void*, FMemTagBudgetMode*> Modes;
+
+	const TCHAR* DefaultCachedPlatformName;
+	const TCHAR* DefaultCachedTrackerName;
+	const TCHAR* PlatformCachedTrackerName;
+	const TCHAR* SystemCachedTagSetName;
+	const TCHAR* AssetCachedTagSetName;
+	const TCHAR* AssetClassCachedTagSetName;
+
+	static const FString STR_Name;
+	static const FString STR_Id;
+	static const FString STR_Budget;
+	static const FString STR_Set;
+	static const FString STR_Tag;
+	static const FString STR_TagMemMax;
+	static const FString STR_Tracker;
+	static const FString STR_Platform;
+	static const FString STR_Grouping;
+	static const FString STR_Group;
+	static const FString STR_Include;
+	static const FString STR_Exclude;
 };
 
 ////////////////////////////////////////////////////////////////////////////////////////////////////
diff --git a/Engine/Source/Developer/TraceInsights/Private/Insights/MemoryProfiler/ViewModels/MemTagBudgetGrouping.cpp b/Engine/Source/Developer/TraceInsights/Private/Insights/MemoryProfiler/ViewModels/MemTagBudgetGrouping.cpp
new file mode 100644
--- /dev/null
+++ b/Engine/Source/Developer/TraceInsights/Private/Insights/MemoryProfiler/ViewModels/MemTagBudgetGrouping.cpp
@@ -0,0 +1,171 @@
+// Copyright Epic Games, Inc. All Rights Reserved.
+
+#include "MemTagBudgetGrouping.h"
+
+#include "Internationalization/Regex.h"
+
+// TraceInsightsCore
+#include "InsightsCore/Common/AsyncOperationProgress.h"
+
+// TraceInsights
+#include "Insights/InsightsManager.h"
+#include "Insights/MemoryProfiler/MemoryProfilerManager.h"
+#include "Insights/MemoryProfiler/ViewModels/MemTagBudget.h"
+#include "Insights/MemoryProfiler/ViewModels/MemTagNode.h"
+#include "Insights/MemoryProfiler/Widgets/SMemTagTreeView.h"
+
+#define LOCTEXT_NAMESPACE "UE::Insights::MemoryProfiler::FMemTagNode"
+
+namespace UE::Insights::MemoryProfiler
+{
+
+////////////////////////////////////////////////////////////////////////////////////////////////////
+
+INSIGHTS_IMPLEMENT_RTTI(FMemTagBudgetNodeGrouping)
+INSIGHTS_IMPLEMENT_RTTI(FMemTagBudgetGroupNode)
+
+////////////////////////////////////////////////////////////////////////////////////////////////////
+// FMemTagBudgetNodeGrouping
+////////////////////////////////////////////////////////////////////////////////////////////////////
+
+FMemTagBudgetNodeGrouping::FMemTagBudgetNodeGrouping(TSharedPtr<SMemTagTreeView> InMemTagTreeView)
+	: FTreeNodeGrouping(
+		LOCTEXT("Grouping_Budget_ShortName", "Budget"),
+		LOCTEXT("Grouping_Budget_TitleName", "Budget Groups"),
+		LOCTEXT("Grouping_Budget_Desc", "Groups tags based on rules specified in the budget xml file."),
+		nullptr)
+	, MemTagTreeView(InMemTagTreeView.ToWeakPtr())
+{
+}
+
+////////////////////////////////////////////////////////////////////////////////////////////////////
+
+void FMemTagBudgetNodeGrouping::GroupNodes(
+	const TArray<FTableTreeNodePtr>& Nodes,
+	FTableTreeNode& ParentGroup,
+	TWeakPtr<FTable> InParentTable,
+	IAsyncOperationProgress& InAsyncOperationProgress) const
+{
+	ParentGroup.ClearChildren();
+
+	class FGroupRule
+	{
+	public:
+		FGroupRule(const TCHAR* InCachedGroupName, int64 InSizeBudget, const FString& InInclude, const FString& InExclude)
+			: CachedGroupName(InCachedGroupName)
+			, SizeBudget(InSizeBudget)
+			, Include(InInclude)
+			, Exclude(InExclude)
+			, bHasExclude(!InExclude.IsEmpty())
+		{
+		}
+
+		const TCHAR* GetGroupName() const { return CachedGroupName; }
+		int64 GetSizeBudget() const { return SizeBudget; }
+		void SetSizeBudget(int64 InSizeBudget) { SizeBudget = InSizeBudget; }
+
+		bool Match(const FMemoryTag& MemTag)
+		{
+			FRegexMatcher IncludeMatcher(Include, MemTag.GetStatFullName());
+			if (IncludeMatcher.FindNext())
+			{
+				FRegexMatcher ExcludeMatcher(Exclude, MemTag.GetStatFullName());
+				if (!bHasExclude || !ExcludeMatcher.FindNext())
+				{
+					return true;
+				}
+			}
+			return false;
+		}
+
+	private:
+		const TCHAR* CachedGroupName = nullptr;
+		int64 SizeBudget = 0;
+		FRegexPattern Include;
+		FRegexPattern Exclude;
+		bool bHasExclude = false;
+	};
+
+	TArray<FGroupRule> GroupRules;
+
+	const FMemTagBudgetGrouping* BudgetGrouping = nullptr;
+	const FMemTagBudgetGrouping* BudgetGroupingOverride = nullptr;
+	auto MemTagTreeViewPtr = MemTagTreeView.Pin();
+	if (MemTagTreeViewPtr.IsValid())
+	{
+		MemTagTreeViewPtr->GetBudgetGrouping(BudgetGrouping, BudgetGroupingOverride);
+	}
+	if (BudgetGrouping)
+	{
+		GroupRules.Reserve(BudgetGrouping->GetNumGroups());
+		BudgetGrouping->EnumerateGroups([&GroupRules](const TCHAR* InCachedGroupName, const FMemTagBudgetGroup& InGroup)
+			{
+				GroupRules.Emplace(InCachedGroupName, InGroup.GetMemMax(), InGroup.GetInclude(), InGroup.GetExclude());
+			});
+		if (BudgetGroupingOverride)
+		{
+			BudgetGroupingOverride->EnumerateGroups([&GroupRules](const TCHAR* InCachedGroupName, const FMemTagBudgetGroup& InGroup)
+				{
+					FGroupRule* GroupRule = GroupRules.FindByPredicate([InCachedGroupName](const FGroupRule& Rule) { return Rule.GetGroupName() == InCachedGroupName; });
+					if (GroupRule)
+					{
+						GroupRule->SetSizeBudget(InGroup.GetMemMax());
+					}
+				});
+		}
+	}
+
+	TMap<const void*, TSharedPtr<FMemTagBudgetGroupNode>> Groups;
+
+	for (FTableTreeNodePtr NodePtr : Nodes)
+	{
+		if (InAsyncOperationProgress.ShouldCancelAsyncOp())
+		{
+			return;
+		}
+
+		if (NodePtr->IsGroup() || !NodePtr->Is<FMemTagNode>())
+		{
+			ParentGroup.AddChildAndSetParent(NodePtr);
+			continue;
+		}
+
+		FMemTagNode& MemTagNode = NodePtr->As<FMemTagNode>();
+		check(MemTagNode.GetMemTag() != nullptr);
+		FMemoryTag& MemTag = *MemTagNode.GetMemTag();
+
+		const TCHAR* CachedGroupName = nullptr;
+		int64 SizeBudget = 0;
+		for (FGroupRule& Rule : GroupRules)
+		{
+			if (Rule.Match(MemTag))
+			{
+				CachedGroupName = Rule.GetGroupName();
+				SizeBudget = Rule.GetSizeBudget();
+				break;
+			}
+		}
+		if (!CachedGroupName)
+		{
+			ParentGroup.AddChildAndSetParent(NodePtr);
+			continue;
+		}
+
+		TSharedPtr<FMemTagBudgetGroupNode> GroupNode = Groups.FindRef(CachedGroupName);
+		if (!GroupNode)
+		{
+			FName NodeName(CachedGroupName);
+			GroupNode = MakeShared<FMemTagBudgetGroupNode>(NodeName, InParentTable, CachedGroupName);
+			GroupNode->SetSizeBudget(SizeBudget);
+			ParentGroup.AddChildAndSetParent(GroupNode);
+			Groups.Add(CachedGroupName, GroupNode);
+		}
+		GroupNode->AddChildAndSetParent(NodePtr);
+	}
+}
+
+////////////////////////////////////////////////////////////////////////////////////////////////////
+
+} // namespace UE::Insights::MemoryProfiler
+
+#undef LOCTEXT_NAMESPACE
diff --git a/Engine/Source/Developer/TraceInsights/Private/Insights/MemoryProfiler/ViewModels/MemTagBudgetGrouping.h b/Engine/Source/Developer/TraceInsights/Private/Insights/MemoryProfiler/ViewModels/MemTagBudgetGrouping.h
new file mode 100644
--- /dev/null
+++ b/Engine/Source/Developer/TraceInsights/Private/Insights/MemoryProfiler/ViewModels/MemTagBudgetGrouping.h
@@ -0,0 +1,98 @@
+// Copyright Epic Games, Inc. All Rights Reserved.
+#pragma once
+
+#include "CoreTypes.h"
+
+// TraceInsightsCore
+#include "InsightsCore/Common/SimpleRtti.h"
+#include "InsightsCore/Table/ViewModels/TableTreeNode.h"
+#include "InsightsCore/Table/ViewModels/TreeNodeGrouping.h"
+
+namespace UE::Insights::MemoryProfiler
+{
+
+class SMemTagTreeView;
+
+////////////////////////////////////////////////////////////////////////////////////////////////////
+
+class FMemTagBudgetNodeGrouping : public FTreeNodeGrouping
+{
+	INSIGHTS_DECLARE_RTTI(FMemTagBudgetNodeGrouping, FTreeNodeGrouping);
+
+public:
+	FMemTagBudgetNodeGrouping(TSharedPtr<SMemTagTreeView> InMemTagTreeView);
+
+	virtual void GroupNodes(const TArray<FTableTreeNodePtr>& Nodes, FTableTreeNode& ParentGroup, TWeakPtr<FTable> InParentTable, IAsyncOperationProgress& InAsyncOperationProgress) const override;
+
+private:
+	TWeakPtr<SMemTagTreeView> MemTagTreeView;
+};
+
+////////////////////////////////////////////////////////////////////////////////////////////////////
+
+class FMemTagBudgetGroupNode : public FTableTreeNode
+{
+	INSIGHTS_DECLARE_RTTI(FMemTagBudgetGroupNode, FTableTreeNode)
+
+public:
+	/** Initialization constructor for a table record node. */
+	explicit FMemTagBudgetGroupNode(const FName InName, TWeakPtr<FTable> InParentTable, int32 InRowIndex, const TCHAR* InCachedBudgetGroupName, bool IsGroup)
+		: FTableTreeNode(InName, InParentTable, InRowIndex, IsGroup)
+		, CachedBudgetGroupName(InCachedBudgetGroupName)
+	{
+	}
+
+	/** Initialization constructor for the group node. */
+	explicit FMemTagBudgetGroupNode(const FName InName, TWeakPtr<FTable> InParentTable, const TCHAR* InCachedBudgetGroupName)
+		: FTableTreeNode(InName, InParentTable)
+		, CachedBudgetGroupName(InCachedBudgetGroupName)
+	{
+	}
+
+	virtual ~FMemTagBudgetGroupNode()
+	{
+	}
+
+	virtual FLinearColor GetIconColor() const override
+	{
+		return FLinearColor(0.75f, 0.5f, 1.0f, 1.0f);
+	}
+
+	virtual FLinearColor GetColor() const override
+	{
+		return FLinearColor(0.75f, 0.5f, 1.0f, 1.0f);
+	}
+
+	const TCHAR* GetBudgetGroupName() const
+	{
+		return CachedBudgetGroupName;
+	}
+
+	bool HasSizeBudget() const
+	{
+		return SizeBudget != 0;
+	}
+
+	void ResetSizeBudget()
+	{
+		SizeBudget = 0;
+	}
+
+	int64 GetSizeBudget() const
+	{
+		return SizeBudget;
+	}
+
+	void SetSizeBudget(int64 InValue)
+	{
+		SizeBudget = InValue;
+	}
+
+private:
+	const TCHAR* CachedBudgetGroupName = nullptr;
+	int64 SizeBudget = 0;
+};
+
+////////////////////////////////////////////////////////////////////////////////////////////////////
+
+} // namespace UE::Insights::MemoryProfiler
diff --git a/Engine/Source/Developer/TraceInsights/Private/Insights/MemoryProfiler/ViewModels/MemTagTable.cpp b/Engine/Source/Developer/TraceInsights/Private/Insights/MemoryProfiler/ViewModels/MemTagTable.cpp
--- a/Engine/Source/Developer/TraceInsights/Private/Insights/MemoryProfiler/ViewModels/MemTagTable.cpp
+++ b/Engine/Source/Developer/TraceInsights/Private/Insights/MemoryProfiler/ViewModels/MemTagTable.cpp
@@ -15,6 +15,7 @@
 #include "InsightsCore/Table/ViewModels/TableColumn.h"
 
 // TraceInsights
+#include "Insights/MemoryProfiler/ViewModels/MemTagBudgetGrouping.h"
 #include "Insights/MemoryProfiler/ViewModels/MemTagNode.h"
 #include "Insights/MemoryProfiler/ViewModels/MemTagNodeGroupingAndSorting.h"
 
@@ -114,11 +115,11 @@ void FMemTagTable::AddDefaultColumns()
 						{
 							if (Node.Is<FMemTagNode>())
 							{
-								TOptional<FTableCellValue> Value = Column.GetValue(Node);
-								if (Value.IsSet())
+								const FMemTagNode& MemTagNode = Node.As<FMemTagNode>();
+								if (MemTagNode.HasSizeBudget())
 								{
-									const FMemTagNode& MemTagNode = Node.As<FMemTagNode>();
-									if (MemTagNode.HasSizeBudget())
+									TOptional<FTableCellValue> Value = Column.GetValue(Node);
+									if (Value.IsSet())
 									{
 										int64 SizeValue = Value.GetValue().AsInt64();
 										if (SizeValue > MemTagNode.GetSizeBudget())
@@ -130,9 +131,28 @@ void FMemTagTable::AddDefaultColumns()
 											return FLinearColor(0.1f, 0.5f, 0.1f, 1.0f);
 										}
 									}
-									else
+								}
+								else
+								{
+									return FLinearColor(0.5f, 0.5f, 0.5f, 1.0f);
+								}
+							}
+							else if (Node.Is<FMemTagBudgetGroupNode>())
+							{
+								const FMemTagBudgetGroupNode& GroupNode = Node.As<FMemTagBudgetGroupNode>();
+								if (GroupNode.HasSizeBudget())
+								{
+									if (GroupNode.HasAggregatedValue(Column.GetId()))
 									{
-										return FLinearColor(0.5f, 0.5f, 0.5f, 1.0f);
+										int64 SizeValue = GroupNode.GetAggregatedValue(Column.GetId()).AsInt64();
+										if (SizeValue > GroupNode.GetSizeBudget())
+										{
+											return FLinearColor(1.0f, 0.3f, 0.3f, 1.0f);
+										}
+										else
+										{
+											return FLinearColor(0.1f, 0.5f, 0.1f, 1.0f);
+										}
 									}
 								}
 							}
@@ -648,6 +668,14 @@ void FMemTagTable::AddDefaultColumns()
 						return FTableCellValue(MemTagNode.GetSizeBudget());
 					}
 				}
+				else if (Node.Is<FMemTagBudgetGroupNode>())
+				{
+					const FMemTagBudgetGroupNode& GroupNode = Node.As<FMemTagBudgetGroupNode>();
+					if (GroupNode.HasSizeBudget())
+					{
+						return FTableCellValue(GroupNode.GetSizeBudget());
+					}
+				}
 				return TOptional<FTableCellValue>();
 			}
 		};
diff --git a/Engine/Source/Developer/TraceInsights/Private/Insights/MemoryProfiler/Widgets/SMemTagTreeView.cpp b/Engine/Source/Developer/TraceInsights/Private/Insights/MemoryProfiler/Widgets/SMemTagTreeView.cpp
--- a/Engine/Source/Developer/TraceInsights/Private/Insights/MemoryProfiler/Widgets/SMemTagTreeView.cpp
+++ b/Engine/Source/Developer/TraceInsights/Private/Insights/MemoryProfiler/Widgets/SMemTagTreeView.cpp
@@ -28,6 +28,7 @@
 #include "Insights/MemoryProfiler/ViewModels/MemoryTag.h"
 #include "Insights/MemoryProfiler/ViewModels/MemoryTracker.h"
 #include "Insights/MemoryProfiler/ViewModels/MemTagBudget.h"
+#include "Insights/MemoryProfiler/ViewModels/MemTagBudgetGrouping.h"
 #include "Insights/MemoryProfiler/ViewModels/MemTagNode.h"
 #include "Insights/MemoryProfiler/ViewModels/MemTagTableViewPresets.h"
 #include "Insights/MemoryProfiler/Widgets/SMemoryProfilerWindow.h"
@@ -789,6 +790,8 @@ void SMemTagTreeView::InternalCreateGroupings()
 
 	int32 Index = 1; // after the Flat ("All") grouping
 
+	AvailableGroupings.Insert(MakeShared<FMemTagBudgetNodeGrouping>(SharedThis(this)), Index++);
+
 #if 0
 	// Creates one group for each tag type.
 	class FMemTagGroupingByTagType : public
@@ -2159,6 +2162,83 @@ FText SMemTagTreeView::BudgetPlatform_GetTooltipText() const
 // Budget
 ////////////////////////////////////////////////////////////////////////////////////////////////////
 
+const TCHAR* SMemTagTreeView::GetTagSetFilterBudgetCachedName() const
+{
+	if (!CurrentBudget.IsValid())
+	{
+		return nullptr;
+	}
+
+	switch (TagSetFilter)
+	{
+		case 0:  return CurrentBudget->FindString(TEXTVIEW("Systems"));
+		case 1:  return CurrentBudget->FindString(TEXTVIEW("Assets"));
+		case 2:  return CurrentBudget->FindString(TEXTVIEW("AssetClasses"));
+		default: return nullptr;
+	}
+}
+
+////////////////////////////////////////////////////////////////////////////////////////////////////
+
+const TCHAR* SMemTagTreeView::GetSelectedBudgetPlatformCachedName() const
+{
+	return CurrentBudget.IsValid() && SelectedBudgetPlatform.IsValid() ?
+		CurrentBudget->FindString(SelectedBudgetPlatform->Name) : nullptr;
+}
+
+////////////////////////////////////////////////////////////////////////////////////////////////////
+
+const FMemTagBudgetMode* SMemTagTreeView::GetSelectedBudgetMode() const
+{
+	return CurrentBudget.IsValid() && SelectedBudgetMode.IsValid() ?
+		CurrentBudget->FindMode(SelectedBudgetMode->Name) : nullptr;
+}
+
+////////////////////////////////////////////////////////////////////////////////////////////////////
+
+void SMemTagTreeView::GetBudgetGrouping(const FMemTagBudgetGrouping*& OutGrouping, const FMemTagBudgetGrouping*& OutGroupingOverride) const
+{
+	OutGrouping = nullptr;
+	OutGroupingOverride = nullptr;
+
+	const FMemTagBudgetMode* BudgetMode = GetSelectedBudgetMode();
+	if (!BudgetMode)
+	{
+		return;
+	}
+
+	const TCHAR* CachedTagSetName = GetTagSetFilterBudgetCachedName();
+	if (!CachedTagSetName)
+	{
+		return;
+	}
+
+	const FMemTagBudgetPlatform& DefaultPlatform = BudgetMode->GetDefaultPlatform();
+	const FMemTagBudgetTagSet* TagSet = DefaultPlatform.FindTagSet(CachedTagSetName);
+	if (!TagSet)
+	{
+		return;
+	}
+
+	OutGrouping = TagSet->GetGrouping();
+
+	const TCHAR* CachedPlatformName = GetSelectedBudgetPlatformCachedName();
+	if (CachedPlatformName)
+	{
+		const FMemTagBudgetPlatform* PlatformOverride = BudgetMode->FindPlatformOverride(CachedPlatformName);
+		if (PlatformOverride)
+		{
+			const FMemTagBudgetTagSet* TagSetOverride = PlatformOverride->FindTagSet(CachedTagSetName);
+			if (TagSetOverride)
+			{
+				OutGroupingOverride = TagSetOverride->GetGrouping();
+			}
+		}
+	}
+}
+
+////////////////////////////////////////////////////////////////////////////////////////////////////
+
 void SMemTagTreeView::InitBudgetOptions()
 {
 	CurrentBudget.Reset();
@@ -2350,61 +2430,52 @@ void SMemTagTreeView::ApplyBudgetToNodes()
 		return;
 	}
 
-	const FMemTagBudgetMode* FoundBudgetMode = SelectedBudgetMode.IsValid() ? CurrentBudget->FindMode(SelectedBudgetMode->Name) : nullptr;
-	if (!FoundBudgetMode)
+	const FMemTagBudgetMode* BudgetMode = GetSelectedBudgetMode();
+	if (!BudgetMode)
+	{
+		ResetBudgetForAllNodes();
+		return;
+	}
+
+	const TCHAR* CachedTagSetName = GetTagSetFilterBudgetCachedName();
+	const FMemTagBudgetPlatform& DefaultPlatform = BudgetMode->GetDefaultPlatform();
+	const FMemTagBudgetTagSet* TagSet = DefaultPlatform.FindTagSet(CachedTagSetName);
+	if (!TagSet)
 	{
 		ResetBudgetForAllNodes();
 		return;
 	}
-	const FMemTagBudgetMode& BudgetMode = *FoundBudgetMode;
 
-	const TCHAR* CachedTagSetNames[3];
-	CachedTagSetNames[0] = CurrentBudget->FindString(TEXTVIEW("Systems"));
-	CachedTagSetNames[1] = CurrentBudget->FindString(TEXTVIEW("Assets"));
-	CachedTagSetNames[2] = CurrentBudget->FindString(TEXTVIEW("AssetClasses"));
+	const FMemTagBudgetGrouping* BudgetGrouping = TagSet->GetGrouping();
+	const FMemTagBudgetGrouping* BudgetGroupingOverride = nullptr;
 
 	const TCHAR* CachedDefaultTrackerName = CurrentBudget->FindString(TEXTVIEW("Default"));
 	const TCHAR* CachedPlatformTrackerName = CurrentBudget->FindString(TEXTVIEW("Platform"));
 
-	const FMemTagBudgetPlatform& DefaultPlatform = BudgetMode.GetDefaultPlatform();
-
-	const FMemTagBudgetTagSet* DefaultTagSet[3];
-	const FMemTagBudgetTagSet* PlatformTagSet[3];
-	const FMemTagBudgetTagSet* DefaultTagSetOverride[3];
-	const FMemTagBudgetTagSet* PlatformTagSetOverride[3];
-
-	bool bHasAnyBudget = false;
+	const FMemTagBudgetTracker* DefaultTracker = TagSet->FindTracker(CachedDefaultTrackerName);
+	const FMemTagBudgetTracker* PlatformTracker = TagSet->FindTracker(CachedPlatformTrackerName);
 
-	for (int32 TagSetIndex = 0; TagSetIndex < 3; ++TagSetIndex)
+	const FMemTagBudgetTracker* DefaultTrackerOverride = nullptr;
+	const FMemTagBudgetTracker* PlatformTrackerOverride = nullptr;
+	const TCHAR* CachedPlatformName = GetSelectedBudgetPlatformCachedName();
+	if (CachedPlatformName)
 	{
-		const TCHAR* CachedTagSetName = CachedTagSetNames[TagSetIndex];
-
-		DefaultTagSet[TagSetIndex] = DefaultPlatform.FindTagSet(CachedDefaultTrackerName, CachedTagSetName);
-		PlatformTagSet[TagSetIndex] = DefaultPlatform.FindTagSet(CachedPlatformTrackerName, CachedTagSetName);
-
-		DefaultTagSetOverride[TagSetIndex] = nullptr;
-		PlatformTagSetOverride[TagSetIndex] = nullptr;
-		if (SelectedBudgetPlatform.IsValid())
+		const FMemTagBudgetPlatform* PlatformOverride = BudgetMode->FindPlatformOverride(CachedPlatformName);
+		if (PlatformOverride)
 		{
-			const TCHAR* CachedPlatformName = CurrentBudget->FindString(SelectedBudgetPlatform->Name);
-			const FMemTagBudgetPlatform* PlatformOverride = BudgetMode.FindPlatformOverride(CachedPlatformName);
-			if (PlatformOverride)
+			const FMemTagBudgetTagSet* TagSetOverride = PlatformOverride->FindTagSet(CachedTagSetName);
+			if (TagSetOverride)
 			{
-				DefaultTagSetOverride[TagSetIndex] = PlatformOverride->FindTagSet(CachedDefaultTrackerName, CachedTagSetName);
-				PlatformTagSetOverride[TagSetIndex] = PlatformOverride->FindTagSet(CachedPlatformTrackerName, CachedTagSetName);
+				BudgetGroupingOverride = TagSetOverride->GetGrouping();
+				DefaultTrackerOverride = TagSetOverride->FindTracker(CachedDefaultTrackerName);
+				PlatformTrackerOverride = TagSetOverride->FindTracker(CachedPlatformTrackerName);
 			}
 		}
-
-		bHasAnyBudget = bHasAnyBudget || (DefaultTagSet[TagSetIndex] != nullptr);
-		bHasAnyBudget = bHasAnyBudget || (PlatformTagSet[TagSetIndex] != nullptr);
-		bHasAnyBudget = bHasAnyBudget || (DefaultTagSetOverride[TagSetIndex] != nullptr);
-		bHasAnyBudget = bHasAnyBudget || (PlatformTagSetOverride[TagSetIndex] != nullptr);
 	}
 
-	if (!bHasAnyBudget)
+	if (BudgetGrouping)
 	{
-		ResetBudgetForAllNodes();
-		return;
+		UpdateBudgetGroupsRec(BudgetGrouping, BudgetGroupingOverride, Root);
 	}
 
 	for (const FTableTreeNodePtr& NodePtr : TableRowNodes)
@@ -2433,10 +2504,10 @@ void SMemTagTreeView::ApplyBudgetToNodes()
 
 		bool bIsDefaultTracker = MemoryTag->GetTrackerId() == SharedState.GetDefaultTrackerId();
 
-		const FMemTagBudgetTagSet* BudgetTagSetOverride = bIsDefaultTracker ? DefaultTagSetOverride[TagSetIndex] : PlatformTagSetOverride[TagSetIndex];
-		if (BudgetTagSetOverride)
+		const FMemTagBudgetTracker* BudgetTrackerOverride = bIsDefaultTracker ? DefaultTrackerOverride : PlatformTrackerOverride;
+		if (BudgetTrackerOverride)
 		{
-			const int64* ValuePtr = BudgetTagSetOverride->FindValue(CachedTagName);
+			const int64* ValuePtr = BudgetTrackerOverride->FindValue(CachedTagName);
 			if (ValuePtr)
 			{
 				MemTagNode.SetSizeBudget(*ValuePtr);
@@ -2445,10 +2516,10 @@ void SMemTagTreeView::ApplyBudgetToNodes()
 			}
 		}
 
-		const FMemTagBudgetTagSet* BudgetTagSet = bIsDefaultTracker ? DefaultTagSet[TagSetIndex] : PlatformTagSet[TagSetIndex];
-		if (BudgetTagSet)
+		const FMemTagBudgetTracker* BudgetTracker = bIsDefaultTracker ? DefaultTracker : PlatformTracker;
+		if (BudgetTracker)
 		{
-			const int64* ValuePtr = BudgetTagSet->FindValue(CachedTagName);
+			const int64* ValuePtr = BudgetTracker->FindValue(CachedTagName);
 			if (ValuePtr)
 			{
 				MemTagNode.SetSizeBudget(*ValuePtr);
@@ -2471,6 +2542,36 @@ void SMemTagTreeView::ApplyBudgetToNodes()
 
 ////////////////////////////////////////////////////////////////////////////////////////////////////
 
+void SMemTagTreeView::UpdateBudgetGroupsRec(const FMemTagBudgetGrouping* BudgetGrouping, const FMemTagBudgetGrouping* BudgetGroupingOverride, const FBaseTreeNodePtr& GroupPtr)
+{
+	for (const FBaseTreeNodePtr& NodePtr : GroupPtr->GetChildren())
+	{
+		if (NodePtr->Is<FMemTagBudgetGroupNode>())
+		{
+			FMemTagBudgetGroupNode& GroupNode = NodePtr->As<FMemTagBudgetGroupNode>();
+			const FMemTagBudgetGroup* BudgetGroup = nullptr;
+			if (BudgetGroupingOverride)
+			{
+				BudgetGroup = BudgetGroupingOverride->FindGroup(GroupNode.GetBudgetGroupName());
+			}
+			if (!BudgetGroup)
+			{
+				BudgetGroup = BudgetGrouping->FindGroup(GroupNode.GetBudgetGroupName());
+			}
+			if (BudgetGroup)
+			{
+				GroupNode.SetSizeBudget(BudgetGroup->GetMemMax());
+			}
+		}
+		else if (NodePtr->IsGroup())
+		{
+			UpdateBudgetGroupsRec(BudgetGrouping, BudgetGroupingOverride, NodePtr);
+		}
+	}
+}
+
+////////////////////////////////////////////////////////////////////////////////////////////////////
+
 void SMemTagTreeView::ResetBudgetForAllNodes()
 {
 	TSharedPtr<SMemoryProfilerWindow> ProfilerWindow = GetProfilerWindow();
@@ -2480,15 +2581,39 @@ void SMemTagTreeView::ResetBudgetForAllNodes()
 	}
 	FMemorySharedState& SharedState = ProfilerWindow->GetSharedState();
 
+	ResetBudgetGroupsRec(Root);
+
 	for (const FTableTreeNodePtr& NodePtr : TableRowNodes)
 	{
-		if (!NodePtr->Is<FMemTagNode>())
+		if (NodePtr->Is<FMemTagNode>())
 		{
-			continue;
+			FMemTagNode& MemTagNode = NodePtr->As<FMemTagNode>();
+			MemTagNode.ResetSizeBudget();
+			UpdateHighThreshold(MemTagNode, SharedState);
+		}
+		else if (NodePtr->Is<FMemTagBudgetGroupNode>())
+		{
+			FMemTagBudgetGroupNode& GroupNode = NodePtr->As<FMemTagBudgetGroupNode>();
+			GroupNode.ResetSizeBudget();
+		}
+	}
+}
+
+////////////////////////////////////////////////////////////////////////////////////////////////////
+
+void SMemTagTreeView::ResetBudgetGroupsRec(const FBaseTreeNodePtr& GroupPtr)
+{
+	for (const FBaseTreeNodePtr& NodePtr : GroupPtr->GetChildren())
+	{
+		if (NodePtr->Is<FMemTagBudgetGroupNode>())
+		{
+			FMemTagBudgetGroupNode& GroupNode = NodePtr->As<FMemTagBudgetGroupNode>();
+			GroupNode.ResetSizeBudget();
+		}
+		else if (NodePtr->IsGroup())
+		{
+			ResetBudgetGroupsRec(NodePtr);
 		}
-		FMemTagNode& MemTagNode = NodePtr->As<FMemTagNode>();
-		MemTagNode.ResetSizeBudget();
-		UpdateHighThreshold(MemTagNode, SharedState);
 	}
 }
 
diff --git a/Engine/Source/Developer/TraceInsights/Private/Insights/MemoryProfiler/Widgets/SMemTagTreeView.h b/Engine/Source/Developer/TraceInsights/Private/Insights/MemoryProfiler/Widgets/SMemTagTreeView.h
--- a/Engine/Source/Developer/TraceInsights/Private/Insights/MemoryProfiler/Widgets/SMemTagTreeView.h
+++ b/Engine/Source/Developer/TraceInsights/Private/Insights/MemoryProfiler/Widgets/SMemTagTreeView.h
@@ -13,6 +13,8 @@ namespace UE::Insights::MemoryProfiler
 {
 
 class FMemTagBudget;
+class FMemTagBudgetGrouping;
+class FMemTagBudgetMode;
 class FMemTagNode;
 class FMemoryGraphTrack;
 class FMemorySharedState;
@@ -183,7 +185,14 @@ public:
 
 	////////////////////////////////////////////////////////////////////////////////////////////////////
 
+	TSharedPtr<FMemTagBudget> GetBudget() const { return CurrentBudget; }
+	void GetBudgetGrouping(const FMemTagBudgetGrouping*& OutGrouping, const FMemTagBudgetGrouping*& OutGroupingOverride) const;
+
 protected:
+	const TCHAR* GetTagSetFilterBudgetCachedName() const;
+	const TCHAR* GetSelectedBudgetPlatformCachedName() const;
+	const FMemTagBudgetMode* GetSelectedBudgetMode() const;
+
 	virtual void InternalCreateGroupings() override;
 
 	virtual void ExtendMenu(TSharedRef<FExtender> Extender) override;
@@ -303,7 +312,9 @@ private:
 	void InitBudgetOptions();
 	void OnBudgetChanged();
 	void ApplyBudgetToNodes();
+	void UpdateBudgetGroupsRec(const FMemTagBudgetGrouping* BudgetGrouping, const FMemTagBudgetGrouping* BudgetGroupingOverride, const FBaseTreeNodePtr& GroupPtr);
 	void ResetBudgetForAllNodes();
+	void ResetBudgetGroupsRec(const FBaseTreeNodePtr& GroupPtr);
 	void UpdateHighThreshold(const FMemTagNode& MemTagNode, FMemorySharedState& SharedState);
 	void UpdateHighThreshold(const FMemTagNode& MemTagNode, FMemoryGraphTrack& GraphTrack);
 
