[Insights] Add callstacks to bookmarks, requires `callstack,module,bookmarks` channels active.

#tests editor

diff --git a/Engine/Source/Developer/TraceInsights/Private/Insights/TimingProfiler/Tracks/MarkersTimingTrack.cpp b/Engine/Source/Developer/TraceInsights/Private/Insights/TimingProfiler/Tracks/MarkersTimingTrack.cpp
--- a/Engine/Source/Developer/TraceInsights/Private/Insights/TimingProfiler/Tracks/MarkersTimingTrack.cpp
+++ b/Engine/Source/Developer/TraceInsights/Private/Insights/TimingProfiler/Tracks/MarkersTimingTrack.cpp
@@ -16,6 +16,8 @@
 // TraceServices
 #include "TraceServices/Model/Log.h"
 #include "TraceServices/Model/Screenshot.h"
+#include "TraceServices/Model/Bookmarks.h"
+#include "TraceServices/Model/Callstack.h"
 
 // TraceInsightsCore
 #include "InsightsCore/Common/PaintUtils.h"
@@ -28,9 +30,11 @@
 #include "Insights/ViewModels/TimingEvent.h"
 #include "Insights/ViewModels/TimingTrackViewport.h"
 #include "Insights/ViewModels/TooltipDrawState.h"
+#include "Insights/MemoryProfiler/ViewModels/CallstackFormatting.h"
 #include "Insights/Widgets/STimingView.h"
 
 #include <limits>
+#include <cmath>
 
 #define LOCTEXT_NAMESPACE "UE::Insights::TimingProfiler::FMarkersTimingTrack"
 
@@ -422,6 +426,28 @@ void FMarkersTimingTrack::UpdateCategory(const TraceServices::FLogCategoryInfo*&
 
 ////////////////////////////////////////////////////////////////////////////////////////////////////
 
+class FScreenShotTimingEvent : public FTimingEvent
+{
+public:
+	FScreenShotTimingEvent(const TSharedRef<const FMarkersTimingTrack> InTrack, double InStartTime, double InEndTime, uint32 InDepth, int32 InScreenshotId)
+		: FTimingEvent(InTrack, InStartTime, InEndTime, InDepth, InScreenshotId) {}
+	INSIGHTS_DECLARE_RTTI(FScreenShotTimingEvent, FTimingEvent)
+};
+INSIGHTS_IMPLEMENT_RTTI(FScreenShotTimingEvent)
+
+////////////////////////////////////////////////////////////////////////////////////////////////////
+
+class FBookmarkTimingEvent : public FTimingEvent
+{
+public:
+	FBookmarkTimingEvent(const TSharedRef<const FMarkersTimingTrack> InTrack, double InStartTime, double InEndTime, uint32 InDepth, uint32 InCallstackId)
+		: FTimingEvent(InTrack, InStartTime, InEndTime, InDepth, InCallstackId) {}
+	INSIGHTS_DECLARE_RTTI(FBookmarkTimingEvent, FTimingEvent)
+};
+INSIGHTS_IMPLEMENT_RTTI(FBookmarkTimingEvent)
+
+////////////////////////////////////////////////////////////////////////////////////////////////////
+
 const TSharedPtr<const ITimingEvent> FMarkersTimingTrack::GetEvent(float InPosX, float InPosY, const FTimingTrackViewport& Viewport) const
 {
 	TSharedPtr<ITimingEvent> TimingEvent;
@@ -450,8 +476,6 @@ const TSharedPtr<const ITimingEvent> FMarkersTimingTrack::GetEvent(float InPosX,
 			Width = TimeMarkerTexts[FoundIndex + 1].X;
 		}
 
-		uint32 ScreenshotId = TraceServices::FScreenshot::InvalidScreenshotId;
-
 		const FTimeMarkerTextInfo& Text = TimeMarkerTexts[FoundIndex];
 
 		TSharedPtr<const TraceServices::IAnalysisSession> Session = FInsightsManager::Get()->GetSession();
@@ -461,24 +485,44 @@ const TSharedPtr<const ITimingEvent> FMarkersTimingTrack::GetEvent(float InPosX,
 		}
 
 		TraceServices::FAnalysisSessionReadScope SessionReadScope(*Session.Get());
+
+		double Time = std::numeric_limits<double>::quiet_NaN();
+		uint32 ScreenshotId = TraceServices::FScreenshot::InvalidScreenshotId;
 		const TraceServices::ILogProvider& LogProvider = TraceServices::ReadLogProvider(*Session.Get());
 		LogProvider.ReadMessage(Text.LogIndex,
-			[&ScreenshotId, this](const TraceServices::FLogMessageInfo& Message)
+			[&ScreenshotId, &Time, this](const TraceServices::FLogMessageInfo& Message)
 			{
 				if (Message.Category == this->ScreenshotCategory)
 				{
 					check(Message.Line >= 0);
 					ScreenshotId = Message.Line;
 				}
+				if (Message.Category == this->BookmarkCategory)
+				{
+					Time = Message.Time;
+				}
 			});
 
-		if (ScreenshotId == TraceServices::FScreenshot::InvalidScreenshotId)
+		if (ScreenshotId != TraceServices::FScreenshot::InvalidScreenshotId)
 		{
-			return TimingEvent;
+			return MakeShared<FScreenShotTimingEvent>(SharedThis(this), Viewport.SlateUnitsToTime(Text.X), Viewport.SlateUnitsToTime(Text.X + Width), 0, ScreenshotId);;
 		}
 
-		//TODO: make a custom FScreenshotEvent
-		TimingEvent = MakeShared<FTimingEvent>(SharedThis(this), Viewport.SlateUnitsToTime(Text.X), Viewport.SlateUnitsToTime(Text.X + Width), 0, ScreenshotId);
+		if (FMath::IsFinite(Time))
+		{
+			uint32 CallstackId = 0;
+			const TraceServices::IBookmarkProvider& BookmarkProvider = TraceServices::ReadBookmarkProvider(*Session.Get());
+			BookmarkProvider.EnumerateBookmarks(
+				Time - UE_SMALL_NUMBER,
+				Time + UE_SMALL_NUMBER,
+				[&CallstackId](const TraceServices::FBookmark& Bookmark) {
+					CallstackId = Bookmark.CallstackId;
+			});
+			if (CallstackId > 0)
+			{
+				return MakeShared<FBookmarkTimingEvent>(SharedThis(this), Viewport.SlateUnitsToTime(Text.X), Viewport.SlateUnitsToTime(Text.X + Width), 0, CallstackId);
+			}
+		}
 	}
 
 	return TimingEvent;
@@ -502,54 +546,75 @@ void FMarkersTimingTrack::InitTooltip(FTooltipDrawState& InOutTooltip, const ITi
 		return;
 	}
 
-	TraceServices::FAnalysisSessionReadScope SessionReadScope(*Session.Get());
-	const FTimingEvent& Event = StaticCast<const FTimingEvent&>(InTooltipEvent);
-
-	const TraceServices::IScreenshotProvider& ScreenshotProvider = TraceServices::ReadScreenshotProvider(*Session.Get());
-	TSharedPtr<const TraceServices::FScreenshot> Screenshot = ScreenshotProvider.GetScreenshot((uint32)Event.GetType());
-
-	if (!Screenshot.IsValid())
+	if (InTooltipEvent.Is<FScreenShotTimingEvent>())
 	{
-		return;
-	}
-
-	IImageWrapperModule& ImageWrapperModule = FModuleManager::LoadModuleChecked<IImageWrapperModule>(FName("ImageWrapper"));
-	FImage Image;
-	if (ImageWrapperModule.DecompressImage(Screenshot->Data.GetData(), Screenshot->Size, Image))
-	{
-		constexpr int32 MAX_WIDTH = 640;
-		constexpr int32 MAX_HEIGHT = 480;
-
-		int32 ResizedX = Screenshot->Width;
-		int32 ResizedY = Screenshot->Height;
+		TraceServices::FAnalysisSessionReadScope SessionReadScope(*Session.Get());
+		const FScreenShotTimingEvent& Event = InTooltipEvent.As<FScreenShotTimingEvent>();
 
-		if (ResizedX > MAX_WIDTH)
+		const TraceServices::IScreenshotProvider& ScreenshotProvider = TraceServices::ReadScreenshotProvider(*Session.Get());
+		TSharedPtr<const TraceServices::FScreenshot> Screenshot = ScreenshotProvider.GetScreenshot((uint32)Event.GetType());
+	
+		if (!Screenshot.IsValid())
 		{
-			ResizedY = (ResizedY * MAX_WIDTH) / ResizedX;
-			ResizedX = MAX_WIDTH;
+			return;
 		}
-
-		if (ResizedY > MAX_HEIGHT)
+	
+		IImageWrapperModule& ImageWrapperModule = FModuleManager::LoadModuleChecked<IImageWrapperModule>(FName("ImageWrapper"));
+		FImage Image;
+		if (ImageWrapperModule.DecompressImage(Screenshot->Data.GetData(), Screenshot->Size, Image))
 		{
-			ResizedX = (ResizedX * MAX_HEIGHT) / ResizedY;
-			ResizedY = MAX_HEIGHT;
-		}
-
-		TSharedPtr<FSlateBrush> ImageBrush;
-		if (Screenshot->Width != ResizedX || Screenshot->Height != ResizedY)
-		{
-			FImage ResizedImage;
-			Image.ResizeTo(ResizedImage, ResizedX, ResizedY, ERawImageFormat::BGRA8, EGammaSpace::sRGB);
-			ImageBrush = FSlateDynamicImageBrush::CreateWithImageData(FName(Screenshot->Name), FVector2D(ResizedX, ResizedY), TArray<uint8>(ResizedImage.RawData));
+			constexpr int32 MAX_WIDTH = 640;
+			constexpr int32 MAX_HEIGHT = 480;
+	
+			int32 ResizedX = Screenshot->Width;
+			int32 ResizedY = Screenshot->Height;
+	
+			if (ResizedX > MAX_WIDTH)
+			{
+				ResizedY = (ResizedY * MAX_WIDTH) / ResizedX;
+				ResizedX = MAX_WIDTH;
+			}
+	
+			if (ResizedY > MAX_HEIGHT)
+			{
+				ResizedX = (ResizedX * MAX_HEIGHT) / ResizedY;
+				ResizedY = MAX_HEIGHT;
+			}
+	
+			TSharedPtr<FSlateBrush> ImageBrush;
+			if (Screenshot->Width != ResizedX || Screenshot->Height != ResizedY)
+			{
+				FImage ResizedImage;
+				Image.ResizeTo(ResizedImage, ResizedX, ResizedY, ERawImageFormat::BGRA8, EGammaSpace::sRGB);
+				ImageBrush = FSlateDynamicImageBrush::CreateWithImageData(FName(Screenshot->Name), FVector2D(ResizedX, ResizedY), TArray<uint8>(ResizedImage.RawData));
+			}
+			else
+			{
+				ImageBrush = FSlateDynamicImageBrush::CreateWithImageData(FName(Screenshot->Name), FVector2D(Screenshot->Width, Screenshot->Height), TArray<uint8>(Image.RawData));
+			}
+			InOutTooltip.SetImage(ImageBrush);
 		}
-		else
+	
+		InOutTooltip.UpdateLayout();
+	}
+	if (InTooltipEvent.Is<FBookmarkTimingEvent>())
+	{
+		const FBookmarkTimingEvent& Event = InTooltipEvent.As<FBookmarkTimingEvent>();
+		if (const TraceServices::ICallstacksProvider* CallstackProvider = TraceServices::ReadCallstacksProvider(*Session.Get()))
 		{
-			ImageBrush = FSlateDynamicImageBrush::CreateWithImageData(FName(Screenshot->Name), FVector2D(Screenshot->Width, Screenshot->Height), TArray<uint8>(Image.RawData));
+			if (const TraceServices::FCallstack* Callstack = CallstackProvider->GetCallstack((uint32)Event.GetType()))
+			{
+				uint32 CallstackSize = Callstack->Num();
+				for (uint8 Idx = 0; Idx < CallstackSize; ++Idx)
+				{
+					TStringBuilder<512> CallstackText;
+					UE::Insights::FormatStackFrame(*Callstack->Frame(Idx), CallstackText, UE::Insights::EStackFrameFormatFlags::ModuleSymbolFileAndLine);
+					InOutTooltip.AddTextLine(CallstackText, FTooltipDrawState::DefaultValueColor);
+				}
+			}
 		}
-		InOutTooltip.SetImage(ImageBrush);
+		InOutTooltip.UpdateLayout();
 	}
-
-	InOutTooltip.UpdateLayout();
 }
 
 ////////////////////////////////////////////////////////////////////////////////////////////////////
diff --git a/Engine/Source/Developer/TraceServices/Private/Analyzers/BookmarksTraceAnalysis.cpp b/Engine/Source/Developer/TraceServices/Private/Analyzers/BookmarksTraceAnalysis.cpp
--- a/Engine/Source/Developer/TraceServices/Private/Analyzers/BookmarksTraceAnalysis.cpp
+++ b/Engine/Source/Developer/TraceServices/Private/Analyzers/BookmarksTraceAnalysis.cpp
@@ -92,11 +92,12 @@ bool FBookmarksAnalyzer::OnEvent(uint16 RouteId, EStyle Style, const FOnEventCon
 		uint64 Cycle = EventData.GetValue<uint64>("Cycle");
 		double Timestamp = Context.EventTime.AsSeconds(Cycle);
 		TArrayView<const uint8> FormatArgsView = FTraceAnalyzerUtils::LegacyAttachmentArray("FormatArgs", Context);
+		uint32 CallstackId = EventData.GetValue<uint32>("CallstackId");
 
 		{
 			// Todo: update this to use provider locking instead of session locking
 			// FProviderEditScopeLock EditableBookmarkProviderLock (EditableBookmarkProvider);
-			EditableBookmarkProvider.AppendBookmark(BookmarkPoint, Timestamp, FormatArgsView.GetData());
+			EditableBookmarkProvider.AppendBookmark(BookmarkPoint, Timestamp, CallstackId, FormatArgsView.GetData());
 		}
 
 		if (EditableLogProvider)
diff --git a/Engine/Source/Developer/TraceServices/Private/Model/Bookmarks.cpp b/Engine/Source/Developer/TraceServices/Private/Model/Bookmarks.cpp
--- a/Engine/Source/Developer/TraceServices/Private/Model/Bookmarks.cpp
+++ b/Engine/Source/Developer/TraceServices/Private/Model/Bookmarks.cpp
@@ -47,7 +47,7 @@ void FBookmarkProvider::UpdateBookmarkSpec(uint64 BookmarkPoint, const TCHAR* Fo
 	BookmarkSpec.Line = Line;
 }
 
-void FBookmarkProvider::AppendBookmark(uint64 BookmarkPoint, double Time, const uint8* FormatArgs)
+void FBookmarkProvider::AppendBookmark(uint64 BookmarkPoint, double Time, uint32 CallstackId, const uint8* FormatArgs)
 {
 	Session.WriteAccessCheck();
 
@@ -56,6 +56,7 @@ void FBookmarkProvider::AppendBookmark(uint64 BookmarkPoint, double Time, const
 	TSharedRef<FBookmarkInternal> Bookmark = MakeShared<FBookmarkInternal>();
 	Bookmark->Time = Time;
 	Bookmark->Text = Session.StoreString(FormatBuffer);
+	Bookmark->CallstackId = CallstackId;
 	Bookmarks.Add(Bookmark);
 
 	CheckBookmarkForRegion(Bookmark);
@@ -63,13 +64,14 @@ void FBookmarkProvider::AppendBookmark(uint64 BookmarkPoint, double Time, const
 	Session.UpdateDurationSeconds(Time);
 }
 
-void FBookmarkProvider::AppendBookmark(uint64 BookmarkPoint, double Time, const TCHAR* Text)
+void FBookmarkProvider::AppendBookmark(uint64 BookmarkPoint, double Time, uint32 CallstackId, const TCHAR* Text)
 {
 	Session.WriteAccessCheck();
 
 	TSharedRef<FBookmarkInternal> Bookmark = MakeShared<FBookmarkInternal>();
 	Bookmark->Time = Time;
 	Bookmark->Text = Text;
+	Bookmark->CallstackId = CallstackId;
 	Bookmarks.Add(Bookmark);
 
 	CheckBookmarkForRegion(Bookmark);
@@ -134,6 +136,7 @@ void FBookmarkProvider::EnumerateBookmarks(double IntervalStart, double Interval
 		FBookmark Bookmark;
 		Bookmark.Time = InternalBookmark.Time;
 		Bookmark.Text = InternalBookmark.Text;
+		Bookmark.CallstackId = InternalBookmark.CallstackId;
 		Callback(Bookmark);
 	}
 }
diff --git a/Engine/Source/Developer/TraceServices/Private/Model/BookmarksPrivate.h b/Engine/Source/Developer/TraceServices/Private/Model/BookmarksPrivate.h
--- a/Engine/Source/Developer/TraceServices/Private/Model/BookmarksPrivate.h
+++ b/Engine/Source/Developer/TraceServices/Private/Model/BookmarksPrivate.h
@@ -22,6 +22,7 @@ struct FBookmarkInternal
 {
 	double Time = 0.0;
 	const TCHAR* Text = nullptr;
+	uint32 CallstackId = 0;
 };
 
 class FBookmarkProvider
@@ -38,8 +39,8 @@ public:
 
 	// implement IEditableBookmarkProvider
 	virtual void UpdateBookmarkSpec(uint64 BookmarkPoint, const TCHAR* FormatString, const TCHAR* File, int32 Line) override;
-	virtual void AppendBookmark(uint64 BookmarkPoint, double Time, const uint8* FormatArgs) override;
-	virtual void AppendBookmark(uint64 BookmarkPoint, double Time, const TCHAR* Text) override;
+	virtual void AppendBookmark(uint64 BookmarkPoint, double Time, uint32 CallstackId, const uint8* FormatArgs) override;
+	virtual void AppendBookmark(uint64 BookmarkPoint, double Time, uint32 CallstackId, const TCHAR* Text) override;
 
 private:
 	// check whether a bookmark is actually a region marker from the deprecated/5.1 regions implementation.
diff --git a/Engine/Source/Developer/TraceServices/Public/TraceServices/Model/Bookmarks.h b/Engine/Source/Developer/TraceServices/Public/TraceServices/Model/Bookmarks.h
--- a/Engine/Source/Developer/TraceServices/Public/TraceServices/Model/Bookmarks.h
+++ b/Engine/Source/Developer/TraceServices/Public/TraceServices/Model/Bookmarks.h
@@ -14,6 +14,7 @@ struct FBookmark
 {
 	double Time;
 	const TCHAR* Text;
+	uint32 CallstackId;
 };
 
 class IBookmarkProvider
@@ -51,7 +52,7 @@ public:
 	* @param Time			The time in seconds of the event.
 	* @param FormatArgs		The arguments to use in conjunction with the spec's FormatString.
 	*/
-	virtual void AppendBookmark(uint64 BookmarkPoint, double Time, const uint8* FormatArgs) = 0;
+	virtual void AppendBookmark(uint64 BookmarkPoint, double Time, uint32 CallstackId, const uint8* FormatArgs) = 0;
 
 	/*
 	* Append a new instance of a bookmark from the trace session.
@@ -61,7 +62,12 @@ public:
 	* @param Text			The fully formatted bookmark string for this instance. This may vary between instances.
 	*						This pointer is valid until the IAnalysisSession is deleted.
 	*/
-	virtual void AppendBookmark(uint64 BookmarkPoint, double Time, const TCHAR* Text) = 0;
+	virtual void AppendBookmark(uint64 BookmarkPoint, double Time, uint32 CallstackId, const TCHAR* Text) = 0;
+
+	UE_DEPRECATED(5.7, "Use AppendBookmark that has CallstackId parameter")
+	void AppendBookmark(uint64 BookmarkPoint, double Time, const uint8* FormatArgs) { AppendBookmark(BookmarkPoint, Time, 0, FormatArgs); }
+	UE_DEPRECATED(5.7, "Use AppendBookmark that has CallstackId parameter")
+	void AppendBookmark(uint64 BookmarkPoint, double Time, const TCHAR* Text) { AppendBookmark(BookmarkPoint, Time, 0, Text); }
 };
 
 TRACESERVICES_API FName GetBookmarkProviderName();
diff --git a/Engine/Source/Editor/UnrealEd/Private/Commandlets/SummarizeTraceCommandlet.cpp b/Engine/Source/Editor/UnrealEd/Private/Commandlets/SummarizeTraceCommandlet.cpp
--- a/Engine/Source/Editor/UnrealEd/Private/Commandlets/SummarizeTraceCommandlet.cpp
+++ b/Engine/Source/Editor/UnrealEd/Private/Commandlets/SummarizeTraceCommandlet.cpp
@@ -476,8 +476,8 @@ class FSummarizeBookmarksProvider
 	: public TraceServices::IEditableBookmarkProvider
 {
 	virtual void UpdateBookmarkSpec(uint64 BookmarkPoint, const TCHAR* FormatString, const TCHAR* File, int32 Line) override;
-	virtual void AppendBookmark(uint64 BookmarkPoint, double Time, const uint8* FormatString) override;
-	virtual void AppendBookmark(uint64 BookmarkPoint, double Time, const TCHAR* Text) override;
+	virtual void AppendBookmark(uint64 BookmarkPoint, double Time, uint32 CallstackId, const uint8* FormatString) override;
+	virtual void AppendBookmark(uint64 BookmarkPoint, double Time, uint32 CallstackId, const TCHAR* Text) override;
 
 	struct FBookmarkSpec
 	{
@@ -530,14 +530,14 @@ void FSummarizeBookmarksProvider::UpdateBookmarkSpec(uint64 BookmarkPoint, const
 	BookmarkSpec.Line = Line;
 }
 
-void FSummarizeBookmarksProvider::AppendBookmark(uint64 BookmarkPoint, double Time, const uint8* FormatArgs)
+void FSummarizeBookmarksProvider::AppendBookmark(uint64 BookmarkPoint, double Time, uint32 CallstackId, const uint8* FormatArgs)
 {
 	FBookmarkSpec& BookmarkSpec = GetSpec(BookmarkPoint);
 	TraceServices::StringFormat(FormatBuffer, FormatBufferSize - 1, TempBuffer, FormatBufferSize - 1, BookmarkSpec.FormatString, FormatArgs);
-	FSummarizeBookmarksProvider::AppendBookmark(BookmarkPoint, Time, FormatBuffer);
+	FSummarizeBookmarksProvider::AppendBookmark(BookmarkPoint, Time, CallstackId, FormatBuffer);
 }
 
-void FSummarizeBookmarksProvider::AppendBookmark(uint64 BookmarkPoint, double Time, const TCHAR* Text)
+void FSummarizeBookmarksProvider::AppendBookmark(uint64 BookmarkPoint, double Time, uint32 CallstackId, const TCHAR* Text)
 {
 	FString Name(Text);
 
diff --git a/Engine/Source/Editor/UnrealEd/Private/Commandlets/SummarizeTraceEditorWorkflowsCommandlet.cpp b/Engine/Source/Editor/UnrealEd/Private/Commandlets/SummarizeTraceEditorWorkflowsCommandlet.cpp
--- a/Engine/Source/Editor/UnrealEd/Private/Commandlets/SummarizeTraceEditorWorkflowsCommandlet.cpp
+++ b/Engine/Source/Editor/UnrealEd/Private/Commandlets/SummarizeTraceEditorWorkflowsCommandlet.cpp
@@ -20,8 +20,8 @@ class FSummarizeBookmarksProviderEditorWorkflows
 {
 public:
 	virtual void UpdateBookmarkSpec(uint64 BookmarkPoint, const TCHAR* FormatString, const TCHAR* File, int32 Line) override;
-	virtual void AppendBookmark(uint64 BookmarkPoint, double Time, const uint8* FormatString) override;
-	virtual void AppendBookmark(uint64 BookmarkPoint, double Time, const TCHAR* Text) override;
+	virtual void AppendBookmark(uint64 BookmarkPoint, double Time, uint32 CallstackId, const uint8* FormatString) override;
+	virtual void AppendBookmark(uint64 BookmarkPoint, double Time, uint32 CallstackId, const TCHAR* Text) override;
 	struct FBookmarkSpec
 	{
 		const TCHAR* File = nullptr;
@@ -67,14 +67,14 @@ void FSummarizeBookmarksProviderEditorWorkflows::UpdateBookmarkSpec(uint64 Bookm
 	BookmarkSpec.Line = Line;
 }
 
-void FSummarizeBookmarksProviderEditorWorkflows::AppendBookmark(uint64 BookmarkPoint, double Time, const uint8* FormatArgs)
+void FSummarizeBookmarksProviderEditorWorkflows::AppendBookmark(uint64 BookmarkPoint, double Time, uint32 CallstackId, const uint8* FormatArgs)
 {
 	FBookmarkSpec& BookmarkSpec = GetSpec(BookmarkPoint);
 	TraceServices::StringFormat(FormatBuffer, FormatBufferSize - 1, TempBuffer, FormatBufferSize - 1, BookmarkSpec.FormatString, FormatArgs);
-	FSummarizeBookmarksProviderEditorWorkflows::AppendBookmark(BookmarkPoint, Time, FormatBuffer);
+	FSummarizeBookmarksProviderEditorWorkflows::AppendBookmark(BookmarkPoint, Time, CallstackId, FormatBuffer);
 }
 
-void FSummarizeBookmarksProviderEditorWorkflows::AppendBookmark(uint64 BookmarkPoint, double Time, const TCHAR* Text)
+void FSummarizeBookmarksProviderEditorWorkflows::AppendBookmark(uint64 BookmarkPoint, double Time, uint32 CallstackId, const TCHAR* Text)
 {
 	FString Name(Text);
 
diff --git a/Engine/Source/Runtime/Core/Private/ProfilingDebugging/MiscTrace.cpp b/Engine/Source/Runtime/Core/Private/ProfilingDebugging/MiscTrace.cpp
--- a/Engine/Source/Runtime/Core/Private/ProfilingDebugging/MiscTrace.cpp
+++ b/Engine/Source/Runtime/Core/Private/ProfilingDebugging/MiscTrace.cpp
@@ -24,6 +24,7 @@ UE_TRACE_EVENT_BEGIN(Misc, Bookmark)
 	UE_TRACE_EVENT_FIELD(uint64, Cycle)
 	UE_TRACE_EVENT_FIELD(const void*, BookmarkPoint)
 	UE_TRACE_EVENT_FIELD(uint8[], FormatArgs)
+	UE_TRACE_EVENT_FIELD(uint32, CallstackId)
 UE_TRACE_EVENT_END()
 
 UE_TRACE_EVENT_BEGIN(Misc, RegionBegin)
@@ -116,20 +117,18 @@ void FMiscTrace::OutputEndRegionWithId(uint64 RegionId)
 		<< RegionEndWithId.RegionId(RegionId);
 }
 
-void FMiscTrace::OutputBookmarkInternal(const void* BookmarkPoint, uint16 EncodedFormatArgsSize, uint8* EncodedFormatArgs)
+void FMiscTrace::OutputBookmarkInternal(const void* BookmarkPoint, uint32 CallstackId, uint16 EncodedFormatArgsSize, uint8* EncodedFormatArgs)
 {
-	UE_TRACE_LOG(Misc, Bookmark, BookmarkChannel)
-		<< Bookmark.Cycle(FPlatformTime::Cycles64())
-		<< Bookmark.BookmarkPoint(BookmarkPoint)
-		<< Bookmark.FormatArgs(EncodedFormatArgs, EncodedFormatArgsSize);
+	OutputBookmarkInternalCycles(FPlatformTime::Cycles64(), BookmarkPoint, CallstackId, EncodedFormatArgsSize, EncodedFormatArgs);
 }
 
-void FMiscTrace::OutputBookmarkInternalCycles(uint64 Cycles, const void* BookmarkPoint, uint16 EncodedFormatArgsSize, uint8* EncodedFormatArgs)
+void FMiscTrace::OutputBookmarkInternalCycles(uint64 Cycles, const void* BookmarkPoint, uint32 CallstackId, uint16 EncodedFormatArgsSize, uint8* EncodedFormatArgs)
 {
 	UE_TRACE_LOG(Misc, Bookmark, BookmarkChannel)
 		<< Bookmark.Cycle(Cycles)
 		<< Bookmark.BookmarkPoint(BookmarkPoint)
-		<< Bookmark.FormatArgs(EncodedFormatArgs, EncodedFormatArgsSize);
+		<< Bookmark.FormatArgs(EncodedFormatArgs, EncodedFormatArgsSize)
+		<< Bookmark.CallstackId(CallstackId);
 }
 
 void FMiscTrace::OutputBeginFrame(ETraceFrameType FrameType)
diff --git a/Engine/Source/Runtime/Core/Public/ProfilingDebugging/MiscTrace.h b/Engine/Source/Runtime/Core/Public/ProfilingDebugging/MiscTrace.h
--- a/Engine/Source/Runtime/Core/Public/ProfilingDebugging/MiscTrace.h
+++ b/Engine/Source/Runtime/Core/Public/ProfilingDebugging/MiscTrace.h
@@ -6,6 +6,7 @@
 #include "CoreTypes.h"
 #include "HAL/PreprocessorHelpers.h"
 #include "Misc/Build.h"
+#include "ProfilingDebugging/CallstackTrace.h"
 #include "ProfilingDebugging/FormatArgsTrace.h"
 #include "Trace/Config.h"
 
@@ -53,25 +54,26 @@ class FName;
 struct FMiscTrace
 {
 	CORE_API static void OutputBookmarkSpec(const void* BookmarkPoint, const ANSICHAR* File, int32 Line, const TCHAR* Format);
+
 	template <typename... Types>
-	static void OutputBookmark(const void* BookmarkPoint, Types... FormatArgs)
+	static void OutputBookmark(uint32 CallstackId, const void* BookmarkPoint, Types... FormatArgs)
 	{
 		uint8 FormatArgsBuffer[4096];
 		uint16 FormatArgsSize = FFormatArgsTrace::EncodeArguments(FormatArgsBuffer, FormatArgs...);
 		if (FormatArgsSize)
 		{
-			OutputBookmarkInternal(BookmarkPoint, FormatArgsSize, FormatArgsBuffer);
+			OutputBookmarkInternal(BookmarkPoint, CallstackId, FormatArgsSize, FormatArgsBuffer);
 		}
 	}
 
 	template <typename... Types>
-	static void OutputBookmarkCycles(uint64 Cycles, const void* BookmarkPoint, Types... FormatArgs)
+	static void OutputBookmarkCycles(uint32 CallstackId, uint64 Cycles, const void* BookmarkPoint, Types... FormatArgs)
 	{
 		uint8 FormatArgsBuffer[4096];
 		uint16 FormatArgsSize = FFormatArgsTrace::EncodeArguments(FormatArgsBuffer, FormatArgs...);
 		if (FormatArgsSize)
 		{
-			OutputBookmarkInternalCycles(Cycles, BookmarkPoint, FormatArgsSize, FormatArgsBuffer);
+			OutputBookmarkInternalCycles(Cycles, BookmarkPoint, CallstackId, FormatArgsSize, FormatArgsBuffer);
 		}
 	}
 
@@ -88,9 +90,23 @@ struct FMiscTrace
 	CORE_API static bool ShouldTraceBookmark();
 	CORE_API static bool ShouldTraceRegion();
 
+	template <typename... Types>
+	UE_DEPRECATED(5.7, "Use OutputBookmark that has CallstackId parameter")
+	static void OutputBookmark(const void* BookmarkPoint, Types... FormatArgs)
+	{
+		OutputBookmark(CallstackTrace_GetCurrentId(), BookmarkPoint, Forward<Types>(FormatArgs)...);
+	}
+
+	template <typename... Types>
+	UE_DEPRECATED(5.7, "Use OutputBookmarkCycles that has CallstackId parameter")
+	static void OutputBookmarkCycles(uint64 Cycles, const void* BookmarkPoint, Types... FormatArgs)
+	{
+		OutputBookmarkCycles(CallstackTrace_GetCurrentId(), Cycles, BookmarkPoint, Forward<Types>(FormatArgs)...);
+	}
+
 private:
-	CORE_API static void OutputBookmarkInternal(const void* BookmarkPoint, uint16 EncodedFormatArgsSize, uint8* EncodedFormatArgs);
-	CORE_API static void OutputBookmarkInternalCycles(uint64 Cycles, const void* BookmarkPoint, uint16 EncodedFormatArgsSize, uint8* EncodedFormatArgs);
+	CORE_API static void OutputBookmarkInternal(const void* BookmarkPoint, uint32 CallstackId, uint16 EncodedFormatArgsSize, uint8* EncodedFormatArgs);
+	CORE_API static void OutputBookmarkInternalCycles(uint64 Cycles, const void* BookmarkPoint, uint32 CallstackId, uint16 EncodedFormatArgsSize, uint8* EncodedFormatArgs);
 };
 
 #define TRACE_BOOKMARK(Format, ...) \
@@ -105,7 +121,7 @@ if (UE_TRACE_CHANNELEXPR_IS_ENABLED(BookmarkChannel)) \
 		FMiscTrace::OutputBookmarkSpec(&PREPROCESSOR_JOIN(__BookmarkPoint, __LINE__), __FILE__, __LINE__, Format); \
 		PREPROCESSOR_JOIN(__BookmarkPoint, __LINE__) = true; \
 	} \
-	FMiscTrace::OutputBookmark(&PREPROCESSOR_JOIN(__BookmarkPoint, __LINE__), ##__VA_ARGS__); \
+	FMiscTrace::OutputBookmark(CallstackTrace_GetCurrentId(), &PREPROCESSOR_JOIN(__BookmarkPoint, __LINE__), ##__VA_ARGS__); \
 }
 
 #define TRACE_BOOKMARK_CYCLES(Cycles, Format, ...) \
@@ -120,7 +136,7 @@ if (UE_TRACE_CHANNELEXPR_IS_ENABLED(BookmarkChannel)) \
 		FMiscTrace::OutputBookmarkSpec(&PREPROCESSOR_JOIN(__BookmarkPoint, __LINE__), __FILE__, __LINE__, Format); \
 		PREPROCESSOR_JOIN(__BookmarkPoint, __LINE__) = true; \
 	} \
-	FMiscTrace::OutputBookmarkCycles(Cycles, &PREPROCESSOR_JOIN(__BookmarkPoint, __LINE__), ##__VA_ARGS__); \
+	FMiscTrace::OutputBookmarkCycles(CallstackTrace_GetCurrentId(), Cycles, &PREPROCESSOR_JOIN(__BookmarkPoint, __LINE__), ##__VA_ARGS__); \
 }
 
 #define TRACE_BEGIN_REGION(RegionName) \
