Add LiveCoding to UnrealInsights

diff --git a/Engine/Source/Developer/TraceInsights/Private/Insights/Common/InsightsMenuBuilder.cpp b/Engine/Source/Developer/TraceInsights/Private/Insights/Common/InsightsMenuBuilder.cpp
--- a/Engine/Source/Developer/TraceInsights/Private/Insights/Common/InsightsMenuBuilder.cpp
+++ b/Engine/Source/Developer/TraceInsights/Private/Insights/Common/InsightsMenuBuilder.cpp
@@ -17,6 +17,10 @@
 #include "Insights/InsightsStyle.h"
 #include "Insights/InsightsManager.h"
 
+#if WITH_LIVE_CODING
+#include "ILiveCodingModule.h"
+#endif
+
 #define LOCTEXT_NAMESPACE "UE::Insights::MenuBuilder"
 
 namespace UE::Insights
@@ -126,6 +130,31 @@ void FInsightsMenuBuilder::PopulateMenu(FMenuBuilder& MenuBuilder)
 		MenuBuilder.EndSection();
 	}
 
+#if WITH_LIVE_CODING
+	if (ILiveCodingModule* LiveCoding = FModuleManager::GetModulePtr<ILiveCodingModule>(LIVE_CODING_MODULE_NAME))
+	{
+		MenuBuilder.BeginSection("Developer");
+		MenuBuilder.AddMenuEntry(
+			LOCTEXT("LiveCodingCompile", "Live Coding Compile"),
+			LOCTEXT("LiveCodingCompile_Tooltip", "Compiles the Unreal Insights code and patches this executable"),
+			FSlateIcon(FAppStyle::Get().GetStyleSetName(), "WidgetReflector.Icon"),
+			FUIAction(
+				FExecuteAction::CreateLambda([]()
+				{
+					ILiveCodingModule& LiveCoding = FModuleManager::GetModuleChecked<ILiveCodingModule>(LIVE_CODING_MODULE_NAME);
+					LiveCoding.Compile();
+				}),
+				FCanExecuteAction::CreateLambda([]() -> bool
+				{
+					ILiveCodingModule* LiveCoding = FModuleManager::GetModulePtr<ILiveCodingModule>(LIVE_CODING_MODULE_NAME);
+					return LiveCoding != nullptr && LiveCoding->IsEnabledForSession();
+				}))
+		);
+			
+		MenuBuilder.EndSection();
+	}
+#endif
+
 #if !UE_BUILD_SHIPPING
 	// Open Starship Test Suite
 	{
diff --git a/Engine/Source/Developer/TraceInsights/TraceInsights.Build.cs b/Engine/Source/Developer/TraceInsights/TraceInsights.Build.cs
--- a/Engine/Source/Developer/TraceInsights/TraceInsights.Build.cs
+++ b/Engine/Source/Developer/TraceInsights/TraceInsights.Build.cs
@@ -70,5 +70,10 @@ public class TraceInsights : ModuleRules
 				"MessageLog",
 			}
 		);
+
+		if (Target.bWithLiveCoding)
+		{
+			PrivateDependencyModuleNames.Add("LiveCoding");
+		}
 	}
 }
diff --git a/Engine/Source/Programs/UnrealInsights/Private/UserInterfaceCommand.cpp b/Engine/Source/Programs/UnrealInsights/Private/UserInterfaceCommand.cpp
--- a/Engine/Source/Programs/UnrealInsights/Private/UserInterfaceCommand.cpp
+++ b/Engine/Source/Programs/UnrealInsights/Private/UserInterfaceCommand.cpp
@@ -76,6 +76,10 @@ void FUserInterfaceCommand::Run(bool bFrontendMode, const FString& TraceFileToOp
 	// Crank up a normal Slate application using the platform's standalone renderer.
 	FSlateApplication::InitializeAsStandaloneApplication(GetStandardStandaloneRenderer());
 
+	// Used by Live Coding, among other things.
+	FPlatformMisc::SetUBTTargetName(TEXT("UnrealInsights"));
+	FModuleManager::Get().LoadModule("LiveCoding");
+
 	// Load required modules.
 	FModuleManager::Get().LoadModuleChecked("TraceInsightsCore");
 	if (bFrontendMode)
diff --git a/Engine/Source/Programs/UnrealInsights/UnrealInsights.Build.cs b/Engine/Source/Programs/UnrealInsights/UnrealInsights.Build.cs
--- a/Engine/Source/Programs/UnrealInsights/UnrealInsights.Build.cs
+++ b/Engine/Source/Programs/UnrealInsights/UnrealInsights.Build.cs
@@ -77,5 +77,10 @@ public class UnrealInsights : ModuleRules
 				}
 			);
 		}
+
+		if (Target.bWithLiveCoding)
+		{
+			PrivateDependencyModuleNames.Add("LiveCoding");
+		}
 	}
 }
diff --git a/Engine/Source/Programs/UnrealInsights/UnrealInsights.Target.cs b/Engine/Source/Programs/UnrealInsights/UnrealInsights.Target.cs
--- a/Engine/Source/Programs/UnrealInsights/UnrealInsights.Target.cs
+++ b/Engine/Source/Programs/UnrealInsights/UnrealInsights.Target.cs
@@ -33,6 +33,8 @@ public class UnrealInsightsTarget : TargetRules
 
 		// For UI functionality
 		bBuildDeveloperTools = true;
+		
+		bWithLiveCoding = true;
 
 		bHasExports = false;
 
