System.Insights.Analysis.TimingInsights test cases OOM or fail.

#rnx

diff --git a/Engine/Source/Developer/TraceInsights/Private/Insights/Tests/FunctionalTests/TimingProfilerTests.cpp b/Engine/Source/Developer/TraceInsights/Private/Insights/Tests/FunctionalTests/TimingProfilerTests.cpp
--- a/Engine/Source/Developer/TraceInsights/Private/Insights/Tests/FunctionalTests/TimingProfilerTests.cpp
+++ b/Engine/Source/Developer/TraceInsights/Private/Insights/Tests/FunctionalTests/TimingProfilerTests.cpp
@@ -23,14 +23,17 @@ DEFINE_LOG_CATEGORY(TimingProfilerTests);
 
 #if WITH_AUTOMATION_TESTS
 
+bool AnalyzeTrace(const FString& RelativePath, FAutomationTestBase* Test)
+{
+	FInsightsTestUtils Utils(Test);
+	FString AbsolutePath = FPaths::RootDir() / RelativePath;
+	return Utils.AnalyzeTrace(*AbsolutePath);
+}
+
 IMPLEMENT_SIMPLE_AUTOMATION_TEST(FEnumerateTest, "System.Insights.Analysis.TimingInsights.Enumerate", EAutomationTestFlags::EditorContext | EAutomationTestFlags::ProgramContext | EAutomationTestFlags::EngineFilter)
 bool FEnumerateTest::RunTest(const FString& Parameters)
 {
-	FInsightsTestUtils Utils(this);
-
-	FString Path = FPaths::RootDir() + TEXT("EngineTest/SourceAssets/Utrace/r424_win64_game_11590231.utrace");
-	bool ret = Utils.AnalyzeTrace(*Path);
-	if (!ret)
+	if (!AnalyzeTrace(TEXT("EngineTest/SourceAssets/Utrace/r424_win64_game_11590231.utrace"), this))
 	{
 		return !HasAnyErrors();
 	}
@@ -56,11 +59,7 @@ bool FEnumerateTest::RunTest(const FString& Parameters)
 IMPLEMENT_SIMPLE_AUTOMATION_TEST(FEnumeratePerformanceTest, "System.Insights.Analysis.TimingInsights.EnumeratePerformance", EAutomationTestFlags::EditorContext | EAutomationTestFlags::ProgramContext | EAutomationTestFlags::EngineFilter)
 bool FEnumeratePerformanceTest::RunTest(const FString& Parameters)
 {
-	FInsightsTestUtils Utils(this);
-
-	FString Path = FPaths::RootDir() + TEXT("EngineTest/SourceAssets/Utrace/r425_win64_game_13649855.utrace");
-	bool ret = Utils.AnalyzeTrace(*Path);
-	if (!ret)
+	if (!AnalyzeTrace(TEXT("EngineTest/SourceAssets/Utrace/r425_win64_game_13649855.utrace"), this))
 	{
 		return !HasAnyErrors();
 	}
@@ -96,11 +95,7 @@ bool FEnumeratePerformanceTest::RunTest(const FString& Parameters)
 IMPLEMENT_SIMPLE_AUTOMATION_TEST(FEnumerateFastTest, "System.Insights.Analysis.TimingInsights.EnumerateFast", EAutomationTestFlags::EditorContext | EAutomationTestFlags::ProgramContext | EAutomationTestFlags::EngineFilter)
 bool FEnumerateFastTest::RunTest(const FString& Parameters)
 {
-	FInsightsTestUtils Utils(this);
-
-	FString Path = FPaths::RootDir() + TEXT("EngineTest/SourceAssets/Utrace/r423_win64_game_10478456.utrace");
-	bool ret = Utils.AnalyzeTrace(*Path);
-	if (!ret)
+	if (!AnalyzeTrace(TEXT("EngineTest/SourceAssets/Utrace/r423_win64_game_10478456.utrace"), this))
 	{
 		return !HasAnyErrors();
 	}
@@ -126,6 +121,11 @@ bool FEnumerateFastTest::RunTest(const FString& Parameters)
 IMPLEMENT_SIMPLE_AUTOMATION_TEST(EnumerateEventsToFile, "System.Insights.Trace.Analysis.TimingInsights.EnumerateEventsToFile", EAutomationTestFlags::EditorContext | EAutomationTestFlags::ProgramContext | EAutomationTestFlags::EngineFilter)
 bool EnumerateEventsToFile::RunTest(const FString& Parameters)
 {
+	if (!AnalyzeTrace(TEXT("EngineTest/SourceAssets/Utrace/r423_win64_game_10478456.utrace"), this))
+	{
+		return !HasAnyErrors();
+	}
+
 	double SessionTime = 0.0;
 	uint32 TimelineIndex = (uint32)-1;
 
@@ -144,21 +144,21 @@ bool EnumerateEventsToFile::RunTest(const FString& Parameters)
 		TimingProfilerProvider.ReadTimers([&TimerReader](const TraceServices::ITimingProfilerTimerReader& Out) { TimerReader = &Out; });
 
 		FString FilePath = FPaths::ProjectSavedDir() + TEXT("/EnumerateEventsToFile.txt");
-		TUniquePtr<FArchive> Ar = TUniquePtr<FArchive>(IFileManager::Get().CreateFileWriter(*FilePath));
+		TUniquePtr<FArchive> ArchiveWriter = TUniquePtr<FArchive>(IFileManager::Get().CreateFileWriter(*FilePath));
 		TimingProfilerProvider.ReadTimeline(TimelineIndex,
-			[Session, &Ar, TimerReader](const TraceServices::ITimingProfilerProvider::Timeline& Timeline)
+			[Session, &ArchiveWriter, TimerReader](const TraceServices::ITimingProfilerProvider::Timeline& Timeline)
 			{
 				Timeline.EnumerateEvents(-1.0, Session->GetDurationSeconds() + 1.0,
-					[&Ar, TimerReader](double EventStartTime, double EventEndTime, uint32 EventDepth, const TraceServices::FTimingProfilerEvent& Event)
+					[&ArchiveWriter, TimerReader](double EventStartTime, double EventEndTime, uint32 EventDepth, const TraceServices::FTimingProfilerEvent& Event)
 					{
 						const TraceServices::FTimingProfilerTimer* Timer = TimerReader->GetTimer(Event.TimerIndex);
 
-						(*Ar) << FString::Printf(TEXT("%s %f %f %d\n"), Timer->Name, EventStartTime, EventEndTime, EventDepth).GetCharArray();
+						(*ArchiveWriter) << FString::Printf(TEXT("%s %f %f %d\n"), Timer->Name, EventStartTime, EventEndTime, EventDepth).GetCharArray();
 						return TraceServices::EEventEnumerate::Continue;
 					});
 			});
 
-		Ar->Close();
+		ArchiveWriter->Close();
 	}
 
 	return !HasAnyErrors();
@@ -185,21 +185,21 @@ bool EnumerateScopesToFile::RunTest(const FString& Parameters)
 		TimingProfilerProvider.ReadTimers([&TimerReader](const TraceServices::ITimingProfilerTimerReader& Out) { TimerReader = &Out; });
 
 		FString FilePath = FPaths::ProjectSavedDir() + TEXT("/EnumerateScopesToFile.txt");
-		TUniquePtr<FArchive> Ar = TUniquePtr<FArchive>(IFileManager::Get().CreateFileWriter(*FilePath));
+		TUniquePtr<FArchive> ArchiveWriter = TUniquePtr<FArchive>(IFileManager::Get().CreateFileWriter(*FilePath));
 		TimingProfilerProvider.ReadTimeline(TimelineIndex,
-			[Session, &Ar, TimerReader](const TraceServices::ITimingProfilerProvider::Timeline& Timeline)
+			[Session, &ArchiveWriter, TimerReader](const TraceServices::ITimingProfilerProvider::Timeline& Timeline)
 			{
 				Timeline.EnumerateEvents(-1.0, Session->GetDurationSeconds() + 1.0,
-					[&Ar, TimerReader](bool bStart, double Time, const TraceServices::FTimingProfilerEvent& Event)
+					[&ArchiveWriter, TimerReader](bool bStart, double Time, const TraceServices::FTimingProfilerEvent& Event)
 					{
 						const TraceServices::FTimingProfilerTimer* Timer = TimerReader->GetTimer(Event.TimerIndex);
 
-						(*Ar) << FString::Printf(TEXT("%s %d %f\n"), Timer->Name, (int32)bStart, Time).GetCharArray();
+						(*ArchiveWriter) << FString::Printf(TEXT("%s %d %f\n"), Timer->Name, (int32)bStart, Time).GetCharArray();
 						return TraceServices::EEventEnumerate::Continue;
 					});
 			});
 
-		Ar->Close();
+		ArchiveWriter->Close();
 	}
 
 	return !HasAnyErrors();
