[Insights] Fixed grouping for the budget file defined groups to put the ungrouped tags into a separate group node. This resolves the issue with further groupings in the hierarchy.

diff --git a/Engine/Source/Developer/TraceInsights/Private/Insights/MemoryProfiler/ViewModels/MemTagBudgetGrouping.cpp b/Engine/Source/Developer/TraceInsights/Private/Insights/MemoryProfiler/ViewModels/MemTagBudgetGrouping.cpp
--- a/Engine/Source/Developer/TraceInsights/Private/Insights/MemoryProfiler/ViewModels/MemTagBudgetGrouping.cpp
+++ b/Engine/Source/Developer/TraceInsights/Private/Insights/MemoryProfiler/ViewModels/MemTagBudgetGrouping.cpp
@@ -116,6 +116,7 @@ void FMemTagBudgetNodeGrouping::GroupNodes(
 	}
 
 	TMap<const void*, TSharedPtr<FMemTagBudgetGroupNode>> Groups;
+	TSharedPtr<FTableTreeNode> UngroupedTagsGroupNode;
 
 	for (FTableTreeNodePtr NodePtr : Nodes)
 	{
@@ -124,43 +125,52 @@ void FMemTagBudgetNodeGrouping::GroupNodes(
 			return;
 		}
 
-		if (NodePtr->IsGroup() || !NodePtr->Is<FMemTagNode>())
+		if (NodePtr->IsGroup())
 		{
 			ParentGroup.AddChildAndSetParent(NodePtr);
 			continue;
 		}
 
-		FMemTagNode& MemTagNode = NodePtr->As<FMemTagNode>();
-		check(MemTagNode.GetMemTag() != nullptr);
-		FMemoryTag& MemTag = *MemTagNode.GetMemTag();
-
-		const TCHAR* CachedGroupName = nullptr;
-		int64 SizeBudget = 0;
-		for (FGroupRule& Rule : GroupRules)
+		if (NodePtr->Is<FMemTagNode>())
 		{
-			if (Rule.Match(MemTag))
+			FMemTagNode& MemTagNode = NodePtr->As<FMemTagNode>();
+			check(MemTagNode.GetMemTag() != nullptr);
+			FMemoryTag& MemTag = *MemTagNode.GetMemTag();
+
+			const TCHAR* CachedGroupName = nullptr;
+			int64 SizeBudget = 0;
+			for (FGroupRule& Rule : GroupRules)
 			{
-				CachedGroupName = Rule.GetGroupName();
-				SizeBudget = Rule.GetSizeBudget();
-				break;
+				if (Rule.Match(MemTag))
+				{
+					CachedGroupName = Rule.GetGroupName();
+					SizeBudget = Rule.GetSizeBudget();
+					break;
+				}
+			}
+			if (CachedGroupName)
+			{
+				TSharedPtr<FMemTagBudgetGroupNode> GroupNode = Groups.FindRef(CachedGroupName);
+				if (!GroupNode)
+				{
+					FName GroupName(CachedGroupName);
+					GroupNode = MakeShared<FMemTagBudgetGroupNode>(GroupName, InParentTable, CachedGroupName);
+					GroupNode->SetSizeBudget(SizeBudget);
+					ParentGroup.AddChildAndSetParent(GroupNode);
+					Groups.Add(CachedGroupName, GroupNode);
+				}
+				GroupNode->AddChildAndSetParent(NodePtr);
+				continue;
 			}
-		}
-		if (!CachedGroupName)
-		{
-			ParentGroup.AddChildAndSetParent(NodePtr);
-			continue;
 		}
 
-		TSharedPtr<FMemTagBudgetGroupNode> GroupNode = Groups.FindRef(CachedGroupName);
-		if (!GroupNode)
+		if (!UngroupedTagsGroupNode)
 		{
-			FName NodeName(CachedGroupName);
-			GroupNode = MakeShared<FMemTagBudgetGroupNode>(NodeName, InParentTable, CachedGroupName);
-			GroupNode->SetSizeBudget(SizeBudget);
-			ParentGroup.AddChildAndSetParent(GroupNode);
-			Groups.Add(CachedGroupName, GroupNode);
+			static FName UngroupedTagsGroupName("Ungrouped");
+			UngroupedTagsGroupNode = MakeShared<FTableTreeNode>(UngroupedTagsGroupName, InParentTable);
+			ParentGroup.AddChildAndSetParent(UngroupedTagsGroupNode);
 		}
-		GroupNode->AddChildAndSetParent(NodePtr);
+		UngroupedTagsGroupNode->AddChildAndSetParent(NodePtr);
 	}
 }
 
