[Insights] STaskTableTreeView: Cleaned up the unused included header files. Also fixed the usage of TraceServices namespace to not bleed into other compile units.

#rnx

diff --git a/Engine/Source/Developer/TraceInsights/Private/Insights/TaskGraphProfiler/Widgets/STaskTableTreeView.cpp b/Engine/Source/Developer/TraceInsights/Private/Insights/TaskGraphProfiler/Widgets/STaskTableTreeView.cpp
--- a/Engine/Source/Developer/TraceInsights/Private/Insights/TaskGraphProfiler/Widgets/STaskTableTreeView.cpp
+++ b/Engine/Source/Developer/TraceInsights/Private/Insights/TaskGraphProfiler/Widgets/STaskTableTreeView.cpp
@@ -3,7 +3,6 @@
 #include "STaskTableTreeView.h"
 
 #include "Framework/MultiBox/MultiBoxBuilder.h"
-#include "Internationalization/Regex.h"
 #include "ISourceCodeAccessModule.h"
 #include "ISourceCodeAccessor.h"
 #include "Misc/Paths.h"
@@ -14,15 +13,10 @@
 // TraceServices
 #include "TraceServices/Model/TasksProfiler.h"
 
-// TraceInsightsCore
-#include "InsightsCore/Filter/ViewModels/FilterConfigurator.h"
-#include "InsightsCore/Table/ViewModels/TableColumn.h"
-
 // TraceInsights
 #include "Insights/InsightsStyle.h"
 #include "Insights/TaskGraphProfiler/TaskGraphProfilerManager.h"
 #include "Insights/TaskGraphProfiler/ViewModels/TaskEntry.h"
-#include "Insights/TaskGraphProfiler/ViewModels/TaskGraphRelation.h"
 #include "Insights/TaskGraphProfiler/ViewModels/TaskNode.h"
 #include "Insights/TaskGraphProfiler/ViewModels/TaskTimingTrack.h"
 #include "Insights/TimingProfiler/TimingProfilerManager.h"
@@ -31,11 +25,7 @@
 #include "Insights/TimingProfiler/Widgets/STimingProfilerWindow.h"
 #include "Insights/Widgets/STimingView.h"
 
-#include <limits>
-
-#define LOCTEXT_NAMESPACE "STaskTableTreeView"
-
-using namespace TraceServices;
+#define LOCTEXT_NAMESPACE "UE::Insights::TaskGraphProfiler::STaskTableTreeView"
 
 namespace UE::Insights::TaskGraphProfiler
 {
@@ -331,7 +321,7 @@ void STaskTableTreeView::ConstructHeaderArea(TSharedRef<SVerticalBox> InWidgetCo
 			SNew(SBox)
 			.MinDesiredWidth(160.0f)
 			[
-				SNew(SComboBox<TSharedPtr<ETaskEnumerationOption>>)
+				SNew(SComboBox<TSharedPtr<TraceServices::ETaskEnumerationOption>>)
 				.OptionsSource(GetAvailableTasksSelectionOptions())
 				.OnSelectionChanged(this, &STaskTableTreeView::TasksSelectionOptions_OnSelectionChanged)
 				.OnGenerateWidget(this, &STaskTableTreeView::TasksSelectionOptions_OnGenerateWidget)
@@ -444,8 +434,10 @@ TSharedRef<SWidget> STaskTableTreeView::TimestampOptions_OnGenerateWidget(TShare
 
 ////////////////////////////////////////////////////////////////////////////////////////////////////
 
-TSharedRef<SWidget> STaskTableTreeView::TasksSelectionOptions_OnGenerateWidget(TSharedPtr<ETaskEnumerationOption> InOption)
+TSharedRef<SWidget> STaskTableTreeView::TasksSelectionOptions_OnGenerateWidget(TSharedPtr<TraceServices::ETaskEnumerationOption> InOption)
 {
+	using namespace TraceServices;
+
 	auto GetTooltipText = [](ETaskEnumerationOption InOption)
 	{
 		switch (InOption)
@@ -524,10 +516,11 @@ const TArray<TSharedPtr<STaskTableTreeView::ETimestampOptions>>* STaskTableTreeV
 
 ////////////////////////////////////////////////////////////////////////////////////////////////////
 
-const TArray<TSharedPtr<ETaskEnumerationOption>>* STaskTableTreeView::GetAvailableTasksSelectionOptions()
+const TArray<TSharedPtr<TraceServices::ETaskEnumerationOption>>* STaskTableTreeView::GetAvailableTasksSelectionOptions()
 {
 	if (AvailableTasksSelectionOptions.Num() == 0)
 	{
+		using namespace TraceServices;
 		AvailableTasksSelectionOptions.Add(MakeShared<ETaskEnumerationOption>(ETaskEnumerationOption::Alive));
 		AvailableTasksSelectionOptions.Add(MakeShared<ETaskEnumerationOption>(ETaskEnumerationOption::Launched));
 		AvailableTasksSelectionOptions.Add(MakeShared<ETaskEnumerationOption>(ETaskEnumerationOption::Active));
@@ -578,14 +571,14 @@ void STaskTableTreeView::TimestampOptions_OnSelectionChanged(ETimestampOptions I
 
 ////////////////////////////////////////////////////////////////////////////////////////////////////
 
-void STaskTableTreeView::TasksSelectionOptions_OnSelectionChanged(TSharedPtr<ETaskEnumerationOption> InOption, ESelectInfo::Type SelectInfo)
+void STaskTableTreeView::TasksSelectionOptions_OnSelectionChanged(TSharedPtr<TraceServices::ETaskEnumerationOption> InOption, ESelectInfo::Type SelectInfo)
 {
 	TasksSelectionOptions_OnSelectionChanged(*InOption);
 }
 
 ////////////////////////////////////////////////////////////////////////////////////////////////////
 
-void STaskTableTreeView::TasksSelectionOptions_OnSelectionChanged(ETaskEnumerationOption InOption)
+void STaskTableTreeView::TasksSelectionOptions_OnSelectionChanged(TraceServices::ETaskEnumerationOption InOption)
 {
 	SelectedTasksSelectionOption = InOption;
 	RebuildTree(true);
@@ -630,8 +623,10 @@ FText STaskTableTreeView::TasksSelectionOptions_GetSelectionText() const
 
 ////////////////////////////////////////////////////////////////////////////////////////////////////
 
-FText STaskTableTreeView::TasksSelectionOptions_GetText(ETaskEnumerationOption InOption) const
+FText STaskTableTreeView::TasksSelectionOptions_GetText(TraceServices::ETaskEnumerationOption InOption) const
 {
+	using namespace TraceServices;
+
 	switch (InOption)
 	{
 	case ETaskEnumerationOption::Alive:
