[Insights]
* Timing Insights / Callers & Callees: Added status bar warning mesage regarding the needed selection (time range + timer).
* Timing Insights: Updated the selection warning messages for Timers and Counters panels.
* Timing Insights / Timers: Added and updated the tooltips re Aggregation Mode combo box.

diff --git a/Engine/Source/Developer/TraceInsights/Private/Insights/TimingProfiler/TimingProfilerManager.h b/Engine/Source/Developer/TraceInsights/Private/Insights/TimingProfiler/TimingProfilerManager.h
--- a/Engine/Source/Developer/TraceInsights/Private/Insights/TimingProfiler/TimingProfilerManager.h
+++ b/Engine/Source/Developer/TraceInsights/Private/Insights/TimingProfiler/TimingProfilerManager.h
@@ -134,11 +134,14 @@ public:
 
 	void OnSessionChanged();
 
+	bool IsValidTimeSelection() const { return SelectionStartTime < SelectionEndTime; }
 	double GetSelectionStartTime() const { return SelectionStartTime; }
 	double GetSelectionEndTime() const { return SelectionEndTime; }
 	void SetSelectedTimeRange(double StartTime, double EndTime);
 
 	FTimerNodePtr GetTimerNode(uint32 TimerId) const;
+	static constexpr uint32 GetInvalidTimerId() { return InvalidTimerId; }
+	bool IsValidSelectedTimer() const { return SelectedTimerId != InvalidTimerId; }
 	uint32 GetSelectedTimer() const { return SelectedTimerId; }
 	void SetSelectedTimer(uint32 TimerId);
 	void ToggleTimingViewMainGraphEventSeries(uint32 InTimerId);
diff --git a/Engine/Source/Developer/TraceInsights/Private/Insights/TimingProfiler/ViewModels/TimerButterflyAggregator.cpp b/Engine/Source/Developer/TraceInsights/Private/Insights/TimingProfiler/ViewModels/TimerButterflyAggregator.cpp
--- a/Engine/Source/Developer/TraceInsights/Private/Insights/TimingProfiler/ViewModels/TimerButterflyAggregator.cpp
+++ b/Engine/Source/Developer/TraceInsights/Private/Insights/TimingProfiler/ViewModels/TimerButterflyAggregator.cpp
@@ -68,11 +68,12 @@ void FTimerButterflyAggregationWorker::DoWork(TSharedPtr<TraceServices::FCancell
 		TraceServices::FAnalysisSessionReadScope SessionReadScope(*Session.Get());
 		const TraceServices::ITimingProfilerProvider& TimingProfilerProvider = *TraceServices::ReadTimingProfilerProvider(*Session.Get());
 
-		if (StartTime <= Session->GetDurationSeconds())
+		const double AdjustedEndTime = FMath::Min(Session->GetDurationSeconds(), EndTime);
+		if (StartTime <= AdjustedEndTime)
 		{
 			TraceServices::FCreateButterflyParams Params;
 			Params.IntervalStart = StartTime;
-			Params.IntervalEnd = FMath::Min(Session->GetDurationSeconds(), EndTime);
+			Params.IntervalEnd = AdjustedEndTime;
 			Params.bIncludeOldGpu1 = bIncludeOldGpu1;
 			Params.bIncludeOldGpu2 = bIncludeOldGpu2;
 			Params.GpuQueueFilter = [this](uint32 QueueId) { return GpuQueues.Contains(QueueId); };
@@ -82,10 +83,10 @@ void FTimerButterflyAggregationWorker::DoWork(TSharedPtr<TraceServices::FCancell
 		}
 		else
 		{
-			// Exclude all GPU and CPU timelines, to receive empty results (in case infinite events are in the session).
+			// Generate empty results (invalid time range selection).
 			TraceServices::FCreateButterflyParams Params;
-			Params.IntervalStart = StartTime;
-			Params.IntervalEnd = EndTime;
+			Params.IntervalStart = 0.0;
+			Params.IntervalEnd = -1.0;
 
 			ResultButterfly.Reset(TimingProfilerProvider.CreateButterfly(Params));
 		}
diff --git a/Engine/Source/Developer/TraceInsights/Private/Insights/TimingProfiler/Widgets/SStatsView.cpp b/Engine/Source/Developer/TraceInsights/Private/Insights/TimingProfiler/Widgets/SStatsView.cpp
--- a/Engine/Source/Developer/TraceInsights/Private/Insights/TimingProfiler/Widgets/SStatsView.cpp
+++ b/Engine/Source/Developer/TraceInsights/Private/Insights/TimingProfiler/Widgets/SStatsView.cpp
@@ -355,9 +355,13 @@ void SStatsView::Construct(const FArguments& InArgs)
 			[
 				SNew(STextBlock)
 				.Margin(FMargin(4.0f, 1.0f, 4.0f, 1.0f))
-				.Text(LOCTEXT("EmptyAggregationNote", "-- Select a time region to update the aggregated statistics! --"))
+				.Text(LOCTEXT("EmptyAggregationNote", "Please select a time range to update the aggregated statistics!"))
 				.ColorAndOpacity(FLinearColor(1.0f, 0.75f, 0.5f, 1.0f))
-				.Visibility_Lambda([this]() { return Aggregator->IsEmptyTimeInterval() && !Aggregator->IsRunning() ? EVisibility::Visible : EVisibility::Collapsed; })
+				.Visibility_Lambda([this]()
+				{
+					return Aggregator->IsEmptyTimeInterval() && !Aggregator->IsRunning() ?
+						EVisibility::Visible : EVisibility::Collapsed;
+				})
 			]
 		]
 	];
diff --git a/Engine/Source/Developer/TraceInsights/Private/Insights/TimingProfiler/Widgets/STimerTreeView.cpp b/Engine/Source/Developer/TraceInsights/Private/Insights/TimingProfiler/Widgets/STimerTreeView.cpp
--- a/Engine/Source/Developer/TraceInsights/Private/Insights/TimingProfiler/Widgets/STimerTreeView.cpp
+++ b/Engine/Source/Developer/TraceInsights/Private/Insights/TimingProfiler/Widgets/STimerTreeView.cpp
@@ -125,51 +125,91 @@ void STimerTreeView::Construct(const FArguments& InArgs, const FText& InViewName
 
 	ChildSlot
 	[
-		SNew(SHorizontalBox)
+		SNew(SVerticalBox)
 
-		+ SHorizontalBox::Slot()
-		.FillWidth(1.0f)
+		// Tree view
+		+ SVerticalBox::Slot()
+		.FillHeight(1.0f)
 		.Padding(0.0f)
 		[
-			SNew(SOverlay)
+			SNew(SHorizontalBox)
 
-			+ SOverlay::Slot()
-			.HAlign(HAlign_Fill)
-			.VAlign(VAlign_Fill)
+			+ SHorizontalBox::Slot()
+			.FillWidth(1.0f)
+			.Padding(0.0f)
 			[
-				SAssignNew(TreeView, STreeView<FTimerNodePtr>)
-				.ExternalScrollbar(ExternalScrollbar)
-				.SelectionMode(ESelectionMode::Multi)
-				.TreeItemsSource(&TreeNodes)
-				.OnGetChildren(this, &STimerTreeView::TreeView_OnGetChildren)
-				.OnGenerateRow(this, &STimerTreeView::TreeView_OnGenerateRow)
-				//.OnSelectionChanged(this, &STimerTreeView::TreeView_OnSelectionChanged)
-				//.OnMouseButtonDoubleClick(this, &STimerTreeView::TreeView_OnMouseButtonDoubleClick)
-				.OnContextMenuOpening(FOnContextMenuOpening::CreateSP(this, &STimerTreeView::TreeView_GetMenuContent))
-				.HeaderRow
-				(
-					SAssignNew(TreeViewHeaderRow, SHeaderRow)
-					.Visibility(EVisibility::Visible)
-				)
+				SNew(SOverlay)
+
+				+ SOverlay::Slot()
+				.HAlign(HAlign_Fill)
+				.VAlign(VAlign_Fill)
+				[
+					SAssignNew(TreeView, STreeView<FTimerNodePtr>)
+					.ExternalScrollbar(ExternalScrollbar)
+					.SelectionMode(ESelectionMode::Multi)
+					.TreeItemsSource(&TreeNodes)
+					.OnGetChildren(this, &STimerTreeView::TreeView_OnGetChildren)
+					.OnGenerateRow(this, &STimerTreeView::TreeView_OnGenerateRow)
+					//.OnSelectionChanged(this, &STimerTreeView::TreeView_OnSelectionChanged)
+					//.OnMouseButtonDoubleClick(this, &STimerTreeView::TreeView_OnMouseButtonDoubleClick)
+					.OnContextMenuOpening(FOnContextMenuOpening::CreateSP(this, &STimerTreeView::TreeView_GetMenuContent))
+					.HeaderRow
+					(
+						SAssignNew(TreeViewHeaderRow, SHeaderRow)
+						.Visibility(EVisibility::Visible)
+					)
+				]
+
+				+ SOverlay::Slot()
+				.HAlign(HAlign_Right)
+				.VAlign(VAlign_Bottom)
+				.Padding(16.0f)
+				[
+					SAssignNew(AsyncOperationStatus, SAsyncOperationStatus, TimerButterflyAggregator)
+				]
 			]
 
-			+ SOverlay::Slot()
-			.HAlign(HAlign_Right)
-			.VAlign(VAlign_Bottom)
-			.Padding(16.0f)
+			+ SHorizontalBox::Slot()
+			.AutoWidth()
+			.Padding(0.0f)
 			[
-				SAssignNew(AsyncOperationStatus, SAsyncOperationStatus, TimerButterflyAggregator)
+				SNew(SBox)
+				.WidthOverride(FOptionalSize(13.0f))
+				[
+					ExternalScrollbar.ToSharedRef()
+				]
 			]
 		]
 
-		+ SHorizontalBox::Slot()
-		.AutoWidth()
+		// Status bar
+		+ SVerticalBox::Slot()
+		.AutoHeight()
 		.Padding(0.0f)
 		[
-			SNew(SBox)
-			.WidthOverride(FOptionalSize(13.0f))
+			SNew(SBorder)
+			.BorderImage(FAppStyle::Get().GetBrush("WhiteBrush"))
+			.BorderBackgroundColor(FLinearColor(0.05f, 0.1f, 0.2f, 1.0f))
+			.HAlign(HAlign_Center)
 			[
-				ExternalScrollbar.ToSharedRef()
+				SNew(STextBlock)
+				.Margin(FMargin(4.0f, 1.0f, 4.0f, 1.0f))
+				.Text(LOCTEXT("SelectionWarning", "Please select a time range and a timer!"))
+				.ColorAndOpacity(FLinearColor(1.0f, 0.75f, 0.5f, 1.0f))
+				.Visibility_Lambda([]()
+				{
+					TSharedPtr<FTimingProfilerManager> TimingProfilerManager = FTimingProfilerManager::Get();
+					if (!TimingProfilerManager.IsValid())
+					{
+						return EVisibility::Collapsed;
+					}
+					TSharedRef<FTimerButterflyAggregator> Aggregator = TimingProfilerManager->GetTimerButterflyAggregator();
+					if (Aggregator->IsRunning())
+					{
+						return EVisibility::Collapsed;
+					}
+					return !TimingProfilerManager->IsValidTimeSelection() || !TimingProfilerManager->IsValidSelectedTimer() ?
+						EVisibility::Visible : EVisibility::Collapsed;
+				})
 			]
 		]
 	];
diff --git a/Engine/Source/Developer/TraceInsights/Private/Insights/TimingProfiler/Widgets/STimersView.cpp b/Engine/Source/Developer/TraceInsights/Private/Insights/TimingProfiler/Widgets/STimersView.cpp
--- a/Engine/Source/Developer/TraceInsights/Private/Insights/TimingProfiler/Widgets/STimersView.cpp
+++ b/Engine/Source/Developer/TraceInsights/Private/Insights/TimingProfiler/Widgets/STimersView.cpp
@@ -358,6 +358,7 @@ void STimersView::Construct(const FArguments& InArgs)
 					[
 						SNew(STextBlock)
 						.Text(LOCTEXT("Mode", "Mode"))
+						.ToolTipText(LOCTEXT("ModeToolTip", "Aggregation Mode"))
 					]
 
 					+ SHorizontalBox::Slot()
@@ -449,9 +450,13 @@ void STimersView::Construct(const FArguments& InArgs)
 			[
 				SNew(STextBlock)
 				.Margin(FMargin(4.0f, 1.0f, 4.0f, 1.0f))
-				.Text(LOCTEXT("EmptyAggregationNote", "-- Select a time region to update the aggregated statistics! --"))
+				.Text(LOCTEXT("EmptyAggregationNote", "Please select a time range to update the aggregated statistics!"))
 				.ColorAndOpacity(FLinearColor(1.0f, 0.75f, 0.5f, 1.0f))
-				.Visibility_Lambda([this]() { return Aggregator->IsEmptyTimeInterval() && !Aggregator->IsRunning() ? EVisibility::Visible : EVisibility::Collapsed; })
+				.Visibility_Lambda([this]()
+				{
+					return Aggregator->IsEmptyTimeInterval() && !Aggregator->IsRunning() ?
+						EVisibility::Visible : EVisibility::Collapsed;
+				})
 			]
 		]
 	];
@@ -1959,11 +1964,11 @@ FText STimersView::Mode_GetTooltipText(ETraceFrameType InFrameType) const
 	switch (InFrameType)
 	{
 	case ETraceFrameType::TraceFrameType_Count:
-		return LOCTEXT("ModeInstanceTooltip", "Calculate the stats per event instance.");
+		return LOCTEXT("ModeInstanceTooltip", "'Instance' Aggregation Mode\nCalculates the stats per event instance.");
 	case ETraceFrameType::TraceFrameType_Game:
-		return LOCTEXT("ModeGameFrameTooltip", "Calculate the timer stats per game frame.");
+		return LOCTEXT("ModeGameFrameTooltip", "'Game Frame' Aggregation Mode\nCalculates the timer stats per game frame.");
 	case ETraceFrameType::TraceFrameType_Rendering:
-		return LOCTEXT("ModeRenderingFrameTooltip", "Calculate the timer stats per rendering frame.");
+		return LOCTEXT("ModeRenderingFrameTooltip", "'Rendering Frame' Aggregation Mode\nCalculates the timer stats per rendering frame.");
 	}
 
 	return FText();
