[Insights] Timing View: Updated colors for the track headers to use "blue" for selected tracks.

diff --git a/Engine/Source/Developer/TraceInsights/Private/Insights/ViewModels/TimingViewDrawHelper.cpp b/Engine/Source/Developer/TraceInsights/Private/Insights/ViewModels/TimingViewDrawHelper.cpp
--- a/Engine/Source/Developer/TraceInsights/Private/Insights/ViewModels/TimingViewDrawHelper.cpp
+++ b/Engine/Source/Developer/TraceInsights/Private/Insights/ViewModels/TimingViewDrawHelper.cpp
@@ -701,12 +701,6 @@ void FTimingViewDrawHelper::DrawTrackHeader(const FBaseTimingTrack& Track, const
 			const float FontScale = DrawContext.Geometry.Scale;
 			float TextWidth = static_cast<float>(FontMeasureService->Measure(Track.GetName(), EventFont, FontScale).X / FontScale);
 
-			constexpr float PinWidth = 8.0f;
-			if (Track.IsSelected())
-			{
-				TextWidth += PinWidth;
-			}
-
 			const float HeaderX = 0.0f;
 			const float HeaderW = TextWidth + 4.0f;
 
@@ -715,21 +709,13 @@ void FTimingViewDrawHelper::DrawTrackHeader(const FBaseTimingTrack& Track, const
 
 			if (HeaderH > 0)
 			{
-				DrawContext.DrawBox(HeaderLayerId, HeaderX, HeaderY, HeaderW, HeaderH, WhiteBrush, EdgeColor);
+				const FLinearColor HeaderBackgroundColor = GetTrackNameBackgroundColor(Track);
+				DrawContext.DrawBox(HeaderLayerId, HeaderX, HeaderY, HeaderW, HeaderH, WhiteBrush, HeaderBackgroundColor);
 
-				const FLinearColor TextColor = GetTrackNameTextColor(Track);
-
-				float TextX = HeaderX + 2.0f;
+				const float TextX = HeaderX + 2.0f;
 				const float TextY = HeaderY + HeaderH / 2.0f - 7.0f;
-
-				if (Track.IsSelected())
-				{
-					// TODO: Use a "pin" image brush instead.
-					DrawContext.DrawText(HeaderTextLayerId, TextX, TextY, TEXT("\u25CA"), EventFont, TextColor); // lozenge shape
-					TextX += PinWidth;
-				}
-
-				DrawContext.DrawText(HeaderTextLayerId, TextX, TextY, Track.GetName(), EventFont, TextColor);
+				const FLinearColor HeaderTextColor = GetTrackNameTextColor(Track);
+				DrawContext.DrawText(HeaderTextLayerId, TextX, TextY, Track.GetName(), EventFont, HeaderTextColor);
 			}
 		}
 	}
@@ -871,9 +857,27 @@ void FTimingViewDrawHelper::DrawTimingEventHighlight(double StartTime, double En
 
 FLinearColor FTimingViewDrawHelper::GetTrackNameTextColor(const FBaseTimingTrack& Track) const
 {
-	return  Track.IsHovered() ?  FLinearColor(1.0f, 1.0f, 0.0f, 1.0f) :
-			Track.IsSelected() ? FLinearColor(1.0f, 1.0f, 0.5f, 1.0f) :
-			                     FLinearColor(1.0f, 1.0f, 1.0f, 1.0f);
+	return Track.IsHovered() ?
+		FLinearColor(1.0f, 1.0f, 0.5f, 1.0f) :
+		FLinearColor(1.0f, 1.0f, 1.0f, 1.0f);
+}
+
+////////////////////////////////////////////////////////////////////////////////////////////////////
+
+FLinearColor FTimingViewDrawHelper::GetTrackNameBackgroundColor(const FBaseTimingTrack& Track) const
+{
+	if (Track.IsSelected())
+	{
+		return Track.IsHovered() ?
+			FLinearColor(0.00f, 0.1f, 0.4f, 1.0f) :
+			FLinearColor(0.00f, 0.2f, 0.5f, 1.0f);
+	}
+	else
+	{
+		return Track.IsHovered() ?
+			FLinearColor(0.04f, 0.04f, 0.04f, 1.0f) :
+			FLinearColor(0.05f, 0.05f, 0.05f, 1.0f);
+	}
 }
 
 ////////////////////////////////////////////////////////////////////////////////////////////////////
diff --git a/Engine/Source/Developer/TraceInsights/Private/Insights/ViewModels/TimingViewDrawHelper.h b/Engine/Source/Developer/TraceInsights/Private/Insights/ViewModels/TimingViewDrawHelper.h
--- a/Engine/Source/Developer/TraceInsights/Private/Insights/ViewModels/TimingViewDrawHelper.h
+++ b/Engine/Source/Developer/TraceInsights/Private/Insights/ViewModels/TimingViewDrawHelper.h
@@ -184,6 +184,7 @@ public:
 	virtual FLinearColor GetValidAreaColor() const override { return ValidAreaColor; }
 	virtual FLinearColor GetInvalidAreaColor() const override { return InvalidAreaColor; }
 	virtual FLinearColor GetTrackNameTextColor(const FBaseTimingTrack& Track) const override;
+	virtual FLinearColor GetTrackNameBackgroundColor(const FBaseTimingTrack& Track) const override;
 	virtual int32 GetHeaderBackgroundLayerId() const override { return ReservedLayerId + ToInt32(EDrawLayer::HeaderBackground); }
 	virtual int32 GetHeaderTextLayerId() const override { return ReservedLayerId + ToInt32(EDrawLayer::HeaderText); }
 	virtual int32 GetRelationLayerId() const override {return ReservedLayerId + ToInt32(EDrawLayer::Relation); };
diff --git a/Engine/Source/Developer/TraceInsights/Public/Insights/ViewModels/ITimingViewDrawHelper.h b/Engine/Source/Developer/TraceInsights/Public/Insights/ViewModels/ITimingViewDrawHelper.h
--- a/Engine/Source/Developer/TraceInsights/Public/Insights/ViewModels/ITimingViewDrawHelper.h
+++ b/Engine/Source/Developer/TraceInsights/Public/Insights/ViewModels/ITimingViewDrawHelper.h
@@ -21,6 +21,7 @@ public:
 	virtual FLinearColor GetValidAreaColor() const = 0;
 	virtual FLinearColor GetInvalidAreaColor() const = 0;
 	virtual FLinearColor GetTrackNameTextColor(const FBaseTimingTrack& Track) const = 0;
+	virtual FLinearColor GetTrackNameBackgroundColor(const FBaseTimingTrack& Track) const { return GetEdgeColor(); }
 	virtual int32 GetHeaderBackgroundLayerId() const = 0;
 	virtual int32 GetHeaderTextLayerId() const = 0;
 	virtual int32 GetRelationLayerId() const = 0;
