[Insights] Improved UI for LLM Tags panel (tags sets as segmented control with initial placeholder icons; combobox style for the trackers filter; fixed padding for all toolbar elements).

#rnx

diff --git a/Engine/Source/Developer/TraceInsights/Private/Insights/InsightsStyle.cpp b/Engine/Source/Developer/TraceInsights/Private/Insights/InsightsStyle.cpp
--- a/Engine/Source/Developer/TraceInsights/Private/Insights/InsightsStyle.cpp
+++ b/Engine/Source/Developer/TraceInsights/Private/Insights/InsightsStyle.cpp
@@ -220,6 +220,10 @@ void FInsightsStyle::FStyle::Initialize()
 	Set("Icons.AddAllMemTagGraphs", new CORE_IMAGE_BRUSH_SVG("Starship/Common/plus-circle", Icon16x16));
 	Set("Icons.RemoveAllMemTagGraphs", new CORE_IMAGE_BRUSH_SVG("Starship/Common/Delete", Icon16x16));
 
+	Set("Icons.TagSet.Systems", new CORE_IMAGE_BRUSH_SVG("Starship/Common/File", Icon16x16)); 
+	Set("Icons.TagSet.Assets", new EDITOR_IMAGE_BRUSH_SVG("Starship/Common/FileLined", Icon16x16));
+	Set("Icons.TagSet.AssetClasses", new EDITOR_IMAGE_BRUSH_SVG("Starship/Common/Class", Icon16x16));
+
 	//////////////////////////////////////////////////
 	// Tasks
 
diff --git a/Engine/Source/Developer/TraceInsights/Private/Insights/MemoryProfiler/Widgets/SMemTagTreeView.cpp b/Engine/Source/Developer/TraceInsights/Private/Insights/MemoryProfiler/Widgets/SMemTagTreeView.cpp
--- a/Engine/Source/Developer/TraceInsights/Private/Insights/MemoryProfiler/Widgets/SMemTagTreeView.cpp
+++ b/Engine/Source/Developer/TraceInsights/Private/Insights/MemoryProfiler/Widgets/SMemTagTreeView.cpp
@@ -76,8 +76,107 @@ void SMemTagTreeView::Construct(const FArguments& InArgs, TSharedPtr<SMemoryProf
 
 ////////////////////////////////////////////////////////////////////////////////////////////////////
 
+void SMemTagTreeView::ConstructHeaderArea(TSharedRef<SVerticalBox> InHostBox)
+{
+	InHostBox->AddSlot()
+		.AutoHeight()
+		.Padding(2.0f, 2.0f, 2.0f, 2.0f)
+		[
+			SNew(SBox)
+			[
+				SAssignNew(TagSetsSegmentedControl, SSegmentedControl<FMemoryTagSetId>)
+				.SupportsMultiSelection(false)
+				.OnValueChanged_Lambda([this](FMemoryTagSetId InValue)
+				{
+					TagSetFilter = InValue;
+					OnNodeFilteringChanged();
+					bShouldUpdateStats = true;
+					bShouldUpdateBudgets = true;
+				})
+				.Value_Lambda([this]
+				{
+					return TagSetFilter;
+				})
+			]
+		];
+	UpdateAvailableTagSets();
+
+	TSharedPtr<SWidget> TopSettingsWidget = ConstructTopSettings();
+	if (TopSettingsWidget.IsValid())
+	{
+		InHostBox->AddSlot()
+			.AutoHeight()
+			.Padding(2.0f, 2.0f, 2.0f, 2.0f)
+			[
+				TopSettingsWidget.ToSharedRef()
+			];
+	}
+
+	TSharedPtr<SWidget> BudgetSettingsWidget = ConstructBudgetSettings();
+	if (BudgetSettingsWidget.IsValid())
+	{
+		InHostBox->AddSlot()
+			.AutoHeight()
+			.Padding(2.0f, 2.0f, 2.0f, 2.0f)
+			[
+				BudgetSettingsWidget.ToSharedRef()
+			];
+	}
+
+	TSharedPtr<SWidget> FilterToolbarWidget = ConstructFilterToolbar();
+	if (FilterToolbarWidget.IsValid())
+	{
+		InHostBox->AddSlot()
+			.AutoHeight()
+			.Padding(2.0f, 2.0f, 2.0f, 2.0f)
+			[
+				FilterToolbarWidget.ToSharedRef()
+			];
+	}
+
+	TSharedPtr<SWidget> HierarchyBreadcrumbTrailWidget = ConstructHierarchyBreadcrumbTrail();
+	if (HierarchyBreadcrumbTrailWidget.IsValid())
+	{
+		InHostBox->AddSlot()
+			.AutoHeight()
+			.Padding(2.0f, 2.0f, 2.0f, 2.0f)
+			[
+				HierarchyBreadcrumbTrailWidget.ToSharedRef()
+			];
+	}
+
+	TSharedPtr<SWidget> TimeMarkersWidget = ConstructTimeMarkers();
+	if (TimeMarkersWidget.IsValid())
+	{
+		InHostBox->AddSlot()
+			.AutoHeight()
+			.Padding(20.0f, 2.0f, 2.0f, -2.0f)
+			[
+				TimeMarkersWidget.ToSharedRef()
+			];
+	}
+}
+
+////////////////////////////////////////////////////////////////////////////////////////////////////
+
+void SMemTagTreeView::ConstructFooterArea(TSharedRef<SVerticalBox> InHostBox)
+{
+	TSharedPtr<SWidget> FooterWidget = ConstructFooter();
+	if (FooterWidget.IsValid())
+	{
+		InHostBox->AddSlot()
+			.AutoHeight()
+			.Padding(0.0f, 6.0f, 0.0f, 0.0f)
+			[
+				FooterWidget.ToSharedRef()
+			];
+	}
+}
+
+////////////////////////////////////////////////////////////////////////////////////////////////////
+
 BEGIN_SLATE_FUNCTION_BUILD_OPTIMIZATION
-TSharedPtr<SWidget> SMemTagTreeView::ConstructFilterToolbar()
+TSharedPtr<SWidget> SMemTagTreeView::ConstructTopSettings()
 {
 	TSharedRef<SHorizontalBox> TopLine = SNew(SHorizontalBox)
 
@@ -105,6 +204,7 @@ TSharedPtr<SWidget> SMemTagTreeView::ConstructFilterToolbar()
 		ConstructTrackHeightControls()
 	]
 
+#if 0
 	+ SHorizontalBox::Slot()
 	.FillWidth(1.0f)
 	.Padding(0.0f, 0.0f, 0.0f, 0.0f)
@@ -129,27 +229,75 @@ TSharedPtr<SWidget> SMemTagTreeView::ConstructFilterToolbar()
 			SNew(STextBlock)
 			.Text(this, &SMemTagTreeView::TagSet_GetSelectedText)
 		]
-	];
+	]
+#else
+	+ SHorizontalBox::Slot()
+	.FillWidth(1.0f)
+	.Padding(0.0f, 0.0f, 0.0f, 0.0f)
+	.HAlign(HAlign_Right)
+	.VAlign(VAlign_Center)
+	[
+		SNew(SBox)
+	]
+#endif
+	;
 
 	ConstructViewPreset(TopLine);
 
-	//////////////////////////////////////////////////
+	return TopLine;
+}
+END_SLATE_FUNCTION_BUILD_OPTIMIZATION
+
+////////////////////////////////////////////////////////////////////////////////////////////////////
+
+void SMemTagTreeView::UpdateAvailableTagSets()
+{
+	if (TagSetsSegmentedControl->NumSlots() != 0)
+	{
+		return;
+	}
 
-	FSlimHorizontalToolBarBuilder ToolbarBuilder(TSharedPtr<const FUICommandList>(), FMultiBoxCustomization::None);
-	ToolbarBuilder.SetStyle(&FInsightsStyle::Get(), "SecondaryToolbar");
-	ToolbarBuilder.BeginSection("Filters");
+	for (TSharedPtr<FMemoryTagSetId> TagSetPtr : AvailableTagSets)
 	{
-		// Filter Trackers
-		ToolbarBuilder.AddComboButton(
-			FUIAction(),
-			FOnGetContent::CreateSP(this, &SMemTagTreeView::MakeTrackersMenu),
-			LOCTEXT("TrackersMenuText", "Trackers"),
-			LOCTEXT("TrackersMenuToolTip", "Filter list of LLM tags by tracker."),
-			FSlateIcon(FInsightsStyle::GetStyleSetName(), "Icons.Filter.ToolBar"),
-			false);
+		FMemoryTagSetId TagSet = *TagSetPtr;
+
+		const FSlateBrush* TagSetIcon = nullptr;
+		FText TagSetName;
+		switch (TagSet)
+		{
+		case 0:
+			TagSetIcon = FInsightsStyle::GetBrush("Icons.TagSet.Systems");
+			TagSetName = FText::FromString("Systems");
+			break;
+		case 1:
+			TagSetIcon = FInsightsStyle::GetBrush("Icons.TagSet.Assets");
+			TagSetName = FText::FromString("Assets");
+			break;
+		case 2:
+			TagSetIcon = FInsightsStyle::GetBrush("Icons.TagSet.AssetClasses");
+			TagSetName = FText::FromString("AssetClasses");
+			break;
+		default:
+			TagSetName = FText::FromString("Unknown");
+			TagSetIcon = FInsightsStyle::GetBrush("Icons.TagSet.Systems");
+		}
+
+		FText TagSetToolTip;
+		const FText TagSetToolTipFmt = LOCTEXT("TagSetToolTipFmt", "'{0}' Tag Set\n\nThe tree view shows the memory tags from the selected tag set.");
+		TagSetToolTip = FText::Format(TagSetToolTipFmt, TagSetName);
+
+		TagSetsSegmentedControl->AddSlot(TagSet)
+			.Icon(TagSetIcon)
+			.Text(TagSetName)
+			.ToolTip(TagSetToolTip);
 	}
-	ToolbarBuilder.EndSection();
+}
+
+////////////////////////////////////////////////////////////////////////////////////////////////////
 
+BEGIN_SLATE_FUNCTION_BUILD_OPTIMIZATION
+TSharedPtr<SWidget> SMemTagTreeView::ConstructFilterToolbar()
+{
 	TSharedRef<SHorizontalBox> FilterLine = SNew(SHorizontalBox)
 
 	+ SHorizontalBox::Slot()
@@ -161,41 +309,46 @@ TSharedPtr<SWidget> SMemTagTreeView::ConstructFilterToolbar()
 
 	+ SHorizontalBox::Slot()
 	.AutoWidth()
+	.MinSize(10.0f)
 	.Padding(4.0f, 0.0f, 0.0f, 0.0f)
 	.VAlign(VAlign_Center)
 	[
-		ConstructFilterConfiguratorButton()
+		SNew(SComboButton)
+		.OnGetMenuContent(this, &SMemTagTreeView::MakeTrackersMenu)
+		.ToolTipText(LOCTEXT("TrackersMenuToolTip", "Filter list of LLM tags by tracker."))
+		.ContentPadding(FMargin(0.0f, 1.0f))
+		.ButtonContent()
+		[
+			SNew(SHorizontalBox)
+			+ SHorizontalBox::Slot()
+			.VAlign(VAlign_Center)
+			.AutoWidth()
+			.Padding(0.0f, 0.0f, 4.0f, 0.0f)
+			[
+				SNew(SImage)
+				.Image(FInsightsStyle::GetBrush("Icons.Filter.ToolBar"))
+				.ColorAndOpacity(FSlateColor(EStyleColor::Foreground))
+			]
+			+ SHorizontalBox::Slot()
+			.VAlign(VAlign_Center)
+			.AutoWidth()
+			.Padding(0.0f, 0.0f, 0.0f, 0.0f)
+			[
+				SNew(STextBlock)
+				.Text(LOCTEXT("TrackersMenuText", "Trackers"))
+			]
+		]
 	]
 
 	+ SHorizontalBox::Slot()
 	.AutoWidth()
-	.MinSize(10.0f)
-	.Padding(0.0f, 0.0f, -6.0f, 0.0f)
+	.Padding(2.0f, 0.0f, 0.0f, 0.0f)
 	.VAlign(VAlign_Center)
 	[
-		ToolbarBuilder.MakeWidget()
+		ConstructFilterConfiguratorButton()
 	];
 
-	//////////////////////////////////////////////////
-
-	return SNew(SVerticalBox)
-
-	+ SVerticalBox::Slot()
-	[
-		TopLine
-	]
-
-	+ SVerticalBox::Slot()
-	.AutoHeight()
-	.Padding(FMargin(0.0f, 2.0f, 0.0f, 2.0f))
-	[
-		ConstructBudgetSettingsToolbar()
-	]
-
-	+ SVerticalBox::Slot()
-	[
-		FilterLine
-	];
+	return FilterLine;
 }
 END_SLATE_FUNCTION_BUILD_OPTIMIZATION
 
@@ -211,40 +364,56 @@ END_SLATE_FUNCTION_BUILD_OPTIMIZATION
 ////////////////////////////////////////////////////////////////////////////////////////////////////
 
 BEGIN_SLATE_FUNCTION_BUILD_OPTIMIZATION
-TSharedPtr<SWidget> SMemTagTreeView::ConstructExtraToolbar()
+TSharedPtr<SWidget> SMemTagTreeView::ConstructTimeMarkers()
 {
-	return SNew(SVerticalBox)
+	TSharedPtr<SWidget> WidgetA = ConstructTimeMarkerA();
+	TSharedPtr<SWidget> WidgetB = ConstructTimeMarkerB();
 
-	+ SVerticalBox::Slot()
-	.AutoHeight()
-	[
-		ConstructTimeMarkerA()
-	]
+	if (!WidgetA.IsValid() && !WidgetB.IsValid())
+	{
+		return nullptr;
+	}
 
-	+ SVerticalBox::Slot()
-	.AutoHeight()
-	[
-		ConstructTimeMarkerB()
-	];
+	TSharedRef<SVerticalBox> VerticalBox = SNew(SVerticalBox);
+
+	if (WidgetA.IsValid())
+	{
+		VerticalBox->AddSlot()
+			.AutoHeight()
+			[
+				WidgetA.ToSharedRef()
+			];
+	}
+
+	if (WidgetB.IsValid())
+	{
+		VerticalBox->AddSlot()
+			.AutoHeight()
+			[
+				WidgetB.ToSharedRef()
+			];
+	}
+
+	return VerticalBox;
 }
 END_SLATE_FUNCTION_BUILD_OPTIMIZATION
 
 ////////////////////////////////////////////////////////////////////////////////////////////////////
 
 BEGIN_SLATE_FUNCTION_BUILD_OPTIMIZATION
-TSharedRef<SWidget> SMemTagTreeView::ConstructTimeMarkerA()
+TSharedPtr<SWidget> SMemTagTreeView::ConstructTimeMarkerA()
 {
 	TSharedPtr<SMemoryProfilerWindow> ProfilerWindow = GetProfilerWindow();
 	if (!ProfilerWindow.IsValid() || ProfilerWindow->GetNumCustomTimeMarkers() < 2)
 	{
-		return SNew(SBox);
+		return nullptr;
 	}
 
 	using FTimeMarker = TimingProfiler::FTimeMarker;
 	TSharedRef<FTimeMarker> TimeMarkerA = ProfilerWindow->GetCustomTimeMarker(0);
 
 	return SNew(TimingProfiler::STimeMarkerEditBlock, TimeMarkerA)
-		.Padding(FMargin(20.0f, 0.0f, 0.0f, 0.0f))
+		.Padding(FMargin(0.0f, 0.0f, 0.0f, 0.0f))
 		.OnGetTimingView_Lambda(
 			[ProfilerWindowWeak = ProfilerWindowWeakPtr](TSharedRef<FTimeMarker> TimeMarker) -> TSharedPtr<TimingProfiler::STimingView>
 			{
@@ -270,12 +439,12 @@ END_SLATE_FUNCTION_BUILD_OPTIMIZATION
 ////////////////////////////////////////////////////////////////////////////////////////////////////
 
 BEGIN_SLATE_FUNCTION_BUILD_OPTIMIZATION
-TSharedRef<SWidget> SMemTagTreeView::ConstructTimeMarkerB()
+TSharedPtr<SWidget> SMemTagTreeView::ConstructTimeMarkerB()
 {
 	TSharedPtr<SMemoryProfilerWindow> ProfilerWindow = GetProfilerWindow();
 	if (!ProfilerWindow.IsValid() || ProfilerWindow->GetNumCustomTimeMarkers() < 2)
 	{
-		return SNew(SBox);
+		return nullptr;
 	}
 
 	using FTimeMarker = TimingProfiler::FTimeMarker;
@@ -283,7 +452,7 @@ TSharedRef<SWidget> SMemTagTreeView::ConstructTimeMarkerB()
 	TSharedRef<FTimeMarker> TimeMarkerB = ProfilerWindow->GetCustomTimeMarker(1);
 
 	return SNew(TimingProfiler::STimeMarkerEditBlock, TimeMarkerB)
-		.Padding(FMargin(20.0f, 0.0f, 0.0f, 0.0f))
+		.Padding(FMargin(0.0f, 0.0f, 0.0f, 0.0f))
 		.PreviousTimeMarker(TimeMarkerA)
 		.OnGetTimingView_Lambda(
 			[ProfilerWindowWeak = ProfilerWindowWeakPtr](TSharedRef<FTimeMarker> TimeMarker) -> TSharedPtr<TimingProfiler::STimingView>
@@ -317,6 +486,8 @@ TSharedRef<SWidget> SMemTagTreeView::MakeTrackersMenu()
 	CreateTrackersMenuSection(MenuBuilder);
 	MenuBuilder.EndSection();
 
+	MenuBuilder.SetSearchable(false);
+
 	return MenuBuilder.MakeWidget();
 }
 
@@ -511,7 +682,7 @@ END_SLATE_FUNCTION_BUILD_OPTIMIZATION
 ////////////////////////////////////////////////////////////////////////////////////////////////////
 
 BEGIN_SLATE_FUNCTION_BUILD_OPTIMIZATION
-TSharedRef<SWidget> SMemTagTreeView::ConstructBudgetSettingsToolbar()
+TSharedRef<SWidget> SMemTagTreeView::ConstructBudgetSettings()
 {
 	float constexpr TopPadding = 0.0f;
 
@@ -1022,6 +1193,7 @@ void SMemTagTreeView::RebuildTree(bool bResync)
 				{
 					AvailableTagSets.Add(MakeShared<FMemoryTagSetId>(InTagSet.GetId()));
 				});
+			UpdateAvailableTagSets();
 		}
 	}
 
diff --git a/Engine/Source/Developer/TraceInsights/Private/Insights/MemoryProfiler/Widgets/SMemTagTreeView.h b/Engine/Source/Developer/TraceInsights/Private/Insights/MemoryProfiler/Widgets/SMemTagTreeView.h
--- a/Engine/Source/Developer/TraceInsights/Private/Insights/MemoryProfiler/Widgets/SMemTagTreeView.h
+++ b/Engine/Source/Developer/TraceInsights/Private/Insights/MemoryProfiler/Widgets/SMemTagTreeView.h
@@ -3,6 +3,7 @@
 #pragma once
 
 #include "CoreTypes.h"
+#include "Widgets/Input/SSegmentedControl.h"
 
 // TraceInsights
 #include "Insights/Table/Widgets/SSessionTableTreeView.h"
@@ -140,9 +141,10 @@ public:
 	 */
 	void Construct(const FArguments& InArgs, TSharedPtr<SMemoryProfilerWindow> InProfilerWindow);
 
+	virtual void ConstructHeaderArea(TSharedRef<SVerticalBox> InHostBox) override;
+	virtual void ConstructFooterArea(TSharedRef<SVerticalBox> InHostBox) override;
 	virtual TSharedPtr<SWidget> ConstructFilterToolbar() override;
 	virtual TSharedPtr<SWidget> ConstructToolbar() override;
-	virtual TSharedPtr<SWidget> ConstructExtraToolbar() override;
 	virtual TSharedPtr<SWidget> ConstructFooter() override;
 
 	TSharedPtr<FMemTagTable> GetMemTagTable()
@@ -204,12 +206,14 @@ protected:
 	virtual void TreeView_OnSelectionChanged(FTableTreeNodePtr SelectedItem, ESelectInfo::Type SelectInfo) override;
 
 private:
+	TSharedPtr<SWidget> ConstructTopSettings();
 	TSharedRef<SButton> ConstructHideAllButton();
 	TSharedRef<SButton> ConstructShowAllButton();
 	TSharedRef<SWidget> ConstructTrackHeightControls();
-	TSharedRef<SWidget> ConstructBudgetSettingsToolbar();
-	TSharedRef<SWidget> ConstructTimeMarkerA();
-	TSharedRef<SWidget> ConstructTimeMarkerB();
+	TSharedRef<SWidget> ConstructBudgetSettings();
+	TSharedPtr<SWidget> ConstructTimeMarkers();
+	TSharedPtr<SWidget> ConstructTimeMarkerA();
+	TSharedPtr<SWidget> ConstructTimeMarkerB();
 
 	void UpdateSelectionStatsText();
 	FText GetNumSelectedTagsText() const { return NumSelectedTagsText; }
@@ -242,6 +246,7 @@ private:
 	////////////////////////////////////////////////////////////////////////////////////////////////////
 	// Tag Set Selection
 
+	void UpdateAvailableTagSets();
 	void TagSet_OnSelectionChanged(TSharedPtr<FMemoryTagSetId> NewTagSet, ESelectInfo::Type SelectInfo);
 	TSharedRef<SWidget> TagSet_OnGenerateWidget(TSharedPtr<FMemoryTagSetId> InTagSet) const;
 	FText TagSet_GetSelectedText() const;
@@ -340,6 +345,7 @@ private:
 	/** Filter the LLM tags by tag set. */
 	FMemoryTagSetId TagSetFilter = FMemoryTagSet::DefaultTagSetId;
 	TArray<TSharedPtr<FMemoryTagSetId>> AvailableTagSets;
+	TSharedPtr<SSegmentedControl<FMemoryTagSetId>> TagSetsSegmentedControl;
 
 	//////////////////////////////////////////////////
 
