[Insights]
* Memory Insights / LLM Tags: Updated the view presets to automatically change visibility of the time markers needed. Also renamed presets: "Default" to "Time A" and "Min-Max" to "Time Range".
* Table Tree Views: Changed "Preset" to "View".

diff --git a/Engine/Source/Developer/TraceInsights/Private/Insights/CookProfiler/Widgets/SPackageTableTreeView.cpp b/Engine/Source/Developer/TraceInsights/Private/Insights/CookProfiler/Widgets/SPackageTableTreeView.cpp
--- a/Engine/Source/Developer/TraceInsights/Private/Insights/CookProfiler/Widgets/SPackageTableTreeView.cpp
+++ b/Engine/Source/Developer/TraceInsights/Private/Insights/CookProfiler/Widgets/SPackageTableTreeView.cpp
@@ -189,45 +189,9 @@ BEGIN_SLATE_FUNCTION_BUILD_OPTIMIZATION
 TSharedPtr<SWidget> SPackageTableTreeView::ConstructToolbar()
 {
 	TSharedPtr<SHorizontalBox> Box = SNew(SHorizontalBox);
-
-	Box->AddSlot()
-		.AutoWidth()
-		.VAlign(VAlign_Center)
-		.Padding(0.0f, 0.0f, 4.0f, 0.0f)
-		[
-			SNew(STextBlock)
-			.Text(LOCTEXT("Preset", "Preset:"))
-		];
-
-	Box->AddSlot()
-		.AutoWidth()
-		.VAlign(VAlign_Center)
-		[
-			SNew(SBox)
-			.MinDesiredWidth(150.0f)
-			[
-				SAssignNew(PresetComboBox, SComboBox<TSharedRef<ITableTreeViewPreset>>)
-				.ToolTipText(this, &SPackageTableTreeView::ViewPreset_GetSelectedToolTipText)
-				.OptionsSource(GetAvailableViewPresets())
-				.OnSelectionChanged(this, &SPackageTableTreeView::ViewPreset_OnSelectionChanged)
-				.InitiallySelectedItem(SelectedViewPreset)
-				.OnGenerateWidget(this, &SPackageTableTreeView::ViewPreset_OnGenerateWidget)
-				[
-					SNew(STextBlock)
-					.Text(this, &SPackageTableTreeView::ViewPreset_GetSelectedText)
-				]
-			]
-		];
-
+	ConstructViewPreset(Box, 128.0f);
 	return Box;
 }
-
-////////////////////////////////////////////////////////////////////////////////////////////////////
-
-TSharedPtr<SWidget> SPackageTableTreeView::ConstructFooter()
-{
-	return nullptr;
-}
 END_SLATE_FUNCTION_BUILD_OPTIMIZATION
 
 ////////////////////////////////////////////////////////////////////////////////////////////////////
diff --git a/Engine/Source/Developer/TraceInsights/Private/Insights/CookProfiler/Widgets/SPackageTableTreeView.h b/Engine/Source/Developer/TraceInsights/Private/Insights/CookProfiler/Widgets/SPackageTableTreeView.h
--- a/Engine/Source/Developer/TraceInsights/Private/Insights/CookProfiler/Widgets/SPackageTableTreeView.h
+++ b/Engine/Source/Developer/TraceInsights/Private/Insights/CookProfiler/Widgets/SPackageTableTreeView.h
@@ -37,7 +37,6 @@ public:
 	void Construct(const FArguments& InArgs, TSharedPtr<FPackageTable> InTablePtr);
 
 	virtual TSharedPtr<SWidget> ConstructToolbar() override;
-	virtual TSharedPtr<SWidget> ConstructFooter() override;
 
 	TSharedPtr<FPackageTable> GetPackageTable() { return StaticCastSharedPtr<FPackageTable>(GetTable()); }
 	const TSharedPtr<FPackageTable> GetPackageTable() const { return StaticCastSharedPtr<FPackageTable>(GetTable()); }
diff --git a/Engine/Source/Developer/TraceInsights/Private/Insights/MemoryProfiler/ViewModels/MemTagTableViewPresets.cpp b/Engine/Source/Developer/TraceInsights/Private/Insights/MemoryProfiler/ViewModels/MemTagTableViewPresets.cpp
--- a/Engine/Source/Developer/TraceInsights/Private/Insights/MemoryProfiler/ViewModels/MemTagTableViewPresets.cpp
+++ b/Engine/Source/Developer/TraceInsights/Private/Insights/MemoryProfiler/ViewModels/MemTagTableViewPresets.cpp
@@ -8,7 +8,9 @@
 
 // TraceInsights
 #include "Insights/MemoryProfiler/ViewModels/MemTagTable.h"
+#include "Insights/MemoryProfiler/Widgets/SMemoryProfilerWindow.h"
 #include "Insights/MemoryProfiler/Widgets/SMemTagTreeView.h"
+#include "Insights/TimingProfiler/ViewModels/TimeMarker.h"
 
 #define LOCTEXT_NAMESPACE "UE::Insights::MemoryProfiler::SMemTagTableTreeView"
 
@@ -25,11 +27,11 @@ TSharedRef<ITableTreeViewPreset> FMemTagTableViewPresets::CreateDefaultViewPrese
 	public:
 		virtual FText GetName() const override
 		{
-			return LOCTEXT("Default_PresetName", "Default");
+			return LOCTEXT("Default_PresetName", "Time A");
 		}
 		virtual FText GetToolTip() const override
 		{
-			return LOCTEXT("Default_PresetToolTip", "Default View\nConfigure the tree view to show the default LLM tags info.");
+			return LOCTEXT("Default_PresetToolTip", "Time A View\nConfigure the tree view to show the LLM tags and their values at time A.");
 		}
 		virtual FName GetSortColumn() const override
 		{
@@ -52,6 +54,18 @@ TSharedRef<ITableTreeViewPreset> FMemTagTableViewPresets::CreateDefaultViewPrese
 			InOutConfigSet.Add({ FMemTagTableColumns::SizeAColumnId,      true, 100.0f });
 			InOutConfigSet.Add({ FMemTagTableColumns::SizeBudgetColumnId, true, 100.0f });
 		}
+		virtual void OnAppliedToView(STableTreeView& TableTreeView) const override
+		{
+			SMemTagTreeView& MemTagTreeView = (SMemTagTreeView&)TableTreeView;
+			TSharedPtr<SMemoryProfilerWindow> ProfilerWindow = MemTagTreeView.GetProfilerWindow();
+			if (ProfilerWindow.IsValid())
+			{
+				const TSharedRef<TimingProfiler::FTimeMarker>& MarkerA = ProfilerWindow->GetCustomTimeMarker(0);
+				MarkerA->SetVisibility(true);
+				const TSharedRef<TimingProfiler::FTimeMarker>& MarkerB = ProfilerWindow->GetCustomTimeMarker(1);
+				MarkerB->SetVisibility(false);
+			}
+		}
 	};
 	return MakeShared<FDefaultViewPreset>();
 }
@@ -95,25 +109,37 @@ TSharedRef<ITableTreeViewPreset> FMemTagTableViewPresets::CreateDiffViewPreset(S
 			InOutConfigSet.Add({ FMemTagTableColumns::SizeDiffColumnId,   true, 80.0f });
 			InOutConfigSet.Add({ FMemTagTableColumns::SizeBudgetColumnId, true, 80.0f });
 		}
+		virtual void OnAppliedToView(STableTreeView& TableTreeView) const override
+		{
+			SMemTagTreeView& MemTagTreeView = (SMemTagTreeView&)TableTreeView;
+			TSharedPtr<SMemoryProfilerWindow> ProfilerWindow = MemTagTreeView.GetProfilerWindow();
+			if (ProfilerWindow.IsValid())
+			{
+				const TSharedRef<TimingProfiler::FTimeMarker>& MarkerA = ProfilerWindow->GetCustomTimeMarker(0);
+				MarkerA->SetVisibility(true);
+				const TSharedRef<TimingProfiler::FTimeMarker>& MarkerB = ProfilerWindow->GetCustomTimeMarker(1);
+				MarkerB->SetVisibility(true);
+			}
+		}
 	};
 	return MakeShared<FDiffViewPreset>();
 }
 
 ////////////////////////////////////////////////////////////////////////////////////////////////////
-// Min-Max View
+// Time Range View
 
-TSharedRef<ITableTreeViewPreset> FMemTagTableViewPresets::CreateMinMaxViewPreset(SMemTagTreeView& TableTreeView)
+TSharedRef<ITableTreeViewPreset> FMemTagTableViewPresets::CreateTimeRangeViewPreset(SMemTagTreeView& TableTreeView)
 {
-	class FMinMaxViewPreset : public ITableTreeViewPreset
+	class FTimeRangeViewPreset : public ITableTreeViewPreset
 	{
 	public:
 		virtual FText GetName() const override
 		{
-			return LOCTEXT("MinMax_PresetName", "Min-Max");
+			return LOCTEXT("TimeRange_PresetName", "Time Range");
 		}
 		virtual FText GetToolTip() const override
 		{
-			return LOCTEXT("MinMax_PresetToolTip", "Min-Max View\nConfigure the tree view to investigate the aggregated LLM stats for the selected time region.");
+			return LOCTEXT("TimeRange_PresetToolTip", "Time Range View\nConfigure the tree view to investigate the aggregated LLM stats for the selected time range.");
 		}
 		virtual FName GetSortColumn() const override
 		{
@@ -138,8 +164,20 @@ TSharedRef<ITableTreeViewPreset> FMemTagTableViewPresets::CreateMinMaxViewPreset
 			InOutConfigSet.Add({ FMemTagTableColumns::SizeAverageColumnId, true, 80.0f });
 			InOutConfigSet.Add({ FMemTagTableColumns::SizeBudgetColumnId,  true, 80.0f });
 		}
+		virtual void OnAppliedToView(STableTreeView& TableTreeView) const override
+		{
+			SMemTagTreeView& MemTagTreeView = (SMemTagTreeView&)TableTreeView;
+			TSharedPtr<SMemoryProfilerWindow> ProfilerWindow = MemTagTreeView.GetProfilerWindow();
+			if (ProfilerWindow.IsValid())
+			{
+				const TSharedRef<TimingProfiler::FTimeMarker>& MarkerA = ProfilerWindow->GetCustomTimeMarker(0);
+				MarkerA->SetVisibility(false);
+				const TSharedRef<TimingProfiler::FTimeMarker>& MarkerB = ProfilerWindow->GetCustomTimeMarker(1);
+				MarkerB->SetVisibility(false);
+			}
+		}
 	};
-	return MakeShared<FMinMaxViewPreset>();
+	return MakeShared<FTimeRangeViewPreset>();
 }
 
 ////////////////////////////////////////////////////////////////////////////////////////////////////
diff --git a/Engine/Source/Developer/TraceInsights/Private/Insights/MemoryProfiler/ViewModels/MemTagTableViewPresets.h b/Engine/Source/Developer/TraceInsights/Private/Insights/MemoryProfiler/ViewModels/MemTagTableViewPresets.h
--- a/Engine/Source/Developer/TraceInsights/Private/Insights/MemoryProfiler/ViewModels/MemTagTableViewPresets.h
+++ b/Engine/Source/Developer/TraceInsights/Private/Insights/MemoryProfiler/ViewModels/MemTagTableViewPresets.h
@@ -18,7 +18,7 @@ class FMemTagTableViewPresets
 public:
 	static TSharedRef<ITableTreeViewPreset> CreateDefaultViewPreset(SMemTagTreeView& TableTreeView);
 	static TSharedRef<ITableTreeViewPreset> CreateDiffViewPreset(SMemTagTreeView& TableTreeView);
-	static TSharedRef<ITableTreeViewPreset> CreateMinMaxViewPreset(SMemTagTreeView& TableTreeView);
+	static TSharedRef<ITableTreeViewPreset> CreateTimeRangeViewPreset(SMemTagTreeView& TableTreeView);
 };
 
 } // namespace UE::Insights::MemoryProfiler
diff --git a/Engine/Source/Developer/TraceInsights/Private/Insights/MemoryProfiler/Widgets/SMemTagTreeView.cpp b/Engine/Source/Developer/TraceInsights/Private/Insights/MemoryProfiler/Widgets/SMemTagTreeView.cpp
--- a/Engine/Source/Developer/TraceInsights/Private/Insights/MemoryProfiler/Widgets/SMemTagTreeView.cpp
+++ b/Engine/Source/Developer/TraceInsights/Private/Insights/MemoryProfiler/Widgets/SMemTagTreeView.cpp
@@ -1067,7 +1067,7 @@ void SMemTagTreeView::InitAvailableViewPresets()
 {
 	AvailableViewPresets.Add(FMemTagTableViewPresets::CreateDefaultViewPreset(*this));
 	AvailableViewPresets.Add(FMemTagTableViewPresets::CreateDiffViewPreset(*this));
-	AvailableViewPresets.Add(FMemTagTableViewPresets::CreateMinMaxViewPreset(*this));
+	AvailableViewPresets.Add(FMemTagTableViewPresets::CreateTimeRangeViewPreset(*this));
 
 	SelectedViewPreset = AvailableViewPresets[0];
 }
diff --git a/Engine/Source/Developer/TraceInsightsCore/Private/InsightsCore/Table/Widgets/STableTreeView.cpp b/Engine/Source/Developer/TraceInsightsCore/Private/InsightsCore/Table/Widgets/STableTreeView.cpp
--- a/Engine/Source/Developer/TraceInsightsCore/Private/InsightsCore/Table/Widgets/STableTreeView.cpp
+++ b/Engine/Source/Developer/TraceInsightsCore/Private/InsightsCore/Table/Widgets/STableTreeView.cpp
@@ -3929,6 +3929,8 @@ void STableTreeView::ApplyViewPreset(const ITableTreeViewPreset& InPreset)
 	TArray<FTableColumnConfig > ColumnConfigSet;
 	InPreset.GetColumnConfigSet(ColumnConfigSet);
 	ApplyColumnConfig(ColumnConfigSet);
+
+	InPreset.OnAppliedToView(*this);
 }
 
 ////////////////////////////////////////////////////////////////////////////////////////////////////
@@ -3997,7 +3999,7 @@ FText STableTreeView::ViewPreset_GetSelectedToolTipText() const
 ////////////////////////////////////////////////////////////////////////////////////////////////////
 
 BEGIN_SLATE_FUNCTION_BUILD_OPTIMIZATION
-void STableTreeView::ConstructViewPreset(TSharedPtr<SHorizontalBox> Box)
+void STableTreeView::ConstructViewPreset(TSharedPtr<SHorizontalBox> Box, float MinDesiredWidth)
 {
 	Box->AddSlot()
 		.AutoWidth()
@@ -4005,7 +4007,7 @@ void STableTreeView::ConstructViewPreset(TSharedPtr<SHorizontalBox> Box)
 		.Padding(0.0f, 0.0f, 4.0f, 0.0f)
 		[
 			SNew(STextBlock)
-			.Text(LOCTEXT("Preset", "Preset:"))
+			.Text(LOCTEXT("View", "View:"))
 		];
 
 	Box->AddSlot()
@@ -4013,6 +4015,7 @@ void STableTreeView::ConstructViewPreset(TSharedPtr<SHorizontalBox> Box)
 		.VAlign(VAlign_Center)
 		[
 			SNew(SBox)
+			.MinDesiredWidth(MinDesiredWidth)
 			[
 				SAssignNew(PresetComboBox, SComboBox<TSharedRef<ITableTreeViewPreset>>)
 				.ToolTipText(this, &STableTreeView::ViewPreset_GetSelectedToolTipText)
diff --git a/Engine/Source/Developer/TraceInsightsCore/Public/InsightsCore/Table/Widgets/STableTreeView.h b/Engine/Source/Developer/TraceInsightsCore/Public/InsightsCore/Table/Widgets/STableTreeView.h
--- a/Engine/Source/Developer/TraceInsightsCore/Public/InsightsCore/Table/Widgets/STableTreeView.h
+++ b/Engine/Source/Developer/TraceInsightsCore/Public/InsightsCore/Table/Widgets/STableTreeView.h
@@ -66,6 +66,8 @@ struct TRACEINSIGHTSCORE_API FTableColumnConfig
 	float Width;
 };
 
+class STableTreeView;
+
 class TRACEINSIGHTSCORE_API ITableTreeViewPreset
 {
 public:
@@ -75,6 +77,7 @@ public:
 	virtual EColumnSortMode::Type GetSortMode() const = 0;
 	virtual void SetCurrentGroupings(const TArray<TSharedPtr<FTreeNodeGrouping>>& InAvailableGroupings, TArray<TSharedPtr<FTreeNodeGrouping>>& InOutCurrentGroupings) const = 0;
 	virtual void GetColumnConfigSet(TArray<FTableColumnConfig>& InOutConfigSet) const = 0;
+	virtual void OnAppliedToView(STableTreeView& TableTreeView) const {}
 };
 
 class TRACEINSIGHTSCORE_API FTableTaskCancellationToken
@@ -475,7 +478,7 @@ protected:
 	TSharedRef<SWidget> ViewPreset_OnGenerateWidget(TSharedRef<ITableTreeViewPreset> InPreset);
 	FText ViewPreset_GetSelectedText() const;
 	FText ViewPreset_GetSelectedToolTipText() const;
-	void ConstructViewPreset(TSharedPtr<SHorizontalBox> InHorizontalBox);
+	void ConstructViewPreset(TSharedPtr<SHorizontalBox> InHorizontalBox, float InMinDesiredWidth = 50.0f);
 
 	////////////////////////////////////////////////////////////////////////////////////////////////////
 
