UnrealInsights: enhancement to TimingInsights.ExportTimerStatistics.

- Add ability to sort the resulting CSV by a field (only TotalInclusiveTime supported atm) and limit the number of exported timers to "top 100"

diff --git a/Engine/Source/Developer/TraceInsights/Private/Insights/TimingProfiler/TimingProfilerManager.cpp b/Engine/Source/Developer/TraceInsights/Private/Insights/TimingProfiler/TimingProfilerManager.cpp
--- a/Engine/Source/Developer/TraceInsights/Private/Insights/TimingProfiler/TimingProfilerManager.cpp
+++ b/Engine/Source/Developer/TraceInsights/Private/Insights/TimingProfiler/TimingProfilerManager.cpp
@@ -828,6 +828,9 @@ bool FTimingProfilerManager::Exec(const TCHAR* Cmd, FOutputDevice& Ar)
 				static constexpr TCHAR StartTimeToken[] = TEXT("-startTime=");
 				static constexpr TCHAR EndTimeToken[] = TEXT("-endTime=");
 				static constexpr TCHAR RegionToken[] = TEXT("-region=");
+				static constexpr TCHAR MaxTimerCountToken[] = TEXT("-maxTimerCount=");
+				static constexpr TCHAR SortByToken[] = TEXT("-sortBy=");
+				static constexpr TCHAR SortOrderToken[] = TEXT("-sortOrder=");
 
 				if (Token.StartsWith(ColumnsToken))
 				{
@@ -882,6 +885,57 @@ bool FTimingProfilerManager::Exec(const TCHAR* Cmd, FOutputDevice& Ar)
 					Token.TrimQuotesInline();
 					Params.Region = TCHAR_TO_ANSI(*Token);
 				}
+				else if (Token.StartsWith(MaxTimerCountToken))
+				{
+					// Integer limiting the number of timers to export
+					// Example: -MaxTimerCount=100
+					Token.RightChopInline(UE_ARRAY_COUNT(MaxTimerCountToken) - 1);
+					Params.MaxExportedEvents = FCString::Atoi(*Token);
+				}
+				else if (Token.StartsWith(SortByToken))
+				{
+					static constexpr TCHAR TotalInclusiveTimeToken[] = TEXT("TotalInclusiveTime");
+
+					// Choice of field to sort the exported timers by
+					// Example: -sortBy=TotalInclusiveTime
+					Token.RightChopInline(UE_ARRAY_COUNT(SortByToken) - 1);
+					Token.TrimQuotesInline();
+					if (Token.Equals(TotalInclusiveTimeToken, ESearchCase::IgnoreCase))
+					{
+						Params.SortBy = FTimingExporter::FExportTimerStatisticsParams::ESortBy::TotalInclusiveTime;
+						// default to descending order to avoid needing to pass
+						if (Params.SortOrder == FTimingExporter::FExportTimerStatisticsParams::ESortOrder::DontSort)
+						{
+							Params.SortOrder = FTimingExporter::FExportTimerStatisticsParams::ESortOrder::Descending;
+						}
+					}
+					else
+					{
+						Ar.Logf(ELogVerbosity::Warning, TEXT("Unsupported sortBy value: %s"), *Token);
+					}
+				}
+				else if (Token.StartsWith(SortOrderToken))
+				{
+					static constexpr TCHAR DescendingToken[] = TEXT("Descending");
+					static constexpr TCHAR AscendingToken[] = TEXT("Ascending");
+
+					// Sorting order for the exported timers.
+					// Example: -sortOrder=Descending
+					Token.RightChopInline(UE_ARRAY_COUNT(SortOrderToken) - 1);
+					Token.TrimQuotesInline();
+					if (Token.Equals(DescendingToken, ESearchCase::IgnoreCase))
+					{
+						Params.SortOrder = FTimingExporter::FExportTimerStatisticsParams::ESortOrder::Descending;
+					}
+					else if (Token.Equals(AscendingToken, ESearchCase::IgnoreCase))
+					{
+						Params.SortOrder = FTimingExporter::FExportTimerStatisticsParams::ESortOrder::Ascending;
+					}
+					else
+					{
+						Ar.Logf(ELogVerbosity::Warning, TEXT("Unsupported sortOrder value: %s"), *Token);
+					}
+				}
 				else
 				{
 					Ar.Logf(ELogVerbosity::Warning, TEXT("Unknown Cmd Param: %s"), *Token);
diff --git a/Engine/Source/Developer/TraceInsights/Private/Insights/TimingProfiler/ViewModels/TimingExporter.cpp b/Engine/Source/Developer/TraceInsights/Private/Insights/TimingProfiler/ViewModels/TimingExporter.cpp
--- a/Engine/Source/Developer/TraceInsights/Private/Insights/TimingProfiler/ViewModels/TimingExporter.cpp
+++ b/Engine/Source/Developer/TraceInsights/Private/Insights/TimingProfiler/ViewModels/TimingExporter.cpp
@@ -748,6 +748,44 @@ int32 FTimingExporter::ExportTimerStatisticsAsTextByRegions(const FString& Filen
 
 ////////////////////////////////////////////////////////////////////////////////////////////////////
 
+namespace TimingExporter
+{
+
+TraceServices::FCreateAggregationParams::ESortBy MapTimerStatisticsParams(FTimingExporter::FExportTimerStatisticsParams::ESortBy Arg)
+{
+	switch(Arg)
+	{
+		default:
+			ensureMsgf(false, TEXT("Unmapped FExportTimerStatisticsParams::ESortBy value %d"), static_cast<int32>(Arg));
+			// intentional fall-through
+		case FTimingExporter::FExportTimerStatisticsParams::ESortBy::DontSort:
+			return TraceServices::FCreateAggregationParams::ESortBy::DontSort;
+
+		case FTimingExporter::FExportTimerStatisticsParams::ESortBy::TotalInclusiveTime:
+			return TraceServices::FCreateAggregationParams::ESortBy::TotalInclusiveTime;
+	}
+}
+
+TraceServices::FCreateAggregationParams::ESortOrder MapTimerStatisticsParams(FTimingExporter::FExportTimerStatisticsParams::ESortOrder Arg)
+{
+	switch (Arg)
+	{
+		default:
+			ensureMsgf(false, TEXT("Unmapped FExportTimerStatisticsParams::ESortOrder value %d"), static_cast<int32>(Arg));
+			// intentional fall-through
+		case FTimingExporter::FExportTimerStatisticsParams::ESortOrder::DontSort:
+			return TraceServices::FCreateAggregationParams::ESortOrder::DontSort;
+
+		case FTimingExporter::FExportTimerStatisticsParams::ESortOrder::Descending:
+			return TraceServices::FCreateAggregationParams::ESortOrder::Descending;
+
+		case FTimingExporter::FExportTimerStatisticsParams::ESortOrder::Ascending:
+			return TraceServices::FCreateAggregationParams::ESortOrder::Ascending;
+	}
+}
+
+} // namespace TimingExporter
+
 int32 FTimingExporter::ExportTimerStatisticsAsText(const FString& Filename, FExportTimerStatisticsParams& Params) const
 {
 	if (!Params.Region.IsEmpty())
@@ -774,6 +812,10 @@ int32 FTimingExporter::ExportTimerStatisticsAsText(const FString& Filename, FExp
 		CreateAggregationParams.GpuQueueFilter = [](uint32) { return true; };
 		CreateAggregationParams.CpuThreadFilter = Params.ThreadFilter;
 
+		CreateAggregationParams.SortBy = TimingExporter::MapTimerStatisticsParams(Params.SortBy);
+		CreateAggregationParams.SortOrder = TimingExporter::MapTimerStatisticsParams(Params.SortOrder);
+		CreateAggregationParams.TableEntryLimit = Params.MaxExportedEvents;
+
 		//@Todo: this does not yet handle the -column and -timers parameters.
 		StatsTable = TimingProfilerProvider->CreateAggregation(CreateAggregationParams);
 	}
diff --git a/Engine/Source/Developer/TraceInsights/Private/Insights/TimingProfiler/ViewModels/TimingExporter.h b/Engine/Source/Developer/TraceInsights/Private/Insights/TimingProfiler/ViewModels/TimingExporter.h
--- a/Engine/Source/Developer/TraceInsights/Private/Insights/TimingProfiler/ViewModels/TimingExporter.h
+++ b/Engine/Source/Developer/TraceInsights/Private/Insights/TimingProfiler/ViewModels/TimingExporter.h
@@ -91,6 +91,39 @@ public:
 
 	struct FExportTimerStatisticsParams : public FExportTimingEventsParams
 	{
+		/**
+		 * Enum governing field to use for sorting of exported events
+		 */
+		enum class ESortBy
+		{
+			DontSort,
+			TotalInclusiveTime
+		};
+
+		/**
+		 * Enum governing sorting order of exported events
+		 */
+		enum class ESortOrder
+		{
+			DontSort,
+			Descending,
+			Ascending
+		};
+
+		/**
+		 * Whether to sort the exported timers by a field, and which one if so.
+		 */
+		ESortBy SortBy = ESortBy::DontSort;
+
+		/**
+		 * Sorting order of the exported timers, descending or ascending
+		 */
+		ESortOrder SortOrder = ESortOrder::DontSort;
+
+		/**
+		 * Whether to limit the exported timing events (e.g. "top 100"). 0 means none.
+		 */
+		int MaxExportedEvents = 0;
 	};
 
 	struct FExportTimerCalleesParams : private FExportTimingEventsParams
diff --git a/Engine/Source/Developer/TraceServices/Private/Model/TimingProfiler.cpp b/Engine/Source/Developer/TraceServices/Private/Model/TimingProfiler.cpp
--- a/Engine/Source/Developer/TraceServices/Private/Model/TimingProfiler.cpp
+++ b/Engine/Source/Developer/TraceServices/Private/Model/TimingProfiler.cpp
@@ -669,6 +669,23 @@ ITable<FTimingProfilerAggregatedStats>* FTimingProfilerProvider::CreateAggregati
 		return Table;
 	}
 
+	if (Params.SortBy == FCreateAggregationParams::ESortBy::TotalInclusiveTime)
+	{
+		auto SortLambda = (Params.SortOrder == FCreateAggregationParams::ESortOrder::Descending) ?
+			[](const FAggregatedTimingStats& First, const FAggregatedTimingStats& Second) -> bool
+			{
+				return First.TotalInclusiveTime > Second.TotalInclusiveTime;
+			}
+			:
+			[](const FAggregatedTimingStats& First, const FAggregatedTimingStats& Second) -> bool
+			{
+				return First.TotalInclusiveTime < Second.TotalInclusiveTime;
+			};
+
+		Aggregation.ValueStableSort(SortLambda);
+	}
+
+	int EntryLimit = Params.TableEntryLimit;
 	for (const auto& KV : Aggregation)
 	{
 		FTimingProfilerAggregatedStats& Row = Table->AddRow();
@@ -685,6 +702,14 @@ ITable<FTimingProfilerAggregatedStats>* FTimingProfilerProvider::CreateAggregati
 		Row.MaxExclusiveTime = Stats.MaxExclusiveTime;
 		Row.AverageExclusiveTime = Stats.AverageExclusiveTime;
 		Row.MedianExclusiveTime = Stats.MedianExclusiveTime;
+
+		if (UNLIKELY(EntryLimit > 0))
+		{
+			if (UNLIKELY(--EntryLimit == 0))
+			{
+				break;
+			}
+		}
 	}
 	return Table;
 }
diff --git a/Engine/Source/Developer/TraceServices/Public/TraceServices/Model/TimingProfiler.h b/Engine/Source/Developer/TraceServices/Public/TraceServices/Model/TimingProfiler.h
--- a/Engine/Source/Developer/TraceServices/Public/TraceServices/Model/TimingProfiler.h
+++ b/Engine/Source/Developer/TraceServices/Public/TraceServices/Model/TimingProfiler.h
@@ -68,6 +68,20 @@ struct FTimingProfilerButterflyNode
 
 struct FCreateAggregationParams
 {
+	// whether to sort the created aggregation
+	enum class ESortBy
+	{
+		DontSort,
+		TotalInclusiveTime
+	};
+
+	enum class ESortOrder
+	{
+		DontSort,
+		Descending,
+		Ascending
+	};
+
 	// The start timestamp in seconds.
 	double IntervalStart = 0.0;
 
@@ -80,6 +94,15 @@ struct FCreateAggregationParams
 	// A boolean value to specify if aggregation should include the old GPU timeline.
 	bool bIncludeOldGpu2 = false;
 
+	// Whether to sort the table by a field
+	ESortBy SortBy = ESortBy::DontSort;
+
+	// Descending or ascending
+	ESortOrder SortOrder = ESortOrder::DontSort;
+
+	// Limit the entries in aggregation (e.g. "top 100")
+	int TableEntryLimit = 0;
+
 	// A function to filter the GPU queues to aggregate.
 	TFunction<bool(uint32)> GpuQueueFilter;
 
