[Insights] Removed the Search box from few context menus.

#rnx

diff --git a/Engine/Source/Developer/TraceInsights/Private/Insights/TimingProfiler/Widgets/STimersView.cpp b/Engine/Source/Developer/TraceInsights/Private/Insights/TimingProfiler/Widgets/STimersView.cpp
--- a/Engine/Source/Developer/TraceInsights/Private/Insights/TimingProfiler/Widgets/STimersView.cpp
+++ b/Engine/Source/Developer/TraceInsights/Private/Insights/TimingProfiler/Widgets/STimersView.cpp
@@ -818,7 +818,9 @@ void STimersView::TreeView_BuildViewColumnMenu(FMenuBuilder& MenuBuilder)
 
 void STimersView::TreeView_BuildExportMenu(FMenuBuilder& MenuBuilder)
 {
-	MenuBuilder.BeginSection("Export", LOCTEXT("ContextMenu_Section_Export", "Export"));
+	MenuBuilder.SetSearchable(false);
+
+	MenuBuilder.BeginSection("Export");
 	{
 		MenuBuilder.AddMenuEntry
 		(
@@ -854,6 +856,8 @@ void STimersView::TreeView_BuildExportMenu(FMenuBuilder& MenuBuilder)
 
 void STimersView::TreeView_BuildPlotTimerMenu(FMenuBuilder& MenuBuilder)
 {
+	MenuBuilder.SetSearchable(false);
+
 	const TArray<FTimerNodePtr> SelectedNodes = TreeView->GetSelectedItems();
 	const int32 NumSelectedNodes = SelectedNodes.Num();
 	FTimerNodePtr SelectedNode = NumSelectedNodes ? SelectedNodes[0] : nullptr;
@@ -1059,6 +1063,8 @@ void STimersView::TreeView_BuildPlotTimerMenu(FMenuBuilder& MenuBuilder)
 
 void STimersView::TreeView_FindMenu(FMenuBuilder& MenuBuilder)
 {
+	MenuBuilder.SetSearchable(false);
+
 	MenuBuilder.AddMenuEntry
 	(
 		FTimersViewCommands::Get().Command_FindMaxInstance,
diff --git a/Engine/Source/Developer/TraceInsights/Private/Insights/Widgets/STimingView.cpp b/Engine/Source/Developer/TraceInsights/Private/Insights/Widgets/STimingView.cpp
--- a/Engine/Source/Developer/TraceInsights/Private/Insights/Widgets/STimingView.cpp
+++ b/Engine/Source/Developer/TraceInsights/Private/Insights/Widgets/STimingView.cpp
@@ -4431,6 +4431,7 @@ TSharedRef<SWidget> STimingView::MakeCompactAutoScrollOptionsMenu()
 TSharedRef<SWidget> STimingView::MakeAutoScrollOptionsMenu()
 {
 	FMenuBuilder MenuBuilder(/*bInShouldCloseWindowAfterMenuSelection=*/true, CommandList);
+	MenuBuilder.SetSearchable(false);
 
 	MenuBuilder.BeginSection("Alignment", LOCTEXT("AutoScrollOptionsMenu_Section_Alignment", "Alignment"));
 	{
@@ -4621,6 +4622,7 @@ TSharedRef<SWidget> STimingView::MakeAutoScrollOptionsMenu()
 TSharedRef<SWidget> STimingView::MakeAllTracksMenu()
 {
 	FMenuBuilder MenuBuilder(/*bInShouldCloseWindowAfterMenuSelection=*/true, CommandList);
+	MenuBuilder.SetSearchable(false);
 
 	CreateAllTracksMenu(MenuBuilder);
 
