[Insights]
* Memory Insights / LLM Tags: Changed the tag sets to use a smaller font and icon size. Also added label.
* Memory Insights / LLM Tags: Moved budget settings into a single dropdown menu.
* Memory Insights / LLM Tags: Made the time marker settings collapsible.
* Memory Insights / LLM Tags: Removed Search box from some context menus.

diff --git a/Engine/Source/Developer/TraceInsights/Private/Insights/InsightsStyle.cpp b/Engine/Source/Developer/TraceInsights/Private/Insights/InsightsStyle.cpp
--- a/Engine/Source/Developer/TraceInsights/Private/Insights/InsightsStyle.cpp
+++ b/Engine/Source/Developer/TraceInsights/Private/Insights/InsightsStyle.cpp
@@ -224,6 +224,9 @@ void FInsightsStyle::FStyle::Initialize()
 	Set("Icons.TagSet.Assets", new EDITOR_IMAGE_BRUSH_SVG("Starship/Common/FileLined", Icon16x16));
 	Set("Icons.TagSet.AssetClasses", new EDITOR_IMAGE_BRUSH_SVG("Starship/Common/Class", Icon16x16));
 
+	Set("Icons.BudgetSettings", new CORE_IMAGE_BRUSH_SVG("Starship/Common/file-tree", Icon16x16));
+	Set("Icons.TimeMarkerSettings", new EDITOR_IMAGE_BRUSH_SVG("Starship/Common/Timeline", Icon16x16));
+
 	//////////////////////////////////////////////////
 	// Tasks
 
diff --git a/Engine/Source/Developer/TraceInsights/Private/Insights/MemoryProfiler/Widgets/SMemTagTreeView.cpp b/Engine/Source/Developer/TraceInsights/Private/Insights/MemoryProfiler/Widgets/SMemTagTreeView.cpp
--- a/Engine/Source/Developer/TraceInsights/Private/Insights/MemoryProfiler/Widgets/SMemTagTreeView.cpp
+++ b/Engine/Source/Developer/TraceInsights/Private/Insights/MemoryProfiler/Widgets/SMemTagTreeView.cpp
@@ -6,6 +6,7 @@
 #include "Framework/MultiBox/MultiBoxBuilder.h"
 #include "Misc/ConfigCacheIni.h"
 #include "SlateOptMacros.h"
+#include "Styling/StarshipCoreStyle.h"
 #include "Styling/StyleColors.h"
 #include "Widgets/Colors/SColorPicker.h"
 #include "Widgets/Images/SImage.h"
@@ -82,7 +83,21 @@ void SMemTagTreeView::ConstructHeaderArea(TSharedRef<SVerticalBox> InHostBox)
 		.AutoHeight()
 		.Padding(2.0f, 2.0f, 2.0f, 2.0f)
 		[
-			SNew(SBox)
+			SNew(SHorizontalBox)
+
+			+ SHorizontalBox::Slot()
+			.AutoWidth()
+			.Padding(0.0f, 0.0f, 4.0f, 0.0f)
+			.VAlign(VAlign_Center)
+			[
+				SNew(STextBlock)
+				.Text(LOCTEXT("TagSetText", "Tag Set:"))
+			]
+
+			+ SHorizontalBox::Slot()
+			.AutoWidth()
+			.Padding(0.0f, 0.0f, 0.0f, 0.0f)
+			.VAlign(VAlign_Center)
 			[
 				SAssignNew(TagSetsSegmentedControl, SSegmentedControl<FMemoryTagSetId>)
 				.SupportsMultiSelection(false)
@@ -112,17 +127,6 @@ void SMemTagTreeView::ConstructHeaderArea(TSharedRef<SVerticalBox> InHostBox)
 			];
 	}
 
-	TSharedPtr<SWidget> BudgetSettingsWidget = ConstructBudgetSettings();
-	if (BudgetSettingsWidget.IsValid())
-	{
-		InHostBox->AddSlot()
-			.AutoHeight()
-			.Padding(2.0f, 2.0f, 2.0f, 2.0f)
-			[
-				BudgetSettingsWidget.ToSharedRef()
-			];
-	}
-
 	TSharedPtr<SWidget> FilterToolbarWidget = ConstructFilterToolbar();
 	if (FilterToolbarWidget.IsValid())
 	{
@@ -134,25 +138,33 @@ void SMemTagTreeView::ConstructHeaderArea(TSharedRef<SVerticalBox> InHostBox)
 			];
 	}
 
-	TSharedPtr<SWidget> HierarchyBreadcrumbTrailWidget = ConstructHierarchyBreadcrumbTrail();
-	if (HierarchyBreadcrumbTrailWidget.IsValid())
+	TSharedPtr<SWidget> TimeMarkersWidget = ConstructTimeMarkers();
+	if (TimeMarkersWidget.IsValid())
 	{
 		InHostBox->AddSlot()
 			.AutoHeight()
-			.Padding(2.0f, 2.0f, 2.0f, 2.0f)
+			.Padding(0.0f, 0.0f, 0.0f, 0.0f)
 			[
-				HierarchyBreadcrumbTrailWidget.ToSharedRef()
+				SNew(SBox)
+				.Padding(20.0f, 1.0f, 2.0f, 1.0f)
+				.Visibility_Lambda([this]
+				{
+					return bAreTimeMarkerSettingsVisible ? EVisibility::Visible : EVisibility::Collapsed;
+				})
+				[
+					TimeMarkersWidget.ToSharedRef()
+				]
 			];
 	}
 
-	TSharedPtr<SWidget> TimeMarkersWidget = ConstructTimeMarkers();
-	if (TimeMarkersWidget.IsValid())
+	TSharedPtr<SWidget> HierarchyBreadcrumbTrailWidget = ConstructHierarchyBreadcrumbTrail();
+	if (HierarchyBreadcrumbTrailWidget.IsValid())
 	{
 		InHostBox->AddSlot()
 			.AutoHeight()
-			.Padding(20.0f, 2.0f, 2.0f, -2.0f)
+			.Padding(2.0f, 2.0f, 2.0f, -2.0f)
 			[
-				TimeMarkersWidget.ToSharedRef()
+				HierarchyBreadcrumbTrailWidget.ToSharedRef()
 			];
 	}
 }
@@ -204,6 +216,14 @@ TSharedPtr<SWidget> SMemTagTreeView::ConstructTopSettings()
 		ConstructTrackHeightControls()
 	]
 
+	+ SHorizontalBox::Slot()
+	.AutoWidth()
+	.Padding(FMargin(4.0f, 0.0f, 0.0f, 0.0f))
+	.VAlign(VAlign_Center)
+	[
+		ConstructBudgetSettings()
+	]
+
 #if 0
 	+ SHorizontalBox::Slot()
 	.FillWidth(1.0f)
@@ -233,7 +253,7 @@ TSharedPtr<SWidget> SMemTagTreeView::ConstructTopSettings()
 #else
 	+ SHorizontalBox::Slot()
 	.FillWidth(1.0f)
-	.Padding(0.0f, 0.0f, 0.0f, 0.0f)
+	.Padding(4.0f, 0.0f, 0.0f, 0.0f)
 	.HAlign(HAlign_Right)
 	.VAlign(VAlign_Center)
 	[
@@ -286,10 +306,36 @@ void SMemTagTreeView::UpdateAvailableTagSets()
 		const FText TagSetToolTipFmt = LOCTEXT("TagSetToolTipFmt", "'{0}' Tag Set\n\nThe tree view shows the memory tags from the selected tag set.");
 		TagSetToolTip = FText::Format(TagSetToolTipFmt, TagSetName);
 
+		TSharedRef<SWidget> Widget = SNew(SHorizontalBox)
+
+			+ SHorizontalBox::Slot()
+			.AutoWidth()
+			.Padding(0.0f, 0.0f, 2.0f, 0.0f)
+			.VAlign(VAlign_Center)
+			[
+				SNew(SBox)
+				.WidthOverride(12.0f)
+				.HeightOverride(12.0f)
+				[
+					SNew(SImage)
+					.Image(TagSetIcon)
+					.ColorAndOpacity(FSlateColor(EStyleColor::Foreground))
+				]
+			]
+
+			+ SHorizontalBox::Slot()
+			.AutoWidth()
+			.Padding(0.0f, 0.0f, 0.0f, 0.0f)
+			.VAlign(VAlign_Center)
+			[
+				SNew(STextBlock)
+				.Font(FStyleFonts::Get().Small)
+				.Text(TagSetName)
+			];
+
 		TagSetsSegmentedControl->AddSlot(TagSet)
-			.Icon(TagSetIcon)
-			.Text(TagSetName)
-			.ToolTip(TagSetToolTip);
+			.ToolTip(TagSetToolTip)
+			.AttachWidget(Widget);
 	}
 }
 
@@ -346,6 +392,30 @@ TSharedPtr<SWidget> SMemTagTreeView::ConstructFilterToolbar()
 	.VAlign(VAlign_Center)
 	[
 		ConstructFilterConfiguratorButton()
+	]
+
+	+ SHorizontalBox::Slot()
+	.AutoWidth()
+	.Padding(2.0f, 0.0f, 0.0f, 0.0f)
+	.VAlign(VAlign_Center)
+	[
+		SNew(SCheckBox)
+		.Style(FAppStyle::Get(), "ToggleButtonCheckbox")
+		.Padding(FMargin(4.0f, 3.0f, 4.0f, 3.0f))
+		.IsChecked_Lambda([this]()
+			{
+				return bAreTimeMarkerSettingsVisible ? ECheckBoxState::Checked : ECheckBoxState::Unchecked;
+			})
+		.OnCheckStateChanged_Lambda([this](ECheckBoxState InCheckBoxState)
+			{
+				bAreTimeMarkerSettingsVisible = (InCheckBoxState == ECheckBoxState::Checked);
+			})
+		.Content()
+		[
+			SNew(SImage)
+			.Image(FInsightsStyle::GetBrush("Icons.TimeMarkerSettings"))
+		]
+		.ToolTipText(LOCTEXT("TimeMarkerSettingsVisibilityToolTip", "Toggle visibility for the advanced Time Marker settings."))
 	];
 
 	return FilterLine;
@@ -481,13 +551,12 @@ END_SLATE_FUNCTION_BUILD_OPTIMIZATION
 TSharedRef<SWidget> SMemTagTreeView::MakeTrackersMenu()
 {
 	FMenuBuilder MenuBuilder(/*bInShouldCloseWindowAfterMenuSelection=*/true, nullptr);
+	MenuBuilder.SetSearchable(false);
 
 	MenuBuilder.BeginSection("Trackers"/*, LOCTEXT("ContextMenu_Section_Trackers", "Trackers")*/);
 	CreateTrackersMenuSection(MenuBuilder);
 	MenuBuilder.EndSection();
 
-	MenuBuilder.SetSearchable(false);
-
 	return MenuBuilder.MakeWidget();
 }
 
@@ -684,86 +753,16 @@ END_SLATE_FUNCTION_BUILD_OPTIMIZATION
 BEGIN_SLATE_FUNCTION_BUILD_OPTIMIZATION
 TSharedRef<SWidget> SMemTagTreeView::ConstructBudgetSettings()
 {
-	float constexpr TopPadding = 0.0f;
-
-	TSharedRef<SHorizontalBox> Box = SNew(SHorizontalBox)
-
-	+ SHorizontalBox::Slot()
-	.AutoWidth()
-	.VAlign(VAlign_Center)
-	.Padding(0.0f, TopPadding, 4.0f, 0.0f)
-	[
-		SNew(STextBlock)
-		.Text(LOCTEXT("BudgetFileText", "Budget File:"))
-	]
-
-	+ SHorizontalBox::Slot()
-	.FillWidth(1.0f)
-	.VAlign(VAlign_Center)
-	.Padding(0.0f, TopPadding, 0.0f, 0.0f)
-	[
-		SAssignNew(BudgetFileComboBox, SComboBox<TSharedPtr<FMemTagBudgetFileDesc>>)
-		.ToolTipText(this, &SMemTagTreeView::BudgetFile_GetTooltipText)
-		.OptionsSource(&AvailableBudgetFiles)
-		.OnSelectionChanged(this, &SMemTagTreeView::BudgetFile_OnSelectionChanged)
-		.OnGenerateWidget(this, &SMemTagTreeView::BudgetFile_OnGenerateWidget)
-		[
-			SNew(STextBlock)
-			.Text(this, &SMemTagTreeView::BudgetFile_GetSelectedText)
-		]
-	]
-
-	+ SHorizontalBox::Slot()
-	.AutoWidth()
-	.VAlign(VAlign_Center)
-	.Padding(6.0f, TopPadding, 4.0f, 0.0f)
-	[
-		SNew(STextBlock)
-		.Text(LOCTEXT("BudgetModeText", "Mode:"))
-	]
-
-	+ SHorizontalBox::Slot()
-	.AutoWidth()
-	.VAlign(VAlign_Center)
-	.Padding(0.0f, TopPadding, 0.0f, 0.0f)
-	[
-		SAssignNew(BudgetModeComboBox, SComboBox<TSharedPtr<FMemTagBudgetModeDesc>>)
-		.ToolTipText(this, &SMemTagTreeView::BudgetMode_GetTooltipText)
-		.OptionsSource(&AvailableBudgetModes)
-		.OnSelectionChanged(this, &SMemTagTreeView::BudgetMode_OnSelectionChanged)
-		.OnGenerateWidget(this, &SMemTagTreeView::BudgetMode_OnGenerateWidget)
-		[
-			SNew(STextBlock)
-			.Text(this, &SMemTagTreeView::BudgetMode_GetSelectedText)
-		]
-	]
-
-	+ SHorizontalBox::Slot()
-	.AutoWidth()
-	.VAlign(VAlign_Center)
-	.Padding(6.0f, TopPadding, 4.0f, 0.0f)
-	[
-		SNew(STextBlock)
-		.Text(LOCTEXT("PlatformModeText", "Platform:"))
-	]
-
-	+ SHorizontalBox::Slot()
-	.AutoWidth()
-	.VAlign(VAlign_Center)
-	.Padding(0.0f, TopPadding, 0.0f, 0.0f)
-	[
-		SAssignNew(BudgetPlatformComboBox, SComboBox<TSharedPtr<FMemTagBudgetPlatformDesc>>)
-		.ToolTipText(this, &SMemTagTreeView::BudgetPlatform_GetTooltipText)
-		.OptionsSource(&AvailableBudgetPlatforms)
-		.OnSelectionChanged(this, &SMemTagTreeView::BudgetPlatform_OnSelectionChanged)
-		.OnGenerateWidget(this, &SMemTagTreeView::BudgetPlatform_OnGenerateWidget)
+	return SNew(SComboButton)
+		.OnGetMenuContent(this, &SMemTagTreeView::MakeBudgetSettingsMenu)
+		.ToolTipText(LOCTEXT("BudgetSettingsToolTip", "Budget Settings"))
+		.ContentPadding(FMargin(0.0f, 1.0f))
+		.ButtonContent()
 		[
-			SNew(STextBlock)
-			.Text(this, &SMemTagTreeView::BudgetPlatform_GetSelectedText)
-		]
-	];
-
-	return Box;
+			SNew(SImage)
+			.Image(FInsightsStyle::GetBrush("Icons.BudgetSettings"))
+			.ColorAndOpacity(FSlateColor(EStyleColor::Foreground))
+		];
 }
 END_SLATE_FUNCTION_BUILD_OPTIMIZATION
 
@@ -1532,6 +1531,8 @@ void SMemTagTreeView::ExtendMenuBeforeMisc(FMenuBuilder& MenuBuilder)
 
 void SMemTagTreeView::ExtendMenuCreateGraphTracks(FMenuBuilder& MenuBuilder)
 {
+	MenuBuilder.SetSearchable(false);
+
 	MenuBuilder.BeginSection("CreateGraphTracks");
 	{
 		// Create memory graph tracks for selected LLM tag(s)
@@ -1589,6 +1590,8 @@ void SMemTagTreeView::ExtendMenuCreateGraphTracks(FMenuBuilder& MenuBuilder)
 
 void SMemTagTreeView::ExtendMenuRemoveGraphTracks(FMenuBuilder& MenuBuilder)
 {
+	MenuBuilder.SetSearchable(false);
+
 	MenuBuilder.BeginSection("RemoveGraphTracks");
 	{
 		// Remove memory graph tracks for selected LLM tag(s)
@@ -1700,38 +1703,43 @@ bool SMemTagTreeView::CanLoadReportXML() const
 void SMemTagTreeView::LoadReportXML()
 {
 	TSharedPtr<SMemoryProfilerWindow> ProfilerWindow = GetProfilerWindow();
-	if (ProfilerWindow.IsValid())
+	if (!ProfilerWindow.IsValid())
 	{
-		TArray<FString> Files;
-		bool bOpened = false;
+		return;
+	}
 
-		IDesktopPlatform* DesktopPlatform = FDesktopPlatformModule::Get();
-		if (DesktopPlatform != nullptr)
-		{
-			FSlateApplication::Get().CloseToolTip();
-
-			const FString DefaultPath(FPaths::RootDir() / TEXT("Engine/Binaries/DotNET/CsvTools"));
-			const FString DefaultFile(TEXT("LLMReportTypes.xml"));
-
-			bOpened = DesktopPlatform->OpenFileDialog
-			(
-				FSlateApplication::Get().FindBestParentWindowHandleForDialogs(nullptr),
-				LOCTEXT("LoadReportXML_FileDesc", "Open the LLMReportTypes.xml file...").ToString(),
-				DefaultPath,
-				DefaultFile, // Not actually used. See FDesktopPlatformWindows::FileDialogShared implementation. :(
-				LOCTEXT("LoadReportXML_FileFilter", "XML files (*.xml)|*.xml|All files (*.*)|*.*").ToString(),
-				EFileDialogFlags::None,
-				Files
-			);
-		}
+	IDesktopPlatform* DesktopPlatform = FDesktopPlatformModule::Get();
+	if (DesktopPlatform != nullptr)
+	{
+		return;
+	}
 
-		if (bOpened && Files.Num() == 1)
-		{
-			FMemorySharedState& SharedState = ProfilerWindow->GetSharedState();
-			SharedState.RemoveAllMemTagGraphTracks();
-			SharedState.CreateTracksFromReport(Files[0]);
-		}
+	FSlateApplication::Get().CloseToolTip();
+
+	const FString DefaultPath(FPaths::RootDir() / TEXT("Engine/Binaries/DotNET/CsvTools"));
+	const FString DefaultFile(TEXT("LLMReportTypes.xml"));
+
+	TArray<FString> Files;
+	if (!DesktopPlatform->OpenFileDialog(
+		FSlateApplication::Get().FindBestParentWindowHandleForDialogs(nullptr),
+		LOCTEXT("LoadReportXML_FileDesc", "Open the LLMReportTypes.xml file...").ToString(),
+		DefaultPath,
+		DefaultFile, // Not actually used. See FDesktopPlatformWindows::FileDialogShared implementation. :(
+		LOCTEXT("LoadReportXML_FileFilter", "XML files (*.xml)|*.xml|All files (*.*)|*.*").ToString(),
+		EFileDialogFlags::None,
+		Files))
+	{
+		return;
+	}
+
+	if (Files.Num() != 1)
+	{
+		return;
 	}
+
+	FMemorySharedState& SharedState = ProfilerWindow->GetSharedState();
+	SharedState.RemoveAllMemTagGraphTracks();
+	SharedState.CreateTracksFromReport(Files[0]);
 }
 
 ////////////////////////////////////////////////////////////////////////////////////////////////////
@@ -2116,62 +2124,167 @@ void SMemTagTreeView::ColorPickerCancelled(FLinearColor OriginalColor)
 }
 
 ////////////////////////////////////////////////////////////////////////////////////////////////////
-// Budget File
+// Budgets
 ////////////////////////////////////////////////////////////////////////////////////////////////////
 
-void SMemTagTreeView::BudgetFile_OnSelectionChanged(TSharedPtr<FMemTagBudgetFileDesc> InBudgetFile, ESelectInfo::Type SelectInfo)
+TSharedRef<SWidget> SMemTagTreeView::MakeBudgetSettingsMenu()
 {
-	if (InBudgetFile.IsValid() &&
-		InBudgetFile->Type == FMemTagBudgetFileDesc::EType::OpenFileDialog)
+	FMenuBuilder MenuBuilder(/*bInShouldCloseWindowAfterMenuSelection=*/true, nullptr);
+	MenuBuilder.SetSearchable(false);
+
+	MenuBuilder.BeginSection("File", LOCTEXT("BudgetSettingsMenu_Section_File", "File"));
+	for (const auto& BudgetFile : AvailableBudgetFiles)
 	{
-		TArray<FString> Files;
-		bool bOpened = false;
+		FText Label = FText::FromString(FString::Printf(TEXT("%s (%s)"), *BudgetFile->Name, *BudgetFile->FileName));
+		MenuBuilder.AddMenuEntry(
+			Label,
+			TAttribute<FText>(),
+			FSlateIcon(),
+			FUIAction(
+				FExecuteAction::CreateSP(this, &SMemTagTreeView::SelectBudgetFile, BudgetFile),
+				FCanExecuteAction(),
+				FIsActionChecked::CreateSP(this, &SMemTagTreeView::IsBudgetFileSelected, BudgetFile)),
+			NAME_None,
+			EUserInterfaceActionType::Check);
+	}
+	MenuBuilder.AddMenuEntry(
+		LOCTEXT("OpenBudgetFile_Text", "Open File..."),
+		LOCTEXT("OpenBudgetFile_ToolTip", "Choose budget file..."),
+		FSlateIcon(),
+		FUIAction(
+			FExecuteAction::CreateSP(this, &SMemTagTreeView::OpenAndSelectBudgetFile),
+			FCanExecuteAction()),
+		NAME_None,
+		EUserInterfaceActionType::Button);
+	MenuBuilder.EndSection();
 
-		IDesktopPlatform* DesktopPlatform = FDesktopPlatformModule::Get();
-		if (DesktopPlatform != nullptr)
-		{
-			FSlateApplication::Get().CloseToolTip();
-
-			const FString DefaultPath(FPaths::RootDir() / TEXT("Engine/Programs/UnrealInsights/Config"));
-
-			bOpened = DesktopPlatform->OpenFileDialog
-			(
-				FSlateApplication::Get().FindBestParentWindowHandleForDialogs(nullptr),
-				LOCTEXT("LoadBudgetFile_FileDesc", "Open the Memory Budget .xml file...").ToString(),
-				DefaultPath,
-				FString(),
-				LOCTEXT("LoadBudgetFile_FileFilter", "XML files (*.xml)|*.xml|All files (*.*)|*.*").ToString(),
-				EFileDialogFlags::None,
-				Files
-			);
-		}
+	MenuBuilder.BeginSection("BudgetMode", LOCTEXT("BudgetSettingsMenu_Section_Budget", "Budget"));
+	for (const auto& BudgetMode : AvailableBudgetModes)
+	{
+		MenuBuilder.AddMenuEntry(
+			FText::FromString(BudgetMode->Name),
+			TAttribute<FText>(),
+			FSlateIcon(),
+			FUIAction(
+				FExecuteAction::CreateSP(this, &SMemTagTreeView::SelectBudgetMode, BudgetMode),
+				FCanExecuteAction(),
+				FIsActionChecked::CreateSP(this, &SMemTagTreeView::IsBudgetModeSelected, BudgetMode)),
+			NAME_None,
+			EUserInterfaceActionType::Check);
+	}
+	MenuBuilder.EndSection();
 
-		if (bOpened && Files.Num() == 1)
-		{
-			InBudgetFile = MakeShared<FMemTagBudgetFileDesc>(FPaths::GetCleanFilename(Files[0]), Files[0]);
-			bool bIsAlreadyAvailable = false;
-			for (int32 BudgetFileIndex = 0; BudgetFileIndex < AvailableBudgetFiles.Num() - 1; ++BudgetFileIndex)
-			{
-				if (*InBudgetFile == *AvailableBudgetFiles[BudgetFileIndex])
-				{
-					InBudgetFile = AvailableBudgetFiles[BudgetFileIndex];
-					bIsAlreadyAvailable = true;
-				}
-			}
-			if (!bIsAlreadyAvailable)
-			{
-				AvailableBudgetFiles.Insert(InBudgetFile, AvailableBudgetFiles.Num() - 1);
-			}
-			BudgetFileComboBox->SetSelectedItem(InBudgetFile);
-		}
-		else
+	MenuBuilder.BeginSection("BudgetPlatform", LOCTEXT("BudgetSettingsMenu_Section_Platform", "Platform"));
+	for (const auto& BudgetPlatform : AvailableBudgetPlatforms)
+	{
+		MenuBuilder.AddMenuEntry(
+			FText::FromString(BudgetPlatform->Name),
+			TAttribute<FText>(),
+			FSlateIcon(),
+			FUIAction(
+				FExecuteAction::CreateSP(this, &SMemTagTreeView::SelectBudgetPlatform, BudgetPlatform),
+				FCanExecuteAction(),
+				FIsActionChecked::CreateSP(this, &SMemTagTreeView::IsBudgetPlatformSelected, BudgetPlatform)),
+			NAME_None,
+			EUserInterfaceActionType::Check);
+	}
+	MenuBuilder.EndSection();
+
+	return MenuBuilder.MakeWidget();
+}
+
+////////////////////////////////////////////////////////////////////////////////////////////////////
+
+TSharedRef<SWidget> MakeBudgetModeSubMenu()
+{
+	FMenuBuilder MenuBuilder(/*bInShouldCloseWindowAfterMenuSelection=*/true, nullptr);
+	MenuBuilder.SetSearchable(false);
+
+	MenuBuilder.BeginSection("BudgetMode");
+	//
+	MenuBuilder.EndSection();
+
+	return MenuBuilder.MakeWidget();
+}
+
+////////////////////////////////////////////////////////////////////////////////////////////////////
+
+TSharedRef<SWidget> MakeBudgetPlatformSubMenu()
+{
+	FMenuBuilder MenuBuilder(/*bInShouldCloseWindowAfterMenuSelection=*/true, nullptr);
+	MenuBuilder.SetSearchable(false);
+
+	MenuBuilder.BeginSection("BudgetPlatform");
+	//
+	MenuBuilder.EndSection();
+
+	return MenuBuilder.MakeWidget();
+}
+
+////////////////////////////////////////////////////////////////////////////////////////////////////
+
+void SMemTagTreeView::OpenAndSelectBudgetFile()
+{
+	TSharedPtr<FMemTagBudgetFileDesc> BudgetFile = OpenBudgetFile();
+	if (BudgetFile)
+	{
+		SelectBudgetFile(BudgetFile);
+	}
+}
+
+////////////////////////////////////////////////////////////////////////////////////////////////////
+
+TSharedPtr<FMemTagBudgetFileDesc> SMemTagTreeView::OpenBudgetFile()
+{
+	IDesktopPlatform* DesktopPlatform = FDesktopPlatformModule::Get();
+	if (!DesktopPlatform)
+	{
+		return nullptr;
+	}
+
+	FSlateApplication::Get().CloseToolTip();
+
+	const FString DefaultPath(FPaths::RootDir() / TEXT("Engine/Programs/UnrealInsights/Config"));
+
+	TArray<FString> Files;
+	if (!DesktopPlatform->OpenFileDialog(
+		FSlateApplication::Get().FindBestParentWindowHandleForDialogs(nullptr),
+		LOCTEXT("LoadBudgetFile_FileDesc", "Open the Memory Budget .xml file...").ToString(),
+		DefaultPath,
+		FString(),
+		LOCTEXT("LoadBudgetFile_FileFilter", "XML files (*.xml)|*.xml|All files (*.*)|*.*").ToString(),
+		EFileDialogFlags::None,
+		Files))
+	{
+		return nullptr;
+	}
+
+	if (Files.Num() == 1)
+	{
+		return nullptr;
+	}
+
+	const FString& FileName = Files[0];
+	FString Name = FPaths::GetCleanFilename(FileName);
+
+	for (TSharedPtr<FMemTagBudgetFileDesc>& BudgetFile : AvailableBudgetFiles)
+	{
+		if (Name.Equals(BudgetFile->Name, ESearchCase::IgnoreCase) &&
+			FileName.Equals(BudgetFile->FileName, ESearchCase::IgnoreCase))
 		{
-			// No file opened, revert combo box selection
-			InBudgetFile = SelectedBudgetFile;
-			BudgetFileComboBox->SetSelectedItem(SelectedBudgetFile);
+			return BudgetFile;
 		}
 	}
 
+	TSharedPtr<FMemTagBudgetFileDesc> NewBudgetFile = MakeShared<FMemTagBudgetFileDesc>(Name, FileName);
+	AvailableBudgetFiles.Add(NewBudgetFile);
+	return NewBudgetFile;
+}
+
+////////////////////////////////////////////////////////////////////////////////////////////////////
+
+void SMemTagTreeView::SelectBudgetFile(TSharedPtr<FMemTagBudgetFileDesc> InBudgetFile)
+{
 	if (InBudgetFile.IsValid())
 	{
 		if (!SelectedBudgetFile.IsValid() || *SelectedBudgetFile != *InBudgetFile)
@@ -2201,98 +2314,44 @@ void SMemTagTreeView::BudgetFile_OnSelectionChanged(TSharedPtr<FMemTagBudgetFile
 
 ////////////////////////////////////////////////////////////////////////////////////////////////////
 
-TSharedRef<SWidget> SMemTagTreeView::BudgetFile_OnGenerateWidget(TSharedPtr<FMemTagBudgetFileDesc> InBudgetFile)
-{
-	if (InBudgetFile.IsValid())
-	{
-		if (InBudgetFile->Type == FMemTagBudgetFileDesc::EType::NamedBudgetFile)
-		{
-			return SNew(STextBlock)
-				.Text(FText::FromString(FString::Printf(TEXT("%s (%s)"), *InBudgetFile->Name, *InBudgetFile->FileName)));
-				//.Text(FText::FromString(InBudgetFile->Name))
-				//.ToolTip(FSlateApplicationBase::Get().MakeToolTip(FText::FromString(InBudgetFile->FileName)));
-		}
-		else if (InBudgetFile->Type == FMemTagBudgetFileDesc::EType::OpenFileDialog)
-		{
-			return SNew(STextBlock)
-				.Text(LOCTEXT("BudgetFileComboBox_OpenFileOption", "Open file..."));
-		}
-	}
-	return SNew(SBox);
-}
-
-////////////////////////////////////////////////////////////////////////////////////////////////////
-
-FText SMemTagTreeView::BudgetFile_GetSelectedText() const
+bool SMemTagTreeView::IsBudgetFileSelected(TSharedPtr<FMemTagBudgetFileDesc> InBudgetFile) const
 {
-	if (SelectedBudgetFile.IsValid())
-	{
-		if (SelectedBudgetFile->Type == FMemTagBudgetFileDesc::EType::NamedBudgetFile)
-		{
-			return FText::FromString(*SelectedBudgetFile->Name);
-		}
-		else if (SelectedBudgetFile->Type == FMemTagBudgetFileDesc::EType::OpenFileDialog)
-		{
-			return LOCTEXT("BudgetFileComboBox_OpenFileOption2", "Choose budget file...");
-		}
-	}
-	return FText::GetEmpty();
+	return SelectedBudgetFile.IsValid () && *InBudgetFile == *SelectedBudgetFile;
 }
 
 ////////////////////////////////////////////////////////////////////////////////////////////////////
 
-FText SMemTagTreeView::BudgetFile_GetTooltipText() const
-{
-	return LOCTEXT("BudgetFileComboBox_Tooltip", "Choose the current memory budget file");
-}
-
-////////////////////////////////////////////////////////////////////////////////////////////////////
-// Budget Mode
-////////////////////////////////////////////////////////////////////////////////////////////////////
-
-void SMemTagTreeView::BudgetMode_OnSelectionChanged(TSharedPtr<FMemTagBudgetModeDesc> InBudgetMode, ESelectInfo::Type SelectInfo)
+void SMemTagTreeView::SelectBudgetMode(TSharedPtr<FMemTagBudgetModeDesc> InBudgetMode)
 {
 	if (InBudgetMode.IsValid())
 	{
-		if (!SelectedBudgetMode.IsValid() || !SelectedBudgetMode->Name.Equals(InBudgetMode->Name, ESearchCase::IgnoreCase))
+		if (!SelectedBudgetMode.IsValid() ||
+			!SelectedBudgetMode->Name.Equals(InBudgetMode->Name, ESearchCase::IgnoreCase))
 		{
 			SelectedBudgetMode = InBudgetMode;
 			bShouldUpdateBudgets = true;
 		}
 	}
-}
-
-////////////////////////////////////////////////////////////////////////////////////////////////////
-
-TSharedRef<SWidget> SMemTagTreeView::BudgetMode_OnGenerateWidget(TSharedPtr<FMemTagBudgetModeDesc> InBudgetMode)
-{
-	return SNew(STextBlock)
-		.Text(FText::FromString(InBudgetMode->Name));
-}
-
-////////////////////////////////////////////////////////////////////////////////////////////////////
-
-FText SMemTagTreeView::BudgetMode_GetSelectedText() const
-{
-	if (SelectedBudgetMode.IsValid())
+	else
 	{
-		return FText::FromString(SelectedBudgetMode->Name);
+		if (SelectedBudgetMode.IsValid())
+		{
+			SelectedBudgetMode = nullptr;
+			bShouldUpdateBudgets = true;
+		}
 	}
-	return FText::GetEmpty();
 }
 
 ////////////////////////////////////////////////////////////////////////////////////////////////////
 
-FText SMemTagTreeView::BudgetMode_GetTooltipText() const
+bool SMemTagTreeView::IsBudgetModeSelected(TSharedPtr<FMemTagBudgetModeDesc> InBudgetMode) const
 {
-	return LOCTEXT("BudgetModeComboBox_Tooltip", "Choose the current memory budget mode");
+	return SelectedBudgetMode.IsValid() && *InBudgetMode == *SelectedBudgetMode;
 }
 
-////////////////////////////////////////////////////////////////////////////////////////////////////
-// Budget Platform
 ////////////////////////////////////////////////////////////////////////////////////////////////////
 
-void SMemTagTreeView::BudgetPlatform_OnSelectionChanged(TSharedPtr<FMemTagBudgetPlatformDesc> InBudgetPlatform, ESelectInfo::Type SelectInfo)
+void SMemTagTreeView::SelectBudgetPlatform(TSharedPtr<FMemTagBudgetPlatformDesc> InBudgetPlatform)
 {
 	if (InBudgetPlatform.IsValid())
 	{
@@ -2302,36 +2361,23 @@ void SMemTagTreeView::BudgetPlatform_OnSelectionChanged(TSharedPtr<FMemTagBudget
 			bShouldUpdateBudgets = true;
 		}
 	}
-}
-
-////////////////////////////////////////////////////////////////////////////////////////////////////
-
-TSharedRef<SWidget> SMemTagTreeView::BudgetPlatform_OnGenerateWidget(TSharedPtr<FMemTagBudgetPlatformDesc> InBudgetPlatform)
-{
-	return SNew(STextBlock)
-		.Text(FText::FromString(InBudgetPlatform->Name));
-}
-
-////////////////////////////////////////////////////////////////////////////////////////////////////
-
-FText SMemTagTreeView::BudgetPlatform_GetSelectedText() const
-{
-	if (SelectedBudgetPlatform.IsValid())
+	else
 	{
-		return FText::FromString(SelectedBudgetPlatform->Name);
+		if (SelectedBudgetPlatform.IsValid())
+		{
+			SelectedBudgetPlatform = nullptr;
+			bShouldUpdateBudgets = true;
+		}
 	}
-	return FText::GetEmpty();
 }
 
 ////////////////////////////////////////////////////////////////////////////////////////////////////
 
-FText SMemTagTreeView::BudgetPlatform_GetTooltipText() const
+bool SMemTagTreeView::IsBudgetPlatformSelected(TSharedPtr<FMemTagBudgetPlatformDesc> InBudgetPlatform) const
 {
-	return LOCTEXT("BudgetPlatformComboBox_Tooltip", "Choose the current memory budget platform");
+	return SelectedBudgetPlatform.IsValid() && *InBudgetPlatform == *SelectedBudgetPlatform;
 }
 
-////////////////////////////////////////////////////////////////////////////////////////////////////
-// Budget
 ////////////////////////////////////////////////////////////////////////////////////////////////////
 
 const TCHAR* SMemTagTreeView::GetTagSetFilterBudgetCachedName() const
@@ -2456,15 +2502,13 @@ void SMemTagTreeView::InitBudgetOptions()
 		UE_LOG(LogMemoryProfiler, Log, TEXT("Couldn't find UnrealInsightsSettings config"));
 	}
 
-	AvailableBudgetFiles.Add(MakeShared<FMemTagBudgetFileDesc>()); // open file dialog
-
 	AvailableBudgetModes.Reset();
 	AvailableBudgetPlatforms.Reset();
 	AvailablePlatforms.Reset();
 
-	if (AvailableBudgetFiles.Num() > 1)
+	if (AvailableBudgetFiles.Num() > 0)
 	{
-		BudgetFileComboBox->SetSelectedItem(AvailableBudgetFiles[0]);
+		SelectBudgetFile(AvailableBudgetFiles[0]);
 	}
 }
 
@@ -2517,7 +2561,7 @@ void SMemTagTreeView::OnBudgetChanged()
 		{
 			if (PrevSelectedBudgetMode->Name.Equals(BudgetMode->Name, ESearchCase::CaseSensitive))
 			{
-				BudgetModeComboBox->SetSelectedItem(BudgetMode);
+				SelectBudgetMode(BudgetMode);
 				bFound = true;
 				break;
 			}
@@ -2527,18 +2571,18 @@ void SMemTagTreeView::OnBudgetChanged()
 			if (AvailableBudgetModes.Num() > 0)
 			{
 				// Select first available mode.
-				BudgetModeComboBox->SetSelectedItem(AvailableBudgetModes[0]);
+				SelectBudgetMode(AvailableBudgetModes[0]);
 			}
 			else
 			{
-				BudgetModeComboBox->ClearSelection();
+				SelectBudgetMode(nullptr);
 			}
 		}
 	}
 	if (!SelectedBudgetMode.IsValid() && AvailableBudgetModes.Num() > 0)
 	{
 		// Select the first available budget mode.
-		BudgetModeComboBox->SetSelectedItem(AvailableBudgetModes[0]);
+		SelectBudgetMode(AvailableBudgetModes[0]);
 	}
 
 	// Restore selection for Budget Platform.
@@ -2549,7 +2593,7 @@ void SMemTagTreeView::OnBudgetChanged()
 		{
 			if (PrevSelectedBudgetPlatform->Name.Equals(BudgetPlatform->Name, ESearchCase::CaseSensitive))
 			{
-				BudgetPlatformComboBox->SetSelectedItem(BudgetPlatform);
+				SelectBudgetPlatform(BudgetPlatform);
 				bFound = true;
 				break;
 			}
@@ -2559,18 +2603,18 @@ void SMemTagTreeView::OnBudgetChanged()
 			if (AvailableBudgetPlatforms.Num() > 0)
 			{
 				// Select the first available platform.
-				BudgetPlatformComboBox->SetSelectedItem(AvailableBudgetPlatforms[0]);
+				SelectBudgetPlatform(AvailableBudgetPlatforms[0]);
 			}
 			else
 			{
-				BudgetPlatformComboBox->ClearSelection();
+				SelectBudgetPlatform(nullptr);
 			}
 		}
 	}
 	if (!SelectedBudgetPlatform.IsValid() && AvailableBudgetPlatforms.Num() > 0)
 	{
 		// Select the first available platform.
-		BudgetPlatformComboBox->SetSelectedItem(AvailableBudgetPlatforms[0]);
+		SelectBudgetPlatform(AvailableBudgetPlatforms[0]);
 	}
 
 	bIsLoadingBudget = false;
diff --git a/Engine/Source/Developer/TraceInsights/Private/Insights/MemoryProfiler/Widgets/SMemTagTreeView.h b/Engine/Source/Developer/TraceInsights/Private/Insights/MemoryProfiler/Widgets/SMemTagTreeView.h
--- a/Engine/Source/Developer/TraceInsights/Private/Insights/MemoryProfiler/Widgets/SMemTagTreeView.h
+++ b/Engine/Source/Developer/TraceInsights/Private/Insights/MemoryProfiler/Widgets/SMemTagTreeView.h
@@ -25,28 +25,15 @@ class SMemoryProfilerWindow;
 
 struct FMemTagBudgetFileDesc
 {
-	enum class EType
-	{
-		NamedBudgetFile,
-		OpenFileDialog
-	};
-
-	FMemTagBudgetFileDesc()
-		: Type(EType::OpenFileDialog)
-	{
-	}
-
 	FMemTagBudgetFileDesc(const FString& InName, const FString& InFileName)
-		: Type(EType::NamedBudgetFile)
-		, Name(InName)
+		: Name(InName)
 		, FileName(InFileName)
 	{
 	}
 
 	bool operator==(const FMemTagBudgetFileDesc& Other) const
 	{
-		return (Type == Other.Type)
-			&& Name.Equals(Other.Name, ESearchCase::IgnoreCase)
+		return Name.Equals(Other.Name, ESearchCase::IgnoreCase)
 			&& FileName.Equals(Other.FileName, ESearchCase::IgnoreCase);
 	}
 
@@ -55,7 +42,6 @@ struct FMemTagBudgetFileDesc
 		return !(*this == Other);
 	}
 
-	EType Type;
 	FString Name;
 	FString FileName;
 };
@@ -299,20 +285,19 @@ private:
 	////////////////////////////////////////////////////////////////////////////////////////////////////
 	// Budgets
 
-	void BudgetFile_OnSelectionChanged(TSharedPtr<FMemTagBudgetFileDesc> InBudgetFile, ESelectInfo::Type SelectInfo);
-	TSharedRef<SWidget> BudgetFile_OnGenerateWidget(TSharedPtr<FMemTagBudgetFileDesc> InBudgetFile);
-	FText BudgetFile_GetSelectedText() const;
-	FText BudgetFile_GetTooltipText() const;
+	TSharedRef<SWidget> MakeBudgetSettingsMenu();
+	TSharedRef<SWidget> MakeBudgetModeSubMenu();
+	TSharedRef<SWidget> MakeBudgetPlatformSubMenu();
 
-	void BudgetMode_OnSelectionChanged(TSharedPtr<FMemTagBudgetModeDesc> InBudgetMode, ESelectInfo::Type SelectInfo);
-	TSharedRef<SWidget> BudgetMode_OnGenerateWidget(TSharedPtr<FMemTagBudgetModeDesc> InBudgetMode);
-	FText BudgetMode_GetSelectedText() const;
-	FText BudgetMode_GetTooltipText() const;
+	TSharedPtr<FMemTagBudgetFileDesc> OpenBudgetFile();
 
-	void BudgetPlatform_OnSelectionChanged(TSharedPtr<FMemTagBudgetPlatformDesc> InPlatform, ESelectInfo::Type SelectInfo);
-	TSharedRef<SWidget> BudgetPlatform_OnGenerateWidget(TSharedPtr<FMemTagBudgetPlatformDesc> InPlatform);
-	FText BudgetPlatform_GetSelectedText() const;
-	FText BudgetPlatform_GetTooltipText() const;
+	void OpenAndSelectBudgetFile();
+	void SelectBudgetFile(TSharedPtr<FMemTagBudgetFileDesc> InBudgetFile);
+	bool IsBudgetFileSelected(TSharedPtr<FMemTagBudgetFileDesc> InBudgetFile) const;
+	void SelectBudgetMode(TSharedPtr<FMemTagBudgetModeDesc> InBudgetMode);
+	bool IsBudgetModeSelected(TSharedPtr<FMemTagBudgetModeDesc> InBudgetMode) const;
+	void SelectBudgetPlatform(TSharedPtr<FMemTagBudgetPlatformDesc> InBudgetPlatform);
+	bool IsBudgetPlatformSelected(TSharedPtr<FMemTagBudgetPlatformDesc> InBudgetPlatform) const;
 
 	void InitBudgetOptions();
 	void OnBudgetChanged();
@@ -361,10 +346,6 @@ private:
 
 	TSharedPtr<FMemTagBudget> CurrentBudget;
 
-	TSharedPtr<SComboBox<TSharedPtr<FMemTagBudgetFileDesc>>> BudgetFileComboBox;
-	TSharedPtr<SComboBox<TSharedPtr<FMemTagBudgetModeDesc>>> BudgetModeComboBox;
-	TSharedPtr<SComboBox<TSharedPtr<FMemTagBudgetPlatformDesc>>> BudgetPlatformComboBox;
-
 	TArray<TSharedPtr<FMemTagBudgetFileDesc>> AvailableBudgetFiles;
 	TArray<TSharedPtr<FMemTagBudgetModeDesc>> AvailableBudgetModes;
 	TArray<TSharedPtr<FMemTagBudgetPlatformDesc>> AvailableBudgetPlatforms;
@@ -390,6 +371,7 @@ private:
 
 	bool bShouldUpdateStats = false;
 	bool bShouldUpdateBudgets = false;
+	bool bAreTimeMarkerSettingsVisible = false;
 };
 
 ////////////////////////////////////////////////////////////////////////////////////////////////////
