[Insights] Mem Insights / LLM Tags: Fixed inverted conditions for validation code regarding OpenFileDialog usage.

#rnx

diff --git a/Engine/Source/Developer/TraceInsights/Private/Insights/MemoryProfiler/Widgets/SMemTagTreeView.cpp b/Engine/Source/Developer/TraceInsights/Private/Insights/MemoryProfiler/Widgets/SMemTagTreeView.cpp
--- a/Engine/Source/Developer/TraceInsights/Private/Insights/MemoryProfiler/Widgets/SMemTagTreeView.cpp
+++ b/Engine/Source/Developer/TraceInsights/Private/Insights/MemoryProfiler/Widgets/SMemTagTreeView.cpp
@@ -1709,7 +1709,7 @@ void SMemTagTreeView::LoadReportXML()
 	}
 
 	IDesktopPlatform* DesktopPlatform = FDesktopPlatformModule::Get();
-	if (DesktopPlatform != nullptr)
+	if (!DesktopPlatform)
 	{
 		return;
 	}
@@ -2259,7 +2259,7 @@ TSharedPtr<FMemTagBudgetFileDesc> SMemTagTreeView::OpenBudgetFile()
 		return nullptr;
 	}
 
-	if (Files.Num() == 1)
+	if (Files.Num() != 1)
 	{
 		return nullptr;
 	}
