Trace RHIBreadcrumbs using metadata with specs

diff --git a/Engine/Source/Developer/TraceInsights/Private/Insights/TimingProfiler/Tracks/ThreadTimingTrack.cpp b/Engine/Source/Developer/TraceInsights/Private/Insights/TimingProfiler/Tracks/ThreadTimingTrack.cpp
--- a/Engine/Source/Developer/TraceInsights/Private/Insights/TimingProfiler/Tracks/ThreadTimingTrack.cpp
+++ b/Engine/Source/Developer/TraceInsights/Private/Insights/TimingProfiler/Tracks/ThreadTimingTrack.cpp
@@ -41,28 +41,55 @@ INSIGHTS_IMPLEMENT_RTTI(FThreadTimingTrack);
 
 ////////////////////////////////////////////////////////////////////////////////////////////////////
 
-static void AppendMetadataToTooltip(FTooltipDrawState& Tooltip, TArrayView<const uint8>& Metadata)
+static void AppendMetadataToTooltip(FTooltipDrawState& Tooltip, const TraceServices::FMetadataSpec* MetadataSpec, TArrayView<const uint8>& Metadata)
 {
 	FMemoryReaderView MemoryReader(Metadata);
 	FCborReader CborReader(&MemoryReader, ECborEndianness::StandardCompliant);
 	FCborContext Context;
 
-	if (!CborReader.ReadNext(Context) || Context.MajorType() != ECborCode::Map)
+	if (MetadataSpec == nullptr)
 	{
-		return;
+		if (!CborReader.ReadNext(Context))
+		{
+			return;
+		}
+
+		if (Context.MajorType() != ECborCode::Map)
+		{
+			return;
+		}
 	}
 
 	Tooltip.AddTitle(TEXT("Metadata:"));
 
+	uint32 Index = 0;
 	while (true)
 	{
-		// Read key
-		if (!CborReader.ReadNext(Context) || !Context.IsString())
+		FString Key;
+
+		if (MetadataSpec)
 		{
-			break;
+			if (Index < (uint32) MetadataSpec->FieldNames.Num())
+			{
+				Key = MetadataSpec->FieldNames[Index];
+				++Index;
+			}
+			else
+			{
+				Key = TEXT("UndefinedField");
+			}
+		}
+		else
+		{
+			// Read key
+			if (!CborReader.ReadNext(Context) || !Context.IsString())
+			{
+				break;
+			}
+
+			Key = FString::ConstructFromPtrSize(Context.AsCString(), static_cast<int32>(Context.AsLength()));
 		}
 
-		FString Key = FString::ConstructFromPtrSize(Context.AsCString(), static_cast<int32>(Context.AsLength()));
 		Key += TEXT(":");
 
 		if (!CborReader.ReadNext(Context))
@@ -98,8 +125,7 @@ static void AppendMetadataToTooltip(FTooltipDrawState& Tooltip, TArrayView<const
 
 		case ECborCode::ByteString:
 			{
-				FAnsiStringView Value(Context.AsCString(), static_cast<int32>(Context.AsLength()));
-				FString ValueStr(Value);
+				FString ValueStr = FString::ConstructFromPtrSize(Context.AsCString(), static_cast<int32>(Context.AsLength()));
 				Tooltip.AddNameValueTextLine(Key, ValueStr);
 				continue;
 			}
@@ -186,55 +212,192 @@ static void AppendMetadataToString(FString& Str, TArrayView<const uint8>& Metada
 		{
 		case ECborCode::Int:
 		case ECborCode::Uint:
+		{
+			uint64 Value = Context.AsUInt();
+			if (Value > 999'999'999ULL)
 			{
-				uint64 Value = Context.AsUInt();
-				if (Value > 999'999'999ULL)
-				{
-					Str += FString::Printf(TEXT("0x%llX"), Value);
-				}
-				else
+				Str += FString::Printf(TEXT("0x%llX"), Value);
+			}
+			else
+			{
+				Str += FString::Printf(TEXT("%llu"), Value);
+			}
+			continue;
+		}
+
+		case ECborCode::TextString:
+		{
+			Str += Context.AsString();
+			continue;
+		}
+
+		case ECborCode::ByteString:
+		{
+			Str.AppendChars(Context.AsCString(), static_cast<int32>(Context.AsLength()));
+			continue;
+		}
+		}
+
+		if (Context.RawCode() == (ECborCode::Prim | ECborCode::Value_4Bytes))
+		{
+			float Value = Context.AsFloat();
+			Str += FString::Printf(TEXT("%f"), Value);
+			continue;
+		}
+
+		if (Context.RawCode() == (ECborCode::Prim | ECborCode::Value_8Bytes))
+		{
+			double Value = Context.AsDouble();
+			Str += FString::Printf(TEXT("%g"), Value);
+			continue;
+		}
+
+		if (Context.RawCode() == (ECborCode::Prim | ECborCode::False))
+		{
+			Str += TEXT("false");
+			continue;
+		}
+
+		if (Context.RawCode() == (ECborCode::Prim | ECborCode::True))
+		{
+			Str += TEXT("true");
+			continue;
+		}
+
+		if (Context.IsFiniteContainer())
+		{
+			CborReader.SkipContainer(ECborCode::Array);
+		}
+	}
+}
+
+////////////////////////////////////////////////////////////////////////////////////////////////////
+
+static void AppendMetadataToString(FString& Str, const TraceServices::FMetadataSpec* MetadataSpec, TArrayView<const uint8>& Metadata)
+{
+	if (MetadataSpec == nullptr)
+	{
+		return;
+	}
+
+	FMemoryReaderView MemoryReader(Metadata);
+	FCborReader CborReader(&MemoryReader, ECborEndianness::StandardCompliant);
+	FCborContext Context;
+
+	FString Format;
+	Format = MetadataSpec->Format;
+
+	const FString Specifiers = TEXT("diuoxXfFeEgGaAcspn");
+
+	auto GetNextFormatSection = [&Format, &Specifiers]()
+	{
+		int Index = 0;
+		bool bIsInFormatSpecifier = false;
+		while (Index < Format.Len())
+		{
+			if (bIsInFormatSpecifier)
+			{
+				int32 SpecIndex = -1;
+				if (Specifiers.FindChar(Format[Index], SpecIndex))
 				{
-					Str += FString::Printf(TEXT("%llu"), Value);
+					FString NextFormat = Format.Left(Index+1);
+					Format.MidInline(Index+1);
+					return NextFormat;
 				}
+			}
+
+			if (Format[Index] == TEXT('%'))
+			{
+				bIsInFormatSpecifier = !bIsInFormatSpecifier;
+			}
+
+			++Index;
+		}
+
+		FString Copy = Format;
+		Format.Empty();
+		return Copy;
+	};
+
+	bool bFirst = true;
+
+	while (!Format.IsEmpty())
+	{
+		if (!CborReader.ReadNext(Context))
+		{
+			break;
+		}
+
+		if (bFirst)
+		{
+			bFirst = false;
+			Str += TEXT(" - ");
+		}
+		else
+		{
+			if (Format.IsEmpty())
+			{
+				Str += TEXT(", ");
+			}
+		}
+		
+		constexpr int MaxLength = 256;
+		TCHAR Data[MaxLength];
+
+		auto AddValueToName = [&Str](TCHAR* Dest, int MaxLength, const FString& Format, auto Value)
+		{
+			int32 Result = FCString::Snprintf(Dest, MaxLength, reinterpret_cast<TCHAR const(&)[1]>(**Format), Value);
+			if (Result > 0)
+			{
+				Str.Append(Dest);
+			}
+		};
+
+		switch (Context.MajorType())
+		{
+		case ECborCode::Int:
+		case ECborCode::Uint:
+			{
+				uint64 Value = Context.AsUInt();
+				AddValueToName(Data, MaxLength, GetNextFormatSection(), Value);
+
 				continue;
 			}
 
 		case ECborCode::TextString:
 			{
-				Str += Context.AsString();
+				AddValueToName(Data, MaxLength, GetNextFormatSection(), *Context.AsString());
 				continue;
 			}
 
 		case ECborCode::ByteString:
 			{
-				Str.AppendChars(Context.AsCString(), static_cast<int32>(Context.AsLength()));
+				AddValueToName(Data, FMath::Min(MaxLength, (int)Context.AsLength()), GetNextFormatSection(), Context.AsCString());
 				continue;
 			}
 		}
 
 		if (Context.RawCode() == (ECborCode::Prim | ECborCode::Value_4Bytes))
 		{
-			float Value = Context.AsFloat();
-			Str += FString::Printf(TEXT("%f"), Value);
+			AddValueToName(Data, MaxLength, GetNextFormatSection(), Context.AsFloat());
 			continue;
 		}
 
 		if (Context.RawCode() == (ECborCode::Prim | ECborCode::Value_8Bytes))
 		{
-			double Value = Context.AsDouble();
-			Str += FString::Printf(TEXT("%g"), Value);
+			AddValueToName(Data, MaxLength, GetNextFormatSection(), Context.AsDouble());
 			continue;
 		}
 
 		if (Context.RawCode() == (ECborCode::Prim | ECborCode::False))
 		{
-			Str += TEXT("false");
+			AddValueToName(Data, MaxLength, GetNextFormatSection(), false);
 			continue;
 		}
 
 		if (Context.RawCode() == (ECborCode::Prim | ECborCode::True))
 		{
-			Str += TEXT("true");
+			AddValueToName(Data, MaxLength, GetNextFormatSection(), true);
 			continue;
 		}
 
@@ -243,6 +406,9 @@ static void AppendMetadataToString(FString& Str, TArrayView<const uint8>& Metada
 			CborReader.SkipContainer(ECborCode::Array);
 		}
 	}
+
+	// Append what's left of the format string.
+	Str.Append(Format);
 }
 
 ////////////////////////////////////////////////////////////////////////////////////////////////////
@@ -311,9 +477,23 @@ static void AddTimingEventToBuilder(ITimingEventsTrackDrawStateBuilder& Builder,
 					TimingProfilerProvider.ReadTimers([&TimerReader](const TraceServices::ITimingProfilerTimerReader& Out) { TimerReader = &Out; });
 
 					TArrayView<const uint8> Metadata = TimerReader->GetMetadata(TimerIndex);
+
 					if (Metadata.Num() > 0)
 					{
-						AppendMetadataToString(EventName, Metadata);
+						const TraceServices::FMetadataSpec* MetadataSpec = nullptr;
+						if (Timer && Timer->HasValidMetadataSpecId())
+						{
+							MetadataSpec = TimingProfilerProvider.GetMetadataSpec(Timer->MetadataSpecId);
+						}
+
+						if (MetadataSpec)
+						{
+							AppendMetadataToString(EventName, MetadataSpec, Metadata);
+						}
+						else
+						{
+							AppendMetadataToString(EventName, Metadata);
+						}
 					}
 				}
 			}
@@ -743,10 +923,15 @@ void FThreadTimingTrackImpl::InitTooltip(FTooltipDrawState& InOutTooltip, const
 
 		InOutTooltip.AddNameValueTextLine(TEXTVIEW("Depth:"), FString::Printf(TEXT("%d"), TooltipEvent.GetDepth()));
 
+		const TraceServices::FMetadataSpec* MetadataSpec = nullptr;
+		if (Timer && Timer->HasValidMetadataSpecId())
+		{
+			MetadataSpec = TimingProfilerProvider.GetMetadataSpec(Timer->MetadataSpecId);
+		}
 		TArrayView<const uint8> Metadata = TimerReader->GetMetadata(TooltipEvent.GetTimerIndex());
 		if (Metadata.Num() > 0)
 		{
-			AppendMetadataToTooltip(InOutTooltip, Metadata);
+			AppendMetadataToTooltip(InOutTooltip, MetadataSpec, Metadata);
 		}
 
 		PostInitTooltip(InOutTooltip, TooltipEvent, *Session.Get(), TimerName);
@@ -958,9 +1143,14 @@ void FThreadTimingTrackImpl::OnClipboardCopyEvent(const ITimingEvent& InSelected
 			if (int32(TimerIndex) < 0) // has metadata?
 			{
 				TArrayView<const uint8> Metadata = TimerReader->GetMetadata(TimerIndex);
+				const TraceServices::FMetadataSpec* MetadataSpec = nullptr;
+				if (TimerPtr && TimerPtr->HasValidMetadataSpecId())
+				{
+					MetadataSpec = TimingProfilerProvider.GetMetadataSpec(TimerPtr->MetadataSpecId);
+				}
 				if (Metadata.Num() > 0)
 				{
-					AppendMetadataToString(EventName, Metadata);
+					AppendMetadataToString(EventName, MetadataSpec, Metadata);
 				}
 			}
 
diff --git a/Engine/Source/Developer/TraceServices/Private/Analyzers/GpuProfilerTraceAnalysis.cpp b/Engine/Source/Developer/TraceServices/Private/Analyzers/GpuProfilerTraceAnalysis.cpp
--- a/Engine/Source/Developer/TraceServices/Private/Analyzers/GpuProfilerTraceAnalysis.cpp
+++ b/Engine/Source/Developer/TraceServices/Private/Analyzers/GpuProfilerTraceAnalysis.cpp
@@ -4,9 +4,11 @@
 
 #include "AnalysisServicePrivate.h"
 #include "CborWriter.h"
+#include "CborReader.h"
 #include "Common/Utils.h"
 #include "HAL/LowLevelMemTracker.h"
 #include "Model/CountersPrivate.h"
+#include "Serialization/MemoryReader.h"
 #include "Serialization/MemoryWriter.h"
 
 namespace TraceServices
@@ -27,6 +29,7 @@ void FGpuProfilerAnalyzer::OnAnalysisBegin(const FOnAnalysisContext& Context)
 	Builder.RouteEvent(RouteId_Init, "GpuProfiler", "Init");
 	Builder.RouteEvent(RouteId_QueueSpec, "GpuProfiler", "QueueSpec");
 	Builder.RouteEvent(RouteId_EventFrameBoundary, "GpuProfiler", "EventFrameBoundary");
+	Builder.RouteEvent(RouteId_EventBreadcrumbSpec, "GpuProfiler", "EventBreadcrumbSpec");
 	Builder.RouteEvent(RouteId_EventBeginBreadcrumb, "GpuProfiler", "EventBeginBreadcrumb");
 	Builder.RouteEvent(RouteId_EventEndBreadcrumb, "GpuProfiler", "EventEndBreadcrumb");
 	Builder.RouteEvent(RouteId_EventBeginWork, "GpuProfiler", "EventBeginWork");
@@ -143,8 +146,113 @@ bool FGpuProfilerAnalyzer::OnEvent(uint16 RouteId, EStyle Style, const FOnEventC
 
 		break;
 	}
+	case RouteId_EventBreadcrumbSpec:
+	{
+		FString Name, NameFormat;
+		uint32 SpecId = EventData.GetValue<uint32>("SpecId");
+		EventData.GetString("StaticName", Name);
+		EventData.GetString("NameFormat", NameFormat);
+		TArrayView<const uint8> FieldNames = EventData.GetArrayView<uint8>("FieldNames");
+
+		if (Name.Compare(NameFormat) == 0)
+		{
+			Name.Empty();
+		}
+
+		int Index = 0;
+		while (Index < NameFormat.Len() && Index < Name.Len())
+		{
+			if (Name[Index] != NameFormat[Index])
+			{
+				break;
+			}
+
+			++Index;
+		}
+
+		if (Index > 1)
+		{
+			NameFormat.MidInline(Index);
+		}
+
+		if (Name.IsEmpty())
+		{
+			const FString Separators = TEXT("% (=");
+			Index = 0;
+			bool bIsInFormatSpecifier = false;
+			while (Index < NameFormat.Len())
+			{
+				int32 SpecIndex = -1;
+				if (Separators.FindChar(NameFormat[Index], SpecIndex))
+				{
+					Name = NameFormat.Left(Index);
+					NameFormat.MidInline(Index);
+					break;
+				}
+
+				++Index;
+			}
+
+			if (Name.IsEmpty())
+			{
+				if (!NameFormat.IsEmpty())
+				{
+					Name = NameFormat;
+					NameFormat.Empty();
+				}
+				else
+				{
+					Name = Session.StoreString(TEXT("Unknown"));
+				}
+			}
+		}
+
+		FMetadataSpec Spec;
+		if (!NameFormat.IsEmpty())
+		{
+			Spec.Format = Session.StoreString(NameFormat);
+		}
+
+		if (FieldNames.Num() > 0)
+		{
+			FMemoryReaderView MemoryReader(FieldNames);
+			FCborReader CborReader(&MemoryReader, ECborEndianness::StandardCompliant);
+			FCborContext CborContext;
+
+			while (CborReader.ReadNext(CborContext))
+			{
+				if (CborContext.MajorType() == ECborCode::TextString)
+				{
+					FString Field = CborContext.AsString();
+					Spec.FieldNames.Add(Session.StoreString(Field));
+				}
+			}
+		}
+		
+		FAnalysisSessionEditScope _(Session);
+		uint32 TimerId = 0;
+		uint32* FoundSpecIdPtr = BreadcrumbSpecMap.Find(SpecId);
+		if (FoundSpecIdPtr == nullptr)
+		{
+			TimerId = GetOrAddTimerId(Name);
+			BreadcrumbSpecMap.Add(SpecId, TimerId);
+		}
+		else
+		{
+			TimingProfilerProvider.SetTimerName(TimerId, Name);
+		}
+
+		if (Spec.FieldNames.Num() > 0 || Spec.Format != nullptr)
+		{
+			uint32 MetadataSpecId = TimingProfilerProvider.AddMetadataSpec(MoveTemp(Spec));
+			TimingProfilerProvider.SetMetadataSpec(TimerId, MetadataSpecId);
+		}
+
+		break;
+	}
 	case RouteId_EventBeginBreadcrumb:
 	{
+		const uint32 SpecId = EventData.GetValue<uint32>("SpecId");
 		const uint32 QueueId = EventData.GetValue<uint32>("QueueId");
 		const uint64 GPUTimestampTOP = EventData.GetValue<uint64>("GPUTimestampTOP");
 		
@@ -154,15 +262,31 @@ bool FGpuProfilerAnalyzer::OnEvent(uint16 RouteId, EStyle Style, const FOnEventC
 			break;
 		}
 
+		TArray<uint8> Metadata(EventData.GetArrayView<uint8>("Metadata"));
 		const double Time = Context.EventTime.AsSeconds(GPUTimestampTOP);
 
-		FString Name;
-		EventData.GetString("Name", Name);
+		UE_LOG(LogTraceServices, VeryVerbose, TEXT("[GPU] EventBeginBreadcrumb QueueId=%u, Time=%f, SpecId=\"%u\""), QueueId, Time, SpecId);
+
+		uint32 TimerId = 0;
+		uint32* TimerIdPtr;
+		TimerIdPtr = BreadcrumbSpecMap.Find(SpecId);
+		
+		{
+			FAnalysisSessionEditScope _(Session);
+			if (TimerIdPtr)
+			{
+				TimerId = *TimerIdPtr;
+			}
+			else
+			{
+				TimerId = GetOrAddTimerId(TEXT("Unknown"));
+				BreadcrumbSpecMap.Add(SpecId, TimerId);
+			}
 
-		UE_LOG(LogTraceServices, VeryVerbose, TEXT("[GPU] EventBeginBreadcrumb QueueId=%u, Time=%f, Name=\"%s\""), QueueId, Time, *Name);
+			TimerId = TimingProfilerProvider.AddMetadata(TimerId, MoveTemp(Metadata));
+		}
 
 		FQueue& Queue = GetOrAddQueue(QueueId);
-		const uint32 TimerId = GetOrAddTimerId(Name);
 		BeginEvent(Queue, 0, Time, TimerId);
 		break;
 	}
@@ -185,8 +309,7 @@ bool FGpuProfilerAnalyzer::OnEvent(uint16 RouteId, EStyle Style, const FOnEventC
 		UE_LOG(LogTraceServices, VeryVerbose, TEXT("[GPU] EventEndBreadcrumb QueueId=%u, Time=%f, Name=\"%s\""), QueueId, Time, *Name);
 
 		FQueue& Queue = GetOrAddQueue(QueueId);
-		const uint32 TimerId = GetOrAddTimerId(Name);
-		EndEvent(Queue, 0, Time, TimerId);
+		EndEvent(Queue, 0, Time, (uint32)- 1);
 		break;
 	}
 	case RouteId_EventBeginWork:
@@ -537,7 +660,7 @@ void FGpuProfilerAnalyzer::EndEvent(FQueue& Queue, int32 StackIndex, double EndE
 		FAnalysisSessionReadScope _(Session);
 		BeginEventTimerId = TimingProfilerProvider.GetOriginalTimerIdFromMetadata(BeginEvent.TimerId);
 	}
-	if (BeginEventTimerId != EndEventTimerId)
+	if (EndEventTimerId != (uint32)-1 && BeginEventTimerId != EndEventTimerId)
 	{
 		++ErrorData.NumMismatchedEvents;
 
diff --git a/Engine/Source/Developer/TraceServices/Private/Analyzers/GpuProfilerTraceAnalysis.h b/Engine/Source/Developer/TraceServices/Private/Analyzers/GpuProfilerTraceAnalysis.h
--- a/Engine/Source/Developer/TraceServices/Private/Analyzers/GpuProfilerTraceAnalysis.h
+++ b/Engine/Source/Developer/TraceServices/Private/Analyzers/GpuProfilerTraceAnalysis.h
@@ -21,6 +21,7 @@ private:
 		RouteId_Init,
 		RouteId_QueueSpec,
 		RouteId_EventFrameBoundary,
+		RouteId_EventBreadcrumbSpec,
 		RouteId_EventBeginBreadcrumb,
 		RouteId_EventEndBreadcrumb,
 		RouteId_EventBeginWork,
@@ -95,6 +96,7 @@ private:
 	uint32 Version = 0;
 	uint32 GpuWorkTimerId = ~0;
 	uint32 GpuWaitTimerId = ~0;
+	TMap<uint32, uint32> BreadcrumbSpecMap; // breadcrumb spec id --> GPU timer id
 	TMap<FString, uint32> BreadcrumbMap; // breadcrumb name --> GPU timer id
 	TMap<uint32, const TCHAR*> TimerMap; // GPU timer id --> persistent timer name
 	TMap<uint32, FQueue> Queues; // QueueId --> FQueue
diff --git a/Engine/Source/Developer/TraceServices/Private/Model/TimingProfiler.cpp b/Engine/Source/Developer/TraceServices/Private/Model/TimingProfiler.cpp
--- a/Engine/Source/Developer/TraceServices/Private/Model/TimingProfiler.cpp
+++ b/Engine/Source/Developer/TraceServices/Private/Model/TimingProfiler.cpp
@@ -87,6 +87,14 @@ void FTimingProfilerProvider::SetTimerNameAndLocation(uint32 TimerId, FStringVie
 	Timer.Line = Line;
 }
 
+void FTimingProfilerProvider::SetMetadataSpec(uint32 TimerId, uint32 MetadataSpecId)
+{
+	Session.WriteAccessCheck();
+
+	FTimingProfilerTimer& Timer = Timers[TimerId];
+	Timer.MetadataSpecId = MetadataSpecId;
+}
+
 uint32 FTimingProfilerProvider::AddMetadata(uint32 OriginalTimerId, TArray<uint8>&& Metadata)
 {
 	Session.WriteAccessCheck();
@@ -154,6 +162,28 @@ TArrayView<uint8> FTimingProfilerProvider::GetEditableMetadata(uint32 MetadataTi
 	return Metadata.Payload;
 }
 
+uint32 FTimingProfilerProvider::AddMetadataSpec(FMetadataSpec&& Metadata)
+{
+	Session.WriteAccessCheck();
+
+	uint32 MetadataSpecId = MetadataSpecs.Num();
+	MetadataSpecs.Add(MoveTemp(Metadata));
+
+	return MetadataSpecId;
+}
+
+const FMetadataSpec* FTimingProfilerProvider::GetMetadataSpec(uint32 MetadataSpecId) const
+{
+	Session.ReadAccessCheck();
+
+	if (MetadataSpecId < (uint32)MetadataSpecs.Num())
+	{
+		return &MetadataSpecs[MetadataSpecId];
+	}
+
+	return nullptr;
+}
+
 IEditableTimeline<FTimingProfilerEvent>& FTimingProfilerProvider::GetCpuThreadEditableTimeline(uint32 ThreadId)
 {
 	Session.WriteAccessCheck();
diff --git a/Engine/Source/Developer/TraceServices/Private/Model/TimingProfilerPrivate.h b/Engine/Source/Developer/TraceServices/Private/Model/TimingProfilerPrivate.h
--- a/Engine/Source/Developer/TraceServices/Private/Model/TimingProfilerPrivate.h
+++ b/Engine/Source/Developer/TraceServices/Private/Model/TimingProfilerPrivate.h
@@ -51,6 +51,7 @@ public:
 	virtual void ReadTimers(TFunctionRef<void(const ITimingProfilerTimerReader&)> Callback) const override;
 	virtual uint32 GetOriginalTimerIdFromMetadata(uint32 MetadataTimerId) const override;
 	virtual TArrayView<const uint8> GetMetadata(uint32 MetadataTimerId) const override;
+	virtual const FMetadataSpec* GetMetadataSpec(uint32 MetadataSpecId) const;
 	virtual ITable<FTimingProfilerAggregatedStats>* CreateAggregation(const FCreateAggregationParams& Params) const override;
 	virtual ITimingProfilerButterfly* CreateButterfly(const FCreateButterflyParams& Params) const override;
 
@@ -68,6 +69,9 @@ public:
 	virtual uint32 AddMetadata(uint32 OriginalTimerId, TArray<uint8>&& Metadata) override;
 	virtual TArrayView<uint8> GetEditableMetadata(uint32 MetadataTimerId) override;
 
+	virtual uint32 AddMetadataSpec(FMetadataSpec&& Metadata) override;
+	virtual void SetMetadataSpec(uint32 TimerId, uint32 MetadataSpecId) override;
+
 	virtual void AddGpuQueue(uint32 QueueId, uint8 GPU, uint8 Index, uint8 Type, const TCHAR* Name) override;
 	virtual void AddGpuSignalFence(uint32 QueueId, const FGpuSignalFence& SignalFence) override;
 	virtual void AddGpuWaitFence(uint32 QueueId, const FGpuWaitFence& WaitFence) override;
@@ -89,6 +93,7 @@ private:
 
 	IAnalysisSession& Session;
 	TArray<FMetadata> Metadatas;
+	TArray<FMetadataSpec> MetadataSpecs;
 	TArray<FTimingProfilerTimer> Timers;
 	TArray<TSharedRef<TimelineInternal>> Timelines;
 	TMap<uint32, uint32> CpuThreadTimelineIndexMap; // CPU Thread Id --> timeline index, in Timelines
diff --git a/Engine/Source/Developer/TraceServices/Public/TraceServices/Model/TimingProfiler.h b/Engine/Source/Developer/TraceServices/Public/TraceServices/Model/TimingProfiler.h
--- a/Engine/Source/Developer/TraceServices/Public/TraceServices/Model/TimingProfiler.h
+++ b/Engine/Source/Developer/TraceServices/Public/TraceServices/Model/TimingProfiler.h
@@ -19,10 +19,20 @@ namespace TraceServices
 template <typename InEventType> class IEditableTimeline;
 template <typename RowType> class ITable;
 
+struct FMetadataSpec
+{
+	const TCHAR* Format = nullptr;
+	TArray<const TCHAR*> FieldNames;
+
+	static const uint32 InvalidMetadataSpecId = (uint32)-1;
+};
+
 struct FTimingProfilerTimer
 {
 	const TCHAR* Name = nullptr;
 	const TCHAR* File = nullptr;
+	uint32 MetadataSpecId = FMetadataSpec::InvalidMetadataSpecId;
+
 	uint32 Id = 0;
 	union
 	{
@@ -33,6 +43,11 @@ struct FTimingProfilerTimer
 		};
 		uint32 LineAndFlags = 0; // used only to default initialize Line and IsGpuTimer with 0
 	};
+
+	bool HasValidMetadataSpecId() const
+	{
+		return MetadataSpecId != FMetadataSpec::InvalidMetadataSpecId;
+	}
 };
 
 struct FTimingProfilerEvent
@@ -226,6 +241,14 @@ public:
 	virtual uint32 GetOriginalTimerIdFromMetadata(uint32 MetadataTimerId) const { return MetadataTimerId; }
 	virtual TArrayView<const uint8> GetMetadata(uint32 MetadataTimerId) const { return TArrayView<const uint8>(); }
 
+	/**
+	* Get the metadata spec associated with the MetadataSpecId.
+	*
+	* @param MetadataSpecId The Metadata spec id.
+	* @return A pointer to the FMetadataSpec or nullptr if no spec is associated with the provided Id. The pointer is only valid in the same Session Read Scope as the function call.
+	*/
+	virtual const FMetadataSpec* GetMetadataSpec(uint32 MetadataSpecId) const { return nullptr; }
+
 	/**
 	* Creates a table of aggregated stats.
 	*
@@ -303,6 +326,14 @@ public:
 	*/
 	virtual void SetTimerNameAndLocation(uint32 TimerId, FStringView Name, const TCHAR* File, uint32 Line) = 0;
 
+	/*
+	* Sets the metadata spec for an existing timer.
+	*
+	* @param TimerId		The identity of the timer to update.
+	* @param MetadataSpecId	The metadata spec to associate with the timer.
+	*/
+	virtual void SetMetadataSpec(uint32 TimerId, uint32 MetadataSpecId) {}
+
 	/*
 	* Adds metadata to a CPU or GPU timer.
 	*
@@ -322,6 +353,15 @@ public:
 	*/
 	virtual TArrayView<uint8> GetEditableMetadata(uint32 MetadataTimerId) = 0;
 
+	/*
+	* Adds a metadata spec to storage.
+	*
+	* @param Metadata	The metadata spec to add.
+	*
+	* @return			The identity of the metadata spec.
+	*/
+	virtual uint32 AddMetadataSpec(FMetadataSpec&& Metadata) { return 0; }
+
 	/*
 	* Adds a new GPU queue.
 	*
diff --git a/Engine/Source/Runtime/RHI/Private/GPUProfiler.cpp b/Engine/Source/Runtime/RHI/Private/GPUProfiler.cpp
--- a/Engine/Source/Runtime/RHI/Private/GPUProfiler.cpp
+++ b/Engine/Source/Runtime/RHI/Private/GPUProfiler.cpp
@@ -1970,14 +1970,14 @@ namespace UE::RHI::GPUProfiler
 				}
 
 				Breadcrumb = Event.Breadcrumb;
-
-				FRHIBreadcrumb::FBuffer Buffer;
-				const TCHAR* Name = Event.Breadcrumb->GetTCHAR(Buffer);
-				FGpuProfilerTrace::BeginBreadcrumb(Queue.Value, Event.GPUTimestampTOP, Name);
+				Breadcrumb->TraceBeginGPU(Queue.Value, Event.GPUTimestampTOP);
 
 			#if WITH_PROFILEGPU
 				if (Profile.bProfileFrame)
 				{
+					FRHIBreadcrumb::FBuffer Buffer;
+					const TCHAR* Name = Event.Breadcrumb->GetTCHAR(Buffer);
+
 					// Push a new node
 					Profile.PushNode(Name);
 				}
@@ -2006,12 +2006,11 @@ namespace UE::RHI::GPUProfiler
 					check(ActiveStatsStack.Last() == Stat);
 					ActiveStatsStack.RemoveAt(ActiveStatsStack.Num() - 1, EAllowShrinking::No);
 				}
+				
+				Breadcrumb->TraceEndGPU(Queue.Value, Event.GPUTimestampBOP);
 
 				Breadcrumb = Event.Breadcrumb->GetParent();
 
-				FRHIBreadcrumb::FBuffer Buffer;
-				const TCHAR* Name = Event.Breadcrumb->GetTCHAR(Buffer);
-				FGpuProfilerTrace::EndBreadcrumb(Queue.Value, Event.GPUTimestampBOP, Name);
 
 			#if WITH_PROFILEGPU
 				if (Profile.bProfileFrame)
diff --git a/Engine/Source/Runtime/RHI/Private/GpuProfilerTrace.cpp b/Engine/Source/Runtime/RHI/Private/GpuProfilerTrace.cpp
--- a/Engine/Source/Runtime/RHI/Private/GpuProfilerTrace.cpp
+++ b/Engine/Source/Runtime/RHI/Private/GpuProfilerTrace.cpp
@@ -1,8 +1,11 @@
 // Copyright Epic Games, Inc. All Rights Reserved.
 
 #include "GpuProfilerTrace.h"
+
+#include "CborWriter.h"
 #include "GPUProfiler.h"
 #include "HAL/IConsoleManager.h"
+#include "Serialization/MemoryWriter.h"
 #include "ProfilingDebugging/MiscTrace.h"
 #include "Trace/Trace.inl"
 #include "RHI.h"
@@ -263,11 +266,10 @@ void FGpuProfilerTrace::Deinitialize()
 
 #endif // GPUPROFILERTRACE_ENABLED
 
-// Tracing for the new GPU Profiler
-#if UE_TRACE_GPU_PROFILER_ENABLED
-
 namespace UE::RHI::GPUProfiler
 {
+// Tracing for the new GPU Profiler
+#if UE_TRACE_GPU_PROFILER_ENABLED
 
 UE_TRACE_EVENT_BEGIN(GpuProfiler, Init, NoSync | Important)
 	UE_TRACE_EVENT_FIELD(uint8, Version)
@@ -284,19 +286,24 @@ UE_TRACE_EVENT_BEGIN(GpuProfiler, EventFrameBoundary)
 UE_TRACE_EVENT_END()
 
 #if WITH_RHI_BREADCRUMBS
+
+UE_TRACE_EVENT_BEGIN(GpuProfiler, EventBreadcrumbSpec)
+	UE_TRACE_EVENT_FIELD(uint32, SpecId)
+	UE_TRACE_EVENT_FIELD(UE::Trace::WideString, StaticName)
+	UE_TRACE_EVENT_FIELD(UE::Trace::WideString, NameFormat)
+	UE_TRACE_EVENT_FIELD(uint8[], FieldNames)
+UE_TRACE_EVENT_END()
+
 UE_TRACE_EVENT_BEGIN(GpuProfiler, EventBeginBreadcrumb)
+	UE_TRACE_EVENT_FIELD(uint32, SpecId)
 	UE_TRACE_EVENT_FIELD(uint32, QueueId)
 	UE_TRACE_EVENT_FIELD(uint64, GPUTimestampTOP)
-	UE_TRACE_EVENT_FIELD(UE::Trace::WideString, Name)
-	//TODO: long name format + serialized args
-	//UE_TRACE_EVENT_FIELD(UE::Trace::WideString, LongNameFormat)
-	//UE_TRACE_EVENT_FIELD(uint8[], LongNameArgs)
+	UE_TRACE_EVENT_FIELD(uint8[], Metadata)
 UE_TRACE_EVENT_END()
 
 UE_TRACE_EVENT_BEGIN(GpuProfiler, EventEndBreadcrumb)
 	UE_TRACE_EVENT_FIELD(uint32, QueueId)
 	UE_TRACE_EVENT_FIELD(uint64, GPUTimestampBOP)
-	UE_TRACE_EVENT_FIELD(UE::Trace::WideString, Name) // for debugging
 UE_TRACE_EVENT_END()
 #endif // WITH_RHI_BREADCRUMBS
 
@@ -372,26 +379,6 @@ void FGpuProfilerTrace::FrameBoundary(uint32 QueueId, uint32 FrameId)
 		<< EventFrameBoundary.FrameNumber(FrameId);
 }
 
-void FGpuProfilerTrace::BeginBreadcrumb(uint32 QueueId, uint64 GPUTimestampTOP, const TCHAR* Name)
-{
-#if WITH_RHI_BREADCRUMBS
-	UE_TRACE_LOG(GpuProfiler, EventBeginBreadcrumb, GpuChannel)
-		<< EventBeginBreadcrumb.QueueId(QueueId)
-		<< EventBeginBreadcrumb.GPUTimestampTOP(GPUTimestampTOP)
-		<< EventBeginBreadcrumb.Name(Name);
-#endif
-}
-
-void FGpuProfilerTrace::EndBreadcrumb(uint32 QueueId, uint64 GPUTimestampTOP, const TCHAR* Name)
-{
-#if WITH_RHI_BREADCRUMBS
-	UE_TRACE_LOG(GpuProfiler, EventEndBreadcrumb, GpuChannel)
-		<< EventEndBreadcrumb.QueueId(QueueId)
-		<< EventEndBreadcrumb.GPUTimestampBOP(GPUTimestampTOP)
-		<< EventEndBreadcrumb.Name(Name); // for debugging
-#endif
-}
-
 void FGpuProfilerTrace::BeginWork(uint32 QueueId, uint64 GPUTimestampTOP, uint64 CPUTimestamp)
 {
 	UE_TRACE_LOG(GpuProfiler, EventBeginWork, GpuChannel)
@@ -440,6 +427,113 @@ void FGpuProfilerTrace::WaitFence(uint32 QueueId, uint64 ResolvedTimestamp, uint
 		<< WaitFence.Value(Value);
 }
 
-} // namespace UE::RHI::GPUProfiler
+std::atomic<uint32> FGpuProfilerTrace::NextSpecId = 1;
+TMap<const void*, uint32> FGpuProfilerTrace::DynamicNamesMap;
+FCriticalSection FGpuProfilerTrace::Cs;
+
+uint32 FGpuProfilerTrace::InternalBreadcrumbSpec(const TCHAR* StaticName, const TCHAR* NameFormat, const TArray<uint8>& FieldNames)
+{
+	if (!FGpuProfilerTrace::IsAvailable())
+	{
+		return 0;
+	}
+
+	uint32 SpecId = NextSpecId.fetch_add(1);
+
+#if WITH_RHI_BREADCRUMBS
+	UE_TRACE_LOG(GpuProfiler, EventBreadcrumbSpec, GpuChannel)
+		<< EventBreadcrumbSpec.SpecId(SpecId)
+		<< EventBreadcrumbSpec.StaticName(StaticName)
+		<< EventBreadcrumbSpec.NameFormat(NameFormat)
+		<< EventBreadcrumbSpec.FieldNames(FieldNames.GetData(), FieldNames.Num());
+#endif
+	return SpecId;
+}
+
+void FGpuProfilerTrace::BeginBreadcrumb(uint32 SpecId, uint32 QueueId, uint64 GPUTimestampTOP, const TArray<uint8>& CborData)
+{
+#if WITH_RHI_BREADCRUMBS
+	UE_TRACE_LOG(GpuProfiler, EventBeginBreadcrumb, GpuChannel)
+		<< EventBeginBreadcrumb.SpecId(SpecId)
+		<< EventBeginBreadcrumb.QueueId(QueueId)
+		<< EventBeginBreadcrumb.GPUTimestampTOP(GPUTimestampTOP)
+		<< EventBeginBreadcrumb.Metadata(CborData.GetData(), CborData.Num());
+#endif
+}
+
+void FGpuProfilerTrace::EndBreadcrumb(uint32 QueueId, uint64 GPUTimestampTOP)
+{
+#if WITH_RHI_BREADCRUMBS
+	UE_TRACE_LOG(GpuProfiler, EventEndBreadcrumb, GpuChannel)
+		<< EventEndBreadcrumb.QueueId(QueueId)
+		<< EventEndBreadcrumb.GPUTimestampBOP(GPUTimestampTOP);
+#endif
+}
 
 #endif // UE_TRACE_GPU_PROFILER_ENABLED
+
+FMetadataSerializer::FMetadataSerializer()
+{
+	CborData.Reserve(128);
+	MemoryWriter = new FMemoryWriter(CborData, false, true);
+	CborWriter = new FCborWriter(MemoryWriter, ECborEndianness::StandardCompliant);
+}
+
+FMetadataSerializer::~FMetadataSerializer()
+{
+	delete CborWriter;
+	delete MemoryWriter;
+}
+
+void FMetadataSerializer::AppendValue(const ANSICHAR* Value)
+{
+	CborWriter->WriteValue(Value, FCStringAnsi::Strlen(Value));
+}
+
+void FMetadataSerializer::AppendValue(const WIDECHAR* Value)
+{
+	CborWriter->WriteValue(FWideStringView(Value));
+}
+
+void FMetadataSerializer::AppendValue(const UTF8CHAR* Value)
+{
+	CborWriter->WriteValue((FUtf8StringView)Value);
+}
+
+void FMetadataSerializer::AppendValue(uint64 Value)
+{
+	CborWriter->WriteValue(Value);
+}
+
+void FMetadataSerializer::AppendValue(int64 Value)
+{
+	CborWriter->WriteValue(Value);
+}
+
+void FMetadataSerializer::AppendValue(double Value)
+{
+	CborWriter->WriteValue(Value);
+}
+
+void FMetadataSerializer::AppendValue(bool Value)
+{
+	CborWriter->WriteValue(Value);
+}
+
+void FMetadataSerializer::AppendValue(const FName& Value)
+{
+	CborWriter->WriteValue(Value.ToString());
+}
+
+void FMetadataSerializer::AppendValue(const FDebugName& Value)
+{
+	CborWriter->WriteValue(Value.ToString());
+}
+
+void FMetadataSerializer::AppendValue(const FString& Value)
+{
+	CborWriter->WriteValue(Value);
+}
+
+} // namespace UE::RHI::GPUProfiler
+
diff --git a/Engine/Source/Runtime/RHI/Public/GpuProfilerTrace.h b/Engine/Source/Runtime/RHI/Public/GpuProfilerTrace.h
--- a/Engine/Source/Runtime/RHI/Public/GpuProfilerTrace.h
+++ b/Engine/Source/Runtime/RHI/Public/GpuProfilerTrace.h
@@ -3,7 +3,10 @@
 #pragma once
 
 #include "CoreTypes.h"
+#include "Containers/StringFwd.h"
+#include "RHIDefinitions.h"
 #include "Trace/Config.h"
+#include "UObject/NameTypes.h"
 
 #if !defined(GPUPROFILERTRACE_ENABLED)
 #if UE_TRACE_ENABLED && !UE_BUILD_SHIPPING
@@ -66,12 +69,79 @@ struct FGpuProfilerTrace
 #define TRACE_GPUPROFILER_END_FRAME(...)                 UE_DEPRECATED_MACRO(5.6, "TRACE_GPUPROFILER_END_FRAME has been deprecated and it is non functional. Use standard RHI breadcrumb events instead.")
 #define TRACE_GPUPROFILER_DEINITIALIZE(...)              UE_DEPRECATED_MACRO(5.6, "TRACE_GPUPROFILER_DEINITIALIZE has been deprecated and it is non functional. Use standard RHI breadcrumb events instead.")
 
+class FCborWriter;
+class FDebugName;
+class FMemoryWriter;
+
 namespace UE::RHI::GPUProfiler
 {
+class FMetadataSerializer
+{
+public:
+	RHI_API FMetadataSerializer();
+	RHI_API ~FMetadataSerializer();
+
+	RHI_API void AppendValue(const ANSICHAR* Value);
+	RHI_API void AppendValue(const WIDECHAR* Value);
+	RHI_API void AppendValue(const UTF8CHAR* Value);
+	RHI_API void AppendValue(uint64 Value);
+	RHI_API void AppendValue(int64 Value);
+	RHI_API void AppendValue(bool Value);
+	RHI_API void AppendValue(const FName& Value);
+	RHI_API void AppendValue(const FString& Value);
+	RHI_API void AppendValue(const FDebugName& Value);
+	RHI_API void AppendValue(double Value);
+
+	RHI_API void AppendValue(int32 Value)
+	{
+		AppendValue((int64)Value);
+	}
+	RHI_API void AppendValue(int16 Value)
+	{
+		AppendValue((int64)Value);
+	}
+	RHI_API void AppendValue(int8 Value)
+	{
+		AppendValue((int64)Value);
+	}
+	RHI_API void AppendValue(uint32 Value)
+	{
+		AppendValue((uint64)Value);
+	}
+	RHI_API void AppendValue(uint16 Value)
+	{
+		AppendValue((uint64)Value);
+	}
+	RHI_API void AppendValue(uint8 Value)
+	{
+		AppendValue((uint64)Value);
+	}
+
+	RHI_API void AppendValue(float Value)
+	{
+		AppendValue((double)Value);
+	}
+
+	RHI_API const TArray<uint8>& GetData() const { return CborData; }
+
+private:
+	FCborWriter* CborWriter;
+	FMemoryWriter* MemoryWriter;
+	TArray<uint8> CborData;
+};
 
 struct FGpuProfilerTrace
 {
 #if UE_TRACE_GPU_PROFILER_ENABLED
+private:
+	
+	RHI_API static uint32 InternalBreadcrumbSpec(const TCHAR* StaticName, const TCHAR* NameFormat, const TArray<uint8>& FieldNames);
+
+	static std::atomic<uint32> NextSpecId;
+	static TMap<const void*, uint32> DynamicNamesMap;
+	static FCriticalSection Cs;
+
+public:
 
 	/**
 	 * Returns true if the Gpu channel is enabled.
@@ -94,14 +164,26 @@ struct FGpuProfilerTrace
 	RHI_API static void FrameBoundary(uint32 QueueId, uint32 FrameId);
 
 	/**
-	 * Trace the start of a breadcrump on a Gpu Queue. 
+	 * Traces a breadcrumb spec and returns an id that can be used in BeginBreadcrumb. 
 	 */
-	RHI_API static void BeginBreadcrumb(uint32 QueueId, uint64 GPUTimestampTOP, const TCHAR* Name);
+	template<size_t Size>
+	static uint32 BreadcrumbSpec(const TCHAR* StaticName, const TCHAR* NameFormat, const std::array<const TCHAR*, Size>& FieldNames)
+	{
+			FMetadataSerializer Serializer;
+			for (const TCHAR* FieldName : FieldNames)
+			{
+				Serializer.AppendValue(FieldName);
+			}
+
+			return InternalBreadcrumbSpec(StaticName, NameFormat, Serializer.GetData());
+	}
+
+	RHI_API static void BeginBreadcrumb(uint32 SpecId, uint32 QueueId, uint64 GPUTimestampTOP, const TArray<uint8>& CborData);
 
 	/**
-	 * Trace the end of a breadcrump on a Gpu Queue. 
+	 * Trace the end of a breadcrumb on a Gpu Queue. 
 	 */
-	RHI_API static void EndBreadcrumb(uint32 QueueId, uint64 GPUTimestampBOP, const TCHAR* Name);
+	RHI_API static void EndBreadcrumb(uint32 QueueId, uint64 GPUTimestampBOP);
 	
 	/**
 	 * Trace the start of a work event on a Gpu Queue. 
@@ -138,14 +220,17 @@ struct FGpuProfilerTrace
 	RHI_API static void Initialize() {}
 	RHI_API static void InitializeQueue(uint32 QueueId, const TCHAR* Name) {}
 	RHI_API static void FrameBoundary(uint32 QueueId, uint32 FrameId) {}
-	RHI_API static void BeginBreadcrumb(uint32 QueueId, uint64 GPUTimestampTOP, const TCHAR* Name) {}
-	RHI_API static void EndBreadcrumb(uint32 QueueId, uint64 GPUTimestampBOP, const TCHAR* Name) {}
+	RHI_API static void BeginBreadcrumb(uint32 SpecId, uint32 QueueId, uint64 GPUTimestampTOP, const TArray<uint8>& CborData) {}
+	RHI_API static void EndBreadcrumb(uint32 QueueId, uint64 GPUTimestampBOP) {}
 	RHI_API static void BeginWork(uint32 QueueId, uint64 GPUTimestampTOP, uint64 CPUTimestamp) {}
 	RHI_API static void EndWork(uint32 QueueId, uint64 GPUTimestampBOP) {}
 	RHI_API static void TraceWait(uint32 QueueId, uint64 StartTime, uint64 EndTime) {}
 	RHI_API static void Stats(uint32 QueueId, uint32 NumDraws, uint32 NumPrimitives) {}
 	RHI_API static void SignalFence(uint32 QueueId, uint64 ResolvedTimestamp, uint64 Value) {}
 	RHI_API static void WaitFence(uint32 QueueId, uint64 ResolvedTimestamp, uint32 QueueToWaitForId, uint64 Value) {}
+	
+	template<size_t Size>
+	static uint32 BreadcrumbSpec(const TCHAR* StaticName, const TCHAR* NameFormat, const std::array<const TCHAR*, Size>& FieldNames) { return 0; }
 #endif
 };
 
diff --git a/Engine/Source/Runtime/RHI/Public/RHIBreadcrumbs.h b/Engine/Source/Runtime/RHI/Public/RHIBreadcrumbs.h
--- a/Engine/Source/Runtime/RHI/Public/RHIBreadcrumbs.h
+++ b/Engine/Source/Runtime/RHI/Public/RHIBreadcrumbs.h
@@ -9,6 +9,7 @@
 #include "ProfilingDebugging/CsvProfiler.h"
 #include "GenericPlatform/GenericPlatformCrashContext.h"
 
+#include "GpuProfilerTrace.h"
 #include "RHIFwd.h"
 #include "RHIPipeline.h"
 #include "MultiGPU.h"
@@ -248,6 +249,9 @@ namespace UE::RHI::GPUProfiler
 		FRHIBreadcrumbNode* GetParent() const { return Parent; }
 		inline void SetParent(FRHIBreadcrumbNode* Node);
 
+		virtual void TraceBeginGPU(uint32 QueueId, uint64 GPUTimestampTOP) const = 0;
+		virtual void TraceEndGPU(uint32 QueueId, uint64 GPUTimestampTOP) const = 0;
+
 		inline void BeginCPU() const;
 		inline void EndCPU() const;
 
@@ -716,6 +720,11 @@ namespace UE::RHI::GPUProfiler
 					: Inner(Value.Value)
 				{}
 			};
+
+			void Serialize(UE::RHI::GPUProfiler::FMetadataSerializer& Serializer) const
+			{
+				Serializer.AppendValue(Value);
+			}
 		};
 
 		// Disallow all pointer types. Provide a specific assert message for TCHAR pointers.
@@ -749,6 +758,8 @@ namespace UE::RHI::GPUProfiler
 					: Inner(0)
 				{}
 			};
+
+			void Serialize(UE::RHI::GPUProfiler::FMetadataSerializer& Serializer) const {}
 		};
 
 		// String literal - keep the string pointer
@@ -770,6 +781,11 @@ namespace UE::RHI::GPUProfiler
 					: Inner(Value.Value)
 				{}
 			};
+
+			void Serialize(UE::RHI::GPUProfiler::FMetadataSerializer& Serializer) const
+			{
+				Serializer.AppendValue(Value);
+			}
 		};
 
 		// FName - keep the FName itself and defer resolving
@@ -791,6 +807,11 @@ namespace UE::RHI::GPUProfiler
 					Name.Value.ToStringTruncate(Inner);
 				}
 			};
+
+			void Serialize(UE::RHI::GPUProfiler::FMetadataSerializer& Serializer) const
+			{
+				Serializer.AppendValue(Value);
+			}
 		};
 
 		// FDebugName - keep the FDebugName itself and defer resolving
@@ -812,6 +833,11 @@ namespace UE::RHI::GPUProfiler
 					Name.Value.ToStringTruncate(Inner);
 				}
 			};
+
+			void Serialize(UE::RHI::GPUProfiler::FMetadataSerializer& Serializer) const
+			{
+				Serializer.AppendValue(Value);
+			}
 		};
 
 		// FString - Take an immediate copy of the string. Total length is limited by fixed buffer size.
@@ -833,6 +859,11 @@ namespace UE::RHI::GPUProfiler
 					: Inner(String.Buffer)
 				{}
 			};
+
+			void Serialize(UE::RHI::GPUProfiler::FMetadataSerializer& Serializer) const
+			{
+				Serializer.AppendValue(Buffer);
+			}
 		};
 
 		// Determines if a vararg type is a string literal / value tuple, as defined by RHI_BREADCRUMB_FIELD.
@@ -1005,6 +1036,9 @@ namespace UE::RHI::GPUProfiler
 			// Array of the string literal field names
 			std::array<TCHAR const*, NumFields> const FieldNames;
 
+			// Id of the spec of the gpu events associated with this breadcrumb.
+			mutable uint32 TraceGpuSpecId = 0;
+
 			template <typename... TInnerValues>
 			TRHIBreadcrumbDesc(FRHIBreadcrumbData&& BaseData, TFormatString FormatString, TInnerValues&&... Values)
 				: FRHIBreadcrumbData(MoveTemp(BaseData))
@@ -1016,6 +1050,38 @@ namespace UE::RHI::GPUProfiler
 			{
 				return TFormatString::ToStringImpl(StaticName, Buffer, Values, TValuesIndexSequence());
 			}
+
+			void SerializeValues(UE::RHI::GPUProfiler::FMetadataSerializer& Serializer, TValuesTuple const& Values) const
+			{
+				std::apply([&](const auto&... Value)
+					{
+						(SerializeValue(Value, Serializer), ...);
+					}, Values);
+			}
+
+			uint32 GetTraceGpuSpec() const
+			{
+				if (TraceGpuSpecId == 0 && UE::RHI::GPUProfiler::FGpuProfilerTrace::IsAvailable())
+				{
+					if constexpr (FormatStringSize > 0)
+					{
+						TraceGpuSpecId = UE::RHI::GPUProfiler::FGpuProfilerTrace::BreadcrumbSpec(StaticName, TFormatString::FormatString, FieldNames);
+					}
+					else
+					{
+						TraceGpuSpecId = UE::RHI::GPUProfiler::FGpuProfilerTrace::BreadcrumbSpec(StaticName, TEXT(""), FieldNames);
+					}
+				}
+
+				return TraceGpuSpecId;
+			}
+
+		private:
+			template<typename T>
+			void SerializeValue(const T& Item, UE::RHI::GPUProfiler::FMetadataSerializer& Serializer) const
+			{
+				Item.Serialize(Serializer);
+			}
 		};
 
 		// Breadcrumb implementation for printf formatted names
@@ -1046,6 +1112,24 @@ namespace UE::RHI::GPUProfiler
 				return Desc.ToString(Buffer, *this);
 			}
 
+			virtual void TraceBeginGPU(uint32 QueueId, uint64 GPUTimestampTOP) const override
+			{
+				if (uint32 SpecId = Desc.GetTraceGpuSpec())
+				{
+					UE::RHI::GPUProfiler::FMetadataSerializer Serializer;
+					Desc.SerializeValues(Serializer, *this);
+					UE::RHI::GPUProfiler::FGpuProfilerTrace::BeginBreadcrumb(SpecId, QueueId, GPUTimestampTOP, Serializer.GetData());
+				}
+			}
+
+			virtual void TraceEndGPU(uint32 QueueId, uint64 GPUTimestampBOP) const override
+			{
+				if (Desc.TraceGpuSpecId)
+				{
+					UE::RHI::GPUProfiler::FGpuProfilerTrace::EndBreadcrumb(QueueId, GPUTimestampBOP);
+				}
+			}
+
 		private:
 			// Constructor for the sentinel value
 			TRHIBreadcrumb(TDesc const& Desc)
diff --git a/Engine/Source/Runtime/RHI/RHI.Build.cs b/Engine/Source/Runtime/RHI/RHI.Build.cs
--- a/Engine/Source/Runtime/RHI/RHI.Build.cs
+++ b/Engine/Source/Runtime/RHI/RHI.Build.cs
@@ -11,6 +11,7 @@ public class RHI : ModuleRules
 		PrivateDependencyModuleNames.Add("Core");
 		PrivateDependencyModuleNames.Add("TraceLog");
 		PrivateDependencyModuleNames.Add("ApplicationCore");
+		PrivateDependencyModuleNames.Add("Cbor");
 
 		// @todo - new gpu profiler. This is experimental.
 		PublicDefinitions.Add("RHI_NEW_GPU_PROFILER=0");
diff --git a/Engine/Source/Runtime/RenderCore/Public/RenderGraphEvent.h b/Engine/Source/Runtime/RenderCore/Public/RenderGraphEvent.h
--- a/Engine/Source/Runtime/RenderCore/Public/RenderGraphEvent.h
+++ b/Engine/Source/Runtime/RenderCore/Public/RenderGraphEvent.h
@@ -382,6 +382,11 @@ private:
 					: Inner(Value.StringPointer)
 				{}
 			};
+
+			void Serialize(UE::RHI::GPUProfiler::FMetadataSerializer& Serializer) const
+			{
+				Serializer.AppendValue(StringPointer);
+			}
 		};
 
 		// Type to force the selection of the below TValue specialization.
@@ -409,6 +414,11 @@ private:
 					: Inner(Value.Value)
 				{}
 			};
+
+			void Serialize(UE::RHI::GPUProfiler::FMetadataSerializer& Serializer) const
+			{
+				Serializer.AppendValue(Value);
+			}
 		};
 	}
 
