Fix static analysis warning

diff --git a/Engine/Source/Developer/TraceInsights/Private/Insights/TimingProfiler/Tracks/ThreadTimingTrack.cpp b/Engine/Source/Developer/TraceInsights/Private/Insights/TimingProfiler/Tracks/ThreadTimingTrack.cpp
--- a/Engine/Source/Developer/TraceInsights/Private/Insights/TimingProfiler/Tracks/ThreadTimingTrack.cpp
+++ b/Engine/Source/Developer/TraceInsights/Private/Insights/TimingProfiler/Tracks/ThreadTimingTrack.cpp
@@ -481,7 +481,7 @@ static void AddTimingEventToBuilder(ITimingEventsTrackDrawStateBuilder& Builder,
 					if (Metadata.Num() > 0)
 					{
 						const TraceServices::FMetadataSpec* MetadataSpec = nullptr;
-						if (Timer && Timer->HasValidMetadataSpecId())
+						if (Timer->HasValidMetadataSpecId())
 						{
 							MetadataSpec = TimingProfilerProvider.GetMetadataSpec(Timer->MetadataSpecId);
 						}
