[Insights] Timing Regions: Support coloring Timing Regions by category and separate Tracks per category

- `FTimingRegionsTrack` can now display a filtered view when a specific Category is set.
- Added checkbox to allow coloring regions by Category instead of by Name
- Added checkbox to split the TimingRegionsTrack into multiple tracks, one per Category. The individual tracks can be toggled via the "All Tracks" menu.

diff --git a/Engine/Source/Developer/TraceInsights/Private/Insights/TimingProfiler/Tracks/RegionsTimingTrack.cpp b/Engine/Source/Developer/TraceInsights/Private/Insights/TimingProfiler/Tracks/RegionsTimingTrack.cpp
--- a/Engine/Source/Developer/TraceInsights/Private/Insights/TimingProfiler/Tracks/RegionsTimingTrack.cpp
+++ b/Engine/Source/Developer/TraceInsights/Private/Insights/TimingProfiler/Tracks/RegionsTimingTrack.cpp
@@ -21,6 +21,7 @@
 #include "Insights/InsightsStyle.h"
 #include "Insights/ITimingViewSession.h"
 #include "Insights/TimingProfilerCommon.h"
+#include "Insights/TimingProfiler/ViewModels/TimingRegionsSharedState.h"
 #include "Insights/ViewModels/TimingTrackViewport.h"
 #include "Insights/ViewModels/TimingViewDrawHelper.h"
 #include "Insights/Widgets/STimingView.h"
@@ -41,6 +42,13 @@ FTimingRegionsTrack::~FTimingRegionsTrack()
 {
 }
 
+
+void FTimingRegionsTrack::SetRegionsCategory(const TCHAR* InRegionsCategory)
+{
+	RegionsCategory = InRegionsCategory;
+	SetName(FString(TEXT("Timing Regions - ")) + InRegionsCategory);
+}
+
 void FTimingRegionsTrack::BuildContextMenu(FMenuBuilder& MenuBuilder)
 {
 	FTimingEventsTrack::BuildContextMenu(MenuBuilder);
@@ -68,6 +76,10 @@ void FTimingRegionsTrack::InitTooltip(FTooltipDrawState& InOutTooltip, const ITi
 			InOutTooltip.AddTitle(InRegion.Text, FLinearColor::White);
 			InOutTooltip.AddNameValueTextLine(TEXT("Duration:"),  FormatTimeAuto(InRegion.EndTime - InRegion.BeginTime));
 			InOutTooltip.AddNameValueTextLine(TEXT("Depth:"),  FString::FromInt(InRegion.Depth));
+			if (InRegion.Category)
+			{
+				InOutTooltip.AddNameValueTextLine(TEXT("Category:"),  InRegion.Category);
+			}
 			InOutTooltip.UpdateLayout();
 		});
 	}
@@ -89,13 +101,17 @@ void FTimingRegionsTrack::BuildDrawState(ITimingEventsTrackDrawStateBuilder& Bui
 
 	// We are counting only non-empty lanes, so we can collapse empty ones in the visualization.
 	int32 CurDepth = 0;
-	RegionProvider.GetDefaultTimeline().EnumerateLanes([this, Viewport, &CurDepth, &Builder](const TraceServices::FRegionLane& Lane, const int32 Depth)
+	const TraceServices::IRegionTimeline* Timeline = RegionProvider.GetTimelineForCategory(RegionsCategory);
+	check(Timeline)
+	
+	Timeline->EnumerateLanes([this, Viewport, &CurDepth, &Builder](const TraceServices::FRegionLane& Lane, const int32 Depth)
 	{
 		bool RegionHadEvents = false;
-		Lane.EnumerateRegions(Viewport.GetStartTime(), Viewport.GetEndTime(), [&Builder, &RegionHadEvents, &CurDepth](const TraceServices::FTimeRegion& Region) -> bool
+		Lane.EnumerateRegions(Viewport.GetStartTime(), Viewport.GetEndTime(), [this, &Builder, &RegionHadEvents, &CurDepth](const TraceServices::FTimeRegion& Region) -> bool
 		{
 			RegionHadEvents = true;
-			Builder.AddEvent(Region.BeginTime, Region.EndTime,CurDepth, Region.Text);
+			uint32 EventColor = SharedState.bColorRegionsByCategory ? FTimingEvent::ComputeEventColor(Region.Category) : FTimingEvent::ComputeEventColor(Region.Text);
+			Builder.AddEvent(Region.BeginTime, Region.EndTime,CurDepth, Region.Text, 0 , EventColor);
 			return true;
 		});
 
@@ -134,7 +150,9 @@ void FTimingRegionsTrack::BuildFilteredDrawState(ITimingEventsTrackDrawStateBuil
 			if (bFilterOnlyByEventType)
 			{
 				int32 CurDepth = 0;
-				RegionProvider.GetDefaultTimeline().EnumerateLanes([this, Viewport, &CurDepth, &Builder, FilterEventType](const TraceServices::FRegionLane& Lane, const int32 Depth)
+				const TraceServices::IRegionTimeline* Timeline = RegionProvider.GetTimelineForCategory(RegionsCategory);
+				check(Timeline)
+				Timeline->EnumerateLanes([this, Viewport, &CurDepth, &Builder, FilterEventType](const TraceServices::FRegionLane& Lane, const int32 Depth)
 					{
 						bool RegionHadEvents = false;
 						Lane.EnumerateRegions(Viewport.GetStartTime(), Viewport.GetEndTime(), [&Builder, &RegionHadEvents, &CurDepth, FilterEventType](const TraceServices::FTimeRegion& Region) -> bool
@@ -181,7 +199,9 @@ void FTimingRegionsTrack::BuildFilteredDrawState(ITimingEventsTrackDrawStateBuil
 			const FTimingTrackViewport& Viewport = Context.GetViewport();
 
 			int32 CurDepth = 0;
-			RegionProvider.GetDefaultTimeline().EnumerateLanes([this, Viewport, &CurDepth, &Builder, &FilterContext](const TraceServices::FRegionLane& Lane, const int32 Depth)
+			const TraceServices::IRegionTimeline* Timeline = RegionProvider.GetTimelineForCategory(RegionsCategory);
+			check(Timeline)
+			Timeline->EnumerateLanes([this, Viewport, &CurDepth, &Builder, &FilterContext](const TraceServices::FRegionLane& Lane, const int32 Depth)
 				{
 					bool RegionHadEvents = false;
 					Lane.EnumerateRegions(Viewport.GetStartTime(), Viewport.GetEndTime(), [&Builder, &RegionHadEvents, &CurDepth, &FilterContext, this](const TraceServices::FTimeRegion& Region) -> bool
@@ -253,7 +273,9 @@ bool FTimingRegionsTrack::FindRegionEvent(const FTimingEventSearchParameters& In
 		const TraceServices::IRegionProvider& RegionProvider = TraceServices::ReadRegionProvider(*Session);
 		TraceServices::FProviderReadScopeLock RegionProviderScopedLock(RegionProvider);
 
-		RegionProvider.GetDefaultTimeline().EnumerateRegions(InContext.GetParameters().StartTime, InContext.GetParameters().EndTime, [&InContext](const TraceServices::FTimeRegion& Region)
+		const TraceServices::IRegionTimeline* Timeline = RegionProvider.GetTimelineForCategory(RegionsCategory);
+		check(Timeline)
+		Timeline->EnumerateRegions(InContext.GetParameters().StartTime, InContext.GetParameters().EndTime, [&InContext](const TraceServices::FTimeRegion& Region)
 		{
 			InContext.Check(Region.BeginTime, Region.EndTime, Region.Depth, Region);
 
diff --git a/Engine/Source/Developer/TraceInsights/Private/Insights/TimingProfiler/Tracks/RegionsTimingTrack.h b/Engine/Source/Developer/TraceInsights/Private/Insights/TimingProfiler/Tracks/RegionsTimingTrack.h
--- a/Engine/Source/Developer/TraceInsights/Private/Insights/TimingProfiler/Tracks/RegionsTimingTrack.h
+++ b/Engine/Source/Developer/TraceInsights/Private/Insights/TimingProfiler/Tracks/RegionsTimingTrack.h
@@ -43,6 +43,7 @@ public:
 	virtual bool HasCustomFilter() const override;
 	virtual void OnClipboardCopyEvent(const ITimingEvent& InSelectedEvent) const override;
 
+	void SetRegionsCategory(const TCHAR* InRegionsCategory);
 protected:
 	bool FindRegionEvent(const FTimingEventSearchParameters& InParameters, TFunctionRef<void(double, double, uint32, const TraceServices::FTimeRegion&)> InFoundPredicate) const;
 
@@ -50,6 +51,7 @@ private:
 	FTimingRegionsSharedState& SharedState;
 
 	TSharedPtr<FFilterConfigurator> FilterConfigurator;
+	const TCHAR* RegionsCategory = nullptr;
 };
 
 ////////////////////////////////////////////////////////////////////////////////////////////////////
diff --git a/Engine/Source/Developer/TraceInsights/Private/Insights/TimingProfiler/ViewModels/TimingRegionsSharedState.cpp b/Engine/Source/Developer/TraceInsights/Private/Insights/TimingProfiler/ViewModels/TimingRegionsSharedState.cpp
--- a/Engine/Source/Developer/TraceInsights/Private/Insights/TimingProfiler/ViewModels/TimingRegionsSharedState.cpp
+++ b/Engine/Source/Developer/TraceInsights/Private/Insights/TimingProfiler/ViewModels/TimingRegionsSharedState.cpp
@@ -10,11 +10,14 @@
 #include "TraceServices/Model/AnalysisSession.h"
 
 // TraceInsights
+#include "Common/ProviderLock.h"
+#include "Insights/InsightsManager.h"
 #include "Insights/Common/InsightsMenuBuilder.h"
 #include "Insights/InsightsStyle.h"
 #include "Insights/ITimingViewSession.h"
 #include "Insights/TimingProfiler/Tracks/RegionsTimingTrack.h"
 #include "Insights/Widgets/STimingView.h"
+#include "TraceServices/Model/Regions.h"
 
 #define LOCTEXT_NAMESPACE "UE::Insights::TimingProfiler::TimingRegions"
 
@@ -34,6 +37,8 @@ public:
 
 public:
 	TSharedPtr<FUICommandInfo> ShowHideTimingRegionsTrack;
+	TSharedPtr<FUICommandInfo> ColorTimingRegionsTrackByCategory;
+	TSharedPtr<FUICommandInfo> CreateRegionTracksByCategory;
 };
 
 ////////////////////////////////////////////////////////////////////////////////////////////////////
@@ -55,9 +60,19 @@ void FTimingRegionsViewCommands::RegisterCommands()
 {
 	UI_COMMAND(ShowHideTimingRegionsTrack,
 		"Timing Regions Track",
-		"Shows/hides the Timing Regions track.",
+		"Shows/hides the Timing Regions track(s).",
 		EUserInterfaceActionType::ToggleButton,
 		FInputChord(EModifierKey::Control, EKeys::R));
+	UI_COMMAND(ColorTimingRegionsTrackByCategory,
+		"Color Regions by Category",
+		"Color Timing Regions by Category instead of by Name.",
+		EUserInterfaceActionType::ToggleButton,
+		FInputChord());
+	UI_COMMAND(CreateRegionTracksByCategory,
+		"Split Timing Regions into individual Tracks per Category",
+		"Creates a Timing Regions track for each category instead of a single combined one.",
+		EUserInterfaceActionType::ToggleButton,
+		FInputChord());
 }
 UE_ENABLE_OPTIMIZATION_SHIP
 
@@ -80,7 +95,8 @@ void FTimingRegionsSharedState::OnBeginSession(Timing::ITimingViewSession& InSes
 		return;
 	}
 
-	TimingRegionsTrack.Reset();
+	AllRegionsTrack.Reset();
+	TimingRegionTracksPerCategory.Reset();
 }
 
 ////////////////////////////////////////////////////////////////////////////////////////////////////
@@ -92,25 +108,88 @@ void FTimingRegionsSharedState::OnEndSession(Timing::ITimingViewSession& InSessi
 		return;
 	}
 
-	TimingRegionsTrack.Reset();
+	AllRegionsTrack.Reset();
+	TimingRegionTracksPerCategory.Reset();
 }
 
 ////////////////////////////////////////////////////////////////////////////////////////////////////
 
-void FTimingRegionsSharedState::Tick(Timing::ITimingViewSession& InSession, const TraceServices::IAnalysisSession& InAnalysisSession)
+void FTimingRegionsSharedState::CreateRequiredTracks()
 {
-	if (&InSession != TimingView)
+	// delete all tracks
+	if (!bShowHideRegionsTrack)
 	{
+		if (AllRegionsTrack)
+		{
+			TimingView->RemoveScrollableTrack(AllRegionsTrack);
+			AllRegionsTrack.Reset();
+		}
+		
+		for (auto& KV : TimingRegionTracksPerCategory)
+		{
+			TimingView->RemoveScrollableTrack(KV.Value);
+		}
+		TimingRegionTracksPerCategory.Reset();
 		return;
 	}
+	
+	if (bCreateRegionTracksByCategory)
+	{   // delete combined track
+		if (AllRegionsTrack)
+		{
+			TimingView->RemoveScrollableTrack(AllRegionsTrack);
+			AllRegionsTrack.Reset();
+		}
+	}
+	else
+	{   // create combined track
+		if (!AllRegionsTrack)
+		{
+			AllRegionsTrack = MakeShared<FTimingRegionsTrack>(*this);
+			AllRegionsTrack->SetOrder(FTimingTrackOrder::First + 100);
+			AllRegionsTrack->SetVisibilityFlag(bShowHideRegionsTrack);
+			TimingView->AddScrollableTrack(AllRegionsTrack);
+		}
+	}
 
-	if (!TimingRegionsTrack.IsValid())
+	if (!bCreateRegionTracksByCategory)
+	{   // delete per category tracks
+		for (auto& KV : TimingRegionTracksPerCategory)
+		{
+			TimingView->RemoveScrollableTrack(KV.Value);
+		}
+		TimingRegionTracksPerCategory.Reset();
+	}
+	else
+	{   // create per category tracks
+		// known categories might change during analysis so always create on demand
+		TSharedPtr<const TraceServices::IAnalysisSession> Session = FInsightsManager::Get()->GetSession();
+		const TraceServices::IRegionProvider& RegionProvider = TraceServices::ReadRegionProvider(*Session);
+		TraceServices::FProviderReadScopeLock RegionProviderScopedLock(RegionProvider);
+
+		RegionProvider.EnumerateTimelinesByCategory([this, &RegionProvider](const TraceServices::IRegionTimeline& Timeline, const TCHAR* Category)
+		{			
+			if (!TimingRegionTracksPerCategory.Contains(Category))
+			{
+				auto NewTrack = MakeShared<FTimingRegionsTrack>(*this);
+				NewTrack->SetRegionsCategory(Category);
+				NewTrack->SetOrder( Category == RegionProvider.GetUncategorizedRegionCategoryName() ? FTimingTrackOrder::First + 100 : FTimingTrackOrder::First + 101);
+				NewTrack->SetVisibilityFlag(true);
+				TimingView->AddScrollableTrack(NewTrack);
+				TimingRegionTracksPerCategory.Add(Category, NewTrack);
+			}
+		});
+	}
+}
+
+void FTimingRegionsSharedState::Tick(Timing::ITimingViewSession& InSession, const TraceServices::IAnalysisSession& InAnalysisSession)
+{
+	if (&InSession != TimingView)
 	{
-		TimingRegionsTrack = MakeShared<FTimingRegionsTrack>(*this);
-		TimingRegionsTrack->SetOrder(FTimingTrackOrder::First + 100);
-		TimingRegionsTrack->SetVisibilityFlag(true);
-		InSession.AddScrollableTrack(TimingRegionsTrack);
+		return;
 	}
+
+	CreateRequiredTracks();
 }
 
 ////////////////////////////////////////////////////////////////////////////////////////////////////
@@ -119,19 +198,37 @@ void FTimingRegionsSharedState::ShowHideRegionsTrack()
 {
 	bShowHideRegionsTrack = !bShowHideRegionsTrack;
 
-	if (TimingRegionsTrack.IsValid())
+	CreateRequiredTracks();
+	
+	if (TimingView)
 	{
-		TimingRegionsTrack->SetVisibilityFlag(bShowHideRegionsTrack);
+		TimingView->HandleTrackVisibilityChanged();
 	}
+}
 
-	if (TimingView)
+void FTimingRegionsSharedState::ToggleColorRegionsByCategory()
+{
+	bColorRegionsByCategory = !bColorRegionsByCategory;
+	// redraw whatever tracks exist right now
+	if (AllRegionsTrack.IsValid())
 	{
-		TimingView->HandleTrackVisibilityChanged();
+		AllRegionsTrack->SetDirtyFlag();
 	}
+	for (auto& KV : TimingRegionTracksPerCategory)
+	{
+		KV.Value->SetDirtyFlag();
+	}
+}
+
+void FTimingRegionsSharedState::ToggleShouldCreateRegionTracksByCategory()
+{
+	bCreateRegionTracksByCategory = !bCreateRegionTracksByCategory;
 
-	if (bShowHideRegionsTrack)
+	CreateRequiredTracks();
+	
+	if (TimingView)
 	{
-		TimingRegionsTrack->SetDirtyFlag();
+		TimingView->HandleTrackVisibilityChanged();
 	}
 }
 
@@ -147,6 +244,8 @@ void FTimingRegionsSharedState::ExtendOtherTracksFilterMenu(Timing::ITimingViewS
 	InOutMenuBuilder.BeginSection("Timing Regions", LOCTEXT("ContextMenu_Section_Regions", "Timing Regions"));
 	{
 		InOutMenuBuilder.AddMenuEntry(FTimingRegionsViewCommands::Get().ShowHideTimingRegionsTrack);
+		InOutMenuBuilder.AddMenuEntry(FTimingRegionsViewCommands::Get().ColorTimingRegionsTrackByCategory);
+		InOutMenuBuilder.AddMenuEntry(FTimingRegionsViewCommands::Get().CreateRegionTracksByCategory);
 	}
 	InOutMenuBuilder.EndSection();
 }
@@ -165,6 +264,18 @@ void FTimingRegionsSharedState::BindCommands()
 		FExecuteAction::CreateSP(this, &FTimingRegionsSharedState::ShowHideRegionsTrack),
 		FCanExecuteAction(),
 		FIsActionChecked::CreateSP(this, &FTimingRegionsSharedState::IsRegionsTrackVisible));
+
+	CommandList->MapAction(
+		FTimingRegionsViewCommands::Get().ColorTimingRegionsTrackByCategory,
+		FExecuteAction::CreateSP(this, &FTimingRegionsSharedState::ToggleColorRegionsByCategory),
+		FCanExecuteAction(),
+		FIsActionChecked::CreateSP(this, &FTimingRegionsSharedState::ShouldColorRegionsByCategory));
+	
+	CommandList->MapAction(
+		FTimingRegionsViewCommands::Get().CreateRegionTracksByCategory,
+		FExecuteAction::CreateSP(this, &FTimingRegionsSharedState::ToggleShouldCreateRegionTracksByCategory),
+		FCanExecuteAction(),
+		FIsActionChecked::CreateSP(this, &FTimingRegionsSharedState::ShouldCreateRegionTracksByCategory));
 }
 
 ////////////////////////////////////////////////////////////////////////////////////////////////////
diff --git a/Engine/Source/Developer/TraceInsights/Private/Insights/TimingProfiler/ViewModels/TimingRegionsSharedState.h b/Engine/Source/Developer/TraceInsights/Private/Insights/TimingProfiler/ViewModels/TimingRegionsSharedState.h
--- a/Engine/Source/Developer/TraceInsights/Private/Insights/TimingProfiler/ViewModels/TimingRegionsSharedState.h
+++ b/Engine/Source/Developer/TraceInsights/Private/Insights/TimingProfiler/ViewModels/TimingRegionsSharedState.h
@@ -44,13 +44,25 @@ public:
 
 	bool IsRegionsTrackVisible() const { return bShowHideRegionsTrack; }
 	void ShowHideRegionsTrack();
+	bool ShouldColorRegionsByCategory() const { return bColorRegionsByCategory; }
+	void ToggleColorRegionsByCategory();
+	bool ShouldCreateRegionTracksByCategory() const { return bCreateRegionTracksByCategory; }
+	void ToggleShouldCreateRegionTracksByCategory();
 
 private:
+	/// creates/destroys AllRegions/per-category tracks to make sure they match with bCreateRegionTracksByCategory
+	void CreateRequiredTracks();
+
 	STimingView* TimingView = nullptr;
 
-	TSharedPtr<FTimingRegionsTrack> TimingRegionsTrack;
+	// unfiltered view
+	TSharedPtr<FTimingRegionsTrack> AllRegionsTrack;
+	// filtered views/tracks per category
+	TMap<const TCHAR*, TSharedPtr<FTimingRegionsTrack>, FDefaultSetAllocator, TStringPointerMapKeyFuncs_DEPRECATED<const TCHAR*, TSharedPtr<FTimingRegionsTrack>>> TimingRegionTracksPerCategory;
 
 	bool bShowHideRegionsTrack = true;
+	bool bColorRegionsByCategory = false;
+	bool bCreateRegionTracksByCategory = true;
 };
 
 ////////////////////////////////////////////////////////////////////////////////////////////////////
