[Insights]
* Memory Tags: Renamed "LLM Tags" panel to "Memory Tags".
* Memory Tags: Moved the graph tracks related buttons to the status bar.
* Memory Tags: Updated the UI layout and style for the toolbar area.
* Memory Tags: Added separate tree item icon for each tag set.

diff --git a/Engine/Source/Developer/TraceInsights/Private/Insights/InsightsStyle.cpp b/Engine/Source/Developer/TraceInsights/Private/Insights/InsightsStyle.cpp
--- a/Engine/Source/Developer/TraceInsights/Private/Insights/InsightsStyle.cpp
+++ b/Engine/Source/Developer/TraceInsights/Private/Insights/InsightsStyle.cpp
@@ -217,10 +217,10 @@ void FInsightsStyle::FStyle::Initialize()
 	Set("Icons.ModulesView", new EDITOR_IMAGE_BRUSH_SVG("Starship/Common/Spreadsheet", Icon16x16));
 	Set("Icons.ModulesView.ToolBar", new EDITOR_IMAGE_BRUSH_SVG("Starship/Common/Spreadsheet", Icon20x20));
 
-	Set("Icons.AddAllMemTagGraphs", new CORE_IMAGE_BRUSH_SVG("Starship/Common/plus-circle", Icon16x16));
-	Set("Icons.RemoveAllMemTagGraphs", new CORE_IMAGE_BRUSH_SVG("Starship/Common/Delete", Icon16x16));
+	Set("Icons.AddMemTagGraphs", new CORE_IMAGE_BRUSH_SVG("Starship/Common/plus", Icon16x16));
+	Set("Icons.RemoveMemTagGraphs", new CORE_IMAGE_BRUSH_SVG("Starship/Common/Delete", Icon16x16));
 
-	Set("Icons.TagSet.Systems", new CORE_IMAGE_BRUSH_SVG("Starship/Common/File", Icon16x16)); 
+	Set("Icons.TagSet.Systems", new CORE_IMAGE_BRUSH_SVG("Starship/Common/File", Icon16x16));
 	Set("Icons.TagSet.Assets", new EDITOR_IMAGE_BRUSH_SVG("Starship/Common/FileLined", Icon16x16));
 	Set("Icons.TagSet.AssetClasses", new EDITOR_IMAGE_BRUSH_SVG("Starship/Common/Class", Icon16x16));
 
@@ -277,8 +277,14 @@ void FInsightsStyle::FStyle::Initialize()
 	Set("Icons.DataTypeDouble.TreeItem", new TODO_IMAGE_BRUSH(Icon12x12));
 	Set("Icons.DataTypeInt64.TreeItem", new TODO_IMAGE_BRUSH(Icon12x12));
 	Set("Icons.NetEvent.TreeItem", new TODO_IMAGE_BRUSH(Icon12x12));
+
 	Set("Icons.MemTag.TreeItem", new IMAGE_BRUSH_SVG("MemTags", Icon12x12));
-	Set("Icons.MemTagWithGraph.TreeItem", new IMAGE_BRUSH_SVG("MemTags", Icon12x12));
+	Set("Icons.SystemMemTag.TreeItem", new IMAGE_BRUSH_SVG("MemTags", Icon12x12));
+	Set("Icons.AssetMemTag.TreeItem", new EDITOR_IMAGE_BRUSH_SVG("Starship/Common/FileLined", Icon12x12));
+	Set("Icons.ClassMemTag.TreeItem", new EDITOR_IMAGE_BRUSH_SVG("Starship/Common/Class", Icon12x12));
+	Set("Icons.UObject.TreeItem", new CORE_IMAGE_BRUSH_SVG("Starship/Common/File", Icon12x12));
+
+	Set("Icons.HasGraph.TreeItem", new CORE_IMAGE_BRUSH_SVG("Starship/Common/Stop", FVector2D(8.0f, 12.0f)));
 
 	//////////////////////////////////////////////////
 	// Trace Control
diff --git a/Engine/Source/Developer/TraceInsights/Private/Insights/MemoryProfiler/MemoryProfilerCommands.cpp b/Engine/Source/Developer/TraceInsights/Private/Insights/MemoryProfiler/MemoryProfilerCommands.cpp
--- a/Engine/Source/Developer/TraceInsights/Private/Insights/MemoryProfiler/MemoryProfilerCommands.cpp
+++ b/Engine/Source/Developer/TraceInsights/Private/Insights/MemoryProfiler/MemoryProfilerCommands.cpp
@@ -73,8 +73,8 @@ void FMemoryProfilerCommands::RegisterCommands()
 		FInputChord());
 
 	UI_COMMAND(ToggleMemTagTreeViewVisibility,
-		"LLM Tags",
-		"Toggles the visibility of the LLM Tags tree view.",
+		"Memory Tags",
+		"Toggles the visibility of the Memory Tags (LLM) tree view.",
 		EUserInterfaceActionType::ToggleButton,
 		FInputChord());
 
diff --git a/Engine/Source/Developer/TraceInsights/Private/Insights/MemoryProfiler/MemoryProfilerCommands.h b/Engine/Source/Developer/TraceInsights/Private/Insights/MemoryProfiler/MemoryProfilerCommands.h
--- a/Engine/Source/Developer/TraceInsights/Private/Insights/MemoryProfiler/MemoryProfilerCommands.h
+++ b/Engine/Source/Developer/TraceInsights/Private/Insights/MemoryProfiler/MemoryProfilerCommands.h
@@ -37,7 +37,7 @@ public:
 	/** Toggles visibility for the Memory Investigation view. Global and custom command. */
 	TSharedPtr<FUICommandInfo> ToggleMemInvestigationViewVisibility;
 
-	/** Toggles visibility for the LLM Tags tree view. Global and custom command. */
+	/** Toggles visibility for the Memory Tags tree view. Global and custom command. */
 	TSharedPtr<FUICommandInfo> ToggleMemTagTreeViewVisibility;
 
 	/** Toggles visibility for the Modules view. Global and custom command. */
diff --git a/Engine/Source/Developer/TraceInsights/Private/Insights/MemoryProfiler/MemoryProfilerManager.h b/Engine/Source/Developer/TraceInsights/Private/Insights/MemoryProfiler/MemoryProfilerManager.h
--- a/Engine/Source/Developer/TraceInsights/Private/Insights/MemoryProfiler/MemoryProfilerManager.h
+++ b/Engine/Source/Developer/TraceInsights/Private/Insights/MemoryProfiler/MemoryProfilerManager.h
@@ -155,7 +155,7 @@ private:
 	/** If the Memory Investigation (Alloc Queries) view is visible or hidden. */
 	bool bIsMemInvestigationViewVisible;
 
-	/** If the LLM Tags tree view is visible or hidden. */
+	/** If the Memory Tags tree view is visible or hidden. */
 	bool bIsMemTagTreeViewVisible;
 
 	/** If the Modules view is visible or hidden. */
diff --git a/Engine/Source/Developer/TraceInsights/Private/Insights/MemoryProfiler/ViewModels/MemTagNode.cpp b/Engine/Source/Developer/TraceInsights/Private/Insights/MemoryProfiler/ViewModels/MemTagNode.cpp
--- a/Engine/Source/Developer/TraceInsights/Private/Insights/MemoryProfiler/ViewModels/MemTagNode.cpp
+++ b/Engine/Source/Developer/TraceInsights/Private/Insights/MemoryProfiler/ViewModels/MemTagNode.cpp
@@ -14,6 +14,9 @@ namespace UE::Insights::MemoryProfiler
 {
 
 INSIGHTS_IMPLEMENT_RTTI(FMemTagNode)
+INSIGHTS_IMPLEMENT_RTTI(FSystemMemTagNode)
+INSIGHTS_IMPLEMENT_RTTI(FAssetMemTagNode)
+INSIGHTS_IMPLEMENT_RTTI(FClassMemTagNode)
 
 ////////////////////////////////////////////////////////////////////////////////////////////////////
 
@@ -161,7 +164,7 @@ const FSlateBrush* FMemTagNode::GetIcon() const
 {
 	if (MemTag && MemTag->IsAddedToGraph())
 	{
-		return FInsightsStyle::GetBrush("Icons.MemTagWithGraph.TreeItem");
+		return FInsightsStyle::GetBrush("Icons.HasGraph.TreeItem");
 	}
 	else
 	{
@@ -179,7 +182,7 @@ FLinearColor FMemTagNode::GetIconColor() const
 	}
 	else
 	{
-		return FLinearColor(0.1f, 0.1f, 0.1f, 1.0f);
+		return GetDefaultColor(IsGroup());
 	}
 }
 
@@ -187,21 +190,56 @@ FLinearColor FMemTagNode::GetIconColor() const
 
 FLinearColor FMemTagNode::GetColor() const
 {
-	if (MemTag && MemTag->IsAddedToGraph())
+	return GetDefaultColor(IsGroup());
+}
+
+////////////////////////////////////////////////////////////////////////////////////////////////////
+
+void FMemTagNode::ResetAggregatedStats()
+{
+	FMemory::Memzero<FMemTagStats>(Stats);
+}
+
+////////////////////////////////////////////////////////////////////////////////////////////////////
+
+const FSlateBrush* FSystemMemTagNode::GetIcon() const
+{
+	if (GetMemTag() && GetMemTag()->IsAddedToGraph())
 	{
-		return MemTag->GetColor();
+		return FInsightsStyle::GetBrush("Icons.HasGraph.TreeItem");
 	}
 	else
 	{
-		return GetDefaultColor(IsGroup());
+		return FInsightsStyle::GetBrush("Icons.SystemMemTag.TreeItem");
 	}
 }
 
 ////////////////////////////////////////////////////////////////////////////////////////////////////
 
-void FMemTagNode::ResetAggregatedStats()
+const FSlateBrush* FAssetMemTagNode::GetIcon() const
 {
-	FMemory::Memzero<FMemTagStats>(Stats);
+	if (GetMemTag() && GetMemTag()->IsAddedToGraph())
+	{
+		return FInsightsStyle::GetBrush("Icons.HasGraph.TreeItem");
+	}
+	else
+	{
+		return FInsightsStyle::GetBrush("Icons.AssetMemTag.TreeItem");
+	}
+}
+
+////////////////////////////////////////////////////////////////////////////////////////////////////
+
+const FSlateBrush* FClassMemTagNode::GetIcon() const
+{
+	if (GetMemTag() && GetMemTag()->IsAddedToGraph())
+	{
+		return FInsightsStyle::GetBrush("Icons.HasGraph.TreeItem");
+	}
+	else
+	{
+		return FInsightsStyle::GetBrush("Icons.ClassMemTag.TreeItem");
+	}
 }
 
 ////////////////////////////////////////////////////////////////////////////////////////////////////
diff --git a/Engine/Source/Developer/TraceInsights/Private/Insights/MemoryProfiler/ViewModels/MemTagNode.h b/Engine/Source/Developer/TraceInsights/Private/Insights/MemoryProfiler/ViewModels/MemTagNode.h
--- a/Engine/Source/Developer/TraceInsights/Private/Insights/MemoryProfiler/ViewModels/MemTagNode.h
+++ b/Engine/Source/Developer/TraceInsights/Private/Insights/MemoryProfiler/ViewModels/MemTagNode.h
@@ -46,22 +46,6 @@ struct FMemTagStats
 	int64 SizeAverage = 0;
 };
 
-////////////////////////////////////////////////////////////////////////////////////////////////////
-
-class FMemTagNode;
-
-/** Type definition for shared pointers to instances of FMemTagNode. */
-typedef TSharedPtr<class FMemTagNode> FMemTagNodePtr;
-
-/** Type definition for shared references to instances of FMemTagNode. */
-typedef TSharedRef<class FMemTagNode> FMemTagNodeRef;
-
-/** Type definition for shared references to const instances of FMemTagNode. */
-typedef TSharedRef<const class FMemTagNode> FMemTagNodeRefConst;
-
-/** Type definition for weak references to instances of FMemTagNode. */
-typedef TWeakPtr<class FMemTagNode> FMemTagNodeWeak;
-
 ////////////////////////////////////////////////////////////////////////////////////////////////////
 /**
  * Class used to store information about an LLM tag node (used in the SMemTagTreeView).
@@ -111,10 +95,6 @@ public:
 	virtual FLinearColor GetIconColor() const override;
 	virtual FLinearColor GetColor() const override;
 
-	FMemTagNodePtr GetParentTagNode() const { return ParentTagNode; }
-	FMemoryTag* GetParentMemTag() const { return ParentTagNode.IsValid() ? ParentTagNode->GetMemTag() : nullptr; }
-	void SetParentTagNode(FMemTagNodePtr NodePtr) { ParentTagNode = NodePtr; }
-
 	const FMemTagStats& GetStats() const { return Stats; }
 	FMemTagStats& GetStats() { return Stats; }
 
@@ -133,11 +113,85 @@ public:
 	void ResetAggregatedStats();
 
 private:
-	FMemoryTag* MemTag;
-	FMemTagNodePtr ParentTagNode = nullptr;
+	FMemoryTag* MemTag = nullptr;
 	FMemTagStats Stats;
 };
 
 ////////////////////////////////////////////////////////////////////////////////////////////////////
 
+class FSystemMemTagNode : public FMemTagNode
+{
+	INSIGHTS_DECLARE_RTTI(FSystemMemTagNode, FMemTagNode)
+
+public:
+	/** Initialization constructor for the Asset MemTag node. */
+	explicit FSystemMemTagNode(TWeakPtr<FMemTagTable> InParentTable, FMemoryTag& InMemTag)
+		: FMemTagNode(InParentTable, InMemTag)
+	{
+	}
+
+	virtual const FSlateBrush* GetIcon() const override;
+
+	TSharedPtr<FMemTagNode> GetParentTagNode() const { return ParentTagNode; }
+	FMemoryTag* GetParentMemTag() const { return ParentTagNode.IsValid() ? ParentTagNode->GetMemTag() : nullptr; }
+	void SetParentTagNode(TSharedPtr<FMemTagNode> NodePtr) { ParentTagNode = NodePtr; }
+
+private:
+	TSharedPtr<FMemTagNode> ParentTagNode = nullptr;
+};
+
+////////////////////////////////////////////////////////////////////////////////////////////////////
+
+class FAssetMemTagNode : public FMemTagNode
+{
+	INSIGHTS_DECLARE_RTTI(FAssetMemTagNode, FMemTagNode)
+
+public:
+	/** Initialization constructor for the Asset MemTag node. */
+	explicit FAssetMemTagNode(TWeakPtr<FMemTagTable> InParentTable, FMemoryTag& InMemTag)
+		: FMemTagNode(InParentTable, InMemTag)
+	{
+	}
+
+	virtual const FSlateBrush* GetIcon() const override;
+
+	int32 GetObjectSerial() const { return ObjectSerial; }
+	void SetObjectSerial(int32 InObjectSerial) { ObjectSerial = InObjectSerial; }
+
+	int64 GetInclusiveSizeA() const { return InclusiveSizeA; }
+	void SetInclusiveSizeA(int64 Size) { InclusiveSizeA = Size; }
+
+	int64 GetInclusiveSizeB() const { return InclusiveSizeB; }
+	void SetInclusiveSizeB(int64 Size) { InclusiveSizeB = Size; }
+
+private:
+	int32 ObjectSerial = 0;
+	int64 InclusiveSizeA = 0;
+	int64 InclusiveSizeB = 0;
+};
+
+////////////////////////////////////////////////////////////////////////////////////////////////////
+
+class FClassMemTagNode : public FMemTagNode
+{
+	INSIGHTS_DECLARE_RTTI(FClassMemTagNode, FMemTagNode)
+
+public:
+	/** Initialization constructor for the Asset MemTag node. */
+	explicit FClassMemTagNode(TWeakPtr<FMemTagTable> InParentTable, FMemoryTag& InMemTag)
+		: FMemTagNode(InParentTable, InMemTag)
+	{
+	}
+
+	virtual const FSlateBrush* GetIcon() const override;
+
+	void SetClassSerial(int32 InClassSerial) { ClassSerial = InClassSerial; }
+	int32 GetClassSerial() const { return ClassSerial; }
+
+private:
+	int32 ClassSerial = 0;
+};
+
+////////////////////////////////////////////////////////////////////////////////////////////////////
+
 } // namespace UE::Insights::MemoryProfiler
diff --git a/Engine/Source/Developer/TraceInsights/Private/Insights/MemoryProfiler/ViewModels/MemTagNodeGroupingAndSorting.cpp b/Engine/Source/Developer/TraceInsights/Private/Insights/MemoryProfiler/ViewModels/MemTagNodeGroupingAndSorting.cpp
--- a/Engine/Source/Developer/TraceInsights/Private/Insights/MemoryProfiler/ViewModels/MemTagNodeGroupingAndSorting.cpp
+++ b/Engine/Source/Developer/TraceInsights/Private/Insights/MemoryProfiler/ViewModels/MemTagNodeGroupingAndSorting.cpp
@@ -55,10 +55,10 @@ void FMemTagNodeSortingByTracker::Sort(TArray<FBaseTreeNodePtr>& NodesToSort, ES
 			INSIGHTS_DEFAULT_PRESORTING_NODES(A, B);
 
 			INSIGHTS_ENSURE(A.IsValid() && A->Is<FMemTagNode>());
-			const FMemTagNodePtr MemTagNodeA = StaticCastSharedPtr<FMemTagNode, FBaseTreeNode>(A);
+			const TSharedPtr<FMemTagNode> MemTagNodeA = StaticCastSharedPtr<FMemTagNode, FBaseTreeNode>(A);
 
 			INSIGHTS_ENSURE(B.IsValid() && B->Is<FMemTagNode>());
-			const FMemTagNodePtr MemTagNodeB = StaticCastSharedPtr<FMemTagNode, FBaseTreeNode>(B);
+			const TSharedPtr<FMemTagNode> MemTagNodeB = StaticCastSharedPtr<FMemTagNode, FBaseTreeNode>(B);
 
 			if (MemTagNodeA->GetMemTrackerId() == MemTagNodeB->GetMemTrackerId())
 			{
@@ -78,10 +78,10 @@ void FMemTagNodeSortingByTracker::Sort(TArray<FBaseTreeNodePtr>& NodesToSort, ES
 			INSIGHTS_DEFAULT_PRESORTING_NODES(A, B);
 
 			INSIGHTS_ENSURE(A.IsValid() && A->Is<FMemTagNode>());
-			const FMemTagNodePtr MemTagNodeA = StaticCastSharedPtr<FMemTagNode, FBaseTreeNode>(A);
+			const TSharedPtr<FMemTagNode> MemTagNodeA = StaticCastSharedPtr<FMemTagNode, FBaseTreeNode>(A);
 
 			INSIGHTS_ENSURE(B.IsValid() && B->Is<FMemTagNode>());
-			const FMemTagNodePtr MemTagNodeB = StaticCastSharedPtr<FMemTagNode, FBaseTreeNode>(B);
+			const TSharedPtr<FMemTagNode> MemTagNodeB = StaticCastSharedPtr<FMemTagNode, FBaseTreeNode>(B);
 
 			if (MemTagNodeA->GetMemTrackerId() == MemTagNodeB->GetMemTrackerId())
 			{
diff --git a/Engine/Source/Developer/TraceInsights/Private/Insights/MemoryProfiler/ViewModels/MemorySharedState.cpp b/Engine/Source/Developer/TraceInsights/Private/Insights/MemoryProfiler/ViewModels/MemorySharedState.cpp
--- a/Engine/Source/Developer/TraceInsights/Private/Insights/MemoryProfiler/ViewModels/MemorySharedState.cpp
+++ b/Engine/Source/Developer/TraceInsights/Private/Insights/MemoryProfiler/ViewModels/MemorySharedState.cpp
@@ -135,6 +135,8 @@ void FMemorySharedState::OnEndSession(Timing::ITimingViewSession& InSession)
 		return;
 	}
 
+	//TODO: move TagList, Trackers and TagSets to MemoryManager.
+
 	TagList.Reset();
 
 	Trackers.Reset();
@@ -186,7 +188,7 @@ void FMemorySharedState::Tick(Timing::ITimingViewSession& InSession, const Trace
 		check(MainGraphTrack);
 
 		MainGraphTrack->SetOrder(GraphTrackOrder);
-		MainGraphTrack->SetName(TEXT("Main Memory Graph"));
+		MainGraphTrack->SetName(TEXT("MAIN MEMORY GRAPH"));
 
 		MainGraphTrack->AddTimelineSeries(FAllocationsGraphSeries::ETimeline::MaxTotalMem);
 		MainGraphTrack->AddTimelineSeries(FAllocationsGraphSeries::ETimeline::MinTotalMem);
@@ -354,8 +356,8 @@ void FMemorySharedState::CreateDefaultTracks()
 			for (const FMemoryTag* Tag : Tags)
 			{
 				if (Tag->GetTrackerId() == DefaultTrackerId && // is it used by the default tracker?
-					Tag->GetGraphTracks().Num() == 0 && // a graph isn't already added for this llm tag?
-					FCString::Stricmp(*Tag->GetStatName(), DefaultTags[DefaultTagIndex]) == 0) // is it one of the llm tags to show as default?
+					Tag->GetGraphTracks().Num() == 0 && // a graph isn't already added for this LLM tag?
+					FCString::Stricmp(*Tag->GetStatName(), DefaultTags[DefaultTagIndex]) == 0) // is it one of the LLM tags to show as default?
 				{
 					CreateMemTagGraphTrack(DefaultTrackerId, Tag->GetId());
 					CreatedDefaultTracks[DefaultTagIndex] = true;
@@ -1052,7 +1054,7 @@ TSharedPtr<FMemoryGraphTrack> FMemorySharedState::CreateGraphTrack(const FReport
 		const double MaxValue = GraphConfig.MaxY * 1024.0 * 1024.0;
 		GraphTrack->SetDefaultValueRange(MinValue, MaxValue);
 
-		UE_LOG(LogMemoryProfiler, Log, TEXT("[LLM Tags] Created graph \"%s\" (H=%.1f%s, MainStat=%s, Stats=%s)"),
+		UE_LOG(LogMemoryProfiler, Log, TEXT("[Memory] Created graph \"%s\" (H=%.1f%s, MainStat=%s, Stats=%s)"),
 			*ReportTypeGraphConfig.Title,
 			GraphTrack->GetHeight(),
 			GraphConfig.bStacked ? TEXT(", stacked") : TEXT(""),
diff --git a/Engine/Source/Developer/TraceInsights/Private/Insights/MemoryProfiler/Widgets/SMemTagTreeView.cpp b/Engine/Source/Developer/TraceInsights/Private/Insights/MemoryProfiler/Widgets/SMemTagTreeView.cpp
--- a/Engine/Source/Developer/TraceInsights/Private/Insights/MemoryProfiler/Widgets/SMemTagTreeView.cpp
+++ b/Engine/Source/Developer/TraceInsights/Private/Insights/MemoryProfiler/Widgets/SMemTagTreeView.cpp
@@ -64,7 +64,7 @@ void SMemTagTreeView::Construct(const FArguments& InArgs, TSharedPtr<SMemoryProf
 
 	TSharedRef<FMemTagTable> TablePtr = MakeShared<FMemTagTable>();
 	TablePtr->Reset();
-	TablePtr->SetDisplayName(FText::FromString(TEXT("LLM Tags")));
+	TablePtr->SetDisplayName(FText::FromString(TEXT("Memory Tags")));
 
 	ConstructWidget(TablePtr);
 
@@ -79,65 +79,28 @@ void SMemTagTreeView::Construct(const FArguments& InArgs, TSharedPtr<SMemoryProf
 
 void SMemTagTreeView::ConstructHeaderArea(TSharedRef<SVerticalBox> InHostBox)
 {
-	InHostBox->AddSlot()
-		.AutoHeight()
-		.Padding(2.0f, 2.0f, 2.0f, 2.0f)
-		[
-			SNew(SHorizontalBox)
-
-			+ SHorizontalBox::Slot()
-			.AutoWidth()
-			.Padding(0.0f, 0.0f, 4.0f, 0.0f)
-			.VAlign(VAlign_Center)
-			[
-				SNew(STextBlock)
-				.Text(LOCTEXT("TagSetText", "Tag Set:"))
-			]
-
-			+ SHorizontalBox::Slot()
-			.AutoWidth()
-			.Padding(0.0f, 0.0f, 0.0f, 0.0f)
-			.VAlign(VAlign_Center)
-			[
-				SAssignNew(TagSetsSegmentedControl, SSegmentedControl<FMemoryTagSetId>)
-				.SupportsMultiSelection(false)
-				.OnValueChanged_Lambda([this](FMemoryTagSetId InValue)
-				{
-					TagSetFilter = InValue;
-					OnNodeFilteringChanged();
-					bShouldUpdateStats = true;
-					bShouldUpdateBudgets = true;
-				})
-				.Value_Lambda([this]
-				{
-					return TagSetFilter;
-				})
-			]
-		];
-	UpdateAvailableTagSets();
-
 	TSharedPtr<SWidget> TopSettingsWidget = ConstructTopSettings();
 	if (TopSettingsWidget.IsValid())
 	{
 		InHostBox->AddSlot()
 			.AutoHeight()
-			.Padding(2.0f, 2.0f, 2.0f, 2.0f)
+			.Padding(4.0f, 4.0f, 4.0f, 2.0f)
 			[
 				TopSettingsWidget.ToSharedRef()
 			];
 	}
 
-	TSharedPtr<SWidget> FilterToolbarWidget = ConstructFilterToolbar();
-	if (FilterToolbarWidget.IsValid())
+	TSharedPtr<SWidget> TagSetAndViewPresetWidget = ConstructTagSetAndViewPreset();
+	if (TagSetAndViewPresetWidget.IsValid())
 	{
 		InHostBox->AddSlot()
 			.AutoHeight()
-			.Padding(2.0f, 2.0f, 2.0f, 2.0f)
+			.Padding(4.0f, 2.0f, 4.0f, 2.0f)
 			[
-				FilterToolbarWidget.ToSharedRef()
+				TagSetAndViewPresetWidget.ToSharedRef()
 			];
 	}
-
+	
 	TSharedPtr<SWidget> TimeMarkersWidget = ConstructTimeMarkers();
 	if (TimeMarkersWidget.IsValid())
 	{
@@ -157,12 +120,23 @@ void SMemTagTreeView::ConstructHeaderArea(TSharedRef<SVerticalBox> InHostBox)
 			];
 	}
 
+	TSharedPtr<SWidget> FilterToolbarWidget = ConstructFilterToolbar();
+	if (FilterToolbarWidget.IsValid())
+	{
+		InHostBox->AddSlot()
+			.AutoHeight()
+			.Padding(4.0f, 2.0f, 4.0f, 2.0f)
+			[
+				FilterToolbarWidget.ToSharedRef()
+			];
+	}
+
 	TSharedPtr<SWidget> HierarchyBreadcrumbTrailWidget = ConstructHierarchyBreadcrumbTrail();
 	if (HierarchyBreadcrumbTrailWidget.IsValid())
 	{
 		InHostBox->AddSlot()
 			.AutoHeight()
-			.Padding(2.0f, 2.0f, 2.0f, -2.0f)
+			.Padding(4.0f, 2.0f, 4.0f, -2.0f)
 			[
 				HierarchyBreadcrumbTrailWidget.ToSharedRef()
 			];
@@ -194,79 +168,76 @@ TSharedPtr<SWidget> SMemTagTreeView::ConstructTopSettings()
 
 	+ SHorizontalBox::Slot()
 	.AutoWidth()
-	.Padding(FMargin(0.0f, 0.0f, 0.0f, 0.0f))
+	.Padding(0.0f, 0.0f, 0.0f, 0.0f)
 	.VAlign(VAlign_Center)
 	[
-		ConstructHideAllButton()
+		ConstructBudgetSettings()
 	]
 
 	+ SHorizontalBox::Slot()
 	.AutoWidth()
-	.Padding(FMargin(0.0f, 0.0f, 0.0f, 0.0f))
+	.Padding(4.0f, 0.0f, 0.0f, 0.0f)
 	.VAlign(VAlign_Center)
 	[
-		ConstructShowAllButton()
-	]
+		SNew(SCheckBox)
+		.Style(FAppStyle::Get(), "ToggleButtonCheckbox")
+		.Padding(FMargin(4.0f, 3.0f, 4.0f, 3.0f))
+		.IsChecked_Lambda([this]()
+			{
+				return bAreTimeMarkerSettingsVisible ? ECheckBoxState::Checked : ECheckBoxState::Unchecked;
+			})
+		.OnCheckStateChanged_Lambda([this](ECheckBoxState InCheckBoxState)
+			{
+				bAreTimeMarkerSettingsVisible = (InCheckBoxState == ECheckBoxState::Checked);
+			})
+		.Content()
+		[
+			SNew(SImage)
+			.Image(FInsightsStyle::GetBrush("Icons.TimeMarkerSettings"))
+		]
+		.ToolTipText(LOCTEXT("TimeMarkerSettingsVisibilityToolTip", "Toggle visibility for the advanced Time Marker settings."))
+	];
 
-	+ SHorizontalBox::Slot()
-	.AutoWidth()
-	.Padding(FMargin(4.0f, 0.0f, 0.0f, 0.0f))
-	.VAlign(VAlign_Center)
-	[
-		ConstructTrackHeightControls()
-	]
+	return TopLine;
+}
+END_SLATE_FUNCTION_BUILD_OPTIMIZATION
 
-	+ SHorizontalBox::Slot()
-	.AutoWidth()
-	.Padding(FMargin(4.0f, 0.0f, 0.0f, 0.0f))
-	.VAlign(VAlign_Center)
-	[
-		ConstructBudgetSettings()
-	]
+////////////////////////////////////////////////////////////////////////////////////////////////////
 
-#if 0
-	+ SHorizontalBox::Slot()
-	.FillWidth(1.0f)
-	.Padding(0.0f, 0.0f, 0.0f, 0.0f)
-	.HAlign(HAlign_Right)
-	.VAlign(VAlign_Center)
-	[
-		SNew(STextBlock)
-		.Text(LOCTEXT("TagSetText", "Tag Set:"))
-	]
+TSharedPtr<SWidget> SMemTagTreeView::ConstructTagSetAndViewPreset()
+{
+	TagSetsSegmentedControl = SNew(SSegmentedControl<FMemoryTagSetId>)
+		.SupportsMultiSelection(false)
+		.OnValueChanged_Lambda([this](FMemoryTagSetId InValue)
+			{
+				TagSetFilter = InValue;
+				OnNodeFilteringChanged();
+				bShouldUpdateStats = true;
+				bShouldUpdateBudgets = true;
+			})
+		.Value_Lambda([this]
+			{
+				return TagSetFilter;
+			});
+	UpdateAvailableTagSets();
 
-	+ SHorizontalBox::Slot()
-	.AutoWidth()
-	.Padding(4.0f, 0.0f, 6.0f, 0.0f)
-	.VAlign(VAlign_Center)
-	[
-		SNew(SComboBox<TSharedPtr<FMemoryTagSetId>>)
-		.ToolTipText(LOCTEXT("TagSetToolTip", "Select the LLM tag set."))
-		.OptionsSource(&AvailableTagSets)
-		.OnSelectionChanged(this, &SMemTagTreeView::TagSet_OnSelectionChanged)
-		.OnGenerateWidget(this, &SMemTagTreeView::TagSet_OnGenerateWidget)
+	TSharedPtr<SHorizontalBox> Widget = SNew(SHorizontalBox)
+
+		+ SHorizontalBox::Slot()
+		.AutoWidth()
+		.Padding(0.0f, 0.0f, 0.0f, 0.0f)
+		.VAlign(VAlign_Center)
 		[
-			SNew(STextBlock)
-			.Text(this, &SMemTagTreeView::TagSet_GetSelectedText)
+			TagSetsSegmentedControl.ToSharedRef()
 		]
-	]
-#else
-	+ SHorizontalBox::Slot()
-	.FillWidth(1.0f)
-	.Padding(4.0f, 0.0f, 0.0f, 0.0f)
-	.HAlign(HAlign_Right)
-	.VAlign(VAlign_Center)
-	[
-		SNew(SBox)
-	]
-#endif
-	;
 
-	ConstructViewPreset(TopLine);
+		+ SHorizontalBox::Slot()
+		.FillWidth(1.0);
+
+	ConstructViewPreset(Widget);
 
-	return TopLine;
+	return Widget;
 }
-END_SLATE_FUNCTION_BUILD_OPTIMIZATION
 
 ////////////////////////////////////////////////////////////////////////////////////////////////////
 
@@ -283,27 +254,29 @@ void SMemTagTreeView::UpdateAvailableTagSets()
 
 		const FSlateBrush* TagSetIcon = nullptr;
 		FText TagSetName;
-		switch (TagSet)
+		if (TagSet == SystemsTagSet)
 		{
-		case 0:
 			TagSetIcon = FInsightsStyle::GetBrush("Icons.TagSet.Systems");
 			TagSetName = FText::FromString("Systems");
-			break;
-		case 1:
+		}
+		else if (TagSet == AssetsTagSet)
+		{
 			TagSetIcon = FInsightsStyle::GetBrush("Icons.TagSet.Assets");
 			TagSetName = FText::FromString("Assets");
-			break;
-		case 2:
+		}
+		else if (TagSet == AssetClassesTagSet)
+		{
 			TagSetIcon = FInsightsStyle::GetBrush("Icons.TagSet.AssetClasses");
 			TagSetName = FText::FromString("AssetClasses");
-			break;
-		default:
+		}
+		else
+		{
 			TagSetName = FText::FromString("Unknown");
 			TagSetIcon = FInsightsStyle::GetBrush("Icons.TagSet.Systems");
 		}
 
 		FText TagSetToolTip;
-		const FText TagSetToolTipFmt = LOCTEXT("TagSetToolTipFmt", "'{0}' Tag Set\n\nThe tree view shows the memory tags from the selected tag set.");
+		const FText TagSetToolTipFmt = LOCTEXT("TagSetToolTipFmt", "'{0}' Tag Set\n\nThe tree view shows memory tags from the selected tag set.");
 		TagSetToolTip = FText::Format(TagSetToolTipFmt, TagSetName);
 
 		TSharedRef<SWidget> Widget = SNew(SHorizontalBox)
@@ -329,7 +302,7 @@ void SMemTagTreeView::UpdateAvailableTagSets()
 			.VAlign(VAlign_Center)
 			[
 				SNew(STextBlock)
-				.Font(FStyleFonts::Get().Small)
+				//.Font(FStyleFonts::Get().Small)
 				.Text(TagSetName)
 			];
 
@@ -360,8 +333,9 @@ TSharedPtr<SWidget> SMemTagTreeView::ConstructFilterToolbar()
 	.VAlign(VAlign_Center)
 	[
 		SNew(SComboButton)
+		.ComboButtonStyle(FAppStyle::Get(), "SimpleComboButton")
 		.OnGetMenuContent(this, &SMemTagTreeView::MakeTrackersMenu)
-		.ToolTipText(LOCTEXT("TrackersMenuToolTip", "Filter list of LLM tags by tracker."))
+		.ToolTipText(LOCTEXT("TrackersMenuToolTip", "Filter the list of memory tags by LLM tracker."))
 		.ContentPadding(FMargin(0.0f, 1.0f))
 		.ButtonContent()
 		[
@@ -392,30 +366,6 @@ TSharedPtr<SWidget> SMemTagTreeView::ConstructFilterToolbar()
 	.VAlign(VAlign_Center)
 	[
 		ConstructFilterConfiguratorButton()
-	]
-
-	+ SHorizontalBox::Slot()
-	.AutoWidth()
-	.Padding(2.0f, 0.0f, 0.0f, 0.0f)
-	.VAlign(VAlign_Center)
-	[
-		SNew(SCheckBox)
-		.Style(FAppStyle::Get(), "ToggleButtonCheckbox")
-		.Padding(FMargin(4.0f, 3.0f, 4.0f, 3.0f))
-		.IsChecked_Lambda([this]()
-			{
-				return bAreTimeMarkerSettingsVisible ? ECheckBoxState::Checked : ECheckBoxState::Unchecked;
-			})
-		.OnCheckStateChanged_Lambda([this](ECheckBoxState InCheckBoxState)
-			{
-				bAreTimeMarkerSettingsVisible = (InCheckBoxState == ECheckBoxState::Checked);
-			})
-		.Content()
-		[
-			SNew(SImage)
-			.Image(FInsightsStyle::GetBrush("Icons.TimeMarkerSettings"))
-		]
-		.ToolTipText(LOCTEXT("TimeMarkerSettingsVisibilityToolTip", "Toggle visibility for the advanced Time Marker settings."))
 	];
 
 	return FilterLine;
@@ -666,9 +616,9 @@ BEGIN_SLATE_FUNCTION_BUILD_OPTIMIZATION
 TSharedRef<SButton> SMemTagTreeView::ConstructHideAllButton()
 {
 	return SNew(SButton)
-		.ToolTipText(LOCTEXT("HideAll_ToolTip", "Remove memory graph tracks for all LLM tags."))
-		.OnClicked(this, &SMemTagTreeView::HideAllTracks_OnClicked)
 		.ButtonStyle(&FAppStyle::Get().GetWidgetStyle<FButtonStyle>("SimpleButton"))
+		.ToolTipText(LOCTEXT("HideAll_ToolTip", "Remove all LLM graph tracks."))
+		.OnClicked(this, &SMemTagTreeView::HideAllTracks_OnClicked)
 		.ContentPadding(FMargin(2.0f, 2.0f, 2.0f, 2.0f))
 		[
 			SNew(SBox)
@@ -676,9 +626,8 @@ TSharedRef<SButton> SMemTagTreeView::ConstructHideAllButton()
 			.VAlign(VAlign_Center)
 			[
 				SNew(SImage)
-				//.ColorAndOpacity(FSlateColor::UseForeground())
-				.ColorAndOpacity(FSlateColor(FLinearColor(1.0f, 0.5f, 0.5f, 1.0f)))
-				.Image(FInsightsStyle::Get().GetBrush("Icons.RemoveAllMemTagGraphs"))
+				.ColorAndOpacity(FSlateColor::UseForeground())
+				.Image(FInsightsStyle::Get().GetBrush("Icons.RemoveMemTagGraphs"))
 			]
 		];
 }
@@ -687,12 +636,13 @@ END_SLATE_FUNCTION_BUILD_OPTIMIZATION
 ////////////////////////////////////////////////////////////////////////////////////////////////////
 
 BEGIN_SLATE_FUNCTION_BUILD_OPTIMIZATION
-TSharedRef<SButton> SMemTagTreeView::ConstructShowAllButton()
+TSharedRef<SButton> SMemTagTreeView::ConstructShowSelectedButton()
 {
 	return SNew(SButton)
-		.ToolTipText(LOCTEXT("ShowAll_ToolTip", "Create memory graph tracks for all visible (filtered) LLM tags."))
-		.OnClicked(this, &SMemTagTreeView::ShowAllTracks_OnClicked)
 		.ButtonStyle(&FAppStyle::Get().GetWidgetStyle<FButtonStyle>("SimpleButton"))
+		.ToolTipText(LOCTEXT("ShowSelected_ToolTip", "Create LLM graph tracks for the selected memory tags."))
+		.OnClicked(this, &SMemTagTreeView::ShowSelectedTracks_OnClicked)
+		.IsEnabled(this, &SMemTagTreeView::CanCreateGraphTracksForSelectedMemTags)
 		.ContentPadding(FMargin(2.0f, 2.0f, 2.0f, 2.0f))
 		[
 			SNew(SBox)
@@ -700,9 +650,8 @@ TSharedRef<SButton> SMemTagTreeView::ConstructShowAllButton()
 			.VAlign(VAlign_Center)
 			[
 				SNew(SImage)
-				//.ColorAndOpacity(FSlateColor::UseForeground())
-				.ColorAndOpacity(FSlateColor(FLinearColor(0.5f, 1.0f, 0.5f, 1.0f)))
-				.Image(FInsightsStyle::Get().GetBrush("Icons.AddAllMemTagGraphs"))
+				.ColorAndOpacity(FSlateColor::UseForeground())
+				.Image(FInsightsStyle::Get().GetBrush("Icons.AddMemTagGraphs"))
 			]
 		];
 }
@@ -736,15 +685,15 @@ TSharedRef<SWidget> SMemTagTreeView::ConstructTrackHeightControls()
 
 		+ SSegmentedControl<uint32>::Slot((uint32)EMemoryTrackHeightMode::Small)
 		.Icon(FInsightsStyle::GetBrush("Icons.SizeSmall"))
-		.ToolTip(LOCTEXT("SmallHeight_ToolTip", "Change height of LLM Tag Graph tracks to Small."))
+		.ToolTip(LOCTEXT("SmallHeight_ToolTip", "Change height of all LLM Tag graph tracks to Small."))
 
 		+ SSegmentedControl<uint32>::Slot((uint32)EMemoryTrackHeightMode::Medium)
 		.Icon(FInsightsStyle::GetBrush("Icons.SizeMedium"))
-		.ToolTip(LOCTEXT("MediumHeight_ToolTip", "Change height of LLM Tag Graph tracks to Medium."))
+		.ToolTip(LOCTEXT("MediumHeight_ToolTip", "Change height of all LLM Tag graph tracks to Medium."))
 
 		+ SSegmentedControl<uint32>::Slot((uint32)EMemoryTrackHeightMode::Large)
 		.Icon(FInsightsStyle::GetBrush("Icons.SizeLarge"))
-		.ToolTip(LOCTEXT("LargeHeight_ToolTip", "Change height of LLM Tag Graph tracks to Large."));
+		.ToolTip(LOCTEXT("LargeHeight_ToolTip", "Change height of all LLM Tag graph tracks to Large."));
 }
 END_SLATE_FUNCTION_BUILD_OPTIMIZATION
 
@@ -754,14 +703,30 @@ BEGIN_SLATE_FUNCTION_BUILD_OPTIMIZATION
 TSharedRef<SWidget> SMemTagTreeView::ConstructBudgetSettings()
 {
 	return SNew(SComboButton)
+		.ComboButtonStyle(FAppStyle::Get(), "SimpleComboButton")
 		.OnGetMenuContent(this, &SMemTagTreeView::MakeBudgetSettingsMenu)
 		.ToolTipText(LOCTEXT("BudgetSettingsToolTip", "Budget Settings"))
 		.ContentPadding(FMargin(0.0f, 1.0f))
 		.ButtonContent()
 		[
-			SNew(SImage)
-			.Image(FInsightsStyle::GetBrush("Icons.BudgetSettings"))
-			.ColorAndOpacity(FSlateColor(EStyleColor::Foreground))
+			SNew(SHorizontalBox)
+
+			+ SHorizontalBox::Slot()
+			.AutoWidth()
+			.Padding(0.0f, 0.0f, 0.0f, 0.0f)
+			[
+				SNew(SImage)
+				.Image(FInsightsStyle::GetBrush("Icons.BudgetSettings"))
+				.ColorAndOpacity(FSlateColor(EStyleColor::Foreground))
+			]
+
+			+ SHorizontalBox::Slot()
+			.AutoWidth()
+			.Padding(4.0f, 0.0f, 0.0f, 0.0f)
+			[
+				SNew(STextBlock)
+				.Text(LOCTEXT("BudgetSettingsText", "Budget"))
+			]
 		];
 }
 END_SLATE_FUNCTION_BUILD_OPTIMIZATION
@@ -774,6 +739,30 @@ TSharedPtr<SWidget> SMemTagTreeView::ConstructFooter()
 	return
 		SNew(SHorizontalBox)
 
+		+ SHorizontalBox::Slot()
+		.AutoWidth()
+		.Padding(4.0f, -2.0f, 0.0f, 2.0f)
+		.VAlign(VAlign_Top)
+		[
+			ConstructShowSelectedButton()
+		]
+
+		+ SHorizontalBox::Slot()
+		.AutoWidth()
+		.Padding(2.0f, -2.0f, 0.0f, 2.0f)
+		.VAlign(VAlign_Top)
+		[
+			ConstructHideAllButton()
+		]
+
+		+ SHorizontalBox::Slot()
+		.AutoWidth()
+		.Padding(4.0f, -3.0f, 0.0f, 3.0f)
+		.VAlign(VAlign_Top)
+		[
+			ConstructTrackHeightControls()
+		]
+
 		+ SHorizontalBox::Slot()
 		.FillWidth(1.0f)
 
@@ -962,80 +951,6 @@ void SMemTagTreeView::InternalCreateGroupings()
 
 	AvailableGroupings.Insert(MakeShared<FMemTagBudgetNodeGrouping>(SharedThis(this)), Index++);
 
-#if 0
-	// Creates one group for each tag type.
-	class FMemTagGroupingByTagType : public
-	{
-		{
-			TMap<EMemTagNodeType, FMemTagNodePtr> GroupNodeSet;
-			for (const FMemTagNodePtr& NodePtr : MemTagNodes)
-			{
-				const EMemTagNodeType NodeType = NodePtr->GetType();
-				FMemTagNodePtr GroupPtr = GroupNodeSet.FindRef(NodeType);
-				if (!GroupPtr)
-				{
-					const FName GroupName = *MemTagNodeTypeHelper::ToText(NodeType).ToString();
-					GroupPtr = GroupNodeSet.Add(NodeType, MakeShared<FMemTagNode>(GroupName));
-				}
-				GroupPtr->AddChildAndSetParent(NodePtr);
-				TreeView->SetItemExpansion(GroupPtr, true);
-			}
-			GroupNodeSet.KeySort([](const EMemTagNodeType& A, const EMemTagNodeType& B) { return A < B; }); // sort groups by type
-			GroupNodeSet.GenerateValueArray(GroupNodes);
-		}
-	};
-	AvailableGroupings.Insert(MakeShared<FMemTagGroupingByTagType>(), Index++);
-#endif
-
-#if 0
-	// Creates one group for one letter.
-	class FMemTagGroupingByName : public
-	{
-		{
-			TMap<TCHAR, FMemTagNodePtr> GroupNodeSet;
-			for (const FMemTagNodePtr& NodePtr : MemTagNodes)
-			{
-				FString FirstLetterStr(NodePtr->GetName().GetPlainNameString().Left(1).ToUpper());
-				const TCHAR FirstLetter = FirstLetterStr[0];
-				FMemTagNodePtr GroupPtr = GroupNodeSet.FindRef(FirstLetter);
-				if (!GroupPtr)
-				{
-					const FName GroupName(FirstLetterStr);
-					GroupPtr = GroupNodeSet.Add(FirstLetter, MakeShared<FMemTagNode>(GroupName));
-				}
-				GroupPtr->AddChildAndSetParent(NodePtr);
-			}
-			GroupNodeSet.KeySort([](const TCHAR& A, const TCHAR& B) { return A < B; }); // sort groups alphabetically
-			GroupNodeSet.GenerateValueArray(GroupNodes);
-		}
-	};
-	AvailableGroupings.Insert(MakeShared<FMemTagGroupingByName>(), Index++);
-#endif
-
-#if 0
-	// Groups LLM tags by tracker.
-	class FMemTagGroupingByTracker : public
-	{
-		{
-			TMap<FMemoryTrackerId, FMemTagNodePtr> GroupNodeSet;
-			for (const FMemTagNodePtr& NodePtr : MemTagNodes)
-			{
-				FMemoryTrackerId TrackerId = NodePtr->GetMemTrackerId();
-				FMemTagNodePtr GroupPtr = GroupNodeSet.FindRef(TrackerId);
-				if (!GroupPtr)
-				{
-					const FName GroupName = *NodePtr->GetTrackerText().ToString();
-					GroupPtr = GroupNodeSet.Add(TrackerId, MakeShared<FMemTagNode>(GroupName));
-				}
-				GroupPtr->AddChildAndSetParent(NodePtr);
-				TreeView->SetItemExpansion(GroupPtr, true);
-			}
-			GroupNodeSet.GenerateValueArray(GroupNodes);
-		}
-	};
-	AvailableGroupings.Insert(MakeShared<FMemTagGroupingByTracker>(), Index++);
-#endif
-
 #if 0
 	// Groups LLM tags by their hierarchy.
 	class FMemTagGroupingByHierarchy : public
@@ -1162,7 +1077,23 @@ void SMemTagTreeView::RebuildTree(bool bResync)
 			for (FMemoryTag* MemTagPtr : MemTags)
 			{
 				check(MemTagPtr);
-				FMemTagNodePtr MemTagNodePtr = MakeShared<FMemTagNode>(ParentMemTagTable, *MemTagPtr);
+				TSharedPtr<FMemTagNode> MemTagNodePtr;
+				if (MemTagPtr->GetTagSetId() == SystemsTagSet)
+				{
+					MemTagNodePtr = MakeShared<FSystemMemTagNode>(ParentMemTagTable, *MemTagPtr);
+				}
+				else if (MemTagPtr->GetTagSetId() == AssetsTagSet)
+				{
+					MemTagNodePtr = MakeShared<FAssetMemTagNode>(ParentMemTagTable, *MemTagPtr);
+				}
+				else if (MemTagPtr->GetTagSetId() == AssetClassesTagSet)
+				{
+					MemTagNodePtr = MakeShared<FClassMemTagNode>(ParentMemTagTable, *MemTagPtr);
+				}
+				else
+				{
+					MemTagNodePtr = MakeShared<FMemTagNode>(ParentMemTagTable, *MemTagPtr);
+				}
 				TableRowNodes.Add(MemTagNodePtr);
 				MemTagNodesIdMap.Add(MemTagPtr->GetId(), MemTagNodePtr);
 			}
@@ -1170,15 +1101,17 @@ void SMemTagTreeView::RebuildTree(bool bResync)
 			// Resolve pointers to parent tags.
 			for (FTableTreeNodePtr& NodePtr : TableRowNodes)
 			{
-				check(NodePtr->Is<FMemTagNode>());
-				FMemTagNode& MemTagNode = NodePtr->As<FMemTagNode>();
-				check(MemTagNode.GetMemTag() != nullptr);
-				FMemoryTag& MemTag = *MemTagNode.GetMemTag();
-				FMemTagNodePtr ParentNodePtr = MemTagNodesIdMap.FindRef(MemTag.GetParentId());
-				if (ParentNodePtr)
+				if (NodePtr->Is<FSystemMemTagNode>())
 				{
-					check(ParentNodePtr != NodePtr);
-					MemTagNode.SetParentTagNode(ParentNodePtr);
+					FSystemMemTagNode& MemTagNode = NodePtr->As<FSystemMemTagNode>();
+					check(MemTagNode.GetMemTag() != nullptr);
+					FMemoryTag& MemTag = *MemTagNode.GetMemTag();
+					TSharedPtr<FMemTagNode> ParentNodePtr = MemTagNodesIdMap.FindRef(MemTag.GetParentId());
+					if (ParentNodePtr)
+					{
+						check(ParentNodePtr != NodePtr);
+						MemTagNode.SetParentTagNode(ParentNodePtr);
+					}
 				}
 			}
 
@@ -1236,7 +1169,7 @@ void SMemTagTreeView::RebuildTree(bool bResync)
 	if (TotalTime > 0.01)
 	{
 		const double SyncTime = SyncStopwatch.GetAccumulatedTime();
-		UE_LOG(LogMemoryProfiler, Log, TEXT("[LLM Tags] Tree view rebuilt in %.4fs (sync: %.4fs + update: %.4fs) --> %d tags (%d added)"),
+		UE_LOG(LogMemoryProfiler, Log, TEXT("[MemTags] Tree view rebuilt in %.4fs (sync: %.4fs + update: %.4fs) --> %d tags (%d added)"),
 			TotalTime, SyncTime, TotalTime - SyncTime, TableRowNodes.Num(), TableRowNodes.Num() - PreviousNodeCount);
 	}
 }
@@ -1348,7 +1281,7 @@ void SMemTagTreeView::UpdateStats()
 	const double TotalTime = Stopwatch.GetAccumulatedTime();
 	if (TotalTime > 0.01)
 	{
-		UE_LOG(LogMemoryProfiler, Log, TEXT("[LLM Tags] Aggregated stats updated in %.4fs"), TotalTime);
+		UE_LOG(LogMemoryProfiler, Log, TEXT("[MemTags] Aggregated stats updated in %.4fs"), TotalTime);
 	}
 }
 
@@ -1479,19 +1412,19 @@ void SMemTagTreeView::ExtendMenuBeforeMisc(FMenuBuilder& MenuBuilder)
 		MenuBuilder.AddSubMenu
 		(
 			LOCTEXT("ContextMenu_CreateGraphTracks_SubMenu", "Create Graph Tracks"),
-			LOCTEXT("ContextMenu_CreateGraphTracks_SubMenu_Desc", "Create memory graph tracks."),
+			LOCTEXT("ContextMenu_CreateGraphTracks_SubMenu_Desc", "Create memory graph tracks (LLM graph series)."),
 			FNewMenuDelegate::CreateSP(this, &SMemTagTreeView::ExtendMenuCreateGraphTracks),
 			false,
-			FSlateIcon()
+			FSlateIcon(FInsightsStyle::GetStyleSetName(), "Icons.AddMemTagGraphs")
 		);
 
 		MenuBuilder.AddSubMenu
 		(
 			LOCTEXT("ContextMenu_RemoveGraphTracks_SubMenu", "Remove Graph Tracks"),
-			LOCTEXT("ContextMenu_RemoveGraphTracks_SubMenu_Desc", "Remove memory graph tracks."),
+			LOCTEXT("ContextMenu_RemoveGraphTracks_SubMenu_Desc", "Remove memory graph tracks (LLM graph series)."),
 			FNewMenuDelegate::CreateSP(this, &SMemTagTreeView::ExtendMenuRemoveGraphTracks),
 			false,
-			FSlateIcon()
+			FSlateIcon(FInsightsStyle::GetStyleSetName(), "Icons.RemoveMemTagGraphs")
 		);
 
 		FUIAction Action_GenerateColorForSelectedMemTags
@@ -1502,7 +1435,7 @@ void SMemTagTreeView::ExtendMenuBeforeMisc(FMenuBuilder& MenuBuilder)
 		MenuBuilder.AddMenuEntry
 		(
 			LOCTEXT("ContextMenu_GenerateColorForSelectedMemTags", "Generate New Color"),
-			LOCTEXT("ContextMenu_GenerateColorForSelectedMemTags_Desc", "Generate new color for selected LLM tag(s)."),
+			LOCTEXT("ContextMenu_GenerateColorForSelectedMemTags_Desc", "Generate new color for selected memory tag(s)."),
 			FSlateIcon(),
 			Action_GenerateColorForSelectedMemTags,
 			NAME_None,
@@ -1517,7 +1450,7 @@ void SMemTagTreeView::ExtendMenuBeforeMisc(FMenuBuilder& MenuBuilder)
 		MenuBuilder.AddMenuEntry
 		(
 			LOCTEXT("ContextMenu_EditColorForSelectedMemTags", "Edit Color..."),
-			LOCTEXT("ContextMenu_EditColorForSelectedMemTags_Desc", "Edit color for selected LLM tag(s)."),
+			LOCTEXT("ContextMenu_EditColorForSelectedMemTags_Desc", "Edit color for selected memory tag(s)."),
 			FSlateIcon(),
 			Action_EditColorForSelectedMemTags,
 			NAME_None,
@@ -1544,7 +1477,7 @@ void SMemTagTreeView::ExtendMenuCreateGraphTracks(FMenuBuilder& MenuBuilder)
 		MenuBuilder.AddMenuEntry
 		(
 			LOCTEXT("ContextMenu_CreateGraphTracksForSelectedMemTags", "Selected"),
-			LOCTEXT("ContextMenu_CreateGraphTracksForSelectedMemTags_Desc", "Create memory graph tracks for selected LLM tag(s)."),
+			LOCTEXT("ContextMenu_CreateGraphTracksForSelectedMemTags_Desc", "Create memory graph tracks (LLM graph series) for the selected memory tag(s)."),
 			FSlateIcon(),
 			Action_CreateGraphTracksForSelectedMemTags,
 			NAME_None,
@@ -1560,7 +1493,7 @@ void SMemTagTreeView::ExtendMenuCreateGraphTracks(FMenuBuilder& MenuBuilder)
 		MenuBuilder.AddMenuEntry
 		(
 			LOCTEXT("ContextMenu_CreateGraphTracksForFilteredMemTags", "Filtered"),
-			LOCTEXT("ContextMenu_CreateGraphTracksForFilteredMemTags_Desc", "Create memory graph tracks for all visible (filtered) LLM tags."),
+			LOCTEXT("ContextMenu_CreateGraphTracksForFilteredMemTags_Desc", "Create memory graph tracks (LLM graph series) for the visible (filtered) memory tags."),
 			FSlateIcon(),
 			Action_CreateGraphTracksForFilteredMemTags,
 			NAME_None,
@@ -1576,7 +1509,7 @@ void SMemTagTreeView::ExtendMenuCreateGraphTracks(FMenuBuilder& MenuBuilder)
 		MenuBuilder.AddMenuEntry
 		(
 			LOCTEXT("ContextMenu_CreateAllGraphTracks", "All"),
-			LOCTEXT("ContextMenu_CreateAllGraphTracks_Desc", "Create memory graph tracks for all LLM tags."),
+			LOCTEXT("ContextMenu_CreateAllGraphTracks_Desc", "Create memory graph tracks (LLM graph series) for all memory tags."),
 			FSlateIcon(),
 			Action_CreateAllGraphTracks,
 			NAME_None,
@@ -1603,7 +1536,7 @@ void SMemTagTreeView::ExtendMenuRemoveGraphTracks(FMenuBuilder& MenuBuilder)
 		MenuBuilder.AddMenuEntry
 		(
 			LOCTEXT("ContextMenu_RemoveGraphTracksForSelectedMemTags", "Selected"),
-			LOCTEXT("ContextMenu_RemoveGraphTracksForSelectedMemTags_Desc", "Remove memory graph tracks for selected LLM tag(s)."),
+			LOCTEXT("ContextMenu_RemoveGraphTracksForSelectedMemTags_Desc", "Remove memory graph tracks (LLM graph series) for the selected memory tag(s)."),
 			FSlateIcon(),
 			Action_RemoveGraphTracksForSelectedMemTags,
 			NAME_None,
@@ -1619,7 +1552,7 @@ void SMemTagTreeView::ExtendMenuRemoveGraphTracks(FMenuBuilder& MenuBuilder)
 		MenuBuilder.AddMenuEntry
 		(
 			LOCTEXT("ContextMenu_RemoveAllGraphTracks", "All"),
-			LOCTEXT("ContextMenu_RemoveAllGraphTracks_Desc", "Remove memory graph tracks for all LLM tags."),
+			LOCTEXT("ContextMenu_RemoveAllGraphTracks_Desc", "Remove memory graph tracks (LLM graph series) for all memory tags."),
 			FSlateIcon(),
 			Action_RemoveAllGraphTracks,
 			NAME_None,
@@ -1681,7 +1614,7 @@ bool SMemTagTreeView::FilterNodeCustom(const FTableTreeNode& InNode) const
 
 void SMemTagTreeView::SelectMemTagNode(FMemoryTagId MemTagId)
 {
-	FMemTagNodePtr NodePtr = GetMemTagNode(MemTagId);
+	TSharedPtr<FMemTagNode> NodePtr = GetMemTagNode(MemTagId);
 	if (NodePtr)
 	{
 		TreeView->SetSelection(NodePtr);
@@ -1746,9 +1679,9 @@ void SMemTagTreeView::LoadReportXML()
 // Button actions re graph tracks
 ////////////////////////////////////////////////////////////////////////////////////////////////////
 
-FReply SMemTagTreeView::ShowAllTracks_OnClicked()
+FReply SMemTagTreeView::ShowSelectedTracks_OnClicked()
 {
-	CreateGraphTracksForFilteredMemTags();
+	CreateGraphTracksForSelectedMemTags();
 	return FReply::Handled();
 }
 
@@ -1793,7 +1726,7 @@ void SMemTagTreeView::TryRemoveGraphTrackForNode(FMemorySharedState& SharedState
 }
 
 ////////////////////////////////////////////////////////////////////////////////////////////////////
-// Create memory graph tracks for selected LLM tag(s)
+// Create memory graph tracks (LLM graph series) for the selected memory tags
 ////////////////////////////////////////////////////////////////////////////////////////////////////
 
 bool SMemTagTreeView::CanCreateGraphTracksForSelectedMemTags() const
@@ -1829,7 +1762,7 @@ void SMemTagTreeView::CreateGraphTracksForSelectedMemTags()
 }
 
 ////////////////////////////////////////////////////////////////////////////////////////////////////
-// Create memory graph tracks for filtered LLM tags
+// Create memory graph tracks (LLM graph series) for the visible (filtered) memory tags
 ////////////////////////////////////////////////////////////////////////////////////////////////////
 
 bool SMemTagTreeView::CanCreateGraphTracksForFilteredMemTags() const
@@ -1854,13 +1787,13 @@ void SMemTagTreeView::CreateGraphTracksForFilteredMemTags()
 }
 
 ////////////////////////////////////////////////////////////////////////////////////////////////////
-// Create all memory graph tracks for selected LLM tag(s)
+// Create memory graph tracks (LLM graph series) for all memory tags
 ////////////////////////////////////////////////////////////////////////////////////////////////////
 
 bool SMemTagTreeView::CanCreateAllGraphTracks() const
 {
 	TSharedPtr<SMemoryProfilerWindow> ProfilerWindow = GetProfilerWindow();
-	return ProfilerWindow.IsValid();
+	return ProfilerWindow.IsValid() && TagSetFilter == SystemsTagSet;
 }
 
 ////////////////////////////////////////////////////////////////////////////////////////////////////
@@ -1890,7 +1823,7 @@ void SMemTagTreeView::CreateAllGraphTracks()
 }
 
 ////////////////////////////////////////////////////////////////////////////////////////////////////
-// Remove memory graph tracks for selected LLM tag(s)
+// Remove memory graph tracks (LLM graph series) for the selected memory tags
 ////////////////////////////////////////////////////////////////////////////////////////////////////
 
 bool SMemTagTreeView::CanRemoveGraphTracksForSelectedMemTags() const
@@ -1926,7 +1859,7 @@ void SMemTagTreeView::RemoveGraphTracksForSelectedMemTags()
 }
 
 ////////////////////////////////////////////////////////////////////////////////////////////////////
-// Remove all graph series
+// Remove all memory graph tracks (LLM graph series)
 ////////////////////////////////////////////////////////////////////////////////////////////////////
 
 bool SMemTagTreeView::CanRemoveAllGraphTracks() const
@@ -2195,34 +2128,6 @@ TSharedRef<SWidget> SMemTagTreeView::MakeBudgetSettingsMenu()
 
 ////////////////////////////////////////////////////////////////////////////////////////////////////
 
-TSharedRef<SWidget> MakeBudgetModeSubMenu()
-{
-	FMenuBuilder MenuBuilder(/*bInShouldCloseWindowAfterMenuSelection=*/true, nullptr);
-	MenuBuilder.SetSearchable(false);
-
-	MenuBuilder.BeginSection("BudgetMode");
-	//
-	MenuBuilder.EndSection();
-
-	return MenuBuilder.MakeWidget();
-}
-
-////////////////////////////////////////////////////////////////////////////////////////////////////
-
-TSharedRef<SWidget> MakeBudgetPlatformSubMenu()
-{
-	FMenuBuilder MenuBuilder(/*bInShouldCloseWindowAfterMenuSelection=*/true, nullptr);
-	MenuBuilder.SetSearchable(false);
-
-	MenuBuilder.BeginSection("BudgetPlatform");
-	//
-	MenuBuilder.EndSection();
-
-	return MenuBuilder.MakeWidget();
-}
-
-////////////////////////////////////////////////////////////////////////////////////////////////////
-
 void SMemTagTreeView::OpenAndSelectBudgetFile()
 {
 	TSharedPtr<FMemTagBudgetFileDesc> BudgetFile = OpenBudgetFile();
@@ -2387,13 +2292,20 @@ const TCHAR* SMemTagTreeView::GetTagSetFilterBudgetCachedName() const
 		return nullptr;
 	}
 
-	switch (TagSetFilter)
+	if (TagSetFilter == SystemsTagSet)
+	{
+		return CurrentBudget->FindString(TEXTVIEW("Systems"));
+	}
+	if (TagSetFilter == AssetsTagSet)
 	{
-		case 0:  return CurrentBudget->FindString(TEXTVIEW("Systems"));
-		case 1:  return CurrentBudget->FindString(TEXTVIEW("Assets"));
-		case 2:  return CurrentBudget->FindString(TEXTVIEW("AssetClasses"));
-		default: return nullptr;
+		return CurrentBudget->FindString(TEXTVIEW("Assets"));
 	}
+	if (TagSetFilter == AssetClassesTagSet)
+	{
+		return CurrentBudget->FindString(TEXTVIEW("AssetClasses"));
+	}
+
+	return nullptr;
 }
 
 ////////////////////////////////////////////////////////////////////////////////////////////////////
@@ -2484,22 +2396,22 @@ void SMemTagTreeView::InitBudgetOptions()
 				if (FPaths::FileExists(FullPath))
 				{
 					AvailableBudgetFiles.Insert(MakeShared<FMemTagBudgetFileDesc>(Label, FullPath), 0);
-					UE_LOG(LogMemoryProfiler, Log, TEXT("Budget file found (Label=\"%s\", Path=\"%s\")"), *Label, *FullPath);
+					UE_LOG(LogMemoryProfiler, Log, TEXT("[MemTags] Budget file found (Label=\"%s\", Path=\"%s\")"), *Label, *FullPath);
 				}
 				else
 				{
-					UE_LOG(LogMemoryProfiler, Warning, TEXT("Budget file not found (Label=\"%s\", Path=\"%s\")"), *Label, *FullPath);
+					UE_LOG(LogMemoryProfiler, Warning, TEXT("[MemTags] Budget file not found (Label=\"%s\", Path=\"%s\")"), *Label, *FullPath);
 				}
 			}
 		}
 		else
 		{
-			UE_LOG(LogMemoryProfiler, Log, TEXT("Couldn't find Insights.MemoryProfiler config section"));
+			UE_LOG(LogMemoryProfiler, Log, TEXT("[MemTags] Couldn't find Insights.MemoryProfiler config section"));
 		}
 	}
 	else
 	{
-		UE_LOG(LogMemoryProfiler, Log, TEXT("Couldn't find UnrealInsightsSettings config"));
+		UE_LOG(LogMemoryProfiler, Log, TEXT("[MemTags] Couldn't find UnrealInsightsSettings config"));
 	}
 
 	AvailableBudgetModes.Reset();
@@ -2752,7 +2664,7 @@ void SMemTagTreeView::ApplyBudgetToNodes()
 	const double TotalTime = Stopwatch.GetAccumulatedTime();
 	if (TotalTime > 0.01)
 	{
-		UE_LOG(LogMemoryProfiler, Log, TEXT("[LLM Tags] Budgets updated in %.4fs"), TotalTime);
+		UE_LOG(LogMemoryProfiler, Log, TEXT("[MemTags] Budgets updated in %.4fs"), TotalTime);
 	}
 }
 
diff --git a/Engine/Source/Developer/TraceInsights/Private/Insights/MemoryProfiler/Widgets/SMemTagTreeView.h b/Engine/Source/Developer/TraceInsights/Private/Insights/MemoryProfiler/Widgets/SMemTagTreeView.h
--- a/Engine/Source/Developer/TraceInsights/Private/Insights/MemoryProfiler/Widgets/SMemTagTreeView.h
+++ b/Engine/Source/Developer/TraceInsights/Private/Insights/MemoryProfiler/Widgets/SMemTagTreeView.h
@@ -176,6 +176,11 @@ public:
 	TSharedPtr<FMemTagBudget> GetBudget() const { return CurrentBudget; }
 	void GetBudgetGrouping(const FMemTagBudgetGrouping*& OutGrouping, const FMemTagBudgetGrouping*& OutGroupingOverride) const;
 
+	FMemoryTagSetId GetTagSet() const { return TagSetFilter; }
+	bool IsSystemsTagSet() const { return TagSetFilter == SystemsTagSet; }
+	bool IsAssetsTagSet() const { return TagSetFilter == AssetsTagSet; }
+	bool IsAssetClassesTagSet() const { return TagSetFilter == AssetClassesTagSet; }
+
 protected:
 	const TCHAR* GetTagSetFilterBudgetCachedName() const;
 	const TCHAR* GetSelectedBudgetPlatformCachedName() const;
@@ -193,8 +198,9 @@ protected:
 
 private:
 	TSharedPtr<SWidget> ConstructTopSettings();
+	TSharedPtr<SWidget> ConstructTagSetAndViewPreset();
 	TSharedRef<SButton> ConstructHideAllButton();
-	TSharedRef<SButton> ConstructShowAllButton();
+	TSharedRef<SButton> ConstructShowSelectedButton();
 	TSharedRef<SWidget> ConstructTrackHeightControls();
 	TSharedRef<SWidget> ConstructBudgetSettings();
 	TSharedPtr<SWidget> ConstructTimeMarkers();
@@ -243,7 +249,7 @@ private:
 	bool CanLoadReportXML() const;
 	void LoadReportXML();
 
-	FReply ShowAllTracks_OnClicked();
+	FReply ShowSelectedTracks_OnClicked();
 	FReply HideAllTracks_OnClicked();
 
 	void TryCreateGraphTrackForNode(FMemorySharedState& SharedState, const FBaseTreeNode& Node) const;
@@ -286,8 +292,6 @@ private:
 	// Budgets
 
 	TSharedRef<SWidget> MakeBudgetSettingsMenu();
-	TSharedRef<SWidget> MakeBudgetModeSubMenu();
-	TSharedRef<SWidget> MakeBudgetPlatformSubMenu();
 
 	TSharedPtr<FMemTagBudgetFileDesc> OpenBudgetFile();
 
@@ -329,6 +333,9 @@ private:
 
 	/** Filter the LLM tags by tag set. */
 	FMemoryTagSetId TagSetFilter = FMemoryTagSet::DefaultTagSetId;
+	FMemoryTagSetId SystemsTagSet = 0;
+	FMemoryTagSetId AssetsTagSet = 1;
+	FMemoryTagSetId AssetClassesTagSet = 2;
 	TArray<TSharedPtr<FMemoryTagSetId>> AvailableTagSets;
 	TSharedPtr<SSegmentedControl<FMemoryTagSetId>> TagSetsSegmentedControl;
 
diff --git a/Engine/Source/Developer/TraceInsights/Private/Insights/MemoryProfiler/Widgets/SMemoryProfilerWindow.cpp b/Engine/Source/Developer/TraceInsights/Private/Insights/MemoryProfiler/Widgets/SMemoryProfilerWindow.cpp
--- a/Engine/Source/Developer/TraceInsights/Private/Insights/MemoryProfiler/Widgets/SMemoryProfilerWindow.cpp
+++ b/Engine/Source/Developer/TraceInsights/Private/Insights/MemoryProfiler/Widgets/SMemoryProfilerWindow.cpp
@@ -426,7 +426,7 @@ void SMemoryProfilerWindow::RegisterTabSpawners()
 	if (Config.ShouldRegisterMinorTab(FMemoryProfilerTabs::MemTagTreeViewID))
 	{
 		TabManagerPtr->RegisterTabSpawner(FMemoryProfilerTabs::MemTagTreeViewID, FOnSpawnTab::CreateRaw(this, &SMemoryProfilerWindow::SpawnTab_MemTagTreeView))
-			.SetDisplayName(LOCTEXT("MemTagTreeViewTabTitle", "LLM Tags"))
+			.SetDisplayName(LOCTEXT("MemTagTreeViewTabTitle", "Memory Tags"))
 			.SetIcon(FSlateIcon(FInsightsStyle::GetStyleSetName(), "Icons.MemTagTreeView"))
 			.SetGroup(Group);
 	}
@@ -445,7 +445,7 @@ void SMemoryProfilerWindow::RegisterTabSpawners()
 BEGIN_SLATE_FUNCTION_BUILD_OPTIMIZATION
 TSharedRef<FTabManager::FLayout> SMemoryProfilerWindow::CreateDefaultTabLayout() const
 {
-	return FTabManager::NewLayout("InsightsMemoryProfilerLayout_v1.2")
+	return FTabManager::NewLayout("InsightsMemoryProfilerLayout_v1.3")
 		->AddArea
 		(
 			FTabManager::NewPrimaryArea()
@@ -453,14 +453,14 @@ TSharedRef<FTabManager::FLayout> SMemoryProfilerWindow::CreateDefaultTabLayout()
 			->Split
 			(
 				FTabManager::NewStack()
-				->SetSizeCoefficient(0.77f)
+				->SetSizeCoefficient(0.65f)
 				->SetHideTabWell(true)
 				->AddTab(FMemoryProfilerTabs::TimingViewID, ETabState::OpenedTab)
 			)
 			->Split
 			(
 				FTabManager::NewStack()
-				->SetSizeCoefficient(0.23f)
+				->SetSizeCoefficient(0.35f)
 				->AddTab(FMemoryProfilerTabs::MemInvestigationViewID, ETabState::OpenedTab)
 				->AddTab(FMemoryProfilerTabs::MemTagTreeViewID, ETabState::OpenedTab)
 				->AddTab(FMemoryProfilerTabs::ModulesViewID, ETabState::OpenedTab)
diff --git a/Engine/Source/Developer/TraceInsights/Private/Insights/MemoryProfiler/Widgets/SMemoryProfilerWindow.h b/Engine/Source/Developer/TraceInsights/Private/Insights/MemoryProfiler/Widgets/SMemoryProfilerWindow.h
--- a/Engine/Source/Developer/TraceInsights/Private/Insights/MemoryProfiler/Widgets/SMemoryProfilerWindow.h
+++ b/Engine/Source/Developer/TraceInsights/Private/Insights/MemoryProfiler/Widgets/SMemoryProfilerWindow.h
@@ -107,7 +107,7 @@ private:
 	/** The Memory Investigation (Allocation Queries) view widget */
 	TSharedPtr<SMemInvestigationView> MemInvestigationView;
 
-	/** The "LLM Tags" tree view widget */
+	/** The Memory Tags tree view widget */
 	TSharedPtr<SMemTagTreeView> MemTagTreeView;
 
 	/** The list of Allocations table tree view widgets */
