[Insights]
* Timing View: Added toolitp with complete bookmark name when hovering the Markers (Bookmarks & Screenshots) track. If the Markers track shows all Log messages the tootlip will also show the category of the log message.
* Timing View: Fixed identification of a bookmark or screenshot. Now it correctly identify them even if their text is not displayed.

diff --git a/Engine/Source/Developer/TraceInsights/Private/Insights/TimingProfiler/Tracks/MarkersTimingTrack.cpp b/Engine/Source/Developer/TraceInsights/Private/Insights/TimingProfiler/Tracks/MarkersTimingTrack.cpp
--- a/Engine/Source/Developer/TraceInsights/Private/Insights/TimingProfiler/Tracks/MarkersTimingTrack.cpp
+++ b/Engine/Source/Developer/TraceInsights/Private/Insights/TimingProfiler/Tracks/MarkersTimingTrack.cpp
@@ -426,106 +426,164 @@ void FMarkersTimingTrack::UpdateCategory(const TraceServices::FLogCategoryInfo*&
 
 ////////////////////////////////////////////////////////////////////////////////////////////////////
 
-class FScreenShotTimingEvent : public FTimingEvent
+class FLogMessageTimingEvent : public FTimingEvent
 {
+	INSIGHTS_DECLARE_RTTI(FLogMessageTimingEvent, FTimingEvent)
+
 public:
-	FScreenShotTimingEvent(const TSharedRef<const FMarkersTimingTrack> InTrack, double InStartTime, double InEndTime, uint32 InDepth, int32 InScreenshotId)
-		: FTimingEvent(InTrack, InStartTime, InEndTime, InDepth, InScreenshotId) {}
-	INSIGHTS_DECLARE_RTTI(FScreenShotTimingEvent, FTimingEvent)
+	FLogMessageTimingEvent(const TSharedRef<const FMarkersTimingTrack> InTrack, double InTime, uint64 InLogIndex)
+		: FTimingEvent(InTrack, InTime, InTime, 0, InLogIndex)
+	{
+	}
+
+	uint64 GetLogIndex() const { return GetType(); }
 };
+
+INSIGHTS_IMPLEMENT_RTTI(FLogMessageTimingEvent)
+
+////////////////////////////////////////////////////////////////////////////////////////////////////
+
+class FScreenShotTimingEvent : public FLogMessageTimingEvent
+{
+	INSIGHTS_DECLARE_RTTI(FScreenShotTimingEvent, FLogMessageTimingEvent)
+
+public:
+	FScreenShotTimingEvent(const TSharedRef<const FMarkersTimingTrack> InTrack, double InTime, uint64 InLogIndex, uint32 InScreenshotId)
+		: FLogMessageTimingEvent(InTrack, InTime, InLogIndex)
+		, ScreenshotId(InScreenshotId)
+	{
+	}
+
+	uint32 GetScreenshotId() const { return ScreenshotId; }
+
+private:
+	uint32 ScreenshotId;
+};
+
 INSIGHTS_IMPLEMENT_RTTI(FScreenShotTimingEvent)
 
 ////////////////////////////////////////////////////////////////////////////////////////////////////
 
-class FBookmarkTimingEvent : public FTimingEvent
+class FBookmarkTimingEvent : public FLogMessageTimingEvent
 {
+	INSIGHTS_DECLARE_RTTI(FBookmarkTimingEvent, FLogMessageTimingEvent)
+
 public:
-	FBookmarkTimingEvent(const TSharedRef<const FMarkersTimingTrack> InTrack, double InStartTime, double InEndTime, uint32 InDepth, uint32 InCallstackId)
-		: FTimingEvent(InTrack, InStartTime, InEndTime, InDepth, InCallstackId) {}
-	INSIGHTS_DECLARE_RTTI(FBookmarkTimingEvent, FTimingEvent)
+	FBookmarkTimingEvent(const TSharedRef<const FMarkersTimingTrack> InTrack, double InTime, uint64 InLogIndex, uint32 InCallstackId)
+		: FLogMessageTimingEvent(InTrack, InTime, InLogIndex)
+		, CallstackId(InCallstackId)
+	{}
+
+	uint32 GetCallstackId() const { return CallstackId; }
+
+private:
+	uint32 CallstackId;
 };
+
 INSIGHTS_IMPLEMENT_RTTI(FBookmarkTimingEvent)
 
 ////////////////////////////////////////////////////////////////////////////////////////////////////
 
 const TSharedPtr<const ITimingEvent> FMarkersTimingTrack::GetEvent(float InPosX, float InPosY, const FTimingTrackViewport& Viewport) const
 {
-	TSharedPtr<ITimingEvent> TimingEvent;
-
 	const FTimingViewLayout& Layout = Viewport.GetLayout();
 
 	const float DY = InPosY - GetPosY();
 
 	if (DY >= 0 && DY < GetHeight())
 	{
-		const int32 NumBoxes = TimeMarkerTexts.Num();
-		int32 FoundIndex = Algo::LowerBoundBy(TimeMarkerTexts, InPosX, [](const FTimeMarkerTextInfo& Text) { return Text.X; });
-		if (FoundIndex > 0)
-		{
-			--FoundIndex;
-		}
-
-		if (FoundIndex < 0 || FoundIndex >= TimeMarkerTexts.Num())
-		{
-			return TimingEvent;
-		}
-
-		float Width = Viewport.GetWidth();
-		if (FoundIndex + 1 < TimeMarkerTexts.Num())
-		{
-			Width = TimeMarkerTexts[FoundIndex + 1].X;
-		}
-
-		const FTimeMarkerTextInfo& Text = TimeMarkerTexts[FoundIndex];
-
 		TSharedPtr<const TraceServices::IAnalysisSession> Session = FInsightsManager::Get()->GetSession();
 		if (!Session.IsValid())
 		{
-			return TimingEvent;
+			return nullptr;
 		}
 
+		const double PosTime = Viewport.SlateUnitsToTime(InPosX);
+
 		TraceServices::FAnalysisSessionReadScope SessionReadScope(*Session.Get());
 
-		double Time = std::numeric_limits<double>::quiet_NaN();
-		uint32 ScreenshotId = TraceServices::FScreenshot::InvalidScreenshotId;
 		const TraceServices::ILogProvider& LogProvider = TraceServices::ReadLogProvider(*Session.Get());
-		LogProvider.ReadMessage(Text.LogIndex,
-			[&ScreenshotId, &Time, this](const TraceServices::FLogMessageInfo& Message)
-			{
-				if (Message.Category == this->ScreenshotCategory)
+
+		uint64 LogIndex = 0;
+		const TraceServices::FLogCategoryInfo* Category = nullptr;
+		double Time = 0.0;
+		uint32 ScreenshotId = TraceServices::FScreenshot::InvalidScreenshotId;
+
+		if (IsBookmarksTrack())
+		{
+			//TODO: Reverse enumeration with filter; see FTimeMarkerTrackBuilder::AddLogMessage
+			LogProvider.EnumerateMessages(Viewport.GetStartTime(), PosTime,
+				[&LogIndex, &Category, &Time, &ScreenshotId, this](const TraceServices::FLogMessageInfo& Message)
 				{
-					check(Message.Line >= 0);
-					ScreenshotId = Message.Line;
-				}
-				if (Message.Category == this->BookmarkCategory)
+					if (Message.Category == this->BookmarkCategory)
+					{
+						LogIndex = Message.Index;
+						Category = Message.Category;
+						Time = Message.Time;
+					}
+					else if (Message.Category == this->ScreenshotCategory)
+					{
+						LogIndex = Message.Index;
+						Category = Message.Category;
+						Time = Message.Time;
+						check(Message.Line >= 0);
+						ScreenshotId = Message.Line;
+					}
+				});
+		}
+		else
+		{
+			LogIndex = LogProvider.UpperBoundByTime(PosTime);
+			if (LogIndex == 0 || LogIndex == ~0ull)
+			{
+				return nullptr;
+			}
+			--LogIndex;
+
+			if (!LogProvider.ReadMessage(LogIndex,
+				[&Category, &Time, &ScreenshotId, this](const TraceServices::FLogMessageInfo& Message)
 				{
+					Category = Message.Category;
 					Time = Message.Time;
-				}
-			});
+					if (Message.Category == this->ScreenshotCategory)
+					{
+						check(Message.Line >= 0);
+						ScreenshotId = Message.Line;
+					}
+				}))
+			{
+				return nullptr;
+			}
+		}
+
+		if (Category == nullptr)
+		{
+			return nullptr;
+		}
 
-		if (ScreenshotId != TraceServices::FScreenshot::InvalidScreenshotId)
+		if (Category == ScreenshotCategory)
 		{
-			return MakeShared<FScreenShotTimingEvent>(SharedThis(this), Viewport.SlateUnitsToTime(Text.X), Viewport.SlateUnitsToTime(Text.X + Width), 0, ScreenshotId);;
+			return MakeShared<FScreenShotTimingEvent>(SharedThis(this), Time, LogIndex, ScreenshotId);
 		}
 
-		if (FMath::IsFinite(Time))
+		if (Category == BookmarkCategory)
 		{
 			uint32 CallstackId = 0;
 			const TraceServices::IBookmarkProvider& BookmarkProvider = TraceServices::ReadBookmarkProvider(*Session.Get());
 			BookmarkProvider.EnumerateBookmarks(
 				Time - UE_SMALL_NUMBER,
 				Time + UE_SMALL_NUMBER,
-				[&CallstackId](const TraceServices::FBookmark& Bookmark) {
+				[&CallstackId](const TraceServices::FBookmark& Bookmark)
+				{
 					CallstackId = Bookmark.CallstackId;
-			});
-			if (CallstackId > 0)
-			{
-				return MakeShared<FBookmarkTimingEvent>(SharedThis(this), Viewport.SlateUnitsToTime(Text.X), Viewport.SlateUnitsToTime(Text.X + Width), 0, CallstackId);
-			}
+				});
+			return MakeShared<FBookmarkTimingEvent>(SharedThis(this), Time, LogIndex, CallstackId);
 		}
+
+		return MakeShared<FLogMessageTimingEvent>(SharedThis(this), Time, LogIndex);
 	}
 
-	return TimingEvent;
+	return nullptr;
 }
 
 ////////////////////////////////////////////////////////////////////////////////////////////////////
@@ -535,52 +593,83 @@ void FMarkersTimingTrack::InitTooltip(FTooltipDrawState& InOutTooltip, const ITi
 	InOutTooltip.ResetContent();
 	InOutTooltip.UpdateLayout();
 
+	TSharedPtr<const TraceServices::IAnalysisSession> Session = FInsightsManager::Get()->GetSession();
+	if (!Session.IsValid())
+	{
+		return;
+	}
+
 	if (!InTooltipEvent.CheckTrack(this))
 	{
 		return;
 	}
 
-	TSharedPtr<const TraceServices::IAnalysisSession> Session = FInsightsManager::Get()->GetSession();
-	if (!Session.IsValid())
+	if (!InTooltipEvent.Is<FLogMessageTimingEvent>())
 	{
 		return;
 	}
 
-	if (InTooltipEvent.Is<FScreenShotTimingEvent>())
+	const FLogMessageTimingEvent& LogMessageTimingEvent = InTooltipEvent.As<FLogMessageTimingEvent>();
 	{
 		TraceServices::FAnalysisSessionReadScope SessionReadScope(*Session.Get());
+
+		const TraceServices::ILogProvider& LogProvider = TraceServices::ReadLogProvider(*Session.Get());
+		LogProvider.ReadMessage(LogMessageTimingEvent.GetLogIndex(),
+			[&InOutTooltip, this](const TraceServices::FLogMessageInfo& Message)
+			{
+				const FLinearColor Color = FTimeMarkerTrackBuilder::GetColorByCategory(Message.Category->Name);
+				if (IsLogsTrack())
+				{
+					const TCHAR* CategoryName = Message.Category->Name != nullptr ? Message.Category->Name : TEXT("");
+					FStringView CategoryStr(CategoryName);
+					if (CategoryStr.StartsWith(TEXT("Log")))
+					{
+						// Strip the "Log" prefix.
+						CategoryStr.RightChopInline(3);
+					}
+					InOutTooltip.AddTextLine(CategoryStr, Color);
+				}
+				InOutTooltip.AddTextLine(Message.Message, Color);
+				InOutTooltip.UpdateLayout();
+			});
+	}
+
+	if (InTooltipEvent.Is<FScreenShotTimingEvent>())
+	{
 		const FScreenShotTimingEvent& Event = InTooltipEvent.As<FScreenShotTimingEvent>();
 
+		TraceServices::FAnalysisSessionReadScope SessionReadScope(*Session.Get());
+
 		const TraceServices::IScreenshotProvider& ScreenshotProvider = TraceServices::ReadScreenshotProvider(*Session.Get());
 		TSharedPtr<const TraceServices::FScreenshot> Screenshot = ScreenshotProvider.GetScreenshot((uint32)Event.GetType());
-	
+
 		if (!Screenshot.IsValid())
 		{
 			return;
 		}
-	
+
 		IImageWrapperModule& ImageWrapperModule = FModuleManager::LoadModuleChecked<IImageWrapperModule>(FName("ImageWrapper"));
 		FImage Image;
 		if (ImageWrapperModule.DecompressImage(Screenshot->Data.GetData(), Screenshot->Size, Image))
 		{
 			constexpr int32 MAX_WIDTH = 640;
 			constexpr int32 MAX_HEIGHT = 480;
-	
+
 			int32 ResizedX = Screenshot->Width;
 			int32 ResizedY = Screenshot->Height;
-	
+
 			if (ResizedX > MAX_WIDTH)
 			{
 				ResizedY = (ResizedY * MAX_WIDTH) / ResizedX;
 				ResizedX = MAX_WIDTH;
 			}
-	
+
 			if (ResizedY > MAX_HEIGHT)
 			{
 				ResizedX = (ResizedX * MAX_HEIGHT) / ResizedY;
 				ResizedY = MAX_HEIGHT;
 			}
-	
+
 			TSharedPtr<FSlateBrush> ImageBrush;
 			if (Screenshot->Width != ResizedX || Screenshot->Height != ResizedY)
 			{
@@ -593,27 +682,32 @@ void FMarkersTimingTrack::InitTooltip(FTooltipDrawState& InOutTooltip, const ITi
 				ImageBrush = FSlateDynamicImageBrush::CreateWithImageData(FName(Screenshot->Name), FVector2D(Screenshot->Width, Screenshot->Height), TArray<uint8>(Image.RawData));
 			}
 			InOutTooltip.SetImage(ImageBrush);
+			InOutTooltip.UpdateLayout();
 		}
-	
-		InOutTooltip.UpdateLayout();
 	}
+
 	if (InTooltipEvent.Is<FBookmarkTimingEvent>())
 	{
 		const FBookmarkTimingEvent& Event = InTooltipEvent.As<FBookmarkTimingEvent>();
-		if (const TraceServices::ICallstacksProvider* CallstackProvider = TraceServices::ReadCallstacksProvider(*Session.Get()))
+
+		if (Event.GetCallstackId() != 0 &&
+			Event.GetCallstackId() != TraceServices::FScreenshot::InvalidScreenshotId)
 		{
-			if (const TraceServices::FCallstack* Callstack = CallstackProvider->GetCallstack((uint32)Event.GetType()))
+			if (const TraceServices::ICallstacksProvider* CallstackProvider = TraceServices::ReadCallstacksProvider(*Session.Get()))
 			{
-				uint32 CallstackSize = Callstack->Num();
-				for (uint8 Idx = 0; Idx < CallstackSize; ++Idx)
+				if (const TraceServices::FCallstack* Callstack = CallstackProvider->GetCallstack(Event.GetCallstackId()))
 				{
-					TStringBuilder<512> CallstackText;
-					UE::Insights::FormatStackFrame(*Callstack->Frame(Idx), CallstackText, UE::Insights::EStackFrameFormatFlags::ModuleSymbolFileAndLine);
-					InOutTooltip.AddTextLine(CallstackText, FTooltipDrawState::DefaultValueColor);
+					uint32 CallstackSize = Callstack->Num();
+					for (uint8 Idx = 0; Idx < CallstackSize; ++Idx)
+					{
+						TStringBuilder<512> CallstackText;
+						UE::Insights::FormatStackFrame(*Callstack->Frame(Idx), CallstackText, UE::Insights::EStackFrameFormatFlags::ModuleSymbolFileAndLine);
+						InOutTooltip.AddTextLine(CallstackText, FTooltipDrawState::DefaultValueColor);
+					}
+					InOutTooltip.UpdateLayout();
 				}
 			}
 		}
-		InOutTooltip.UpdateLayout();
 	}
 }
 
@@ -675,37 +769,48 @@ void FMarkersTimingTrack::SaveScreenshot_Execute()
 
 ////////////////////////////////////////////////////////////////////////////////////////////////////
 
-bool FMarkersTimingTrack::TryGetHoveredEventScreenshotId(uint32& OutScreenshotId)
+const TSharedPtr<const ITimingEvent> FMarkersTimingTrack::TryGetHoveredEvent()
 {
-	OutScreenshotId = TraceServices::FScreenshot::InvalidScreenshotId;
-
 	TSharedPtr<STimingProfilerWindow> Window = FTimingProfilerManager::Get()->GetProfilerWindow();
 	if (!Window.IsValid())
 	{
-		return false;
+		return nullptr;
 	}
 
 	TSharedPtr<STimingView> TimingView = Window->GetTimingView();
 	if (!TimingView.IsValid())
 	{
-		return false;
+		return nullptr;
 	}
 
 	const TSharedPtr<const ITimingEvent> HoveredTimingEvent = TimingView->GetHoveredEvent();
 	if (!HoveredTimingEvent.IsValid())
 	{
-		return false;
+		return nullptr;
 	}
 
-	const FBaseTimingTrack& Track = HoveredTimingEvent->GetTrack().Get();
 	if (!HoveredTimingEvent->CheckTrack(this))
 	{
-		return false;
+		return nullptr;
+	}
+
+	return HoveredTimingEvent;
+}
+
+////////////////////////////////////////////////////////////////////////////////////////////////////
+
+bool FMarkersTimingTrack::TryGetHoveredEventScreenshotId(uint32& OutScreenshotId)
+{
+	const TSharedPtr<const ITimingEvent> Event = TryGetHoveredEvent();
+
+	if (Event.IsValid() && Event->Is<FScreenShotTimingEvent>())
+	{
+		OutScreenshotId = Event->As<FScreenShotTimingEvent>().GetScreenshotId();
+		return true;
 	}
 
-	const FTimingEvent& Event = StaticCast<const FTimingEvent&>(*HoveredTimingEvent);
-	OutScreenshotId = (uint32) Event.GetType();
-	return true;
+	OutScreenshotId = TraceServices::FScreenshot::InvalidScreenshotId;
+	return false;
 }
 
 ////////////////////////////////////////////////////////////////////////////////////////////////////
diff --git a/Engine/Source/Developer/TraceInsights/Private/Insights/TimingProfiler/Tracks/MarkersTimingTrack.h b/Engine/Source/Developer/TraceInsights/Private/Insights/TimingProfiler/Tracks/MarkersTimingTrack.h
--- a/Engine/Source/Developer/TraceInsights/Private/Insights/TimingProfiler/Tracks/MarkersTimingTrack.h
+++ b/Engine/Source/Developer/TraceInsights/Private/Insights/TimingProfiler/Tracks/MarkersTimingTrack.h
@@ -110,6 +110,7 @@ private:
 
 	void UpdateCategory(const TraceServices::FLogCategoryInfo*& InOutCategory, const TCHAR* CategoryName);
 
+	const TSharedPtr<const ITimingEvent> TryGetHoveredEvent();
 	bool TryGetHoveredEventScreenshotId(uint32& OutScreenshotId);
 
 private:
@@ -130,7 +131,10 @@ private:
 	// Slate resources
 	const FSlateBrush* WhiteBrush;
 	const FSlateFontInfo Font;
+
 	uint32 LastScreenshotId;
+	mutable FString HoveredCategory;
+	mutable FString HoveredMessage;
 };
 
 ////////////////////////////////////////////////////////////////////////////////////////////////////
