[Insights] Callers & Callees: Added the timer specific options (i.e. ones available in Timers panel) also in the Callers and Callees context menus: filter/highlight, plot graphs, find min/max instance, open source file in IDE.

diff --git a/Engine/Source/Developer/TraceInsights/Private/Insights/TimingProfiler/Widgets/STimerTreeView.cpp b/Engine/Source/Developer/TraceInsights/Private/Insights/TimingProfiler/Widgets/STimerTreeView.cpp
--- a/Engine/Source/Developer/TraceInsights/Private/Insights/TimingProfiler/Widgets/STimerTreeView.cpp
+++ b/Engine/Source/Developer/TraceInsights/Private/Insights/TimingProfiler/Widgets/STimerTreeView.cpp
@@ -7,6 +7,9 @@
 #include "Framework/Commands/UICommandList.h"
 #include "Framework/MultiBox/MultiBoxBuilder.h"
 #include "HAL/PlatformApplicationMisc.h"
+#include "ISourceCodeAccessModule.h"
+#include "ISourceCodeAccessor.h"
+#include "Modules/ModuleManager.h"
 #include "SlateOptMacros.h"
 #include "Styling/AppStyle.h"
 #include "Widgets/Layout/SScrollBox.h"
@@ -20,12 +23,16 @@
 
 // TraceInsights
 #include "Insights/InsightsStyle.h"
+#include "Insights/Table/ViewModels/TableCommands.h"
 #include "Insights/TimingProfiler/TimingProfilerManager.h"
 #include "Insights/TimingProfiler/ViewModels/TimerButterflyAggregator.h"
 #include "Insights/TimingProfiler/ViewModels/TimersViewColumnFactory.h"
+#include "Insights/TimingProfiler/Widgets/SFrameTrack.h"
 #include "Insights/TimingProfiler/Widgets/STimersViewTooltip.h"
 #include "Insights/TimingProfiler/Widgets/STimerTableRow.h"
-#include "Insights/Table/ViewModels/TableCommands.h"
+#include "Insights/TimingProfiler/Widgets/STimingProfilerWindow.h"
+#include "Insights/ViewModels/TimingGraphTrack.h"
+#include "Insights/Widgets/STimingView.h"
 
 #define LOCTEXT_NAMESPACE "UE::Insights::TimingProfiler::STimerTreeView"
 
@@ -67,11 +74,46 @@ public:
 			"Copies the name of the selected timer to the clipboard.",
 			EUserInterfaceActionType::Button,
 			FInputChord(EModifierKey::Control | EModifierKey::Shift, EKeys::C));
+
+		UI_COMMAND(Command_OpenSource,
+			"Open Source",
+			"Opens the source file of the selected timer in the registered IDE.",
+			EUserInterfaceActionType::Button,
+			FInputChord());
+
+		UI_COMMAND(Command_FindMaxInstance,
+			"Maximum Duration Instance",
+			"Navigates to and selects the timing event instance with the maximum duration, for the selected timer.",
+			EUserInterfaceActionType::Button,
+			FInputChord());
+
+		UI_COMMAND(Command_FindMinInstance,
+			"Minimum Duration Instance",
+			"Navigates to and selects the timing event instance with the minimum duration, for the selected timer.",
+			EUserInterfaceActionType::Button,
+			FInputChord());
+
+		UI_COMMAND(Command_FindMaxInstanceInSelection,
+			"Maximum Duration Instance in Selection",
+			"Navigates to and selects the timing event instance with the maximum duration, for the selected timer, in the selected time range.",
+			EUserInterfaceActionType::Button,
+			FInputChord());
+
+		UI_COMMAND(Command_FindMinInstanceInSelection,
+			"Minimum Duration Instance in Selection",
+			"Navigates to and selects the timing event instance with the minimum duration, for the selected timer, in the selected time range.",
+			EUserInterfaceActionType::Button,
+			FInputChord());
 	}
 	UE_ENABLE_OPTIMIZATION_SHIP
 
 	TSharedPtr<FUICommandInfo> Command_CopyToClipboard;
 	TSharedPtr<FUICommandInfo> Command_CopyNameToClipboard;
+	TSharedPtr<FUICommandInfo> Command_OpenSource;
+	TSharedPtr<FUICommandInfo> Command_FindMaxInstance;
+	TSharedPtr<FUICommandInfo> Command_FindMinInstance;
+	TSharedPtr<FUICommandInfo> Command_FindMaxInstanceInSelection;
+	TSharedPtr<FUICommandInfo> Command_FindMinInstanceInSelection;
 };
 
 ////////////////////////////////////////////////////////////////////////////////////////////////////
@@ -80,16 +122,6 @@ public:
 
 STimerTreeView::STimerTreeView()
 	: Table(MakeShared<FTable>())
-	, ViewName()
-	, TreeView(nullptr)
-	, TreeViewHeaderRow(nullptr)
-	, ExternalScrollbar(nullptr)
-	, HoveredColumnId()
-	, HoveredNodePtr(nullptr)
-	, HighlightedNodeName()
-	, TreeNodes()
-	, AvailableSorters()
-	, CurrentSorter(nullptr)
 	, ColumnBeingSorted(GetDefaultColumnBeingSorted())
 	, ColumnSortMode(GetDefaultColumnSortMode())
 {
@@ -99,6 +131,7 @@ STimerTreeView::STimerTreeView()
 
 STimerTreeView::~STimerTreeView()
 {
+	FTimerTreeViewCommands::Unregister();
 }
 
 ////////////////////////////////////////////////////////////////////////////////////////////////////
@@ -107,8 +140,13 @@ void STimerTreeView::InitCommandList()
 {
 	FTimerTreeViewCommands::Register();
 	CommandList = MakeShared<FUICommandList>();
-	CommandList->MapAction(FTimerTreeViewCommands::Get().Command_CopyToClipboard, FExecuteAction::CreateSP(this, &STimerTreeView::ContextMenu_CopySelectedToClipboard_Execute), FCanExecuteAction::CreateSP(this, &STimerTreeView::ContextMenu_CopySelectedToClipboard_CanExecute));
-	CommandList->MapAction(FTimerTreeViewCommands::Get().Command_CopyNameToClipboard, FExecuteAction::CreateSP(this, &STimerTreeView::ContextMenu_CopySelectedTimerNameToClipboard_Execute), FCanExecuteAction::CreateSP(this, &STimerTreeView::ContextMenu_CopySelectedTimerNameToClipboard_CanExecute));
+	CommandList->MapAction(FTimerTreeViewCommands::Get().Command_CopyToClipboard, FExecuteAction::CreateSP(this, &STimerTreeView::ContextMenu_CopyToClipboard_Execute), FCanExecuteAction::CreateSP(this, &STimerTreeView::ContextMenu_CopyToClipboard_CanExecute));
+	CommandList->MapAction(FTimerTreeViewCommands::Get().Command_CopyNameToClipboard, FExecuteAction::CreateSP(this, &STimerTreeView::ContextMenu_CopyTimerNameToClipboard_Execute), FCanExecuteAction::CreateSP(this, &STimerTreeView::ContextMenu_CopyTimerNameToClipboard_CanExecute));
+	CommandList->MapAction(FTimerTreeViewCommands::Get().Command_OpenSource, FExecuteAction::CreateSP(this, &STimerTreeView::ContextMenu_OpenSource_Execute), FCanExecuteAction::CreateSP(this, &STimerTreeView::ContextMenu_OpenSource_CanExecute));
+	CommandList->MapAction(FTimerTreeViewCommands::Get().Command_FindMaxInstance, FExecuteAction::CreateSP(this, &STimerTreeView::ContextMenu_FindInstance_Execute, true), FCanExecuteAction::CreateSP(this, &STimerTreeView::ContextMenu_FindInstance_CanExecute));
+	CommandList->MapAction(FTimerTreeViewCommands::Get().Command_FindMinInstance, FExecuteAction::CreateSP(this, &STimerTreeView::ContextMenu_FindInstance_Execute, false), FCanExecuteAction::CreateSP(this, &STimerTreeView::ContextMenu_FindInstance_CanExecute));
+	CommandList->MapAction(FTimerTreeViewCommands::Get().Command_FindMaxInstanceInSelection, FExecuteAction::CreateSP(this, &STimerTreeView::ContextMenu_FindInstanceInSelection_Execute, true), FCanExecuteAction::CreateSP(this, &STimerTreeView::ContextMenu_FindInstanceInSelection_CanExecute));
+	CommandList->MapAction(FTimerTreeViewCommands::Get().Command_FindMinInstanceInSelection, FExecuteAction::CreateSP(this, &STimerTreeView::ContextMenu_FindInstanceInSelection_Execute, false), FCanExecuteAction::CreateSP(this, &STimerTreeView::ContextMenu_FindInstanceInSelection_CanExecute));
 }
 
 ////////////////////////////////////////////////////////////////////////////////////////////////////
@@ -278,6 +316,126 @@ TSharedPtr<SWidget> STimerTreeView::TreeView_GetMenuContent()
 	}
 	MenuBuilder.EndSection();
 
+	// Timer options section
+	MenuBuilder.BeginSection("TimerOptions", LOCTEXT("ContextMenu_Section_TimerOptions", "Timer Options"));
+	{
+		TSharedPtr<STimingView> TimingView = GetTimingView();
+
+		auto CanExecute = [TimingView, NumSelectedNodes, SelectedNode]()
+		{
+			return TimingView.IsValid() && NumSelectedNodes == 1 && SelectedNode.IsValid() && SelectedNode->GetType() != ETimerNodeType::Group;
+		};
+
+		// Highlight event
+		{
+			FUIAction Action_ToggleHighlight;
+			Action_ToggleHighlight.CanExecuteAction = FCanExecuteAction::CreateLambda(CanExecute);
+			Action_ToggleHighlight.ExecuteAction = FExecuteAction::CreateSP(this, &STimerTreeView::ToggleTimingViewEventFilter, SelectedNode);
+
+			if (SelectedNode.IsValid() &&
+				SelectedNode->GetType() != ETimerNodeType::Group &&
+				TimingView.IsValid() &&
+				TimingView->IsFilterByEventType(SelectedNode->GetTimerId()))
+			{
+				MenuBuilder.AddMenuEntry
+				(
+					LOCTEXT("ContextMenu_StopHighlightEvent", "Stop Highlighting Event"),
+					LOCTEXT("ContextMenu_StopHighlightEvent_Desc", "Stops highlighting timing event instances for the selected timer."),
+					FSlateIcon(FAppStyle::Get().GetStyleSetName(), "Icons.Visible"),
+					Action_ToggleHighlight,
+					NAME_None,
+					EUserInterfaceActionType::Button
+				);
+			}
+			else
+			{
+				MenuBuilder.AddMenuEntry
+				(
+					LOCTEXT("ContextMenu_HighlightEvent", "Highlight Event"),
+					LOCTEXT("ContextMenu_HighlightEvent_Desc", "Highlights all timing event instances for the selected timer."),
+					FSlateIcon(FAppStyle::Get().GetStyleSetName(), "Icons.Visible"),
+					Action_ToggleHighlight,
+					NAME_None,
+					EUserInterfaceActionType::Button
+				);
+			}
+		}
+
+		MenuBuilder.AddSubMenu
+		(
+			LOCTEXT("ContextMenu_PlotTimer_SubMenu", "Plot Timer"),
+			LOCTEXT("ContextMenu_PlotTimer_SubMenu_Desc", "Options to add the timer series to graph or frame tracks."),
+			FNewMenuDelegate::CreateSP(this, &STimerTreeView::TreeView_BuildPlotTimerMenu),
+			false,
+			FSlateIcon(FAppStyle::Get().GetStyleSetName(), "Icons.AddGraphSeries")
+		);
+
+		MenuBuilder.AddSubMenu
+		(
+			LOCTEXT("ContextMenu_FindInstance_SubMenu", "Find Instance"),
+			LOCTEXT("ContextMenu_PlotInstance_SubMenu_Desc", "Find the instance of this timer with the minimum or maximum duration."),
+			FNewMenuDelegate::CreateSP(this, &STimerTreeView::TreeView_FindMenu),
+			false,
+			FSlateIcon(FAppStyle::Get().GetStyleSetName(), "Icons.FindInstance")
+		);
+
+		// Open Source in IDE
+		{
+			ISourceCodeAccessModule& SourceCodeAccessModule = FModuleManager::LoadModuleChecked<ISourceCodeAccessModule>("SourceCodeAccess");
+			ISourceCodeAccessor& SourceCodeAccessor = SourceCodeAccessModule.GetAccessor();
+
+			FString File;
+			uint32 Line = 0;
+			bool bIsValidSource = false;
+			if (NumSelectedNodes == 1 && SelectedNode.IsValid() && SelectedNode->GetType() != ETimerNodeType::Group)
+			{
+				bIsValidSource = SelectedNode->GetSourceFileAndLine(File, Line);
+			}
+
+			FText ItemLabel;
+			FText ItemToolTip;
+
+			if (SourceCodeAccessor.CanAccessSourceCode())
+			{
+				ItemLabel = FText::Format(LOCTEXT("ContextMenu_OpenSource", "Open Source in {0}"), SourceCodeAccessor.GetNameText());
+				if (bIsValidSource)
+				{
+					ItemToolTip = FText::Format(LOCTEXT("ContextMenu_OpenSource_Desc1", "Opens the source file of the selected timer in {0}.\n{1} ({2})"),
+						SourceCodeAccessor.GetNameText(),
+						FText::FromString(File),
+						FText::AsNumber(Line, &FNumberFormattingOptions::DefaultNoGrouping()));
+				}
+				else
+				{
+					ItemToolTip = FText::Format(LOCTEXT("ContextMenu_OpenSource_Desc2", "Opens the source file of the selected timer in {0}."),
+						SourceCodeAccessor.GetNameText());
+				}
+			}
+			else
+			{
+				ItemLabel = LOCTEXT("ContextMenu_OpenSourceNA", "Open Source");
+				if (bIsValidSource)
+				{
+					ItemToolTip = FText::Format(LOCTEXT("ContextMenu_OpenSourceNA_Desc1", "{0} ({1})\nSource Code Accessor is not available."),
+						FText::FromString(File),
+						FText::AsNumber(Line, &FNumberFormattingOptions::DefaultNoGrouping()));
+				}
+				else
+				{
+					ItemToolTip = LOCTEXT("ContextMenu_OpenSourceNA_Desc2", "Source Code Accessor is not available.");
+				}
+			}
+
+			MenuBuilder.AddMenuEntry(
+				FTimerTreeViewCommands::Get().Command_OpenSource,
+				NAME_None,
+				ItemLabel,
+				ItemToolTip,
+				FSlateIcon(SourceCodeAccessor.GetStyleSet(), SourceCodeAccessor.GetOpenIconName()));
+		}
+	}
+	MenuBuilder.EndSection();
+
 	MenuBuilder.BeginSection("Misc", LOCTEXT("ContextMenu_Section_Misc", "Miscellaneous"));
 	{
 		MenuBuilder.AddMenuEntry
@@ -404,6 +562,258 @@ void STimerTreeView::TreeView_BuildViewColumnMenu(FMenuBuilder& MenuBuilder)
 
 ////////////////////////////////////////////////////////////////////////////////////////////////////
 
+void STimerTreeView::TreeView_BuildPlotTimerMenu(FMenuBuilder& MenuBuilder)
+{
+	MenuBuilder.SetSearchable(false);
+
+	const TArray<FTimerNodePtr> SelectedNodes = TreeView->GetSelectedItems();
+	const int32 NumSelectedNodes = SelectedNodes.Num();
+	FTimerNodePtr SelectedNode = NumSelectedNodes ? SelectedNodes[0] : nullptr;
+
+	auto CanExecuteAddToGraphTrack = [NumSelectedNodes, SelectedNode]()
+	{
+		TSharedPtr<STimingProfilerWindow> Wnd = FTimingProfilerManager::Get()->GetProfilerWindow();
+		TSharedPtr<STimingView> TimingView = Wnd.IsValid() ? Wnd->GetTimingView() : nullptr;
+		return TimingView.IsValid() && NumSelectedNodes == 1 && SelectedNode.IsValid() && SelectedNode->GetType() != ETimerNodeType::Group;
+	};
+
+	auto CanExecuteAddToFramesTrack = [NumSelectedNodes, SelectedNode]()
+	{
+		TSharedPtr<STimingProfilerWindow> Wnd = FTimingProfilerManager::Get()->GetProfilerWindow();
+		TSharedPtr<SFrameTrack> FrameTrack = Wnd.IsValid() ? Wnd->GetFrameView() : nullptr;
+		return FrameTrack.IsValid() && NumSelectedNodes == 1 && SelectedNode.IsValid() && SelectedNode->GetType() != ETimerNodeType::Group;
+	};
+
+	MenuBuilder.BeginSection("Instance", LOCTEXT("Plot_Series_Instance_Section", "Instance"));
+
+	// Add/remove series to/from graph track
+	{
+		FUIAction Action_ToggleTimerInGraphTrack;
+		Action_ToggleTimerInGraphTrack.CanExecuteAction = FCanExecuteAction::CreateLambda(CanExecuteAddToGraphTrack);
+		Action_ToggleTimerInGraphTrack.ExecuteAction = FExecuteAction::CreateSP(this, &STimerTreeView::ToggleTimingViewMainGraphEventInstanceSeries, SelectedNode);
+
+		if (SelectedNode.IsValid() &&
+			SelectedNode->GetType() != ETimerNodeType::Group &&
+			IsInstanceSeriesInTimingViewMainGraph(SelectedNode))
+		{
+			MenuBuilder.AddMenuEntry
+			(
+				LOCTEXT("ContextMenu_RemoveFromGraphTrack", "Remove instance series from graph track"),
+				LOCTEXT("ContextMenu_RemoveFromGraphTrack_Desc", "Removes the series containing event instances of the selected timer from the Main Graph track."),
+				FSlateIcon(FInsightsStyle::GetStyleSetName(), "Icons.RemoveGraphSeries"),
+				Action_ToggleTimerInGraphTrack,
+				NAME_None,
+				EUserInterfaceActionType::Button
+			);
+		}
+		else
+		{
+			MenuBuilder.AddMenuEntry
+			(
+				LOCTEXT("ContextMenu_AddToGraphTrack", "Add instance series to graph track"),
+				LOCTEXT("ContextMenu_AddToGraphTrack_Desc", "Adds a series containing event instances of the selected timer to the Main Graph track."),
+				FSlateIcon(FInsightsStyle::GetStyleSetName(), "Icons.AddGraphSeries"),
+				Action_ToggleTimerInGraphTrack,
+				NAME_None,
+				EUserInterfaceActionType::Button
+			);
+		}
+	}
+
+	MenuBuilder.EndSection();
+
+	MenuBuilder.BeginSection("Game Frame", LOCTEXT("Plot_Series_GameFrame_Section", "Game Frame"));
+
+	// Add/remove game frame stats series to/from graph track
+	{
+		FUIAction Action_ToggleFrameStatsTimerInGraphTrack;
+		Action_ToggleFrameStatsTimerInGraphTrack.CanExecuteAction = FCanExecuteAction::CreateLambda(CanExecuteAddToGraphTrack);
+		Action_ToggleFrameStatsTimerInGraphTrack.ExecuteAction = FExecuteAction::CreateSP(this, &STimerTreeView::ToggleTimingViewMainGraphEventFrameStatsSeries, SelectedNode, ETraceFrameType::TraceFrameType_Game);
+
+		if (SelectedNode.IsValid() &&
+			SelectedNode->GetType() != ETimerNodeType::Group &&
+			IsFrameStatsSeriesInTimingViewMainGraph(SelectedNode, ETraceFrameType::TraceFrameType_Game))
+		{
+			MenuBuilder.AddMenuEntry
+			(
+				LOCTEXT("ContextMenu_RemoveGameFrameStatsFromGraphTrack", "Remove game frame stats series from graph track"),
+				LOCTEXT("ContextMenu_RemoveGameFrameStatsFromGraphTrack_Desc", "Remove the game frame stats series for this timer."),
+				FSlateIcon(FInsightsStyle::GetStyleSetName(), "Icons.RemoveGraphSeries"),
+				Action_ToggleFrameStatsTimerInGraphTrack,
+				NAME_None,
+				EUserInterfaceActionType::Button
+			);
+		}
+		else
+		{
+			MenuBuilder.AddMenuEntry
+			(
+				LOCTEXT("ContextMenu_AddGameFrameStatsSeriesToGraphTrack", "Add game frame stats series to graph track"),
+				LOCTEXT("ContextMenu_AddGameFrameStatsSeriesToGraphTrack_Desc", "Adds a game frame stats series for this timer. Each data entry is computed as the sum of all instances of this timer in a game frame."),
+				FSlateIcon(FInsightsStyle::GetStyleSetName(), "Icons.AddGraphSeries"),
+				Action_ToggleFrameStatsTimerInGraphTrack,
+				NAME_None,
+				EUserInterfaceActionType::Button
+			);
+		}
+	}
+
+	// Add/remove game frame stats series to/from frame track
+	{
+		FUIAction Action_ToggleFrameStatsTimerInFrameTrack;
+		Action_ToggleFrameStatsTimerInFrameTrack.CanExecuteAction = FCanExecuteAction::CreateLambda(CanExecuteAddToFramesTrack);
+		Action_ToggleFrameStatsTimerInFrameTrack.ExecuteAction = FExecuteAction::CreateSP(this, &STimerTreeView::ToggleFrameTrackSeries, SelectedNode, ETraceFrameType::TraceFrameType_Game);
+
+		if (SelectedNode.IsValid() &&
+			SelectedNode->GetType() != ETimerNodeType::Group &&
+			IsSeriesInFrameTrack(SelectedNode, ETraceFrameType::TraceFrameType_Game))
+		{
+			MenuBuilder.AddMenuEntry
+			(
+				LOCTEXT("ContextMenu_RemoveGameFrameStatsSeriesFromFrameTrack", "Remove game frame stats series from frame track"),
+				LOCTEXT("ContextMenu_RemoveGameFrameStatsSeriesFromFrameTrack_Desc", "Remove the game frame stats series for this timer from the frame track."),
+				FSlateIcon(FInsightsStyle::GetStyleSetName(), "Icons.RemoveGraphSeries"),
+				Action_ToggleFrameStatsTimerInFrameTrack,
+				NAME_None,
+				EUserInterfaceActionType::Button
+			);
+		}
+		else
+		{
+			MenuBuilder.AddMenuEntry
+			(
+				LOCTEXT("ContextMenu_AddGameFrameStatsSeriesToFrameTrack", "Add game frame stats series to the frame track"),
+				LOCTEXT("ContextMenu_AddGameFrameStatsSeriesToFrameTrack_Desc", "Adds a game frame stats series for this timer to the frame track. Each data entry is computed as the sum of all instances of this timer in a game frame."),
+				FSlateIcon(FInsightsStyle::GetStyleSetName(), "Icons.AddGraphSeries"),
+				Action_ToggleFrameStatsTimerInFrameTrack,
+				NAME_None,
+				EUserInterfaceActionType::Button
+			);
+		}
+	}
+
+	MenuBuilder.EndSection();
+
+	MenuBuilder.BeginSection("Rendering Frame", LOCTEXT("Plot_Series_RenderingFrame_Section", "Rendering frame"));
+
+	// Add/remove rendering frame stats series to/from graph track
+	{
+		FUIAction Action_ToggleFrameStatsTimerInGraphTrack;
+		Action_ToggleFrameStatsTimerInGraphTrack.CanExecuteAction = FCanExecuteAction::CreateLambda(CanExecuteAddToGraphTrack);
+		Action_ToggleFrameStatsTimerInGraphTrack.ExecuteAction = FExecuteAction::CreateSP(this, &STimerTreeView::ToggleTimingViewMainGraphEventFrameStatsSeries, SelectedNode, ETraceFrameType::TraceFrameType_Rendering);
+
+		if (SelectedNode.IsValid() &&
+			SelectedNode->GetType() != ETimerNodeType::Group &&
+			IsFrameStatsSeriesInTimingViewMainGraph(SelectedNode, ETraceFrameType::TraceFrameType_Rendering))
+		{
+			MenuBuilder.AddMenuEntry
+			(
+				LOCTEXT("ContextMenu_RemoveRenderingFrameStatsFromGraphTrack", "Remove rendering frame stats series from graph track"),
+				LOCTEXT("ContextMenu_RemoveRenderingFrameStatsFromGraphTrack_Desc", "Remove the rendering frame stats series for this timer."),
+				FSlateIcon(FInsightsStyle::GetStyleSetName(), "Icons.RemoveGraphSeries"),
+				Action_ToggleFrameStatsTimerInGraphTrack,
+				NAME_None,
+				EUserInterfaceActionType::Button
+			);
+		}
+		else
+		{
+			MenuBuilder.AddMenuEntry
+			(
+				LOCTEXT("ContextMenu_AddRenderingFrameStatsSeriesToGraphTrack", "Add rendering frame stats series to graph track"),
+				LOCTEXT("ContextMenu_AddRenderingFrameStatsSeriesToGraphTrack_Desc", "Adds a rendering frame stats series for this timer. Each data entry is computed as the sum of all instances of this timer in a rendering frame."),
+				FSlateIcon(FInsightsStyle::GetStyleSetName(), "Icons.AddGraphSeries"),
+				Action_ToggleFrameStatsTimerInGraphTrack,
+				NAME_None,
+				EUserInterfaceActionType::Button
+			);
+		}
+	}
+
+	// Add/remove rendering frame stats series to/from frame track
+	{
+		FUIAction Action_ToggleFrameStatsTimerInFrameTrack;
+		Action_ToggleFrameStatsTimerInFrameTrack.CanExecuteAction = FCanExecuteAction::CreateLambda(CanExecuteAddToFramesTrack);
+		Action_ToggleFrameStatsTimerInFrameTrack.ExecuteAction = FExecuteAction::CreateSP(this, &STimerTreeView::ToggleFrameTrackSeries, SelectedNode, ETraceFrameType::TraceFrameType_Rendering);
+
+		if (SelectedNode.IsValid() &&
+			SelectedNode->GetType() != ETimerNodeType::Group &&
+			IsSeriesInFrameTrack(SelectedNode, ETraceFrameType::TraceFrameType_Rendering))
+		{
+			MenuBuilder.AddMenuEntry
+			(
+				LOCTEXT("ContextMenu_RemoveRenderingFrameStatsFromFrameTrac", "Remove rendering frame stats series from frame track"),
+				LOCTEXT("ContextMenu_RemoveRenderingFrameStatsFromFrameTrack_Desc", "Remove the rendering frame stats series for this timer from the frame track."),
+				FSlateIcon(FInsightsStyle::GetStyleSetName(), "Icons.RemoveGraphSeries"),
+				Action_ToggleFrameStatsTimerInFrameTrack,
+				NAME_None,
+				EUserInterfaceActionType::Button
+			);
+		}
+		else
+		{
+			MenuBuilder.AddMenuEntry
+			(
+				LOCTEXT("ContextMenu_AddRenderingFrameStatsSeriesToFrameTrack", "Add rendering frame stats series to the frame track"),
+				LOCTEXT("ContextMenu_AddRenderingFrameStatsSeriesToFrameTrack_Desc", "Adds a rendering frame stats series for this timer to the frame track. Each data entry is computed as the sum of all instances of this timer in a rendering frame."),
+				FSlateIcon(FInsightsStyle::GetStyleSetName(), "Icons.AddGraphSeries"),
+				Action_ToggleFrameStatsTimerInFrameTrack,
+				NAME_None,
+				EUserInterfaceActionType::Button
+			);
+		}
+	}
+
+	MenuBuilder.EndSection();
+}
+
+////////////////////////////////////////////////////////////////////////////////////////////////////
+
+void STimerTreeView::TreeView_FindMenu(FMenuBuilder& MenuBuilder)
+{
+	MenuBuilder.SetSearchable(false);
+
+	MenuBuilder.AddMenuEntry
+	(
+		FTimerTreeViewCommands::Get().Command_FindMaxInstance,
+		NAME_None,
+		TAttribute<FText>(),
+		TAttribute<FText>(),
+		FSlateIcon(FAppStyle::Get().GetStyleSetName(), "Icons.FindMaxInstance")
+	);
+
+	MenuBuilder.AddMenuEntry
+	(
+		FTimerTreeViewCommands::Get().Command_FindMinInstance,
+		NAME_None,
+		TAttribute<FText>(),
+		TAttribute<FText>(),
+		FSlateIcon(FAppStyle::Get().GetStyleSetName(), "Icons.FindMinInstance")
+	);
+
+	MenuBuilder.AddMenuSeparator();
+
+	MenuBuilder.AddMenuEntry
+	(
+		FTimerTreeViewCommands::Get().Command_FindMaxInstanceInSelection,
+		NAME_None,
+		TAttribute<FText>(),
+		TAttribute<FText>(),
+		FSlateIcon(FAppStyle::Get().GetStyleSetName(), "Icons.FindMaxInstance")
+	);
+
+	MenuBuilder.AddMenuEntry
+	(
+		FTimerTreeViewCommands::Get().Command_FindMinInstanceInSelection,
+		NAME_None,
+		TAttribute<FText>(),
+		TAttribute<FText>(),
+		FSlateIcon(FAppStyle::Get().GetStyleSetName(), "Icons.FindMinInstance")
+	);
+}
+
+////////////////////////////////////////////////////////////////////////////////////////////////////
+
 void STimerTreeView::InitializeAndShowHeaderColumns()
 {
 	// Create columns.
@@ -573,6 +983,22 @@ void STimerTreeView::TreeView_OnSelectionChanged(FTimerNodePtr SelectedItem, ESe
 
 ////////////////////////////////////////////////////////////////////////////////////////////////////
 
+FTimerNodePtr STimerTreeView::GetSingleSelectedTimerNode() const
+{
+	if (TreeView->GetNumItemsSelected() != 1)
+	{
+		return nullptr;
+	}
+	const TArray<FTimerNodePtr> SelectedNodes = TreeView->GetSelectedItems();
+	if (SelectedNodes.Num() == 1)
+	{
+		return SelectedNodes[0];
+	}
+	return nullptr;
+}
+
+////////////////////////////////////////////////////////////////////////////////////////////////////
+
 void STimerTreeView::TreeView_OnGetChildren(FTimerNodePtr InParent, TArray<FTimerNodePtr>& OutChildren)
 {
 	constexpr bool bUseFiltering = false;
@@ -598,13 +1024,12 @@ void STimerTreeView::TreeView_OnGetChildren(FTimerNodePtr InParent, TArray<FTime
 
 ////////////////////////////////////////////////////////////////////////////////////////////////////
 
-void STimerTreeView::TreeView_OnMouseButtonDoubleClick(FTimerNodePtr TimerNodePtr)
+void STimerTreeView::TreeView_OnMouseButtonDoubleClick(FTimerNodePtr NodePtr)
 {
-	//if (TimerNodePtr->IsGroup())
-	if (TimerNodePtr->GetChildrenCount() > 0)
+	if (NodePtr->GetChildrenCount() > 0)
 	{
-		const bool bIsGroupExpanded = TreeView->IsItemExpanded(TimerNodePtr);
-		TreeView->SetItemExpansion(TimerNodePtr, !bIsGroupExpanded);
+		const bool bIsGroupExpanded = TreeView->IsItemExpanded(NodePtr);
+		TreeView->SetItemExpansion(NodePtr, !bIsGroupExpanded);
 	}
 }
 
@@ -612,7 +1037,7 @@ void STimerTreeView::TreeView_OnMouseButtonDoubleClick(FTimerNodePtr TimerNodePt
 // Tree View's Table Row
 ////////////////////////////////////////////////////////////////////////////////////////////////////
 
-TSharedRef<ITableRow> STimerTreeView::TreeView_OnGenerateRow(FTimerNodePtr TimerNodePtr, const TSharedRef<STableViewBase>& OwnerTable)
+TSharedRef<ITableRow> STimerTreeView::TreeView_OnGenerateRow(FTimerNodePtr NodePtr, const TSharedRef<STableViewBase>& OwnerTable)
 {
 	TSharedRef<ITableRow> TableRow =
 		SNew(STimerTableRow, OwnerTable)
@@ -623,7 +1048,7 @@ TSharedRef<ITableRow> STimerTreeView::TreeView_OnGenerateRow(FTimerNodePtr Timer
 		.HighlightText(this, &STimerTreeView::TableRow_GetHighlightText)
 		.HighlightedNodeName(this, &STimerTreeView::TableRow_GetHighlightedNodeName)
 		.TablePtr(Table)
-		.TimerNodePtr(TimerNodePtr);
+		.TimerNodePtr(NodePtr);
 
 	return TableRow;
 }
@@ -753,7 +1178,6 @@ void STimerTreeView::SortTreeNodesRec(FTimerNode& Node, const ITableCellValueSor
 
 	for (FBaseTreeNodePtr ChildPtr : Node.GetChildren())
 	{
-		//if (ChildPtr->IsGroup())
 		if (ChildPtr->GetChildrenCount() > 0)
 		{
 			SortTreeNodesRec(*StaticCastSharedPtr<FTimerNode>(ChildPtr), Sorter);
@@ -782,7 +1206,6 @@ void STimerTreeView::SetSortModeForColumn(const FName& ColumnId, const EColumnSo
 	UpdateCurrentSortingByColumn();
 
 	SortTreeNodes();
-	//ApplyFiltering();
 }
 
 ////////////////////////////////////////////////////////////////////////////////////////////////////
@@ -952,7 +1375,7 @@ void STimerTreeView::HideColumn(const FName ColumnId)
 // ToggleColumn action
 ////////////////////////////////////////////////////////////////////////////////////////////////////
 
-bool STimerTreeView::IsColumnVisible(const FName ColumnId)
+bool STimerTreeView::IsColumnVisible(const FName ColumnId) const
 {
 	const FTableColumn& Column = *Table->FindColumnChecked(ColumnId);
 	return Column.IsVisible();
@@ -1151,7 +1574,180 @@ void STimerTreeView::ExpandNodesRec(FTimerNodePtr NodePtr, int32 Depth)
 
 ////////////////////////////////////////////////////////////////////////////////////////////////////
 
-bool STimerTreeView::ContextMenu_CopySelectedToClipboard_CanExecute() const
+void STimerTreeView::ToggleTimingViewEventFilter(FTimerNodePtr TimerNode) const
+{
+	TSharedPtr<STimingView> TimingView = GetTimingView();
+	if (TimingView.IsValid())
+	{
+		const uint64 EventType = static_cast<uint64>(TimerNode->GetTimerId());
+		TimingView->ToggleEventFilterByEventType(EventType);
+	}
+}
+
+////////////////////////////////////////////////////////////////////////////////////////////////////
+
+TSharedPtr<FTimingGraphTrack> STimerTreeView::GetTimingViewMainGraphTrack() const
+{
+	TSharedPtr<STimingView> TimingView = GetTimingView();
+	return TimingView.IsValid() ? TimingView->GetMainTimingGraphTrack() : nullptr;
+}
+
+////////////////////////////////////////////////////////////////////////////////////////////////////
+
+TSharedPtr<SFrameTrack> STimerTreeView::GetFrameTrack() const
+{
+	TSharedPtr<STimingProfilerWindow> Wnd = FTimingProfilerManager::Get()->GetProfilerWindow();
+	return Wnd.IsValid() ? Wnd->GetFrameView() : nullptr;
+}
+
+////////////////////////////////////////////////////////////////////////////////////////////////////
+
+void STimerTreeView::ToggleGraphInstanceSeries(TSharedRef<FTimingGraphTrack> GraphTrack, FTimerNodeRef NodePtr) const
+{
+	const uint32 TimerId = NodePtr->GetTimerId();
+	TSharedPtr<FTimingGraphSeries> Series = GraphTrack->GetTimerSeries(TimerId);
+	if (Series.IsValid())
+	{
+		GraphTrack->RemoveTimerSeries(TimerId);
+		GraphTrack->SetDirtyFlag();
+		NodePtr->OnRemovedFromGraph();
+	}
+	else
+	{
+		GraphTrack->Show();
+		Series = GraphTrack->AddTimerSeries(TimerId, NodePtr->GetColor());
+		Series->SetName(FText::FromName(NodePtr->GetName()));
+		GraphTrack->SetDirtyFlag();
+		NodePtr->OnAddedToGraph();
+	}
+}
+
+////////////////////////////////////////////////////////////////////////////////////////////////////
+
+bool STimerTreeView::IsInstanceSeriesInTimingViewMainGraph(FTimerNodePtr TimerNode) const
+{
+	TSharedPtr<FTimingGraphTrack> GraphTrack = GetTimingViewMainGraphTrack();
+
+	if (GraphTrack.IsValid())
+	{
+		const uint32 TimerId = TimerNode->GetTimerId();
+		TSharedPtr<FTimingGraphSeries> Series = GraphTrack->GetTimerSeries(TimerId);
+
+		return Series.IsValid();
+	}
+
+	return false;
+}
+
+////////////////////////////////////////////////////////////////////////////////////////////////////
+
+void STimerTreeView::ToggleTimingViewMainGraphEventInstanceSeries(FTimerNodePtr TimerNode) const
+{
+	TSharedPtr<FTimingGraphTrack> GraphTrack = GetTimingViewMainGraphTrack();
+	if (GraphTrack.IsValid())
+	{
+		ToggleGraphInstanceSeries(GraphTrack.ToSharedRef(), TimerNode.ToSharedRef());
+	}
+}
+
+////////////////////////////////////////////////////////////////////////////////////////////////////
+
+void STimerTreeView::ToggleGraphFrameStatsSeries(TSharedRef<FTimingGraphTrack> GraphTrack, FTimerNodeRef NodePtr, ETraceFrameType FrameType) const
+{
+	const uint32 TimerId = NodePtr->GetTimerId();
+	TSharedPtr<FTimingGraphSeries> Series = GraphTrack->GetFrameStatsTimerSeries(TimerId, FrameType);
+	if (Series.IsValid())
+	{
+		GraphTrack->RemoveFrameStatsTimerSeries(TimerId, FrameType);
+		GraphTrack->SetDirtyFlag();
+		NodePtr->OnRemovedFromGraph();
+	}
+	else
+	{
+		GraphTrack->Show();
+		Series = GraphTrack->AddFrameStatsTimerSeries(TimerId, FrameType, NodePtr->GetColor());
+		FText SeriesName = FText::Format(LOCTEXT("FrameStatsTimerSeriesName_Fmt", "{0} ({1})"),
+			FText::FromName(NodePtr->GetName()),
+			FFrameTrackDrawHelper::FrameTypeToText(FrameType));
+		Series->SetName(SeriesName);
+		GraphTrack->SetDirtyFlag();
+		NodePtr->OnAddedToGraph();
+	}
+}
+
+////////////////////////////////////////////////////////////////////////////////////////////////////
+
+bool STimerTreeView::IsFrameStatsSeriesInTimingViewMainGraph(FTimerNodePtr TimerNode, ETraceFrameType FrameType) const
+{
+	TSharedPtr<FTimingGraphTrack> GraphTrack = GetTimingViewMainGraphTrack();
+
+	if (GraphTrack.IsValid())
+	{
+		const uint32 TimerId = TimerNode->GetTimerId();
+		TSharedPtr<FTimingGraphSeries> Series = GraphTrack->GetFrameStatsTimerSeries(TimerId, FrameType);
+
+		return Series.IsValid();
+	}
+
+	return false;
+}
+
+////////////////////////////////////////////////////////////////////////////////////////////////////
+
+void STimerTreeView::ToggleTimingViewMainGraphEventFrameStatsSeries(FTimerNodePtr TimerNode, ETraceFrameType FrameType) const
+{
+	TSharedPtr<FTimingGraphTrack> GraphTrack = GetTimingViewMainGraphTrack();
+	if (GraphTrack.IsValid())
+	{
+		ToggleGraphFrameStatsSeries(GraphTrack.ToSharedRef(), TimerNode.ToSharedRef(), FrameType);
+	}
+}
+
+////////////////////////////////////////////////////////////////////////////////////////////////////
+
+bool STimerTreeView::IsSeriesInFrameTrack(FTimerNodePtr TimerNode, ETraceFrameType FrameType) const
+{
+	TSharedPtr<SFrameTrack> FrameTrack = GetFrameTrack();
+
+	if (!FrameTrack.IsValid())
+	{
+		return false;
+	}
+
+	return FrameTrack->HasFrameStatSeries(FrameType, TimerNode->GetTimerId());
+}
+
+////////////////////////////////////////////////////////////////////////////////////////////////////
+
+void STimerTreeView::ToggleFrameTrackSeries(FTimerNodePtr TimerNode, ETraceFrameType FrameType) const
+{
+	TSharedPtr<SFrameTrack> FrameTrack = GetFrameTrack();
+
+	if (!FrameTrack.IsValid())
+	{
+		return;
+	}
+
+	if (FrameTrack->HasFrameStatSeries(FrameType, TimerNode->GetTimerId()))
+	{
+		FrameTrack->RemoveTimerFrameStatSeries(FrameType, TimerNode->GetTimerId());
+		TimerNode->OnRemovedFromGraph();
+	}
+	else
+	{
+		FText SeriesName = FText::Format(LOCTEXT("FrameStatsTimerSeriesName_Fmt", "{0} ({1})"),
+			TimerNode->GetDisplayName(),
+			FFrameTrackDrawHelper::FrameTypeToText(FrameType));
+		TSharedPtr<FTimerFrameStatsTrackSeries> Series = FrameTrack->AddTimerFrameStatSeries(FrameType, TimerNode->GetTimerId(), TimerNode->GetColor(), SeriesName);
+		TimerNode->OnAddedToGraph();
+	}
+}
+
+////////////////////////////////////////////////////////////////////////////////////////////////////
+// Copy to Clipboard
+////////////////////////////////////////////////////////////////////////////////////////////////////
+
+bool STimerTreeView::ContextMenu_CopyToClipboard_CanExecute() const
 {
 	const TArray<FTimerNodePtr> SelectedNodes = TreeView->GetSelectedItems();
 
@@ -1160,7 +1756,7 @@ bool STimerTreeView::ContextMenu_CopySelectedToClipboard_CanExecute() const
 
 ////////////////////////////////////////////////////////////////////////////////////////////////////
 
-void STimerTreeView::ContextMenu_CopySelectedToClipboard_Execute()
+void STimerTreeView::ContextMenu_CopyToClipboard_Execute()
 {
 	if (Table->IsValid())
 	{
@@ -1171,14 +1767,14 @@ void STimerTreeView::ContextMenu_CopySelectedToClipboard_Execute()
 
 ////////////////////////////////////////////////////////////////////////////////////////////////////
 
-bool STimerTreeView::ContextMenu_CopySelectedTimerNameToClipboard_CanExecute() const
+bool STimerTreeView::ContextMenu_CopyTimerNameToClipboard_CanExecute() const
 {
 	return TreeView->GetSelectedItems().Num() > 0;
 }
 
 ////////////////////////////////////////////////////////////////////////////////////////////////////
 
-void STimerTreeView::ContextMenu_CopySelectedTimerNameToClipboard_Execute()
+void STimerTreeView::ContextMenu_CopyTimerNameToClipboard_Execute()
 {
 	if (Table->IsValid())
 	{
@@ -1187,6 +1783,148 @@ void STimerTreeView::ContextMenu_CopySelectedTimerNameToClipboard_Execute()
 	}
 }
 
+////////////////////////////////////////////////////////////////////////////////////////////////////
+// Open Source File in IDE
+////////////////////////////////////////////////////////////////////////////////////////////////////
+
+bool STimerTreeView::ContextMenu_OpenSource_CanExecute() const
+{
+	ISourceCodeAccessModule& SourceCodeAccessModule = FModuleManager::LoadModuleChecked<ISourceCodeAccessModule>("SourceCodeAccess");
+	ISourceCodeAccessor& SourceCodeAccessor = SourceCodeAccessModule.GetAccessor();
+
+	if (!SourceCodeAccessor.CanAccessSourceCode())
+	{
+		return false;
+	}
+
+	FTimerNodePtr SelectedNode = GetSingleSelectedTimerNode();
+	if (!SelectedNode.IsValid())
+	{
+		return false;
+	}
+
+	FString File;
+	uint32 Line = 0;
+	return SelectedNode->GetSourceFileAndLine(File, Line);
+}
+
+////////////////////////////////////////////////////////////////////////////////////////////////////
+
+void STimerTreeView::ContextMenu_OpenSource_Execute() const
+{
+	FTimerNodePtr SelectedNode = GetSingleSelectedTimerNode();
+	if (SelectedNode.IsValid())
+	{
+		OpenSourceFileInIDE(SelectedNode);
+	}
+}
+
+////////////////////////////////////////////////////////////////////////////////////////////////////
+
+void STimerTreeView::OpenSourceFileInIDE(FTimerNodePtr InNode) const
+{
+	if (!InNode.IsValid() || InNode->GetType() == ETimerNodeType::Group)
+	{
+		return;
+	}
+
+	FString File;
+	uint32 Line = 0;
+	if (!InNode->GetSourceFileAndLine(File, Line))
+	{
+		return;
+	}
+
+	ISourceCodeAccessModule& SourceCodeAccessModule = FModuleManager::LoadModuleChecked<ISourceCodeAccessModule>("SourceCodeAccess");
+	if (FPaths::FileExists(File))
+	{
+		ISourceCodeAccessor& SourceCodeAccessor = SourceCodeAccessModule.GetAccessor();
+		SourceCodeAccessor.OpenFileAtLine(File, Line);
+	}
+	else
+	{
+		SourceCodeAccessModule.OnOpenFileFailed().Broadcast(File);
+	}
+}
+
+////////////////////////////////////////////////////////////////////////////////////////////////////
+// Find Min/Max Instance
+////////////////////////////////////////////////////////////////////////////////////////////////////
+
+bool STimerTreeView::ContextMenu_FindInstance_CanExecute() const
+{
+	FTimerNodePtr SelectedNode = GetSingleSelectedTimerNode();
+	return SelectedNode.IsValid() && SelectedNode->GetType() != ETimerNodeType::Group;
+}
+
+////////////////////////////////////////////////////////////////////////////////////////////////////
+
+void STimerTreeView::ContextMenu_FindInstance_Execute(bool bFindMax) const
+{
+	FTimerNodePtr SelectedNode = GetSingleSelectedTimerNode();
+	if (!SelectedNode.IsValid())
+	{
+		return;
+	}
+
+	TSharedPtr<STimingView> TimingView = GetTimingView();
+	if (!TimingView.IsValid())
+	{
+		return;
+	}
+
+	ESelectEventType Type = bFindMax ? ESelectEventType::Max : ESelectEventType::Min;
+	TimingView->SelectEventInstance(SelectedNode->GetTimerId(), Type, false);
+}
+
+////////////////////////////////////////////////////////////////////////////////////////////////////
+
+bool STimerTreeView::ContextMenu_FindInstanceInSelection_CanExecute() const
+{
+	FTimerNodePtr SelectedNode = GetSingleSelectedTimerNode();
+	if (!SelectedNode.IsValid() || SelectedNode->GetType() == ETimerNodeType::Group)
+	{
+		return false;
+	}
+
+	TSharedPtr<STimingView> TimingView = GetTimingView();
+	if (TimingView.IsValid())
+	{
+		return TimingView->GetSelectionEndTime() > TimingView->GetSelectionStartTime();
+	}
+
+	return false;
+}
+
+////////////////////////////////////////////////////////////////////////////////////////////////////
+
+void STimerTreeView::ContextMenu_FindInstanceInSelection_Execute(bool bFindMax) const
+{
+	FTimerNodePtr SelectedNode = GetSingleSelectedTimerNode();
+	if (!SelectedNode.IsValid() || SelectedNode->GetType() == ETimerNodeType::Group)
+	{
+		return;
+	}
+
+	TSharedPtr<STimingView> TimingView = GetTimingView();
+	if (!TimingView.IsValid())
+	{
+		return;
+	}
+
+	ESelectEventType Type = bFindMax ? ESelectEventType::Max : ESelectEventType::Min;
+	TimingView->SelectEventInstance(SelectedNode->GetTimerId(), Type, true);
+}
+
+////////////////////////////////////////////////////////////////////////////////////////////////////
+////////////////////////////////////////////////////////////////////////////////////////////////////
+
+TSharedPtr<STimingView> STimerTreeView::GetTimingView() const
+{
+	TSharedPtr<STimingProfilerWindow> ProfilerWindow = FTimingProfilerManager::Get()->GetProfilerWindow();
+	return ProfilerWindow.IsValid() ? ProfilerWindow->GetTimingView() : nullptr;
+}
+
 ////////////////////////////////////////////////////////////////////////////////////////////////////
 
 FReply STimerTreeView::OnKeyDown(const FGeometry& MyGeometry, const FKeyEvent& InKeyEvent)
diff --git a/Engine/Source/Developer/TraceInsights/Private/Insights/TimingProfiler/Widgets/STimerTreeView.h b/Engine/Source/Developer/TraceInsights/Private/Insights/TimingProfiler/Widgets/STimerTreeView.h
--- a/Engine/Source/Developer/TraceInsights/Private/Insights/TimingProfiler/Widgets/STimerTreeView.h
+++ b/Engine/Source/Developer/TraceInsights/Private/Insights/TimingProfiler/Widgets/STimerTreeView.h
@@ -31,9 +31,13 @@ namespace UE::Insights
 	class SAsyncOperationStatus;
 }
 
+class FTimingGraphTrack;
+
 namespace UE::Insights::TimingProfiler
 {
 
+class SFrameTrack;
+class STimingView;
 ////////////////////////////////////////////////////////////////////////////////////////////////////
 
 /**
@@ -74,11 +78,38 @@ private:
 	void TreeView_BuildSortByMenu(FMenuBuilder& MenuBuilder);
 	void TreeView_BuildViewColumnMenu(FMenuBuilder& MenuBuilder);
 
-	bool ContextMenu_CopySelectedToClipboard_CanExecute() const;
-	void ContextMenu_CopySelectedToClipboard_Execute();
+	//////////////////////////////////////////////////
+	// Copy to Clipboard
+
+	bool ContextMenu_CopyToClipboard_CanExecute() const;
+	void ContextMenu_CopyToClipboard_Execute();
+
+	bool ContextMenu_CopyTimerNameToClipboard_CanExecute() const;
+	void ContextMenu_CopyTimerNameToClipboard_Execute();
+
+	//////////////////////////////////////////////////
+	// Open Source File in IDE
+
+	bool ContextMenu_OpenSource_CanExecute() const;
+	void ContextMenu_OpenSource_Execute() const;
+
+	void OpenSourceFileInIDE(FTimerNodePtr InNode) const;
+
+	//////////////////////////////////////////////////
+	// Plot (Graph Series)
+
+	void TreeView_BuildPlotTimerMenu(FMenuBuilder& MenuBuilder);
+
+	//////////////////////////////////////////////////
+	// Find Instance
+
+	void TreeView_FindMenu(FMenuBuilder& MenuBuilder);
 
-	bool ContextMenu_CopySelectedTimerNameToClipboard_CanExecute() const;
-	void ContextMenu_CopySelectedTimerNameToClipboard_Execute();
+	bool ContextMenu_FindInstance_CanExecute() const;
+	void ContextMenu_FindInstance_Execute(bool bFindMax) const;
+
+	bool ContextMenu_FindInstanceInSelection_CanExecute() const;
+	void ContextMenu_FindInstanceInSelection_Execute(bool bFindMax) const;
 
 	////////////////////////////////////////////////////////////////////////////////////////////////////
 	// Tree View - Columns' Header
@@ -104,6 +135,8 @@ private:
 	/** Called by STreeView when selection has changed. */
 	void TreeView_OnSelectionChanged(FTimerNodePtr SelectedItem, ESelectInfo::Type SelectInfo);
 
+	FTimerNodePtr GetSingleSelectedTimerNode() const;
+
 	/** Called by STreeView when a tree item is double clicked. */
 	void TreeView_OnMouseButtonDoubleClick(FTimerNodePtr TreeNode);
 
@@ -138,7 +171,7 @@ private:
 	void OnSortModeChanged(const EColumnSortPriority::Type SortPriority, const FName& ColumnId, const EColumnSortMode::Type SortMode);
 
 	////////////////////////////////////////////////////////////////////////////////////////////////////
-	// Sorting actions
+	// Sorting Actions
 
 	// SortMode (HeaderMenu)
 	bool HeaderMenu_SortMode_IsChecked(const FName ColumnId, const EColumnSortMode::Type InSortMode);
@@ -156,7 +189,7 @@ private:
 	void ContextMenu_SortByColumn_Execute(const FName ColumnId);
 
 	////////////////////////////////////////////////////////////////////////////////////////////////////
-	// Column visibility actions
+	// Column Visibility Actions
 
 	// ShowColumn
 	bool CanShowColumn(const FName ColumnId) const;
@@ -167,7 +200,7 @@ private:
 	void HideColumn(const FName ColumnId);
 
 	// ToggleColumnVisibility
-	bool IsColumnVisible(const FName ColumnId);
+	bool IsColumnVisible(const FName ColumnId) const;
 	bool CanToggleColumnVisibility(const FName ColumnId) const;
 	void ToggleColumnVisibility(const FName ColumnId);
 
@@ -179,8 +212,32 @@ private:
 	bool ContextMenu_ResetColumns_CanExecute() const;
 	void ContextMenu_ResetColumns_Execute();
 
+	////////////////////////////////////////////////////////////////////////////////////////////////////
+	// Event Highlight/Filtering in Timing View
+
+	void ToggleTimingViewEventFilter(FTimerNodePtr TimerNode) const;
+
+	////////////////////////////////////////////////////////////////////////////////////////////////////
+	// Plot (Graph Series)
+
+	TSharedPtr<FTimingGraphTrack> GetTimingViewMainGraphTrack() const;
+	TSharedPtr<SFrameTrack> GetFrameTrack() const;
+
+	void ToggleGraphInstanceSeries(TSharedRef<FTimingGraphTrack> GraphTrack, FTimerNodeRef NodePtr) const;
+	bool IsInstanceSeriesInTimingViewMainGraph(FTimerNodePtr TimerNode) const;
+	void ToggleTimingViewMainGraphEventInstanceSeries(FTimerNodePtr TimerNode) const;
+
+	void ToggleGraphFrameStatsSeries(TSharedRef<FTimingGraphTrack> GraphTrack, FTimerNodeRef NodePtr, ETraceFrameType FrameType) const;
+	bool IsFrameStatsSeriesInTimingViewMainGraph(FTimerNodePtr TimerNode, ETraceFrameType FrameType) const;
+	void ToggleTimingViewMainGraphEventFrameStatsSeries(FTimerNodePtr TimerNode, ETraceFrameType FrameType) const;
+
+	bool IsSeriesInFrameTrack(FTimerNodePtr TimerNode, ETraceFrameType FrameType) const;
+	void ToggleFrameTrackSeries(FTimerNodePtr TimerNode, ETraceFrameType FrameType) const;
+
 	////////////////////////////////////////////////////////////////////////////////////////////////////
 
+	TSharedPtr<STimingView> GetTimingView() const;
+
 private:
 	/** Table view model. */
 	TSharedPtr<FTable> Table;
@@ -233,7 +290,9 @@ private:
 	FName ColumnBeingSorted;
 
 	/** How we sort the nodes? Ascending or Descending. */
-	EColumnSortMode::Type ColumnSortMode;
+	EColumnSortMode::Type ColumnSortMode = EColumnSortMode::Type::Ascending;
+
+	//////////////////////////////////////////////////
 
 	TSharedPtr<SAsyncOperationStatus> AsyncOperationStatus;
 };
diff --git a/Engine/Source/Developer/TraceInsights/Private/Insights/TimingProfiler/Widgets/STimersView.cpp b/Engine/Source/Developer/TraceInsights/Private/Insights/TimingProfiler/Widgets/STimersView.cpp
--- a/Engine/Source/Developer/TraceInsights/Private/Insights/TimingProfiler/Widgets/STimersView.cpp
+++ b/Engine/Source/Developer/TraceInsights/Private/Insights/TimingProfiler/Widgets/STimersView.cpp
@@ -73,7 +73,7 @@ public:
 		TEXT("TimersViewCommands"),
 		NSLOCTEXT("Contexts", "TimersViewCommands", "Insights - Timers View"),
 		NAME_None,
-		FInsightsStyle::Get().GetStyleSetName())
+		FInsightsStyle::GetStyleSetName())
 	{
 	}
 
@@ -167,7 +167,6 @@ public:
 	TSharedPtr<FUICommandInfo> Command_ExportThreads;
 	TSharedPtr<FUICommandInfo> Command_ExportTimers;
 	TSharedPtr<FUICommandInfo> Command_OpenSource;
-
 	TSharedPtr<FUICommandInfo> Command_FindMaxInstance;
 	TSharedPtr<FUICommandInfo> Command_FindMinInstance;
 	TSharedPtr<FUICommandInfo> Command_FindMaxInstanceInSelection;
@@ -180,11 +179,6 @@ public:
 
 STimersView::STimersView()
 	: Table(MakeShared<FTable>())
-	, bExpansionSaved(false)
-	, bFilterOutZeroCountTimers(false)
-	, GroupingMode(ETimerGroupingMode::ByType)
-	, AvailableSorters()
-	, CurrentSorter(nullptr)
 	, ColumnBeingSorted(GetDefaultColumnBeingSorted())
 	, ColumnSortMode(GetDefaultColumnSortMode())
 	, Aggregator(MakeShared<FTimerAggregator>())
@@ -1556,8 +1550,7 @@ FTimerNodePtr STimersView::GetSingleSelectedTimerNode() const
 		return nullptr;
 	}
 	const TArray<FTimerNodePtr> SelectedNodes = TreeView->GetSelectedItems();
-	const int32 NumSelectedNodes = SelectedNodes.Num();
-	if (NumSelectedNodes == 1)
+	if (SelectedNodes.Num() == 1)
 	{
 		return SelectedNodes[0];
 	}
@@ -2893,6 +2886,8 @@ void STimersView::ContextMenu_CopyTimerNameToClipboard_Execute()
 	}
 }
 
+////////////////////////////////////////////////////////////////////////////////////////////////////
+// Export
 ////////////////////////////////////////////////////////////////////////////////////////////////////
 
 bool STimersView::ContextMenu_Export_CanExecute() const
@@ -3315,6 +3310,8 @@ IFileHandle* STimersView::OpenExportFile(const TCHAR* InFilename) const
 	return ExportFileHandle;
 }
 
+////////////////////////////////////////////////////////////////////////////////////////////////////
+// OpenSource
 ////////////////////////////////////////////////////////////////////////////////////////////////////
 
 bool STimersView::ContextMenu_OpenSource_CanExecute() const
@@ -3351,6 +3348,36 @@ void STimersView::ContextMenu_OpenSource_Execute() const
 
 ////////////////////////////////////////////////////////////////////////////////////////////////////
 
+void STimersView::OpenSourceFileInIDE(FTimerNodePtr InNode) const
+{
+	if (!InNode.IsValid() || InNode->GetType() == ETimerNodeType::Group)
+	{
+		return;
+	}
+
+	FString File;
+	uint32 Line = 0;
+	if (!InNode->GetSourceFileAndLine(File, Line))
+	{
+		return;
+	}
+
+	ISourceCodeAccessModule& SourceCodeAccessModule = FModuleManager::LoadModuleChecked<ISourceCodeAccessModule>("SourceCodeAccess");
+	if (FPaths::FileExists(File))
+	{
+		ISourceCodeAccessor& SourceCodeAccessor = SourceCodeAccessModule.GetAccessor();
+		SourceCodeAccessor.OpenFileAtLine(File, Line);
+	}
+	else
+	{
+		SourceCodeAccessModule.OnOpenFileFailed().Broadcast(File);
+	}
+}
+
+////////////////////////////////////////////////////////////////////////////////////////////////////
+// Find Min/Max Instance
+////////////////////////////////////////////////////////////////////////////////////////////////////
+
 bool STimersView::ContextMenu_FindInstance_CanExecute() const
 {
 	FTimerNodePtr SelectedNode = GetSingleSelectedTimerNode();
@@ -3417,51 +3444,7 @@ void STimersView::ContextMenu_FindInstanceInSelection_Execute(bool bFindMax) con
 }
 
 ////////////////////////////////////////////////////////////////////////////////////////////////////
-
-void STimersView::OpenSourceFileInIDE(FTimerNodePtr InNode) const
-{
-	if (!InNode.IsValid() || InNode->GetType() == ETimerNodeType::Group)
-	{
-		return;
-	}
-
-	FString File;
-	uint32 Line = 0;
-	if (!InNode->GetSourceFileAndLine(File, Line))
-	{
-		return;
-	}
-
-	ISourceCodeAccessModule& SourceCodeAccessModule = FModuleManager::LoadModuleChecked<ISourceCodeAccessModule>("SourceCodeAccess");
-	if (FPaths::FileExists(File))
-	{
-		ISourceCodeAccessor& SourceCodeAccessor = SourceCodeAccessModule.GetAccessor();
-		SourceCodeAccessor.OpenFileAtLine(File, Line);
-	}
-	else
-	{
-		SourceCodeAccessModule.OnOpenFileFailed().Broadcast(File);
-	}
-}
-
-////////////////////////////////////////////////////////////////////////////////////////////////////
-
-FReply STimersView::OnKeyDown(const FGeometry& MyGeometry, const FKeyEvent& InKeyEvent)
-{
-	return CommandList->ProcessCommandBindings(InKeyEvent) == true ? FReply::Handled() : FReply::Unhandled();
-}
-
-////////////////////////////////////////////////////////////////////////////////////////////////////
-
-void STimersView::SetTimingViewFrameType()
-{
-	TSharedPtr<STimingView> TimingView = GetTimingView();
-	if (TimingView.IsValid())
-	{
-		TimingView->SelectTimeInterval(TimingView->GetSelectionStartTime(), TimingView->GetSelectionEndTime() - TimingView->GetSelectionStartTime());
-	}
-}
-
+// Persistent Settings
 ////////////////////////////////////////////////////////////////////////////////////////////////////
 
 void STimersView::SaveVisibleColumnsSettings()
@@ -3571,6 +3554,18 @@ void STimersView::LoadSettings()
 	LoadVisibleColumnsSettings();
 }
 
+////////////////////////////////////////////////////////////////////////////////////////////////////
+////////////////////////////////////////////////////////////////////////////////////////////////////
+
+void STimersView::SetTimingViewFrameType()
+{
+	TSharedPtr<STimingView> TimingView = GetTimingView();
+	if (TimingView.IsValid())
+	{
+		TimingView->SelectTimeInterval(TimingView->GetSelectionStartTime(), TimingView->GetSelectionEndTime() - TimingView->GetSelectionStartTime());
+	}
+}
+
 ////////////////////////////////////////////////////////////////////////////////////////////////////
 
 void STimersView::OnTimingViewTrackListChanged()
@@ -3593,6 +3588,13 @@ TSharedPtr<STimingView> STimersView::GetTimingView() const
 
 ////////////////////////////////////////////////////////////////////////////////////////////////////
 
+FReply STimersView::OnKeyDown(const FGeometry& MyGeometry, const FKeyEvent& InKeyEvent)
+{
+	return CommandList->ProcessCommandBindings(InKeyEvent) == true ? FReply::Handled() : FReply::Unhandled();
+}
+
+////////////////////////////////////////////////////////////////////////////////////////////////////
+
 } // namespace UE::Insights::TimingProfiler
 
 #undef LOCTEXT_NAMESPACE
diff --git a/Engine/Source/Developer/TraceInsights/Private/Insights/TimingProfiler/Widgets/STimersView.h b/Engine/Source/Developer/TraceInsights/Private/Insights/TimingProfiler/Widgets/STimersView.h
--- a/Engine/Source/Developer/TraceInsights/Private/Insights/TimingProfiler/Widgets/STimersView.h
+++ b/Engine/Source/Developer/TraceInsights/Private/Insights/TimingProfiler/Widgets/STimersView.h
@@ -88,14 +88,14 @@ public:
 	FTimerNodePtr GetTimerNode(uint32 TimerId) const;
 	void SelectTimerNode(uint32 TimerId);
 
-	virtual FReply OnKeyDown(const FGeometry& MyGeometry, const FKeyEvent& InKeyEvent) override;
-
 	ETraceFrameType GetFrameTypeMode() { return ModeFrameType; }
 
 	void OnTimingViewTrackListChanged();
 
 	void ToggleTimingViewMainGraphEventSeries(FTimerNodePtr TimerNode) const;
 
+	virtual FReply OnKeyDown(const FGeometry& MyGeometry, const FKeyEvent& InKeyEvent) override;
+
 private:
 	void InitCommandList();
 
@@ -125,7 +125,9 @@ private:
 	TSharedPtr<SWidget> TreeView_GetMenuContent();
 	void TreeView_BuildSortByMenu(FMenuBuilder& MenuBuilder);
 	void TreeView_BuildViewColumnMenu(FMenuBuilder& MenuBuilder);
-	void TreeView_BuildExportMenu(FMenuBuilder& MenuBuilder);
+
+	//////////////////////////////////////////////////
+	// Copy to Clipboard
 
 	bool ContextMenu_CopyToClipboard_CanExecute() const;
 	void ContextMenu_CopyToClipboard_Execute();
@@ -133,6 +135,11 @@ private:
 	bool ContextMenu_CopyTimerNameToClipboard_CanExecute() const;
 	void ContextMenu_CopyTimerNameToClipboard_Execute();
 
+	//////////////////////////////////////////////////
+	// Export
+
+	void TreeView_BuildExportMenu(FMenuBuilder& MenuBuilder);
+
 	bool ContextMenu_Export_CanExecute() const;
 	void ContextMenu_Export_Execute();
 
@@ -153,9 +160,24 @@ private:
 	bool OpenSaveTextFileDialog(const FString& InDialogTitle, const FString& InDefaultFile, FString& OutFilename) const;
 	class IFileHandle* OpenExportFile(const TCHAR* InFilename) const;
 
+	//////////////////////////////////////////////////
+	// Open Source File in IDE
+
 	bool ContextMenu_OpenSource_CanExecute() const;
 	void ContextMenu_OpenSource_Execute() const;
 
+	void OpenSourceFileInIDE(FTimerNodePtr InNode) const;
+
+	//////////////////////////////////////////////////
+	// Plot (Graph Series)
+
+	void TreeView_BuildPlotTimerMenu(FMenuBuilder& MenuBuilder);
+
+	//////////////////////////////////////////////////
+	// Find Instance
+
+	void TreeView_FindMenu(FMenuBuilder& MenuBuilder);
+
 	bool ContextMenu_FindInstance_CanExecute() const;
 	void ContextMenu_FindInstance_Execute(bool bFindMax) const;
 
@@ -245,7 +267,7 @@ private:
 	TSharedRef<SWidget> Mode_OnGenerateWidget(TSharedPtr<ETraceFrameType> InGroupingMode) const;
 
 	FText Mode_GetSelectedText() const;
-	
+
 	FText Mode_GetText(ETraceFrameType InFrameType) const;
 
 	FText Mode_GetSelectedTooltipText() const;
@@ -269,7 +291,7 @@ private:
 	void OnSortModeChanged(const EColumnSortPriority::Type SortPriority, const FName& ColumnId, const EColumnSortMode::Type SortMode);
 
 	////////////////////////////////////////////////////////////////////////////////////////////////////
-	// Sorting actions
+	// Sorting Actions
 
 	// SortMode (HeaderMenu)
 	bool HeaderMenu_SortMode_IsChecked(const FName ColumnId, const EColumnSortMode::Type InSortMode);
@@ -287,7 +309,7 @@ private:
 	void ContextMenu_SortByColumn_Execute(const FName ColumnId);
 
 	////////////////////////////////////////////////////////////////////////////////////////////////////
-	// Column visibility actions
+	// Column Visibility Actions
 
 	// ShowColumn
 	bool CanShowColumn(const FName ColumnId) const;
@@ -315,12 +337,12 @@ private:
 	void ContextMenu_ResetColumns_Execute();
 
 	////////////////////////////////////////////////////////////////////////////////////////////////////
+	// Event Highlight/Filtering in Timing View
 
 	void ToggleTimingViewEventFilter(FTimerNodePtr TimerNode) const;
 
-	void TreeView_BuildPlotTimerMenu(FMenuBuilder& MenuBuilder);
-
-	void TreeView_FindMenu(FMenuBuilder& MenuBuilder);
+	////////////////////////////////////////////////////////////////////////////////////////////////////
+	// Plot (Graph Series)
 
 	TSharedPtr<FTimingGraphTrack> GetTimingViewMainGraphTrack() const;
 	TSharedPtr<SFrameTrack> GetFrameTrack() const;
@@ -347,8 +369,6 @@ private:
 	 */
 	virtual void Tick(const FGeometry& AllottedGeometry, const double InCurrentTime, const float InDeltaTime) override;
 
-	void OpenSourceFileInIDE(FTimerNodePtr InNode) const;
-
 	void SaveSettings();
 	void SaveVisibleColumnsSettings();
 	void LoadSettings();
@@ -407,7 +427,7 @@ private:
 	TSet<FTimerNodePtr> ExpandedNodes;
 
 	/** If true, the expanded nodes have been saved before applying a text filter. */
-	bool bExpansionSaved;
+	bool bExpansionSaved = false;
 
 	//////////////////////////////////////////////////
 	// Search box and filters
@@ -427,7 +447,7 @@ private:
 	bool bTimerTypeIsVisible[static_cast<int>(ETimerNodeType::InvalidOrMax)];
 
 	/** Filter out the timers having zero total instance count (aggregated stats). */
-	bool bFilterOutZeroCountTimers;
+	bool bFilterOutZeroCountTimers = false;
 
 	//////////////////////////////////////////////////
 	// Grouping
@@ -439,7 +459,7 @@ private:
 	TSharedPtr<SComboBox<TSharedPtr<ETimerGroupingMode>>> GroupByComboBox;
 
 	/** How we group the timers? */
-	ETimerGroupingMode GroupingMode;
+	ETimerGroupingMode GroupingMode = ETimerGroupingMode::ByType;
 
 	//////////////////////////////////////////////////
 	// Frame Type
@@ -448,13 +468,11 @@ private:
 
 	TSharedPtr<SComboBox<TSharedPtr<ETraceFrameType>>> ModeComboBox;
 
-	ETraceFrameType ModeFrameType;
+	ETraceFrameType ModeFrameType = ETraceFrameType::TraceFrameType_Game;
 
 	//////////////////////////////////////////////////
 	// Sorting
 
-	//bool bUseSorting;
-
 	/** All available sorters. */
 	TArray<TSharedPtr<ITableCellValueSorter>> AvailableSorters;
 
@@ -465,7 +483,7 @@ private:
 	FName ColumnBeingSorted;
 
 	/** How we sort the nodes? Ascending or Descending. */
-	EColumnSortMode::Type ColumnSortMode;
+	EColumnSortMode::Type ColumnSortMode = EColumnSortMode::Type::Ascending;
 
 	//////////////////////////////////////////////////
 
