Automate Test Case: System.Insights.Hub.SessionBrowser

#rnx

diff --git a/Engine/Source/Developer/TraceInsights/Private/Insights/TraceInsightsModule.cpp b/Engine/Source/Developer/TraceInsights/Private/Insights/TraceInsightsModule.cpp
--- a/Engine/Source/Developer/TraceInsights/Private/Insights/TraceInsightsModule.cpp
+++ b/Engine/Source/Developer/TraceInsights/Private/Insights/TraceInsightsModule.cpp
@@ -211,7 +211,7 @@ FString FTraceInsightsModule::GetDefaultStoreDir()
 	}
 
 	const FStoreClient::FStatus* Status = StoreClient->GetStatus();
-	if (!StoreClient)
+	if (!Status)
 	{
 		UE_LOG(TraceInsights, Error, TEXT("Failed to get the status of the store client."));
 		return FString("");
