[Insights] Timers, Callers, Callees: Fixed "added to graphs" colored box not being in sync between different panels.

diff --git a/Engine/Source/Developer/TraceInsights/Private/Insights/TimingProfiler/ViewModels/TimerNode.cpp b/Engine/Source/Developer/TraceInsights/Private/Insights/TimingProfiler/ViewModels/TimerNode.cpp
--- a/Engine/Source/Developer/TraceInsights/Private/Insights/TimingProfiler/ViewModels/TimerNode.cpp
+++ b/Engine/Source/Developer/TraceInsights/Private/Insights/TimingProfiler/ViewModels/TimerNode.cpp
@@ -33,8 +33,6 @@ FTimerNode::FTimerNode(uint32 InTimerId, const TCHAR* InName, ETimerNodeType InT
 					InType == ETimerNodeType::CpuScope ? CpuGroup :
 					NAME_None)
 	, Type(InType)
-	, NumGraphs(0)
-	, bIsHotPath(false)
 {
 	uint32 Color32 = FTimingEvent::ComputeEventColor(InName);
 	Color.R = ((Color32 >> 16) & 0xFF) / 255.0f;
@@ -52,9 +50,7 @@ FTimerNode::FTimerNode(const FName InGroupName)
 	: FBaseTreeNode(InGroupName, true)
 	, TimerId(InvalidTimerId)
 	, Type(ETimerNodeType::Group)
-	, Color(0.0, 0.0, 0.0, 1.0)
-	, NumGraphs(0)
-	, bIsHotPath(false)
+	, Color(0.0f, 0.0f, 0.0f, 1.0f)
 {
 	ResetAggregatedStats();
 }
diff --git a/Engine/Source/Developer/TraceInsights/Private/Insights/TimingProfiler/ViewModels/TimerNode.h b/Engine/Source/Developer/TraceInsights/Private/Insights/TimingProfiler/ViewModels/TimerNode.h
--- a/Engine/Source/Developer/TraceInsights/Private/Insights/TimingProfiler/ViewModels/TimerNode.h
+++ b/Engine/Source/Developer/TraceInsights/Private/Insights/TimingProfiler/ViewModels/TimerNode.h
@@ -99,18 +99,12 @@ public:
 
 	bool IsAddedToGraph() const
 	{
-		return NumGraphs > 0;
+		return bIsAddedToGraphs;
 	}
 
-	void OnAddedToGraph()
+	void SetAddedToGraphsFlag(bool bInIsAddedToGraphs)
 	{
-		++NumGraphs;
-	}
-
-	void OnRemovedFromGraph()
-	{
-		ensure(NumGraphs > 0);
-		--NumGraphs;
+		bIsAddedToGraphs = bInIsAddedToGraphs;
 	}
 
 	/**
@@ -139,12 +133,12 @@ private:
 
 	/** Color of the node. */
 	FLinearColor Color;
-	
-	/** The number of graphs created for this timer. */
-	int32 NumGraphs;
+
+	/** True if the timer is added to a frame track or to a graph track. */
+	bool bIsAddedToGraphs = false;
 
 	/** True if this tree node is on the hot path. */
-	bool bIsHotPath;
+	bool bIsHotPath = false;
 
 	/** Aggregated stats. */
 	TraceServices::FTimingProfilerAggregatedStats AggregatedStats;
diff --git a/Engine/Source/Developer/TraceInsights/Private/Insights/TimingProfiler/Widgets/SFrameTrack.cpp b/Engine/Source/Developer/TraceInsights/Private/Insights/TimingProfiler/Widgets/SFrameTrack.cpp
--- a/Engine/Source/Developer/TraceInsights/Private/Insights/TimingProfiler/Widgets/SFrameTrack.cpp
+++ b/Engine/Source/Developer/TraceInsights/Private/Insights/TimingProfiler/Widgets/SFrameTrack.cpp
@@ -32,7 +32,6 @@
 #include "Insights/TimingProfiler/ViewModels/FrameStatsHelper.h"
 #include "Insights/TimingProfiler/ViewModels/FrameTrackHelper.h"
 #include "Insights/TimingProfiler/ViewModels/ThreadTimingSharedState.h"
-#include "Insights/TimingProfiler/Widgets/STimersView.h"
 #include "Insights/TimingProfiler/Widgets/STimingProfilerWindow.h"
 #include "Insights/TimingProfilerCommon.h"
 #include "Insights/Widgets/SLogView.h"
@@ -77,25 +76,22 @@ SFrameTrack::~SFrameTrack()
 		}
 	}
 
-	TSharedPtr<STimersView> TimersView;
 	if (ProfilerWindow.IsValid())
 	{
-		TimersView = ProfilerWindow->GetTimersView();
-	}
-	if (TimersView)
-	{
-		for (const TSharedPtr<FFrameTrackSeries>& Series : AllSeries)
+		TArray<TSharedPtr<FTimerFrameStatsTrackSeries>> TimerSeries;
+		for (TSharedPtr<FFrameTrackSeries>& Series : AllSeries)
 		{
-			if (Series.IsValid() && Series->Type == EFrameTrackSeriesType::TimerFrameStats)
+			if (Series.IsValid() &&
+				Series->Type == EFrameTrackSeriesType::TimerFrameStats)
 			{
-				const TSharedPtr<FTimerFrameStatsTrackSeries> TimerSeries = StaticCastSharedPtr<FTimerFrameStatsTrackSeries>(Series);
-				FTimerNodePtr TimerNode = TimersView->GetTimerNode(TimerSeries->TimerId);
-				if (TimerNode)
-				{
-					TimerNode->OnRemovedFromGraph();
-				}
+				TimerSeries.Add(StaticCastSharedPtr<FTimerFrameStatsTrackSeries>(Series));
 			}
-		};
+		}
+		AllSeries.Reset();
+		for (TSharedPtr<FTimerFrameStatsTrackSeries>& Series : TimerSeries)
+		{
+			ProfilerWindow->OnTimerAddedToGraphsChanged(Series->TimerId);
+		}
 	}
 }
 
@@ -2139,9 +2135,9 @@ void SFrameTrack::UpdateHorizontalScrollBar()
 
 ////////////////////////////////////////////////////////////////////////////////////////////////////
 
-bool SFrameTrack::HasFrameStatSeries(ETraceFrameType FrameType, uint32 TimerId)
+bool SFrameTrack::HasFrameStatSeries(ETraceFrameType FrameType, uint32 TimerId) const
 {
-	TSharedPtr<FFrameTrackSeries>* ExistingSeries = AllSeries.FindByPredicate([FrameType, TimerId](TSharedPtr<FFrameTrackSeries> Series)
+	const TSharedPtr<FFrameTrackSeries>* ExistingSeries = AllSeries.FindByPredicate([FrameType, TimerId](TSharedPtr<FFrameTrackSeries> Series)
 		{
 			return Series->Type == EFrameTrackSeriesType::TimerFrameStats &&
 				Series->FrameType == FrameType &&
@@ -2193,7 +2189,26 @@ bool SFrameTrack::RemoveTimerFrameStatSeries(ETraceFrameType FrameType, uint32 T
 
 ////////////////////////////////////////////////////////////////////////////////////////////////////
 
-uint32 SFrameTrack::GetNumSeriesForTimer(uint32 TimerId)
+bool SFrameTrack::HasAnySeriesForTimer(uint32 TimerId) const
+{
+	for (const TSharedPtr<FFrameTrackSeries>& Series : AllSeries)
+	{
+		if (Series->Type == EFrameTrackSeriesType::TimerFrameStats)
+		{
+			const TSharedPtr<FTimerFrameStatsTrackSeries> TimerSeries = StaticCastSharedPtr<FTimerFrameStatsTrackSeries>(Series);
+			if (TimerSeries->TimerId == TimerId)
+			{
+				return true;
+			}
+		}
+	}
+
+	return false;
+}
+
+////////////////////////////////////////////////////////////////////////////////////////////////////
+
+uint32 SFrameTrack::GetNumSeriesForTimer(uint32 TimerId) const
 {
 	uint32 NumSeries = 0;
 
@@ -2207,7 +2222,7 @@ uint32 SFrameTrack::GetNumSeriesForTimer(uint32 TimerId)
 				++NumSeries;
 			}
 		}
-	};
+	}
 
 	return NumSeries;
 }
diff --git a/Engine/Source/Developer/TraceInsights/Private/Insights/TimingProfiler/Widgets/SFrameTrack.h b/Engine/Source/Developer/TraceInsights/Private/Insights/TimingProfiler/Widgets/SFrameTrack.h
--- a/Engine/Source/Developer/TraceInsights/Private/Insights/TimingProfiler/Widgets/SFrameTrack.h
+++ b/Engine/Source/Developer/TraceInsights/Private/Insights/TimingProfiler/Widgets/SFrameTrack.h
@@ -137,11 +137,12 @@ public:
 
 	virtual FCursorReply OnCursorQuery(const FGeometry& MyGeometry, const FPointerEvent& CursorEvent) const override;
 
-	bool HasFrameStatSeries(ETraceFrameType FrameType, uint32 TimerId);
+	bool HasFrameStatSeries(ETraceFrameType FrameType, uint32 TimerId) const;
 	TSharedPtr<FTimerFrameStatsTrackSeries> AddTimerFrameStatSeries(ETraceFrameType FrameType, uint32 TimerId, FLinearColor Color, FText Name);
 	bool RemoveTimerFrameStatSeries(ETraceFrameType FrameType, uint32 TimerId);
 
-	uint32 GetNumSeriesForTimer(uint32 TimerId);
+	bool HasAnySeriesForTimer(uint32 TimerId) const;
+	uint32 GetNumSeriesForTimer(uint32 TimerId) const;
 
 protected:
 	TSharedRef<FFrameTrackSeries> FindOrAddSeries(ETraceFrameType FrameType);
diff --git a/Engine/Source/Developer/TraceInsights/Private/Insights/TimingProfiler/Widgets/STimerTreeView.cpp b/Engine/Source/Developer/TraceInsights/Private/Insights/TimingProfiler/Widgets/STimerTreeView.cpp
--- a/Engine/Source/Developer/TraceInsights/Private/Insights/TimingProfiler/Widgets/STimerTreeView.cpp
+++ b/Engine/Source/Developer/TraceInsights/Private/Insights/TimingProfiler/Widgets/STimerTreeView.cpp
@@ -1008,7 +1008,8 @@ void STimerTreeView::TreeView_OnGetChildren(FTimerNodePtr InParent, TArray<FTime
 		OutChildren.Reset(Children.Num());
 		for (const FBaseTreeNodePtr& Child : Children)
 		{
-			OutChildren.Add(StaticCastSharedPtr<FTimerNode, FBaseTreeNode>(Child));
+			check(Child->Is<FTimerNode>());
+			OutChildren.Add(StaticCastSharedPtr<FTimerNode>(Child));
 		}
 	}
 	else
@@ -1017,7 +1018,8 @@ void STimerTreeView::TreeView_OnGetChildren(FTimerNodePtr InParent, TArray<FTime
 		OutChildren.Reset(Children.Num());
 		for (const FBaseTreeNodePtr& Child : Children)
 		{
-			OutChildren.Add(StaticCastSharedPtr<FTimerNode, FBaseTreeNode>(Child));
+			check(Child->Is<FTimerNode>());
+			OutChildren.Add(StaticCastSharedPtr<FTimerNode>(Child));
 		}
 	}
 }
@@ -1180,6 +1182,7 @@ void STimerTreeView::SortTreeNodesRec(FTimerNode& Node, const ITableCellValueSor
 	{
 		if (ChildPtr->GetChildrenCount() > 0)
 		{
+			check(ChildPtr->Is<FTimerNode>());
 			SortTreeNodesRec(*StaticCastSharedPtr<FTimerNode>(ChildPtr), Sorter);
 		}
 	}
@@ -1487,7 +1490,15 @@ void STimerTreeView::SetTree(const TraceServices::FTimingProfilerButterflyNode&
 		{
 			TimerNodePtr->SetIsHotPath(true);
 			const TArray<FBaseTreeNodePtr>& Children = TimerNodePtr->GetChildren();
-			TimerNodePtr = Children.Num() > 0 ? StaticCastSharedPtr<FTimerNode, FBaseTreeNode>(Children[0]) : nullptr;
+			if (Children.Num() > 0)
+			{
+				check(Children[0]->Is<FTimerNode>());
+				TimerNodePtr = StaticCastSharedPtr<FTimerNode>(Children[0]);
+			}
+			else
+			{
+				TimerNodePtr = nullptr;
+			}
 		}
 
 		TreeNodes.Add(RootTimerNodePtr);
@@ -1517,6 +1528,10 @@ FTimerNodePtr STimerTreeView::CreateTimerNodeRec(const TraceServices::FTimingPro
 														   ETimerNodeType::CpuScope;
 	FTimerNodePtr TimerNodePtr = MakeShared<FTimerNode>(Node.Timer->Id, Node.Timer->Name, Type, true);
 
+	TSharedPtr<STimingProfilerWindow> ProfilerWindow = FTimingProfilerManager::Get()->GetProfilerWindow();
+	const bool bIsTimerAddedToGraphs = ProfilerWindow.IsValid() ? ProfilerWindow->IsTimerAddedToGraphs(Node.Timer->Id) : false;
+	TimerNodePtr->SetAddedToGraphsFlag(bIsTimerAddedToGraphs);
+
 	TraceServices::FTimingProfilerAggregatedStats AggregatedStats;
 	AggregatedStats.InstanceCount = Node.Count;
 	AggregatedStats.TotalInclusiveTime = Node.InclusiveTime;
@@ -1547,7 +1562,9 @@ FTimerNodePtr STimerTreeView::CreateTimerNodeRec(const TraceServices::FTimingPro
 	// Sort children by InclTime (descending).
 	TimerNodePtr->SortChildren([](const FBaseTreeNodePtr& A, const FBaseTreeNodePtr& B) -> bool
 	{
+		checkSlow(A->Is<FTimerNode>());
 		const double InclTimeA = StaticCastSharedPtr<FTimerNode>(A)->GetAggregatedStats().TotalInclusiveTime;
+		checkSlow(B->Is<FTimerNode>());
 		const double InclTimeB = StaticCastSharedPtr<FTimerNode>(B)->GetAggregatedStats().TotalInclusiveTime;
 		return InclTimeA >= InclTimeB;
 	});
@@ -1567,6 +1584,7 @@ void STimerTreeView::ExpandNodesRec(FTimerNodePtr NodePtr, int32 Depth)
 	{
 		for (const FBaseTreeNodePtr& ChildPtr : NodePtr->GetChildren())
 		{
+			check(ChildPtr->Is<FTimerNode>());
 			ExpandNodesRec(StaticCastSharedPtr<FTimerNode>(ChildPtr), Depth + 1);
 		}
 	}
@@ -1574,6 +1592,44 @@ void STimerTreeView::ExpandNodesRec(FTimerNodePtr NodePtr, int32 Depth)
 
 ////////////////////////////////////////////////////////////////////////////////////////////////////
 
+FTimerNodePtr STimerTreeView::GetTimerNode(uint32 TimerId) const
+{
+	for (FTimerNodePtr TimerNode : TreeNodes)
+	{
+		FTimerNodePtr FoundNode = GetTimerNodeRec(TimerId, TimerNode);
+		if (FoundNode)
+		{
+			return FoundNode;
+		}
+	}
+
+	return nullptr;
+}
+
+////////////////////////////////////////////////////////////////////////////////////////////////////
+
+FTimerNodePtr STimerTreeView::GetTimerNodeRec(uint32 TimerId, const FTimerNodePtr TimerNode) const
+{
+	if (TimerNode->GetTimerId() == TimerId)
+	{
+		return TimerNode;
+	}
+
+	for (const FBaseTreeNodePtr& ChildPtr : TimerNode->GetChildren())
+	{
+		check(ChildPtr->Is<FTimerNode>());
+		FTimerNodePtr FoundNode = GetTimerNodeRec(TimerId, StaticCastSharedPtr<FTimerNode>(ChildPtr));
+		if (FoundNode)
+		{
+			return FoundNode;
+		}
+	}
+
+	return nullptr;
+}
+
+////////////////////////////////////////////////////////////////////////////////////////////////////
+
 void STimerTreeView::ToggleTimingViewEventFilter(FTimerNodePtr TimerNode) const
 {
 	TSharedPtr<STimingView> TimingView = GetTimingView();
@@ -1605,12 +1661,12 @@ TSharedPtr<SFrameTrack> STimerTreeView::GetFrameTrack() const
 void STimerTreeView::ToggleGraphInstanceSeries(TSharedRef<FTimingGraphTrack> GraphTrack, FTimerNodeRef NodePtr) const
 {
 	const uint32 TimerId = NodePtr->GetTimerId();
+
 	TSharedPtr<FTimingGraphSeries> Series = GraphTrack->GetTimerSeries(TimerId);
 	if (Series.IsValid())
 	{
 		GraphTrack->RemoveTimerSeries(TimerId);
 		GraphTrack->SetDirtyFlag();
-		NodePtr->OnRemovedFromGraph();
 	}
 	else
 	{
@@ -1618,7 +1674,12 @@ void STimerTreeView::ToggleGraphInstanceSeries(TSharedRef<FTimingGraphTrack> Gra
 		Series = GraphTrack->AddTimerSeries(TimerId, NodePtr->GetColor());
 		Series->SetName(FText::FromName(NodePtr->GetName()));
 		GraphTrack->SetDirtyFlag();
-		NodePtr->OnAddedToGraph();
+	}
+
+	TSharedPtr<STimingProfilerWindow> ProfilerWindow = FTimingProfilerManager::Get()->GetProfilerWindow();
+	if (ProfilerWindow.IsValid())
+	{
+		ProfilerWindow->OnTimerAddedToGraphsChanged(TimerId);
 	}
 }
 
@@ -1655,12 +1716,12 @@ void STimerTreeView::ToggleTimingViewMainGraphEventInstanceSeries(FTimerNodePtr
 void STimerTreeView::ToggleGraphFrameStatsSeries(TSharedRef<FTimingGraphTrack> GraphTrack, FTimerNodeRef NodePtr, ETraceFrameType FrameType) const
 {
 	const uint32 TimerId = NodePtr->GetTimerId();
+
 	TSharedPtr<FTimingGraphSeries> Series = GraphTrack->GetFrameStatsTimerSeries(TimerId, FrameType);
 	if (Series.IsValid())
 	{
 		GraphTrack->RemoveFrameStatsTimerSeries(TimerId, FrameType);
 		GraphTrack->SetDirtyFlag();
-		NodePtr->OnRemovedFromGraph();
 	}
 	else
 	{
@@ -1671,7 +1732,12 @@ void STimerTreeView::ToggleGraphFrameStatsSeries(TSharedRef<FTimingGraphTrack> G
 			FFrameTrackDrawHelper::FrameTypeToText(FrameType));
 		Series->SetName(SeriesName);
 		GraphTrack->SetDirtyFlag();
-		NodePtr->OnAddedToGraph();
+	}
+
+	TSharedPtr<STimingProfilerWindow> ProfilerWindow = FTimingProfilerManager::Get()->GetProfilerWindow();
+	if (ProfilerWindow.IsValid())
+	{
+		ProfilerWindow->OnTimerAddedToGraphsChanged(TimerId);
 	}
 }
 
@@ -1722,24 +1788,29 @@ bool STimerTreeView::IsSeriesInFrameTrack(FTimerNodePtr TimerNode, ETraceFrameTy
 void STimerTreeView::ToggleFrameTrackSeries(FTimerNodePtr TimerNode, ETraceFrameType FrameType) const
 {
 	TSharedPtr<SFrameTrack> FrameTrack = GetFrameTrack();
-
 	if (!FrameTrack.IsValid())
 	{
 		return;
 	}
 
-	if (FrameTrack->HasFrameStatSeries(FrameType, TimerNode->GetTimerId()))
+	const uint32 TimerId = TimerNode->GetTimerId();
+
+	if (FrameTrack->HasFrameStatSeries(FrameType, TimerId))
 	{
-		FrameTrack->RemoveTimerFrameStatSeries(FrameType, TimerNode->GetTimerId());
-		TimerNode->OnRemovedFromGraph();
+		FrameTrack->RemoveTimerFrameStatSeries(FrameType, TimerId);
 	}
 	else
 	{
 		FText SeriesName = FText::Format(LOCTEXT("FrameStatsTimerSeriesName_Fmt", "{0} ({1})"),
 			TimerNode->GetDisplayName(),
 			FFrameTrackDrawHelper::FrameTypeToText(FrameType));
-		TSharedPtr<FTimerFrameStatsTrackSeries> Series = FrameTrack->AddTimerFrameStatSeries(FrameType, TimerNode->GetTimerId(), TimerNode->GetColor(), SeriesName);
-		TimerNode->OnAddedToGraph();
+		FrameTrack->AddTimerFrameStatSeries(FrameType, TimerId, TimerNode->GetColor(), SeriesName);
+	}
+
+	TSharedPtr<STimingProfilerWindow> ProfilerWindow = FTimingProfilerManager::Get()->GetProfilerWindow();
+	if (ProfilerWindow.IsValid())
+	{
+		ProfilerWindow->OnTimerAddedToGraphsChanged(TimerId);
 	}
 }
 
diff --git a/Engine/Source/Developer/TraceInsights/Private/Insights/TimingProfiler/Widgets/STimerTreeView.h b/Engine/Source/Developer/TraceInsights/Private/Insights/TimingProfiler/Widgets/STimerTreeView.h
--- a/Engine/Source/Developer/TraceInsights/Private/Insights/TimingProfiler/Widgets/STimerTreeView.h
+++ b/Engine/Source/Developer/TraceInsights/Private/Insights/TimingProfiler/Widgets/STimerTreeView.h
@@ -62,8 +62,11 @@ public:
 	void Construct(const FArguments& InArgs, const FText& InViewName);
 
 	void Reset();
+
 	void SetTree(const TraceServices::FTimingProfilerButterflyNode& Root);
 
+	FTimerNodePtr GetTimerNode(uint32 TimerId) const;
+
 	virtual FReply OnKeyDown(const FGeometry& MyGeometry, const FKeyEvent& InKeyEvent) override;
 
 private:
@@ -238,6 +241,8 @@ private:
 
 	TSharedPtr<STimingView> GetTimingView() const;
 
+	FTimerNodePtr GetTimerNodeRec(uint32 TimerId, const FTimerNodePtr TimerNode) const;
+
 private:
 	/** Table view model. */
 	TSharedPtr<FTable> Table;
diff --git a/Engine/Source/Developer/TraceInsights/Private/Insights/TimingProfiler/Widgets/STimersView.cpp b/Engine/Source/Developer/TraceInsights/Private/Insights/TimingProfiler/Widgets/STimersView.cpp
--- a/Engine/Source/Developer/TraceInsights/Private/Insights/TimingProfiler/Widgets/STimersView.cpp
+++ b/Engine/Source/Developer/TraceInsights/Private/Insights/TimingProfiler/Widgets/STimersView.cpp
@@ -1347,7 +1347,8 @@ void STimersView::ApplyFiltering()
 		int32 NumVisibleChildren = 0;
 		for (const FBaseTreeNodePtr& ChildPtr : GroupChildren)
 		{
-			const FTimerNodePtr& NodePtr = StaticCastSharedPtr<FTimerNode, FBaseTreeNode>(ChildPtr);
+			check(ChildPtr->Is<FTimerNode>());
+			const FTimerNodePtr& NodePtr = StaticCastSharedPtr<FTimerNode>(ChildPtr);
 
 			const bool bIsChildVisible = (!bFilterOutZeroCountTimers || NodePtr->GetAggregatedStats().InstanceCount > 0)
 									  && bTimerTypeIsVisible[static_cast<int>(NodePtr->GetType())]
@@ -1418,7 +1419,8 @@ void STimersView::ApplyFiltering()
 		const TArray<FBaseTreeNodePtr>& GroupChildren = GroupPtr->GetFilteredChildren();
 		for (const FBaseTreeNodePtr& ChildPtr : GroupChildren)
 		{
-			const FTimerNodePtr& NodePtr = StaticCastSharedPtr<FTimerNode, FBaseTreeNode>(ChildPtr);
+			check(ChildPtr->Is<FTimerNode>());
+			const FTimerNodePtr& NodePtr = StaticCastSharedPtr<FTimerNode>(ChildPtr);
 			const TraceServices::FTimingProfilerAggregatedStats& NodeAggregatedStats = NodePtr->GetAggregatedStats();
 
 			if (NodeAggregatedStats.InstanceCount > 0)
@@ -1565,7 +1567,8 @@ void STimersView::TreeView_OnGetChildren(FTimerNodePtr InParent, TArray<FTimerNo
 	OutChildren.Reset(Children.Num());
 	for (const FBaseTreeNodePtr& Child : Children)
 	{
-		OutChildren.Add(StaticCastSharedPtr<FTimerNode, FBaseTreeNode>(Child));
+		check(Child->Is<FTimerNode>());
+		OutChildren.Add(StaticCastSharedPtr<FTimerNode>(Child));
 	}
 }
 
@@ -2069,6 +2072,7 @@ void STimersView::SortTreeNodesRec(FTimerNode& Node, const ITableCellValueSorter
 	{
 		if (ChildPtr->GetChildrenCount() > 0)
 		{
+			check(ChildPtr->Is<FTimerNode>());
 			SortTreeNodesRec(*StaticCastSharedPtr<FTimerNode>(ChildPtr), Sorter);
 		}
 	}
@@ -2471,8 +2475,8 @@ void STimersView::RebuildTree(bool bResync)
 			check(TimerCount > PreviousNodeCount);
 			TimerNodes.Reserve(TimerCount);
 
-			TSharedPtr<FTimingGraphTrack> GraphTrack = GetTimingViewMainGraphTrack();
 			TSharedPtr<SFrameTrack> FrameTrack = GetFrameTrack();
+			TSharedPtr<FTimingGraphTrack> GraphTrack = GetTimingViewMainGraphTrack();
 
 			// Add nodes only for new timers.
 			for (uint32 TimerIndex = PreviousNodeCount; TimerIndex < TimerCount; ++TimerIndex)
@@ -2485,25 +2489,10 @@ void STimersView::RebuildTree(bool bResync)
 				FTimerNodePtr TimerNodePtr = MakeShared<FTimerNode>(Timer.Id, Timer.Name, Type, false);
 				TimerNodePtr->SetDefaultSortOrder(TimerIndex + 1);
 
-				if (GraphTrack.IsValid())
-				{
-					uint32 NumSeries = GraphTrack->GetNumSeriesForTimer(Timer.Id);
-					while (NumSeries > 0)
-					{
-						TimerNodePtr->OnAddedToGraph();
-						--NumSeries;
-					}
-				}
-
-				if (FrameTrack.IsValid())
-				{
-					uint32 NumSeries = FrameTrack->GetNumSeriesForTimer(Timer.Id);
-					while (NumSeries > 0)
-					{
-						TimerNodePtr->OnAddedToGraph();
-						--NumSeries;
-					}
-				}
+				const bool bIsTimerAddedToGraphs =
+					(FrameTrack.IsValid() && FrameTrack->HasAnySeriesForTimer(Timer.Id)) ||
+					(GraphTrack.IsValid() && GraphTrack->HasAnySeriesForTimer(Timer.Id));
+				TimerNodePtr->SetAddedToGraphsFlag(bIsTimerAddedToGraphs);
 
 				TimerNodes.Add(TimerNodePtr);
 			}
@@ -2710,12 +2699,12 @@ void STimersView::ToggleTimingViewMainGraphEventSeries(FTimerNodePtr TimerNode)
 void STimersView::ToggleGraphInstanceSeries(TSharedRef<FTimingGraphTrack> GraphTrack, FTimerNodeRef NodePtr) const
 {
 	const uint32 TimerId = NodePtr->GetTimerId();
+
 	TSharedPtr<FTimingGraphSeries> Series = GraphTrack->GetTimerSeries(TimerId);
 	if (Series.IsValid())
 	{
 		GraphTrack->RemoveTimerSeries(TimerId);
 		GraphTrack->SetDirtyFlag();
-		NodePtr->OnRemovedFromGraph();
 	}
 	else
 	{
@@ -2723,7 +2712,12 @@ void STimersView::ToggleGraphInstanceSeries(TSharedRef<FTimingGraphTrack> GraphT
 		Series = GraphTrack->AddTimerSeries(TimerId, NodePtr->GetColor());
 		Series->SetName(FText::FromName(NodePtr->GetName()));
 		GraphTrack->SetDirtyFlag();
-		NodePtr->OnAddedToGraph();
+	}
+
+	TSharedPtr<STimingProfilerWindow> ProfilerWindow = FTimingProfilerManager::Get()->GetProfilerWindow();
+	if (ProfilerWindow.IsValid())
+	{
+		ProfilerWindow->OnTimerAddedToGraphsChanged(TimerId);
 	}
 }
 
@@ -2760,12 +2754,12 @@ void STimersView::ToggleTimingViewMainGraphEventInstanceSeries(FTimerNodePtr Tim
 void STimersView::ToggleGraphFrameStatsSeries(TSharedRef<FTimingGraphTrack> GraphTrack, FTimerNodeRef NodePtr, ETraceFrameType FrameType) const
 {
 	const uint32 TimerId = NodePtr->GetTimerId();
+
 	TSharedPtr<FTimingGraphSeries> Series = GraphTrack->GetFrameStatsTimerSeries(TimerId, FrameType);
 	if (Series.IsValid())
 	{
 		GraphTrack->RemoveFrameStatsTimerSeries(TimerId, FrameType);
 		GraphTrack->SetDirtyFlag();
-		NodePtr->OnRemovedFromGraph();
 	}
 	else
 	{
@@ -2776,7 +2770,12 @@ void STimersView::ToggleGraphFrameStatsSeries(TSharedRef<FTimingGraphTrack> Grap
 			FFrameTrackDrawHelper::FrameTypeToText(FrameType));
 		Series->SetName(SeriesName);
 		GraphTrack->SetDirtyFlag();
-		NodePtr->OnAddedToGraph();
+	}
+
+	TSharedPtr<STimingProfilerWindow> ProfilerWindow = FTimingProfilerManager::Get()->GetProfilerWindow();
+	if (ProfilerWindow.IsValid())
+	{
+		ProfilerWindow->OnTimerAddedToGraphsChanged(TimerId);
 	}
 }
 
@@ -2833,18 +2832,24 @@ void STimersView::ToggleFrameTrackSeries(FTimerNodePtr TimerNode, ETraceFrameTyp
 		return;
 	}
 
-	if (FrameTrack->HasFrameStatSeries(FrameType, TimerNode->GetTimerId()))
+	const uint32 TimerId = TimerNode->GetTimerId();
+
+	if (FrameTrack->HasFrameStatSeries(FrameType, TimerId))
 	{
-		FrameTrack->RemoveTimerFrameStatSeries(FrameType, TimerNode->GetTimerId());
-		TimerNode->OnRemovedFromGraph();
+		FrameTrack->RemoveTimerFrameStatSeries(FrameType, TimerId);
 	}
 	else
 	{
 		FText SeriesName = FText::Format(LOCTEXT("FrameStatsTimerSeriesName_Fmt", "{0} ({1})"),
 			TimerNode->GetDisplayName(),
 			FFrameTrackDrawHelper::FrameTypeToText(FrameType));
-		TSharedPtr<FTimerFrameStatsTrackSeries> Series = FrameTrack->AddTimerFrameStatSeries(FrameType, TimerNode->GetTimerId(), TimerNode->GetColor(), SeriesName);
-		TimerNode->OnAddedToGraph();
+		FrameTrack->AddTimerFrameStatSeries(FrameType, TimerId, TimerNode->GetColor(), SeriesName);
+	}
+
+	TSharedPtr<STimingProfilerWindow> ProfilerWindow = FTimingProfilerManager::Get()->GetProfilerWindow();
+	if (ProfilerWindow.IsValid())
+	{
+		ProfilerWindow->OnTimerAddedToGraphsChanged(TimerId);
 	}
 }
 
@@ -3048,6 +3053,7 @@ void STimersView::AddTimerNodeRecursive(FTimerNodePtr InNode, TSet<uint32>& InOu
 	{
 		for (FBaseTreeNodePtr ChildNode : InNode->GetFilteredChildren())
 		{
+			check(ChildNode->Is<FTimerNode>());
 			AddTimerNodeRecursive(StaticCastSharedPtr<FTimerNode>(ChildNode), InOutIncludedTimers);
 		}
 	}
diff --git a/Engine/Source/Developer/TraceInsights/Private/Insights/TimingProfiler/Widgets/STimingProfilerWindow.cpp b/Engine/Source/Developer/TraceInsights/Private/Insights/TimingProfiler/Widgets/STimingProfilerWindow.cpp
--- a/Engine/Source/Developer/TraceInsights/Private/Insights/TimingProfiler/Widgets/STimingProfilerWindow.cpp
+++ b/Engine/Source/Developer/TraceInsights/Private/Insights/TimingProfiler/Widgets/STimingProfilerWindow.cpp
@@ -17,6 +17,7 @@
 #include "Insights/TimingProfiler/Widgets/STimersView.h"
 #include "Insights/TimingProfiler/Widgets/STimerTreeView.h"
 #include "Insights/TimingProfiler/Widgets/STimingProfilerToolbar.h"
+#include "Insights/ViewModels/TimingGraphTrack.h"
 #include "Insights/Widgets/SLogView.h"
 #include "Insights/Widgets/STimingView.h"
 
@@ -516,6 +517,63 @@ void STimingProfilerWindow::OnTimeSelectionChanged(Timing::ETimeChangedFlags InF
 
 ////////////////////////////////////////////////////////////////////////////////////////////////////
 
+bool STimingProfilerWindow::IsTimerAddedToGraphs(uint32 TimerId) const
+{
+	if (FrameTrack.IsValid() &&
+		FrameTrack->HasAnySeriesForTimer(TimerId))
+	{
+		return true;
+	}
+
+	if (TimingView.IsValid())
+	{
+		TSharedPtr<FTimingGraphTrack> GraphTrack = TimingView->GetMainTimingGraphTrack();
+		if (GraphTrack.IsValid() &&
+			GraphTrack->HasAnySeriesForTimer(TimerId))
+		{
+			return true;
+		}
+	}
+
+	return false;
+}
+
+////////////////////////////////////////////////////////////////////////////////////////////////////
+
+void STimingProfilerWindow::OnTimerAddedToGraphsChanged(uint32 TimerId)
+{
+	const bool bIsAddedToGraphs = IsTimerAddedToGraphs(TimerId);
+
+	if (TimersView.IsValid())
+	{
+		FTimerNodePtr TimerNode = TimersView->GetTimerNode(TimerId);
+		if (TimerNode.IsValid())
+		{
+			TimerNode->SetAddedToGraphsFlag(bIsAddedToGraphs);
+		}
+	}
+
+	if (CallersTreeView.IsValid())
+	{
+		FTimerNodePtr TimerNode = CallersTreeView->GetTimerNode(TimerId);
+		if (TimerNode.IsValid())
+		{
+			TimerNode->SetAddedToGraphsFlag(bIsAddedToGraphs);
+		}
+	}
+
+	if (CalleesTreeView.IsValid())
+	{
+		FTimerNodePtr TimerNode = CalleesTreeView->GetTimerNode(TimerId);
+		if (TimerNode.IsValid())
+		{
+			TimerNode->SetAddedToGraphsFlag(bIsAddedToGraphs);
+		}
+	}
+}
+
+////////////////////////////////////////////////////////////////////////////////////////////////////
+
 } // namespace UE::Insights::TimingProfiler
 
 #undef LOCTEXT_NAMESPACE
diff --git a/Engine/Source/Developer/TraceInsights/Private/Insights/TimingProfiler/Widgets/STimingProfilerWindow.h b/Engine/Source/Developer/TraceInsights/Private/Insights/TimingProfiler/Widgets/STimingProfilerWindow.h
--- a/Engine/Source/Developer/TraceInsights/Private/Insights/TimingProfiler/Widgets/STimingProfilerWindow.h
+++ b/Engine/Source/Developer/TraceInsights/Private/Insights/TimingProfiler/Widgets/STimingProfilerWindow.h
@@ -45,6 +45,9 @@ public:
 	TSharedPtr<SLogView> GetLogView() const { return LogView; }
 	TSharedPtr<SFrameTrack> GetFrameView() const { return FrameTrack; }
 
+	bool IsTimerAddedToGraphs(uint32 TimerId) const;
+	void OnTimerAddedToGraphsChanged(uint32 TimerId);
+
 protected:
 	virtual const TCHAR* GetAnalyticsEventName() const override;
 	virtual TSharedRef<FWorkspaceItem> CreateWorkspaceMenuGroup() override;
diff --git a/Engine/Source/Developer/TraceInsights/Private/Insights/ViewModels/TimingGraphTrack.cpp b/Engine/Source/Developer/TraceInsights/Private/Insights/ViewModels/TimingGraphTrack.cpp
--- a/Engine/Source/Developer/TraceInsights/Private/Insights/ViewModels/TimingGraphTrack.cpp
+++ b/Engine/Source/Developer/TraceInsights/Private/Insights/ViewModels/TimingGraphTrack.cpp
@@ -19,7 +19,6 @@
 #include "Insights/TimingProfiler/Tracks/ThreadTimingTrack.h"
 #include "Insights/TimingProfiler/ViewModels/FrameStatsHelper.h"
 #include "Insights/TimingProfiler/ViewModels/ThreadTimingSharedState.h"
-#include "Insights/TimingProfiler/Widgets/STimersView.h"
 #include "Insights/TimingProfiler/Widgets/STimingProfilerWindow.h"
 #include "Insights/ViewModels/AxisViewportDouble.h"
 #include "Insights/ViewModels/GraphTrackBuilder.h"
@@ -189,31 +188,28 @@ FTimingGraphTrack::~FTimingGraphTrack()
 
 	using namespace UE::Insights::TimingProfiler;
 
-	TSharedPtr<STimersView> TimersView;
 	if (bNotifyTimersOnDestruction)
 	{
 		TSharedPtr<STimingProfilerWindow> ProfilerWindow = FTimingProfilerManager::Get()->GetProfilerWindow();
 		if (ProfilerWindow.IsValid())
 		{
-			TimersView = ProfilerWindow->GetTimersView();
-		}
-	}
-
-	if (TimersView)
-	{
-		for (const TSharedPtr<FGraphSeries>& Series : AllSeries)
-		{
-			const TSharedPtr<FTimingGraphSeries> TimingSeries = StaticCastSharedPtr<FTimingGraphSeries>(Series);
-			if (TimingSeries.IsValid() &&
-				(TimingSeries->Type == FTimingGraphSeries::ESeriesType::Timer || TimingSeries->Type == FTimingGraphSeries::ESeriesType::FrameStatsTimer))
+			TArray<TSharedPtr<FTimingGraphSeries>> TimerSeries;
+			for (TSharedPtr<FGraphSeries>& Series : AllSeries)
 			{
-				FTimerNodePtr TimerNode = TimersView->GetTimerNode(TimingSeries->TimerId);
-				if (TimerNode)
+				TSharedPtr<FTimingGraphSeries> TimingSeries = StaticCastSharedPtr<FTimingGraphSeries>(Series);
+				if (TimingSeries.IsValid() &&
+					(TimingSeries->Type == FTimingGraphSeries::ESeriesType::Timer ||
+					 TimingSeries->Type == FTimingGraphSeries::ESeriesType::FrameStatsTimer))
 				{
-					TimerNode->OnRemovedFromGraph();
+					TimerSeries.Add(TimingSeries);
 				}
 			}
-		};
+			AllSeries.Reset();
+			for (TSharedPtr<FTimingGraphSeries>& TimingSeries : TimerSeries)
+			{
+				ProfilerWindow->OnTimerAddedToGraphsChanged(TimingSeries->TimerId);
+			}
+		}
 	}
 }
 
@@ -558,7 +554,9 @@ TSharedPtr<FTimingGraphSeries> FTimingGraphTrack::GetFrameStatsTimerSeries(uint3
 	TSharedPtr<FGraphSeries>* Ptr = AllSeries.FindByPredicate([TimerId, FrameType](const TSharedPtr<FGraphSeries>& Series)
 		{
 			const TSharedPtr<FTimingGraphSeries> TimingSeries = StaticCastSharedPtr<FTimingGraphSeries>(Series);
-			return TimingSeries->Type == FTimingGraphSeries::ESeriesType::FrameStatsTimer && TimingSeries->TimerId == TimerId && TimingSeries->FrameType == FrameType;
+			return TimingSeries->Type == FTimingGraphSeries::ESeriesType::FrameStatsTimer &&
+				   TimingSeries->TimerId == TimerId &&
+				   TimingSeries->FrameType == FrameType;
 		});
 	return (Ptr != nullptr) ? StaticCastSharedPtr<FTimingGraphSeries>(*Ptr) : nullptr;
 }
@@ -596,7 +594,9 @@ void FTimingGraphTrack::RemoveFrameStatsTimerSeries(uint32 TimerId, ETraceFrameT
 	AllSeries.RemoveAll([TimerId, FrameType](const TSharedPtr<FGraphSeries>& Series)
 		{
 			const TSharedPtr<FTimingGraphSeries> TimingSeries = StaticCastSharedPtr<FTimingGraphSeries>(Series);
-			return TimingSeries->Type == FTimingGraphSeries::ESeriesType::FrameStatsTimer && TimingSeries->TimerId == TimerId && TimingSeries->FrameType == FrameType;
+			return TimingSeries->Type == FTimingGraphSeries::ESeriesType::FrameStatsTimer &&
+				   TimingSeries->TimerId == TimerId &&
+				   TimingSeries->FrameType == FrameType;
 		});
 }
 
@@ -1021,20 +1021,40 @@ void FTimingGraphTrack::DrawVerticalAxisGrid(const ITimingTrackDrawContext& Cont
 
 ////////////////////////////////////////////////////////////////////////////////////////////////////
 
-uint32 FTimingGraphTrack::GetNumSeriesForTimer(uint32 TimerId)
+bool FTimingGraphTrack::HasAnySeriesForTimer(uint32 TimerId) const
+{
+	for (const TSharedPtr<FGraphSeries>& Series : AllSeries)
+	{
+		const TSharedPtr<FTimingGraphSeries> TimingSeries = StaticCastSharedPtr<FTimingGraphSeries>(Series);
+		if (TimingSeries.IsValid() &&
+			(TimingSeries->Type == FTimingGraphSeries::ESeriesType::Timer ||
+			 TimingSeries->Type == FTimingGraphSeries::ESeriesType::FrameStatsTimer) &&
+			TimingSeries->TimerId == TimerId)
+		{
+			return true;
+		}
+	}
+
+	return false;
+}
+
+////////////////////////////////////////////////////////////////////////////////////////////////////
+
+uint32 FTimingGraphTrack::GetNumSeriesForTimer(uint32 TimerId) const
 {
 	uint32 NumSeries = 0;
 
-	for(const TSharedPtr<FGraphSeries>& Series : AllSeries)
+	for (const TSharedPtr<FGraphSeries>& Series : AllSeries)
 	{
 		const TSharedPtr<FTimingGraphSeries> TimingSeries = StaticCastSharedPtr<FTimingGraphSeries>(Series);
 		if (TimingSeries.IsValid() &&
-			(TimingSeries->Type == FTimingGraphSeries::ESeriesType::Timer || TimingSeries->Type == FTimingGraphSeries::ESeriesType::FrameStatsTimer) &&
+			(TimingSeries->Type == FTimingGraphSeries::ESeriesType::Timer ||
+			 TimingSeries->Type == FTimingGraphSeries::ESeriesType::FrameStatsTimer) &&
 			TimingSeries->TimerId == TimerId)
 		{
 			++NumSeries;
 		}
-	};
+	}
 
 	return NumSeries;
 }
diff --git a/Engine/Source/Developer/TraceInsights/Private/Insights/ViewModels/TimingGraphTrack.h b/Engine/Source/Developer/TraceInsights/Private/Insights/ViewModels/TimingGraphTrack.h
--- a/Engine/Source/Developer/TraceInsights/Private/Insights/ViewModels/TimingGraphTrack.h
+++ b/Engine/Source/Developer/TraceInsights/Private/Insights/ViewModels/TimingGraphTrack.h
@@ -98,7 +98,8 @@ public:
 	TSharedPtr<FTimingGraphSeries> AddStatsCounterSeries(uint32 CounterId, FLinearColor Color);
 	void RemoveStatsCounterSeries(uint32 CounterId);
 
-	uint32 GetNumSeriesForTimer(uint32 TimerId);
+	bool HasAnySeriesForTimer(uint32 TimerId) const;
+	uint32 GetNumSeriesForTimer(uint32 TimerId) const;
 
 protected:
 	void UpdateFrameSeries(FTimingGraphSeries& Series, const FTimingTrackViewport& Viewport);
