[Insights]
* Memory Tags: Updated icons for memory tag sets and for memory tag items in the tree view. Also updated icons for the Budget settings and for  the Timer Marker settings.
* Memory Tags: Added info about source data from LLM (and the required options: LLM_ALLOW_ASSETS_TAGS macro and \"-llm -llmtagsets=assets,assetclasses\" command line parameters) to the tooltip of the segmented control for choosing the memory tag set.
* Memory Tags: Fixed text for several tooltips.
* Memory Insights: Added tooltip to the menu options that opens the Investigation, the Memory Tags and the Modules panels.

diff --git a/Engine/Source/Developer/TraceInsights/Private/Insights/InsightsStyle.cpp b/Engine/Source/Developer/TraceInsights/Private/Insights/InsightsStyle.cpp
--- a/Engine/Source/Developer/TraceInsights/Private/Insights/InsightsStyle.cpp
+++ b/Engine/Source/Developer/TraceInsights/Private/Insights/InsightsStyle.cpp
@@ -166,6 +166,8 @@ void FInsightsStyle::FStyle::Initialize()
 	Set("Icons.HighlightEvents.ToolBar", new CORE_IMAGE_BRUSH_SVG("Starship/Common/Visualizer", Icon20x20));
 	Set("Icons.ResetHighlight.ToolBar", new CORE_IMAGE_BRUSH_SVG("Starship/Common/Reject", Icon20x20));
 
+	Set("Icons.TimeMarker", new IMAGE_BRUSH_SVG("TimeMarker", Icon16x16));
+
 	//////////////////////////////////////////////////
 
 	Set("Icons.FindFirst.ToolBar", new IMAGE_BRUSH_SVG("ControlsFirst", Icon20x20));
@@ -220,12 +222,12 @@ void FInsightsStyle::FStyle::Initialize()
 	Set("Icons.AddMemTagGraphs", new CORE_IMAGE_BRUSH_SVG("Starship/Common/plus", Icon16x16));
 	Set("Icons.RemoveMemTagGraphs", new CORE_IMAGE_BRUSH_SVG("Starship/Common/Delete", Icon16x16));
 
-	Set("Icons.TagSet.Systems", new CORE_IMAGE_BRUSH_SVG("Starship/Common/File", Icon16x16));
-	Set("Icons.TagSet.Assets", new EDITOR_IMAGE_BRUSH_SVG("Starship/Common/FileLined", Icon16x16));
-	Set("Icons.TagSet.AssetClasses", new EDITOR_IMAGE_BRUSH_SVG("Starship/Common/Class", Icon16x16));
+	Set("Icons.TagSet.Systems", new IMAGE_BRUSH_SVG("MemTagSet_Systems", Icon16x16));
+	Set("Icons.TagSet.Assets", new IMAGE_BRUSH_SVG("MemTagSet_Assets", Icon16x16));
+	Set("Icons.TagSet.AssetClasses", new IMAGE_BRUSH_SVG("MemTagSet_AssetClasses", Icon16x16));
 
-	Set("Icons.BudgetSettings", new CORE_IMAGE_BRUSH_SVG("Starship/Common/file-tree", Icon16x16));
-	Set("Icons.TimeMarkerSettings", new EDITOR_IMAGE_BRUSH_SVG("Starship/Common/Timeline", Icon16x16));
+	Set("Icons.BudgetSettings", new IMAGE_BRUSH_SVG("BudgetSettings", Icon16x16));
+	Set("Icons.TimeMarkerSettings", new IMAGE_BRUSH_SVG("TimeMarkerSettings", Icon16x16));
 
 	//////////////////////////////////////////////////
 	// Tasks
@@ -279,12 +281,12 @@ void FInsightsStyle::FStyle::Initialize()
 	Set("Icons.NetEvent.TreeItem", new TODO_IMAGE_BRUSH(Icon12x12));
 
 	Set("Icons.MemTag.TreeItem", new IMAGE_BRUSH_SVG("MemTags", Icon12x12));
-	Set("Icons.SystemMemTag.TreeItem", new IMAGE_BRUSH_SVG("MemTags", Icon12x12));
-	Set("Icons.AssetMemTag.TreeItem", new EDITOR_IMAGE_BRUSH_SVG("Starship/Common/FileLined", Icon12x12));
-	Set("Icons.ClassMemTag.TreeItem", new EDITOR_IMAGE_BRUSH_SVG("Starship/Common/Class", Icon12x12));
-	Set("Icons.UObject.TreeItem", new CORE_IMAGE_BRUSH_SVG("Starship/Common/File", Icon12x12));
+	Set("Icons.SystemMemTag.TreeItem", new IMAGE_BRUSH_SVG("MemTag_System_12", Icon12x12));
+	Set("Icons.AssetMemTag.TreeItem", new IMAGE_BRUSH_SVG("MemTag_Asset_12", Icon12x12));
+	Set("Icons.ClassMemTag.TreeItem", new IMAGE_BRUSH_SVG("MemTag_Class_12", Icon12x12));
+	Set("Icons.UObject.TreeItem", new IMAGE_BRUSH_SVG("UObject_12", Icon12x12));
 
-	Set("Icons.HasGraph.TreeItem", new CORE_IMAGE_BRUSH_SVG("Starship/Common/Stop", FVector2D(8.0f, 12.0f)));
+	Set("Icons.HasGraph.TreeItem", new IMAGE_BRUSH_SVG("RoundedBullet", Icon12x12));
 
 	//////////////////////////////////////////////////
 	// Trace Control
diff --git a/Engine/Source/Developer/TraceInsights/Private/Insights/MemoryProfiler/Widgets/SMemTagTreeView.cpp b/Engine/Source/Developer/TraceInsights/Private/Insights/MemoryProfiler/Widgets/SMemTagTreeView.cpp
--- a/Engine/Source/Developer/TraceInsights/Private/Insights/MemoryProfiler/Widgets/SMemTagTreeView.cpp
+++ b/Engine/Source/Developer/TraceInsights/Private/Insights/MemoryProfiler/Widgets/SMemTagTreeView.cpp
@@ -109,7 +109,7 @@ void SMemTagTreeView::ConstructHeaderArea(TSharedRef<SVerticalBox> InHostBox)
 			.Padding(0.0f, 0.0f, 0.0f, 0.0f)
 			[
 				SNew(SBox)
-				.Padding(20.0f, 1.0f, 2.0f, 1.0f)
+				.Padding(4.0f, 1.0f, 2.0f, 1.0f)
 				.Visibility_Lambda([this]
 				{
 					return bAreTimeMarkerSettingsVisible ? EVisibility::Visible : EVisibility::Collapsed;
@@ -276,7 +276,7 @@ void SMemTagTreeView::UpdateAvailableTagSets()
 		}
 
 		FText TagSetToolTip;
-		const FText TagSetToolTipFmt = LOCTEXT("TagSetToolTipFmt", "'{0}' Tag Set\n\nThe tree view shows memory tags from the selected tag set.");
+		const FText TagSetToolTipFmt = LOCTEXT("TagSetToolTipFmt", "'{0}' Tag Set\n\nThe tree view shows memory tags from the selected tag set.\n\nNote:\n\tTrace data is emitted by the Low-Level Memory Tracker system (LLM).\n\tSee documentation about LLM_ALLOW_ASSETS_TAGS macro and\n\tabout \"-llm -llmtagsets=assets,assetclasses\" command line parameters.");
 		TagSetToolTip = FText::Format(TagSetToolTipFmt, TagSetName);
 
 		TSharedRef<SWidget> Widget = SNew(SHorizontalBox)
@@ -287,12 +287,12 @@ void SMemTagTreeView::UpdateAvailableTagSets()
 			.VAlign(VAlign_Center)
 			[
 				SNew(SBox)
-				.WidthOverride(12.0f)
-				.HeightOverride(12.0f)
+				.WidthOverride(16.0f)
+				.HeightOverride(16.0f)
 				[
 					SNew(SImage)
 					.Image(TagSetIcon)
-					.ColorAndOpacity(FSlateColor(EStyleColor::Foreground))
+					.ColorAndOpacity(FSlateColor::UseForeground())
 				]
 			]
 
@@ -302,7 +302,6 @@ void SMemTagTreeView::UpdateAvailableTagSets()
 			.VAlign(VAlign_Center)
 			[
 				SNew(STextBlock)
-				//.Font(FStyleFonts::Get().Small)
 				.Text(TagSetName)
 			];
 
@@ -347,7 +346,7 @@ TSharedPtr<SWidget> SMemTagTreeView::ConstructFilterToolbar()
 			[
 				SNew(SImage)
 				.Image(FInsightsStyle::GetBrush("Icons.Filter.ToolBar"))
-				.ColorAndOpacity(FSlateColor(EStyleColor::Foreground))
+				.ColorAndOpacity(FSlateColor::UseForeground())
 			]
 			+ SHorizontalBox::Slot()
 			.VAlign(VAlign_Center)
@@ -617,7 +616,7 @@ TSharedRef<SButton> SMemTagTreeView::ConstructHideAllButton()
 {
 	return SNew(SButton)
 		.ButtonStyle(&FAppStyle::Get().GetWidgetStyle<FButtonStyle>("SimpleButton"))
-		.ToolTipText(LOCTEXT("HideAll_ToolTip", "Remove all LLM graph tracks."))
+		.ToolTipText(LOCTEXT("RemoveAllGraphTracks_ToolTip", "Remove memory graph tracks (LLM graph series) for all memory tags."))
 		.OnClicked(this, &SMemTagTreeView::HideAllTracks_OnClicked)
 		.ContentPadding(FMargin(2.0f, 2.0f, 2.0f, 2.0f))
 		[
@@ -640,7 +639,7 @@ TSharedRef<SButton> SMemTagTreeView::ConstructShowSelectedButton()
 {
 	return SNew(SButton)
 		.ButtonStyle(&FAppStyle::Get().GetWidgetStyle<FButtonStyle>("SimpleButton"))
-		.ToolTipText(LOCTEXT("ShowSelected_ToolTip", "Create LLM graph tracks for the selected memory tags."))
+		.ToolTipText(LOCTEXT("CreateGraphTracksForSelectedMemTags_ToolTip", "Create memory graph tracks (LLM graph series) for the selected memory tags."))
 		.OnClicked(this, &SMemTagTreeView::ShowSelectedTracks_OnClicked)
 		.IsEnabled(this, &SMemTagTreeView::CanCreateGraphTracksForSelectedMemTags)
 		.ContentPadding(FMargin(2.0f, 2.0f, 2.0f, 2.0f))
@@ -717,7 +716,7 @@ TSharedRef<SWidget> SMemTagTreeView::ConstructBudgetSettings()
 			[
 				SNew(SImage)
 				.Image(FInsightsStyle::GetBrush("Icons.BudgetSettings"))
-				.ColorAndOpacity(FSlateColor(EStyleColor::Foreground))
+				.ColorAndOpacity(FSlateColor::UseForeground())
 			]
 
 			+ SHorizontalBox::Slot()
@@ -887,7 +886,7 @@ void SMemTagTreeView::UpdateSelectionStatsText()
 	else
 	{
 		NumSelectedTagsText = FText::GetEmpty();
-		SelectedTagsText = LOCTEXT("NoSelectionStats", "No LLM tag selected");
+		SelectedTagsText = LOCTEXT("NoSelectionStats", "No memory tag selected");
 		SelectionSizeAText = FText::GetEmpty();
 		SelectionSizeBText = FText::GetEmpty();
 		SelectionDiffText = FText::GetEmpty();
@@ -952,7 +951,7 @@ void SMemTagTreeView::InternalCreateGroupings()
 	AvailableGroupings.Insert(MakeShared<FMemTagBudgetNodeGrouping>(SharedThis(this)), Index++);
 
 #if 0
-	// Groups LLM tags by their hierarchy.
+	// Groups memory tags by their hierarchy.
 	class FMemTagGroupingByHierarchy : public
 	{
 		{
@@ -1002,7 +1001,7 @@ void SMemTagTreeView::Tick(const FGeometry& AllottedGeometry, const double InCur
 {
 	STableTreeView::Tick(AllottedGeometry, InCurrentTime, InDeltaTime);
 
-	// We need to check if the list of LLM tags has changed.
+	// We need to check if the list of memory tags has changed.
 	// But, ensure we do not check too often.
 	static uint64 NextTimestamp = 0;
 	uint64 Time = FPlatformTime::Cycles64();
@@ -1068,34 +1067,37 @@ void SMemTagTreeView::RebuildTree(bool bResync)
 			const TArray<FMemoryTag*>& MemTags = TagList.GetTags();
 			const int32 MemTagCount = MemTags.Num();
 
-			TableRowNodes.Empty(MemTagCount);
-			MemTagNodesIdMap.Empty(MemTagCount);
-			bListHasChanged = true;
+			// Create the table row nodes...
+			{
+				TableRowNodes.Empty(MemTagCount);
+				MemTagNodesIdMap.Empty(MemTagCount);
+				bListHasChanged = true;
 
-			TSharedPtr<FMemTagTable> ParentMemTagTable = GetMemTagTable();
+				TSharedPtr<FMemTagTable> ParentMemTagTable = GetMemTagTable();
 
-			for (FMemoryTag* MemTagPtr : MemTags)
-			{
-				check(MemTagPtr);
-				TSharedPtr<FMemTagNode> MemTagNodePtr;
-				if (MemTagPtr->GetTagSetId() == SystemsTagSet)
-				{
-					MemTagNodePtr = MakeShared<FSystemMemTagNode>(ParentMemTagTable, *MemTagPtr);
-				}
-				else if (MemTagPtr->GetTagSetId() == AssetsTagSet)
-				{
-					MemTagNodePtr = MakeShared<FAssetMemTagNode>(ParentMemTagTable, *MemTagPtr);
-				}
-				else if (MemTagPtr->GetTagSetId() == AssetClassesTagSet)
+				for (FMemoryTag* MemTagPtr : MemTags)
 				{
-					MemTagNodePtr = MakeShared<FClassMemTagNode>(ParentMemTagTable, *MemTagPtr);
-				}
-				else
-				{
-					MemTagNodePtr = MakeShared<FMemTagNode>(ParentMemTagTable, *MemTagPtr);
+					check(MemTagPtr);
+					TSharedPtr<FMemTagNode> MemTagNodePtr;
+					if (MemTagPtr->GetTagSetId() == SystemsTagSet)
+					{
+						MemTagNodePtr = MakeShared<FSystemMemTagNode>(ParentMemTagTable, *MemTagPtr);
+					}
+					else if (MemTagPtr->GetTagSetId() == AssetsTagSet)
+					{
+						MemTagNodePtr = MakeShared<FAssetMemTagNode>(ParentMemTagTable, *MemTagPtr);
+					}
+					else if (MemTagPtr->GetTagSetId() == AssetClassesTagSet)
+					{
+						MemTagNodePtr = MakeShared<FClassMemTagNode>(ParentMemTagTable, *MemTagPtr);
+					}
+					else
+					{
+						MemTagNodePtr = MakeShared<FMemTagNode>(ParentMemTagTable, *MemTagPtr);
+					}
+					TableRowNodes.Add(MemTagNodePtr);
+					MemTagNodesIdMap.Add(MemTagPtr->GetId(), MemTagNodePtr);
 				}
-				TableRowNodes.Add(MemTagNodePtr);
-				MemTagNodesIdMap.Add(MemTagPtr->GetId(), MemTagNodePtr);
 			}
 
 			// Resolve pointers to parent tags.
@@ -1435,7 +1437,7 @@ void SMemTagTreeView::ExtendMenuBeforeMisc(FMenuBuilder& MenuBuilder)
 		MenuBuilder.AddMenuEntry
 		(
 			LOCTEXT("ContextMenu_GenerateColorForSelectedMemTags", "Generate New Color"),
-			LOCTEXT("ContextMenu_GenerateColorForSelectedMemTags_Desc", "Generate new color for selected memory tag(s)."),
+			LOCTEXT("ContextMenu_GenerateColorForSelectedMemTags_Desc", "Generate new color for selected memory tags."),
 			FSlateIcon(),
 			Action_GenerateColorForSelectedMemTags,
 			NAME_None,
@@ -1450,7 +1452,7 @@ void SMemTagTreeView::ExtendMenuBeforeMisc(FMenuBuilder& MenuBuilder)
 		MenuBuilder.AddMenuEntry
 		(
 			LOCTEXT("ContextMenu_EditColorForSelectedMemTags", "Edit Color..."),
-			LOCTEXT("ContextMenu_EditColorForSelectedMemTags_Desc", "Edit color for selected memory tag(s)."),
+			LOCTEXT("ContextMenu_EditColorForSelectedMemTags_Desc", "Edit color for selected memory tags."),
 			FSlateIcon(),
 			Action_EditColorForSelectedMemTags,
 			NAME_None,
@@ -1468,7 +1470,7 @@ void SMemTagTreeView::ExtendMenuCreateGraphTracks(FMenuBuilder& MenuBuilder)
 
 	MenuBuilder.BeginSection("CreateGraphTracks");
 	{
-		// Create memory graph tracks for selected LLM tag(s)
+		// Create memory graph tracks (LLM graph series) for the selected memory tags.
 		FUIAction Action_CreateGraphTracksForSelectedMemTags
 		(
 			FExecuteAction::CreateSP(this, &SMemTagTreeView::CreateGraphTracksForSelectedMemTags),
@@ -1477,14 +1479,14 @@ void SMemTagTreeView::ExtendMenuCreateGraphTracks(FMenuBuilder& MenuBuilder)
 		MenuBuilder.AddMenuEntry
 		(
 			LOCTEXT("ContextMenu_CreateGraphTracksForSelectedMemTags", "Selected"),
-			LOCTEXT("ContextMenu_CreateGraphTracksForSelectedMemTags_Desc", "Create memory graph tracks (LLM graph series) for the selected memory tag(s)."),
+			LOCTEXT("ContextMenu_CreateGraphTracksForSelectedMemTags_Desc", "Create memory graph tracks (LLM graph series) for the selected memory tags."),
 			FSlateIcon(),
 			Action_CreateGraphTracksForSelectedMemTags,
 			NAME_None,
 			EUserInterfaceActionType::Button
 		);
 
-		// Create memory graph tracks for filtered LLM tags
+		// Create memory graph tracks (LLM graph series) for the visible (filtered) memory tags.
 		FUIAction Action_CreateGraphTracksForFilteredMemTags
 		(
 			FExecuteAction::CreateSP(this, &SMemTagTreeView::CreateGraphTracksForFilteredMemTags),
@@ -1500,7 +1502,7 @@ void SMemTagTreeView::ExtendMenuCreateGraphTracks(FMenuBuilder& MenuBuilder)
 			EUserInterfaceActionType::Button
 		);
 
-		// Create memory graph tracks for all LLM tags
+		// Create memory graph tracks (LLM graph series) for all memory tags.
 		FUIAction Action_CreateAllGraphTracks
 		(
 			FExecuteAction::CreateSP(this, &SMemTagTreeView::CreateAllGraphTracks),
@@ -1527,7 +1529,7 @@ void SMemTagTreeView::ExtendMenuRemoveGraphTracks(FMenuBuilder& MenuBuilder)
 
 	MenuBuilder.BeginSection("RemoveGraphTracks");
 	{
-		// Remove memory graph tracks for selected LLM tag(s)
+		// Remove memory graph tracks (LLM graph series) for the selected memory tags.
 		FUIAction Action_RemoveGraphTracksForSelectedMemTags
 		(
 			FExecuteAction::CreateSP(this, &SMemTagTreeView::RemoveGraphTracksForSelectedMemTags),
@@ -1536,14 +1538,14 @@ void SMemTagTreeView::ExtendMenuRemoveGraphTracks(FMenuBuilder& MenuBuilder)
 		MenuBuilder.AddMenuEntry
 		(
 			LOCTEXT("ContextMenu_RemoveGraphTracksForSelectedMemTags", "Selected"),
-			LOCTEXT("ContextMenu_RemoveGraphTracksForSelectedMemTags_Desc", "Remove memory graph tracks (LLM graph series) for the selected memory tag(s)."),
+			LOCTEXT("ContextMenu_RemoveGraphTracksForSelectedMemTags_Desc", "Remove memory graph tracks (LLM graph series) for the selected memory tags."),
 			FSlateIcon(),
 			Action_RemoveGraphTracksForSelectedMemTags,
 			NAME_None,
 			EUserInterfaceActionType::Button
 		);
 
-		// Remove memory graph tracks for all LLM tags
+		// Remove memory graph tracks (LLM graph series) for all memory tags.
 		FUIAction Action_RemoveAllGraphTracks
 		(
 			FExecuteAction::CreateSP(this, &SMemTagTreeView::RemoveAllGraphTracks),
@@ -1881,7 +1883,7 @@ void SMemTagTreeView::RemoveAllGraphTracks()
 }
 
 ////////////////////////////////////////////////////////////////////////////////////////////////////
-// Generate new color for selected LLM tag(s)
+// Generate new color for selected memory tags
 ////////////////////////////////////////////////////////////////////////////////////////////////////
 
 bool SMemTagTreeView::CanGenerateColorForSelectedMemTags() const
@@ -1976,7 +1978,7 @@ void SMemTagTreeView::SetColorToNode(FMemTagNode& MemTagNode, FLinearColor Color
 }
 
 ////////////////////////////////////////////////////////////////////////////////////////////////////
-// Edit color for selected LLM tag(s)
+// Edit color for selected memory tags
 ////////////////////////////////////////////////////////////////////////////////////////////////////
 
 bool SMemTagTreeView::CanEditColorForSelectedMemTags() const
diff --git a/Engine/Source/Developer/TraceInsights/Private/Insights/MemoryProfiler/Widgets/SMemoryProfilerWindow.cpp b/Engine/Source/Developer/TraceInsights/Private/Insights/MemoryProfiler/Widgets/SMemoryProfilerWindow.cpp
--- a/Engine/Source/Developer/TraceInsights/Private/Insights/MemoryProfiler/Widgets/SMemoryProfilerWindow.cpp
+++ b/Engine/Source/Developer/TraceInsights/Private/Insights/MemoryProfiler/Widgets/SMemoryProfilerWindow.cpp
@@ -412,6 +412,7 @@ void SMemoryProfilerWindow::RegisterTabSpawners()
 		TabManagerPtr->RegisterTabSpawner(FMemoryProfilerTabs::TimingViewID, FOnSpawnTab::CreateRaw(this, &SMemoryProfilerWindow::SpawnTab_TimingView))
 			.SetDisplayName(LOCTEXT("TimingViewTabTitle", "Timing View"))
 			.SetIcon(FSlateIcon(FInsightsStyle::GetStyleSetName(), "Icons.TimingView"))
+			.SetTooltipText(LOCTEXT("TimingViewTabToolTip", "Opens the Timing View."))
 			.SetGroup(Group);
 	}
 
@@ -420,6 +421,7 @@ void SMemoryProfilerWindow::RegisterTabSpawners()
 		TabManagerPtr->RegisterTabSpawner(FMemoryProfilerTabs::MemInvestigationViewID, FOnSpawnTab::CreateRaw(this, &SMemoryProfilerWindow::SpawnTab_MemInvestigationView))
 			.SetDisplayName(LOCTEXT("MemInvestigationViewTabTitle", "Investigation"))
 			.SetIcon(FSlateIcon(FInsightsStyle::GetStyleSetName(), "Icons.MemInvestigationView"))
+			.SetTooltipText(LOCTEXT("MemInvestigationViewTabToolTip", "Opens the Memory Investigation panel.\n\nIt analyses the detailed trace data about memory allocations.\nSee documentation about \"-trace=memalloc\" or \"-trace=memory\"."))
 			.SetGroup(Group);
 	}
 
@@ -428,6 +430,7 @@ void SMemoryProfilerWindow::RegisterTabSpawners()
 		TabManagerPtr->RegisterTabSpawner(FMemoryProfilerTabs::MemTagTreeViewID, FOnSpawnTab::CreateRaw(this, &SMemoryProfilerWindow::SpawnTab_MemTagTreeView))
 			.SetDisplayName(LOCTEXT("MemTagTreeViewTabTitle", "Memory Tags"))
 			.SetIcon(FSlateIcon(FInsightsStyle::GetStyleSetName(), "Icons.MemTagTreeView"))
+			.SetTooltipText(LOCTEXT("MemTagTreeViewTabToolTip", "Opens the Memory Tags (LLM) panel.\n\nIt analyses the per-frame trace data emitted by the Low-Level Memory Tracker system (LLM).\nSee documentation about LLM_ALLOW_ASSETS_TAGS macro and\nabout \"-llm -llmtagsets=assets,assetclasses\" command line parameters."))
 			.SetGroup(Group);
 	}
 
@@ -436,6 +439,7 @@ void SMemoryProfilerWindow::RegisterTabSpawners()
 		TabManagerPtr->RegisterTabSpawner(FMemoryProfilerTabs::ModulesViewID, FOnSpawnTab::CreateRaw(this, &SMemoryProfilerWindow::SpawnTab_ModulesView))
 			.SetDisplayName(LOCTEXT("ModulesViewTabTitle", "Modules"))
 			.SetIcon(FSlateIcon(FInsightsStyle::GetStyleSetName(), "Icons.ModulesView"))
+			.SetTooltipText(LOCTEXT("TimingViewTabToolTip", "Opens the Modules panel."))
 			.SetGroup(Group);
 	}
 }
