[Insights] Memory Tags: Fixed rare crash happening when the async filter operation accesses the list of tree nodes in the same time with main thread (re)creating the tree nodes.

#rnx

diff --git a/Engine/Source/Developer/TraceInsights/Private/Insights/MemoryProfiler/Widgets/SMemTagTreeView.cpp b/Engine/Source/Developer/TraceInsights/Private/Insights/MemoryProfiler/Widgets/SMemTagTreeView.cpp
--- a/Engine/Source/Developer/TraceInsights/Private/Insights/MemoryProfiler/Widgets/SMemTagTreeView.cpp
+++ b/Engine/Source/Developer/TraceInsights/Private/Insights/MemoryProfiler/Widgets/SMemTagTreeView.cpp
@@ -1005,7 +1005,7 @@ void SMemTagTreeView::Tick(const FGeometry& AllottedGeometry, const double InCur
 	// But, ensure we do not check too often.
 	static uint64 NextTimestamp = 0;
 	uint64 Time = FPlatformTime::Cycles64();
-	if (Time > NextTimestamp)
+	if (Time > NextTimestamp && !bIsUpdateRunning)
 	{
 		RebuildTree(false);
 
