[Insights] Memory Tags: Fixed "Create Filtered Graph / Filtered" to correctly identify the visible tags when using the Search text box. Also renamed "Filtered" to "Visible" to avoid confusions.

diff --git a/Engine/Source/Developer/TraceInsights/Private/Insights/MemoryProfiler/Widgets/SMemTagTreeView.cpp b/Engine/Source/Developer/TraceInsights/Private/Insights/MemoryProfiler/Widgets/SMemTagTreeView.cpp
--- a/Engine/Source/Developer/TraceInsights/Private/Insights/MemoryProfiler/Widgets/SMemTagTreeView.cpp
+++ b/Engine/Source/Developer/TraceInsights/Private/Insights/MemoryProfiler/Widgets/SMemTagTreeView.cpp
@@ -1414,7 +1414,7 @@ void SMemTagTreeView::ExtendMenuBeforeMisc(FMenuBuilder& MenuBuilder)
 		MenuBuilder.AddSubMenu
 		(
 			LOCTEXT("ContextMenu_CreateGraphTracks_SubMenu", "Create Graph Tracks"),
-			LOCTEXT("ContextMenu_CreateGraphTracks_SubMenu_Desc", "Create memory graph tracks (LLM graph series)."),
+			LOCTEXT("ContextMenu_CreateGraphTracks_SubMenu_Desc", "Creates memory graph tracks (LLM graph series)."),
 			FNewMenuDelegate::CreateSP(this, &SMemTagTreeView::ExtendMenuCreateGraphTracks),
 			false,
 			FSlateIcon(FInsightsStyle::GetStyleSetName(), "Icons.AddMemTagGraphs")
@@ -1423,7 +1423,7 @@ void SMemTagTreeView::ExtendMenuBeforeMisc(FMenuBuilder& MenuBuilder)
 		MenuBuilder.AddSubMenu
 		(
 			LOCTEXT("ContextMenu_RemoveGraphTracks_SubMenu", "Remove Graph Tracks"),
-			LOCTEXT("ContextMenu_RemoveGraphTracks_SubMenu_Desc", "Remove memory graph tracks (LLM graph series)."),
+			LOCTEXT("ContextMenu_RemoveGraphTracks_SubMenu_Desc", "Removes memory graph tracks (LLM graph series)."),
 			FNewMenuDelegate::CreateSP(this, &SMemTagTreeView::ExtendMenuRemoveGraphTracks),
 			false,
 			FSlateIcon(FInsightsStyle::GetStyleSetName(), "Icons.RemoveMemTagGraphs")
@@ -1437,7 +1437,7 @@ void SMemTagTreeView::ExtendMenuBeforeMisc(FMenuBuilder& MenuBuilder)
 		MenuBuilder.AddMenuEntry
 		(
 			LOCTEXT("ContextMenu_GenerateColorForSelectedMemTags", "Generate New Color"),
-			LOCTEXT("ContextMenu_GenerateColorForSelectedMemTags_Desc", "Generate new color for selected memory tags."),
+			LOCTEXT("ContextMenu_GenerateColorForSelectedMemTags_Desc", "Generates new color for the selected memory tags."),
 			FSlateIcon(),
 			Action_GenerateColorForSelectedMemTags,
 			NAME_None,
@@ -1452,7 +1452,7 @@ void SMemTagTreeView::ExtendMenuBeforeMisc(FMenuBuilder& MenuBuilder)
 		MenuBuilder.AddMenuEntry
 		(
 			LOCTEXT("ContextMenu_EditColorForSelectedMemTags", "Edit Color..."),
-			LOCTEXT("ContextMenu_EditColorForSelectedMemTags_Desc", "Edit color for selected memory tags."),
+			LOCTEXT("ContextMenu_EditColorForSelectedMemTags_Desc", "Changes color for the selected memory tags."),
 			FSlateIcon(),
 			Action_EditColorForSelectedMemTags,
 			NAME_None,
@@ -1479,25 +1479,25 @@ void SMemTagTreeView::ExtendMenuCreateGraphTracks(FMenuBuilder& MenuBuilder)
 		MenuBuilder.AddMenuEntry
 		(
 			LOCTEXT("ContextMenu_CreateGraphTracksForSelectedMemTags", "Selected"),
-			LOCTEXT("ContextMenu_CreateGraphTracksForSelectedMemTags_Desc", "Create memory graph tracks (LLM graph series) for the selected memory tags."),
+			LOCTEXT("ContextMenu_CreateGraphTracksForSelectedMemTags_Desc", "Creates memory graph tracks (LLM graph series) for the selected memory tags."),
 			FSlateIcon(),
 			Action_CreateGraphTracksForSelectedMemTags,
 			NAME_None,
 			EUserInterfaceActionType::Button
 		);
 
-		// Create memory graph tracks (LLM graph series) for the visible (filtered) memory tags.
-		FUIAction Action_CreateGraphTracksForFilteredMemTags
+		// Create memory graph tracks (LLM graph series) for the visible memory tags.
+		FUIAction Action_CreateGraphTracksForVisibleMemTags
 		(
-			FExecuteAction::CreateSP(this, &SMemTagTreeView::CreateGraphTracksForFilteredMemTags),
-			FCanExecuteAction::CreateSP(this, &SMemTagTreeView::CanCreateGraphTracksForFilteredMemTags)
+			FExecuteAction::CreateSP(this, &SMemTagTreeView::CreateGraphTracksForVisibleMemTags),
+			FCanExecuteAction::CreateSP(this, &SMemTagTreeView::CanCreateGraphTracksForVisibleMemTags)
 		);
 		MenuBuilder.AddMenuEntry
 		(
-			LOCTEXT("ContextMenu_CreateGraphTracksForFilteredMemTags", "Filtered"),
-			LOCTEXT("ContextMenu_CreateGraphTracksForFilteredMemTags_Desc", "Create memory graph tracks (LLM graph series) for the visible (filtered) memory tags."),
+			LOCTEXT("ContextMenu_CreateGraphTracksForVisibleMemTags", "Visible"),
+			LOCTEXT("ContextMenu_CreateGraphTracksForVisibleMemTags_Desc", "Creates memory graph tracks (LLM graph series) for the visible memory tags."),
 			FSlateIcon(),
-			Action_CreateGraphTracksForFilteredMemTags,
+			Action_CreateGraphTracksForVisibleMemTags,
 			NAME_None,
 			EUserInterfaceActionType::Button
 		);
@@ -1511,7 +1511,7 @@ void SMemTagTreeView::ExtendMenuCreateGraphTracks(FMenuBuilder& MenuBuilder)
 		MenuBuilder.AddMenuEntry
 		(
 			LOCTEXT("ContextMenu_CreateAllGraphTracks", "All"),
-			LOCTEXT("ContextMenu_CreateAllGraphTracks_Desc", "Create memory graph tracks (LLM graph series) for all memory tags."),
+			LOCTEXT("ContextMenu_CreateAllGraphTracks_Desc", "Creates memory graph tracks (LLM graph series) for all memory tags."),
 			FSlateIcon(),
 			Action_CreateAllGraphTracks,
 			NAME_None,
@@ -1538,7 +1538,7 @@ void SMemTagTreeView::ExtendMenuRemoveGraphTracks(FMenuBuilder& MenuBuilder)
 		MenuBuilder.AddMenuEntry
 		(
 			LOCTEXT("ContextMenu_RemoveGraphTracksForSelectedMemTags", "Selected"),
-			LOCTEXT("ContextMenu_RemoveGraphTracksForSelectedMemTags_Desc", "Remove memory graph tracks (LLM graph series) for the selected memory tags."),
+			LOCTEXT("ContextMenu_RemoveGraphTracksForSelectedMemTags_Desc", "Removes memory graph tracks (LLM graph series) for the selected memory tags."),
 			FSlateIcon(),
 			Action_RemoveGraphTracksForSelectedMemTags,
 			NAME_None,
@@ -1554,7 +1554,7 @@ void SMemTagTreeView::ExtendMenuRemoveGraphTracks(FMenuBuilder& MenuBuilder)
 		MenuBuilder.AddMenuEntry
 		(
 			LOCTEXT("ContextMenu_RemoveAllGraphTracks", "All"),
-			LOCTEXT("ContextMenu_RemoveAllGraphTracks_Desc", "Remove memory graph tracks (LLM graph series) for all memory tags."),
+			LOCTEXT("ContextMenu_RemoveAllGraphTracks_Desc", "Removes memory graph tracks (LLM graph series) for all memory tags."),
 			FSlateIcon(),
 			Action_RemoveAllGraphTracks,
 			NAME_None,
@@ -1576,7 +1576,7 @@ void SMemTagTreeView::ExtendMenuAfterMisc(FMenuBuilder& MenuBuilder)
 	MenuBuilder.AddMenuEntry
 	(
 		LOCTEXT("ContextMenu_LoadReportXML", "Load Report XML..."),
-		LOCTEXT("ContextMenu_LoadReportXML_Desc", "Load LLMReportTypes.xml"),
+		LOCTEXT("ContextMenu_LoadReportXML_Desc", "Loads a report specification file (LLMReportTypes.xml)"),
 		FSlateIcon(FAppStyle::Get().GetStyleSetName(), "Icons.FolderOpen"),
 		Action_LoadReportXML,
 		NAME_None,
@@ -1764,26 +1764,46 @@ void SMemTagTreeView::CreateGraphTracksForSelectedMemTags()
 }
 
 ////////////////////////////////////////////////////////////////////////////////////////////////////
-// Create memory graph tracks (LLM graph series) for the visible (filtered) memory tags
+// Create memory graph tracks (LLM graph series) for the visible memory tags
 ////////////////////////////////////////////////////////////////////////////////////////////////////
 
-bool SMemTagTreeView::CanCreateGraphTracksForFilteredMemTags() const
+bool SMemTagTreeView::CanCreateGraphTracksForVisibleMemTags() const
 {
-	return FilteredNodesPtr->Num() > 0;
+	return FilteredNodesPtr->Num() > 0 && FilteredGroupNodes.Num() > 0;
 }
 
 ////////////////////////////////////////////////////////////////////////////////////////////////////
 
-void SMemTagTreeView::CreateGraphTracksForFilteredMemTags()
+void SMemTagTreeView::CreateGraphTracksForVisibleMemTags()
 {
 	TSharedPtr<SMemoryProfilerWindow> ProfilerWindow = GetProfilerWindow();
 	if (ProfilerWindow.IsValid())
 	{
 		FMemorySharedState& SharedState = ProfilerWindow->GetSharedState();
 
-		for (const FTableTreeNodePtr& FilteredNode : *FilteredNodesPtr)
+		for (const FTableTreeNodePtr& GroupNode : FilteredGroupNodes)
 		{
-			TryCreateGraphTrackForNode(SharedState, *FilteredNode);
+			CreateGraphTracksRec(SharedState, *GroupNode);
+		}
+	}
+}
+
+////////////////////////////////////////////////////////////////////////////////////////////////////
+
+void SMemTagTreeView::CreateGraphTracksRec(FMemorySharedState& SharedState, const FBaseTreeNode& Node)
+{
+	if (Node.Is<FMemTagNode>())
+	{
+		const FMemTagNode& MemTagNode = Node.As<FMemTagNode>();
+		TryCreateGraphTrackForNode(SharedState, MemTagNode);
+	}
+
+	if (Node.IsGroup())
+	{
+		const TArray<FBaseTreeNodePtr>& Children = Node.GetFilteredChildren();
+		for (const FBaseTreeNodePtr& Child : Children)
+		{
+			CreateGraphTracksRec(SharedState, *Child);
 		}
 	}
 }
diff --git a/Engine/Source/Developer/TraceInsights/Private/Insights/MemoryProfiler/Widgets/SMemTagTreeView.h b/Engine/Source/Developer/TraceInsights/Private/Insights/MemoryProfiler/Widgets/SMemTagTreeView.h
--- a/Engine/Source/Developer/TraceInsights/Private/Insights/MemoryProfiler/Widgets/SMemTagTreeView.h
+++ b/Engine/Source/Developer/TraceInsights/Private/Insights/MemoryProfiler/Widgets/SMemTagTreeView.h
@@ -255,19 +255,20 @@ private:
 	void TryCreateGraphTrackForNode(FMemorySharedState& SharedState, const FBaseTreeNode& Node) const;
 	void TryRemoveGraphTrackForNode(FMemorySharedState& SharedState, const FBaseTreeNode& Node) const;
 
-	// Create memory graph tracks for selected LLM tag(s)
+	// Create memory graph tracks for the selected LLM tags
 	bool CanCreateGraphTracksForSelectedMemTags() const;
 	void CreateGraphTracksForSelectedMemTags();
 
-	// Create memory graph tracks for filtered LLM tags
-	bool CanCreateGraphTracksForFilteredMemTags() const;
-	void CreateGraphTracksForFilteredMemTags();
+	// Create memory graph tracks for the visible LLM tags
+	bool CanCreateGraphTracksForVisibleMemTags() const;
+	void CreateGraphTracksForVisibleMemTags();
+	void CreateGraphTracksRec(FMemorySharedState& SharedState, const FBaseTreeNode& Node);
 
 	// Create all memory graph tracks
 	bool CanCreateAllGraphTracks() const;
 	void CreateAllGraphTracks();
 
-	// Remove memory graph tracks for selected LLM tag(s)
+	// Remove memory graph tracks for the selected LLM tags
 	bool CanRemoveGraphTracksForSelectedMemTags() const;
 	void RemoveGraphTracksForSelectedMemTags();
 
@@ -275,11 +276,11 @@ private:
 	bool CanRemoveAllGraphTracks() const;
 	void RemoveAllGraphTracks();
 
-	// Generate new color for selected LLM tag(s)
+	// Generate new color for the selected LLM tags
 	bool CanGenerateColorForSelectedMemTags() const;
 	void GenerateColorForSelectedMemTags() const;
 
-	// Edit color for selected LLM tag(s)
+	// Edit color for the selected LLM tags
 	bool CanEditColorForSelectedMemTags() const;
 	void EditColorForSelectedMemTags();
 
