[Insights] Fixed display of Y-axis labels for graph series of per-frame timers.

diff --git a/Engine/Source/Developer/TraceInsights/Private/Insights/ViewModels/TimingGraphTrack.cpp b/Engine/Source/Developer/TraceInsights/Private/Insights/ViewModels/TimingGraphTrack.cpp
--- a/Engine/Source/Developer/TraceInsights/Private/Insights/ViewModels/TimingGraphTrack.cpp
+++ b/Engine/Source/Developer/TraceInsights/Private/Insights/ViewModels/TimingGraphTrack.cpp
@@ -40,7 +40,9 @@ FTimingGraphSeries::FTimingGraphSeries(FTimingGraphSeries::ESeriesType InType)
 	, TimerId(0)
 	, CachedSessionDuration(0.0)
 	, CachedEvents()
-	, bIsTime(InType == FTimingGraphSeries::ESeriesType::Frame || InType == FTimingGraphSeries::ESeriesType::Timer)
+	, bIsTime(InType == FTimingGraphSeries::ESeriesType::Frame ||
+			  InType == FTimingGraphSeries::ESeriesType::Timer ||
+			  InType == FTimingGraphSeries::ESeriesType::FrameStatsTimer)
 	, bIsMemory(false)
 	, bIsFloatingPoint(false)
 {
