[Insights] Memory Tags: Improved the performance of updating the mem tags in UI with the data from the provider/analysis thread.

#rnx

diff --git a/Engine/Source/Developer/TraceInsights/Private/Insights/MemoryProfiler/ViewModels/MemoryTag.cpp b/Engine/Source/Developer/TraceInsights/Private/Insights/MemoryProfiler/ViewModels/MemoryTag.cpp
--- a/Engine/Source/Developer/TraceInsights/Private/Insights/MemoryProfiler/ViewModels/MemoryTag.cpp
+++ b/Engine/Source/Developer/TraceInsights/Private/Insights/MemoryProfiler/ViewModels/MemoryTag.cpp
@@ -137,71 +137,77 @@ void FMemoryTagList::Update()
 
 void FMemoryTagList::UpdateInternal()
 {
-	bool bUpdated = false;
-
 	TSharedPtr<const TraceServices::IAnalysisSession> Session = FInsightsManager::Get()->GetSession();
-	if (Session.IsValid())
+	if (!Session.IsValid())
 	{
-		const TraceServices::IMemoryProvider* MemoryProvider = TraceServices::ReadMemoryProvider(*Session.Get());
-		if (MemoryProvider)
-		{
-			TraceServices::FProviderReadScopeLock _(*MemoryProvider);
+		return;
+	}
 
-			const int32 TraceSerialNumber = MemoryProvider->GetTagSerial();
-			if (LastTraceSerialNumber != TraceSerialNumber)
-			{
-				LastTraceSerialNumber = TraceSerialNumber;
-				++SerialNumber;
-				bUpdated = true;
+	const TraceServices::IMemoryProvider* MemoryProvider = TraceServices::ReadMemoryProvider(*Session.Get());
+	if (!MemoryProvider)
+	{
+		return;
+	}
 
-				MemoryProvider->EnumerateTags([this](const TraceServices::FMemoryTagInfo& TraceTag)
-				{
-					static_assert(FMemoryTag::InvalidTagId == static_cast<FMemoryTagId>(TraceServices::FMemoryTagInfo::InvalidTagId), "Memory TagId type mismatch!");
-					FMemoryTagId TagId = static_cast<FMemoryTagId>(TraceTag.Id);
+	TraceServices::FProviderReadScopeLock _(*MemoryProvider);
 
-					uint64 TrackerFlags = TraceTag.Trackers;
-					int32 TrackerId = 0;
-					while (TrackerFlags != 0)
+	const int32 TraceSerialNumber = MemoryProvider->GetTagSerial();
+	if (LastTraceSerialNumber == TraceSerialNumber)
+	{
+		return;
+	}
+	LastTraceSerialNumber = TraceSerialNumber;
+
+	uint32 AddedOrUpdatedTags = 0;
+
+	MemoryProvider->EnumerateTags([this, &AddedOrUpdatedTags](const TraceServices::FMemoryTagInfo& TraceTag)
+	{
+		static_assert(FMemoryTag::InvalidTagId == static_cast<FMemoryTagId>(TraceServices::FMemoryTagInfo::InvalidTagId), "Memory TagId type mismatch!");
+		FMemoryTagId TagId = static_cast<FMemoryTagId>(TraceTag.Id);
+
+		uint64 TrackerFlags = TraceTag.Trackers;
+		int32 TrackerId = 0;
+		while (TrackerFlags != 0)
+		{
+			if (TrackerFlags & 1)
+			{
+				TMap<FMemoryTagId, FMemoryTag*>& TagsMap = TrackersAndTagsMap.FindOrAdd(FMemoryTrackerId(TrackerId));
+				FMemoryTag* TagPtr = TagsMap.FindRef(TagId);
+				if (!TagPtr)
+				{
+					TagPtr = new FMemoryTag();
+
+					FMemoryTag& Tag = *TagPtr;
+					Tag.Index = Tags.Num();
+					Tag.Id = TagId;
+					Tag.ParentId = static_cast<FMemoryTagId>(TraceTag.ParentId);
+					Tag.StatName = TraceTag.Name;
+					Tag.StatFullName = TraceTag.Name;
+					Tag.TrackerId = static_cast<FMemoryTrackerId>(TrackerId);
+					Tag.TagSetId = static_cast<FMemoryTagSetId>(TraceTag.TagSetId);
+					Tag.SetColorAuto();
+
+					// Skip the parent prefix if it is already included in the name.
+					int32 CharIndex;
+					if (Tag.StatName.FindLastChar(TEXT('/'), CharIndex))
 					{
-						if (TrackerFlags & 1)
-						{
-							TMap<FMemoryTagId, FMemoryTag*>& TagsMap = TrackersAndTagsMap.FindOrAdd(FMemoryTrackerId(TrackerId));
-							FMemoryTag* TagPtr = TagsMap.FindRef(TagId);
-							if (!TagPtr)
-							{
-								TagPtr = new FMemoryTag();
-
-								FMemoryTag& Tag = *TagPtr;
-								Tag.Index = Tags.Num();
-								Tag.Id = TagId;
-								Tag.ParentId = static_cast<FMemoryTagId>(TraceTag.ParentId);
-								Tag.StatName = TraceTag.Name;
-								Tag.StatFullName = TraceTag.Name;
-								Tag.TrackerId = static_cast<FMemoryTrackerId>(TrackerId);
-								Tag.TagSetId = static_cast<FMemoryTagSetId>(TraceTag.TagSetId);
-								Tag.SetColorAuto();
-
-								// Skip the parent prefix if it is already included in the name.
-								int32 CharIndex;
-								if (Tag.StatName.FindLastChar(TEXT('/'), CharIndex))
-								{
-									Tag.StatName = Tag.StatName.RightChop(CharIndex + 1);
-								}
-
-								Tags.Add(TagPtr);
-								TagsMap.Add(TagId, TagPtr);
-							}
-						}
-						TrackerFlags >>= 1;
-						++TrackerId;
+						Tag.StatName = Tag.StatName.RightChop(CharIndex + 1);
 					}
-				});
+
+					Tags.Add(TagPtr);
+					TagsMap.Add(TagId, TagPtr);
+					++AddedOrUpdatedTags;
+				}
 			}
+			TrackerFlags >>= 1;
+			++TrackerId;
 		}
-	}
+	});
 
-	if (bUpdated)
+	if (AddedOrUpdatedTags > 0)
 	{
+		++SerialNumber;
+
 		// Update Parent and StatFullName for each tag.
 		for (FMemoryTag* TagPtr : Tags)
 		{
