Add support for exporting gpu queue events.
Fix the exporting for timing event with metadata

diff --git a/Engine/Source/Developer/TraceInsights/Private/Insights/TimingProfiler/ViewModels/ThreadTimingSharedState.cpp b/Engine/Source/Developer/TraceInsights/Private/Insights/TimingProfiler/ViewModels/ThreadTimingSharedState.cpp
--- a/Engine/Source/Developer/TraceInsights/Private/Insights/TimingProfiler/ViewModels/ThreadTimingSharedState.cpp
+++ b/Engine/Source/Developer/TraceInsights/Private/Insights/TimingProfiler/ViewModels/ThreadTimingSharedState.cpp
@@ -421,8 +421,7 @@ void FThreadTimingSharedState::Tick(Timing::ITimingViewSession& InSession, const
 				if (TrackPtrPtr == nullptr)
 				{
 					// Create new Timing Events track for the GPU Queue.
-					const FString TrackName = FString::Printf(TEXT("GPU%u - %s %u"), QueueInfo.GPU, QueueInfo.Name, QueueInfo.Index);
-					TSharedPtr<FGpuQueueTimingTrack> Track = MakeShared<FGpuQueueTimingTrack>(*this, TrackName, QueueInfo.TimelineIndex, QueueInfo.Id);
+					TSharedPtr<FGpuQueueTimingTrack> Track = MakeShared<FGpuQueueTimingTrack>(*this, QueueInfo.GetDisplayName(), QueueInfo.TimelineIndex, QueueInfo.Id);
 					Track->SetOrder(GpuTrackOrder);
 					Track->SetVisibilityFlag(bShowHideAllGpuTracks);
 					GpuTracks.Add(QueueInfo.Id, Track);
diff --git a/Engine/Source/Developer/TraceInsights/Private/Insights/TimingProfiler/ViewModels/TimingExporter.cpp b/Engine/Source/Developer/TraceInsights/Private/Insights/TimingProfiler/ViewModels/TimingExporter.cpp
--- a/Engine/Source/Developer/TraceInsights/Private/Insights/TimingProfiler/ViewModels/TimingExporter.cpp
+++ b/Engine/Source/Developer/TraceInsights/Private/Insights/TimingProfiler/ViewModels/TimingExporter.cpp
@@ -176,6 +176,26 @@ int32 FTimingExporter::ExportThreadsAsText(const FString& Filename, FExportThrea
 		Writer.AppendLineEnd();
 		++ThreadCount;
 
+		// Iterate the Gpu Queues.
+		{
+			TraceServices::FAnalysisSessionReadScope SessionReadScope(Session);
+
+			const TraceServices::ITimingProfilerProvider* TimingProfilerProvider = TraceServices::ReadTimingProfilerProvider(Session);
+			if (TimingProfilerProvider)
+			{
+				TimingProfilerProvider->EnumerateGpuQueues([&](const TraceServices::FGpuQueueInfo& QueueInfo)
+				{
+					StringBuilder.Appendf(UTF8TEXT("%u"), GetNonCollidingId(QueueInfo.Id));
+					Writer.AppendSeparator();
+					StringBuilder.Append(QueueInfo.GetDisplayName());
+					Writer.AppendSeparator();
+					StringBuilder.Append(UTF8TEXT("GPU"));
+					Writer.AppendLineEnd();
+					++ThreadCount;
+				});
+			}
+		}
+
 		// Iterate the CPU threads.
 		{
 			TraceServices::FAnalysisSessionReadScope SessionReadScope(Session);
@@ -409,11 +429,14 @@ int32 FTimingExporter::ExportTimingEvents_WriteEvents(FExportTimingEventsInterna
 
 		const TraceServices::ITimingProfilerProvider& TimingProfilerProvider = *TraceServices::ReadTimingProfilerProvider(Session);
 
-		auto TimelineEnumerator = [&Params, &TimingEventCount, &Timers](const TraceServices::ITimingProfilerProvider::Timeline& Timeline)
+		const TraceServices::ITimingProfilerTimerReader* TimerReader;
+		TimingProfilerProvider.ReadTimers([&TimerReader, &Timers](const TraceServices::ITimingProfilerTimerReader& Out) { TimerReader = &Out; });
+
+		auto TimelineEnumerator = [&Params, &TimingEventCount, &Timers, TimerReader](const TraceServices::ITimingProfilerProvider::Timeline& Timeline)
 		{
 			// Iterate timing events.
 			Timeline.EnumerateEvents(Params.UserParams.IntervalStartTime, Params.UserParams.IntervalEndTime,
-				[&Params, &TimingEventCount, &Timers](double EventStartTime, double EventEndTime, uint32 EventDepth, const TraceServices::FTimingProfilerEvent& Event)
+				[&Params, &TimingEventCount, &Timers, TimerReader](double EventStartTime, double EventEndTime, uint32 EventDepth, const TraceServices::FTimingProfilerEvent& Event)
 				{
 					if (!Params.UserParams.TimingEventFilter || Params.UserParams.TimingEventFilter(EventStartTime, EventEndTime, EventDepth, Event))
 					{
@@ -442,11 +465,23 @@ int32 FTimingExporter::ExportTimingEvents_WriteEvents(FExportTimingEventsInterna
 								}
 								else if (Column == ExportTimingEvents_TimerIdColumn)
 								{
-									StringBuilder.Appendf(UTF8TEXT("%u"), Event.TimerIndex);
+									uint32 TimerIndex = Event.TimerIndex;
+									if ((int32) TimerIndex < 0)
+									{
+										TimerIndex = TimerReader->GetOriginalTimerIdFromMetadata(TimerIndex);
+									}
+
+									StringBuilder.Appendf(UTF8TEXT("%u"), TimerIndex);
 								}
 								else if (Column == ExportTimingEvents_TimerNameColumn)
 								{
-									const TCHAR* TimerName = Timers.FindRef(Event.TimerIndex);
+									uint32 TimerIndex = Event.TimerIndex;
+									if ((int32)TimerIndex < 0)
+									{
+										TimerIndex = TimerReader->GetOriginalTimerIdFromMetadata(TimerIndex);
+									}
+
+									const TCHAR* TimerName = Timers.FindRef(TimerIndex);
 									Params.Writer.AppendString(TimerName);
 								}
 								else if (Column == ExportTimingEvents_StartTimeColumn)
@@ -496,6 +531,20 @@ int32 FTimingExporter::ExportTimingEvents_WriteEvents(FExportTimingEventsInterna
 			}
 		}
 
+		// Iterate the GPU timelines for the GPU Queues.
+		TimingProfilerProvider.EnumerateGpuQueues([&TimingProfilerProvider, &TimelineEnumerator, &Params](const TraceServices::FGpuQueueInfo& QueueInfo)
+			{
+				if (!Params.UserParams.ThreadFilter || Params.UserParams.ThreadFilter(GetNonCollidingId(QueueInfo.Id)))
+				{
+					FString QueueName = QueueInfo.GetDisplayName();
+					Params.ThreadId = GetNonCollidingId(QueueInfo.Id);
+					Params.ThreadName = *QueueName;
+					TimingProfilerProvider.ReadTimeline(QueueInfo.TimelineIndex, TimelineEnumerator);
+
+					Params.ThreadName = nullptr; // The pointer would no longer be valid outside this scope.
+				}
+			});
+
 		// Iterate the Verse timelines.
 		{
 			if (!Params.UserParams.ThreadFilter || Params.UserParams.ThreadFilter(FVerseTimingTrack::VerseThreadId))
@@ -1032,11 +1081,25 @@ FTimingExporter::FThreadFilterFunc FTimingExporter::MakeThreadFilterInclusive(co
 
 	TMap<FString, uint32> Threads;
 
-	// Add the GPU threads.
+	// Add the GPU threads for the old GPU Profiler.
 	Threads.Add(FString("GPU"), FGpuTimingTrack::Gpu1ThreadId);
 	Threads.Add(FString("GPU1"), FGpuTimingTrack::Gpu1ThreadId);
 	Threads.Add(FString("GPU2"), FGpuTimingTrack::Gpu2ThreadId);
 
+	// Add Gpu Queues for the new Gpu Profiler
+	{
+		TraceServices::FAnalysisSessionReadScope SessionReadScope(Session);
+
+		const TraceServices::ITimingProfilerProvider* TimingProfilerProvider = TraceServices::ReadTimingProfilerProvider(Session);
+		if (TimingProfilerProvider)
+		{
+			TimingProfilerProvider->EnumerateGpuQueues([&Threads](const TraceServices::FGpuQueueInfo& QueueInfo)
+			{
+				Threads.Add(QueueInfo.GetDisplayName(), GetNonCollidingId(QueueInfo.Id));
+			});
+		}
+	}
+
 	// Add the Verse Sampling fake thread.
 	Threads.Add(FString("VerseSampling"), FVerseTimingTrack::VerseThreadId);
 
@@ -1536,6 +1599,13 @@ int32 FTimingExporter::ExportCounterAsText(const FString& FilenamePattern, uint3
 
 ////////////////////////////////////////////////////////////////////////////////////////////////////
 
+uint32 FTimingExporter::GetNonCollidingId(uint32 QueueId)
+{
+	return QueueId + (1 << 16);
+}
+
+////////////////////////////////////////////////////////////////////////////////////////////////////
+
 } // namespace UE::Insights::TimingProfiler
 
 #undef LOCTEXT_NAMESPACE
diff --git a/Engine/Source/Developer/TraceInsights/Private/Insights/TimingProfiler/ViewModels/TimingExporter.h b/Engine/Source/Developer/TraceInsights/Private/Insights/TimingProfiler/ViewModels/TimingExporter.h
--- a/Engine/Source/Developer/TraceInsights/Private/Insights/TimingProfiler/ViewModels/TimingExporter.h
+++ b/Engine/Source/Developer/TraceInsights/Private/Insights/TimingProfiler/ViewModels/TimingExporter.h
@@ -330,6 +330,7 @@ private:
 	int32 ExportTimerStatisticsAsTextByRegions(const FString& Filename, FExportTimerStatisticsParams& Params) const;
 	int32 ExportTimerCalleesByRegions(const FString& FilenamePattern, const FExportTimerCalleesParams& Params) const;
 	int32 ExportCounterAsTextByRegions(const FString& Filename, uint32 CounterId, FExportCounterParams& Params) const;
+	static uint32 GetNonCollidingId(uint32 QueueId);
 
 private:
 	const TraceServices::IAnalysisSession& Session;
diff --git a/Engine/Source/Developer/TraceServices/Public/TraceServices/Model/TimingProfiler.h b/Engine/Source/Developer/TraceServices/Public/TraceServices/Model/TimingProfiler.h
--- a/Engine/Source/Developer/TraceServices/Public/TraceServices/Model/TimingProfiler.h
+++ b/Engine/Source/Developer/TraceServices/Public/TraceServices/Model/TimingProfiler.h
@@ -213,6 +213,11 @@ struct FGpuQueueInfo
 	const TCHAR* Name = nullptr;
 	uint32 TimelineIndex = ~0;
 	uint32 WorkTimelineIndex = ~0;
+
+	FString GetDisplayName() const
+	{
+		return FString::Printf(TEXT("GPU%u-%s%u"), GPU, Name, Index);
+	}
 };
 
 enum class EEnumerateResult
