[Insights]
* Memory Insights: Fixed crash when toggling visibility of Timing View while Memory Tags panel is visible.
* Memory Insights: Fixed Timing View to correctly recreate if visibility is changed (ex.: to recreate the default LLM graph tracks).

diff --git a/Engine/Source/Developer/TraceInsights/Private/Insights/MemoryProfiler/MemoryProfilerManager.cpp b/Engine/Source/Developer/TraceInsights/Private/Insights/MemoryProfiler/MemoryProfilerManager.cpp
--- a/Engine/Source/Developer/TraceInsights/Private/Insights/MemoryProfiler/MemoryProfilerManager.cpp
+++ b/Engine/Source/Developer/TraceInsights/Private/Insights/MemoryProfiler/MemoryProfilerManager.cpp
@@ -283,6 +283,7 @@ bool FMemoryProfilerManager::Tick(float DeltaTime)
 		RegisterMajorTabs(TraceInsightsModule);
 	}
 #endif
+
 	// Check if session has Memory events (to spawn the tab), but not too often.
 	if (!bIsAvailable && AvailabilityCheck.Tick())
 	{
@@ -346,6 +347,12 @@ bool FMemoryProfilerManager::Tick(float DeltaTime)
 		}
 	}
 
+	FMemorySharedState* SharedState = GetSharedState();
+	if (SharedState)
+	{
+		SharedState->UpdateMemoryTags();
+	}
+
 	return true;
 }
 
diff --git a/Engine/Source/Developer/TraceInsights/Private/Insights/MemoryProfiler/ViewModels/MemorySharedState.cpp b/Engine/Source/Developer/TraceInsights/Private/Insights/MemoryProfiler/ViewModels/MemorySharedState.cpp
--- a/Engine/Source/Developer/TraceInsights/Private/Insights/MemoryProfiler/ViewModels/MemorySharedState.cpp
+++ b/Engine/Source/Developer/TraceInsights/Private/Insights/MemoryProfiler/ViewModels/MemorySharedState.cpp
@@ -97,13 +97,8 @@ FMemorySharedState::~FMemorySharedState()
 
 ////////////////////////////////////////////////////////////////////////////////////////////////////
 
-void FMemorySharedState::OnBeginSession(Timing::ITimingViewSession& InSession)
+void FMemorySharedState::ResetMemoryTags()
 {
-	if (&InSession != TimingView.Get())
-	{
-		return;
-	}
-
 	TagList.Reset();
 
 	Trackers.Reset();
@@ -113,57 +108,75 @@ void FMemorySharedState::OnBeginSession(Timing::ITimingViewSession& InSession)
 	TagSets.Reset();
 	NumSyncedTagSets = 0;
 	NumValidTagSets = 0;
+}
+
+////////////////////////////////////////////////////////////////////////////////////////////////////
+
+void FMemorySharedState::UpdateMemoryTags()
+{
+	TagList.Update();
 
+	if (!DefaultTracker)
+	{
+		SyncTrackers();
+	}
+
+	UpdateTagSets();
+}
+
+////////////////////////////////////////////////////////////////////////////////////////////////////
+
+void FMemorySharedState::ResetTimingView()
+{
 	MainGraphTrack = nullptr;
 	LiveAllocsGraphTrack = nullptr;
 	AllocFreeGraphTrack = nullptr;
 	SwapMemoryGraphTrack = nullptr;
 	PageSwapGraphTrack = nullptr;
 	AllTracks.Reset();
+	for (FMemoryTag* TagPtr : TagList.GetTags())
+	{
+		TagPtr->RemoveAllTracks();
+	}
 
-	bShowHideAllMemoryTracks = true;
+	bShowHideAllMemoryTracks = false;
 
 	CreatedDefaultTracks.Reset();
+	LastTagCountForDefaultTracks = 0;
 }
 
 ////////////////////////////////////////////////////////////////////////////////////////////////////
 
-void FMemorySharedState::OnEndSession(Timing::ITimingViewSession& InSession)
+void FMemorySharedState::OnBeginSession(Timing::ITimingViewSession& InSession)
 {
-	if (&InSession != TimingView.Get())
+	TimingProfiler::STimingView* TimingView = GetTimingView().Get();
+	if (&InSession != TimingView)
 	{
 		return;
 	}
 
-	//TODO: move TagList, Trackers and TagSets to MemoryManager.
-
-	TagList.Reset();
-
-	Trackers.Reset();
-	DefaultTracker = nullptr;
-	PlatformTracker = nullptr;
-
-	TagSets.Reset();
-	NumSyncedTagSets = 0;
-	NumValidTagSets = 0;
+	ResetTimingView();
+}
 
-	MainGraphTrack = nullptr;
-	LiveAllocsGraphTrack = nullptr;
-	AllocFreeGraphTrack = nullptr;
-	SwapMemoryGraphTrack = nullptr;
-	PageSwapGraphTrack = nullptr;
-	AllTracks.Reset();
+////////////////////////////////////////////////////////////////////////////////////////////////////
 
-	bShowHideAllMemoryTracks = false;
+void FMemorySharedState::OnEndSession(Timing::ITimingViewSession& InSession)
+{
+	TimingProfiler::STimingView* TimingView = GetTimingView().Get();
+	if (&InSession != TimingView)
+	{
+		return;
+	}
 
-	CreatedDefaultTracks.Reset();
+	ResetTimingView();
 }
 
 ////////////////////////////////////////////////////////////////////////////////////////////////////
 
 void FMemorySharedState::Tick(Timing::ITimingViewSession& InSession, const TraceServices::IAnalysisSession& InAnalysisSession)
 {
-	if (&InSession != TimingView.Get())
+	TimingProfiler::STimingView* TimingView = GetTimingView().Get();
+	if (&InSession != TimingView)
 	{
 		return;
 	}
@@ -294,23 +307,7 @@ void FMemorySharedState::Tick(Timing::ITimingViewSession& InSession, const Trace
 		TimingView->InvalidateScrollableTracksOrder();
 	}
 
-	const int32 PrevTagCount = TagList.GetTags().Num();
-
-	TagList.Update();
-
-	if (!DefaultTracker)
-	{
-		SyncTrackers();
-	}
-
-	UpdateTagSets();
-
-	// Scan for mem tags to show as default, but only when new mem tags are added.
-	const int32 NewTagCount = TagList.GetTags().Num();
-	if (NewTagCount > PrevTagCount)
-	{
-		CreateDefaultTracks();
-	}
+	CreateDefaultTracks();
 }
 
 ////////////////////////////////////////////////////////////////////////////////////////////////////
@@ -322,7 +319,13 @@ void FMemorySharedState::CreateDefaultTracks()
 		return;
 	}
 
-	FMemoryTrackerId DefaultTrackerId = DefaultTracker->GetId();
+	const auto Tags = TagList.GetTags();
+	if (Tags.Num() == LastTagCountForDefaultTracks)
+	{
+		// Only attempt to create default tracks if new tags are added.
+		return;
+	}
+	LastTagCountForDefaultTracks = Tags.Num();
 
 	static const TCHAR* DefaultTags[] =
 	{
@@ -348,7 +351,8 @@ void FMemorySharedState::CreateDefaultTracks()
 		CreatedDefaultTracks.Init(false, DefaultTagCount);
 	}
 
-	const auto Tags = TagList.GetTags();
+	const FMemoryTrackerId DefaultTrackerId = DefaultTracker->GetId();
+
 	for (int32 DefaultTagIndex = 0; DefaultTagIndex < DefaultTagCount; ++DefaultTagIndex)
 	{
 		if (!CreatedDefaultTracks[DefaultTagIndex])
@@ -525,7 +529,8 @@ void FMemorySharedState::SetTrackHeightMode(EMemoryTrackHeightMode InTrackHeight
 
 void FMemorySharedState::ExtendOtherTracksFilterMenu(Timing::ITimingViewSession& InSession, FMenuBuilder& InOutMenuBuilder)
 {
-	if (&InSession != TimingView.Get())
+	TimingProfiler::STimingView* TimingView = GetTimingView().Get();
+	if (&InSession != TimingView)
 	{
 		return;
 	}
@@ -543,6 +548,12 @@ void FMemorySharedState::BindCommands()
 {
 	FMemoryTimingViewCommands::Register();
 
+	TSharedPtr<TimingProfiler::STimingView> TimingView = GetTimingView();
+	if (!TimingView)
+	{
+		return;
+	}
+
 	TSharedPtr<FUICommandList> CommandList = TimingView->GetCommandList();
 	ensure(CommandList.IsValid());
 
@@ -564,6 +575,7 @@ void FMemorySharedState::SetAllMemoryTracksToggle(bool bOnOff)
 		GraphTrack->SetVisibilityFlag(bShowHideAllMemoryTracks);
 	}
 
+	TSharedPtr<TimingProfiler::STimingView> TimingView = GetTimingView();
 	if (TimingView)
 	{
 		TimingView->HandleTrackVisibilityChanged();
@@ -586,6 +598,7 @@ int32 FMemorySharedState::GetNextMemoryGraphTrackOrder()
 
 TSharedPtr<FMemoryGraphTrack> FMemorySharedState::CreateMemoryGraphTrack()
 {
+	TSharedPtr<TimingProfiler::STimingView> TimingView = GetTimingView();
 	if (!TimingView.IsValid())
 	{
 		return nullptr;
@@ -628,7 +641,11 @@ int32 FMemorySharedState::RemoveMemoryGraphTrack(TSharedPtr<FMemoryGraphTrack> G
 		if (GraphTrack->GetSeries().Num() == 0)
 		{
 			GraphTrack->Hide();
-			TimingView->HandleTrackVisibilityChanged();
+			TSharedPtr<TimingProfiler::STimingView> TimingView = GetTimingView();
+			if (TimingView)
+			{
+				TimingView->HandleTrackVisibilityChanged();
+			}
 		}
 		return -1;
 	}
@@ -637,7 +654,11 @@ int32 FMemorySharedState::RemoveMemoryGraphTrack(TSharedPtr<FMemoryGraphTrack> G
 	{
 		RemoveTrackFromMemTags(GraphTrack);
 		GraphTrack->RemoveAllMemTagSeries();
-		TimingView->RemoveTrack(GraphTrack);
+		TSharedPtr<TimingProfiler::STimingView> TimingView = GetTimingView();
+		if (TimingView)
+		{
+			TimingView->RemoveTrack(GraphTrack);
+		}
 		return 1;
 	}
 
@@ -666,6 +687,7 @@ void FMemorySharedState::RemoveTrackFromMemTags(TSharedPtr<FMemoryGraphTrack>& G
 
 TSharedPtr<FMemoryGraphTrack> FMemorySharedState::GetMemTagGraphTrack(FMemoryTrackerId InMemTrackerId, FMemoryTagId InMemTagId)
 {
+	TSharedPtr<TimingProfiler::STimingView> TimingView = GetTimingView();
 	if (!TimingView.IsValid())
 	{
 		return nullptr;
@@ -692,6 +714,7 @@ TSharedPtr<FMemoryGraphTrack> FMemorySharedState::GetMemTagGraphTrack(FMemoryTra
 
 TSharedPtr<FMemoryGraphTrack> FMemorySharedState::CreateMemTagGraphTrack(FMemoryTrackerId InMemTrackerId, FMemoryTagId InMemTagId)
 {
+	TSharedPtr<TimingProfiler::STimingView> TimingView = GetTimingView();
 	if (!TimingView.IsValid())
 	{
 		return nullptr;
@@ -794,6 +817,7 @@ TSharedPtr<FMemoryGraphTrack> FMemorySharedState::CreateMemTagGraphTrack(FMemory
 
 int32 FMemorySharedState::RemoveMemTagGraphTrack(FMemoryTrackerId InMemTrackerId, FMemoryTagId InMemTagId)
 {
+	TSharedPtr<TimingProfiler::STimingView> TimingView = GetTimingView();
 	if (!TimingView.IsValid())
 	{
 		return -1;
@@ -832,6 +856,7 @@ int32 FMemorySharedState::RemoveMemTagGraphTrack(FMemoryTrackerId InMemTrackerId
 
 int32 FMemorySharedState::RemoveAllMemTagGraphTracks()
 {
+	TSharedPtr<TimingProfiler::STimingView> TimingView = GetTimingView();
 	if (!TimingView.IsValid())
 	{
 		return -1;
@@ -880,6 +905,12 @@ TSharedPtr<FMemoryGraphSeries> FMemorySharedState::ToggleMemTagGraphSeries(TShar
 		return nullptr;
 	}
 
+	TSharedPtr<TimingProfiler::STimingView> TimingView = GetTimingView();
+	if (!TimingView.IsValid())
+	{
+		return nullptr;
+	}
+
 	FMemoryTag* TagPtr = TagList.GetTagById(InMemTrackerId, InMemTagId);
 	FMemoryTagSetId TagSetId = TagPtr ? TagPtr->GetTagSetId() : FMemoryTagSet::InvalidTagSetId;
 
@@ -968,6 +999,7 @@ void FMemorySharedState::CreateTracksFromReport(const FReportTypeConfig& ReportT
 
 	if (NumAddedTracks > 0)
 	{
+		TSharedPtr<TimingProfiler::STimingView> TimingView = GetTimingView();
 		if (TimingView)
 		{
 			TimingView->InvalidateScrollableTracksOrder();
@@ -985,6 +1017,7 @@ TSharedPtr<FMemoryGraphTrack> FMemorySharedState::CreateGraphTrack(const FReport
 		return nullptr;
 	}
 
+	TSharedPtr<TimingProfiler::STimingView> TimingView = GetTimingView();
 	if (!TimingView.IsValid())
 	{
 		return nullptr;
diff --git a/Engine/Source/Developer/TraceInsights/Private/Insights/MemoryProfiler/ViewModels/MemorySharedState.h b/Engine/Source/Developer/TraceInsights/Private/Insights/MemoryProfiler/ViewModels/MemorySharedState.h
--- a/Engine/Source/Developer/TraceInsights/Private/Insights/MemoryProfiler/ViewModels/MemorySharedState.h
+++ b/Engine/Source/Developer/TraceInsights/Private/Insights/MemoryProfiler/ViewModels/MemorySharedState.h
@@ -14,10 +14,6 @@
 #include "Insights/MemoryProfiler/Tracks/MemoryGraphTrack.h"
 #include "Insights/MemoryProfiler/ViewModels/MemoryTag.h"
 
-class FTimingEventSearchParameters;
-class FTimingGraphSeries;
-class FTimingGraphTrack;
-
 namespace UE::Insights::TimingProfiler { class STimingView; }
 
 namespace UE::Insights::MemoryProfiler
@@ -100,8 +96,11 @@ public:
 	FMemorySharedState();
 	virtual ~FMemorySharedState();
 
-	TSharedPtr<TimingProfiler::STimingView> GetTimingView() const { return TimingView; }
-	void SetTimingView(TSharedPtr<TimingProfiler::STimingView> InTimingView) { TimingView = InTimingView; }
+	//////////////////////////////////////////////////
+	// Memory Tags
+
+	void ResetMemoryTags();
+	void UpdateMemoryTags();
 
 	const FMemoryTagList& GetTagList() { return TagList; }
 
@@ -121,11 +120,6 @@ public:
 	void EnumerateTagSets(TFunctionRef<void(const FMemoryTagSet& InTagSet)> Callback) const;
 	const FMemoryTagSet* GetTagSetById(FMemoryTagSetId InMemTagSetId) const;
 
-	TSharedPtr<FMemoryGraphTrack> GetMainGraphTrack() const { return MainGraphTrack; }
-
-	EMemoryTrackHeightMode GetTrackHeightMode() const { return TrackHeightMode; }
-	void SetTrackHeightMode(EMemoryTrackHeightMode InTrackHeightMode);
-
 	//////////////////////////////////////////////////
 	// ITimingViewExtender interface
 
@@ -135,8 +129,24 @@ public:
 	virtual void ExtendOtherTracksFilterMenu(Timing::ITimingViewSession& InSession, FMenuBuilder& InOutMenuBuilder) override;
 
 	//////////////////////////////////////////////////
+	// Timing View
 
-	void BindCommands();
+	void ResetTimingView();
+
+	TSharedPtr<TimingProfiler::STimingView> GetTimingView() const
+	{
+		return TimingViewWeak.Pin();
+	}
+	void SetTimingView(TWeakPtr<TimingProfiler::STimingView> InTimingView)
+	{
+		TimingViewWeak = InTimingView;
+		ResetTimingView();
+	}
+
+	TSharedPtr<FMemoryGraphTrack> GetMainGraphTrack() const { return MainGraphTrack; }
+
+	EMemoryTrackHeightMode GetTrackHeightMode() const { return TrackHeightMode; }
+	void SetTrackHeightMode(EMemoryTrackHeightMode InTrackHeightMode);
 
 	bool IsAllMemoryTracksToggleOn() const { return bShowHideAllMemoryTracks; }
 	void SetAllMemoryTracksToggle(bool bOnOff);
@@ -154,15 +164,15 @@ public:
 	void RemoveTrackFromMemTags(TSharedPtr<FMemoryGraphTrack>& GraphTrack);
 	int32 RemoveMemTagGraphTrack(FMemoryTrackerId InMemTrackerId, FMemoryTagId InMemTagId);
 	int32 RemoveAllMemTagGraphTracks();
+	void ResetAllMemTagGraphTracks();
 
 	TSharedPtr<FMemoryGraphSeries> ToggleMemTagGraphSeries(TSharedPtr<FMemoryGraphTrack> InGraphTrack, FMemoryTrackerId InMemTrackerId, FMemoryTagId InMemTagId);
 
-	void CreateTracksFromReport(const FString& Filename);
-	void CreateTracksFromReport(const FReportConfig& ReportConfig);
-	void CreateTracksFromReport(const FReportTypeConfig& ReportTypeConfig);
+	//////////////////////////////////////////////////
+	// Memory Investigation
 
 	const TArray<TSharedPtr<FMemoryRuleSpec>>& GetMemoryRules() const { return MemoryRules; }
-	
+
 	TSharedPtr<FMemoryRuleSpec> GetCurrentMemoryRule() const { return CurrentMemoryRule; }
 	void SetCurrentMemoryRule(TSharedPtr<FMemoryRuleSpec> InRule) { CurrentMemoryRule = InRule; OnMemoryRuleChanged(); }
 
@@ -173,6 +183,17 @@ public:
 	void AddQueryTarget(TSharedPtr<FQueryTargetWindowSpec> InPtr);
 	void RemoveQueryTarget(TSharedPtr<FQueryTargetWindowSpec> InPtr);
 
+	//////////////////////////////////////////////////
+	// Misc
+
+	void BindCommands();
+
+	void CreateTracksFromReport(const FString& Filename);
+	void CreateTracksFromReport(const FReportConfig& ReportConfig);
+	void CreateTracksFromReport(const FReportTypeConfig& ReportTypeConfig);
+
+	//////////////////////////////////////////////////
+
 private:
 	void SyncTrackers();
 	void UpdateTagSets();
@@ -182,7 +203,8 @@ private:
 	void OnMemoryRuleChanged();
 
 private:
-	TSharedPtr<TimingProfiler::STimingView> TimingView;
+	//////////////////////////////////////////////////
+	// Memory Tags
 
 	FMemoryTagList TagList;
 
@@ -194,6 +216,11 @@ private:
 	uint32 NumSyncedTagSets = 0;
 	uint32 NumValidTagSets = 0;
 
+	//////////////////////////////////////////////////
+	// Timing View
+
+	TWeakPtr<TimingProfiler::STimingView> TimingViewWeak;
+
 	TSharedPtr<FMemoryGraphTrack> MainGraphTrack; // the Main Memory Graph track; also hosts the Total Allocated Memory series
 	TSharedPtr<FMemoryGraphTrack> LiveAllocsGraphTrack; // the graph track for the Live Allocation Count series
 	TSharedPtr<FMemoryGraphTrack> AllocFreeGraphTrack; // the graph track for the Alloc Event Count and the Free Event Count series
@@ -206,12 +233,18 @@ private:
 	bool bShowHideAllMemoryTracks = false;
 
 	TBitArray<> CreatedDefaultTracks;
+	uint32 LastTagCountForDefaultTracks = 0;
+
+	//////////////////////////////////////////////////
+	// Memory Investigation
 
 	TArray<TSharedPtr<FMemoryRuleSpec>> MemoryRules;
 	TSharedPtr<FMemoryRuleSpec> CurrentMemoryRule;
 
 	TSharedPtr<FQueryTargetWindowSpec> CurrentQueryTarget;
 	TArray<TSharedPtr<FQueryTargetWindowSpec>> QueryTargetSpecs;
+
+	//////////////////////////////////////////////////
 };
 
 ////////////////////////////////////////////////////////////////////////////////////////////////////
diff --git a/Engine/Source/Developer/TraceInsights/Private/Insights/MemoryProfiler/Widgets/SMemoryProfilerWindow.cpp b/Engine/Source/Developer/TraceInsights/Private/Insights/MemoryProfiler/Widgets/SMemoryProfilerWindow.cpp
--- a/Engine/Source/Developer/TraceInsights/Private/Insights/MemoryProfiler/Widgets/SMemoryProfilerWindow.cpp
+++ b/Engine/Source/Developer/TraceInsights/Private/Insights/MemoryProfiler/Widgets/SMemoryProfilerWindow.cpp
@@ -42,6 +42,8 @@ SMemoryProfilerWindow::SMemoryProfilerWindow()
 	, SharedState(MakeShared<FMemorySharedState>())
 {
 	CreateTimingViewMarkers();
+
+	IModularFeatures::Get().RegisterModularFeature(Timing::TimingViewExtenderFeatureName, &SharedState.Get());
 }
 
 ////////////////////////////////////////////////////////////////////////////////////////////////////
@@ -50,6 +52,8 @@ SMemoryProfilerWindow::~SMemoryProfilerWindow()
 {
 	CloseAllOpenTabs();
 
+	IModularFeatures::Get().UnregisterModularFeature(Timing::TimingViewExtenderFeatureName, &SharedState.Get());
+
 	check(ModulesView == nullptr);
 	check(MemTagTreeView == nullptr);
 	check(MemInvestigationView == nullptr);
@@ -67,6 +71,8 @@ const TCHAR* SMemoryProfilerWindow::GetAnalyticsEventName() const
 
 void SMemoryProfilerWindow::Reset()
 {
+	SharedState->ResetMemoryTags();
+
 	if (TimingView)
 	{
 		TimingView->Reset();
@@ -105,7 +111,6 @@ TSharedRef<SDockTab> SMemoryProfilerWindow::SpawnTab_TimingView(const FSpawnTabA
 
 	SharedState->SetTimingView(TimingView);
 	SharedState->BindCommands();
-	IModularFeatures::Get().RegisterModularFeature(Timing::TimingViewExtenderFeatureName, &SharedState.Get());
 
 	TimingView->Reset(true);
 	TimingView->OnSelectionChanged().AddSP(this, &SMemoryProfilerWindow::OnTimeSelectionChanged);
@@ -133,7 +138,6 @@ void SMemoryProfilerWindow::OnTimingViewTabClosed(TSharedRef<SDockTab> TabBeingC
 		TimingView = nullptr;
 	}
 
-	IModularFeatures::Get().UnregisterModularFeature(Timing::TimingViewExtenderFeatureName, &SharedState.Get());
 	SharedState->SetTimingView(nullptr);
 
 	RemoveOpenTab(TabBeingClosed);
@@ -242,7 +246,7 @@ void SMemoryProfilerWindow::OnMemAllocTableTreeViewTabClosed(TSharedRef<SDockTab
 	ClosingTabId.SetNumber(MemAllocTableTreeView->GetTabIndex());
 
 	TSharedPtr<FQueryTargetWindowSpec> TargetToDelete;
-	const TArray<TSharedPtr<FQueryTargetWindowSpec>>& Targets = this->GetSharedState().GetQueryTargets();
+	const TArray<TSharedPtr<FQueryTargetWindowSpec>>& Targets = SharedState->GetQueryTargets();
 	for (int32 Index = 0; Index < Targets.Num(); ++Index)
 	{
 		if (Targets[Index]->GetName() == ClosingTabId)
@@ -254,13 +258,13 @@ void SMemoryProfilerWindow::OnMemAllocTableTreeViewTabClosed(TSharedRef<SDockTab
 
 	if (TargetToDelete.IsValid())
 	{
-		this->GetSharedState().RemoveQueryTarget(TargetToDelete);
+		SharedState->RemoveQueryTarget(TargetToDelete);
 	}
 
 	if (Targets.Num() > 0)
 	{
 		TSharedPtr<FQueryTargetWindowSpec> NewSelection = Targets[0];
-		this->GetSharedState().SetCurrentQueryTarget(NewSelection);
+		SharedState->SetCurrentQueryTarget(NewSelection);
 		if (MemInvestigationView.IsValid())
 		{
 			MemInvestigationView->QueryTarget_OnSelectionChanged(NewSelection, ESelectInfo::Type::Direct);
@@ -279,11 +283,11 @@ void SMemoryProfilerWindow::OnMemAllocTableTreeViewTabClosed(TSharedRef<SDockTab
 
 void SMemoryProfilerWindow::CloseMemAllocTableTreeTabs()
 {
-	const TArray<TSharedPtr<FQueryTargetWindowSpec>>& Targets = this->GetSharedState().GetQueryTargets();
+	const TArray<TSharedPtr<FQueryTargetWindowSpec>>& Targets = SharedState->GetQueryTargets();
 	while (Targets.Num() > 0)
 	{
 		FName Name = Targets[0]->GetName();
-		this->GetSharedState().RemoveQueryTarget(Targets[0]);
+		SharedState->RemoveQueryTarget(Targets[0]);
 
 		if (Name != FQueryTargetWindowSpec::NewWindow)
 		{
