[Insights] Timing View: Fixed screenshot preview not showing in tooltip of Markers track.

#rnx

diff --git a/Engine/Source/Developer/TraceInsights/Private/Insights/TimingProfiler/Tracks/MarkersTimingTrack.cpp b/Engine/Source/Developer/TraceInsights/Private/Insights/TimingProfiler/Tracks/MarkersTimingTrack.cpp
--- a/Engine/Source/Developer/TraceInsights/Private/Insights/TimingProfiler/Tracks/MarkersTimingTrack.cpp
+++ b/Engine/Source/Developer/TraceInsights/Private/Insights/TimingProfiler/Tracks/MarkersTimingTrack.cpp
@@ -49,17 +49,10 @@ INSIGHTS_IMPLEMENT_RTTI(FMarkersTimingTrack)
 
 FMarkersTimingTrack::FMarkersTimingTrack()
 	: FBaseTimingTrack(TEXT("Markers (Bookmarks / Logs)"))
-	//, TimeMarkerBoxes()
-	//, TimeMarkerTexts()
-	, bUseOnlyBookmarks(true)
-	, BookmarkCategory(nullptr)
-	, ScreenshotCategory(nullptr)
 	, Header(*this)
-	, NumLogMessages(0)
-	, NumDrawBoxes(0)
-	, NumDrawTexts(0)
 	, WhiteBrush(FAppStyle::Get().GetBrush("WhiteBrush"))
 	, Font(FAppStyle::Get().GetFontStyle("SmallFont"))
+	, LastScreenshotId(TraceServices::FScreenshot::InvalidScreenshotId)
 {
 	SetValidLocations(ETimingTrackLocation::TopDocked | ETimingTrackLocation::BottomDocked);
 	SetOrder(FTimingTrackOrder::Markers);
@@ -327,7 +320,7 @@ void FMarkersTimingTrack::BuildContextMenu(FMenuBuilder& MenuBuilder)
 	MenuBuilder.BeginSection(TEXT("Screenshot"), LOCTEXT("ContextMenu_Section_Screenshot", "Screenshot"));
 	{
 		MenuBuilder.AddMenuEntry(
-			LOCTEXT("ContextMenu_SaveScreenshot", "Save Screenshot"),
+			LOCTEXT("ContextMenu_SaveScreenshot", "Save Screenshot..."),
 			LOCTEXT("ContextMenu_SaveScreenshot_Desc", "Save the hovered screenshot to a file."),
 			FSlateIcon(),
 			FUIAction(FExecuteAction::CreateSP(this, &FMarkersTimingTrack::SaveScreenshot_Execute),
@@ -336,7 +329,7 @@ void FMarkersTimingTrack::BuildContextMenu(FMenuBuilder& MenuBuilder)
 			EUserInterfaceActionType::Button
 		);
 
-		TryGetHoveredEventScreenshotId(LastScreenshotId);
+		LastScreenshotId = TryGetHoveredEventScreenshotId();
 	}
 	MenuBuilder.EndSection();
 }
@@ -641,7 +634,7 @@ void FMarkersTimingTrack::InitTooltip(FTooltipDrawState& InOutTooltip, const ITi
 		TraceServices::FAnalysisSessionReadScope SessionReadScope(*Session.Get());
 
 		const TraceServices::IScreenshotProvider& ScreenshotProvider = TraceServices::ReadScreenshotProvider(*Session.Get());
-		TSharedPtr<const TraceServices::FScreenshot> Screenshot = ScreenshotProvider.GetScreenshot((uint32)Event.GetType());
+		TSharedPtr<const TraceServices::FScreenshot> Screenshot = ScreenshotProvider.GetScreenshot(Event.GetScreenshotId());
 
 		if (!Screenshot.IsValid())
 		{
@@ -799,18 +792,18 @@ const TSharedPtr<const ITimingEvent> FMarkersTimingTrack::TryGetHoveredEvent()
 
 ////////////////////////////////////////////////////////////////////////////////////////////////////
 
-bool FMarkersTimingTrack::TryGetHoveredEventScreenshotId(uint32& OutScreenshotId)
+uint32 FMarkersTimingTrack::TryGetHoveredEventScreenshotId()
 {
 	const TSharedPtr<const ITimingEvent> Event = TryGetHoveredEvent();
 
 	if (Event.IsValid() && Event->Is<FScreenShotTimingEvent>())
 	{
-		OutScreenshotId = Event->As<FScreenShotTimingEvent>().GetScreenshotId();
-		return true;
+		return Event->As<FScreenShotTimingEvent>().GetScreenshotId();
+	}
+	else
+	{
+		return TraceServices::FScreenshot::InvalidScreenshotId;
 	}
-
-	OutScreenshotId = TraceServices::FScreenshot::InvalidScreenshotId;
-	return false;
 }
 
 ////////////////////////////////////////////////////////////////////////////////////////////////////
diff --git a/Engine/Source/Developer/TraceInsights/Private/Insights/TimingProfiler/Tracks/MarkersTimingTrack.h b/Engine/Source/Developer/TraceInsights/Private/Insights/TimingProfiler/Tracks/MarkersTimingTrack.h
--- a/Engine/Source/Developer/TraceInsights/Private/Insights/TimingProfiler/Tracks/MarkersTimingTrack.h
+++ b/Engine/Source/Developer/TraceInsights/Private/Insights/TimingProfiler/Tracks/MarkersTimingTrack.h
@@ -111,28 +111,28 @@ private:
 	void UpdateCategory(const TraceServices::FLogCategoryInfo*& InOutCategory, const TCHAR* CategoryName);
 
 	const TSharedPtr<const ITimingEvent> TryGetHoveredEvent();
-	bool TryGetHoveredEventScreenshotId(uint32& OutScreenshotId);
+	uint32 TryGetHoveredEventScreenshotId();
 
 private:
 	TArray<FTimeMarkerBoxInfo> TimeMarkerBoxes;
 	TArray<FTimeMarkerTextInfo> TimeMarkerTexts;
 
-	bool bUseOnlyBookmarks; // If true, uses only bookmarks; otherwise it uses all log messages.
-	const TraceServices::FLogCategoryInfo* BookmarkCategory;
-	const TraceServices::FLogCategoryInfo* ScreenshotCategory;
+	bool bUseOnlyBookmarks = true; // If true, uses only bookmarks; otherwise it uses all log messages.
+	const TraceServices::FLogCategoryInfo* BookmarkCategory = nullptr;
+	const TraceServices::FLogCategoryInfo* ScreenshotCategory = nullptr;
 
 	FTrackHeader Header;
 
 	// Stats
-	int32 NumLogMessages;
-	mutable int32 NumDrawBoxes;
-	mutable int32 NumDrawTexts;
+	int32 NumLogMessages = 0;
+	mutable int32 NumDrawBoxes = 0;
+	mutable int32 NumDrawTexts = 0;
 
 	// Slate resources
-	const FSlateBrush* WhiteBrush;
+	const FSlateBrush* WhiteBrush = nullptr;
 	const FSlateFontInfo Font;
 
-	uint32 LastScreenshotId;
+	uint32 LastScreenshotId = 0;
 	mutable FString HoveredCategory;
 	mutable FString HoveredMessage;
 };
diff --git a/Engine/Source/Developer/TraceInsights/Private/Insights/ViewModels/TooltipDrawState.cpp b/Engine/Source/Developer/TraceInsights/Private/Insights/ViewModels/TooltipDrawState.cpp
--- a/Engine/Source/Developer/TraceInsights/Private/Insights/ViewModels/TooltipDrawState.cpp
+++ b/Engine/Source/Developer/TraceInsights/Private/Insights/ViewModels/TooltipDrawState.cpp
@@ -128,12 +128,6 @@ void FTooltipDrawState::AddTextLine(const float X, const float Y, FStringView Te
 
 void FTooltipDrawState::UpdateLayout()
 {
-	if (ImageBrush.IsValid())
-	{
-		DesiredSize = ImageBrush->GetImageSize();
-		return;
-	}
-
 	ValueOffsetX = 0.0f;
 	for (const FDrawTextInfo& TextInfo : Texts)
 	{
@@ -171,6 +165,19 @@ void FTooltipDrawState::UpdateLayout()
 		}
 	}
 	DesiredSize.X += BorderX;
+
+	if (ImageBrush.IsValid())
+	{
+		FVector2D ImageSize = ImageBrush->GetImageSize();
+		if (DesiredSize.X < ImageSize.X)
+		{
+			DesiredSize.X = ImageSize.X;
+		}
+		if (DesiredSize.Y < ImageSize.Y)
+		{
+			DesiredSize.Y = ImageSize.Y;
+		}
+	}
 }
 
 ////////////////////////////////////////////////////////////////////////////////////////////////////
@@ -262,11 +269,11 @@ void FTooltipDrawState::Draw(const UE::Insights::FDrawContext& DrawContext) cons
 		//DrawContext.DrawBox(static_cast<float>(Position.X), static_cast<float>(Position.Y), static_cast<float>(Size.X), static_cast<float>(Size.Y), BorderBrush, BorderColor.CopyWithNewOpacity(Opacity));
 		//DrawContext.LayerId++;
 
-		if (ImageBrush)
+		// Draw optional background image.
+		if (ImageBrush.IsValid())
 		{
 			DrawContext.DrawBox(static_cast<float>(Position.X), static_cast<float>(Position.Y), static_cast<float>(DesiredSize.X), static_cast<float>(DesiredSize.Y), ImageBrush.Get(), FLinearColor::White);
 			DrawContext.LayerId++;
-			return;
 		}
 
 		// Draw cached texts.
