Color Picker and Random Color Generation for Timers and Counters.

diff --git a/Engine/Source/Developer/TraceInsights/Private/Insights/MemoryProfiler/Widgets/SMemTagTreeView.cpp b/Engine/Source/Developer/TraceInsights/Private/Insights/MemoryProfiler/Widgets/SMemTagTreeView.cpp
--- a/Engine/Source/Developer/TraceInsights/Private/Insights/MemoryProfiler/Widgets/SMemTagTreeView.cpp
+++ b/Engine/Source/Developer/TraceInsights/Private/Insights/MemoryProfiler/Widgets/SMemTagTreeView.cpp
@@ -1437,8 +1437,8 @@ void SMemTagTreeView::ExtendMenuBeforeMisc(FMenuBuilder& MenuBuilder)
 		MenuBuilder.AddMenuEntry
 		(
 			LOCTEXT("ContextMenu_GenerateColorForSelectedMemTags", "Generate New Color"),
-			LOCTEXT("ContextMenu_GenerateColorForSelectedMemTags_Desc", "Generates new color for the selected memory tags."),
-			FSlateIcon(),
+			LOCTEXT("ContextMenu_GenerateColorForSelectedMemTags_Desc", "Generates a new random color for the selected memory tags."),
+			FSlateIcon(FAppStyle::Get().GetStyleSetName(), "Icons.Refresh"),
 			Action_GenerateColorForSelectedMemTags,
 			NAME_None,
 			EUserInterfaceActionType::Button
@@ -1453,7 +1453,7 @@ void SMemTagTreeView::ExtendMenuBeforeMisc(FMenuBuilder& MenuBuilder)
 		(
 			LOCTEXT("ContextMenu_EditColorForSelectedMemTags", "Edit Color..."),
 			LOCTEXT("ContextMenu_EditColorForSelectedMemTags_Desc", "Changes color for the selected memory tags."),
-			FSlateIcon(),
+			FSlateIcon(FAppStyle::Get().GetStyleSetName(), "Icons.EyeDropper"),
 			Action_EditColorForSelectedMemTags,
 			NAME_None,
 			EUserInterfaceActionType::Button
diff --git a/Engine/Source/Developer/TraceInsights/Private/Insights/TimingProfiler/TimingProfilerManager.cpp b/Engine/Source/Developer/TraceInsights/Private/Insights/TimingProfiler/TimingProfilerManager.cpp
--- a/Engine/Source/Developer/TraceInsights/Private/Insights/TimingProfiler/TimingProfilerManager.cpp
+++ b/Engine/Source/Developer/TraceInsights/Private/Insights/TimingProfiler/TimingProfilerManager.cpp
@@ -587,27 +587,27 @@ void FTimingProfilerManager::UpdateCallersAndCallees()
 	{
 		TimerButterflyAggregator->Cancel();
 		TimerButterflyAggregator->SetTimeInterval(SelectionStartTime, SelectionEndTime);
+	}
 
-		TSharedPtr<STimingProfilerWindow> Wnd = GetProfilerWindow();
-		if (Wnd)
-		{
-			TSharedPtr<STimerTreeView> CallersTreeView = Wnd->GetCallersTreeView();
-			TSharedPtr<STimerTreeView> CalleesTreeView = Wnd->GetCalleesTreeView();
+	TSharedPtr<STimingProfilerWindow> Wnd = GetProfilerWindow();
+	if (Wnd)
+	{
+		TSharedPtr<STimerTreeView> CallersTreeView = Wnd->GetCallersTreeView();
+		TSharedPtr<STimerTreeView> CalleesTreeView = Wnd->GetCalleesTreeView();
 
-			if (CallersTreeView)
-			{
-				CallersTreeView->Reset();
-			}
+		if (CallersTreeView)
+		{
+			CallersTreeView->Reset();
+		}
 
-			if (CalleesTreeView)
-			{
-				CalleesTreeView->Reset();
-			}
+		if (CalleesTreeView)
+		{
+			CalleesTreeView->Reset();
+		}
 
-			if (CallersTreeView || CalleesTreeView)
-			{
-				TimerButterflyAggregator->Start();
-			}
+		if (CallersTreeView || CalleesTreeView)
+		{
+			TimerButterflyAggregator->Start();
 		}
 	}
 }
diff --git a/Engine/Source/Developer/TraceInsights/Private/Insights/TimingProfiler/ViewModels/StatsNode.cpp b/Engine/Source/Developer/TraceInsights/Private/Insights/TimingProfiler/ViewModels/StatsNode.cpp
--- a/Engine/Source/Developer/TraceInsights/Private/Insights/TimingProfiler/ViewModels/StatsNode.cpp
+++ b/Engine/Source/Developer/TraceInsights/Private/Insights/TimingProfiler/ViewModels/StatsNode.cpp
@@ -18,6 +18,13 @@ INSIGHTS_IMPLEMENT_RTTI(FStatsNode)
 
 ////////////////////////////////////////////////////////////////////////////////////////////////////
 
+void FStatsNode::SetColor(const FLinearColor& InColor)
+{
+	Color = InColor;
+}
+
+////////////////////////////////////////////////////////////////////////////////////////////////////
+
 void FStatsNode::ResetAggregatedStats()
 {
 	AggregatedStats.Reset();
diff --git a/Engine/Source/Developer/TraceInsights/Private/Insights/TimingProfiler/ViewModels/StatsNode.h b/Engine/Source/Developer/TraceInsights/Private/Insights/TimingProfiler/ViewModels/StatsNode.h
--- a/Engine/Source/Developer/TraceInsights/Private/Insights/TimingProfiler/ViewModels/StatsNode.h
+++ b/Engine/Source/Developer/TraceInsights/Private/Insights/TimingProfiler/ViewModels/StatsNode.h
@@ -214,6 +214,8 @@ public:
 		return Color;
 	}
 
+	void SetColor(const FLinearColor& InColor);
+
 	bool IsAddedToGraph() const
 	{
 		return bIsAddedToGraph;
diff --git a/Engine/Source/Developer/TraceInsights/Private/Insights/TimingProfiler/ViewModels/TimerNode.cpp b/Engine/Source/Developer/TraceInsights/Private/Insights/TimingProfiler/ViewModels/TimerNode.cpp
--- a/Engine/Source/Developer/TraceInsights/Private/Insights/TimingProfiler/ViewModels/TimerNode.cpp
+++ b/Engine/Source/Developer/TraceInsights/Private/Insights/TimingProfiler/ViewModels/TimerNode.cpp
@@ -63,6 +63,13 @@ FTimerNode::~FTimerNode()
 
 ////////////////////////////////////////////////////////////////////////////////////////////////////
 
+void FTimerNode::SetColor(const FLinearColor& InColor)
+{
+	Color = InColor;
+}
+
+////////////////////////////////////////////////////////////////////////////////////////////////////
+
 void FTimerNode::ResetAggregatedStats()
 {
 	AggregatedStats = TraceServices::FTimingProfilerAggregatedStats();
diff --git a/Engine/Source/Developer/TraceInsights/Private/Insights/TimingProfiler/ViewModels/TimerNode.h b/Engine/Source/Developer/TraceInsights/Private/Insights/TimingProfiler/ViewModels/TimerNode.h
--- a/Engine/Source/Developer/TraceInsights/Private/Insights/TimingProfiler/ViewModels/TimerNode.h
+++ b/Engine/Source/Developer/TraceInsights/Private/Insights/TimingProfiler/ViewModels/TimerNode.h
@@ -97,6 +97,8 @@ public:
 		return Color;
 	}
 
+	void SetColor(const FLinearColor& InColor);
+
 	bool IsAddedToGraph() const
 	{
 		return bIsAddedToGraphs;
diff --git a/Engine/Source/Developer/TraceInsights/Private/Insights/TimingProfiler/Widgets/SFrameTrack.h b/Engine/Source/Developer/TraceInsights/Private/Insights/TimingProfiler/Widgets/SFrameTrack.h
--- a/Engine/Source/Developer/TraceInsights/Private/Insights/TimingProfiler/Widgets/SFrameTrack.h
+++ b/Engine/Source/Developer/TraceInsights/Private/Insights/TimingProfiler/Widgets/SFrameTrack.h
@@ -143,6 +143,7 @@ public:
 
 	bool HasAnySeriesForTimer(uint32 TimerId) const;
 	uint32 GetNumSeriesForTimer(uint32 TimerId) const;
+	TArray<TSharedPtr<FFrameTrackSeries>>& GetSeries() { return AllSeries; };
 
 protected:
 	TSharedRef<FFrameTrackSeries> FindOrAddSeries(ETraceFrameType FrameType);
diff --git a/Engine/Source/Developer/TraceInsights/Private/Insights/TimingProfiler/Widgets/SStatsView.cpp b/Engine/Source/Developer/TraceInsights/Private/Insights/TimingProfiler/Widgets/SStatsView.cpp
--- a/Engine/Source/Developer/TraceInsights/Private/Insights/TimingProfiler/Widgets/SStatsView.cpp
+++ b/Engine/Source/Developer/TraceInsights/Private/Insights/TimingProfiler/Widgets/SStatsView.cpp
@@ -6,11 +6,14 @@
 #include "Framework/Commands/Commands.h"
 #include "Framework/Commands/UICommandList.h"
 #include "Framework/MultiBox/MultiBoxBuilder.h"
+#include "GenericPlatform/GenericPlatformTime.h"
 #include "HAL/PlatformApplicationMisc.h"
 #include "HAL/PlatformFileManager.h"
 #include "Logging/MessageLog.h"
 #include "SlateOptMacros.h"
 #include "Styling/AppStyle.h"
+#include "Styling/StyleColors.h"
+#include "Widgets/Colors/SColorPicker.h"
 #include "Widgets/Images/SImage.h"
 #include "Widgets/Input/SCheckBox.h"
 #include "Widgets/Input/SSearchBox.h"
@@ -78,6 +81,18 @@ public:
 	UE_DISABLE_OPTIMIZATION_SHIP
 	virtual void RegisterCommands() override
 	{
+		UI_COMMAND(Command_GenerateNewColor,
+			"Generate New Color",
+			"Generates a new random color for the selected counters.",
+			EUserInterfaceActionType::Button,
+			FInputChord());
+
+		UI_COMMAND(Command_EditColor,
+			"Edit Color...",
+			"Changes color for the selected counters.",
+			EUserInterfaceActionType::Button,
+			FInputChord());
+
 		UI_COMMAND(Command_CopyToClipboard,
 			"Copy To Clipboard",
 			"Copies the selection (counters and their aggregated statistics) to clipboard.",
@@ -116,6 +131,8 @@ public:
 	}
 	UE_ENABLE_OPTIMIZATION_SHIP
 
+	TSharedPtr<FUICommandInfo> Command_GenerateNewColor;
+	TSharedPtr<FUICommandInfo> Command_EditColor;
 	TSharedPtr<FUICommandInfo> Command_CopyToClipboard;
 	TSharedPtr<FUICommandInfo> Command_CopyNameToClipboard;
 	TSharedPtr<FUICommandInfo> Command_Export;
@@ -165,6 +182,8 @@ void SStatsView::InitCommandList()
 {
 	FStatsViewCommands::Register();
 	CommandList = MakeShared<FUICommandList>();
+	CommandList->MapAction(FStatsViewCommands::Get().Command_GenerateNewColor, FExecuteAction::CreateSP(this, &SStatsView::ContextMenu_GenerateNewColor_Execute), FCanExecuteAction::CreateSP(this, &SStatsView::ContextMenu_GenerateNewColor_CanExecute));
+	CommandList->MapAction(FStatsViewCommands::Get().Command_EditColor, FExecuteAction::CreateSP(this, &SStatsView::ContextMenu_EditColor_Execute), FCanExecuteAction::CreateSP(this, &SStatsView::ContextMenu_EditColor_CanExecute));
 	CommandList->MapAction(FStatsViewCommands::Get().Command_CopyToClipboard, FExecuteAction::CreateSP(this, &SStatsView::ContextMenu_CopySelectedToClipboard_Execute), FCanExecuteAction::CreateSP(this, &SStatsView::ContextMenu_CopySelectedToClipboard_CanExecute));
 	CommandList->MapAction(FStatsViewCommands::Get().Command_CopyNameToClipboard, FExecuteAction::CreateSP(this, &SStatsView::ContextMenu_CopySelectedNameToClipboard_Execute), FCanExecuteAction::CreateSP(this, &SStatsView::ContextMenu_CopySelectedNameToClipboard_CanExecute));
 	CommandList->MapAction(FStatsViewCommands::Get().Command_Export, FExecuteAction::CreateSP(this, &SStatsView::ContextMenu_Export_Execute), FCanExecuteAction::CreateSP(this, &SStatsView::ContextMenu_Export_CanExecute));
@@ -486,6 +505,24 @@ TSharedPtr<SWidget> SStatsView::TreeView_GetMenuContent()
 				);
 			}
 		}
+
+		MenuBuilder.AddMenuEntry
+		(
+			FStatsViewCommands::Get().Command_GenerateNewColor,
+			NAME_None,
+			TAttribute<FText>(),
+			TAttribute<FText>(),
+			FSlateIcon(FAppStyle::Get().GetStyleSetName(), "Icons.Refresh")
+		);
+
+		MenuBuilder.AddMenuEntry
+		(
+			FStatsViewCommands::Get().Command_EditColor,
+			NAME_None,
+			TAttribute<FText>(),
+			TAttribute<FText>(),
+			FSlateIcon(FAppStyle::Get().GetStyleSetName(), "Icons.EyeDropper")
+		);
 	}
 	MenuBuilder.EndSection();
 
@@ -2208,6 +2245,140 @@ void SStatsView::ToggleTimingViewMainGraphEventSeries(FStatsNodePtr CounterNode)
 	}
 }
 
+////////////////////////////////////////////////////////////////////////////////////////////////////
+// Generate and Edit Colors
+////////////////////////////////////////////////////////////////////////////////////////////////////
+
+bool SStatsView::ContextMenu_GenerateNewColor_CanExecute() const
+{
+	return AreSelectedCounterNodesValid() && GetStatsView().IsValid();
+}
+
+////////////////////////////////////////////////////////////////////////////////////////////////////
+
+void SStatsView::ContextMenu_GenerateNewColor_Execute() const
+{
+	const TArray<FStatsNodePtr> SelectedNodes = TreeView->GetSelectedItems();
+	for (const FStatsNodePtr& SelectedNode : SelectedNodes)
+	{
+		if (SelectedNode->Is<FStatsNode>())
+		{
+			FStatsNode& StatsNode = SelectedNode->As<FStatsNode>();
+			constexpr bool bSetRandomColor = true;
+			SetColorToNode(StatsNode, FLinearColor(), bSetRandomColor);
+		}
+	}
+};
+
+////////////////////////////////////////////////////////////////////////////////////////////////////
+
+void SStatsView::SetColorToNode(FStatsNode& StatsNode, FLinearColor CustomColor, bool bSetRandomColor) const
+{
+	if (bSetRandomColor)
+	{
+		uint64 Time = FPlatformTime::Cycles64();
+		uint32 Hash = (Time & 0xFFFFFFFF) ^ (Time >> 32);
+		Hash = ((Hash << 5) + Hash) * 0xfb23618f;
+
+		const uint8 H = Hash & 0xFF;
+		const uint8 S = 128 + ((Hash >> 8) & 0x7F);
+		const uint8 V = 128 + ((Hash >> 16) & 0x7F);
+		CustomColor = FLinearColor::MakeFromHSV8(H, S, V);
+		CustomColor.A = 1.0f;
+
+		StatsNode.SetColor(CustomColor);
+	}
+	else
+	{
+		StatsNode.SetColor(CustomColor);
+	}
+
+	const uint32 CounterId = StatsNode.GetCounterId();
+	const FLinearColor BorderColor(CustomColor.R + 0.4f, CustomColor.G + 0.4f, CustomColor.B + 0.4f, 1.0f);
+
+	if (GetStatsView().IsValid())
+	{
+		TSharedPtr<FTimingGraphSeries> CounterSeries = GetTimingViewMainGraphTrack()->GetStatsCounterSeries(CounterId);
+		if (CounterSeries.IsValid())
+		{
+			CounterSeries->SetColor(CustomColor, BorderColor);
+		}
+	}
+}
+
+////////////////////////////////////////////////////////////////////////////////////////////////////
+
+bool SStatsView::ContextMenu_EditColor_CanExecute() const
+{
+	return AreSelectedCounterNodesValid() && GetStatsView().IsValid();
+};
+
+////////////////////////////////////////////////////////////////////////////////////////////////////
+
+void SStatsView::ContextMenu_EditColor_Execute()
+{
+	if (GetStatsView().IsValid())
+	{
+		EditableColorValue = FLinearColor(0.5f, 0.5f, 0.5f, 1.0f);
+		const TArray<FStatsNodePtr> SelectedNodes = TreeView->GetSelectedItems();
+		for (const FStatsNodePtr& SelectedNode : SelectedNodes)
+		{
+			if (SelectedNode->Is<FStatsNode>())
+			{
+				FStatsNode& StatsNode = SelectedNode->As<FStatsNode>();
+				EditableColorValue = StatsNode.GetColor();
+				break;
+			}
+		}
+
+		FColorPickerArgs PickerArgs;
+		{
+			PickerArgs.bUseAlpha = true;
+			PickerArgs.bOnlyRefreshOnMouseUp = false;
+			PickerArgs.bOnlyRefreshOnOk = false;
+			PickerArgs.bExpandAdvancedSection = false;
+			PickerArgs.OnColorCommitted = FOnLinearColorValueChanged::CreateSP(this, &SStatsView::SetEditableColor);
+			PickerArgs.OnColorPickerCancelled = FOnColorPickerCancelled::CreateSP(this, &SStatsView::ColorPickerCancelled);
+			PickerArgs.InitialColor = EditableColorValue;
+			PickerArgs.ParentWidget = SharedThis(this);
+		}
+
+		OpenColorPicker(PickerArgs);
+	}
+};
+
+////////////////////////////////////////////////////////////////////////////////////////////////////
+
+FLinearColor SStatsView::GetEditableColor() const
+{
+	return EditableColorValue;
+}
+
+////////////////////////////////////////////////////////////////////////////////////////////////////
+
+void SStatsView::SetEditableColor(FLinearColor NewColor)
+{
+	EditableColorValue = NewColor;
+
+	const TArray<FStatsNodePtr> SelectedNodes = TreeView->GetSelectedItems();
+	for (const FStatsNodePtr& SelectedNode : SelectedNodes)
+	{
+		if (SelectedNode->Is<FStatsNode>())
+		{
+			FStatsNode& StatsNode = SelectedNode->As<FStatsNode>();
+			constexpr bool bSetRandomColor = false;
+			SetColorToNode(StatsNode, EditableColorValue, bSetRandomColor);
+		}
+	}
+}
+
+////////////////////////////////////////////////////////////////////////////////////////////////////
+
+void SStatsView::ColorPickerCancelled(FLinearColor OriginalColor)
+{
+	SetEditableColor(OriginalColor);
+}
+
 ////////////////////////////////////////////////////////////////////////////////////////////////////
 
 bool SStatsView::ContextMenu_CopySelectedToClipboard_CanExecute() const
@@ -2619,6 +2790,29 @@ FReply SStatsView::OnKeyDown(const FGeometry& MyGeometry, const FKeyEvent& InKey
 
 ////////////////////////////////////////////////////////////////////////////////////////////////////
 
+TSharedPtr<SStatsView> SStatsView::GetStatsView() const
+{
+	TSharedPtr<STimingProfilerWindow> ProfilerWindow = FTimingProfilerManager::Get()->GetProfilerWindow();
+	return ProfilerWindow.IsValid() ? ProfilerWindow->GetStatsView() : nullptr;
+}
+
+////////////////////////////////////////////////////////////////////////////////////////////////////
+
+bool SStatsView::AreSelectedCounterNodesValid() const
+{
+	const TArray<FStatsNodePtr> SelectedNodes = TreeView->GetSelectedItems();
+	for (const FStatsNodePtr& SelectedNode : SelectedNodes)
+	{
+		if (!SelectedNode.IsValid())
+		{
+			return false;
+		}
+	}
+	return true;
+}
+
+////////////////////////////////////////////////////////////////////////////////////////////////////
+
 } // namespace UE::Insights::TimingProfiler
 
 #undef LOCTEXT_NAMESPACE
diff --git a/Engine/Source/Developer/TraceInsights/Private/Insights/TimingProfiler/Widgets/SStatsView.h b/Engine/Source/Developer/TraceInsights/Private/Insights/TimingProfiler/Widgets/SStatsView.h
--- a/Engine/Source/Developer/TraceInsights/Private/Insights/TimingProfiler/Widgets/SStatsView.h
+++ b/Engine/Source/Developer/TraceInsights/Private/Insights/TimingProfiler/Widgets/SStatsView.h
@@ -124,6 +124,22 @@ private:
 	void TreeView_BuildSortByMenu(FMenuBuilder& MenuBuilder);
 	void TreeView_BuildViewColumnMenu(FMenuBuilder& MenuBuilder);
 
+	//////////////////////////////////////////////////
+	// Generate and Edit Counters Colors
+
+	bool ContextMenu_GenerateNewColor_CanExecute() const;
+	void ContextMenu_GenerateNewColor_Execute() const;
+
+	bool ContextMenu_EditColor_CanExecute() const;
+	void ContextMenu_EditColor_Execute();
+
+	void SetColorToNode(FStatsNode& StatsNode, FLinearColor CustomColor, bool bSetRandomColor) const;
+	FLinearColor GetEditableColor() const;
+	void SetEditableColor(FLinearColor NewColor);
+	void ColorPickerCancelled(FLinearColor OriginalColor);
+
+	//////////////////////////////////////////////////
+
 	bool ContextMenu_CopySelectedToClipboard_CanExecute() const;
 	void ContextMenu_CopySelectedToClipboard_Execute();
 
@@ -172,6 +188,9 @@ private:
 	/** Called by STreeView when a tree item is double clicked. */
 	void TreeView_OnMouseButtonDoubleClick(FStatsNodePtr TreeNode);
 
+	/** Checks validity of one or more selected nodes*/
+	bool AreSelectedCounterNodesValid() const;
+
 	////////////////////////////////////////////////////////////////////////////////////////////////////
 	// Tree View - Table Row
 
@@ -300,6 +319,8 @@ private:
 	 */
 	virtual void Tick(const FGeometry& AllottedGeometry, const double InCurrentTime, const float InDeltaTime) override;
 
+	TSharedPtr<SStatsView> GetStatsView() const;
+
 private:
 	/** Table view model. */
 	TSharedPtr<FTable> Table;
@@ -309,6 +330,9 @@ private:
 
 	TSharedPtr<FUICommandList> CommandList;
 
+	/** Eye-Dropper selected color for graphing.*/
+	FLinearColor EditableColorValue;
+
 	//////////////////////////////////////////////////
 	// Tree View, Columns
 
diff --git a/Engine/Source/Developer/TraceInsights/Private/Insights/TimingProfiler/Widgets/STimerTreeView.cpp b/Engine/Source/Developer/TraceInsights/Private/Insights/TimingProfiler/Widgets/STimerTreeView.cpp
--- a/Engine/Source/Developer/TraceInsights/Private/Insights/TimingProfiler/Widgets/STimerTreeView.cpp
+++ b/Engine/Source/Developer/TraceInsights/Private/Insights/TimingProfiler/Widgets/STimerTreeView.cpp
@@ -1526,10 +1526,53 @@ FTimerNodePtr STimerTreeView::CreateTimerNodeRec(const TraceServices::FTimingPro
 	const ETimerNodeType Type = Node.Timer->IsGpuTimer ?   ETimerNodeType::GpuScope :
 								Node.Timer->IsVerseTimer ? ETimerNodeType::VerseScope :
 														   ETimerNodeType::CpuScope;
-	FTimerNodePtr TimerNodePtr = MakeShared<FTimerNode>(Node.Timer->Id, Node.Timer->Name, Type, true);
 
 	TSharedPtr<STimingProfilerWindow> ProfilerWindow = FTimingProfilerManager::Get()->GetProfilerWindow();
 	const bool bIsTimerAddedToGraphs = ProfilerWindow.IsValid() ? ProfilerWindow->IsTimerAddedToGraphs(Node.Timer->Id) : false;
+
+	FLinearColor ColorOverride;
+	if (GetTimingView().IsValid())
+	{
+		TSharedPtr<FTimingGraphTrack> GraphTrack = GetTimingViewMainGraphTrack();
+
+		if (GraphTrack.IsValid())
+		{
+			for (const TSharedPtr<FGraphSeries>& Series : GraphTrack->GetSeries())
+			{
+				if (Series->Is<FTimingGraphSeries>())
+				{
+					if (StaticCastSharedPtr<FTimingGraphSeries>(Series)->TimerId == Node.Timer->Id)
+					{
+						ColorOverride = Series->GetColor();
+						break;
+					}
+				}
+			}
+		}
+	}
+
+	if (GetFrameTrack().IsValid())
+	{
+		TSharedPtr<SFrameTrack> FrameTrack = GetFrameTrack();
+
+		if (FrameTrack.IsValid())
+		{
+			for (const TSharedPtr<FFrameTrackSeries>& Series : GetFrameTrack()->GetSeries())
+			{
+				if(Series->Type == EFrameTrackSeriesType::TimerFrameStats)
+				{
+					if (StaticCastSharedPtr<FTimerFrameStatsTrackSeries>(Series)->TimerId == Node.Timer->Id)
+					{
+						ColorOverride = Series->Color;
+						break;
+					}
+				}
+			}
+		}
+	}
+
+	FTimerNodePtr TimerNodePtr = MakeShared<FTimerNode>(Node.Timer->Id, Node.Timer->Name, Type, true);
+	TimerNodePtr->SetColor(ColorOverride);
 	TimerNodePtr->SetAddedToGraphsFlag(bIsTimerAddedToGraphs);
 
 	TraceServices::FTimingProfilerAggregatedStats AggregatedStats;
diff --git a/Engine/Source/Developer/TraceInsights/Private/Insights/TimingProfiler/Widgets/STimerTreeView.h b/Engine/Source/Developer/TraceInsights/Private/Insights/TimingProfiler/Widgets/STimerTreeView.h
--- a/Engine/Source/Developer/TraceInsights/Private/Insights/TimingProfiler/Widgets/STimerTreeView.h
+++ b/Engine/Source/Developer/TraceInsights/Private/Insights/TimingProfiler/Widgets/STimerTreeView.h
@@ -65,6 +65,8 @@ public:
 
 	void SetTree(const TraceServices::FTimingProfilerButterflyNode& Root);
 
+	TArray<FTimerNodePtr> GetTreeNodes() const { return TreeNodes; };
+
 	FTimerNodePtr GetTimerNode(uint32 TimerId) const;
 
 	virtual FReply OnKeyDown(const FGeometry& MyGeometry, const FKeyEvent& InKeyEvent) override;
@@ -98,6 +100,12 @@ private:
 
 	void OpenSourceFileInIDE(FTimerNodePtr InNode) const;
 
+	//////////////////////////////////////////////////
+	// Generate and Edit Timer Colors
+
+	bool ContextMenu_GenerateNewColor_CanExecute() const;
+	void ContextMenu_GenerateNewColor_Execute() const;
+
 	//////////////////////////////////////////////////
 	// Plot (Graph Series)
 
diff --git a/Engine/Source/Developer/TraceInsights/Private/Insights/TimingProfiler/Widgets/STimersView.cpp b/Engine/Source/Developer/TraceInsights/Private/Insights/TimingProfiler/Widgets/STimersView.cpp
--- a/Engine/Source/Developer/TraceInsights/Private/Insights/TimingProfiler/Widgets/STimersView.cpp
+++ b/Engine/Source/Developer/TraceInsights/Private/Insights/TimingProfiler/Widgets/STimersView.cpp
@@ -6,6 +6,7 @@
 #include "Framework/Commands/Commands.h"
 #include "Framework/Commands/UICommandList.h"
 #include "Framework/MultiBox/MultiBoxBuilder.h"
+#include "GenericPlatform/GenericPlatformTime.h"
 #include "HAL/PlatformApplicationMisc.h"
 #include "HAL/PlatformFileManager.h"
 #include "ISourceCodeAccessModule.h"
@@ -14,6 +15,8 @@
 #include "Modules/ModuleManager.h"
 #include "SlateOptMacros.h"
 #include "Styling/AppStyle.h"
+#include "Styling/StyleColors.h"
+#include "Widgets/Colors/SColorPicker.h"
 #include "Widgets/Images/SImage.h"
 #include "Widgets/Input/SCheckBox.h"
 #include "Widgets/Input/SSearchBox.h"
@@ -133,6 +136,18 @@ public:
 			EUserInterfaceActionType::Button,
 			FInputChord());
 
+		UI_COMMAND(Command_GenerateNewColor,
+			"Generate New Color",
+			"Generates a new random color for the selected timers.",
+			EUserInterfaceActionType::Button,
+			FInputChord());
+
+		UI_COMMAND(Command_EditColor,
+			"Edit Color...",
+			"Changes color for the selected timers.",
+			EUserInterfaceActionType::Button,
+			FInputChord());
+
 		UI_COMMAND(Command_FindMaxInstance,
 			"Maximum Duration Instance",
 			"Navigates to and selects the timing event instance with the maximum duration, for the selected timer.",
@@ -167,6 +182,8 @@ public:
 	TSharedPtr<FUICommandInfo> Command_ExportThreads;
 	TSharedPtr<FUICommandInfo> Command_ExportTimers;
 	TSharedPtr<FUICommandInfo> Command_OpenSource;
+	TSharedPtr<FUICommandInfo> Command_GenerateNewColor;
+	TSharedPtr<FUICommandInfo> Command_EditColor;
 	TSharedPtr<FUICommandInfo> Command_FindMaxInstance;
 	TSharedPtr<FUICommandInfo> Command_FindMinInstance;
 	TSharedPtr<FUICommandInfo> Command_FindMaxInstanceInSelection;
@@ -219,6 +236,8 @@ void STimersView::InitCommandList()
 	CommandList->MapAction(FTimersViewCommands::Get().Command_ExportThreads, FExecuteAction::CreateSP(this, &STimersView::ContextMenu_ExportThreads_Execute), FCanExecuteAction::CreateSP(this, &STimersView::ContextMenu_ExportThreads_CanExecute));
 	CommandList->MapAction(FTimersViewCommands::Get().Command_ExportTimers, FExecuteAction::CreateSP(this, &STimersView::ContextMenu_ExportTimers_Execute), FCanExecuteAction::CreateSP(this, &STimersView::ContextMenu_ExportTimers_CanExecute));
 	CommandList->MapAction(FTimersViewCommands::Get().Command_OpenSource, FExecuteAction::CreateSP(this, &STimersView::ContextMenu_OpenSource_Execute), FCanExecuteAction::CreateSP(this, &STimersView::ContextMenu_OpenSource_CanExecute));
+	CommandList->MapAction(FTimersViewCommands::Get().Command_GenerateNewColor, FExecuteAction::CreateSP(this, &STimersView::ContextMenu_GenerateNewColor_Execute), FCanExecuteAction::CreateSP(this, &STimersView::ContextMenu_GenerateNewColor_CanExecute));
+	CommandList->MapAction(FTimersViewCommands::Get().Command_EditColor, FExecuteAction::CreateSP(this, &STimersView::ContextMenu_EditColor_Execute), FCanExecuteAction::CreateSP(this, &STimersView::ContextMenu_EditColor_CanExecute));
 	CommandList->MapAction(FTimersViewCommands::Get().Command_FindMaxInstance, FExecuteAction::CreateSP(this, &STimersView::ContextMenu_FindInstance_Execute, true), FCanExecuteAction::CreateSP(this, &STimersView::ContextMenu_FindInstance_CanExecute));
 	CommandList->MapAction(FTimersViewCommands::Get().Command_FindMinInstance, FExecuteAction::CreateSP(this, &STimersView::ContextMenu_FindInstance_Execute, false), FCanExecuteAction::CreateSP(this, &STimersView::ContextMenu_FindInstance_CanExecute));
 	CommandList->MapAction(FTimersViewCommands::Get().Command_FindMaxInstanceInSelection, FExecuteAction::CreateSP(this, &STimersView::ContextMenu_FindInstanceInSelection_Execute, true), FCanExecuteAction::CreateSP(this, &STimersView::ContextMenu_FindInstanceInSelection_CanExecute));
@@ -666,6 +685,24 @@ TSharedPtr<SWidget> STimersView::TreeView_GetMenuContent()
 				ItemToolTip,
 				FSlateIcon(SourceCodeAccessor.GetStyleSet(), SourceCodeAccessor.GetOpenIconName()));
 		}
+
+		MenuBuilder.AddMenuEntry
+		(
+			FTimersViewCommands::Get().Command_GenerateNewColor,
+			NAME_None,
+			TAttribute<FText>(),
+			TAttribute<FText>(),
+			FSlateIcon(FAppStyle::Get().GetStyleSetName(), "Icons.Refresh")
+		);
+
+		MenuBuilder.AddMenuEntry
+		(
+			FTimersViewCommands::Get().Command_EditColor,
+			NAME_None,
+			TAttribute<FText>(),
+			TAttribute<FText>(),
+			FSlateIcon(FAppStyle::Get().GetStyleSetName(), "Icons.EyeDropper")
+		);
 	}
 	MenuBuilder.EndSection();
 
@@ -2718,6 +2755,7 @@ void STimersView::ToggleGraphInstanceSeries(TSharedRef<FTimingGraphTrack> GraphT
 	if (ProfilerWindow.IsValid())
 	{
 		ProfilerWindow->OnTimerAddedToGraphsChanged(TimerId);
+		FTimingProfilerManager::Get()->UpdateCallersAndCallees();
 	}
 }
 
@@ -2776,6 +2814,7 @@ void STimersView::ToggleGraphFrameStatsSeries(TSharedRef<FTimingGraphTrack> Grap
 	if (ProfilerWindow.IsValid())
 	{
 		ProfilerWindow->OnTimerAddedToGraphsChanged(TimerId);
+		FTimingProfilerManager::Get()->UpdateCallersAndCallees();
 	}
 }
 
@@ -2850,6 +2889,7 @@ void STimersView::ToggleFrameTrackSeries(FTimerNodePtr TimerNode, ETraceFrameTyp
 	if (ProfilerWindow.IsValid())
 	{
 		ProfilerWindow->OnTimerAddedToGraphsChanged(TimerId);
+		FTimingProfilerManager::Get()->UpdateCallersAndCallees();
 	}
 }
 
@@ -3380,6 +3420,171 @@ void STimersView::OpenSourceFileInIDE(FTimerNodePtr InNode) const
 	}
 }
 
+////////////////////////////////////////////////////////////////////////////////////////////////////
+// Generate and Edit Colors
+////////////////////////////////////////////////////////////////////////////////////////////////////
+
+bool STimersView::ContextMenu_GenerateNewColor_CanExecute() const
+{
+	return AreSelectedTimerNodesValid() && GetTimingView().IsValid();
+}
+
+////////////////////////////////////////////////////////////////////////////////////////////////////
+
+void STimersView::ContextMenu_GenerateNewColor_Execute() const
+{
+	const TArray<FTimerNodePtr> SelectedNodes = TreeView->GetSelectedItems();
+	for (const FTimerNodePtr& SelectedNode : SelectedNodes)
+	{
+		if (SelectedNode->Is<FTimerNode>())
+		{
+			FTimerNode& TimerNode = SelectedNode->As<FTimerNode>();
+			constexpr bool bSetRandomColor = true;
+			SetColorToNode(TimerNode, FLinearColor(), bSetRandomColor);
+		}
+	}
+};
+
+////////////////////////////////////////////////////////////////////////////////////////////////////
+
+void STimersView::SetColorToNode(FTimerNode& TimerNode, FLinearColor CustomColor, bool bSetRandomColor) const
+{
+	if (bSetRandomColor)
+	{
+		uint64 Time = FPlatformTime::Cycles64();
+		uint32 Hash = (Time & 0xFFFFFFFF) ^ (Time >> 32);
+		Hash = ((Hash << 5) + Hash) * 0xfb23618f;
+
+		const uint8 H = Hash & 0xFF;
+		const uint8 S = 128 + ((Hash >> 8) & 0x7F);
+		const uint8 V = 128 + ((Hash >> 16) & 0x7F);
+		CustomColor = FLinearColor::MakeFromHSV8(H, S, V);
+		CustomColor.A = 1.0f;
+
+		TimerNode.SetColor(CustomColor);
+	}
+	else
+	{
+		TimerNode.SetColor(CustomColor);
+	}
+
+	const uint32 TimerId = TimerNode.GetTimerId();
+	const FLinearColor BorderColor(CustomColor.R + 0.4f, CustomColor.G + 0.4f, CustomColor.B + 0.4f, 1.0f);
+
+	if (GetTimingView().IsValid())
+	{
+		TSharedPtr<FTimingGraphTrack> GraphTrack = GetTimingViewMainGraphTrack();
+
+		if (GraphTrack.IsValid())
+		{
+			for (const TSharedPtr<FGraphSeries>& Series : GraphTrack->GetSeries())
+			{
+				if (Series->Is<FTimingGraphSeries>())
+				{
+					if (StaticCastSharedPtr<FTimingGraphSeries>(Series)->TimerId == TimerId)
+					{
+						Series->SetColor(CustomColor, BorderColor);
+					}
+				}
+			}
+		}
+	}
+
+	if (GetFrameTrack().IsValid())
+	{
+		TSharedPtr<SFrameTrack> FrameTrack = GetFrameTrack();
+
+		if (FrameTrack.IsValid())
+		{
+			for (const TSharedPtr<FFrameTrackSeries>& Series : GetFrameTrack()->GetSeries())
+			{
+				if (Series->Type == EFrameTrackSeriesType::TimerFrameStats)
+				{
+					if (StaticCastSharedPtr<FTimerFrameStatsTrackSeries>(Series)->TimerId == TimerId)
+					{
+						Series->Color = CustomColor;
+					}
+				}
+			}
+		}
+	}
+
+	FTimingProfilerManager::Get()->UpdateCallersAndCallees();
+}
+
+////////////////////////////////////////////////////////////////////////////////////////////////////
+
+bool STimersView::ContextMenu_EditColor_CanExecute() const
+{
+	return AreSelectedTimerNodesValid() && GetTimingView().IsValid();
+};
+
+////////////////////////////////////////////////////////////////////////////////////////////////////
+
+void STimersView::ContextMenu_EditColor_Execute()
+{
+	if (GetTimingView().IsValid())
+	{
+		EditableColorValue = FLinearColor(0.5f, 0.5f, 0.5f, 1.0f);
+		const TArray<FTimerNodePtr> SelectedNodes = TreeView->GetSelectedItems();
+		for (const FTimerNodePtr& SelectedNode : SelectedNodes)
+		{
+			if (SelectedNode->Is<FTimerNode>())
+			{
+				FTimerNode& TimerNode = SelectedNode->As<FTimerNode>();
+				EditableColorValue = TimerNode.GetColor();
+				break;
+			}
+		}
+
+		FColorPickerArgs PickerArgs;
+		{
+			PickerArgs.bUseAlpha = true;
+			PickerArgs.bOnlyRefreshOnMouseUp = false;
+			PickerArgs.bOnlyRefreshOnOk = false;
+			PickerArgs.bExpandAdvancedSection = false;
+			PickerArgs.OnColorCommitted = FOnLinearColorValueChanged::CreateSP(this, &STimersView::SetEditableColor);
+			PickerArgs.OnColorPickerCancelled = FOnColorPickerCancelled::CreateSP(this, &STimersView::ColorPickerCancelled);
+			PickerArgs.InitialColor = EditableColorValue;
+			PickerArgs.ParentWidget = SharedThis(this);
+		}
+
+		OpenColorPicker(PickerArgs);
+	}
+};
+
+////////////////////////////////////////////////////////////////////////////////////////////////////
+
+FLinearColor STimersView::GetEditableColor() const
+{
+	return EditableColorValue;
+}
+
+////////////////////////////////////////////////////////////////////////////////////////////////////
+
+void STimersView::SetEditableColor(FLinearColor NewColor)
+{
+	EditableColorValue = NewColor;
+
+	const TArray<FTimerNodePtr> SelectedNodes = TreeView->GetSelectedItems();
+	for (const FTimerNodePtr& SelectedNode : SelectedNodes)
+	{
+		if (SelectedNode->Is<FTimerNode>())
+		{
+			FTimerNode& TimerNode = SelectedNode->As<FTimerNode>();
+			constexpr bool bSetRandomColor = false;
+			SetColorToNode(TimerNode, EditableColorValue, bSetRandomColor);
+		}
+	}
+}
+
+////////////////////////////////////////////////////////////////////////////////////////////////////
+
+void STimersView::ColorPickerCancelled(FLinearColor OriginalColor)
+{
+	SetEditableColor(OriginalColor);
+}
+
 ////////////////////////////////////////////////////////////////////////////////////////////////////
 // Find Min/Max Instance
 ////////////////////////////////////////////////////////////////////////////////////////////////////
@@ -3601,6 +3806,20 @@ FReply STimersView::OnKeyDown(const FGeometry& MyGeometry, const FKeyEvent& InKe
 
 ////////////////////////////////////////////////////////////////////////////////////////////////////
 
+bool STimersView::AreSelectedTimerNodesValid() const
+{
+	const TArray<FTimerNodePtr> SelectedNodes = TreeView->GetSelectedItems();
+	for (const FTimerNodePtr& SelectedNode : SelectedNodes)
+	{
+		if (!SelectedNode.IsValid())
+		{
+			return false;
+		}
+	}
+	return true;
+}
+
+////////////////////////////////////////////////////////////////////////////////////////////////////
 } // namespace UE::Insights::TimingProfiler
 
 #undef LOCTEXT_NAMESPACE
diff --git a/Engine/Source/Developer/TraceInsights/Private/Insights/TimingProfiler/Widgets/STimersView.h b/Engine/Source/Developer/TraceInsights/Private/Insights/TimingProfiler/Widgets/STimersView.h
--- a/Engine/Source/Developer/TraceInsights/Private/Insights/TimingProfiler/Widgets/STimersView.h
+++ b/Engine/Source/Developer/TraceInsights/Private/Insights/TimingProfiler/Widgets/STimersView.h
@@ -168,6 +168,20 @@ private:
 
 	void OpenSourceFileInIDE(FTimerNodePtr InNode) const;
 
+	//////////////////////////////////////////////////
+	// Generate and Edit Timer Colors
+
+	bool ContextMenu_GenerateNewColor_CanExecute() const;
+	void ContextMenu_GenerateNewColor_Execute() const;
+
+	bool ContextMenu_EditColor_CanExecute() const;
+	void ContextMenu_EditColor_Execute();
+
+	void SetColorToNode(FTimerNode& TimerNode, FLinearColor CustomColor, bool bSetRandomColor) const;
+	FLinearColor GetEditableColor() const;
+	void SetEditableColor(FLinearColor NewColor);
+	void ColorPickerCancelled(FLinearColor OriginalColor);
+
 	//////////////////////////////////////////////////
 	// Plot (Graph Series)
 
@@ -210,6 +224,9 @@ private:
 
 	FTimerNodePtr GetSingleSelectedTimerNode() const;
 
+	/** Checks validity of one or more selected nodes*/
+	bool AreSelectedTimerNodesValid() const;
+
 	/** Called by STreeView when a tree item is double clicked. */
 	void TreeView_OnMouseButtonDoubleClick(FTimerNodePtr TreeNode);
 
@@ -387,6 +404,9 @@ private:
 
 	TSharedPtr<FUICommandList> CommandList;
 
+	/** Eye-Dropper selected color for graphing.*/
+	FLinearColor EditableColorValue;
+
 	//////////////////////////////////////////////////
 	// Tree View, Columns
 
diff --git a/Engine/Source/Developer/TraceInsights/Private/Insights/ViewModels/TimingGraphTrack.cpp b/Engine/Source/Developer/TraceInsights/Private/Insights/ViewModels/TimingGraphTrack.cpp
--- a/Engine/Source/Developer/TraceInsights/Private/Insights/ViewModels/TimingGraphTrack.cpp
+++ b/Engine/Source/Developer/TraceInsights/Private/Insights/ViewModels/TimingGraphTrack.cpp
@@ -30,6 +30,8 @@
 
 #define LOCTEXT_NAMESPACE "GraphTrack"
 
+INSIGHTS_IMPLEMENT_RTTI(FTimingGraphSeries)
+
 ////////////////////////////////////////////////////////////////////////////////////////////////////
 // FTimingGraphSeries
 ////////////////////////////////////////////////////////////////////////////////////////////////////
diff --git a/Engine/Source/Developer/TraceInsights/Private/Insights/ViewModels/TimingGraphTrack.h b/Engine/Source/Developer/TraceInsights/Private/Insights/ViewModels/TimingGraphTrack.h
--- a/Engine/Source/Developer/TraceInsights/Private/Insights/ViewModels/TimingGraphTrack.h
+++ b/Engine/Source/Developer/TraceInsights/Private/Insights/ViewModels/TimingGraphTrack.h
@@ -21,6 +21,8 @@ DECLARE_MULTICAST_DELEGATE_OneParam(FSeriesVisibilityChangedDelegate, bool bOnOf
 
 class FTimingGraphSeries : public FGraphSeries
 {
+	INSIGHTS_DECLARE_RTTI(FTimingGraphSeries, FGraphSeries)
+
 public:
 	enum class ESeriesType
 	{
