Fix timing export with timer filters for timers with metadata.
Fix timing export for selection for GPU Queues.

diff --git a/Engine/Source/Developer/TraceInsights/Private/Insights/TimingProfiler/ViewModels/TimingExporter.cpp b/Engine/Source/Developer/TraceInsights/Private/Insights/TimingProfiler/ViewModels/TimingExporter.cpp
--- a/Engine/Source/Developer/TraceInsights/Private/Insights/TimingProfiler/ViewModels/TimingExporter.cpp
+++ b/Engine/Source/Developer/TraceInsights/Private/Insights/TimingProfiler/ViewModels/TimingExporter.cpp
@@ -438,7 +438,7 @@ int32 FTimingExporter::ExportTimingEvents_WriteEvents(FExportTimingEventsInterna
 			Timeline.EnumerateEvents(Params.UserParams.IntervalStartTime, Params.UserParams.IntervalEndTime,
 				[&Params, &TimingEventCount, &Timers, TimerReader](double EventStartTime, double EventEndTime, uint32 EventDepth, const TraceServices::FTimingProfilerEvent& Event)
 				{
-					if (!Params.UserParams.TimingEventFilter || Params.UserParams.TimingEventFilter(EventStartTime, EventEndTime, EventDepth, Event))
+					if (!Params.UserParams.TimingEventFilter || Params.UserParams.TimingEventFilter(EventStartTime, EventEndTime, EventDepth, Event, TimerReader))
 					{
 						FUtf8StringBuilder& StringBuilder = Params.Writer.GetStringBuilder();
 						bool bFirst = true;
@@ -1257,7 +1257,7 @@ FTimingExporter::FTimingEventFilterFunc FTimingExporter::MakeTimingEventFilterBy
 {
 	if (IncludedTimers.Num() == 0)
 	{
-		return [](double EventStartTime, double EventEndTime, uint32 EventDepth, const TraceServices::FTimingProfilerEvent& Event) -> bool
+		return [](double EventStartTime, double EventEndTime, uint32 EventDepth, const TraceServices::FTimingProfilerEvent& Event, const TraceServices::ITimingProfilerTimerReader* TimerReader) -> bool
 		{
 			return false;
 		};
@@ -1266,15 +1266,25 @@ FTimingExporter::FTimingEventFilterFunc FTimingExporter::MakeTimingEventFilterBy
 	if (IncludedTimers.Num() == 1)
 	{
 		const uint32 IncludedTimerId = IncludedTimers[FSetElementId::FromInteger(0)];
-		return [IncludedTimerId](double EventStartTime, double EventEndTime, uint32 EventDepth, const TraceServices::FTimingProfilerEvent& Event) -> bool
+		return [IncludedTimerId](double EventStartTime, double EventEndTime, uint32 EventDepth, const TraceServices::FTimingProfilerEvent& Event, const TraceServices::ITimingProfilerTimerReader* TimerReader) -> bool
 		{
-			return Event.TimerIndex == IncludedTimerId;
+			uint32 TimerIndex = Event.TimerIndex;
+			if ((int32)TimerIndex < 0)
+			{
+				TimerIndex = TimerReader->GetOriginalTimerIdFromMetadata(TimerIndex);
+			}
+			return TimerIndex == IncludedTimerId;
 		};
 	}
 
-	return [&IncludedTimers](double EventStartTime, double EventEndTime, uint32 EventDepth, const TraceServices::FTimingProfilerEvent& Event) -> bool
+	return [&IncludedTimers](double EventStartTime, double EventEndTime, uint32 EventDepth, const TraceServices::FTimingProfilerEvent& Event, const TraceServices::ITimingProfilerTimerReader* TimerReader) -> bool
 	{
-		return IncludedTimers.Contains(Event.TimerIndex);
+		uint32 TimerIndex = Event.TimerIndex;
+		if ((int32)TimerIndex < 0)
+		{
+			TimerIndex = TimerReader->GetOriginalTimerIdFromMetadata(TimerIndex);
+		}
+		return IncludedTimers.Contains(TimerIndex);
 	};
 }
 
@@ -1290,15 +1300,25 @@ FTimingExporter::FTimingEventFilterFunc FTimingExporter::MakeTimingEventFilterBy
 	if (ExcludedTimers.Num() == 1)
 	{
 		const uint32 ExcludedTimerId = ExcludedTimers[FSetElementId::FromInteger(0)];
-		return [ExcludedTimerId](double EventStartTime, double EventEndTime, uint32 EventDepth, const TraceServices::FTimingProfilerEvent& Event) -> bool
+		return [ExcludedTimerId](double EventStartTime, double EventEndTime, uint32 EventDepth, const TraceServices::FTimingProfilerEvent& Event, const TraceServices::ITimingProfilerTimerReader* TimerReader) -> bool
 		{
-			return Event.TimerIndex != ExcludedTimerId;
+			uint32 TimerIndex = Event.TimerIndex;
+			if ((int32)TimerIndex < 0)
+			{
+				TimerIndex = TimerReader->GetOriginalTimerIdFromMetadata(TimerIndex);
+			}
+			return TimerIndex != ExcludedTimerId;
 		};
 	}
 
-	return [&ExcludedTimers](double EventStartTime, double EventEndTime, uint32 EventDepth, const TraceServices::FTimingProfilerEvent& Event) -> bool
+	return [&ExcludedTimers](double EventStartTime, double EventEndTime, uint32 EventDepth, const TraceServices::FTimingProfilerEvent& Event, const TraceServices::ITimingProfilerTimerReader* TimerReader) -> bool
 	{
-		return !ExcludedTimers.Contains(Event.TimerIndex);
+		uint32 TimerIndex = Event.TimerIndex;
+		if ((int32)TimerIndex < 0)
+		{
+			TimerIndex = TimerReader->GetOriginalTimerIdFromMetadata(TimerIndex);
+		}
+		return !ExcludedTimers.Contains(TimerIndex);
 	};
 }
 
diff --git a/Engine/Source/Developer/TraceInsights/Private/Insights/TimingProfiler/ViewModels/TimingExporter.h b/Engine/Source/Developer/TraceInsights/Private/Insights/TimingProfiler/ViewModels/TimingExporter.h
--- a/Engine/Source/Developer/TraceInsights/Private/Insights/TimingProfiler/ViewModels/TimingExporter.h
+++ b/Engine/Source/Developer/TraceInsights/Private/Insights/TimingProfiler/ViewModels/TimingExporter.h
@@ -20,6 +20,7 @@ class IFileHandle;
 namespace TraceServices
 {
 	class IAnalysisSession;
+	class ITimingProfilerTimerReader;
 	struct FTimingProfilerEvent;
 }
 
@@ -35,7 +36,11 @@ class FTimingExporter
 public:
 	typedef TFunction<bool(const FName /*ColumnId*/)> FColumnFilterFunc;
 	typedef TFunction<bool(uint32 /*ThreadId*/)> FThreadFilterFunc;
-	typedef TFunction<bool(double /*EventStartTime*/, double /*EventEndTime*/, uint32 /*EventDepth*/, const TraceServices::FTimingProfilerEvent& /*Event*/)> FTimingEventFilterFunc;
+
+	/**
+	* Use only under session read lock.
+	*/
+	typedef TFunction<bool(double /*EventStartTime*/, double /*EventEndTime*/, uint32 /*EventDepth*/, const TraceServices::FTimingProfilerEvent& /*Event*/, const TraceServices::ITimingProfilerTimerReader* /*TimerReader*/)> FTimingEventFilterFunc;
 
 	struct FExportThreadsParams
 	{
diff --git a/Engine/Source/Developer/TraceInsights/Private/Insights/TimingProfiler/Widgets/STimersView.cpp b/Engine/Source/Developer/TraceInsights/Private/Insights/TimingProfiler/Widgets/STimersView.cpp
--- a/Engine/Source/Developer/TraceInsights/Private/Insights/TimingProfiler/Widgets/STimersView.cpp
+++ b/Engine/Source/Developer/TraceInsights/Private/Insights/TimingProfiler/Widgets/STimersView.cpp
@@ -3141,6 +3141,20 @@ void STimersView::ContextMenu_ExportTimingEventsSelection_Execute() const
 		ExcludedThreads.Add(FGpuTimingTrack::Gpu1ThreadId);
 		ExcludedThreads.Add(FGpuTimingTrack::Gpu2ThreadId);
 
+		// Add Gpu Queues for the new Gpu Profiler
+		{
+			TraceServices::FAnalysisSessionReadScope SessionReadScope(*Session.Get());
+
+			const TraceServices::ITimingProfilerProvider* TimingProfilerProvider = TraceServices::ReadTimingProfilerProvider(*Session.Get());
+			if (TimingProfilerProvider)
+			{
+				TimingProfilerProvider->EnumerateGpuQueues([&ExcludedThreads](const TraceServices::FGpuQueueInfo& QueueInfo)
+					{
+						ExcludedThreads.Add(QueueInfo.Id + (1 << 16));
+					});
+			}
+		}
+
 		// Add the Verse Sampling fake thread to the ExcludedThreads list.
 		ExcludedThreads.Add(FVerseTimingTrack::VerseThreadId);
 
@@ -3158,11 +3172,20 @@ void STimersView::ContextMenu_ExportTimingEventsSelection_Execute() const
 		// Move the threads corresponding to visible CPU/GPU tracks to the IncludedThreads list.
 		TimingView->EnumerateAllTracks([&IncludedThreads, &ExcludedThreads](TSharedPtr<FBaseTimingTrack>& Track) -> bool
 			{
-				if (Track->IsVisible() && Track->Is<FThreadTimingTrack>())
+				if (Track->IsVisible())
 				{
-					const uint32 ThreadId = Track->As<FThreadTimingTrack>().GetThreadId();
-					ExcludedThreads.Remove(ThreadId);
-					IncludedThreads.Add(ThreadId);
+					if (Track->Is<FGpuQueueTimingTrack>())
+					{
+						const uint32 ThreadId = Track->As<FGpuQueueTimingTrack>().GetQueueId() + (1 << 16);
+						ExcludedThreads.Remove(ThreadId);
+						IncludedThreads.Add(ThreadId);
+					}
+					else if (Track->Is<FThreadTimingTrack>())
+					{
+						const uint32 ThreadId = Track->As<FThreadTimingTrack>().GetThreadId();
+						ExcludedThreads.Remove(ThreadId);
+						IncludedThreads.Add(ThreadId);
+					}
 				}
 				return true;
 			});
