Centralized Timer Color Generation.
Various Code Cleanups.

diff --git a/Engine/Source/Developer/TraceInsights/Private/Insights/MemoryProfiler/Widgets/SMemTagTreeView.cpp b/Engine/Source/Developer/TraceInsights/Private/Insights/MemoryProfiler/Widgets/SMemTagTreeView.cpp
--- a/Engine/Source/Developer/TraceInsights/Private/Insights/MemoryProfiler/Widgets/SMemTagTreeView.cpp
+++ b/Engine/Source/Developer/TraceInsights/Private/Insights/MemoryProfiler/Widgets/SMemTagTreeView.cpp
@@ -1925,8 +1925,7 @@ void SMemTagTreeView::GenerateColorForSelectedMemTags() const
 			if (SelectedNode->Is<FMemTagNode>())
 			{
 				FMemTagNode& MemTagNode = SelectedNode->As<FMemTagNode>();
-				constexpr bool bSetRandomColor = true;
-				SetColorToNode(MemTagNode, FLinearColor(), bSetRandomColor);
+				SetRandomColorToNode(MemTagNode);
 			}
 		}
 	}
@@ -1934,7 +1933,7 @@ void SMemTagTreeView::GenerateColorForSelectedMemTags() const
 
 ////////////////////////////////////////////////////////////////////////////////////////////////////
 
-void SMemTagTreeView::SetColorToNode(FMemTagNode& MemTagNode, FLinearColor Color, bool bSetRandomColor) const
+void SMemTagTreeView::SetColorToNode(FMemTagNode& MemTagNode, FLinearColor Color) const
 {
 	if (MemTagNode.IsGroup())
 	{
@@ -1944,7 +1943,7 @@ void SMemTagTreeView::SetColorToNode(FMemTagNode& MemTagNode, FLinearColor Color
 			if (Child->Is<FMemTagNode>())
 			{
 				FMemTagNode& ChildMemTagNode = Child->As<FMemTagNode>();
-				SetColorToNode(ChildMemTagNode, Color, bSetRandomColor);
+				SetColorToNode(ChildMemTagNode, Color);
 			}
 		}
 		return;
@@ -1955,17 +1954,66 @@ void SMemTagTreeView::SetColorToNode(FMemTagNode& MemTagNode, FLinearColor Color
 	{
 		return;
 	}
+	
+	MemTag->SetColor(Color);
+	const FLinearColor BorderColor(FMath::Min(Color.R + 0.4f, 1.0f), FMath::Min(Color.G + 0.4f, 1.0f), FMath::Min(Color.B + 0.4f, 1.0f), 1.0f);
 
-	if (bSetRandomColor)
+	const FMemoryTrackerId MemTrackerId = MemTagNode.GetMemTrackerId();
+	const FMemoryTagSetId MemTagSetId = MemTagNode.GetMemTagSetId();
+	const FMemoryTagId MemTagId = MemTagNode.GetMemTagId();
+
+	TSharedPtr<FMemoryGraphTrack> MainGraphTrack;
+	TSharedPtr<SMemoryProfilerWindow> ProfilerWindow = GetProfilerWindow();
+	if (ProfilerWindow.IsValid())
 	{
-		MemTag->SetRandomColor();
-		Color = MemTag->GetColor();
+		FMemorySharedState& SharedState = ProfilerWindow->GetSharedState();
+		MainGraphTrack = SharedState.GetMainGraphTrack();
 	}
-	else
+
+	for (const TSharedPtr<FMemoryGraphTrack>& GraphTrack : MemTag->GetGraphTracks())
+	{
+		TSharedPtr<FMemTagGraphSeries> MemTagSeries = GraphTrack->GetMemTagSeries(MemTrackerId, MemTagSetId, MemTagId);
+		if (MemTagSeries)
+		{
+			if (GraphTrack == MainGraphTrack)
+			{
+				MemTagSeries->SetColor(Color, BorderColor, Color.CopyWithNewOpacity(0.1f));
+			}
+			else
+			{
+				MemTagSeries->SetColor(Color, BorderColor);
+			}
+		}
+	}
+}
+
+////////////////////////////////////////////////////////////////////////////////////////////////////
+
+void SMemTagTreeView::SetRandomColorToNode(FMemTagNode& MemTagNode) const
+{
+	if (MemTagNode.IsGroup())
 	{
-		MemTag->SetColor(Color);
+		const TArray<FBaseTreeNodePtr>& Children = MemTagNode.GetFilteredChildren();
+		for (const FBaseTreeNodePtr& Child : Children)
+		{
+			if (Child->Is<FMemTagNode>())
+			{
+				FMemTagNode& ChildMemTagNode = Child->As<FMemTagNode>();
+				SetRandomColorToNode(ChildMemTagNode);
+			}
+		}
+		return;
+	}
+
+	FMemoryTag* MemTag = MemTagNode.GetMemTag();
+	if (!MemTag)
+	{
+		return;
 	}
 
+	MemTag->SetRandomColor();
+	FLinearColor Color = MemTag->GetColor();
+
 	const FLinearColor BorderColor(FMath::Min(Color.R + 0.4f, 1.0f), FMath::Min(Color.G + 0.4f, 1.0f), FMath::Min(Color.B + 0.4f, 1.0f), 1.0f);
 
 	const FMemoryTrackerId MemTrackerId = MemTagNode.GetMemTrackerId();
@@ -2021,7 +2069,7 @@ void SMemTagTreeView::EditColorForSelectedMemTags()
 			if (SelectedNode->Is<FMemTagNode>())
 			{
 				FMemTagNode& MemTagNode = SelectedNode->As<FMemTagNode>();
-				EditableColorValue = MemTagNode.GetColor();
+				EditableColorValue = MemTagNode.GetMemTag()->GetColor();
 				break;
 			}
 		}
@@ -2065,8 +2113,7 @@ void SMemTagTreeView::SetEditableColor(FLinearColor NewColor)
 		if (SelectedNode->Is<FMemTagNode>())
 		{
 			FMemTagNode& MemTagNode = SelectedNode->As<FMemTagNode>();
-			constexpr bool bSetRandomColor = false;
-			SetColorToNode(MemTagNode, EditableColorValue, bSetRandomColor);
+			SetColorToNode(MemTagNode, EditableColorValue);
 		}
 	}
 }
diff --git a/Engine/Source/Developer/TraceInsights/Private/Insights/MemoryProfiler/Widgets/SMemTagTreeView.h b/Engine/Source/Developer/TraceInsights/Private/Insights/MemoryProfiler/Widgets/SMemTagTreeView.h
--- a/Engine/Source/Developer/TraceInsights/Private/Insights/MemoryProfiler/Widgets/SMemTagTreeView.h
+++ b/Engine/Source/Developer/TraceInsights/Private/Insights/MemoryProfiler/Widgets/SMemTagTreeView.h
@@ -284,7 +284,8 @@ private:
 	bool CanEditColorForSelectedMemTags() const;
 	void EditColorForSelectedMemTags();
 
-	void SetColorToNode(FMemTagNode& MemTagNode, FLinearColor CustomColor, bool bSetRandomColor) const;
+	void SetColorToNode(FMemTagNode& MemTagNode, FLinearColor CustomColor) const;
+	void SetRandomColorToNode(FMemTagNode& MemTagNode) const;
 	FLinearColor GetEditableColor() const;
 	void SetEditableColor(FLinearColor NewColor);
 	void ColorPickerCancelled(FLinearColor OriginalColor);
diff --git a/Engine/Source/Developer/TraceInsights/Private/Insights/TimingProfiler/TimingProfilerManager.cpp b/Engine/Source/Developer/TraceInsights/Private/Insights/TimingProfiler/TimingProfilerManager.cpp
--- a/Engine/Source/Developer/TraceInsights/Private/Insights/TimingProfiler/TimingProfilerManager.cpp
+++ b/Engine/Source/Developer/TraceInsights/Private/Insights/TimingProfiler/TimingProfilerManager.cpp
@@ -587,27 +587,27 @@ void FTimingProfilerManager::UpdateCallersAndCallees()
 	{
 		TimerButterflyAggregator->Cancel();
 		TimerButterflyAggregator->SetTimeInterval(SelectionStartTime, SelectionEndTime);
-	}
 
-	TSharedPtr<STimingProfilerWindow> Wnd = GetProfilerWindow();
-	if (Wnd)
-	{
-		TSharedPtr<STimerTreeView> CallersTreeView = Wnd->GetCallersTreeView();
-		TSharedPtr<STimerTreeView> CalleesTreeView = Wnd->GetCalleesTreeView();
-
-		if (CallersTreeView)
+		TSharedPtr<STimingProfilerWindow> Wnd = GetProfilerWindow();
+		if (Wnd)
 		{
-			CallersTreeView->Reset();
-		}
+			TSharedPtr<STimerTreeView> CallersTreeView = Wnd->GetCallersTreeView();
+			TSharedPtr<STimerTreeView> CalleesTreeView = Wnd->GetCalleesTreeView();
 
-		if (CalleesTreeView)
-		{
-			CalleesTreeView->Reset();
-		}
+			if (CallersTreeView)
+			{
+				CallersTreeView->Reset();
+			}
 
-		if (CallersTreeView || CalleesTreeView)
-		{
-			TimerButterflyAggregator->Start();
+			if (CalleesTreeView)
+			{
+				CalleesTreeView->Reset();
+			}
+
+			if (CallersTreeView || CalleesTreeView)
+			{
+				TimerButterflyAggregator->Start();
+			}
 		}
 	}
 }
diff --git a/Engine/Source/Developer/TraceInsights/Private/Insights/TimingProfiler/Widgets/SStatsView.cpp b/Engine/Source/Developer/TraceInsights/Private/Insights/TimingProfiler/Widgets/SStatsView.cpp
--- a/Engine/Source/Developer/TraceInsights/Private/Insights/TimingProfiler/Widgets/SStatsView.cpp
+++ b/Engine/Source/Developer/TraceInsights/Private/Insights/TimingProfiler/Widgets/SStatsView.cpp
@@ -2251,7 +2251,7 @@ void SStatsView::ToggleTimingViewMainGraphEventSeries(FStatsNodePtr CounterNode)
 
 bool SStatsView::ContextMenu_GenerateNewColor_CanExecute() const
 {
-	return AreSelectedCounterNodesValid() && GetStatsView().IsValid();
+	return AreSelectedCounterNodesValid();
 }
 
 ////////////////////////////////////////////////////////////////////////////////////////////////////
@@ -2264,45 +2264,50 @@ void SStatsView::ContextMenu_GenerateNewColor_Execute() const
 		if (SelectedNode->Is<FStatsNode>())
 		{
 			FStatsNode& StatsNode = SelectedNode->As<FStatsNode>();
-			constexpr bool bSetRandomColor = true;
-			SetColorToNode(StatsNode, FLinearColor(), bSetRandomColor);
+			SetRandomColorToNode(StatsNode);
 		}
 	}
 };
 
 ////////////////////////////////////////////////////////////////////////////////////////////////////
 
-void SStatsView::SetColorToNode(FStatsNode& StatsNode, FLinearColor CustomColor, bool bSetRandomColor) const
+void SStatsView::SetColorToNode(FStatsNode& StatsNode, FLinearColor CustomColor) const
 {
-	if (bSetRandomColor)
-	{
-		uint64 Time = FPlatformTime::Cycles64();
-		uint32 Hash = (Time & 0xFFFFFFFF) ^ (Time >> 32);
-		Hash = ((Hash << 5) + Hash) * 0xfb23618f;
-
-		const uint8 H = Hash & 0xFF;
-		const uint8 S = 128 + ((Hash >> 8) & 0x7F);
-		const uint8 V = 128 + ((Hash >> 16) & 0x7F);
-		CustomColor = FLinearColor::MakeFromHSV8(H, S, V);
-		CustomColor.A = 1.0f;
+	StatsNode.SetColor(CustomColor);
 
-		StatsNode.SetColor(CustomColor);
-	}
-	else
+	const uint32 CounterId = StatsNode.GetCounterId();
+	const FLinearColor BorderColor(CustomColor.R + 0.4f, CustomColor.G + 0.4f, CustomColor.B + 0.4f, 1.0f);
+	
+	TSharedPtr<FTimingGraphSeries> CounterSeries = GetTimingViewMainGraphTrack()->GetStatsCounterSeries(CounterId);
+	if (CounterSeries.IsValid())
 	{
-		StatsNode.SetColor(CustomColor);
+		CounterSeries->SetColor(CustomColor, BorderColor);
 	}
+}
+
+////////////////////////////////////////////////////////////////////////////////////////////////////
+	
+void SStatsView::SetRandomColorToNode(FStatsNode& StatsNode) const
+{
+	uint64 Time = FPlatformTime::Cycles64();
+	uint32 Hash = (Time & 0xFFFFFFFF) ^ (Time >> 32);
+	Hash = ((Hash << 5) + Hash) * 0xfb23618f;
+
+	const uint8 H = Hash & 0xFF;
+	const uint8 S = 128 + ((Hash >> 8) & 0x7F);
+	const uint8 V = 128 + ((Hash >> 16) & 0x7F);
+	FLinearColor CustomColor = FLinearColor::MakeFromHSV8(H, S, V);
+	CustomColor.A = 1.0f;
+
+	StatsNode.SetColor(CustomColor);
 
 	const uint32 CounterId = StatsNode.GetCounterId();
 	const FLinearColor BorderColor(CustomColor.R + 0.4f, CustomColor.G + 0.4f, CustomColor.B + 0.4f, 1.0f);
-
-	if (GetStatsView().IsValid())
+	
+	TSharedPtr<FTimingGraphSeries> CounterSeries = GetTimingViewMainGraphTrack()->GetStatsCounterSeries(CounterId);
+	if (CounterSeries.IsValid())
 	{
-		TSharedPtr<FTimingGraphSeries> CounterSeries = GetTimingViewMainGraphTrack()->GetStatsCounterSeries(CounterId);
-		if (CounterSeries.IsValid())
-		{
-			CounterSeries->SetColor(CustomColor, BorderColor);
-		}
+		CounterSeries->SetColor(CustomColor, BorderColor);
 	}
 }
 
@@ -2310,41 +2315,38 @@ void SStatsView::SetColorToNode(FStatsNode& StatsNode, FLinearColor CustomColor,
 
 bool SStatsView::ContextMenu_EditColor_CanExecute() const
 {
-	return AreSelectedCounterNodesValid() && GetStatsView().IsValid();
+	return AreSelectedCounterNodesValid();
 };
 
 ////////////////////////////////////////////////////////////////////////////////////////////////////
 
 void SStatsView::ContextMenu_EditColor_Execute()
 {
-	if (GetStatsView().IsValid())
+	EditableColorValue = FLinearColor(0.5f, 0.5f, 0.5f, 1.0f);
+	const TArray<FStatsNodePtr> SelectedNodes = TreeView->GetSelectedItems();
+	for (const FStatsNodePtr& SelectedNode : SelectedNodes)
 	{
-		EditableColorValue = FLinearColor(0.5f, 0.5f, 0.5f, 1.0f);
-		const TArray<FStatsNodePtr> SelectedNodes = TreeView->GetSelectedItems();
-		for (const FStatsNodePtr& SelectedNode : SelectedNodes)
-		{
-			if (SelectedNode->Is<FStatsNode>())
-			{
-				FStatsNode& StatsNode = SelectedNode->As<FStatsNode>();
-				EditableColorValue = StatsNode.GetColor();
-				break;
-			}
-		}
-
-		FColorPickerArgs PickerArgs;
+		if (SelectedNode->Is<FStatsNode>())
 		{
-			PickerArgs.bUseAlpha = true;
-			PickerArgs.bOnlyRefreshOnMouseUp = false;
-			PickerArgs.bOnlyRefreshOnOk = false;
-			PickerArgs.bExpandAdvancedSection = false;
-			PickerArgs.OnColorCommitted = FOnLinearColorValueChanged::CreateSP(this, &SStatsView::SetEditableColor);
-			PickerArgs.OnColorPickerCancelled = FOnColorPickerCancelled::CreateSP(this, &SStatsView::ColorPickerCancelled);
-			PickerArgs.InitialColor = EditableColorValue;
-			PickerArgs.ParentWidget = SharedThis(this);
+			FStatsNode& StatsNode = SelectedNode->As<FStatsNode>();
+			EditableColorValue = StatsNode.GetColor();
+			break;
 		}
+	}
 
-		OpenColorPicker(PickerArgs);
+	FColorPickerArgs PickerArgs;
+	{
+		PickerArgs.bUseAlpha = true;
+		PickerArgs.bOnlyRefreshOnMouseUp = false;
+		PickerArgs.bOnlyRefreshOnOk = false;
+		PickerArgs.bExpandAdvancedSection = false;
+		PickerArgs.OnColorCommitted = FOnLinearColorValueChanged::CreateSP(this, &SStatsView::SetEditableColor);
+		PickerArgs.OnColorPickerCancelled = FOnColorPickerCancelled::CreateSP(this, &SStatsView::ColorPickerCancelled);
+		PickerArgs.InitialColor = EditableColorValue;
+		PickerArgs.ParentWidget = SharedThis(this);
 	}
+
+	OpenColorPicker(PickerArgs);
 };
 
 ////////////////////////////////////////////////////////////////////////////////////////////////////
@@ -2366,8 +2368,7 @@ void SStatsView::SetEditableColor(FLinearColor NewColor)
 		if (SelectedNode->Is<FStatsNode>())
 		{
 			FStatsNode& StatsNode = SelectedNode->As<FStatsNode>();
-			constexpr bool bSetRandomColor = false;
-			SetColorToNode(StatsNode, EditableColorValue, bSetRandomColor);
+			SetColorToNode(StatsNode, EditableColorValue);
 		}
 	}
 }
@@ -2790,14 +2791,6 @@ FReply SStatsView::OnKeyDown(const FGeometry& MyGeometry, const FKeyEvent& InKey
 
 ////////////////////////////////////////////////////////////////////////////////////////////////////
 
-TSharedPtr<SStatsView> SStatsView::GetStatsView() const
-{
-	TSharedPtr<STimingProfilerWindow> ProfilerWindow = FTimingProfilerManager::Get()->GetProfilerWindow();
-	return ProfilerWindow.IsValid() ? ProfilerWindow->GetStatsView() : nullptr;
-}
-
-////////////////////////////////////////////////////////////////////////////////////////////////////
-
 bool SStatsView::AreSelectedCounterNodesValid() const
 {
 	const TArray<FStatsNodePtr> SelectedNodes = TreeView->GetSelectedItems();
diff --git a/Engine/Source/Developer/TraceInsights/Private/Insights/TimingProfiler/Widgets/SStatsView.h b/Engine/Source/Developer/TraceInsights/Private/Insights/TimingProfiler/Widgets/SStatsView.h
--- a/Engine/Source/Developer/TraceInsights/Private/Insights/TimingProfiler/Widgets/SStatsView.h
+++ b/Engine/Source/Developer/TraceInsights/Private/Insights/TimingProfiler/Widgets/SStatsView.h
@@ -133,7 +133,8 @@ private:
 	bool ContextMenu_EditColor_CanExecute() const;
 	void ContextMenu_EditColor_Execute();
 
-	void SetColorToNode(FStatsNode& StatsNode, FLinearColor CustomColor, bool bSetRandomColor) const;
+	void SetColorToNode(FStatsNode& StatsNode, FLinearColor CustomColor) const;
+	void SetRandomColorToNode(FStatsNode& StatsNode) const;
 	FLinearColor GetEditableColor() const;
 	void SetEditableColor(FLinearColor NewColor);
 	void ColorPickerCancelled(FLinearColor OriginalColor);
@@ -319,8 +320,6 @@ private:
 	 */
 	virtual void Tick(const FGeometry& AllottedGeometry, const double InCurrentTime, const float InDeltaTime) override;
 
-	TSharedPtr<SStatsView> GetStatsView() const;
-
 private:
 	/** Table view model. */
 	TSharedPtr<FTable> Table;
diff --git a/Engine/Source/Developer/TraceInsights/Private/Insights/TimingProfiler/Widgets/STimerTreeView.cpp b/Engine/Source/Developer/TraceInsights/Private/Insights/TimingProfiler/Widgets/STimerTreeView.cpp
--- a/Engine/Source/Developer/TraceInsights/Private/Insights/TimingProfiler/Widgets/STimerTreeView.cpp
+++ b/Engine/Source/Developer/TraceInsights/Private/Insights/TimingProfiler/Widgets/STimerTreeView.cpp
@@ -1529,50 +1529,14 @@ FTimerNodePtr STimerTreeView::CreateTimerNodeRec(const TraceServices::FTimingPro
 
 	TSharedPtr<STimingProfilerWindow> ProfilerWindow = FTimingProfilerManager::Get()->GetProfilerWindow();
 	const bool bIsTimerAddedToGraphs = ProfilerWindow.IsValid() ? ProfilerWindow->IsTimerAddedToGraphs(Node.Timer->Id) : false;
-
+	FTimerNodePtr TimerNodePtr = MakeShared<FTimerNode>(Node.Timer->Id, Node.Timer->Name, Type, true);
 	FLinearColor ColorOverride;
-	if (GetTimingView().IsValid())
-	{
-		TSharedPtr<FTimingGraphTrack> GraphTrack = GetTimingViewMainGraphTrack();
-
-		if (GraphTrack.IsValid())
-		{
-			for (const TSharedPtr<FGraphSeries>& Series : GraphTrack->GetSeries())
-			{
-				if (Series->Is<FTimingGraphSeries>())
-				{
-					if (StaticCastSharedPtr<FTimingGraphSeries>(Series)->TimerId == Node.Timer->Id)
-					{
-						ColorOverride = Series->GetColor();
-						break;
-					}
-				}
-			}
-		}
-	}
 
-	if (GetFrameTrack().IsValid())
+	if (ProfilerWindow.IsValid() && ProfilerWindow->GetTimerColor(Node.Timer->Id, ColorOverride))
 	{
-		TSharedPtr<SFrameTrack> FrameTrack = GetFrameTrack();
-
-		if (FrameTrack.IsValid())
-		{
-			for (const TSharedPtr<FFrameTrackSeries>& Series : GetFrameTrack()->GetSeries())
-			{
-				if(Series->Type == EFrameTrackSeriesType::TimerFrameStats)
-				{
-					if (StaticCastSharedPtr<FTimerFrameStatsTrackSeries>(Series)->TimerId == Node.Timer->Id)
-					{
-						ColorOverride = Series->Color;
-						break;
-					}
-				}
-			}
-		}
+		TimerNodePtr->SetColor(ColorOverride);
 	}
-
-	FTimerNodePtr TimerNodePtr = MakeShared<FTimerNode>(Node.Timer->Id, Node.Timer->Name, Type, true);
-	TimerNodePtr->SetColor(ColorOverride);
+	
 	TimerNodePtr->SetAddedToGraphsFlag(bIsTimerAddedToGraphs);
 
 	TraceServices::FTimingProfilerAggregatedStats AggregatedStats;
diff --git a/Engine/Source/Developer/TraceInsights/Private/Insights/TimingProfiler/Widgets/STimerTreeView.h b/Engine/Source/Developer/TraceInsights/Private/Insights/TimingProfiler/Widgets/STimerTreeView.h
--- a/Engine/Source/Developer/TraceInsights/Private/Insights/TimingProfiler/Widgets/STimerTreeView.h
+++ b/Engine/Source/Developer/TraceInsights/Private/Insights/TimingProfiler/Widgets/STimerTreeView.h
@@ -65,7 +65,7 @@ public:
 
 	void SetTree(const TraceServices::FTimingProfilerButterflyNode& Root);
 
-	TArray<FTimerNodePtr> GetTreeNodes() const { return TreeNodes; };
+	TArray<FTimerNodePtr>& GetTreeNodes() { return TreeNodes; };
 
 	FTimerNodePtr GetTimerNode(uint32 TimerId) const;
 
diff --git a/Engine/Source/Developer/TraceInsights/Private/Insights/TimingProfiler/Widgets/STimersView.cpp b/Engine/Source/Developer/TraceInsights/Private/Insights/TimingProfiler/Widgets/STimersView.cpp
--- a/Engine/Source/Developer/TraceInsights/Private/Insights/TimingProfiler/Widgets/STimersView.cpp
+++ b/Engine/Source/Developer/TraceInsights/Private/Insights/TimingProfiler/Widgets/STimersView.cpp
@@ -2531,6 +2531,13 @@ void STimersView::RebuildTree(bool bResync)
 					(GraphTrack.IsValid() && GraphTrack->HasAnySeriesForTimer(Timer.Id));
 				TimerNodePtr->SetAddedToGraphsFlag(bIsTimerAddedToGraphs);
 
+				FLinearColor ColorOverride;
+				TSharedPtr<STimingProfilerWindow> ProfilerWindow = FTimingProfilerManager::Get()->GetProfilerWindow();
+				if (ProfilerWindow.IsValid() && ProfilerWindow->GetTimerColor(Timer.Id, ColorOverride))
+				{
+					TimerNodePtr->SetColor(ColorOverride);
+				}
+
 				TimerNodes.Add(TimerNodePtr);
 			}
 			ensure(TimerNodes.Num() == TimerCount);
@@ -2755,7 +2762,6 @@ void STimersView::ToggleGraphInstanceSeries(TSharedRef<FTimingGraphTrack> GraphT
 	if (ProfilerWindow.IsValid())
 	{
 		ProfilerWindow->OnTimerAddedToGraphsChanged(TimerId);
-		FTimingProfilerManager::Get()->UpdateCallersAndCallees();
 	}
 }
 
@@ -2814,7 +2820,6 @@ void STimersView::ToggleGraphFrameStatsSeries(TSharedRef<FTimingGraphTrack> Grap
 	if (ProfilerWindow.IsValid())
 	{
 		ProfilerWindow->OnTimerAddedToGraphsChanged(TimerId);
-		FTimingProfilerManager::Get()->UpdateCallersAndCallees();
 	}
 }
 
@@ -2889,7 +2894,6 @@ void STimersView::ToggleFrameTrackSeries(FTimerNodePtr TimerNode, ETraceFrameTyp
 	if (ProfilerWindow.IsValid())
 	{
 		ProfilerWindow->OnTimerAddedToGraphsChanged(TimerId);
-		FTimingProfilerManager::Get()->UpdateCallersAndCallees();
 	}
 }
 
@@ -3449,7 +3453,7 @@ void STimersView::OpenSourceFileInIDE(FTimerNodePtr InNode) const
 
 bool STimersView::ContextMenu_GenerateNewColor_CanExecute() const
 {
-	return AreSelectedTimerNodesValid() && GetTimingView().IsValid();
+	return AreSelectedTimerNodesValid();
 }
 
 ////////////////////////////////////////////////////////////////////////////////////////////////////
@@ -3462,119 +3466,88 @@ void STimersView::ContextMenu_GenerateNewColor_Execute() const
 		if (SelectedNode->Is<FTimerNode>())
 		{
 			FTimerNode& TimerNode = SelectedNode->As<FTimerNode>();
-			constexpr bool bSetRandomColor = true;
-			SetColorToNode(TimerNode, FLinearColor(), bSetRandomColor);
+			SetRandomColorToNode(TimerNode);
 		}
 	}
 };
 
 ////////////////////////////////////////////////////////////////////////////////////////////////////
 
-void STimersView::SetColorToNode(FTimerNode& TimerNode, FLinearColor CustomColor, bool bSetRandomColor) const
+void STimersView::SetColorToNode(FTimerNode& TimerNode, FLinearColor CustomColor) const
 {
-	if (bSetRandomColor)
+	TimerNode.SetColor(CustomColor);
+	if (FTimingProfilerManager::Get().IsValid())
 	{
-		uint64 Time = FPlatformTime::Cycles64();
-		uint32 Hash = (Time & 0xFFFFFFFF) ^ (Time >> 32);
-		Hash = ((Hash << 5) + Hash) * 0xfb23618f;
-
-		const uint8 H = Hash & 0xFF;
-		const uint8 S = 128 + ((Hash >> 8) & 0x7F);
-		const uint8 V = 128 + ((Hash >> 16) & 0x7F);
-		CustomColor = FLinearColor::MakeFromHSV8(H, S, V);
-		CustomColor.A = 1.0f;
-
-		TimerNode.SetColor(CustomColor);
-	}
-	else
-	{
-		TimerNode.SetColor(CustomColor);
-	}
-
-	const uint32 TimerId = TimerNode.GetTimerId();
-	const FLinearColor BorderColor(CustomColor.R + 0.4f, CustomColor.G + 0.4f, CustomColor.B + 0.4f, 1.0f);
-
-	if (GetTimingView().IsValid())
-	{
-		TSharedPtr<FTimingGraphTrack> GraphTrack = GetTimingViewMainGraphTrack();
-
-		if (GraphTrack.IsValid())
+		TSharedPtr<STimingProfilerWindow> ProfilerWindow = FTimingProfilerManager::Get()->GetProfilerWindow();
+		if (ProfilerWindow.IsValid())
 		{
-			for (const TSharedPtr<FGraphSeries>& Series : GraphTrack->GetSeries())
-			{
-				if (Series->Is<FTimingGraphSeries>())
-				{
-					if (StaticCastSharedPtr<FTimingGraphSeries>(Series)->TimerId == TimerId)
-					{
-						Series->SetColor(CustomColor, BorderColor);
-					}
-				}
-			}
+			ProfilerWindow->SetTimerColor(TimerNode.GetTimerId(), CustomColor);
 		}
 	}
+}
 
-	if (GetFrameTrack().IsValid())
-	{
-		TSharedPtr<SFrameTrack> FrameTrack = GetFrameTrack();
+////////////////////////////////////////////////////////////////////////////////////////////////////
+	
+void STimersView::SetRandomColorToNode(FTimerNode& TimerNode) const
+{
+	uint64 Time = FPlatformTime::Cycles64();
+	uint32 Hash = (Time & 0xFFFFFFFF) ^ (Time >> 32);
+	Hash = ((Hash << 5) + Hash) * 0xfb23618f;
+
+	const uint8 H = Hash & 0xFF;
+	const uint8 S = 128 + ((Hash >> 8) & 0x7F);
+	const uint8 V = 128 + ((Hash >> 16) & 0x7F);
+	FLinearColor CustomColor = FLinearColor::MakeFromHSV8(H, S, V);
+	CustomColor.A = 1.0f;
 
-		if (FrameTrack.IsValid())
+	TimerNode.SetColor(CustomColor);
+	if (FTimingProfilerManager::Get().IsValid())
+	{
+		TSharedPtr<STimingProfilerWindow> ProfilerWindow = FTimingProfilerManager::Get()->GetProfilerWindow();
+		if (ProfilerWindow.IsValid())
 		{
-			for (const TSharedPtr<FFrameTrackSeries>& Series : GetFrameTrack()->GetSeries())
-			{
-				if (Series->Type == EFrameTrackSeriesType::TimerFrameStats)
-				{
-					if (StaticCastSharedPtr<FTimerFrameStatsTrackSeries>(Series)->TimerId == TimerId)
-					{
-						Series->Color = CustomColor;
-					}
-				}
-			}
+			ProfilerWindow->SetTimerColor(TimerNode.GetTimerId(), CustomColor);
 		}
 	}
-
-	FTimingProfilerManager::Get()->UpdateCallersAndCallees();
 }
 
 ////////////////////////////////////////////////////////////////////////////////////////////////////
 
 bool STimersView::ContextMenu_EditColor_CanExecute() const
 {
-	return AreSelectedTimerNodesValid() && GetTimingView().IsValid();
-};
+	return AreSelectedTimerNodesValid();
+}
 
 ////////////////////////////////////////////////////////////////////////////////////////////////////
 
 void STimersView::ContextMenu_EditColor_Execute()
 {
-	if (GetTimingView().IsValid())
+	EditableColorValue = FLinearColor(0.5f, 0.5f, 0.5f, 1.0f);
+	const TArray<FTimerNodePtr> SelectedNodes = TreeView->GetSelectedItems();
+	for (const FTimerNodePtr& SelectedNode : SelectedNodes)
 	{
-		EditableColorValue = FLinearColor(0.5f, 0.5f, 0.5f, 1.0f);
-		const TArray<FTimerNodePtr> SelectedNodes = TreeView->GetSelectedItems();
-		for (const FTimerNodePtr& SelectedNode : SelectedNodes)
-		{
-			if (SelectedNode->Is<FTimerNode>())
-			{
-				FTimerNode& TimerNode = SelectedNode->As<FTimerNode>();
-				EditableColorValue = TimerNode.GetColor();
-				break;
-			}
-		}
-
-		FColorPickerArgs PickerArgs;
+		if (SelectedNode->Is<FTimerNode>())
 		{
-			PickerArgs.bUseAlpha = true;
-			PickerArgs.bOnlyRefreshOnMouseUp = false;
-			PickerArgs.bOnlyRefreshOnOk = false;
-			PickerArgs.bExpandAdvancedSection = false;
-			PickerArgs.OnColorCommitted = FOnLinearColorValueChanged::CreateSP(this, &STimersView::SetEditableColor);
-			PickerArgs.OnColorPickerCancelled = FOnColorPickerCancelled::CreateSP(this, &STimersView::ColorPickerCancelled);
-			PickerArgs.InitialColor = EditableColorValue;
-			PickerArgs.ParentWidget = SharedThis(this);
+			FTimerNode& TimerNode = SelectedNode->As<FTimerNode>();
+			EditableColorValue = TimerNode.GetColor();
+			break;
 		}
+	}
 
-		OpenColorPicker(PickerArgs);
+	FColorPickerArgs PickerArgs;
+	{
+		PickerArgs.bUseAlpha = true;
+		PickerArgs.bOnlyRefreshOnMouseUp = false;
+		PickerArgs.bOnlyRefreshOnOk = false;
+		PickerArgs.bExpandAdvancedSection = false;
+		PickerArgs.OnColorCommitted = FOnLinearColorValueChanged::CreateSP(this, &STimersView::SetEditableColor);
+		PickerArgs.OnColorPickerCancelled = FOnColorPickerCancelled::CreateSP(this, &STimersView::ColorPickerCancelled);
+		PickerArgs.InitialColor = EditableColorValue;
+		PickerArgs.ParentWidget = SharedThis(this);
 	}
-};
+
+	OpenColorPicker(PickerArgs);
+}
 
 ////////////////////////////////////////////////////////////////////////////////////////////////////
 
@@ -3595,8 +3568,7 @@ void STimersView::SetEditableColor(FLinearColor NewColor)
 		if (SelectedNode->Is<FTimerNode>())
 		{
 			FTimerNode& TimerNode = SelectedNode->As<FTimerNode>();
-			constexpr bool bSetRandomColor = false;
-			SetColorToNode(TimerNode, EditableColorValue, bSetRandomColor);
+			SetColorToNode(TimerNode, EditableColorValue);
 		}
 	}
 }
@@ -3843,6 +3815,7 @@ bool STimersView::AreSelectedTimerNodesValid() const
 }
 
 ////////////////////////////////////////////////////////////////////////////////////////////////////
+
 } // namespace UE::Insights::TimingProfiler
 
 #undef LOCTEXT_NAMESPACE
diff --git a/Engine/Source/Developer/TraceInsights/Private/Insights/TimingProfiler/Widgets/STimersView.h b/Engine/Source/Developer/TraceInsights/Private/Insights/TimingProfiler/Widgets/STimersView.h
--- a/Engine/Source/Developer/TraceInsights/Private/Insights/TimingProfiler/Widgets/STimersView.h
+++ b/Engine/Source/Developer/TraceInsights/Private/Insights/TimingProfiler/Widgets/STimersView.h
@@ -177,7 +177,9 @@ private:
 	bool ContextMenu_EditColor_CanExecute() const;
 	void ContextMenu_EditColor_Execute();
 
-	void SetColorToNode(FTimerNode& TimerNode, FLinearColor CustomColor, bool bSetRandomColor) const;
+	void SetColorToNode(FTimerNode& TimerNode, FLinearColor CustomColor) const;
+	void SetRandomColorToNode(FTimerNode& TimerNode) const;
+	
 	FLinearColor GetEditableColor() const;
 	void SetEditableColor(FLinearColor NewColor);
 	void ColorPickerCancelled(FLinearColor OriginalColor);
diff --git a/Engine/Source/Developer/TraceInsights/Private/Insights/TimingProfiler/Widgets/STimingProfilerWindow.cpp b/Engine/Source/Developer/TraceInsights/Private/Insights/TimingProfiler/Widgets/STimingProfilerWindow.cpp
--- a/Engine/Source/Developer/TraceInsights/Private/Insights/TimingProfiler/Widgets/STimingProfilerWindow.cpp
+++ b/Engine/Source/Developer/TraceInsights/Private/Insights/TimingProfiler/Widgets/STimingProfilerWindow.cpp
@@ -108,6 +108,8 @@ void STimingProfilerWindow::Reset()
 	{
 		LogView->Reset();
 	}
+
+	TimerColors.Reset();
 }
 
 ////////////////////////////////////////////////////////////////////////////////////////////////////
@@ -517,6 +519,73 @@ void STimingProfilerWindow::OnTimeSelectionChanged(Timing::ETimeChangedFlags InF
 
 ////////////////////////////////////////////////////////////////////////////////////////////////////
 
+bool STimingProfilerWindow::GetTimerColor(uint32 TimerId, FLinearColor& OutColor)
+{
+	if (TimerColors.Contains(TimerId))
+	{
+		OutColor = *TimerColors.Find(TimerId);
+		return true;
+	}
+	
+	return false;
+}
+
+////////////////////////////////////////////////////////////////////////////////////////////////////
+
+void STimingProfilerWindow::SetTimerColor(uint32 TimerId, FLinearColor Color)
+{
+	TimerColors.Add(TimerId, Color);
+	const FLinearColor BorderColor(Color.R + 0.4f, Color.G + 0.4f, Color.B + 0.4f, 1.0f);
+	
+	if (TimingView.IsValid())
+	{
+		for (const TSharedPtr<FGraphSeries>& Series : TimingView->GetMainTimingGraphTrack()->GetSeries())
+		{
+			if (Series->Is<FTimingGraphSeries>())
+			{
+				if (Series->As<FTimingGraphSeries>().TimerId == TimerId)
+				{
+					Series->SetColor(Color, BorderColor);
+				}
+			}
+		}
+	}
+
+	if (FrameTrack.IsValid())
+	{
+		for (const TSharedPtr<FFrameTrackSeries>& Series : FrameTrack->GetSeries())
+		{
+			if (Series->Type == EFrameTrackSeriesType::TimerFrameStats)
+			{
+				if (StaticCastSharedPtr<FTimerFrameStatsTrackSeries>(Series)->TimerId == TimerId)
+				{
+					Series->Color = Color;
+				}
+			}
+		}
+	}
+
+	if (CallersTreeView.IsValid())
+	{
+		FTimerNodePtr TimerNode = CallersTreeView->GetTimerNode(TimerId);
+		if (TimerNode.IsValid())
+		{
+			TimerNode->SetColor(Color);
+		}
+	}
+
+	if (CalleesTreeView.IsValid())
+	{
+		FTimerNodePtr TimerNode = CalleesTreeView->GetTimerNode(TimerId);
+		if (TimerNode.IsValid())
+		{
+			TimerNode->SetColor(Color);
+		}
+	}
+}
+
+////////////////////////////////////////////////////////////////////////////////////////////////////
+
 bool STimingProfilerWindow::IsTimerAddedToGraphs(uint32 TimerId) const
 {
 	if (FrameTrack.IsValid() &&
@@ -543,6 +612,7 @@ bool STimingProfilerWindow::IsTimerAddedToGraphs(uint32 TimerId) const
 void STimingProfilerWindow::OnTimerAddedToGraphsChanged(uint32 TimerId)
 {
 	const bool bIsAddedToGraphs = IsTimerAddedToGraphs(TimerId);
+	FLinearColor ColorOverride;
 
 	if (TimersView.IsValid())
 	{
@@ -559,6 +629,10 @@ void STimingProfilerWindow::OnTimerAddedToGraphsChanged(uint32 TimerId)
 		if (TimerNode.IsValid())
 		{
 			TimerNode->SetAddedToGraphsFlag(bIsAddedToGraphs);
+			if (GetTimerColor(TimerId, ColorOverride))
+			{
+				TimerNode->SetColor(ColorOverride);
+			}
 		}
 	}
 
@@ -568,6 +642,10 @@ void STimingProfilerWindow::OnTimerAddedToGraphsChanged(uint32 TimerId)
 		if (TimerNode.IsValid())
 		{
 			TimerNode->SetAddedToGraphsFlag(bIsAddedToGraphs);
+			if (GetTimerColor(TimerId, ColorOverride))
+			{
+				TimerNode->SetColor(ColorOverride);
+			}
 		}
 	}
 }
diff --git a/Engine/Source/Developer/TraceInsights/Private/Insights/TimingProfiler/Widgets/STimingProfilerWindow.h b/Engine/Source/Developer/TraceInsights/Private/Insights/TimingProfiler/Widgets/STimingProfilerWindow.h
--- a/Engine/Source/Developer/TraceInsights/Private/Insights/TimingProfiler/Widgets/STimingProfilerWindow.h
+++ b/Engine/Source/Developer/TraceInsights/Private/Insights/TimingProfiler/Widgets/STimingProfilerWindow.h
@@ -45,6 +45,11 @@ public:
 	TSharedPtr<SLogView> GetLogView() const { return LogView; }
 	TSharedPtr<SFrameTrack> GetFrameView() const { return FrameTrack; }
 
+	/** Gets the color of a given timer, returns true if successful, false if the timer was not found. */
+	bool GetTimerColor(uint32 TimerId, FLinearColor& OutColor);
+	/** Updates the TMap and propagates the color to any relevant graphs and nodes. */
+	void SetTimerColor(uint32 TimerId, FLinearColor Color);
+	
 	bool IsTimerAddedToGraphs(uint32 TimerId) const;
 	void OnTimerAddedToGraphsChanged(uint32 TimerId);
 
@@ -100,6 +105,10 @@ private:
 
 	/** The Log view widget */
 	TSharedPtr<SLogView> LogView;
+
+	/** Color-TimerId Mapping for easy pairing */
+	TMap<uint32, FLinearColor> TimerColors;
+	
 };
 
 } // namespace UE::Insights::TimingProfiler
