Fix copy to clipboard for timing events with static metadata.

diff --git a/Engine/Source/Developer/TraceInsights/Private/Insights/TimingProfiler/Tracks/ThreadTimingTrack.cpp b/Engine/Source/Developer/TraceInsights/Private/Insights/TimingProfiler/Tracks/ThreadTimingTrack.cpp
--- a/Engine/Source/Developer/TraceInsights/Private/Insights/TimingProfiler/Tracks/ThreadTimingTrack.cpp
+++ b/Engine/Source/Developer/TraceInsights/Private/Insights/TimingProfiler/Tracks/ThreadTimingTrack.cpp
@@ -1150,7 +1150,14 @@ void FThreadTimingTrackImpl::OnClipboardCopyEvent(const ITimingEvent& InSelected
 				}
 				if (Metadata.Num() > 0)
 				{
-					AppendMetadataToString(EventName, MetadataSpec, Metadata);
+					if (MetadataSpec)
+					{
+						AppendMetadataToString(EventName, MetadataSpec, Metadata);
+					}
+					else
+					{
+						AppendMetadataToString(EventName, Metadata);
+					}
 				}
 			}
 
