Refactored FFrameTrackSeries and FTimerFrameStatsTrackSeries to use RTTI

diff --git a/Engine/Source/Developer/TraceInsights/Private/Insights/TimingProfiler/ViewModels/FrameTrackHelper.cpp b/Engine/Source/Developer/TraceInsights/Private/Insights/TimingProfiler/ViewModels/FrameTrackHelper.cpp
--- a/Engine/Source/Developer/TraceInsights/Private/Insights/TimingProfiler/ViewModels/FrameTrackHelper.cpp
+++ b/Engine/Source/Developer/TraceInsights/Private/Insights/TimingProfiler/ViewModels/FrameTrackHelper.cpp
@@ -25,6 +25,15 @@
 namespace UE::Insights::TimingProfiler
 {
 
+///////////////////////////////////////////////////////////////////////////////////////////////////
+// FFrameTrackSeries / FTimerFrameStatsTrackSeries
+////////////////////////////////////////////////////////////////////////////////////////////////////
+
+INSIGHTS_IMPLEMENT_RTTI(FFrameTrackSeries)
+INSIGHTS_IMPLEMENT_RTTI(FTimerFrameStatsTrackSeries)
+	
+////////////////////////////////////////////////////////////////////////////////////////////////////
+	
 ////////////////////////////////////////////////////////////////////////////////////////////////////
 // FFrameTrackSeriesBuilder
 ////////////////////////////////////////////////////////////////////////////////////////////////////
@@ -163,17 +172,17 @@ FLinearColor FFrameTrackDrawHelper::GetColorByFrameType(int32 FrameType)
 
 void FFrameTrackDrawHelper::DrawCached(const FFrameTrackSeries& Series) const
 {
-	if (Series.NumAggregatedFrames == 0)
+	if (Series.GetNumAggregatedFrames() == 0)
 	{
 		return;
 	}
 
-	NumFrames += Series.NumAggregatedFrames;
+	NumFrames += Series.GetNumAggregatedFrames();
 
-	FLinearColor SeriesColor = Series.Color;
+	FLinearColor SeriesColor = Series.GetColor();
 
 	const float SampleW = Viewport.GetSampleWidth();
-	const int32 NumSamples = Series.Samples.Num();
+	const int32 NumSamples = Series.GetNumSamples();
 
 	const FAxisViewportDouble& ViewportY = Viewport.GetVerticalAxisViewport();
 
@@ -182,7 +191,7 @@ void FFrameTrackDrawHelper::DrawCached(const FFrameTrackSeries& Series) const
 
 	for (int32 SampleIndex = 0; SampleIndex < NumSamples; SampleIndex++)
 	{
-		const FFrameTrackSample& Sample = Series.Samples[SampleIndex];
+		const FFrameTrackSample& Sample = Series.GetSample(SampleIndex);
 		if (Sample.NumFrames == 0)
 		{
 			continue;
@@ -281,16 +290,16 @@ void FFrameTrackDrawHelper::DrawHoveredSample(const FFrameTrackSample& Sample) c
 
 void FFrameTrackDrawHelper::DrawHighlightedInterval(const FFrameTrackSeries& Series, const double StartTime, const double EndTime) const
 {
-	const int32 NumSamples = Series.Samples.Num();
+	const int32 NumSamples = Series.GetNumSamples();
 
 	//TODO: binary search
 	int32 Index1 = 0;
 	int32 Index2 = NumSamples - 1;
-	while (Index1 < NumSamples && Series.Samples[Index1].EndTime < StartTime)
+	while (Index1 < NumSamples && Series.GetSample(Index1).EndTime < StartTime)
 	{
 		Index1++;
 	}
-	while (Index2 >= Index1 && Series.Samples[Index2].StartTime > EndTime)
+	while (Index2 >= Index1 && Series.GetSample(Index2).StartTime > EndTime)
 	{
 		Index2--;
 	}
diff --git a/Engine/Source/Developer/TraceInsights/Private/Insights/TimingProfiler/ViewModels/FrameTrackHelper.h b/Engine/Source/Developer/TraceInsights/Private/Insights/TimingProfiler/ViewModels/FrameTrackHelper.h
--- a/Engine/Source/Developer/TraceInsights/Private/Insights/TimingProfiler/ViewModels/FrameTrackHelper.h
+++ b/Engine/Source/Developer/TraceInsights/Private/Insights/TimingProfiler/ViewModels/FrameTrackHelper.h
@@ -4,6 +4,9 @@
 
 #include "CoreTypes.h"
 
+// TraceInsightsCore
+#include "InsightsCore/Common/SimpleRtti.h"
+
 #include "Styling/WidgetStyle.h"
 
 ////////////////////////////////////////////////////////////////////////////////////////////////////
@@ -72,50 +75,78 @@ struct FFrameTrackSample
 
 ////////////////////////////////////////////////////////////////////////////////////////////////////
 
-enum class EFrameTrackSeriesType : uint32
-{
-	Frame,
-	TimerFrameStats
-};
-
-////////////////////////////////////////////////////////////////////////////////////////////////////
-
-struct FFrameTrackSeries
+class FFrameTrackSeries
 {
-	int32 FrameType;
-	EFrameTrackSeriesType Type;
-	bool bIsVisible;
-	int32 NumAggregatedFrames; // total number of frames aggregated in samples; i.e. sum of all Sample.NumFrames
-	TArray<FFrameTrackSample> Samples; // the aggregated samples
-	FLinearColor Color;
-	FText Name;
-
-	explicit FFrameTrackSeries(int32 InFrameType, EFrameTrackSeriesType InType)
+	friend class FFrameTrackSeriesBuilder;
+	
+	INSIGHTS_DECLARE_RTTI_BASE(FFrameTrackSeries)
+	
+public:
+	explicit FFrameTrackSeries(int32 InFrameType)
 		: FrameType(InFrameType)
-		, Type(InType)
 		, bIsVisible(true)
 		, NumAggregatedFrames(0)
 		, Samples()
 	{
 	}
+	virtual ~FFrameTrackSeries() = default;
+	
+	int32 GetFrameType() const { return FrameType; }
+	
+	bool IsVisible() const { return bIsVisible; }
+	void SetVisibility(bool bOnOff) { bIsVisible = bOnOff; }
+	
+	/** Get the number of aggregated frames */
+	int32 GetNumAggregatedFrames() const { return NumAggregatedFrames; }
+	/** Set the number of aggregated frames */
+	void SetNumAggregatedFrames(int32 InNumAggregatedFrames) { NumAggregatedFrames = InNumAggregatedFrames; }
+
+	/** Get the number of Samples */
+	int32 GetNumSamples() const { return Samples.Num(); }
+	/** Access a single Sample by an index */
+	FFrameTrackSample& GetSample(int32 InIndex) { return Samples[InIndex]; }
+	const FFrameTrackSample& GetSample(int32 InIndex) const { return Samples[InIndex]; }
+	/** Returns the Samples Array by Reference */
+	TArray<FFrameTrackSample>& GetSamples() { return Samples; }
+
+	FLinearColor GetColor() const { return Color; }
+	void SetColor(const FLinearColor& InColor) { Color = InColor; }
+
+	FText GetName() const { return Name; }
+	void SetName(const FText& InName) { Name = InName; }
+
+	void Reset() { NumAggregatedFrames = 0; Samples.Reset(); }
 
-	void Reset()
-	{
-		NumAggregatedFrames = 0;
-		Samples.Reset();
-	}
+private:
+	int32 FrameType;
+	bool bIsVisible;
+	
+	/** Total number of frames aggregated in samples; i.e. sum of all Sample.NumFrames */
+	int32 NumAggregatedFrames;
+	/** The aggregated samples */
+	TArray<FFrameTrackSample> Samples;
+	
+	FLinearColor Color;
+	FText Name;
 };
 
 ////////////////////////////////////////////////////////////////////////////////////////////////////
 
-struct FTimerFrameStatsTrackSeries : public FFrameTrackSeries
+class FTimerFrameStatsTrackSeries : public FFrameTrackSeries
 {
+	INSIGHTS_DECLARE_RTTI(FTimerFrameStatsTrackSeries, FFrameTrackSeries)
+	
+public:
 	FTimerFrameStatsTrackSeries(int32 InFrameType, uint32 InTimerId)
-		: FFrameTrackSeries(InFrameType, EFrameTrackSeriesType::TimerFrameStats)
+		: FFrameTrackSeries(InFrameType)
 		, TimerId(InTimerId)
 	{
 	}
+	virtual ~FTimerFrameStatsTrackSeries() = default;
 
+	uint32 GetTimerId() const { return TimerId; }
+
+private:
 	uint32 TimerId;
 };
 
diff --git a/Engine/Source/Developer/TraceInsights/Private/Insights/TimingProfiler/Widgets/SFrameTrack.cpp b/Engine/Source/Developer/TraceInsights/Private/Insights/TimingProfiler/Widgets/SFrameTrack.cpp
--- a/Engine/Source/Developer/TraceInsights/Private/Insights/TimingProfiler/Widgets/SFrameTrack.cpp
+++ b/Engine/Source/Developer/TraceInsights/Private/Insights/TimingProfiler/Widgets/SFrameTrack.cpp
@@ -78,19 +78,19 @@ SFrameTrack::~SFrameTrack()
 
 	if (ProfilerWindow.IsValid())
 	{
-		TArray<TSharedPtr<FTimerFrameStatsTrackSeries>> TimerSeries;
+		TArray<uint32> TimerIds;
 		for (TSharedPtr<FFrameTrackSeries>& Series : AllSeries)
 		{
 			if (Series.IsValid() &&
-				Series->Type == EFrameTrackSeriesType::TimerFrameStats)
+				Series->Is<FTimerFrameStatsTrackSeries>())
 			{
-				TimerSeries.Add(StaticCastSharedPtr<FTimerFrameStatsTrackSeries>(Series));
+				TimerIds.Add(Series->As<FTimerFrameStatsTrackSeries>().GetTimerId());
 			}
 		}
 		AllSeries.Reset();
-		for (TSharedPtr<FTimerFrameStatsTrackSeries>& Series : TimerSeries)
+		for (uint32 TimerId : TimerIds)
 		{
-			ProfilerWindow->OnTimerAddedToGraphsChanged(Series->TimerId);
+			ProfilerWindow->OnTimerAddedToGraphsChanged(TimerId);
 		}
 	}
 }
@@ -314,8 +314,8 @@ TSharedRef<FFrameTrackSeries> SFrameTrack::FindOrAddSeries(ETraceFrameType Frame
 {
 	TSharedPtr<FFrameTrackSeries>* ExistingSeries = AllSeries.FindByPredicate([FrameType](TSharedPtr<FFrameTrackSeries> Series)
 		{
-			return Series->Type == EFrameTrackSeriesType::Frame &&
-				Series->FrameType == FrameType;
+			return !Series->Is<FTimerFrameStatsTrackSeries>() &&
+				Series->GetFrameType() == FrameType;
 		});
 
 	if (ExistingSeries != nullptr)
@@ -325,9 +325,9 @@ TSharedRef<FFrameTrackSeries> SFrameTrack::FindOrAddSeries(ETraceFrameType Frame
 
 	LLM_SCOPE_BYTAG(Insights);
 
-	TSharedRef<FFrameTrackSeries> SeriesRef = MakeShared<FFrameTrackSeries>(FrameType, EFrameTrackSeriesType::Frame);
-	SeriesRef->Color = FFrameTrackDrawHelper::GetColorByFrameType(FrameType);
-	SeriesRef->Name = FFrameTrackDrawHelper::FrameTypeToText(FrameType);
+	TSharedRef<FFrameTrackSeries> SeriesRef = MakeShared<FFrameTrackSeries>(FrameType);
+	SeriesRef->SetColor(FFrameTrackDrawHelper::GetColorByFrameType(FrameType));
+	SeriesRef->SetName(FFrameTrackDrawHelper::FrameTypeToText(FrameType));
 	AllSeries.Add(SeriesRef);
 	return SeriesRef;
 }
@@ -338,8 +338,8 @@ TSharedPtr<FFrameTrackSeries> SFrameTrack::FindSeries(ETraceFrameType FrameType)
 {
 	const TSharedPtr<FFrameTrackSeries>* ExistingSeries = AllSeries.FindByPredicate([FrameType](TSharedPtr<FFrameTrackSeries> Series)
 		{
-			return Series->Type == EFrameTrackSeriesType::Frame &&
-				Series->FrameType == FrameType;
+			return !Series->Is<FTimerFrameStatsTrackSeries>() &&
+				Series->GetFrameType() == FrameType;
 		});
 
 	if (ExistingSeries != nullptr)
@@ -356,9 +356,9 @@ TSharedPtr<FFrameTrackSeries> SFrameTrack::FindFrameStatsSeries(ETraceFrameType
 {
 	const TSharedPtr<FFrameTrackSeries>* ExistingSeries = AllSeries.FindByPredicate([FrameType, TimerId](TSharedPtr<FFrameTrackSeries> Series)
 		{
-			return Series->Type == EFrameTrackSeriesType::TimerFrameStats &&
-				Series->FrameType == FrameType &&
-				StaticCastSharedPtr<FTimerFrameStatsTrackSeries>(Series)->TimerId == TimerId;
+			return Series->Is<FTimerFrameStatsTrackSeries>() &&
+					Series->As<FTimerFrameStatsTrackSeries>().GetFrameType() == FrameType &&
+					Series->As<FTimerFrameStatsTrackSeries>().GetTimerId() == TimerId;
 		});
 
 	if (ExistingSeries)
@@ -379,7 +379,7 @@ void SFrameTrack::UpdateState()
 	// Reset stats.
 	for (TSharedPtr<FFrameTrackSeries> Series : AllSeries)
 	{
-		Series->NumAggregatedFrames = 0;
+		Series->SetNumAggregatedFrames(0);
 	}
 	NumUpdatedFrames = 0;
 
@@ -413,14 +413,14 @@ void SFrameTrack::UpdateState()
 		for (int32 Index = 0; Index < AllSeries.Num(); ++Index)
 		{
 			TSharedPtr<FFrameTrackSeries> Series = AllSeries[Index];
-			if (Series->Type != EFrameTrackSeriesType::TimerFrameStats)
+			if (!Series->Is<FTimerFrameStatsTrackSeries>())
 			{
 				continue;
 			}
 
-			TSharedPtr<FTimerFrameStatsTrackSeries> TimerSeries = StaticCastSharedPtr<FTimerFrameStatsTrackSeries>(Series);
+			FTimerFrameStatsTrackSeries& TimerSeries = Series->As<FTimerFrameStatsTrackSeries>();
 			TArray<FFrameStatsCachedEvent> Frames;
-			FramesProvider.EnumerateFrames(static_cast<ETraceFrameType>(Series->FrameType), StartIndex, EndIndex, [&Frames](const TraceServices::FFrame& Frame)
+			FramesProvider.EnumerateFrames(static_cast<ETraceFrameType>(Series->GetFrameType()), StartIndex, EndIndex, [&Frames](const TraceServices::FFrame& Frame)
 				{
 					FFrameStatsCachedEvent Event;
 					Event.FrameStartTime = Frame.StartTime;
@@ -446,7 +446,7 @@ void SFrameTrack::UpdateState()
 					if (ThreadSharedState.IsValid())
 					{
 						ThreadSharedState->GetVisibleTimelineIndexes(Timelines);
-						FFrameStatsHelper::ComputeFrameStatsForTimer(Frames, TimerSeries->TimerId, Timelines);
+						FFrameStatsHelper::ComputeFrameStatsForTimer(Frames, TimerSeries.GetTimerId(), Timelines);
 						bTimingViewExists = true;
 					}
 				}
@@ -455,7 +455,7 @@ void SFrameTrack::UpdateState()
 			if (!bTimingViewExists)
 			{
 				// Compute the stats for all timelines.
-				FFrameStatsHelper::ComputeFrameStatsForTimer(Frames, TimerSeries->TimerId);
+				FFrameStatsHelper::ComputeFrameStatsForTimer(Frames, TimerSeries.GetTimerId());
 			}
 
 			uint64 CurrentIndex = StartIndex;
@@ -493,17 +493,17 @@ FFrameTrackSampleRef SFrameTrack::GetSampleAtMousePosition(double X, double Y)
 			{
 				TSharedPtr<FFrameTrackSeries> SeriesPtr = AllSeries[SeriesIndex];
 
-				if (!SeriesPtr->bIsVisible)
+				if (!SeriesPtr->IsVisible())
 				{
 					continue;
 				}
 
 				if (SeriesPtr.IsValid())
 				{
-					if (SeriesPtr->NumAggregatedFrames > 0 &&
-						SampleIndex < SeriesPtr->Samples.Num())
+					if (SeriesPtr->GetNumAggregatedFrames() > 0 &&
+						SampleIndex < SeriesPtr->GetNumSamples())
 					{
-						const FFrameTrackSample& Sample = SeriesPtr->Samples[SampleIndex];
+						const FFrameTrackSample& Sample = SeriesPtr->GetSample(SampleIndex);
 						if (Sample.NumFrames > 0)
 						{
 							const FAxisViewportDouble& ViewportY = Viewport.GetVerticalAxisViewport();
@@ -628,7 +628,7 @@ int32 SFrameTrack::OnPaint(const FPaintArgs& Args, const FGeometry& AllottedGeom
 		// Draw frames, for each visible Series.
 		for (TSharedPtr<FFrameTrackSeries> Series : AllSeries)
 		{
-			if (!Series->bIsVisible)
+			if (!Series->IsVisible())
 			{
 				continue;
 			}
@@ -707,7 +707,7 @@ int32 SFrameTrack::OnPaint(const FPaintArgs& Args, const FGeometry& AllottedGeom
 			// First line: "Rendering Frame 1,234"
 			TStringBuilder<512> StringBuilder;
 
-			StringBuilder.Append(HoveredSample.Series->Name.ToString());
+			StringBuilder.Append(HoveredSample.Series->GetName().ToString());
 			StringBuilder.Append(TEXT(" "));
 			StringBuilder.Append(FText::AsNumber(HoveredSample.Sample->LargestFrameIndex).ToString());
 			const FString Text1(StringBuilder);
@@ -769,9 +769,9 @@ int32 SFrameTrack::OnPaint(const FPaintArgs& Args, const FGeometry& AllottedGeom
 			DrawContext.LayerId++;
 
 			const FLinearColor TextColor1 =
-				HoveredSample.Series->FrameType == TraceFrameType_Rendering ?
+				HoveredSample.Series->GetFrameType() == TraceFrameType_Rendering ?
 					FLinearColor(0.5f, 0.1f, 0.1f, TooltipOpacity) :
-				HoveredSample.Series->FrameType == TraceFrameType_Game ?
+				HoveredSample.Series->GetFrameType() == TraceFrameType_Game ?
 					FLinearColor(0.1f, 0.1f, 0.5f, TooltipOpacity) :
 					FLinearColor(0.1f, 0.1f, 0.1f, TooltipOpacity);
 			const FLinearColor TextColor2(0.05f, 0.05f, 0.05f, TooltipOpacity);
@@ -1417,25 +1417,26 @@ void SFrameTrack::ShowContextMenu(const FPointerEvent& MouseEvent)
 
 	for (TSharedPtr<FFrameTrackSeries> Series : AllSeries)
 	{
-		if (Series->Type != EFrameTrackSeriesType::TimerFrameStats)
+		if (!Series->Is<FTimerFrameStatsTrackSeries>())
 		{
 			continue;
 		}
 
-		TSharedPtr<FTimerFrameStatsTrackSeries> FrameStatsSeries = StaticCastSharedPtr<FTimerFrameStatsTrackSeries>(Series);
-		ETraceFrameType FrameType = static_cast<ETraceFrameType>(FrameStatsSeries->FrameType);
+
+		FTimerFrameStatsTrackSeries& TimerSeries = Series->As<FTimerFrameStatsTrackSeries>();
+		ETraceFrameType FrameType = static_cast<ETraceFrameType>(TimerSeries.GetFrameType());
 		FUIAction Action_ShowFrameStatsSeries
 		(
-			FExecuteAction::CreateSP(this, &SFrameTrack::ContextMenu_ShowFrameStats_Execute, FrameType, FrameStatsSeries->TimerId),
-			FCanExecuteAction::CreateSP(this, &SFrameTrack::ContextMenu_ShowFrameStats_CanExecute, FrameType, FrameStatsSeries->TimerId),
-			FIsActionChecked::CreateSP(this, &SFrameTrack::ContextMenu_ShowFrameStats_IsChecked, FrameType, FrameStatsSeries->TimerId)
+			FExecuteAction::CreateSP(this, &SFrameTrack::ContextMenu_ShowFrameStats_Execute, FrameType, TimerSeries.GetTimerId()),
+			FCanExecuteAction::CreateSP(this, &SFrameTrack::ContextMenu_ShowFrameStats_CanExecute, FrameType, TimerSeries.GetTimerId()),
+			FIsActionChecked::CreateSP(this, &SFrameTrack::ContextMenu_ShowFrameStats_IsChecked, FrameType, TimerSeries.GetTimerId())
 		);
 		MenuBuilder.AddMenuEntry
 		(
 			Action_ShowFrameStatsSeries,
-			CreateSeriesMenuWidget(FrameStatsSeries->Color, FrameStatsSeries->Name),
+			CreateSeriesMenuWidget(TimerSeries.GetColor(), TimerSeries.GetName()),
 			NAME_None,
-			FText::Format(LOCTEXT("ContextMenu_ShowFrameStatsSeries_Desc", "Shows/hides the {0} timer series."), FrameStatsSeries->Name),
+			FText::Format(LOCTEXT("ContextMenu_ShowFrameStatsSeries_Desc", "Shows/hides the {0} timer series."), TimerSeries.GetName()),
 			EUserInterfaceActionType::ToggleButton
 		);
 	}
@@ -1849,7 +1850,7 @@ void SFrameTrack::CreateSelectedFrameMenu(FMenuBuilder& MenuBuilder)
 	if (SelectedSample.IsValid())
 	{
 		SelectedFrameSectionName = FText::Format(LOCTEXT("ContextMenu_Section_SelectedFrame_Fmt", "{0} {1}"),
-			FFrameTrackDrawHelper::FrameTypeToText(SelectedSample.Series->FrameType),
+			FFrameTrackDrawHelper::FrameTypeToText(SelectedSample.Series->GetFrameType()),
 			FText::AsNumber(SelectedSample.Sample->LargestFrameIndex));
 	}
 	else
@@ -1893,7 +1894,7 @@ void SFrameTrack::ContextMenu_ShowGameFrames_Execute()
 	TSharedPtr<FFrameTrackSeries> Series = FindSeries(ETraceFrameType::TraceFrameType_Game);
 	if (Series.IsValid())
 	{
-		Series->bIsVisible = !Series->bIsVisible;
+		Series->SetVisibility(!Series->IsVisible());
 	}
 }
 
@@ -1909,7 +1910,7 @@ bool SFrameTrack::ContextMenu_ShowGameFrames_CanExecute()
 bool SFrameTrack::ContextMenu_ShowGameFrames_IsChecked()
 {
 	TSharedPtr<FFrameTrackSeries> Series = FindSeries(ETraceFrameType::TraceFrameType_Game);
-	return Series.IsValid() ? Series->bIsVisible : false;
+	return Series.IsValid() ? Series->IsVisible() : false;
 }
 
 ////////////////////////////////////////////////////////////////////////////////////////////////////
@@ -1919,7 +1920,7 @@ void SFrameTrack::ContextMenu_ShowRenderingFrames_Execute()
 	TSharedPtr<FFrameTrackSeries> Series = FindSeries(ETraceFrameType::TraceFrameType_Rendering);
 	if (Series.IsValid())
 	{
-		Series->bIsVisible = !Series->bIsVisible;
+		Series->SetVisibility(!Series->IsVisible());
 	}
 }
 
@@ -1935,7 +1936,7 @@ bool SFrameTrack::ContextMenu_ShowRenderingFrames_CanExecute()
 bool SFrameTrack::ContextMenu_ShowRenderingFrames_IsChecked()
 {
 	TSharedPtr<FFrameTrackSeries> Series = FindSeries(ETraceFrameType::TraceFrameType_Rendering);
-	return Series.IsValid() ? Series->bIsVisible : false;
+	return Series.IsValid() ? Series->IsVisible() : false;
 }
 
 ////////////////////////////////////////////////////////////////////////////////////////////////////
@@ -1945,7 +1946,7 @@ void SFrameTrack::ContextMenu_ShowFrameStats_Execute(ETraceFrameType FrameType,
 	TSharedPtr<FFrameTrackSeries> Series = FindFrameStatsSeries(FrameType, TimerId);
 	if (Series.IsValid())
 	{
-		Series->bIsVisible = !Series->bIsVisible;
+		Series->SetVisibility(!Series->IsVisible());
 	}
 }
 
@@ -1961,7 +1962,7 @@ bool SFrameTrack::ContextMenu_ShowFrameStats_CanExecute(ETraceFrameType FrameTyp
 bool SFrameTrack::ContextMenu_ShowFrameStats_IsChecked(ETraceFrameType FrameType, uint32 TimerId)
 {
 	TSharedPtr<FFrameTrackSeries> Series = FindFrameStatsSeries(FrameType, TimerId);
-	return Series.IsValid() ? Series->bIsVisible : false;
+	return Series.IsValid() ? Series->IsVisible() : false;
 }
 
 ////////////////////////////////////////////////////////////////////////////////////////////////////
@@ -2139,9 +2140,9 @@ bool SFrameTrack::HasFrameStatSeries(ETraceFrameType FrameType, uint32 TimerId)
 {
 	const TSharedPtr<FFrameTrackSeries>* ExistingSeries = AllSeries.FindByPredicate([FrameType, TimerId](TSharedPtr<FFrameTrackSeries> Series)
 		{
-			return Series->Type == EFrameTrackSeriesType::TimerFrameStats &&
-				Series->FrameType == FrameType &&
-				StaticCastSharedPtr<FTimerFrameStatsTrackSeries>(Series)->TimerId == TimerId;
+			return Series->Is<FTimerFrameStatsTrackSeries>() &&
+					Series->As<FTimerFrameStatsTrackSeries>().GetFrameType() == FrameType &&
+					Series->As<FTimerFrameStatsTrackSeries>().GetTimerId() == TimerId;
 		});
 
 	return ExistingSeries != nullptr;
@@ -2149,24 +2150,22 @@ bool SFrameTrack::HasFrameStatSeries(ETraceFrameType FrameType, uint32 TimerId)
 
 ////////////////////////////////////////////////////////////////////////////////////////////////////
 
-TSharedPtr<FTimerFrameStatsTrackSeries> SFrameTrack::AddTimerFrameStatSeries(ETraceFrameType FrameType, uint32 TimerId, FLinearColor Color, FText Name)
+FTimerFrameStatsTrackSeries& SFrameTrack::AddTimerFrameStatSeries(ETraceFrameType FrameType, uint32 TimerId, FLinearColor Color, FText Name)
 {
 	TSharedPtr<FFrameTrackSeries> ExistingSeries = FindFrameStatsSeries(FrameType, TimerId);
 
 	if (ExistingSeries.IsValid())
 	{
-		return StaticCastSharedPtr<FTimerFrameStatsTrackSeries>(ExistingSeries);
+		return ExistingSeries->As<FTimerFrameStatsTrackSeries>();
 	}
-
 	TSharedRef<FTimerFrameStatsTrackSeries> SeriesRef = MakeShared<FTimerFrameStatsTrackSeries>(FrameType, TimerId);
-	SeriesRef->TimerId = TimerId;
-	SeriesRef->Color = Color;
-	SeriesRef->Name = Name;
+	SeriesRef->SetColor(Color);
+	SeriesRef->SetName(Name);
 	AllSeries.Add(SeriesRef);
 
 	bIsStateDirty = true;
 
-	return SeriesRef;
+	return *SeriesRef;
 }
 
 ////////////////////////////////////////////////////////////////////////////////////////////////////
@@ -2175,9 +2174,9 @@ bool SFrameTrack::RemoveTimerFrameStatSeries(ETraceFrameType FrameType, uint32 T
 {
 	int32 NumRemoved = AllSeries.RemoveAll([FrameType, TimerId](TSharedPtr<FFrameTrackSeries> Series)
 		{
-			return Series->Type == EFrameTrackSeriesType::TimerFrameStats &&
-				Series->FrameType == FrameType &&
-				StaticCastSharedPtr<FTimerFrameStatsTrackSeries>(Series)->TimerId == TimerId;
+			return Series->Is<FTimerFrameStatsTrackSeries>() &&
+				Series->As<FTimerFrameStatsTrackSeries>().GetFrameType() == FrameType &&
+				Series->As<FTimerFrameStatsTrackSeries>().GetTimerId() == TimerId;
 		});
 
 	ensure(NumRemoved == 1);
@@ -2193,10 +2192,9 @@ bool SFrameTrack::HasAnySeriesForTimer(uint32 TimerId) const
 {
 	for (const TSharedPtr<FFrameTrackSeries>& Series : AllSeries)
 	{
-		if (Series->Type == EFrameTrackSeriesType::TimerFrameStats)
+		if (Series->Is<FTimerFrameStatsTrackSeries>())
 		{
-			const TSharedPtr<FTimerFrameStatsTrackSeries> TimerSeries = StaticCastSharedPtr<FTimerFrameStatsTrackSeries>(Series);
-			if (TimerSeries->TimerId == TimerId)
+			if (Series->As<FTimerFrameStatsTrackSeries>().GetTimerId() == TimerId)
 			{
 				return true;
 			}
@@ -2214,10 +2212,9 @@ uint32 SFrameTrack::GetNumSeriesForTimer(uint32 TimerId) const
 
 	for (const TSharedPtr<FFrameTrackSeries>& Series : AllSeries)
 	{
-		if (Series->Type == EFrameTrackSeriesType::TimerFrameStats)
+		if (Series->Is<FTimerFrameStatsTrackSeries>())
 		{
-			const TSharedPtr<FTimerFrameStatsTrackSeries> TimerSeries = StaticCastSharedPtr<FTimerFrameStatsTrackSeries>(Series);
-			if (TimerSeries->TimerId == TimerId)
+			if (Series->As<FTimerFrameStatsTrackSeries>().GetTimerId() == TimerId)
 			{
 				++NumSeries;
 			}
diff --git a/Engine/Source/Developer/TraceInsights/Private/Insights/TimingProfiler/Widgets/SFrameTrack.h b/Engine/Source/Developer/TraceInsights/Private/Insights/TimingProfiler/Widgets/SFrameTrack.h
--- a/Engine/Source/Developer/TraceInsights/Private/Insights/TimingProfiler/Widgets/SFrameTrack.h
+++ b/Engine/Source/Developer/TraceInsights/Private/Insights/TimingProfiler/Widgets/SFrameTrack.h
@@ -138,7 +138,7 @@ public:
 	virtual FCursorReply OnCursorQuery(const FGeometry& MyGeometry, const FPointerEvent& CursorEvent) const override;
 
 	bool HasFrameStatSeries(ETraceFrameType FrameType, uint32 TimerId) const;
-	TSharedPtr<FTimerFrameStatsTrackSeries> AddTimerFrameStatSeries(ETraceFrameType FrameType, uint32 TimerId, FLinearColor Color, FText Name);
+	FTimerFrameStatsTrackSeries& AddTimerFrameStatSeries(ETraceFrameType FrameType, uint32 TimerId, FLinearColor Color, FText Name);
 	bool RemoveTimerFrameStatSeries(ETraceFrameType FrameType, uint32 TimerId);
 
 	bool HasAnySeriesForTimer(uint32 TimerId) const;
diff --git a/Engine/Source/Developer/TraceInsights/Private/Insights/TimingProfiler/Widgets/STimingProfilerWindow.cpp b/Engine/Source/Developer/TraceInsights/Private/Insights/TimingProfiler/Widgets/STimingProfilerWindow.cpp
--- a/Engine/Source/Developer/TraceInsights/Private/Insights/TimingProfiler/Widgets/STimingProfilerWindow.cpp
+++ b/Engine/Source/Developer/TraceInsights/Private/Insights/TimingProfiler/Widgets/STimingProfilerWindow.cpp
@@ -555,11 +555,11 @@ void STimingProfilerWindow::SetTimerColor(uint32 TimerId, FLinearColor Color)
 	{
 		for (const TSharedPtr<FFrameTrackSeries>& Series : FrameTrack->GetSeries())
 		{
-			if (Series->Type == EFrameTrackSeriesType::TimerFrameStats)
+			if (Series->Is<FTimerFrameStatsTrackSeries>())
 			{
-				if (StaticCastSharedPtr<FTimerFrameStatsTrackSeries>(Series)->TimerId == TimerId)
+				if (Series->As<FTimerFrameStatsTrackSeries>().GetTimerId() == TimerId)
 				{
-					Series->Color = Color;
+					Series->SetColor(Color);
 				}
 			}
 		}
