[Insights] Memory Tags: Fixed segmented control for available asset tags not updating correctly in some cases.

diff --git a/Engine/Source/Developer/TraceInsights/Private/Insights/MemoryProfiler/Widgets/SMemTagTreeView.cpp b/Engine/Source/Developer/TraceInsights/Private/Insights/MemoryProfiler/Widgets/SMemTagTreeView.cpp
--- a/Engine/Source/Developer/TraceInsights/Private/Insights/MemoryProfiler/Widgets/SMemTagTreeView.cpp
+++ b/Engine/Source/Developer/TraceInsights/Private/Insights/MemoryProfiler/Widgets/SMemTagTreeView.cpp
@@ -206,19 +206,6 @@ END_SLATE_FUNCTION_BUILD_OPTIMIZATION
 
 TSharedPtr<SWidget> SMemTagTreeView::ConstructTagSetAndViewPreset()
 {
-	TagSetsSegmentedControl = SNew(SSegmentedControl<FMemoryTagSetId>)
-		.SupportsMultiSelection(false)
-		.OnValueChanged_Lambda([this](FMemoryTagSetId InValue)
-			{
-				TagSetFilter = InValue;
-				OnNodeFilteringChanged();
-				bShouldUpdateStats = true;
-				bShouldUpdateBudgets = true;
-			})
-		.Value_Lambda([this]
-			{
-				return TagSetFilter;
-			});
 	UpdateAvailableTagSets();
 
 	TSharedPtr<SHorizontalBox> Widget = SNew(SHorizontalBox)
@@ -243,11 +230,26 @@ TSharedPtr<SWidget> SMemTagTreeView::ConstructTagSetAndViewPreset()
 
 void SMemTagTreeView::UpdateAvailableTagSets()
 {
-	if (TagSetsSegmentedControl->NumSlots() != 0)
+	TSharedPtr<SHorizontalBox> ParentBox;
+	if (TagSetsSegmentedControl.IsValid())
 	{
-		return;
+		ParentBox = StaticCastSharedPtr<SHorizontalBox>(TagSetsSegmentedControl->GetParentWidget());
 	}
 
+	TagSetsSegmentedControl = SNew(SSegmentedControl<FMemoryTagSetId>)
+		.SupportsMultiSelection(false)
+		.OnValueChanged_Lambda([this](FMemoryTagSetId InValue)
+			{
+				TagSetFilter = InValue;
+				OnNodeFilteringChanged();
+				bShouldUpdateStats = true;
+				bShouldUpdateBudgets = true;
+			})
+		.Value_Lambda([this]
+			{
+				return TagSetFilter;
+			});
+
 	for (TSharedPtr<FMemoryTagSetId> TagSetPtr : AvailableTagSets)
 	{
 		FMemoryTagSetId TagSet = *TagSetPtr;
@@ -309,6 +311,11 @@ void SMemTagTreeView::UpdateAvailableTagSets()
 			.ToolTip(TagSetToolTip)
 			.AttachWidget(Widget);
 	}
+
+	if (ParentBox.IsValid())
+	{
+		ParentBox->GetSlot(0).AttachWidget(TagSetsSegmentedControl.ToSharedRef());
+	}
 }
 
 ////////////////////////////////////////////////////////////////////////////////////////////////////
