On behalf of ionut.matasaru: Fixed loctext warning in memory profiler window.

#rnx

diff --git a/Engine/Source/Developer/TraceInsights/Private/Insights/MemoryProfiler/Widgets/SMemoryProfilerWindow.cpp b/Engine/Source/Developer/TraceInsights/Private/Insights/MemoryProfiler/Widgets/SMemoryProfilerWindow.cpp
--- a/Engine/Source/Developer/TraceInsights/Private/Insights/MemoryProfiler/Widgets/SMemoryProfilerWindow.cpp
+++ b/Engine/Source/Developer/TraceInsights/Private/Insights/MemoryProfiler/Widgets/SMemoryProfilerWindow.cpp
@@ -443,7 +443,7 @@ void SMemoryProfilerWindow::RegisterTabSpawners()
 		TabManagerPtr->RegisterTabSpawner(FMemoryProfilerTabs::ModulesViewID, FOnSpawnTab::CreateRaw(this, &SMemoryProfilerWindow::SpawnTab_ModulesView))
 			.SetDisplayName(LOCTEXT("ModulesViewTabTitle", "Modules"))
 			.SetIcon(FSlateIcon(FInsightsStyle::GetStyleSetName(), "Icons.ModulesView"))
-			.SetTooltipText(LOCTEXT("TimingViewTabToolTip", "Opens the Modules panel."))
+			.SetTooltipText(LOCTEXT("ModulesViewTabToolTip", "Opens the Modules panel."))
 			.SetGroup(Group);
 	}
 }
