Filtering Logs by Verbosity and Category on the Timing Profiler Markers Track.

diff --git a/Engine/Source/Developer/TraceInsights/Private/Insights/TimingProfiler/TimingProfilerManager.cpp b/Engine/Source/Developer/TraceInsights/Private/Insights/TimingProfiler/TimingProfilerManager.cpp
--- a/Engine/Source/Developer/TraceInsights/Private/Insights/TimingProfiler/TimingProfilerManager.cpp
+++ b/Engine/Source/Developer/TraceInsights/Private/Insights/TimingProfiler/TimingProfilerManager.cpp
@@ -67,6 +67,7 @@ FTimingProfilerManager::FTimingProfilerManager(TSharedRef<FUICommandList> InComm
 	, CommandList(InCommandList)
 	, ActionManager(this)
 	, ProfilerWindowWeakPtr()
+	, Filter(MakeShared<FLogFilter>())
 	, bIsFramesTrackVisible(false)
 	, bIsTimingViewVisible(false)
 	, bIsTimersViewVisible(false)
@@ -132,6 +133,8 @@ void FTimingProfilerManager::Shutdown()
 
 	FTimingProfilerCommands::Unregister();
 
+	Filter->Reset();
+
 	// Unregister tick function.
 	FTSTicker::GetCoreTicker().RemoveTicker(OnTickHandle);
 
diff --git a/Engine/Source/Developer/TraceInsights/Private/Insights/TimingProfiler/TimingProfilerManager.h b/Engine/Source/Developer/TraceInsights/Private/Insights/TimingProfiler/TimingProfilerManager.h
--- a/Engine/Source/Developer/TraceInsights/Private/Insights/TimingProfiler/TimingProfilerManager.h
+++ b/Engine/Source/Developer/TraceInsights/Private/Insights/TimingProfiler/TimingProfilerManager.h
@@ -12,6 +12,7 @@
 #include "Insights/IUnrealInsightsModule.h"
 #include "Insights/TimingProfiler/TimingProfilerCommands.h"
 #include "Insights/TimingProfiler/ViewModels/TimerNode.h"
+#include "Insights/ViewModels/LogFilter.h"
 
 namespace UE::Insights
 {
@@ -92,6 +93,8 @@ public:
 		return ProfilerWindowWeakPtr.Pin();
 	}
 
+	const TSharedRef<FLogFilter>& GetLogFilter() const { return Filter; }
+
 	////////////////////////////////////////////////////////////////////////////////////////////////////
 	// Getters and setters used by Toggle Commands.
 
@@ -210,6 +213,9 @@ private:
 	/** A weak pointer to the Timing Insights window. */
 	TWeakPtr<STimingProfilerWindow> ProfilerWindowWeakPtr;
 
+	/** A reference to the log filter state. */
+	TSharedRef<FLogFilter> Filter;
+
 	/** If the Frames Track is visible or hidden. */
 	bool bIsFramesTrackVisible;
 
diff --git a/Engine/Source/Developer/TraceInsights/Private/Insights/TimingProfiler/Tracks/MarkersTimingTrack.cpp b/Engine/Source/Developer/TraceInsights/Private/Insights/TimingProfiler/Tracks/MarkersTimingTrack.cpp
--- a/Engine/Source/Developer/TraceInsights/Private/Insights/TimingProfiler/Tracks/MarkersTimingTrack.cpp
+++ b/Engine/Source/Developer/TraceInsights/Private/Insights/TimingProfiler/Tracks/MarkersTimingTrack.cpp
@@ -27,6 +27,7 @@
 #include "Insights/InsightsStyle.h"
 #include "Insights/TimingProfiler/TimingProfilerManager.h"
 #include "Insights/TimingProfiler/Widgets/STimingProfilerWindow.h"
+#include "Insights/ViewModels/LogFilter.h"
 #include "Insights/ViewModels/TimingEvent.h"
 #include "Insights/ViewModels/TimingTrackViewport.h"
 #include "Insights/ViewModels/TooltipDrawState.h"
@@ -81,6 +82,8 @@ void FMarkersTimingTrack::Reset()
 	Header.SetIsInBackground(true);
 	Header.SetCanBeCollapsed(true);
 
+	ChangeNumber = 0;
+
 	NumLogMessages = 0;
 	NumDrawBoxes = 0;
 	NumDrawTexts = 0;
@@ -128,11 +131,18 @@ void FMarkersTimingTrack::Update(const ITimingTrackUpdateContext& Context)
 	Header.Update(Context);
 
 	const FTimingTrackViewport& Viewport = Context.GetViewport();
-	if (IsDirty() || Viewport.IsHorizontalViewportDirty())
+	
+	if (FTimingProfilerManager::Get().IsValid())
 	{
-		ClearDirtyFlag();
+		uint64 NewChangeNumber = FTimingProfilerManager::Get()->GetLogFilter()->GetChangeNumber();
+		if (IsDirty() || Viewport.IsHorizontalViewportDirty() || ChangeNumber != NewChangeNumber)
+		{
+			ClearDirtyFlag();
 
-		UpdateDrawState(Context);
+			ChangeNumber = NewChangeNumber;
+			
+			UpdateDrawState(Context);
+		}
 	}
 }
 
@@ -526,27 +536,34 @@ const TSharedPtr<const ITimingEvent> FMarkersTimingTrack::GetEvent(float InPosX,
 		}
 		else
 		{
-			LogIndex = LogProvider.UpperBoundByTime(PosTime);
-			if (LogIndex == 0 || LogIndex == ~0ull)
-			{
-				return nullptr;
-			}
-			--LogIndex;
+			ELogVerbosity::Type VerbosityThreshold = FTimingProfilerManager::Get()->GetLogFilter()->GetVerbosityThreshold();
+			const TSet<FName>&  EnabledLogCategories = FTimingProfilerManager::Get()->GetLogFilter()->GetEnabledLogCategories();
 
-			if (!LogProvider.ReadMessage(LogIndex,
-				[&Category, &Time, &ScreenshotId, this](const TraceServices::FLogMessageInfo& Message)
+			// Enumerate through from the current position backwards until there is an appropriately filtered element 
+			LogProvider.ReverseEnumerateMessages(PosTime, 0.0f,
+			[&LogIndex, &Time, &Category, &ScreenshotId, &VerbosityThreshold, &EnabledLogCategories, this](const TraceServices::FLogMessageInfo& Message)
+			{
+				FStringView CategoryStr(Message.Category->Name);
+				if (CategoryStr.StartsWith(TEXT("Log")))
 				{
-					Category = Message.Category;
+					CategoryStr.RightChopInline(3);
+				}
+				FName CleanedCategoryStr(CategoryStr);
+				
+				if (Message.Verbosity <= VerbosityThreshold && EnabledLogCategories.Contains(CleanedCategoryStr))
+				{
+					LogIndex = Message.Index;
 					Time = Message.Time;
+					Category = Message.Category;
 					if (Message.Category == this->ScreenshotCategory)
 					{
 						check(Message.Line >= 0);
 						ScreenshotId = Message.Line;
 					}
-				}))
-			{
-				return nullptr;
-			}
+					return true;
+				}
+				return false;
+			});
 		}
 
 		if (Category == nullptr)
@@ -852,19 +869,42 @@ void FTimeMarkerTrackBuilder::AddLogMessage(const TraceServices::FLogMessageInfo
 	}
 
 	check(Message.Category != nullptr);
-	if (!Track.bUseOnlyBookmarks ||
-		Message.Category == Track.BookmarkCategory ||
-		Message.Category == Track.ScreenshotCategory)
+	
+	FStringView CategoryStr(Message.Category->Name);
+	if (CategoryStr.StartsWith(TEXT("Log")))
 	{
-		float X = Viewport.TimeToSlateUnitsRounded(Message.Time);
-		if (X < 0.0f)
+		CategoryStr.RightChopInline(3);
+	}
+	FName CleanedCategoryStr(CategoryStr);
+
+	// Check whether we are using the bookmarks track and if the current message is allowed
+	if (Track.IsBookmarksTrack())
+	{
+		if (!(Message.Category == Track.BookmarkCategory || Message.Category == Track.ScreenshotCategory))
 		{
-			X = -1.0f;
+			return;
 		}
+	}
+		
+	// Check whether we can allow the message to be added to the logs track
+	else if (Track.IsLogsTrack() && FTimingProfilerManager::Get().IsValid())
+	{
+		ELogVerbosity::Type VerbosityThreshold = FTimingProfilerManager::Get()->GetLogFilter()->GetVerbosityThreshold();
+		const TSet<FName>&  EnabledLogCategories = FTimingProfilerManager::Get()->GetLogFilter()->GetEnabledLogCategories();
+		if (!(Message.Verbosity <= VerbosityThreshold && EnabledLogCategories.Contains(CleanedCategoryStr)))
+		{
+			return;
+		}
+	}
 
-		const TCHAR* CategoryName = Message.Category->Name != nullptr ? Message.Category->Name : TEXT("");
-		AddTimeMarker(X, Message.Index, Message.Verbosity, CategoryName, Message.Message);
+	float X = Viewport.TimeToSlateUnitsRounded(Message.Time);
+	if (X < 0.0f)
+	{
+		X = -1.0f;
 	}
+	
+	const TCHAR* CategoryName = Message.Category->Name != nullptr ? Message.Category->Name : TEXT("");
+	AddTimeMarker(X, Message.Index, Message.Verbosity, CategoryName, Message.Message);
 }
 
 ////////////////////////////////////////////////////////////////////////////////////////////////////
diff --git a/Engine/Source/Developer/TraceInsights/Private/Insights/TimingProfiler/Tracks/MarkersTimingTrack.h b/Engine/Source/Developer/TraceInsights/Private/Insights/TimingProfiler/Tracks/MarkersTimingTrack.h
--- a/Engine/Source/Developer/TraceInsights/Private/Insights/TimingProfiler/Tracks/MarkersTimingTrack.h
+++ b/Engine/Source/Developer/TraceInsights/Private/Insights/TimingProfiler/Tracks/MarkersTimingTrack.h
@@ -123,6 +123,8 @@ private:
 
 	FTrackHeader Header;
 
+	uint64 ChangeNumber = 0;
+
 	// Stats
 	int32 NumLogMessages = 0;
 	mutable int32 NumDrawBoxes = 0;
diff --git a/Engine/Source/Developer/TraceInsights/Private/Insights/ViewModels/LogFilter.h b/Engine/Source/Developer/TraceInsights/Private/Insights/ViewModels/LogFilter.h
--- a/Engine/Source/Developer/TraceInsights/Private/Insights/ViewModels/LogFilter.h
+++ b/Engine/Source/Developer/TraceInsights/Private/Insights/ViewModels/LogFilter.h
@@ -94,6 +94,9 @@ public:
 	/** Returns true if the specified log category is enabled. */
 	bool IsLogCategoryEnabled(const FName& LogCategory) const;
 
+	/** Returns the set of filter-enabled categories. */
+	const TSet<FName>& GetEnabledLogCategories() const { return EnabledLogCategories; }
+
 private:
 	void OnFilterByCategoryChanged()
 	{
diff --git a/Engine/Source/Developer/TraceInsights/Private/Insights/Widgets/SLogView.cpp b/Engine/Source/Developer/TraceInsights/Private/Insights/Widgets/SLogView.cpp
--- a/Engine/Source/Developer/TraceInsights/Private/Insights/Widgets/SLogView.cpp
+++ b/Engine/Source/Developer/TraceInsights/Private/Insights/Widgets/SLogView.cpp
@@ -506,7 +506,8 @@ private:
 ////////////////////////////////////////////////////////////////////////////////////////////////////
 
 SLogView::SLogView()
-	: FilteringStartIndex(0)
+	: Filter(*TimingProfiler::FTimingProfilerManager::Get()->GetLogFilter())
+	, FilteringStartIndex(0)
 	, FilteringEndIndex(0)
 	, FilteringChangeNumber(0)
 	, bIsFilteringAsyncTaskCancelRequested(false)
diff --git a/Engine/Source/Developer/TraceInsights/Private/Insights/Widgets/SLogView.h b/Engine/Source/Developer/TraceInsights/Private/Insights/Widgets/SLogView.h
--- a/Engine/Source/Developer/TraceInsights/Private/Insights/Widgets/SLogView.h
+++ b/Engine/Source/Developer/TraceInsights/Private/Insights/Widgets/SLogView.h
@@ -170,7 +170,7 @@ protected:
 	/** The search box widget used to filter logs by message text. */
 	TSharedPtr<SSearchBox> FilterTextBox;
 
-	FLogFilter Filter;
+	FLogFilter& Filter;
 	uint64 FilterChangeNumber;
 
 	int32 FilteringStartIndex; // Start index (of the range of log messages to filter) currently used by the async task
diff --git a/Engine/Source/Developer/TraceServices/Private/Model/Log.cpp b/Engine/Source/Developer/TraceServices/Private/Model/Log.cpp
--- a/Engine/Source/Developer/TraceServices/Private/Model/Log.cpp
+++ b/Engine/Source/Developer/TraceServices/Private/Model/Log.cpp
@@ -280,6 +280,46 @@ void FLogProvider::EnumerateMessages(double StartTime, double EndTime, TFunction
 	}
 }
 
+void FLogProvider::ReverseEnumerateMessages(double StartTime, double EndTime, TFunctionRef<bool(const FLogMessageInfo&)> Callback) const
+{
+	Session.ReadAccessCheck();
+
+	if (StartTime < EndTime)
+	{
+		return;
+	}
+
+	uint64 MessageCount = Messages.Num();
+	if (MessageCount == 0)
+	{
+		return;
+	}
+
+	// Find the first log message before the current position
+	uint64 StartIndex = TraceServices::PagedArrayAlgo::LowerBoundBy(Messages, StartTime,
+		[](const FLogMessageInternal& Item) { return Item.Time; });
+	if (StartIndex == 0 || StartIndex >= Messages.Num())
+	{
+		return;
+	}
+	--StartIndex;
+
+	// Iterate from StartIndex and stop if we've exceeded the StartTime
+	for (auto It = Messages.GetIteratorFromItem(StartIndex); It; --It)
+	{
+		double Time = It.GetCurrentItem()->Time;
+		if (Time < EndTime)
+		{
+			break;
+		}
+		// Once we've found an appropriate element we can stop
+		if (ConstructMessage(*It.GetCurrentItem(), It.GetCurrentItemIndex(), Callback))
+		{
+			break;
+		}
+	}
+}
+
 uint64 FLogProvider::LowerBoundByTime(double Time) const
 {
 	Session.ReadAccessCheck();
@@ -303,8 +343,8 @@ uint64 FLogProvider::BinarySearchClosestByTime(double Time) const
 	return TraceServices::PagedArrayAlgo::BinarySearchClosestBy(Messages, Time,
 		[](const FLogMessageInternal& Item) { return Item.Time; });
 }
-
-void FLogProvider::ConstructMessage(const FLogMessageInternal& InternalMessage, uint64 Index, TFunctionRef<void(const FLogMessageInfo&)> Callback) const
+	
+bool FLogProvider::ConstructMessage(const FLogMessageInternal& InternalMessage, uint64 Index, TFunctionRef<void(const FLogMessageInfo&)> Callback) const
 {
 	FLogMessageInfo Message;
 	Message.Index = Index;
@@ -315,6 +355,20 @@ void FLogProvider::ConstructMessage(const FLogMessageInternal& InternalMessage,
 	Message.Verbosity = InternalMessage.Spec->Verbosity;
 	Message.Message = InternalMessage.Message;
 	Callback(Message);
+	return true;
+}
+
+bool FLogProvider::ConstructMessage(const FLogMessageInternal& InternalMessage, uint64 Index, TFunctionRef<bool(const FLogMessageInfo&)> Callback) const
+{
+	FLogMessageInfo Message;
+	Message.Index = Index;
+	Message.Time = InternalMessage.Time;
+	Message.Category = InternalMessage.Spec->Category;
+	Message.File = InternalMessage.Spec->File;
+	Message.Line = InternalMessage.Spec->Line;
+	Message.Verbosity = InternalMessage.Spec->Verbosity;
+	Message.Message = InternalMessage.Message;
+	return Callback(Message);
 }
 
 void FLogProvider::EnumerateCategories(TFunctionRef<void(const FLogCategoryInfo&)> Callback) const
diff --git a/Engine/Source/Developer/TraceServices/Private/Model/LogPrivate.h b/Engine/Source/Developer/TraceServices/Private/Model/LogPrivate.h
--- a/Engine/Source/Developer/TraceServices/Private/Model/LogPrivate.h
+++ b/Engine/Source/Developer/TraceServices/Private/Model/LogPrivate.h
@@ -45,7 +45,8 @@ public:
 	virtual bool ReadMessage(uint64 Index, TFunctionRef<void(const FLogMessageInfo&)> Callback) const override;
 	virtual void EnumerateMessagesByIndex(uint64 StartIndex, uint64 EndIndex, TFunctionRef<void(const FLogMessageInfo&)> Callback) const override;
 	virtual void EnumerateMessages(double StartTime, double EndTime, TFunctionRef<void(const FLogMessageInfo&)> Callback) const override;
-
+	virtual void ReverseEnumerateMessages(double StartTime, double EndTime, TFunctionRef<bool(const FLogMessageInfo&)> Callback) const override;
+	
 	virtual uint64 LowerBoundByTime(double Time) const override;
 	virtual uint64 UpperBoundByTime(double Time) const override;
 	virtual uint64 BinarySearchClosestByTime(double Time) const override;
@@ -78,7 +79,8 @@ public:
 private:
 	FLogMessageInternal& AppendMessageInternal(double Time);
 	void AppendUnknownMessageInternal(uint64 LogPoint, double Time, const FStringView Message);
-	void ConstructMessage(const FLogMessageInternal& InternalMessage, uint64 Index, TFunctionRef<void(const FLogMessageInfo&)> Callback) const;
+	bool ConstructMessage(const FLogMessageInternal& InternalMessage, uint64 Index, TFunctionRef<void(const FLogMessageInfo&)> Callback) const;
+	bool ConstructMessage(const FLogMessageInternal& InternalMessage, uint64 Index, TFunctionRef<bool(const FLogMessageInfo&)> Callback) const;
 
 	enum
 	{
diff --git a/Engine/Source/Developer/TraceServices/Public/TraceServices/Model/Log.h b/Engine/Source/Developer/TraceServices/Public/TraceServices/Model/Log.h
--- a/Engine/Source/Developer/TraceServices/Public/TraceServices/Model/Log.h
+++ b/Engine/Source/Developer/TraceServices/Public/TraceServices/Model/Log.h
@@ -66,7 +66,7 @@ public:
 	virtual void EnumerateMessagesByIndex(uint64 StartIndex, uint64 EndIndex, TFunctionRef<void(const FLogMessageInfo&)> Callback) const = 0;
 
 	/**
-	 * Enumerates messages with timestamp in the specified time interval [StartTime, EndTime].
+	 * Enumerates messages with timestamp in the specified time interval [EndTime, StartTime].
 	 *
 	 * @param StartTime The inclusive start time of the interval.
 	 * @param EndTime The inclusive end time of the interval.
@@ -76,6 +76,18 @@ public:
 	 */
 	virtual void EnumerateMessages(double StartTime, double EndTime, TFunctionRef<void(const FLogMessageInfo&)> Callback) const = 0;
 
+	/**
+	 * Enumerates messages in reverse with timestamp in the specified time interval [StartTime, EndTime] until either it finds the last message
+	 * or is stopped on callback.
+	 *
+	 * @param StartTime The inclusive start time of the interval.
+	 * @param EndTime The inclusive end time of the interval.
+	 * @param Callback The function to be called for each log message with the message's information. Return True to signify a stop, False to continue.
+	 *
+	 * @note The current number of messages could be different than the value returned by the last call to GetMessageCount().
+	 */
+	virtual void ReverseEnumerateMessages(double StartTime, double EndTime, TFunctionRef<bool(const FLogMessageInfo&)> Callback) const = 0;
+
 	/**
 	 * Performs binary search, resulting in position of the first log message with time >= provided time value.
 	 *
