[Insights] Added timer id column for Timers, Callers and Callees panels.

diff --git a/Engine/Source/Developer/TraceInsights/Private/Insights/TimingProfiler/ViewModels/TimersViewColumnFactory.cpp b/Engine/Source/Developer/TraceInsights/Private/Insights/TimingProfiler/ViewModels/TimersViewColumnFactory.cpp
--- a/Engine/Source/Developer/TraceInsights/Private/Insights/TimingProfiler/ViewModels/TimersViewColumnFactory.cpp
+++ b/Engine/Source/Developer/TraceInsights/Private/Insights/TimingProfiler/ViewModels/TimersViewColumnFactory.cpp
@@ -21,6 +21,7 @@ namespace UE::Insights::TimingProfiler
 // Column identifiers
 
 const FName FTimersViewColumns::NameColumnID(TEXT("Name")); // TEXT("_Hierarchy")
+const FName FTimersViewColumns::IdColumnID(TEXT("Id"));
 const FName FTimersViewColumns::MetaGroupNameColumnID(TEXT("MetaGroupName"));
 const FName FTimersViewColumns::TypeColumnID(TEXT("Type"));
 const FName FTimersViewColumns::InstanceCountColumnID(TEXT("Count"));
@@ -51,6 +52,7 @@ void FTimersViewColumnFactory::CreateTimersViewColumns(TArray<TSharedRef<FTableC
 	Columns.Reset();
 
 	Columns.Add(CreateNameColumn());
+	Columns.Add(CreateIdColumn());
 	Columns.Add(CreateMetaGroupNameColumn());
 	Columns.Add(CreateTypeColumn());
 	Columns.Add(CreateInstanceCountColumn());
@@ -75,6 +77,7 @@ void FTimersViewColumnFactory::CreateTimerTreeViewColumns(TArray<TSharedRef<FTab
 	Columns.Reset();
 
 	Columns.Add(CreateNameColumn());
+	Columns.Add(CreateIdColumn());
 	Columns.Add(CreateTypeColumn());
 	Columns.Add(CreateInstanceCountColumn());
 	Columns.Add(CreateChildInstanceCountColumn());
@@ -123,6 +126,61 @@ TSharedRef<FTableColumn> FTimersViewColumnFactory::CreateNameColumn()
 
 ////////////////////////////////////////////////////////////////////////////////////////////////////
 
+TSharedRef<FTableColumn> FTimersViewColumnFactory::CreateIdColumn()
+{
+	TSharedRef<FTimersTableColumn> ColumnRef = MakeShared<FTimersTableColumn>(FTimersViewColumns::IdColumnID);
+	FTimersTableColumn& Column = *ColumnRef;
+
+	Column.SetShortName(LOCTEXT("Id_ColumnName", "Id"));
+	Column.SetTitleName(LOCTEXT("Id_ColumnTitle", "Timer Id"));
+
+	FText Description = LOCTEXT("Id_ColumnDesc", "Temporary index/id assigned to timers during analysis");
+	Column.SetDescription(ETraceFrameType::TraceFrameType_Count, Description);
+	Column.SetDescription(ETraceFrameType::TraceFrameType_Game, Description);
+	Column.SetDescription(ETraceFrameType::TraceFrameType_Rendering, Description);
+
+	Column.SetFlags(ETableColumnFlags::CanBeHidden |
+					ETableColumnFlags::CanBeFiltered);
+
+	Column.SetHorizontalAlignment(HAlign_Right);
+	Column.SetInitialWidth(60.0f);
+
+	Column.SetDataType(ETableCellDataType::Int64);
+
+	class FIdValueGetter : public FTableCellValueGetter
+	{
+	public:
+		virtual const TOptional<FTableCellValue> GetValue(const FTableColumn& Column, const FBaseTreeNode& Node) const
+		{
+			ensure(Column.GetId() == FTimersViewColumns::IdColumnID);
+			const FTimerNode& TimerNode = static_cast<const FTimerNode&>(Node);
+			const uint32 TimerId = TimerNode.GetTimerId();
+			if (TimerId != FTimerNode::InvalidTimerId)
+			{
+				return TOptional<FTableCellValue>(FTableCellValue(static_cast<int64>(TimerId)));
+			}
+			else
+			{
+				return TOptional<FTableCellValue>();
+			}
+		}
+	};
+
+	TSharedRef<ITableCellValueGetter> Getter = MakeShared<FIdValueGetter>();
+	Column.SetValueGetter(Getter);
+
+	TSharedRef<ITableCellValueFormatter> Formatter = MakeShared<FInt64ValueFormatterAsNumber>();
+	Column.SetValueFormatter(Formatter);
+
+	TSharedRef<ITableCellValueSorter> Sorter = MakeShared<FSorterByInt64Value>(ColumnRef);
+	Column.SetValueSorter(Sorter);
+	Column.SetInitialSortMode(EColumnSortMode::Ascending);
+
+	return ColumnRef;
+}
+
+////////////////////////////////////////////////////////////////////////////////////////////////////
+
 TSharedRef<FTableColumn> FTimersViewColumnFactory::CreateMetaGroupNameColumn()
 {
 	TSharedRef<FTimersTableColumn> ColumnRef = MakeShared<FTimersTableColumn>(FTimersViewColumns::MetaGroupNameColumnID);
diff --git a/Engine/Source/Developer/TraceInsights/Private/Insights/TimingProfiler/ViewModels/TimersViewColumnFactory.h b/Engine/Source/Developer/TraceInsights/Private/Insights/TimingProfiler/ViewModels/TimersViewColumnFactory.h
--- a/Engine/Source/Developer/TraceInsights/Private/Insights/TimingProfiler/ViewModels/TimersViewColumnFactory.h
+++ b/Engine/Source/Developer/TraceInsights/Private/Insights/TimingProfiler/ViewModels/TimersViewColumnFactory.h
@@ -20,6 +20,7 @@ namespace UE::Insights::TimingProfiler
 struct FTimersViewColumns
 {
 	static const FName NameColumnID;
+	static const FName IdColumnID;
 	static const FName MetaGroupNameColumnID;
 	static const FName TypeColumnID;
 	static const FName InstanceCountColumnID;
@@ -100,6 +101,7 @@ public:
 	static void CreateTimerTreeViewColumns(TArray<TSharedRef<FTableColumn>>& Columns);
 
 	static TSharedRef<FTableColumn> CreateNameColumn();
+	static TSharedRef<FTableColumn> CreateIdColumn();
 	static TSharedRef<FTableColumn> CreateMetaGroupNameColumn();
 	static TSharedRef<FTableColumn> CreateTypeColumn();
 	static TSharedRef<FTableColumn> CreateInstanceCountColumn();
