Timing Event Copy function renamed to Copy Event To Clipboard

diff --git a/Engine/Source/Developer/TraceInsights/Private/Insights/Widgets/STimingView.cpp b/Engine/Source/Developer/TraceInsights/Private/Insights/Widgets/STimingView.cpp
--- a/Engine/Source/Developer/TraceInsights/Private/Insights/Widgets/STimingView.cpp
+++ b/Engine/Source/Developer/TraceInsights/Private/Insights/Widgets/STimingView.cpp
@@ -3052,7 +3052,7 @@ void STimingView::ShowContextMenu(const FPointerEvent& MouseEvent)
 			);
 
 			MenuBuilder.AddMenuEntry(
-				LOCTEXT("ContextMenu_Copy", "Copy"),
+				LOCTEXT("ContextMenu_Copy", "Copy Event To Clipboard"),
 				FText(),
 				FSlateIcon(FAppStyle::Get().GetStyleSetName(), "GenericCommands.Copy"),
 				FUIAction(FExecuteAction::CreateLambda([Event=HoveredEvent]()
