Insights CLI Flag to prevent AutoQuit until all symbols have resolved.

-AutoQuit -WaitForSymbolResolver

diff --git a/Engine/Source/Developer/TraceInsights/Private/Insights/InsightsManager.cpp b/Engine/Source/Developer/TraceInsights/Private/Insights/InsightsManager.cpp
--- a/Engine/Source/Developer/TraceInsights/Private/Insights/InsightsManager.cpp
+++ b/Engine/Source/Developer/TraceInsights/Private/Insights/InsightsManager.cpp
@@ -29,6 +29,7 @@
 // TraceServices
 #include "TraceServices/AnalysisService.h"
 #include "TraceServices/Model/Diagnostics.h"
+#include "TraceServices/Model/Modules.h"
 #include "TraceServices/Model/NetProfiler.h"
 
 // TraceInsightsCore
@@ -394,6 +395,19 @@ void FInsightsManager::UpdateSessionDuration()
 
 			OnSessionAnalysisCompleted();
 		}
+
+		if (bIsAnalysisComplete && !bIsSymbolResolvingComplete)
+		{
+			const TraceServices::IModuleProvider* ModuleProvider = ReadModuleProvider(*Session.Get());
+			if (ModuleProvider)
+			{
+				if (ModuleProvider->HasFinishedResolving())
+				{
+					bIsSymbolResolvingComplete = true;
+					OnSessionSymbolResolverCompleted();
+				}
+			}
+		}
 	}
 }
 
@@ -468,6 +482,7 @@ void FInsightsManager::ResetSession(bool bNotify)
 
 	bIsSessionInfoSet = false;
 	bIsAnalysisComplete = false;
+	bIsSymbolResolvingComplete = false;
 	SessionDuration = 0.0;
 	AnalysisStopwatch.Restart();
 	AnalysisDuration = 0.0;
@@ -713,7 +728,7 @@ void FInsightsManager::LoadLastLiveSession(float RetryTime)
 
 ////////////////////////////////////////////////////////////////////////////////////////////////////
 
-void FInsightsManager::LoadTrace(uint32 InTraceId, bool InAutoQuit)
+void FInsightsManager::LoadTrace(uint32 InTraceId, bool InAutoQuit, bool InWaitForSymbolResolver)
 {
 	ResetSession();
 
@@ -770,6 +785,7 @@ void FInsightsManager::LoadTrace(uint32 InTraceId, bool InAutoQuit)
 		bIsSessionInfoSet = false;
 		OnSessionChanged();
 		bSessionAnalysisCompletedAutoQuit = InAutoQuit;
+		bWaitForSymbolResolver = InAutoQuit && InWaitForSymbolResolver;
 	}
 	else if (InAutoQuit)
 	{
@@ -790,7 +806,7 @@ void FInsightsManager::LoadTraceFile()
 
 ////////////////////////////////////////////////////////////////////////////////////////////////////
 
-void FInsightsManager::LoadTraceFile(const FString& InTraceFilename, bool InAutoQuit)
+void FInsightsManager::LoadTraceFile(const FString& InTraceFilename, bool InAutoQuit, bool InWaitForSymbolResolver)
 {
 	IPlatformFile& PlatformFile = IPlatformFile::GetPlatformPhysical();
 	if (!PlatformFile.FileExists(*InTraceFilename))
@@ -810,6 +826,7 @@ void FInsightsManager::LoadTraceFile(const FString& InTraceFilename, bool InAuto
 		bIsSessionInfoSet = false;
 		OnSessionChanged();
 		bSessionAnalysisCompletedAutoQuit = InAutoQuit;
+		bWaitForSymbolResolver = InAutoQuit && InWaitForSymbolResolver;
 	}
 	else if (InAutoQuit)
 	{
@@ -976,7 +993,7 @@ void FInsightsManager::OnSessionAnalysisCompleted()
 	}
 #endif
 
-	if (bSessionAnalysisCompletedAutoQuit)
+	if (bSessionAnalysisCompletedAutoQuit && !bWaitForSymbolResolver)
 	{
 		bSessionAnalysisCompletedAutoQuit = false;
 		RequestEngineExit(AutoQuitMsg);
@@ -985,6 +1002,17 @@ void FInsightsManager::OnSessionAnalysisCompleted()
 
 ////////////////////////////////////////////////////////////////////////////////////////////////////
 
+void FInsightsManager::OnSessionSymbolResolverCompleted()
+{
+	if (bWaitForSymbolResolver)
+	{
+		bWaitForSymbolResolver = false;
+		RequestEngineExit(AutoQuitMsg);
+	}
+}
+
+////////////////////////////////////////////////////////////////////////////////////////////////////
+
 bool FInsightsManager::Exec(const TCHAR* Cmd, FOutputDevice& Ar)
 {
 	// @=ResponseFile
diff --git a/Engine/Source/Developer/TraceInsights/Private/Insights/InsightsManager.h b/Engine/Source/Developer/TraceInsights/Private/Insights/InsightsManager.h
--- a/Engine/Source/Developer/TraceInsights/Private/Insights/InsightsManager.h
+++ b/Engine/Source/Developer/TraceInsights/Private/Insights/InsightsManager.h
@@ -209,8 +209,9 @@ public:
 	 * Replaces the current analysis session.
 	 * @param TraceId - The id of the trace to analyze
 	 * @param InAutoQuit - The Application will close when session analysis is complete or fails to start
+	 * @param InWaitForSymbolResolver - The Application will automatically close when both session analysis is complete and symbols have been resolved. Must be set together with the InAutoQuit parameter
 	 */
-	void LoadTrace(uint32 TraceId, bool InAutoQuit = false);
+	void LoadTrace(uint32 TraceId, bool InAutoQuit = false, bool InWaitForSymbolResolver = false);
 
 	/**
 	 * Shows the open file dialog and creates a new analysis session instance for the chosen trace file.
@@ -223,8 +224,9 @@ public:
 	 * Replaces the current analysis session.
 	 * @param TraceFilename - The trace file to analyze
 	 * @param InAutoQuit - The Application will close when session analysis is complete or fails to start
+	 * @param InWaitForSymbolResolver - The Application will automatically close when both session analysis is complete and symbols have been resolved. Must be set together with the InAutoQuit parameter
 	 */
-	void LoadTraceFile(const FString& TraceFilename, bool InAutoQuit = false);
+	void LoadTraceFile(const FString& TraceFilename, bool InAutoQuit = false, bool InWaitForSymbolResolver = false);
 
 
 	uint16 ListenForDirectTrace();
@@ -296,6 +298,7 @@ private:
 
 	void OnSessionChanged();
 	void OnSessionAnalysisCompleted();
+	void OnSessionSymbolResolverCompleted();
 
 	void SpawnAndActivateTabs();
 
@@ -362,6 +365,9 @@ private:
 
 	bool bIsAnalysisComplete = false;
 	bool bSessionAnalysisCompletedAutoQuit = false;
+	
+	bool bIsSymbolResolvingComplete = false;
+	bool bWaitForSymbolResolver = false;
 
 	float RetryLoadLastLiveSessionTimer = 0.0f;
 	bool bIsAutoLoadLiveSessionEnabled = false;
diff --git a/Engine/Source/Developer/TraceInsights/Private/Insights/TraceInsightsModule.cpp b/Engine/Source/Developer/TraceInsights/Private/Insights/TraceInsightsModule.cpp
--- a/Engine/Source/Developer/TraceInsights/Private/Insights/TraceInsightsModule.cpp
+++ b/Engine/Source/Developer/TraceInsights/Private/Insights/TraceInsightsModule.cpp
@@ -465,11 +465,11 @@ TSharedPtr<const TraceServices::IAnalysisSession> FTraceInsightsModule::GetAnaly
 
 ////////////////////////////////////////////////////////////////////////////////////////////////////
 
-void FTraceInsightsModule::StartAnalysisForTrace(uint32 InTraceId, bool InAutoQuit)
+void FTraceInsightsModule::StartAnalysisForTrace(uint32 InTraceId, bool InAutoQuit, bool InWaitForSymbolResolver)
 {
 	if (InTraceId != 0)
 	{
-		FInsightsManager::Get()->LoadTrace(InTraceId, InAutoQuit);
+		FInsightsManager::Get()->LoadTrace(InTraceId, InAutoQuit, InWaitForSymbolResolver);
 	}
 	UpdateAppTitle();
 }
@@ -484,11 +484,11 @@ void FTraceInsightsModule::StartAnalysisForLastLiveSession(float InRetryTime)
 
 ////////////////////////////////////////////////////////////////////////////////////////////////////
 
-void FTraceInsightsModule::StartAnalysisForTraceFile(const TCHAR* InTraceFile, bool InAutoQuit)
+void FTraceInsightsModule::StartAnalysisForTraceFile(const TCHAR* InTraceFile, bool InAutoQuit, bool InWaitForSymbolResolver)
 {
 	if (InTraceFile != nullptr)
 	{
-		FInsightsManager::Get()->LoadTraceFile(FString(InTraceFile), InAutoQuit);
+		FInsightsManager::Get()->LoadTraceFile(FString(InTraceFile), InAutoQuit, InWaitForSymbolResolver);
 	}
 	UpdateAppTitle();
 }
diff --git a/Engine/Source/Developer/TraceInsights/Private/Insights/TraceInsightsModule.h b/Engine/Source/Developer/TraceInsights/Private/Insights/TraceInsightsModule.h
--- a/Engine/Source/Developer/TraceInsights/Private/Insights/TraceInsightsModule.h
+++ b/Engine/Source/Developer/TraceInsights/Private/Insights/TraceInsightsModule.h
@@ -55,9 +55,9 @@ public:
 	virtual void CreateSessionViewer(bool bAllowDebugTools = false) override;
 
 	virtual TSharedPtr<const TraceServices::IAnalysisSession> GetAnalysisSession() const override;
-	virtual void StartAnalysisForTrace(uint32 InTraceId, bool InAutoQuit = false) override;
+	virtual void StartAnalysisForTrace(uint32 InTraceId, bool InAutoQuit = false, bool InWaitForSymbolResolver = false) override;
 	virtual void StartAnalysisForLastLiveSession(float InRetryTime = 1.0f) override;
-	virtual void StartAnalysisForTraceFile(const TCHAR* InTraceFile, bool InAutoQuit = false) override;
+	virtual void StartAnalysisForTraceFile(const TCHAR* InTraceFile, bool InAutoQuit = false, bool InWaitForSymbolResolver = false) override;
 	virtual uint16 StartAnalysisWithDirectTrace() override;
 
 	virtual void ShutdownUserInterface() override;
diff --git a/Engine/Source/Developer/TraceInsights/Public/Insights/IUnrealInsightsModule.h b/Engine/Source/Developer/TraceInsights/Public/Insights/IUnrealInsightsModule.h
--- a/Engine/Source/Developer/TraceInsights/Public/Insights/IUnrealInsightsModule.h
+++ b/Engine/Source/Developer/TraceInsights/Public/Insights/IUnrealInsightsModule.h
@@ -254,8 +254,9 @@ public:
 	 *
 	 * @param InTraceId The id of the trace to analyze.
 	 * @param InAutoQuit - The Application will close when session analysis is complete or fails to start
+	 * @param InWaitForSymbolResolver - The Application will automatically close when both session analysis is complete and symbols have been resolved. Must be set together with the InAutoQuit parameter.
 	 */
-	virtual void StartAnalysisForTrace(uint32 InTraceId, bool InAutoQuit = false) = 0;
+	virtual void StartAnalysisForTrace(uint32 InTraceId, bool InAutoQuit = false, bool InWaitForSymbolResolver = false) = 0;
 
 	/**
 	 * Starts analysis of the last live session. Called when the application starts in "Viewer" mode.
@@ -270,8 +271,9 @@ public:
 	 *
 	 * @param InTraceFile The filename (*.utrace) of the trace to analyze.
 	 * @param InAutoQuit - The Application will close when session analysis is complete or fails to start
+	 * @param InWaitForSymbolResolver - The Application will automatically close when both session analysis is complete and symbols have been resolved. Must be set together with the InAutoQuit parameter.
 	 */
-	virtual void StartAnalysisForTraceFile(const TCHAR* InTraceFile, bool InAutoQuit = false) = 0;
+	virtual void StartAnalysisForTraceFile(const TCHAR* InTraceFile, bool InAutoQuit = false, bool InWaitForSymbolResolver = false) = 0;
 
 	/**
 	 * Starts analysis for a direct connection from the tracing application (might be self). Notice that
diff --git a/Engine/Source/Developer/TraceServices/Private/Model/DbgHelpResolver.cpp b/Engine/Source/Developer/TraceServices/Private/Model/DbgHelpResolver.cpp
--- a/Engine/Source/Developer/TraceServices/Private/Model/DbgHelpResolver.cpp
+++ b/Engine/Source/Developer/TraceServices/Private/Model/DbgHelpResolver.cpp
@@ -210,6 +210,13 @@ void FDbgHelpResolver::GetStats(IModuleProvider::FStats* OutStats) const
 
 ////////////////////////////////////////////////////////////////////////////////////////////////////
 
+bool FDbgHelpResolver::HasFinishedResolving() const
+{
+	return LoadSymbolsQueue.IsEmpty() && ResolveQueue.IsEmpty();
+}
+
+////////////////////////////////////////////////////////////////////////////////////////////////////
+
 void FDbgHelpResolver::EnumerateSymbolSearchPaths(TFunctionRef<void(FStringView Path)> Callback) const
 {
 	{
diff --git a/Engine/Source/Developer/TraceServices/Private/Model/DbgHelpResolver.h b/Engine/Source/Developer/TraceServices/Private/Model/DbgHelpResolver.h
--- a/Engine/Source/Developer/TraceServices/Private/Model/DbgHelpResolver.h
+++ b/Engine/Source/Developer/TraceServices/Private/Model/DbgHelpResolver.h
@@ -31,6 +31,7 @@ public:
 	void QueueSymbolResolve(uint64 Address, FResolvedSymbol* Symbol);
 	void OnAnalysisComplete();
 	void GetStats(IModuleProvider::FStats* OutStats) const;
+	bool HasFinishedResolving() const;
 	void EnumerateSymbolSearchPaths(TFunctionRef<void(FStringView Path)> Callback) const;
 
 private:
diff --git a/Engine/Source/Developer/TraceServices/Private/Model/ModuleProvider.cpp b/Engine/Source/Developer/TraceServices/Private/Model/ModuleProvider.cpp
--- a/Engine/Source/Developer/TraceServices/Private/Model/ModuleProvider.cpp
+++ b/Engine/Source/Developer/TraceServices/Private/Model/ModuleProvider.cpp
@@ -71,6 +71,7 @@ public:
 	FGraphEventRef LoadSymbolsForModuleUsingPath(uint64 Base, const TCHAR* Path) override;
 	void EnumerateSymbolSearchPaths(TFunctionRef<void(FStringView Path)> Callback) const override;
 	void GetStats(FStats* OutStats) const override;
+	bool HasFinishedResolving() const override;
 
 	// Analysis interface
 	void OnModuleLoad(const FStringView& Module, uint64 Base, uint32 Size, const uint8* Checksum, uint32 ChecksumSize) override;
@@ -326,6 +327,14 @@ void TModuleProvider<SymbolResolverType>::GetStats(FStats* OutStats) const
 	OutStats->SymbolsResolved += NumCachedSymbols;
 }
 
+////////////////////////////////////////////////////////////////////////////////////////////////////
+	
+template<typename SymbolResolverType>
+bool TModuleProvider<SymbolResolverType>::HasFinishedResolving() const
+{
+	return Resolver->HasFinishedResolving();
+}
+
 ////////////////////////////////////////////////////////////////////////////////////////////////////
 
 template<typename SymbolResolverType>
diff --git a/Engine/Source/Developer/TraceServices/Private/Model/PsymResolver.cpp b/Engine/Source/Developer/TraceServices/Private/Model/PsymResolver.cpp
--- a/Engine/Source/Developer/TraceServices/Private/Model/PsymResolver.cpp
+++ b/Engine/Source/Developer/TraceServices/Private/Model/PsymResolver.cpp
@@ -189,6 +189,12 @@ void FPsymResolver::GetStats(IModuleProvider::FStats* OutStats) const
 	OutStats->ModulesLoaded = ModulesLoaded.load();
 }
 
+/////////////////////////////////////////////////////////////////////
+bool FPsymResolver::HasFinishedResolving() const
+{
+	return LoadSymbolsQueue.IsEmpty() && ResolveQueue.IsEmpty();
+}
+
 /////////////////////////////////////////////////////////////////////
 void FPsymResolver::EnumerateSymbolSearchPaths(TFunctionRef<void(FStringView Path)> Callback) const
 {
diff --git a/Engine/Source/Developer/TraceServices/Private/Model/PsymResolver.h b/Engine/Source/Developer/TraceServices/Private/Model/PsymResolver.h
--- a/Engine/Source/Developer/TraceServices/Private/Model/PsymResolver.h
+++ b/Engine/Source/Developer/TraceServices/Private/Model/PsymResolver.h
@@ -29,6 +29,7 @@ public:
 	void QueueSymbolResolve(uint64 Address, FResolvedSymbol* Symbol);
 	void OnAnalysisComplete();
 	void GetStats(IModuleProvider::FStats* OutStats) const;
+	bool HasFinishedResolving() const;
 	void EnumerateSymbolSearchPaths(TFunctionRef<void(FStringView Path)> Callback) const;
 
 private:
diff --git a/Engine/Source/Developer/TraceServices/Private/Model/SymslibResolver.cpp b/Engine/Source/Developer/TraceServices/Private/Model/SymslibResolver.cpp
--- a/Engine/Source/Developer/TraceServices/Private/Model/SymslibResolver.cpp
+++ b/Engine/Source/Developer/TraceServices/Private/Model/SymslibResolver.cpp
@@ -780,6 +780,17 @@ void FSymslibResolver::GetStats(IModuleProvider::FStats* OutStats) const
 
 ////////////////////////////////////////////////////////////////////////////////////////////////////
 
+bool FSymslibResolver::HasFinishedResolving() const
+{
+	if (CleanupTask)
+	{
+		return CleanupTask->IsComplete();
+	}
+	return false;
+}
+	
+////////////////////////////////////////////////////////////////////////////////////////////////////
+
 void FSymslibResolver::EnumerateSymbolSearchPaths(TFunctionRef<void(FStringView Path)> Callback) const
 {
 	{
diff --git a/Engine/Source/Developer/TraceServices/Private/Model/SymslibResolver.h b/Engine/Source/Developer/TraceServices/Private/Model/SymslibResolver.h
--- a/Engine/Source/Developer/TraceServices/Private/Model/SymslibResolver.h
+++ b/Engine/Source/Developer/TraceServices/Private/Model/SymslibResolver.h
@@ -37,6 +37,7 @@ public:
 	void QueueSymbolResolve(uint64 Address, FResolvedSymbol* Symbol);
 	void OnAnalysisComplete();
 	void GetStats(IModuleProvider::FStats* OutStats) const;
+	bool HasFinishedResolving() const;
 	void EnumerateSymbolSearchPaths(TFunctionRef<void(FStringView Path)> Callback) const;
 
 private:
diff --git a/Engine/Source/Developer/TraceServices/Public/TraceServices/Model/Modules.h b/Engine/Source/Developer/TraceServices/Public/TraceServices/Model/Modules.h
--- a/Engine/Source/Developer/TraceServices/Public/TraceServices/Model/Modules.h
+++ b/Engine/Source/Developer/TraceServices/Public/TraceServices/Model/Modules.h
@@ -212,8 +212,11 @@ public:
 	 */
 	virtual void EnumerateSymbolSearchPaths(TFunctionRef<void(FStringView Path)> Callback) const = 0;
 
-	/** Gets statistics from provider */
+	/** Gets statistics from provider. */
 	virtual void GetStats(FStats* OutStats) const = 0;
+
+	/** Checks if all symbols have been resolved in resolver queues. Only useful after analysis has completed (e.g., OnAnalysisComplete). */
+	virtual bool HasFinishedResolving() const = 0;
 };
 
 ////////////////////////////////////////////////////////////////////////////////
diff --git a/Engine/Source/Programs/UnrealInsights/Private/UserInterfaceCommand.cpp b/Engine/Source/Programs/UnrealInsights/Private/UserInterfaceCommand.cpp
--- a/Engine/Source/Programs/UnrealInsights/Private/UserInterfaceCommand.cpp
+++ b/Engine/Source/Programs/UnrealInsights/Private/UserInterfaceCommand.cpp
@@ -222,6 +222,9 @@ void FUserInterfaceCommand::InitializeSlateApplication(bool bFrontendMode, const
 		// This parameter will cause the application to close when analysis fails to start or completes successfully.
 		const bool bAutoQuit = FParse::Param(FCommandLine::Get(), TEXT("AutoQuit"));
 
+		// This parameter will prevent the application from closing until analysis is completed and symbols have been resolved. Must be used with the -AutoQuit flag to work.
+		const bool bWaitForSymbolResolver = FParse::Param(FCommandLine::Get(), TEXT("WaitForSymbolResolver"));
+
 		const bool bInitializeTesting = FParse::Param(FCommandLine::Get(), TEXT("InsightsTest"));
 		if (bInitializeTesting)
 		{
@@ -261,11 +264,11 @@ void FUserInterfaceCommand::InitializeSlateApplication(bool bFrontendMode, const
 		if (bUseTraceId)
 		{
 			TraceInsightsModule.ConnectToStore(*StoreHost, StorePort);
-			TraceInsightsModule.StartAnalysisForTrace(TraceId, bAutoQuit);
+			TraceInsightsModule.StartAnalysisForTrace(TraceId, bAutoQuit, bWaitForSymbolResolver);
 		}
 		else
 		{
-			TraceInsightsModule.StartAnalysisForTraceFile(*TraceFileToOpen, bAutoQuit);
+			TraceInsightsModule.StartAnalysisForTraceFile(*TraceFileToOpen, bAutoQuit, bWaitForSymbolResolver);
 		}
 	}
 	else // frontend mode
