Add callstack clipboard operations to Insight's memory allocation context menu.

diff --git a/Engine/Source/Developer/TraceInsights/Private/Insights/MemoryProfiler/Widgets/SMemAllocTableTreeView.cpp b/Engine/Source/Developer/TraceInsights/Private/Insights/MemoryProfiler/Widgets/SMemAllocTableTreeView.cpp
--- a/Engine/Source/Developer/TraceInsights/Private/Insights/MemoryProfiler/Widgets/SMemAllocTableTreeView.cpp
+++ b/Engine/Source/Developer/TraceInsights/Private/Insights/MemoryProfiler/Widgets/SMemAllocTableTreeView.cpp
@@ -7,6 +7,7 @@
 #include "Framework/MultiBox/MultiBoxBuilder.h"
 #include "GenericPlatform/GenericPlatformMisc.h"
 #include "HAL/PlatformFileManager.h"
+#include "HAL/PlatformApplicationMisc.h"
 #include "ISourceCodeAccessModule.h"
 #include "ISourceCodeAccessor.h"
 #include "Logging/MessageLog.h"
@@ -1252,7 +1253,7 @@ void SMemAllocTableTreeView::ExtendMenuAllocation(FMenuBuilder& MenuBuilder)
 		(
 			ItemLabel,
 			ItemToolTip,
-			FNewMenuDelegate::CreateSP(this, &SMemAllocTableTreeView::BuildOpenSourceSubMenu, true),
+			FNewMenuDelegate::CreateSP(this, &SMemAllocTableTreeView::BuildOpenSourceSubMenu, true, ESourceSubMenuType::OpenInSourceEditor),
 			false,
 			FSlateIcon(SourceCodeAccessor.GetStyleSet(), SourceCodeAccessor.GetOpenIconName())
 		);
@@ -1275,10 +1276,30 @@ void SMemAllocTableTreeView::ExtendMenuAllocation(FMenuBuilder& MenuBuilder)
 		(
 			ItemLabel,
 			ItemToolTip,
-			FNewMenuDelegate::CreateSP(this, &SMemAllocTableTreeView::BuildOpenSourceSubMenu, false),
+			FNewMenuDelegate::CreateSP(this, &SMemAllocTableTreeView::BuildOpenSourceSubMenu, false, ESourceSubMenuType::OpenInSourceEditor),
 			false,
 			FSlateIcon(SourceCodeAccessor.GetStyleSet(), SourceCodeAccessor.GetOpenIconName())
 		);
+
+		// Copy Callstack
+		MenuBuilder.AddSubMenu
+		(
+			LOCTEXT("ContextMenu_CopyAllocCallstack_SubMenu", "Copy to Clipboard | Alloc Callstack"),
+			LOCTEXT("ContextMenu_CopyAllocCallstack_SubMenu_Desc", "Copy the source file path and line of the selected allocation callstack frame to the clipboard."),
+			FNewMenuDelegate::CreateSP(this, &SMemAllocTableTreeView::BuildOpenSourceSubMenu, true, ESourceSubMenuType::CopyToClipboard),
+			false,
+			FSlateIcon(FAppStyle::GetAppStyleSetName(), "GenericCommands.Copy")
+		);
+
+		// Free Callstack
+		MenuBuilder.AddSubMenu
+		(
+			LOCTEXT("ContextMenu_CopyFreeCallstack_SubMenu", "Copy to Clipboard | Free Callstack"),
+			LOCTEXT("ContextMenu_CopyFreeCallstack_SubMenu_Desc", "Copy the source file path and line of the selected free callstack frame to the clipboard."),
+			FNewMenuDelegate::CreateSP(this, &SMemAllocTableTreeView::BuildOpenSourceSubMenu, false, ESourceSubMenuType::CopyToClipboard),
+			false,
+			FSlateIcon(FAppStyle::GetAppStyleSetName(), "GenericCommands.Copy")
+		);
 	}
 	MenuBuilder.EndSection();
 }
@@ -1298,12 +1319,12 @@ void SMemAllocTableTreeView::ExtendMenuCallstackFrame(FMenuBuilder & MenuBuilder
 
 	MenuBuilder.BeginSection("CallstackFrame", LOCTEXT("ContextMenu_Section_CallstackFrame", "Callstack Frame"));
 	{
+		FText FileName = GetSelectedCallstackFrameFileName();
 		if (SourceCodeAccessor.CanAccessSourceCode())
 		{
 			FText ItemLabel = FText::Format(LOCTEXT("ContextMenu_OpenSourceFile", "Open Source File in {0}"),
 				SourceCodeAccessor.GetNameText());
 
-			FText FileName = GetSelectedCallstackFrameFileName();
 			FText ItemToolTip = FText::Format(LOCTEXT("ContextMenu_OpenSourceFile_Desc", "Opens the source file of the selected callstack frame in {0}.\n{1}"),
 				SourceCodeAccessor.GetNameText(),
 				FileName);
@@ -1320,7 +1341,6 @@ void SMemAllocTableTreeView::ExtendMenuCallstackFrame(FMenuBuilder & MenuBuilder
 		{
 			FText ItemLabel = LOCTEXT("ContextMenu_OpenSourceFile_NoAccessor", "Open Source File");
 
-			FText FileName = GetSelectedCallstackFrameFileName();
 			FText ItemToolTip = FText::Format(LOCTEXT("ContextMenu_OpenSourceFile_NoAccessor_Desc_Fmt", "{0}\nSource Code Accessor is not available."),
 				FileName);
 
@@ -1332,6 +1352,16 @@ void SMemAllocTableTreeView::ExtendMenuCallstackFrame(FMenuBuilder & MenuBuilder
 				NAME_None,
 				EUserInterfaceActionType::None);
 		}
+
+		{
+			MenuBuilder.AddMenuEntry(
+				LOCTEXT("ContextMenu_CopySourceFilePath_ToClipboard", "Copy Source File Path to Clipboard"),
+				FileName,
+				FSlateIcon(FAppStyle::GetAppStyleSetName(), "GenericCommands.Copy"),
+				FUIAction(
+					FExecuteAction::CreateSP(this, &SMemAllocTableTreeView::CopySelectedFilePathToClipboard),
+					FCanExecuteAction() ));
+		}
 	}
 	MenuBuilder.EndSection();
 }
@@ -1385,7 +1415,7 @@ uint32 SMemAllocTableTreeView::CountSourceFiles(FMemAllocNode& MemAllocNode)
 
 ////////////////////////////////////////////////////////////////////////////////////////////////////
 
-bool SMemAllocTableTreeView::BuildOpenSourceSubMenuItems(FMenuBuilder& MenuBuilder, const TraceServices::FCallstack& Callstack)
+bool SMemAllocTableTreeView::BuildOpenSourceSubMenuItems(FMenuBuilder& MenuBuilder, const TraceServices::FCallstack& Callstack, ESourceSubMenuType SourceSubMenuType)
 {
 	ISourceCodeAccessModule& SourceCodeAccessModule = FModuleManager::LoadModuleChecked<ISourceCodeAccessModule>("SourceCodeAccess");
 	ISourceCodeAccessor& SourceCodeAccessor = SourceCodeAccessModule.GetAccessor();
@@ -1467,17 +1497,37 @@ bool SMemAllocTableTreeView::BuildOpenSourceSubMenuItems(FMenuBuilder& MenuBuild
 				}
 			}
 
-			const bool bCanOpenSource = SourceCodeAccessor.CanAccessSourceCode() && FPaths::FileExists(Frame->Symbol->File);
-
-			MenuBuilder.AddMenuEntry(
-				ItemLabel,
-				ItemToolTip,
-				FSlateIcon(),
-				FUIAction(
-					FExecuteAction::CreateSP(this, &SMemAllocTableTreeView::OpenSourceFileInIDE, Frame->Symbol->File, uint32(Frame->Symbol->Line)),
-					FCanExecuteAction::CreateLambda([bCanOpenSource]() { return bCanOpenSource; })),
-				NAME_None,
-				EUserInterfaceActionType::Button);
+			switch (SourceSubMenuType)
+			{
+				case ESourceSubMenuType::OpenInSourceEditor:
+				{
+					const bool bCanOpenSource = SourceCodeAccessor.CanAccessSourceCode() && FPaths::FileExists(Frame->Symbol->File);
+
+					MenuBuilder.AddMenuEntry(
+						ItemLabel,
+						ItemToolTip,
+						FSlateIcon(),
+						FUIAction(
+							FExecuteAction::CreateSP(this, &SMemAllocTableTreeView::OpenSourceFileInIDE, Frame->Symbol->File, uint32(Frame->Symbol->Line)),
+							FCanExecuteAction::CreateLambda([bCanOpenSource]() { return bCanOpenSource; })),
+						NAME_None,
+						EUserInterfaceActionType::Button);
+					break;
+				}
+				case ESourceSubMenuType::CopyToClipboard:
+				{
+					MenuBuilder.AddMenuEntry(
+						ItemLabel,
+						LOCTEXT("ContextMenu_CopySourceCode_Desc", "Copy the source file path and line to the clipboard."),
+						FSlateIcon(),
+						FUIAction(
+							FExecuteAction::CreateSP(this, &SMemAllocTableTreeView::CopySourceFileToClipboard, Frame->Symbol->File, uint32(Frame->Symbol->Line)),
+							FCanExecuteAction()),
+						NAME_None,
+						EUserInterfaceActionType::Button);
+						break;
+				}
+			}
 
 			++NumSourceFiles;
 		}
@@ -1488,7 +1538,7 @@ bool SMemAllocTableTreeView::BuildOpenSourceSubMenuItems(FMenuBuilder& MenuBuild
 
 ////////////////////////////////////////////////////////////////////////////////////////////////////
 
-void SMemAllocTableTreeView::BuildOpenSourceSubMenu(FMenuBuilder& MenuBuilder, bool bIsAllocCallstack)
+void SMemAllocTableTreeView::BuildOpenSourceSubMenu(FMenuBuilder& MenuBuilder, bool bIsAllocCallstack, ESourceSubMenuType SourceSubMenuType)
 {
 	MenuBuilder.BeginSection("OpenSource");
 	{
@@ -1502,7 +1552,7 @@ void SMemAllocTableTreeView::BuildOpenSourceSubMenu(FMenuBuilder& MenuBuilder, b
 				const TraceServices::FCallstack* Callstack = bIsAllocCallstack ? Alloc->AllocCallstack : Alloc->FreeCallstack;
 				if (Callstack)
 				{
-					bHasAnySourceFilesToOpen = BuildOpenSourceSubMenuItems(MenuBuilder, *Callstack);
+					bHasAnySourceFilesToOpen = BuildOpenSourceSubMenuItems(MenuBuilder, *Callstack, SourceSubMenuType);
 				}
 			}
 		}
@@ -1541,6 +1591,20 @@ void SMemAllocTableTreeView::OpenSourceFileInIDE(const TCHAR* InFile, uint32 Lin
 
 ////////////////////////////////////////////////////////////////////////////////////////////////////
 
+void SMemAllocTableTreeView::CopySourceFileToClipboard(const TCHAR* File, uint32 Line) const
+{
+	FPlatformApplicationMisc::ClipboardCopy(*FString::Format(TEXT("{0} ({1})"), { File, Line }));
+}
+
+////////////////////////////////////////////////////////////////////////////////////////////////////
+
+void SMemAllocTableTreeView::CopySelectedFilePathToClipboard() const
+{
+	FPlatformApplicationMisc::ClipboardCopy(*GetSelectedCallstackFrameFileName().ToString());
+}
+
+////////////////////////////////////////////////////////////////////////////////////////////////////
+
 void SMemAllocTableTreeView::ExportMemorySnapshot() const
 {
 	// 1. Choose file
@@ -1626,7 +1690,11 @@ void SMemAllocTableTreeView::ExportMemorySnapshot() const
 		};
 	auto WriteCallstackColumn = [Separator](FStringBuilderBase& OutData, const FTableTreeNode& Node, bool bIsAllocCallstack)
 		{
-			const FMemAllocNode& MemAllocNode = static_cast<const FMemAllocNode&>(Node);
+			if (!Node.Is<FMemAllocNode>())
+			{
+				return;
+			}
+			const FMemAllocNode& MemAllocNode = Node.As<FMemAllocNode>();
 			const FMemoryAlloc* Alloc = MemAllocNode.GetMemAlloc();
 			if (Alloc)
 			{
diff --git a/Engine/Source/Developer/TraceInsights/Private/Insights/MemoryProfiler/Widgets/SMemAllocTableTreeView.h b/Engine/Source/Developer/TraceInsights/Private/Insights/MemoryProfiler/Widgets/SMemAllocTableTreeView.h
--- a/Engine/Source/Developer/TraceInsights/Private/Insights/MemoryProfiler/Widgets/SMemAllocTableTreeView.h
+++ b/Engine/Source/Developer/TraceInsights/Private/Insights/MemoryProfiler/Widgets/SMemAllocTableTreeView.h
@@ -119,12 +119,16 @@ protected:
 	uint32 CountSourceFiles(FMemAllocNode& MemAllocNode);
 
 	virtual void ExtendMenu(TSharedRef<FExtender> Extender) override;
-	bool BuildOpenSourceSubMenuItems(FMenuBuilder& MenuBuilder, const TraceServices::FCallstack& Callstack);
-	void BuildOpenSourceSubMenu(FMenuBuilder& MenuBuilder, bool bIsAllocCallstack);
+	enum class ESourceSubMenuType : uint8 { OpenInSourceEditor, CopyToClipboard };
+	bool BuildOpenSourceSubMenuItems(FMenuBuilder& MenuBuilder, const TraceServices::FCallstack& Callstack, ESourceSubMenuType SourceSubMenuType);
+	void BuildOpenSourceSubMenu(FMenuBuilder& MenuBuilder, bool bIsAllocCallstack, ESourceSubMenuType SourceSubMenuType);
 	bool CanOpenCallstackFrameSourceFileInIDE() const;
 	void OpenCallstackFrameSourceFileInIDE();
 	FText GetSelectedCallstackFrameFileName() const;
 	void OpenSourceFileInIDE(const TCHAR* File, uint32 Line) const;
+	void CopySourceFileToClipboard(const TCHAR* File, uint32 Line) const;
+	void CopySelectedFilePathToClipboard() const;
+
 	void ExportMemorySnapshot() const;
 	bool IsExportMemorySnapshotAvailable() const;
 
