Update Insights Swap viewer to show the callstacks of allocations that have been moved to swap memory.

diff --git a/Engine/Source/Developer/TraceInsights/Private/Insights/MemoryProfiler/ViewModels/MemAllocGroupingByCallstack.cpp b/Engine/Source/Developer/TraceInsights/Private/Insights/MemoryProfiler/ViewModels/MemAllocGroupingByCallstack.cpp
--- a/Engine/Source/Developer/TraceInsights/Private/Insights/MemoryProfiler/ViewModels/MemAllocGroupingByCallstack.cpp
+++ b/Engine/Source/Developer/TraceInsights/Private/Insights/MemoryProfiler/ViewModels/MemAllocGroupingByCallstack.cpp
@@ -66,24 +66,35 @@ INSIGHTS_IMPLEMENT_RTTI(FMemAllocGroupingByCallstack)
 
 ////////////////////////////////////////////////////////////////////////////////////////////////////
 
+FMemAllocGroupingByCallstack::FMemAllocGroupingByCallstack(const FText& InShortName, const FText& InTitleName, const FText& InDescription, const FSlateBrush* InIcon, bool bInIsAllocCallstack, bool bInIsInverted, bool bInIsGroupingByFunction)
+	: FTreeNodeGrouping(InShortName, InTitleName, InDescription, InIcon)
+		, bIsAllocCallstack(bInIsAllocCallstack)
+		, bIsInverted(bInIsInverted)
+		, bIsGroupingByFunction(bInIsGroupingByFunction)
+		, bShouldSkipFilteredFrames(bIsInverted)
+{
+}
+
+////////////////////////////////////////////////////////////////////////////////////////////////////
+
 FMemAllocGroupingByCallstack::FMemAllocGroupingByCallstack(bool bInIsAllocCallstack, bool bInIsInverted, bool bInIsGroupingByFunction)
-	: FTreeNodeGrouping(
-		bInIsAllocCallstack
-		? (bInIsInverted ? LOCTEXT("Grouping_ByCallstack2_ShortName", "Inverted Alloc Callstack")
-						 : LOCTEXT("Grouping_ByCallstack1_ShortName", "Alloc Callstack"))
-		: (bInIsInverted ? LOCTEXT("Grouping_ByCallstack4_ShortName", "Inverted Free Callstack")
-						 : LOCTEXT("Grouping_ByCallstack3_ShortName", "Free Callstack")),
+	: FMemAllocGroupingByCallstack(
 		bInIsAllocCallstack
-		? (bInIsInverted ? LOCTEXT("Grouping_ByCallstack2_TitleName", "By Inverted Alloc Callstack")
-						 : LOCTEXT("Grouping_ByCallstack1_TitleName", "By Alloc Callstack"))
-		: (bInIsInverted ? LOCTEXT("Grouping_ByCallstack4_TitleName", "By Inverted Free Callstack")
-						 : LOCTEXT("Grouping_ByCallstack3_TitleName", "By Free Callstack")),
-		LOCTEXT("Grouping_Callstack_Desc", "Creates a tree based on callstack of each allocation."),
-		nullptr)
-	, bIsAllocCallstack(bInIsAllocCallstack)
-	, bIsInverted(bInIsInverted)
-	, bIsGroupingByFunction(bInIsGroupingByFunction)
-	, bShouldSkipFilteredFrames(bIsInverted)
+			? (bInIsInverted ? LOCTEXT("Grouping_ByCallstack2_ShortName", "Inverted Alloc Callstack")
+							 : LOCTEXT("Grouping_ByCallstack1_ShortName", "Alloc Callstack"))
+			: (bInIsInverted ? LOCTEXT("Grouping_ByCallstack4_ShortName", "Inverted Free Callstack")
+							 : LOCTEXT("Grouping_ByCallstack3_ShortName", "Free Callstack"))
+		, bInIsAllocCallstack
+			? (bInIsInverted ? LOCTEXT("Grouping_ByCallstack2_TitleName", "By Inverted Alloc Callstack")
+							 : LOCTEXT("Grouping_ByCallstack1_TitleName", "By Alloc Callstack"))
+			: (bInIsInverted ? LOCTEXT("Grouping_ByCallstack4_TitleName", "By Inverted Free Callstack")
+							 : LOCTEXT("Grouping_ByCallstack3_TitleName", "By Free Callstack"))
+		,LOCTEXT("Grouping_Callstack_Desc", "Creates a tree based on callstack of each allocation.")
+		,nullptr
+	, bInIsAllocCallstack
+	, bInIsInverted
+	, bInIsGroupingByFunction
+	)
 {
 }
 
@@ -100,6 +111,18 @@ void FMemAllocGroupingByCallstack::GroupNodes(
 	FTableTreeNode& ParentGroup,
 	TWeakPtr<FTable> InParentTable,
 	IAsyncOperationProgress& InAsyncOperationProgress) const
+{
+	GroupNodesInternal(Nodes, ParentGroup, InParentTable, InAsyncOperationProgress, [](FTableTreeNodePtr& InTableTreeNodePtr) { return InTableTreeNodePtr; });
+}
+
+////////////////////////////////////////////////////////////////////////////////////////////////////
+
+void FMemAllocGroupingByCallstack::GroupNodesInternal(
+	const TArray<FTableTreeNodePtr>& Nodes,
+	FTableTreeNode& ParentGroup,
+	TWeakPtr<FTable> InParentTable,
+	IAsyncOperationProgress& InAsyncOperationProgress,
+	TUniqueFunction<FTableTreeNodePtr(FTableTreeNodePtr&)> ProcessNode) const
 {
 	const bool bLocalIsGroupingByFunction = bIsGroupingByFunction;
 	const bool bLocalShouldSkipFilteredFrames = bShouldSkipFilteredFrames;
@@ -107,6 +130,13 @@ void FMemAllocGroupingByCallstack::GroupNodes(
 	ParentGroup.ClearChildren();
 
 	TArray<FCallstackGroup*> CallstackGroups;
+	ON_SCOPE_EXIT
+	{
+		for (FCallstackGroup* Group : CallstackGroups)
+		{
+			delete Group;
+		}
+	};
 
 	FCallstackGroup* Root = new FCallstackGroup();
 	Root->Node = &ParentGroup;
@@ -130,16 +160,27 @@ void FMemAllocGroupingByCallstack::GroupNodes(
 			continue;
 		}
 
+		if (!NodePtr->Is<FMemAllocNode>())
+		{
+			continue;
+		}
+
 		FCallstackGroup* GroupPtr = Root;
 		bool bIsEmptyCallstack = false;
 		bool bIsDefaultEmptyCallstack = false;
 
-		const FMemAllocNode& MemAllocNode = static_cast<const FMemAllocNode&>(*NodePtr);
+		const FMemAllocNode& MemAllocNode = NodePtr->As<FMemAllocNode>();
 		const FMemoryAlloc* Alloc = MemAllocNode.GetMemAlloc();
 		if (Alloc)
 		{
 			const TraceServices::FCallstack* Callstack = bIsAllocCallstack ? Alloc->GetAllocCallstack() : Alloc->GetFreeCallstack();
 
+			NodePtr = ProcessNode(NodePtr);
+			if (!NodePtr.IsValid())
+			{
+				continue;
+			}
+
 			FCallstackGroup** FoundGroupPtrPtr = GroupMapByCallstack.Find(Callstack);
 			if (FoundGroupPtrPtr)
 			{
@@ -246,11 +287,6 @@ void FMemAllocGroupingByCallstack::GroupNodes(
 			UnsetGroupPtr->AddChildAndSetParent(NodePtr);
 		}
 	}
-
-	for (FCallstackGroup* Group : CallstackGroups)
-	{
-		delete Group;
-	}
 }
 
 ////////////////////////////////////////////////////////////////////////////////////////////////////
diff --git a/Engine/Source/Developer/TraceInsights/Private/Insights/MemoryProfiler/ViewModels/MemAllocGroupingByCallstack.h b/Engine/Source/Developer/TraceInsights/Private/Insights/MemoryProfiler/ViewModels/MemAllocGroupingByCallstack.h
--- a/Engine/Source/Developer/TraceInsights/Private/Insights/MemoryProfiler/ViewModels/MemAllocGroupingByCallstack.h
+++ b/Engine/Source/Developer/TraceInsights/Private/Insights/MemoryProfiler/ViewModels/MemAllocGroupingByCallstack.h
@@ -80,7 +80,9 @@ public:
 
 	bool ShouldSkipFilteredFrames() const { return bShouldSkipFilteredFrames; }
 	void SetSkipFilteredFrames(bool bOnOff) { bShouldSkipFilteredFrames = bOnOff; }
-
+protected:
+	FMemAllocGroupingByCallstack(const FText& InShortName, const FText& InTitleName, const FText& InDescription, const FSlateBrush* InIcon, bool bInIsAllocCallstack, bool bInIsInverted, bool bInIsGroupingByFunction);
+	void GroupNodesInternal(const TArray<FTableTreeNodePtr>& Nodes, FTableTreeNode& ParentGroup, TWeakPtr<FTable> InParentTable, IAsyncOperationProgress& InAsyncOperationProgress, TUniqueFunction<FTableTreeNodePtr(FTableTreeNodePtr&)> ProcessNode) const;
 private:
 	FName GetGroupName(const TraceServices::FStackFrame* Frame) const;
 
diff --git a/Engine/Source/Developer/TraceInsights/Private/Insights/MemoryProfiler/ViewModels/MemAllocGroupingBySwapPage.cpp b/Engine/Source/Developer/TraceInsights/Private/Insights/MemoryProfiler/ViewModels/MemAllocGroupingBySwapPage.cpp
--- a/Engine/Source/Developer/TraceInsights/Private/Insights/MemoryProfiler/ViewModels/MemAllocGroupingBySwapPage.cpp
+++ b/Engine/Source/Developer/TraceInsights/Private/Insights/MemoryProfiler/ViewModels/MemAllocGroupingBySwapPage.cpp
@@ -7,12 +7,9 @@
 // TraceServices
 #include "Common/ProviderLock.h"
 
-// TraceInsightsCore
-#include "InsightsCore/Common/AsyncOperationProgress.h"
-
 // TraceInsights
-#include "Insights/MemoryProfiler/ViewModels/CallstackFormatting.h"
-#include "Insights/MemoryProfiler/ViewModels/MemAllocInSwapNode.h"
+#include "Insights/MemoryProfiler/ViewModels/MemAllocNode.h"
+#include "Insights/MemoryProfiler/ViewModels/MemAllocValueNode.h"
 
 #define LOCTEXT_NAMESPACE "UE::Insights::MemoryProfiler::FMemAllocNode"
 
@@ -23,12 +20,19 @@ INSIGHTS_IMPLEMENT_RTTI(FMemAllocGroupingBySwapPage)
 
 ////////////////////////////////////////////////////////////////////////////////////////////////////
 
-FMemAllocGroupingBySwapPage::FMemAllocGroupingBySwapPage(const TraceServices::IAllocationsProvider& InAllocProvider)
-	: FTreeNodeGrouping(
-		LOCTEXT("Grouping_BySwap_ShortName", "Swap"),
-		LOCTEXT("Grouping_BySwap_TitleName", "By Swap"),
-		LOCTEXT("Grouping_BySwap_Desc", "Creates a tree based on swap state."),
-		nullptr)
+FMemAllocGroupingBySwapPage::FMemAllocGroupingBySwapPage(const TraceServices::IAllocationsProvider& InAllocProvider, bool bInIsInverted, bool bInIsGroupingByFunction)
+	: FMemAllocGroupingByCallstack(
+ 		bInIsInverted ? 
+			LOCTEXT("Grouping_BySwapCallstack2_ShortName", "In-Swap Inverted Alloc Callstack") 
+			: LOCTEXT("Grouping_BySwapCallstack1_ShortName", "In-Swap Alloc Callstack"),
+		bInIsInverted ? 
+			LOCTEXT("Grouping_BySwapCallstack2_TitleName", "By Inverted In-Swap Inverted Alloc Callstack")
+			: LOCTEXT("Grouping_BySwapCallstack1_TitleName", "By In-Swap Alloc Callstack"),
+		LOCTEXT("Grouping_SwapCallstack_Desc", "Creates a tree based on callstack of each allocation which is in swap memory."),
+		nullptr
+	, true /*bInIsAllocCallstack*/
+	, bInIsInverted
+	, bInIsGroupingByFunction)
 	, AllocProvider(InAllocProvider)
 {
 }
@@ -41,9 +45,9 @@ FMemAllocGroupingBySwapPage::~FMemAllocGroupingBySwapPage()
 
 ////////////////////////////////////////////////////////////////////////////////////////////////////
 
-struct TMemoryPageMapKeyFuncs : BaseKeyFuncs<TPair<uint64, FTableTreeNode*>, uint64, false>
+struct TMemoryPageMapKeyFuncs : BaseKeyFuncs<TPair<uint64, FTableTreeNodePtr>, uint64, false>
 {
-	static FORCEINLINE uint64 GetSetKey(const TPair<uint64, FTableTreeNode*>& Element)
+	static FORCEINLINE uint64 GetSetKey(const TPair<uint64, FTableTreeNodePtr>& Element)
 	{
 		return Element.Key;
 	}
@@ -85,8 +89,46 @@ public:
 		return FLinearColor(0.2f, 0.8f, 0.4f, 1.0f);
 	}
 };
+////////////////////////////////////////////////////////////////////////////////////////////////////
+
+class FMemSwappedAllocsTreeNode : public FMemAllocValueNode
+{
+	INSIGHTS_DECLARE_RTTI(FMemSwappedAllocsTreeNode, FMemAllocValueNode)
+
+	int64 SwappedSize;
+
+public:
+	/** Initialization constructor for the group node. */
+	explicit FMemSwappedAllocsTreeNode(const FName InName, int64 InSwappedSize, TWeakPtr<FTable> InParentTable)
+		: FMemAllocValueNode(InName, InParentTable)
+		, SwappedSize(InSwappedSize)
+	{
+	}
+
+	virtual ~FMemSwappedAllocsTreeNode()
+	{
+	}
+
+	virtual FLinearColor GetIconColor() const override final
+	{
+		return FLinearColor(0.3f, 0.8f, 0.4f, 1.0f);
+	}
+
+	virtual FLinearColor GetColor() const override final
+	{
+		return FLinearColor(0.4f, 0.5f, 0.4f, 1.0f);
+	}
+
+	virtual int64 GetValueI64() const override
+	{
+		return SwappedSize;
+	}
+};
+
+////////////////////////////////////////////////////////////////////////////////////////////////////
 
 INSIGHTS_IMPLEMENT_RTTI(FMemSwapPageTreeNode)
+INSIGHTS_IMPLEMENT_RTTI(FMemSwappedAllocsTreeNode)
 
 ////////////////////////////////////////////////////////////////////////////////////////////////////
 
@@ -98,13 +140,16 @@ void FMemAllocGroupingBySwapPage::GroupNodes(
 {
 	ParentGroup.ClearChildren();
 
-	auto InSwapGroup = MakeShared<FTableTreeNode>(FName(TEXT("In Swap")), InParentTable);
-	auto InRamGroup = MakeShared<FTableTreeNode>(FName(TEXT("In RAM")), InParentTable);
-	ParentGroup.AddChildAndSetParent(InSwapGroup);
-	ParentGroup.AddChildAndSetParent(InRamGroup);
+	FSwapNodeMap SwapNodes;
+	GatherSwapNodes(Nodes, SwapNodes, InParentTable, InAsyncOperationProgress);
+
+	GroupNodesInternal(Nodes, ParentGroup, InParentTable, InAsyncOperationProgress, [this,InParentTable, SwapNodes](FTableTreeNodePtr& InTableTreeNodePtr) { return ProcessNodeForSwapInfo(InTableTreeNodePtr, InParentTable, SwapNodes); });
+}
 
-	TMap<uint64, FTableTreeNode*, FDefaultSetAllocator, TMemoryPageMapKeyFuncs> SwapNodes;
+////////////////////////////////////////////////////////////////////////////////////////////////////
 
+void FMemAllocGroupingBySwapPage::GatherSwapNodes(const TArray<FTableTreeNodePtr>& Nodes, FSwapNodeMap& SwapNodeMap, TWeakPtr<FTable> InParentTable, IAsyncOperationProgress& InAsyncOperationProgress) const
+{
 	for (FTableTreeNodePtr NodePtr : Nodes)
 	{
 		if (InAsyncOperationProgress.ShouldCancelAsyncOp())
@@ -112,87 +157,130 @@ void FMemAllocGroupingBySwapPage::GroupNodes(
 			return;
 		}
 
-		const FMemAllocNode& MemAllocNode = static_cast<const FMemAllocNode&>(*NodePtr);
+		if (!NodePtr->Is<FMemAllocNode>())
+		{
+			continue;
+		}
+		const FMemAllocNode& MemAllocNode = NodePtr->As<FMemAllocNode>();
 		const FMemoryAlloc* Alloc = MemAllocNode.GetMemAlloc();
 
 		if (Alloc && Alloc->IsSwap())
 		{
 			auto SwapEntry = MakeShared<FMemSwapPageTreeNode>(FName(FString::Printf(TEXT("0x%016llx"), Alloc->GetAddress())), InParentTable);
-			InSwapGroup->AddChildAndSetParent(SwapEntry);
+			NodePtr->SetExpansion(false);
 			SwapEntry->AddChildAndSetParent(NodePtr);
 
-			ensure(SwapNodes.Find(Alloc->GetAddress()) == nullptr);
-
-			SwapNodes.Add(Alloc->GetAddress(), &SwapEntry.Get());
+			ensure(SwapNodeMap.Find(Alloc->GetAddress()) == nullptr);
+			SwapNodeMap.Add(Alloc->GetAddress(), SwapEntry);
 		}
 	}
+}
 
-	for (FTableTreeNodePtr NodePtr : Nodes)
+////////////////////////////////////////////////////////////////////////////////////////////////////
+
+FTableTreeNodePtr FMemAllocGroupingBySwapPage::ProcessNodeForSwapInfo(FTableTreeNodePtr& NodePtr, TWeakPtr<FTable> InParentTable, const FSwapNodeMap& SwapNodeMap) const
+{
+	if (!NodePtr->Is<FMemAllocNode>())
 	{
-		if (InAsyncOperationProgress.ShouldCancelAsyncOp())
-		{
-			return;
-		}
+		return FTableTreeNodePtr();
+	}
+
+	const FMemAllocNode& MemAllocNode = NodePtr->As<FMemAllocNode>();
+	const FMemoryAlloc* Alloc = MemAllocNode.GetMemAlloc();
+
+	const bool bIsInSwap = Alloc && Alloc->IsSwap();
+	if (bIsInSwap)
+	{
+		return FTableTreeNodePtr(); // Remove the swap_xxx entries, they are added back underneath their respective alloc_xxxs
+	}
 
-		if (NodePtr->IsGroup())
+	int64 SwappedSize = 0;
+
+	const uint64 PageSize = AllocProvider.GetPlatformPageSize();
+	const uint64 PageMask = ~(PageSize - 1);
+
+	const uint64 AllocStart = Alloc->GetAddress();
+	const uint64 AllocEnd = AllocStart + FMath::Max(Alloc->GetSize(), 1); // fake alloc to have at least 1 byte size to make page walking logic work
+	const uint64 PageRangeStart = AllocStart & PageMask;
+	const uint64 PageRangeEnd = (AllocEnd + PageSize - 1) & PageMask;
+	const uint64 PageAmount = (PageRangeEnd - PageRangeStart) / PageSize;
+
+	TArray<FMemSwapPageTreeNodePtr> SwapPagesEncountered;
+	// go through every memory page and check if the alloc is in swap, and record each page.
+	for (uint64 PageAddress = PageRangeStart; PageAddress < PageRangeEnd; PageAddress += PageSize)
+	{
+		const FMemSwapPageTreeNodePtr* SwapNodePtrPtr = SwapNodeMap.Find(PageAddress);
+		if (!SwapNodePtrPtr)
 		{
-			ParentGroup.AddChildAndSetParent(NodePtr);
 			continue;
 		}
+		uint64 AllocSizeInPage = PageSize;
 
-		const FMemAllocNode& MemAllocNode = static_cast<const FMemAllocNode&>(*NodePtr);
-		const FMemoryAlloc* Alloc = MemAllocNode.GetMemAlloc();
-
-		if (!Alloc)
+		if (PageAddress < AllocStart) 
 		{
-			InRamGroup->AddChildAndSetParent(NodePtr);
-			continue;
+			// adjust size if allocation starts after page start 
+			AllocSizeInPage -= AllocStart - PageAddress;
+		}
+		if (PageAddress + PageSize > AllocEnd)
+		{
+			// adjust size if allocation ends before page end
+			AllocSizeInPage -= PageAddress + PageSize - AllocEnd;
 		}
 
-		if (!Alloc->IsSwap())
+		if(AllocSizeInPage)
 		{
-			bool bAnyInSwap = false;
-			const uint64 PageSize = AllocProvider.GetPlatformPageSize();
-			const uint64 PageMask = ~(PageSize - 1);
+			SwapPagesEncountered.Add(*SwapNodePtrPtr);
+		}
 
-			const uint64 AllocStart = Alloc->GetAddress();
-			const uint64 AllocEnd = AllocStart + FMath::Max(Alloc->GetSize(), 1); // fake alloc to have at least 1 byte size to make page walking logic work
-			const uint64 PageRangeStart = AllocStart & PageMask;
-			const uint64 PageRangeEnd = (AllocEnd + PageSize - 1) & PageMask;
+		SwappedSize += AllocSizeInPage;
+	}
 
-			// for every alloc go through every memory page and check if it's in swap
-			for (uint64 PageAddress = PageRangeStart; PageAddress < PageRangeEnd; PageAddress += PageSize)
-			{
-				FTableTreeNode** GroupNodePtr = SwapNodes.Find(PageAddress);
-				if (!GroupNodePtr)
-				{
-					continue;
-				}
+	if (SwapPagesEncountered.IsEmpty())
+	{
+		// remove allocs that are not in swap.
+		return FTableTreeNodePtr();
+	}
 
-				FTableTreeNode* GroupNode = *GroupNodePtr;
+	// Wrap the alloc to provide Swap Size column info.
+	if (SwappedSize)
+	{
+		ensure(SwappedSize <= FMath::Max(Alloc->GetSize(), 1));
+		auto SwappedWrapper = MakeShared<FMemSwappedAllocsTreeNode>(FName(FString::Printf(TEXT("%s swap info"), *NodePtr->GetName().ToString())), SwappedSize, InParentTable);
+		SwappedWrapper->AddChildAndSetParent(NodePtr);
 
-				uint64 AllocSizeInPage = PageSize;
-				if (PageAddress < AllocStart) // adjust size if allocation starts after page start 
-				{
-					AllocSizeInPage -= AllocStart - PageAddress;
-				}
-				if (PageAddress + PageSize > AllocEnd) // adjust size if allocation ends before page end
-				{
-					AllocSizeInPage -= PageAddress + PageSize - AllocEnd;
-				}
+		if (TSharedPtr<FTable> SharedTable = InParentTable.Pin())
+		{
 
-				auto SwapAllocNode = MakeShared<FMemAllocInSwapNode>(MemAllocNode.GetName(), StaticCastWeakPtr<FMemAllocTable>(MemAllocNode.GetParentTable()), MemAllocNode.GetRowIndex(), AllocSizeInPage);
-				GroupNode->AddChildAndSetParent(SwapAllocNode);
 
-				bAnyInSwap = true;
-			}
+			auto SwapEventsGroup = MakeShared<FMemSwapPageTreeNode>(FName(FString::Printf(TEXT("Swap events"))), InParentTable);
+			SwappedWrapper->AddChildAndSetParent(SwapEventsGroup);
 
-			if (!bAnyInSwap)
+			// for each alloc we create a swap group that contains all of the swap events.
+			for (FMemSwapPageTreeNodePtr& SwapPage : SwapPagesEncountered)
 			{
-				InRamGroup->AddChildAndSetParent(NodePtr);
+				TSharedPtr<FMemAllocTable> MemAllocTable = StaticCastSharedPtr<FMemAllocTable>(SharedTable);
+				if (!MemAllocTable.IsValid())
+				{
+					continue;
+				}
+
+				SwapPage->EnumerateChildren([&SwapEventsGroup, &MemAllocTable](const FBaseTreeNodePtr& ChildNode) -> bool
+					{
+						// duplicate swap alloc nodes
+						if (ChildNode->Is<FMemAllocNode>())
+						{
+							const FMemAllocNode& SwapAllocNode = ChildNode->As<FMemAllocNode>();
+							FMemAllocNodePtr DupeSwapNodePtr = MakeShared<FMemAllocNode>(ChildNode->GetName(), MemAllocTable, SwapAllocNode.GetRowIndex());
+							SwapEventsGroup->AddChildAndSetParent(DupeSwapNodePtr);
+						}
+						return true;
+					});
 			}
+			NodePtr = SwappedWrapper;
 		}
 	}
+
+	return NodePtr;
 }
 
 ////////////////////////////////////////////////////////////////////////////////////////////////////
diff --git a/Engine/Source/Developer/TraceInsights/Private/Insights/MemoryProfiler/ViewModels/MemAllocGroupingBySwapPage.h b/Engine/Source/Developer/TraceInsights/Private/Insights/MemoryProfiler/ViewModels/MemAllocGroupingBySwapPage.h
--- a/Engine/Source/Developer/TraceInsights/Private/Insights/MemoryProfiler/ViewModels/MemAllocGroupingBySwapPage.h
+++ b/Engine/Source/Developer/TraceInsights/Private/Insights/MemoryProfiler/ViewModels/MemAllocGroupingBySwapPage.h
@@ -7,25 +7,45 @@
 // TraceInsightsCore
 #include "InsightsCore/Table/ViewModels/BaseTreeNode.h"
 #include "InsightsCore/Table/ViewModels/TreeNodeGrouping.h"
+#include "MemAllocGroupingByCallstack.h"
 
-namespace TraceServices { class IAllocationsProvider; }
-
+namespace TraceServices 
+{
+	class IAllocationsProvider;
+}
 namespace UE::Insights { class IAsyncOperationProgress; }
 
 namespace UE::Insights::MemoryProfiler
 {
+typedef TSharedPtr<class FMemSwapPageTreeNode> FMemSwapPageTreeNodePtr;
 
-class FMemAllocGroupingBySwapPage : public FTreeNodeGrouping
+class FMemAllocGroupingBySwapPage : public FMemAllocGroupingByCallstack
 {
-	INSIGHTS_DECLARE_RTTI(FMemAllocGroupingBySwapPage, FTreeNodeGrouping)
+	INSIGHTS_DECLARE_RTTI(FMemAllocGroupingBySwapPage, FMemAllocGroupingByCallstack)
 
 public:
-	FMemAllocGroupingBySwapPage(const TraceServices::IAllocationsProvider& AllocProvider);
+	FMemAllocGroupingBySwapPage(const TraceServices::IAllocationsProvider& AllocProvider, bool bInIsInverted, bool bInIsGroupingByFunction);
 	virtual ~FMemAllocGroupingBySwapPage() override;
 
 	virtual void GroupNodes(const TArray<FTableTreeNodePtr>& Nodes, FTableTreeNode& ParentGroup, TWeakPtr<FTable> InParentTable, IAsyncOperationProgress& InAsyncOperationProgress) const override;
 
 private:
+	struct TMemoryPageMapKeyFuncs : BaseKeyFuncs<TPair<uint64, FMemSwapPageTreeNodePtr>, uint64, false>
+	{
+		static FORCEINLINE bool Matches(const uint64 A, const uint64 B) { return A == B; }
+		static FORCEINLINE uint64 GetSetKey(const TPair<uint64, FMemSwapPageTreeNodePtr>& Element) { return Element.Key; }
+		static FORCEINLINE uint32 GetKeyHash(const uint64 Key)
+		{
+			// memory page is always at least 4k aligned, so skip lower bytes
+			return (uint32)(Key >> 12);
+		}
+	};
+
+	using FSwapNodeMap = TMap<uint64, FMemSwapPageTreeNodePtr, FDefaultSetAllocator, TMemoryPageMapKeyFuncs>;
+	FTableTreeNodePtr ProcessNodeForSwapInfo(FTableTreeNodePtr& NodePtr, TWeakPtr<FTable> InParentTable, const FSwapNodeMap& SwapNodeMap) const;
+	
+	void GatherSwapNodes(const TArray<FTableTreeNodePtr>& Nodes, FSwapNodeMap& SwapNodeMap, TWeakPtr<FTable> ParentTable, IAsyncOperationProgress& InAsyncOperationProgress) const;
+
 	const TraceServices::IAllocationsProvider& AllocProvider;
 };
 
diff --git a/Engine/Source/Developer/TraceInsights/Private/Insights/MemoryProfiler/ViewModels/MemAllocInSwapNode.cpp b/Engine/Source/Developer/TraceInsights/Private/Insights/MemoryProfiler/ViewModels/MemAllocInSwapNode.cpp
deleted file mode 100644
diff --git a/Engine/Source/Developer/TraceInsights/Private/Insights/MemoryProfiler/ViewModels/MemAllocInSwapNode.h b/Engine/Source/Developer/TraceInsights/Private/Insights/MemoryProfiler/ViewModels/MemAllocInSwapNode.h
deleted file mode 100644
diff --git a/Engine/Source/Developer/TraceInsights/Private/Insights/MemoryProfiler/ViewModels/MemAllocTable.cpp b/Engine/Source/Developer/TraceInsights/Private/Insights/MemoryProfiler/ViewModels/MemAllocTable.cpp
--- a/Engine/Source/Developer/TraceInsights/Private/Insights/MemoryProfiler/ViewModels/MemAllocTable.cpp
+++ b/Engine/Source/Developer/TraceInsights/Private/Insights/MemoryProfiler/ViewModels/MemAllocTable.cpp
@@ -22,7 +22,7 @@
 #include "Insights/MemoryProfiler/ViewModels/MemAllocFilterValueConverter.h"
 #include "Insights/MemoryProfiler/ViewModels/MemAllocGroupingByTag.h"
 #include "Insights/MemoryProfiler/ViewModels/MemAllocNode.h"
-#include "Insights/MemoryProfiler/ViewModels/MemAllocInSwapNode.h"
+#include "Insights/MemoryProfiler/ViewModels/MemAllocValueNode.h"
 #include "Insights/MemoryProfiler/ViewModels/MemAllocTable.h"
 
 #define LOCTEXT_NAMESPACE "UE::Insights::MemoryProfiler::FMemAllocTable"
@@ -741,16 +741,33 @@ void FMemAllocTable::AddDefaultColumns()
 
 		class FSizeInSwapColumnValueGetter : public FAggregatedValueGetter
 		{
-		public:
-			virtual const TOptional<FTableCellValue> GetNodeValue(const FBaseTreeNode& Node) const override
-			{
-				if (Node.Is<FMemAllocInSwapNode>())
+			public:
+				virtual const TOptional<FTableCellValue> GetValue(const FTableColumn& Column, const FBaseTreeNode& Node) const override
 				{
-					const FMemAllocInSwapNode& MemAllocInSwapNode = Node.As<FMemAllocInSwapNode>();
-					return FTableCellValue(static_cast<int64>(MemAllocInSwapNode.GetBytesInSwapPage()));
+ 					if (Node.Is<FMemAllocValueNode>())
+					{
+						const FMemAllocValueNode& SwappedNode = Node.As<FMemAllocValueNode>();
+						return FTableCellValue(SwappedNode.GetValueI64());
+					}
+					else
+					{
+ 						return FAggregatedValueGetter::GetValue(Column, Node);
+ 					}
+				}
+
+				virtual const TOptional<FTableCellValue> GetNodeValue(const FBaseTreeNode& Node) const override
+				{
+					if (Node.Is<FMemAllocNode>())
+					{
+						const FMemAllocNode& MemAllocNode = Node.As<FMemAllocNode>();
+						const FMemoryAlloc* Alloc = MemAllocNode.GetMemAlloc();
+						if (Alloc && Alloc->IsSwap())
+						{
+							return FTableCellValue(static_cast<int64>(Alloc->GetSize()));
+						}
+					}
+					return TOptional<FTableCellValue>();
 				}
-				return TOptional<FTableCellValue>();
-			}
 		};
 		TSharedRef<ITableCellValueGetter> Getter = MakeShared<FSizeInSwapColumnValueGetter>();
 		Column.SetValueGetter(Getter);
@@ -801,8 +818,10 @@ void FMemAllocTable::AddDefaultColumns()
 					{
 						if (Alloc->IsSwap())
 						{
-							const FMemAllocTable& MemAllocTable = MemAllocNode.GetMemTableChecked();
-							return FTableCellValue(static_cast<int64>(MemAllocTable.GetPlatformPageSize()));
+							//Do not give swap nodes a size. The size is aggregated with allocations and pollutes the column.
+							return FTableCellValue(0LL);
+							// const FMemAllocTable& MemAllocTable = MemAllocNode.GetMemTableChecked();
+							// return FTableCellValue(static_cast<int64>(MemAllocTable.GetPlatformPageSize()));
 						}
 						else
 						{
diff --git a/Engine/Source/Developer/TraceInsights/Private/Insights/MemoryProfiler/ViewModels/MemAllocTableViewPresets.cpp b/Engine/Source/Developer/TraceInsights/Private/Insights/MemoryProfiler/ViewModels/MemAllocTableViewPresets.cpp
--- a/Engine/Source/Developer/TraceInsights/Private/Insights/MemoryProfiler/ViewModels/MemAllocTableViewPresets.cpp
+++ b/Engine/Source/Developer/TraceInsights/Private/Insights/MemoryProfiler/ViewModels/MemAllocTableViewPresets.cpp
@@ -564,35 +564,43 @@ TSharedRef<ITableTreeViewPreset> FMemAllocTableViewPresets::CreatePlatformPageVi
 ////////////////////////////////////////////////////////////////////////////////////////////////////
 // Swap Breakdown View
 
-TSharedRef<ITableTreeViewPreset> FMemAllocTableViewPresets::CreateSwapViewPreset(SMemAllocTableTreeView& TableTreeView)
+TSharedRef<ITableTreeViewPreset> FMemAllocTableViewPresets::CreateSwapViewPreset(SMemAllocTableTreeView& TableTreeView, bool bIsInverted )
 {
 	class FSwapViewPreset : public ITableTreeViewPreset
 	{
 	public:
+		FSwapViewPreset(bool bIsInverted)
+			: bIsInvertedCallstack(bIsInverted)
+		{
+		}
 		virtual FText GetName() const override
 		{
-			return LOCTEXT("Swap_PresetName", "Swap");
+			return bIsInvertedCallstack ? LOCTEXT("InvertedSwap_PresetName", "Inverted In-Swap Callstack") : LOCTEXT("Swap_PresetName", "In-Swap Callstack");
 		}
 		virtual FText GetToolTip() const override
 		{
-			return LOCTEXT("Swap_PresetToolTip", "Swap Usage Breakdown View\nConfigure the tree view to show a breakdown of allocations by their swap page.\nIt groups allocs into corresponding swap pages.");
+			return bIsInvertedCallstack ?
+				LOCTEXT("InvertedSwap_PresetToolTip", "Inverted Swap Usage Breakdown View\nConfigure the tree view to show a breakdown of allocations by their swap page.\nIt shows allocs by inverted callstack in which some or all of the allocation was moved into swap pages.") :
+				LOCTEXT("Swap_PresetToolTip", "Swap Usage Breakdown View\nConfigure the tree view to show a breakdown of allocations by their swap page.\nIt shows allocs by callstack in which some or all of the allocation was moved into swap pages.");
 		}
 		virtual FName GetSortColumn() const override
 		{
-			return FTable::GetHierarchyColumnId();
+			return FMemAllocTableColumns::SwapSizeColumnId;
 		}
 		virtual EColumnSortMode::Type GetSortMode() const override
 		{
-			return EColumnSortMode::Type::Ascending;
+			return EColumnSortMode::Type::Descending;
 		}
 		virtual void SetCurrentGroupings(const TArray<TSharedPtr<FTreeNodeGrouping>>& InAvailableGroupings, TArray<TSharedPtr<FTreeNodeGrouping>>& InOutCurrentGroupings) const override
 		{
 			InOutCurrentGroupings.Reset();
+			const bool bIsInverted = bIsInvertedCallstack;
 
 			const TSharedPtr<FTreeNodeGrouping>* HeapGrouping = InAvailableGroupings.FindByPredicate(
-				[](TSharedPtr<FTreeNodeGrouping>& Grouping)
+				[bIsInverted](TSharedPtr<FTreeNodeGrouping>& Grouping)
 				{
-					return Grouping->Is<FMemAllocGroupingBySwapPage>();
+					return Grouping->Is<FMemAllocGroupingBySwapPage>() && 
+ 						Grouping->As<FMemAllocGroupingBySwapPage>().IsInverted() == bIsInverted;
 				});
 			if (HeapGrouping)
 			{
@@ -601,19 +609,20 @@ TSharedRef<ITableTreeViewPreset> FMemAllocTableViewPresets::CreateSwapViewPreset
 		}
 		virtual void GetColumnConfigSet(TArray<FTableColumnConfig>& InOutConfigSet) const override
 		{
-			InOutConfigSet.Add({ FTable::GetHierarchyColumnId(),                    true, 200.0f });
+			InOutConfigSet.Add({ FTable::GetHierarchyColumnId(),                    true, 400.0f });
 			InOutConfigSet.Add({ FMemAllocTableColumns::StartTimeColumnId,          true, 100.0f });
 			InOutConfigSet.Add({ FMemAllocTableColumns::EndTimeColumnId,            true, 100.0f });
-			InOutConfigSet.Add({ FMemAllocTableColumns::AddressColumnId,            true, 120.0f });
+			InOutConfigSet.Add({ FMemAllocTableColumns::AddressColumnId,            true, 100.0f });
 			InOutConfigSet.Add({ FMemAllocTableColumns::CountColumnId,              true, 100.0f });
-			InOutConfigSet.Add({ FMemAllocTableColumns::CompressedSwapSizeColumnId, true, 100.0f });
-			InOutConfigSet.Add({ FMemAllocTableColumns::SwapSizeColumnId,           true, 100.0f });
+			InOutConfigSet.Add({ FMemAllocTableColumns::SwapSizeColumnId,           true, 120.0f });
 			InOutConfigSet.Add({ FMemAllocTableColumns::SizeColumnId,               true, 100.0f });
-			InOutConfigSet.Add({ FMemAllocTableColumns::TagColumnId,                true, 120.0f });
-			InOutConfigSet.Add({ FMemAllocTableColumns::AllocFunctionColumnId,      true, 400.0f });
+			InOutConfigSet.Add({ FMemAllocTableColumns::TagColumnId,                true, 130.0f });
+			InOutConfigSet.Add({ FMemAllocTableColumns::AllocFunctionColumnId,      true, 200.0f });
 		}
+	private:
+		bool bIsInvertedCallstack;
 	};
-	return MakeShared<FSwapViewPreset>();
+	return MakeShared<FSwapViewPreset>(bIsInverted);
 }
 
 ////////////////////////////////////////////////////////////////////////////////////////////////////
diff --git a/Engine/Source/Developer/TraceInsights/Private/Insights/MemoryProfiler/ViewModels/MemAllocTableViewPresets.h b/Engine/Source/Developer/TraceInsights/Private/Insights/MemoryProfiler/ViewModels/MemAllocTableViewPresets.h
--- a/Engine/Source/Developer/TraceInsights/Private/Insights/MemoryProfiler/ViewModels/MemAllocTableViewPresets.h
+++ b/Engine/Source/Developer/TraceInsights/Private/Insights/MemoryProfiler/ViewModels/MemAllocTableViewPresets.h
@@ -25,7 +25,7 @@ public:
 	static TSharedRef<ITableTreeViewPreset> CreateClassNameViewPreset(SMemAllocTableTreeView& TableTreeView);
 	static TSharedRef<ITableTreeViewPreset> CreateCallstackViewPreset(SMemAllocTableTreeView& TableTreeView, bool bIsInverted, bool bIsAlloc);
 	static TSharedRef<ITableTreeViewPreset> CreatePlatformPageViewPreset(SMemAllocTableTreeView& TableTreeView);
-	static TSharedRef<ITableTreeViewPreset> CreateSwapViewPreset(SMemAllocTableTreeView& TableTreeView);
+	static TSharedRef<ITableTreeViewPreset> CreateSwapViewPreset(SMemAllocTableTreeView& TableTreeView, bool bIsInverted);
 };
 
 } // namespace UE::Insights::MemoryProfiler
diff --git a/Engine/Source/Developer/TraceInsights/Private/Insights/MemoryProfiler/ViewModels/MemAllocValueNode.cpp b/Engine/Source/Developer/TraceInsights/Private/Insights/MemoryProfiler/ViewModels/MemAllocValueNode.cpp
new file mode 100644
--- /dev/null
+++ b/Engine/Source/Developer/TraceInsights/Private/Insights/MemoryProfiler/ViewModels/MemAllocValueNode.cpp
@@ -0,0 +1,14 @@
+// Copyright Epic Games, Inc. All Rights Reserved.
+
+#include "MemAllocValueNode.h"
+
+#define LOCTEXT_NAMESPACE "UE::Insights::MemoryProfiler::FMemAllocValueNode"
+
+namespace UE::Insights::MemoryProfiler
+{
+
+INSIGHTS_IMPLEMENT_RTTI(FMemAllocValueNode)
+
+} // namespace UE::Insights::MemoryProfiler
+
+#undef LOCTEXT_NAMESPACE
diff --git a/Engine/Source/Developer/TraceInsights/Private/Insights/MemoryProfiler/ViewModels/MemAllocValueNode.h b/Engine/Source/Developer/TraceInsights/Private/Insights/MemoryProfiler/ViewModels/MemAllocValueNode.h
new file mode 100644
--- /dev/null
+++ b/Engine/Source/Developer/TraceInsights/Private/Insights/MemoryProfiler/ViewModels/MemAllocValueNode.h
@@ -0,0 +1,41 @@
+// Copyright Epic Games, Inc. All Rights Reserved.
+
+#pragma once
+
+#include "CoreTypes.h"
+
+#include "Templates/SharedPointer.h"
+
+// TraceServices
+#include "TraceServices/Model/Memory.h"
+
+// TraceInsightsCore
+#include "InsightsCore/Table/ViewModels/TableTreeNode.h"
+
+////////////////////////////////////////////////////////////////////////////////////////////////////
+
+namespace UE::Insights::MemoryProfiler
+{
+////////////////////////////////////////////////////////////////////////////////////////////////////
+/**
+ * Class used to store information about an allocation node (used in the SMemAllocTreeView).
+ */
+class FMemAllocValueNode : public FTableTreeNode
+{
+	INSIGHTS_DECLARE_RTTI(FMemAllocValueNode, FTableTreeNode)
+public:
+	/** Initialization constructor for the group node. */
+	explicit FMemAllocValueNode(const FName InName, TWeakPtr<FTable> InParentTable)
+		: FTableTreeNode(InName, InParentTable)
+	{
+	}
+
+	virtual ~FMemAllocValueNode()
+	{
+	}
+
+	virtual int64 GetValueI64() const = 0;
+};
+////////////////////////////////////////////////////////////////////////////////////////////////////
+
+} // namespace UE::Insights::MemoryProfiler
diff --git a/Engine/Source/Developer/TraceInsights/Private/Insights/MemoryProfiler/Widgets/SMemAllocTableTreeView.cpp b/Engine/Source/Developer/TraceInsights/Private/Insights/MemoryProfiler/Widgets/SMemAllocTableTreeView.cpp
--- a/Engine/Source/Developer/TraceInsights/Private/Insights/MemoryProfiler/Widgets/SMemAllocTableTreeView.cpp
+++ b/Engine/Source/Developer/TraceInsights/Private/Insights/MemoryProfiler/Widgets/SMemAllocTableTreeView.cpp
@@ -215,7 +215,8 @@ void SMemAllocTableTreeView::InternalCreateGroupings()
 	if (AllocationsProvider)
 	{
 		AvailableGroupings.Insert(MakeShared<FMemAllocGroupingByHeap>(*AllocationsProvider), Index++);
-		AvailableGroupings.Insert(MakeShared<FMemAllocGroupingBySwapPage>(*AllocationsProvider), Index++);
+		AvailableGroupings.Insert(MakeShared<FMemAllocGroupingBySwapPage>(*AllocationsProvider, true, bIsCallstackGroupingByFunction), Index++);
+		AvailableGroupings.Insert(MakeShared<FMemAllocGroupingBySwapPage>(*AllocationsProvider, false, bIsCallstackGroupingByFunction), Index++);
 	}
 }
 
@@ -235,7 +236,8 @@ void SMemAllocTableTreeView::InitAvailableViewPresets()
 	AvailableViewPresets.Add(FMemAllocTableViewPresets::CreateCallstackViewPreset(*this, false, false));
 	AvailableViewPresets.Add(FMemAllocTableViewPresets::CreateCallstackViewPreset(*this, true, false));
 	AvailableViewPresets.Add(FMemAllocTableViewPresets::CreatePlatformPageViewPreset(*this));
-	AvailableViewPresets.Add(FMemAllocTableViewPresets::CreateSwapViewPreset(*this));
+	AvailableViewPresets.Add(FMemAllocTableViewPresets::CreateSwapViewPreset(*this, true));
+	AvailableViewPresets.Add(FMemAllocTableViewPresets::CreateSwapViewPreset(*this, false));
 
 	SelectedViewPreset = AvailableViewPresets[0];
 }
