Fixing static analysis error.

diff --git a/Engine/Source/Developer/TraceInsights/Private/Insights/MemoryProfiler/ViewModels/MemAllocGroupingBySwapPage.cpp b/Engine/Source/Developer/TraceInsights/Private/Insights/MemoryProfiler/ViewModels/MemAllocGroupingBySwapPage.cpp
--- a/Engine/Source/Developer/TraceInsights/Private/Insights/MemoryProfiler/ViewModels/MemAllocGroupingBySwapPage.cpp
+++ b/Engine/Source/Developer/TraceInsights/Private/Insights/MemoryProfiler/ViewModels/MemAllocGroupingBySwapPage.cpp
@@ -189,7 +189,7 @@ FTableTreeNodePtr FMemAllocGroupingBySwapPage::ProcessNodeForSwapInfo(FTableTree
 	const FMemoryAlloc* Alloc = MemAllocNode.GetMemAlloc();
 
 	const bool bIsInSwap = Alloc && Alloc->IsSwap();
-	if (bIsInSwap)
+	if (bIsInSwap || !Alloc)
 	{
 		return FTableTreeNodePtr(); // Remove the swap_xxx entries, they are added back underneath their respective alloc_xxxs
 	}
