PerFrame Average Count column for Timers panel

diff --git a/Engine/Source/Developer/TraceInsights/Private/Insights/TimingProfiler/ViewModels/TimersViewColumnFactory.cpp b/Engine/Source/Developer/TraceInsights/Private/Insights/TimingProfiler/ViewModels/TimersViewColumnFactory.cpp
--- a/Engine/Source/Developer/TraceInsights/Private/Insights/TimingProfiler/ViewModels/TimersViewColumnFactory.cpp
+++ b/Engine/Source/Developer/TraceInsights/Private/Insights/TimingProfiler/ViewModels/TimersViewColumnFactory.cpp
@@ -26,6 +26,7 @@ const FName FTimersViewColumns::MetaGroupNameColumnID(TEXT("MetaGroupName"));
 const FName FTimersViewColumns::TypeColumnID(TEXT("Type"));
 const FName FTimersViewColumns::InstanceCountColumnID(TEXT("Count"));
 const FName FTimersViewColumns::ChildInstanceCountColumnID(TEXT("ChildCount"));
+const FName FTimersViewColumns::AverageInstanceCountColumnID(TEXT("AverageCount")); 
 
 // Inclusive Time columns
 const FName FTimersViewColumns::TotalInclusiveTimeColumnID(TEXT("TotalInclTime"));
@@ -56,6 +57,7 @@ void FTimersViewColumnFactory::CreateTimersViewColumns(TArray<TSharedRef<FTableC
 	Columns.Add(CreateMetaGroupNameColumn());
 	Columns.Add(CreateTypeColumn());
 	Columns.Add(CreateInstanceCountColumn());
+	Columns.Add(CreateAverageInstanceCountColumn());
 
 	Columns.Add(CreateTotalInclusiveTimeColumn());
 	Columns.Add(CreateMaxInclusiveTimeColumn());
@@ -375,6 +377,61 @@ TSharedRef<FTableColumn> FTimersViewColumnFactory::CreateChildInstanceCountColum
 	return ColumnRef;
 }
 
+////////////////////////////////////////////////////////////////////////////////////////////////////
+	
+TSharedRef<FTableColumn> FTimersViewColumnFactory::CreateAverageInstanceCountColumn()
+{
+	TSharedRef<FTimersTableColumn> ColumnRef = MakeShared<FTimersTableColumn>(FTimersViewColumns::AverageInstanceCountColumnID);
+	FTimersTableColumn& Column = *ColumnRef;
+
+	Column.SetShortName(LOCTEXT("AverageInstanceCount_ColumnName", "C.Avg"));
+	Column.SetTitleName(LOCTEXT("AverageInstanceCount_ColumnTitle", "Average Instance Count"));
+	
+	Column.SetDescription(ETraceFrameType::TraceFrameType_Count,
+		LOCTEXT("AverageInstanceCount_ColumnDesc",
+			"Total number of timing event instances of the selected timer in the selected range\n"
+			"In Instance mode, this column has the same value as the Instance Count column."));
+	Column.SetDescription(ETraceFrameType::TraceFrameType_Game,
+		LOCTEXT("AverageInstanceCount_GameFrameColumnDesc",
+			"Game Frame Average Instance Count\n"
+			"Average number of timing event instances of the selected timer per-frame."));
+	Column.SetDescription(ETraceFrameType::TraceFrameType_Rendering,
+		LOCTEXT("AverageInstanceCount_RenderingFrameColumnDesc",
+			"Rendering Frame Average Instance Count\n"
+			"Average number of timing event instances of the selected timer per-frame."));
+
+	Column.SetFlags(ETableColumnFlags::CanBeHidden | ETableColumnFlags::CanBeFiltered);
+
+	Column.SetHorizontalAlignment(HAlign_Right);
+	Column.SetInitialWidth(60.0f);
+
+	Column.SetDataType(ETableCellDataType::Double);
+
+	class FAverageInstanceCountValueGetter : public FTableCellValueGetter
+	{
+	public:
+		virtual const TOptional<FTableCellValue> GetValue(const FTableColumn& Column, const FBaseTreeNode& Node) const
+		{
+			ensure(Column.GetId() == FTimersViewColumns::AverageInstanceCountColumnID);
+			const FTimerNode& TimerNode = static_cast<const FTimerNode&>(Node);
+			return TOptional<FTableCellValue>(FTableCellValue(TimerNode.GetAggregatedStats().AverageInstanceCount));
+		}
+	};
+
+	TSharedRef<ITableCellValueGetter> Getter = MakeShared<FAverageInstanceCountValueGetter>();
+	Column.SetValueGetter(Getter);
+
+	TSharedRef<FDoubleValueFormatterAsCustomNumber> Formatter = MakeShared<FDoubleValueFormatterAsCustomNumber>();
+	Formatter->GetFormattingOptions().MaximumFractionalDigits = 2;
+	Column.SetValueFormatter(Formatter);
+
+	TSharedRef<ITableCellValueSorter> Sorter = MakeShared<FSorterByDoubleValue>(ColumnRef);
+	Column.SetValueSorter(Sorter);
+	Column.SetInitialSortMode(EColumnSortMode::Descending);
+
+	return ColumnRef;
+}
+
 ////////////////////////////////////////////////////////////////////////////////////////////////////
 // Inclusive Time Columns
 ////////////////////////////////////////////////////////////////////////////////////////////////////
diff --git a/Engine/Source/Developer/TraceInsights/Private/Insights/TimingProfiler/ViewModels/TimersViewColumnFactory.h b/Engine/Source/Developer/TraceInsights/Private/Insights/TimingProfiler/ViewModels/TimersViewColumnFactory.h
--- a/Engine/Source/Developer/TraceInsights/Private/Insights/TimingProfiler/ViewModels/TimersViewColumnFactory.h
+++ b/Engine/Source/Developer/TraceInsights/Private/Insights/TimingProfiler/ViewModels/TimersViewColumnFactory.h
@@ -25,6 +25,7 @@ struct FTimersViewColumns
 	static const FName TypeColumnID;
 	static const FName InstanceCountColumnID;
 	static const FName ChildInstanceCountColumnID;
+	static const FName AverageInstanceCountColumnID;
 
 	// Inclusive Time columns
 	static const FName TotalInclusiveTimeColumnID;
@@ -106,6 +107,7 @@ public:
 	static TSharedRef<FTableColumn> CreateTypeColumn();
 	static TSharedRef<FTableColumn> CreateInstanceCountColumn();
 	static TSharedRef<FTableColumn> CreateChildInstanceCountColumn();
+	static TSharedRef<FTableColumn> CreateAverageInstanceCountColumn();
 
 	static TSharedRef<FTableColumn> CreateTotalInclusiveTimeColumn();
 	static TSharedRef<FTableColumn> CreateMaxInclusiveTimeColumn();
diff --git a/Engine/Source/Developer/TraceInsights/Private/Insights/TimingProfiler/Widgets/STimersView.cpp b/Engine/Source/Developer/TraceInsights/Private/Insights/TimingProfiler/Widgets/STimersView.cpp
--- a/Engine/Source/Developer/TraceInsights/Private/Insights/TimingProfiler/Widgets/STimersView.cpp
+++ b/Engine/Source/Developer/TraceInsights/Private/Insights/TimingProfiler/Widgets/STimersView.cpp
@@ -1448,9 +1448,7 @@ void STimersView::ApplyFiltering()
 		GroupPtr->ResetAggregatedStats();
 
 		constexpr double NanTimeValue = std::numeric_limits<double>::quiet_NaN();
-		AggregatedStats.AverageInclusiveTime = NanTimeValue;
 		AggregatedStats.MedianInclusiveTime = NanTimeValue;
-		AggregatedStats.AverageExclusiveTime = NanTimeValue;
 		AggregatedStats.MedianExclusiveTime = NanTimeValue;
 
 		const TArray<FBaseTreeNodePtr>& GroupChildren = GroupPtr->GetFilteredChildren();
@@ -1463,6 +1461,9 @@ void STimersView::ApplyFiltering()
 			if (NodeAggregatedStats.InstanceCount > 0)
 			{
 				AggregatedStats.InstanceCount += NodeAggregatedStats.InstanceCount;
+				AggregatedStats.AverageInstanceCount += NodeAggregatedStats.AverageInstanceCount;
+				AggregatedStats.AverageInclusiveTime += NodeAggregatedStats.AverageInclusiveTime;
+				AggregatedStats.AverageExclusiveTime += NodeAggregatedStats.AverageExclusiveTime;
 				AggregatedStats.TotalInclusiveTime += NodeAggregatedStats.TotalInclusiveTime;
 				AggregatedStats.MinInclusiveTime = FMath::Min(AggregatedStats.MinInclusiveTime, NodeAggregatedStats.MinInclusiveTime);
 				AggregatedStats.MaxInclusiveTime = FMath::Max(AggregatedStats.MaxInclusiveTime, NodeAggregatedStats.MaxInclusiveTime);
diff --git a/Engine/Source/Developer/TraceInsightsCore/Private/InsightsCore/Table/ViewModels/TableCellValueFormatter.cpp b/Engine/Source/Developer/TraceInsightsCore/Private/InsightsCore/Table/ViewModels/TableCellValueFormatter.cpp
--- a/Engine/Source/Developer/TraceInsightsCore/Private/InsightsCore/Table/ViewModels/TableCellValueFormatter.cpp
+++ b/Engine/Source/Developer/TraceInsightsCore/Private/InsightsCore/Table/ViewModels/TableCellValueFormatter.cpp
@@ -312,6 +312,29 @@ FText FDoubleValueFormatterAsNumber::FormatValue(const TOptional<FTableCellValue
 
 ////////////////////////////////////////////////////////////////////////////////////////////////////
 
+FText FDoubleValueFormatterAsCustomNumber::FormatValue(const TOptional<FTableCellValue>& InValue) const
+{
+	if (InValue.IsSet())
+	{
+		const double Value = InValue.GetValue().Double;
+		if (std::isnan(Value))
+		{
+			return FText::FromString(TEXT("NaN"));
+		}
+		else if (Value == 0.0)
+		{
+			return FText::FromString(TEXT("0"));
+		}
+		else
+		{
+			return FText::AsNumber(Value, &FormattingOptions);
+		}
+	}
+	return FText::GetEmpty();
+}
+
+////////////////////////////////////////////////////////////////////////////////////////////////////
+
 FText FDoubleValueFormatterAsTimeAuto::FormatValue(const TOptional<FTableCellValue>& InValue) const
 {
 	if (InValue.IsSet())
diff --git a/Engine/Source/Developer/TraceInsightsCore/Public/InsightsCore/Table/ViewModels/TableCellValueFormatter.h b/Engine/Source/Developer/TraceInsightsCore/Public/InsightsCore/Table/ViewModels/TableCellValueFormatter.h
--- a/Engine/Source/Developer/TraceInsightsCore/Public/InsightsCore/Table/ViewModels/TableCellValueFormatter.h
+++ b/Engine/Source/Developer/TraceInsightsCore/Public/InsightsCore/Table/ViewModels/TableCellValueFormatter.h
@@ -199,6 +199,26 @@ public:
 	UE_API virtual FText FormatValue(const TOptional<FTableCellValue>& InValue) const override;
 };
 
+////////////////////////////////////////////////////////////////////////////////////////////////////
+	
+class FDoubleValueFormatterAsCustomNumber : public FTableCellValueFormatter
+{
+public:
+	FDoubleValueFormatterAsCustomNumber()
+	{
+		FormattingOptions.MaximumFractionalDigits = 2;
+	}
+	virtual ~FDoubleValueFormatterAsCustomNumber() {}
+
+	const FNumberFormattingOptions& GetFormattingOptions() const { return FormattingOptions; }
+	FNumberFormattingOptions& GetFormattingOptions() { return FormattingOptions; }
+
+	UE_API virtual FText FormatValue(const TOptional<FTableCellValue>& InValue) const override;
+
+private:
+	FNumberFormattingOptions FormattingOptions;
+};
+	
 ////////////////////////////////////////////////////////////////////////////////////////////////////
 
 class FDoubleValueFormatterAsTimeAuto : public FTableCellValueFormatter
diff --git a/Engine/Source/Developer/TraceServices/Private/Common/TimelineStatistics.h b/Engine/Source/Developer/TraceServices/Private/Common/TimelineStatistics.h
--- a/Engine/Source/Developer/TraceServices/Private/Common/TimelineStatistics.h
+++ b/Engine/Source/Developer/TraceServices/Private/Common/TimelineStatistics.h
@@ -18,6 +18,7 @@ namespace TraceServices
 struct FAggregatedTimingStats
 {
 	uint64 InstanceCount = 0;
+	double AverageInstanceCount = 0.0;
 	double TotalInclusiveTime = 0.0;
 	double MinInclusiveTime = DBL_MAX;
 	double MaxInclusiveTime = -DBL_MAX;
@@ -169,9 +170,6 @@ public:
 						FAggregatedTimingStats& FrameStats = FrameResultIterator->Value;
 						FAggregatedTimingStats& TaskResultStats = TaskResultIterator->Value;
 
-						// TODO: Add Instance Count per frame column
-						// ResultStats.InstanceCount = (ResultStats.InstanceCount * Index + FrameStats.InstanceCount) / (Index + 1);
-
 						TaskResultStats.InstanceCount += FrameStats.InstanceCount;
 
 						TaskResultStats.TotalInclusiveTime += FrameStats.TotalInclusiveTime;
@@ -234,9 +232,6 @@ public:
 				FAggregatedTimingStats& TaskStats = TaskResultIterator->Value;
 				FAggregatedTimingStats& ResultStats = ResultIterator->Value.Inner.Stats;
 
-				// TODO: Add Instance Count per frame column
-				// ResultStats.InstanceCount = (ResultStats.InstanceCount * Index + FrameStats.InstanceCount) / (Index + 1);
-
 				ResultStats.InstanceCount += TaskStats.InstanceCount;
 
 				ResultStats.TotalInclusiveTime += TaskStats.TotalInclusiveTime;
@@ -277,6 +272,7 @@ public:
 		{
 			KV.Value.Inner.Stats.AverageInclusiveTime = KV.Value.Inner.Stats.TotalInclusiveTime / FramesNum;
 			KV.Value.Inner.Stats.AverageExclusiveTime = KV.Value.Inner.Stats.TotalExclusiveTime / FramesNum;
+			KV.Value.Inner.Stats.AverageInstanceCount = static_cast<double>(KV.Value.Inner.Stats.InstanceCount) / FramesNum;
 
 			Result.Add(KV.Key, KV.Value.Inner.Stats);
 		}
@@ -525,6 +521,7 @@ private:
 		double InvCount = 1.0f / static_cast<double>(Stats.InstanceCount);
 		Stats.AverageInclusiveTime = Stats.TotalInclusiveTime * InvCount;
 		Stats.AverageExclusiveTime = Stats.TotalExclusiveTime * InvCount;
+		Stats.AverageInstanceCount = static_cast<double>(Stats.InstanceCount);
 
 		if (ComputeMedian)
 		{
diff --git a/Engine/Source/Developer/TraceServices/Private/Model/TimingProfiler.cpp b/Engine/Source/Developer/TraceServices/Private/Model/TimingProfiler.cpp
--- a/Engine/Source/Developer/TraceServices/Private/Model/TimingProfiler.cpp
+++ b/Engine/Source/Developer/TraceServices/Private/Model/TimingProfiler.cpp
@@ -29,6 +29,7 @@ FTimingProfilerProvider::FTimingProfilerProvider(IAnalysisSession& InSession)
 			},
 			TEXT("Name")).
 		AddColumn(&FTimingProfilerAggregatedStats::InstanceCount, TEXT("Count")).
+		AddColumn(&FTimingProfilerAggregatedStats::AverageInstanceCount, TEXT("C.Avg")).
 		AddColumn(&FTimingProfilerAggregatedStats::TotalInclusiveTime, TEXT("Incl")).
 		AddColumn(&FTimingProfilerAggregatedStats::MinInclusiveTime, TEXT("I.Min")).
 		AddColumn(&FTimingProfilerAggregatedStats::MaxInclusiveTime, TEXT("I.Max")).
@@ -806,6 +807,7 @@ ITable<FTimingProfilerAggregatedStats>* FTimingProfilerProvider::CreateAggregati
 		Row.Timer = KV.Key;
 		const FAggregatedTimingStats& Stats = KV.Value;
 		Row.InstanceCount = Stats.InstanceCount;
+		Row.AverageInstanceCount = Stats.AverageInstanceCount;
 		Row.TotalInclusiveTime = Stats.TotalInclusiveTime;
 		Row.MinInclusiveTime = Stats.MinInclusiveTime;
 		Row.MaxInclusiveTime = Stats.MaxInclusiveTime;
diff --git a/Engine/Source/Developer/TraceServices/Public/TraceServices/Model/TimingProfiler.h b/Engine/Source/Developer/TraceServices/Public/TraceServices/Model/TimingProfiler.h
--- a/Engine/Source/Developer/TraceServices/Public/TraceServices/Model/TimingProfiler.h
+++ b/Engine/Source/Developer/TraceServices/Public/TraceServices/Model/TimingProfiler.h
@@ -64,6 +64,7 @@ struct FTimingProfilerAggregatedStats
 {
 	const FTimingProfilerTimer* Timer = nullptr;
 	uint64 InstanceCount = 0;
+	double AverageInstanceCount = 0.0;
 	double TotalInclusiveTime = 0.0;
 	double MinInclusiveTime = DBL_MAX;
 	double MaxInclusiveTime = -DBL_MAX;
