[Insights] Add ability to copy tooltip value to clipboard for bookmark track

diff --git a/Engine/Source/Developer/TraceInsights/Private/Insights/TimingProfiler/Tracks/MarkersTimingTrack.cpp b/Engine/Source/Developer/TraceInsights/Private/Insights/TimingProfiler/Tracks/MarkersTimingTrack.cpp
--- a/Engine/Source/Developer/TraceInsights/Private/Insights/TimingProfiler/Tracks/MarkersTimingTrack.cpp
+++ b/Engine/Source/Developer/TraceInsights/Private/Insights/TimingProfiler/Tracks/MarkersTimingTrack.cpp
@@ -2,6 +2,7 @@
 
 #include "MarkersTimingTrack.h"
 
+#include "Containers/StringFwd.h"
 #include "DesktopPlatformModule.h"
 #include "Fonts/FontMeasure.h"
 #include "Framework/Application/SlateApplication.h"
@@ -10,6 +11,7 @@
 #include "IImageWrapperModule.h"
 #include "Misc/FileHelper.h"
 #include "Misc/Paths.h"
+#include "HAL/PlatformApplicationMisc.h"
 #include "Modules/ModuleManager.h"
 #include "Styling/AppStyle.h"
 
@@ -700,8 +702,7 @@ void FMarkersTimingTrack::InitTooltip(FTooltipDrawState& InOutTooltip, const ITi
 	{
 		const FBookmarkTimingEvent& Event = InTooltipEvent.As<FBookmarkTimingEvent>();
 
-		if (Event.GetCallstackId() != 0 &&
-			Event.GetCallstackId() != TraceServices::FScreenshot::InvalidScreenshotId)
+		if (Event.GetCallstackId() != 0)
 		{
 			if (const TraceServices::ICallstacksProvider* CallstackProvider = TraceServices::ReadCallstacksProvider(*Session.Get()))
 			{
@@ -723,6 +724,66 @@ void FMarkersTimingTrack::InitTooltip(FTooltipDrawState& InOutTooltip, const ITi
 
 ////////////////////////////////////////////////////////////////////////////////////////////////////
 
+void FMarkersTimingTrack::OnClipboardCopyEvent(const ITimingEvent& InSelectedEvent) const
+{
+	TSharedPtr<const TraceServices::IAnalysisSession> Session = FInsightsManager::Get()->GetSession();
+	if (!Session.IsValid())
+	{
+		return;
+	}
+
+	if (!InSelectedEvent.CheckTrack(this))
+	{
+		return;
+	}
+
+	if (!InSelectedEvent.Is<FLogMessageTimingEvent>())
+	{
+		return;
+	}
+
+	TStringBuilder<2048> Result;
+
+	const FLogMessageTimingEvent& LogMessageTimingEvent = InSelectedEvent.As<FLogMessageTimingEvent>();
+	{
+		TraceServices::FAnalysisSessionReadScope SessionReadScope(*Session.Get());
+
+		const TraceServices::ILogProvider& LogProvider = TraceServices::ReadLogProvider(*Session.Get());
+		LogProvider.ReadMessage(LogMessageTimingEvent.GetLogIndex(),
+			[&Result, this](const TraceServices::FLogMessageInfo& Message)
+			{
+				Result.Append(Message.Message);
+			});
+	}
+
+	if (InSelectedEvent.Is<FBookmarkTimingEvent>())
+	{
+		const FBookmarkTimingEvent& Event = InSelectedEvent.As<FBookmarkTimingEvent>();
+
+		if (Event.GetCallstackId() != 0)
+		{
+			if (const TraceServices::ICallstacksProvider* CallstackProvider = TraceServices::ReadCallstacksProvider(*Session.Get()))
+			{
+				if (const TraceServices::FCallstack* Callstack = CallstackProvider->GetCallstack(Event.GetCallstackId()))
+				{
+					uint32 CallstackSize = Callstack->Num();
+					for (uint8 Idx = 0; Idx < CallstackSize; ++Idx)
+					{
+						TStringBuilder<512> CallstackText;
+						UE::Insights::FormatStackFrame(*Callstack->Frame(Idx), CallstackText, UE::Insights::EStackFrameFormatFlags::ModuleSymbolFileAndLine);
+						Result.Appendf(TEXT("\n%s"), *CallstackText);
+					}
+				}
+			}
+		}
+	}
+
+	// Copy callstack to the clipboard.
+	FPlatformApplicationMisc::ClipboardCopy(*Result);
+}
+
+////////////////////////////////////////////////////////////////////////////////////////////////////
+
 bool FMarkersTimingTrack::SaveScreenshot_CanExecute()
 {
 	return LastScreenshotId != TraceServices::FScreenshot::InvalidScreenshotId;
diff --git a/Engine/Source/Developer/TraceInsights/Private/Insights/TimingProfiler/Tracks/MarkersTimingTrack.h b/Engine/Source/Developer/TraceInsights/Private/Insights/TimingProfiler/Tracks/MarkersTimingTrack.h
--- a/Engine/Source/Developer/TraceInsights/Private/Insights/TimingProfiler/Tracks/MarkersTimingTrack.h
+++ b/Engine/Source/Developer/TraceInsights/Private/Insights/TimingProfiler/Tracks/MarkersTimingTrack.h
@@ -10,6 +10,7 @@
 #include "Insights/ViewModels/BaseTimingTrack.h"
 #include "Insights/ViewModels/TrackHeader.h"
 
+class ITimingEvent;
 class FSlateFontMeasure;
 struct FSlateBrush;
 
@@ -95,6 +96,7 @@ public:
 	virtual void BuildContextMenu(FMenuBuilder& MenuBuilder) override;
 	virtual const TSharedPtr<const ITimingEvent> GetEvent(float InPosX, float InPosY, const FTimingTrackViewport& Viewport) const override;
 	virtual void InitTooltip(FTooltipDrawState& InOutTooltip, const ITimingEvent& InTooltipEvent) const override;
+	virtual void OnClipboardCopyEvent(const ITimingEvent& InSelectedEvent) const override;
 
 	double Snap(double Time, double SnapTolerance);
 
