[Insights] Fix GameFrame/RenderingFrame mode aggregation in the Timers Panel.

diff --git a/Engine/Source/Developer/TraceInsights/Private/Insights/TimingProfiler/ViewModels/TimerAggregator.cpp b/Engine/Source/Developer/TraceInsights/Private/Insights/TimingProfiler/ViewModels/TimerAggregator.cpp
--- a/Engine/Source/Developer/TraceInsights/Private/Insights/TimingProfiler/ViewModels/TimerAggregator.cpp
+++ b/Engine/Source/Developer/TraceInsights/Private/Insights/TimingProfiler/ViewModels/TimerAggregator.cpp
@@ -102,6 +102,7 @@ IStatsAggregationWorker* FTimerAggregator::CreateWorker(TSharedPtr<const TraceSe
 
 	Params.StartTime = GetIntervalStartTime();
 	Params.EndTime = GetIntervalEndTime();
+	Params.FrameType = GetFrameType();
 
 	TSharedPtr<STimingProfilerWindow> Wnd = FTimingProfilerManager::Get()->GetProfilerWindow();
 	if (Wnd.IsValid())
