Context Menu option to toggle between relative and realtime timestamps in Log View

diff --git a/Engine/Source/Developer/TraceAnalysis/Private/Analysis/Engine.cpp b/Engine/Source/Developer/TraceAnalysis/Private/Analysis/Engine.cpp
--- a/Engine/Source/Developer/TraceAnalysis/Private/Analysis/Engine.cpp
+++ b/Engine/Source/Developer/TraceAnalysis/Private/Analysis/Engine.cpp
@@ -316,6 +316,7 @@ struct FTiming
 	uint64	TimestampHz = 0;
 	double	InvTimestampHz = 0.0;
 	uint64	EventTimestamp = 0;
+	double  BaseDateTimestamp = 0.0;
 };
 
 
@@ -355,6 +356,12 @@ double IAnalyzer::FEventTime::AsSecondsAbsolute(int64 DurationCycles64) const
 	return double(DurationCycles64) * Timing->InvTimestampHz;
 }
 
+////////////////////////////////////////////////////////////////////////////////
+double IAnalyzer::FEventTime::GetBaseDateTimestamp() const
+{
+	const auto* Timing = (const FTiming*)this;
+	return Timing->BaseDateTimestamp;
+}
 
 
 // {{{1 dispatch ---------------------------------------------------------------
@@ -1964,10 +1971,12 @@ void FTraceAnalyzer::OnTiming(const FOnEventContext& Context)
 {
 	uint64 StartCycle = Context.EventData.GetValue<uint64>("StartCycle");
 	uint64 CycleFrequency = Context.EventData.GetValue<uint64>("CycleFrequency");
+	double StartDateTime = Context.EventData.GetValue<double>("StartDateTime");
 
 	State.Timing.BaseTimestamp = StartCycle;
 	State.Timing.TimestampHz = CycleFrequency;
 	State.Timing.InvTimestampHz = 1.0 / double(CycleFrequency);
+	State.Timing.BaseDateTimestamp = StartDateTime;
 }
 
 ////////////////////////////////////////////////////////////////////////////////
diff --git a/Engine/Source/Developer/TraceAnalysis/Public/Trace/Analyzer.h b/Engine/Source/Developer/TraceAnalysis/Public/Trace/Analyzer.h
--- a/Engine/Source/Developer/TraceAnalysis/Public/Trace/Analyzer.h
+++ b/Engine/Source/Developer/TraceAnalysis/Public/Trace/Analyzer.h
@@ -300,6 +300,9 @@ public:
 
 		/** As AsSeconds(Cycles64) but absolute. */
 		UE_API double AsSecondsAbsolute(int64 DurationCycles64) const;
+
+		/** Returns the double timestamp in unix for when the trace started. */
+		UE_API double GetBaseDateTimestamp() const;
 	};
 
 	struct FOnEventContext
diff --git a/Engine/Source/Developer/TraceInsights/Private/Insights/ViewModels/LogMessage.cpp b/Engine/Source/Developer/TraceInsights/Private/Insights/ViewModels/LogMessage.cpp
--- a/Engine/Source/Developer/TraceInsights/Private/Insights/ViewModels/LogMessage.cpp
+++ b/Engine/Source/Developer/TraceInsights/Private/Insights/ViewModels/LogMessage.cpp
@@ -14,6 +14,8 @@
 // TraceInsights
 #include "Insights/Log.h"
 
+#define LOCTEXT_NAMESPACE "UE::Insights::SLogView"
+
 ////////////////////////////////////////////////////////////////////////////////////////////////////
 // FLogMessageRecord
 ////////////////////////////////////////////////////////////////////////////////////////////////////
@@ -57,6 +59,26 @@ FText FLogMessageRecord::GetTimeAsText() const
 
 ////////////////////////////////////////////////////////////////////////////////////////////////////
 
+FText FLogMessageRecord::GetAbsoluteDateTime(double BaseDateTime) const
+{
+	if (BaseDateTime == 0)
+	{
+		return LOCTEXT("UnknownDateTime", "Unknown Date-Time");
+	}
+	
+	FDateTime DT = FDateTime::FromUnixTimestampDecimal(BaseDateTime + Time);
+	return FText::FromString(FString::Printf(TEXT("%04d.%02d.%02d-%02d.%02d.%02d:%03d"),
+		DT.GetYear(),
+		DT.GetMonth(),
+		DT.GetDay(),
+		DT.GetHour(),
+		DT.GetMinute(),
+		DT.GetSecond(),
+		DT.GetMillisecond()));
+}
+
+////////////////////////////////////////////////////////////////////////////////////////////////////
+
 FText FLogMessageRecord::GetVerbosityAsText() const
 {
 	return FText::FromString(FString(ToString(Verbosity)));
@@ -170,3 +192,5 @@ TSharedPtr<FLogMessageRecord> FLogMessageCache::GetUncached(uint64 Index) const
 
 ////////////////////////////////////////////////////////////////////////////////////////////////////
 ////////////////////////////////////////////////////////////////////////////////////////////////////
+
+#undef LOCTEXT_NAMESPACE
diff --git a/Engine/Source/Developer/TraceInsights/Private/Insights/ViewModels/LogMessage.h b/Engine/Source/Developer/TraceInsights/Private/Insights/ViewModels/LogMessage.h
--- a/Engine/Source/Developer/TraceInsights/Private/Insights/ViewModels/LogMessage.h
+++ b/Engine/Source/Developer/TraceInsights/Private/Insights/ViewModels/LogMessage.h
@@ -33,6 +33,7 @@ public:
 
 	FText GetIndexAsText() const;
 	FText GetTimeAsText() const;
+	FText GetAbsoluteDateTime(double BaseDateTime) const;
 	FText GetVerbosityAsText() const;
 	FText GetCategoryAsText() const { return Category; }
 	FText GetMessageAsText() const { return Message; }
diff --git a/Engine/Source/Developer/TraceInsights/Private/Insights/Widgets/SLogView.cpp b/Engine/Source/Developer/TraceInsights/Private/Insights/Widgets/SLogView.cpp
--- a/Engine/Source/Developer/TraceInsights/Private/Insights/Widgets/SLogView.cpp
+++ b/Engine/Source/Developer/TraceInsights/Private/Insights/Widgets/SLogView.cpp
@@ -88,6 +88,12 @@ public:
 			EUserInterfaceActionType::Button,
 			FInputChord());
 
+		UI_COMMAND(Command_ToggleTimeFormat,
+			"Toggle Time Format",
+			"Toggles between Relative and Real Time Formats.",
+			EUserInterfaceActionType::Button,
+			FInputChord());
+
 		UI_COMMAND(Command_CopySelected,
 			"Copy",
 			"Copies the selected log (with all its properties) to clipboard.",
@@ -131,6 +137,7 @@ public:
 	TSharedPtr<FUICommandInfo> Command_HideSelectedCategory;
 	TSharedPtr<FUICommandInfo> Command_ShowOnlySelectedCategory;
 	TSharedPtr<FUICommandInfo> Command_ShowAllCategories;
+	TSharedPtr<FUICommandInfo> Command_ToggleTimeFormat;
 	TSharedPtr<FUICommandInfo> Command_CopySelected;
 	TSharedPtr<FUICommandInfo> Command_CopyMessage;
 	TSharedPtr<FUICommandInfo> Command_CopyRange;
@@ -189,7 +196,17 @@ public:
 					.Text(this, &SLogMessageRow::GetIndex)
 				];
 		}
-		else if (ColumnName == LogViewColumns::TimeColumnName)
+		else if (ColumnName == LogViewColumns::SessionTimeColumnName)
+		{
+			return SNew(SBox)
+				.Padding(FMargin(4.0, 0.0))
+				[
+					SNew(STextBlock)
+					.Text(this, &SLogMessageRow::GetTime)
+					.ColorAndOpacity(this, &SLogMessageRow::GetColorAndOpacity)
+				];
+		}
+		else if (ColumnName == LogViewColumns::RealTimeColumnName)
 		{
 			return SNew(SBox)
 				.Padding(FMargin(4.0, 0.0))
@@ -349,7 +366,8 @@ public:
 		if (ParentWidgetPin.IsValid() && LogMessagePin.IsValid())
 		{
 			FLogMessageRecord& CacheEntry = ParentWidgetPin->GetCache().Get(LogMessagePin->GetIndex());
-			return CacheEntry.GetTimeAsText();
+			TSharedPtr<const TraceServices::IAnalysisSession> Session = FInsightsManager::Get()->GetSession();
+			return ParentWidgetPin->IsUsingRelativeTime() ? CacheEntry.GetTimeAsText() : CacheEntry.GetAbsoluteDateTime(Session->GetBaseDateTime());
 		}
 		else
 		{
@@ -515,6 +533,7 @@ SLogView::SLogView()
 	, TotalNumMessages(0)
 	, TotalNumInserts(0)
 	, bIsDirty(false)
+	, bUseRelativeTime(true)
 {
 }
 
@@ -567,6 +586,8 @@ void SLogView::Reset()
 
 	FilteredMessages.Reset();
 
+	bUseRelativeTime = true;
+
 	ListView->RebuildList();
 }
 
@@ -670,9 +691,9 @@ void SLogView::Construct(const FArguments& InArgs)
 						.ManualWidth(60.0f)
 						.DefaultLabel(LOCTEXT("IdColumn", "Index"))
 
-						+ SHeaderRow::Column(LogViewColumns::TimeColumnName)
-						.ManualWidth(94.0f)
-						.DefaultLabel(LOCTEXT("TimeColumn", "Time"))
+						+ SHeaderRow::Column(LogViewColumns::SessionTimeColumnName)
+						.ManualWidth(100.0f)
+						.DefaultLabel(GetTimeFormatColumnName())
 
 						+ SHeaderRow::Column(LogViewColumns::VerbosityColumnName)
 						.ManualWidth(80.0f)
@@ -727,6 +748,7 @@ void SLogView::InitCommandList()
 	CommandList->MapAction(FLogViewCommands::Get().Command_HideSelectedCategory, FExecuteAction::CreateSP(this, &SLogView::HideSelectedCategory), FCanExecuteAction::CreateSP(this, &SLogView::CanHideSelectedCategory));
 	CommandList->MapAction(FLogViewCommands::Get().Command_ShowOnlySelectedCategory, FExecuteAction::CreateSP(this, &SLogView::ShowOnlySelectedCategory), FCanExecuteAction::CreateSP(this, &SLogView::CanShowOnlySelectedCategory));
 	CommandList->MapAction(FLogViewCommands::Get().Command_ShowAllCategories, FExecuteAction::CreateSP(this, &SLogView::ShowAllCategories), FCanExecuteAction::CreateSP(this, &SLogView::CanShowAllCategories));
+	CommandList->MapAction(FLogViewCommands::Get().Command_ToggleTimeFormat, FExecuteAction::CreateSP(this, &SLogView::ToggleTimeFormat), FCanExecuteAction::CreateSP(this, &SLogView::CanToggleTimeFormat));
 	CommandList->MapAction(FLogViewCommands::Get().Command_CopySelected, FExecuteAction::CreateSP(this, &SLogView::CopySelected), FCanExecuteAction::CreateSP(this, &SLogView::CanCopySelected));
 	CommandList->MapAction(FLogViewCommands::Get().Command_CopyMessage, FExecuteAction::CreateSP(this, &SLogView::CopyMessage), FCanExecuteAction::CreateSP(this, &SLogView::CanCopyMessage));
 	CommandList->MapAction(FLogViewCommands::Get().Command_CopyRange, FExecuteAction::CreateSP(this, &SLogView::CopyRange), FCanExecuteAction::CreateSP(this, &SLogView::CanCopyRange));
@@ -1223,6 +1245,19 @@ TSharedPtr<SWidget> SLogView::ListView_GetContextMenu()
 			FSlateIcon(FAppStyle::Get().GetStyleSetName(), "Icons.Visible")
 		);
 
+		MenuBuilder.AddMenuEntry
+		(
+			FLogViewCommands::Get().Command_ToggleTimeFormat,
+			NAME_None,
+			bUseRelativeTime ?
+			LOCTEXT("SwitchToRealTime", "Switch to Real Time") :
+			LOCTEXT("SwitchToSessionTime", "Switch To Session Time"),
+			bUseRelativeTime ?
+			LOCTEXT("ToggleTimeFormatRealTime_Tooltip", "Changes the Time column to reflect real (UTC) timestamps.") :
+			LOCTEXT("ToggleTimeFormatSessionTime_Tooltip", "Changes the Time column to reflect session relative timestamps."),
+			FSlateIcon(FInsightsStyle::GetStyleSetName(), "Icons.TimersView")
+		);
+
 		MenuBuilder.AddSeparator();
 
 		MenuBuilder.AddMenuEntry
@@ -1646,6 +1681,35 @@ void SLogView::ShowAllCategories()
 
 ////////////////////////////////////////////////////////////////////////////////////////////////////
 
+bool SLogView::CanToggleTimeFormat() const
+{
+	return true;
+}
+
+////////////////////////////////////////////////////////////////////////////////////////////////////
+
+void SLogView::ToggleTimeFormat()
+{
+	bUseRelativeTime = !bUseRelativeTime;
+	TSharedPtr<SHeaderRow> Row = ListView->GetHeaderRow();
+	
+	Row->RemoveColumn(Row->GetColumns()[1].ColumnId);
+	Row->InsertColumn(
+		SHeaderRow::Column(bUseRelativeTime ? LogViewColumns::SessionTimeColumnName : LogViewColumns::RealTimeColumnName)
+		.ManualWidth(bUseRelativeTime ? 100.0f : 155.0f)
+		.DefaultLabel(GetTimeFormatColumnName())
+		, 1);
+}
+
+////////////////////////////////////////////////////////////////////////////////////////////////////
+
+FText SLogView::GetTimeFormatColumnName() const
+{
+	return bUseRelativeTime ? LOCTEXT("SessionTimeColumn", "Session Time") : LOCTEXT("RealTimeColumn", "Real Time");
+}
+
+////////////////////////////////////////////////////////////////////////////////////////////////////
+
 FReply SLogView::OnKeyDown(const FGeometry& MyGeometry, const FKeyEvent& InKeyEvent)
 {
 	return CommandList->ProcessCommandBindings(InKeyEvent) == true ? FReply::Handled() : FReply::Unhandled();
@@ -1656,9 +1720,10 @@ FReply SLogView::OnKeyDown(const FGeometry& MyGeometry, const FKeyEvent& InKeyEv
 void SLogView::AppendFormatMessageDetailed(const FLogMessageRecord& Log, TStringBuilderBase<TCHAR>& InOutStringBuilder) const
 {
 	using namespace UE::Insights;
+	TSharedPtr<const TraceServices::IAnalysisSession> Session = FInsightsManager::Get()->GetSession();
 	InOutStringBuilder.Appendf(TEXT("Index=%d"), Log.GetIndex());
 	InOutStringBuilder.Append(TEXT("\nTime="));
-	InOutStringBuilder.Append(FormatTimeHMS(Log.GetTime(), FTimeValue::Microsecond));
+	InOutStringBuilder.Append(bUseRelativeTime ? FormatTimeHMS(Log.GetTime(), FTimeValue::Microsecond) : Log.GetAbsoluteDateTime(Session->GetBaseDateTime()).ToString());
 	InOutStringBuilder.Append(TEXT("\nVerbosity="));
 	InOutStringBuilder.Append(::ToString(Log.GetVerbosity()));
 	InOutStringBuilder.Append(TEXT("\nCategory="));
@@ -1693,9 +1758,10 @@ void SLogView::AppendFormatMessageDelimitedHeader(TStringBuilderBase<TCHAR>& InO
 
 void SLogView::AppendFormatMessageDelimited(const FLogMessageRecord& Log, TStringBuilderBase<TCHAR>& InOutStringBuilder, TCHAR Separator) const
 {
+	TSharedPtr<const TraceServices::IAnalysisSession> Session = FInsightsManager::Get()->GetSession();
 	InOutStringBuilder.Appendf(TEXT("%d"), Log.GetIndex());
 	InOutStringBuilder.AppendChar(Separator);
-	InOutStringBuilder.Append(UE::Insights::FormatTimeHMS(Log.GetTime(), UE::Insights::FTimeValue::Microsecond));
+	InOutStringBuilder.Append(bUseRelativeTime ? FormatTimeHMS(Log.GetTime(), FTimeValue::Microsecond) : Log.GetAbsoluteDateTime(Session->GetBaseDateTime()).ToString());
 	InOutStringBuilder.AppendChar(Separator);
 	InOutStringBuilder.Append(::ToString(Log.GetVerbosity()));
 	InOutStringBuilder.AppendChar(Separator);
diff --git a/Engine/Source/Developer/TraceInsights/Private/Insights/Widgets/SLogView.h b/Engine/Source/Developer/TraceInsights/Private/Insights/Widgets/SLogView.h
--- a/Engine/Source/Developer/TraceInsights/Private/Insights/Widgets/SLogView.h
+++ b/Engine/Source/Developer/TraceInsights/Private/Insights/Widgets/SLogView.h
@@ -31,7 +31,8 @@ namespace UE::Insights
 namespace LogViewColumns
 {
 	static const FName IdColumnName(TEXT("Id"));
-	static const FName TimeColumnName(TEXT("Time"));
+	static const FName SessionTimeColumnName(TEXT("SessionTime"));
+	static const FName RealTimeColumnName(TEXT("RealTime"));
 	static const FName VerbosityColumnName(TEXT("Verbosity"));
 	static const FName CategoryColumnName(TEXT("Category"));
 	static const FName MessageColumnName(TEXT("Message"));
@@ -88,6 +89,8 @@ public:
 
 	TSharedPtr<TimingProfiler::STimingView> GetTimingView() const;
 
+	bool IsUsingRelativeTime() const { return bUseRelativeTime; }
+
 protected:
 	/** Generate a new list view row. */
 	TSharedRef<ITableRow> OnGenerateRow(TSharedPtr<FLogMessage> InLogMessage, const TSharedRef<STableViewBase>& OwnerTable);
@@ -130,6 +133,10 @@ protected:
 
 	bool CanShowAllCategories() const;
 	void ShowAllCategories();
+	
+	bool CanToggleTimeFormat() const;
+	void ToggleTimeFormat();
+	FText GetTimeFormatColumnName() const;
 
 	void AppendFormatMessageDetailed(const FLogMessageRecord& Log, TStringBuilderBase<TCHAR>& InOutStringBuilder) const;
 	void AppendFormatMessageDelimitedHeader(TStringBuilderBase<TCHAR>& InOutStringBuilder, TCHAR Separator = TEXT('\t')) const;
@@ -214,6 +221,9 @@ protected:
 
 	/** List of trace log messages to show in list view. */
 	TArray<TSharedPtr<FLogMessage>> FilteredMessages;
+
+	/** Whether to use Session Relative or Real Time (Absolute) Timestamp for the Log "Time" Column. */
+	bool bUseRelativeTime;
 };
 
 ////////////////////////////////////////////////////////////////////////////////////////////////////
diff --git a/Engine/Source/Developer/TraceServices/Private/AnalysisServicePrivate.h b/Engine/Source/Developer/TraceServices/Private/AnalysisServicePrivate.h
--- a/Engine/Source/Developer/TraceServices/Private/AnalysisServicePrivate.h
+++ b/Engine/Source/Developer/TraceServices/Private/AnalysisServicePrivate.h
@@ -53,6 +53,9 @@ public:
 	virtual double GetDurationSeconds() const override { Lock.ReadAccessCheck(); return DurationSeconds; }
 	virtual void UpdateDurationSeconds(double Duration) override { Lock.WriteAccessCheck(); DurationSeconds = FMath::Max(Duration, DurationSeconds); }
 
+	virtual double GetBaseDateTime() const override { return BaseDateTime; }
+	virtual void SetBaseDateTime(double InBaseDateTime) override { BaseDateTime = InBaseDateTime; }
+	
 	virtual uint32 GetMetadataCount() const override { Lock.ReadAccessCheck(); return Metadata.Num(); }
 	virtual void EnumerateMetadata(TFunctionRef<void(const FTraceSessionMetadata& Metadata)> Callback) const override;
 	virtual void AddMetadata(FName InName, int64 InValue) override;
@@ -95,6 +98,7 @@ private:
 	FString Name;
 	uint32 TraceId = 0;
 	double DurationSeconds = 0.0;
+	double BaseDateTime = 0.0;
 	TMap<FName, FTraceSessionMetadata> Metadata;
 	FSlabAllocator Allocator;
 	FStringStore StringStore;
diff --git a/Engine/Source/Developer/TraceServices/Private/Analyzers/MiscTraceAnalysis.cpp b/Engine/Source/Developer/TraceServices/Private/Analyzers/MiscTraceAnalysis.cpp
--- a/Engine/Source/Developer/TraceServices/Private/Analyzers/MiscTraceAnalysis.cpp
+++ b/Engine/Source/Developer/TraceServices/Private/Analyzers/MiscTraceAnalysis.cpp
@@ -109,6 +109,12 @@ bool FMiscTraceAnalyzer::OnEvent(uint16 RouteId, EStyle Style, const FOnEventCon
 	FAnalysisSessionEditScope _(Session);
 
 	const auto& EventData = Context.EventData;
+
+	if (Session.GetBaseDateTime() == 0.0)
+	{
+		Session.SetBaseDateTime(Context.EventTime.GetBaseDateTimestamp());
+	}
+
 	switch (RouteId)
 	{
 	case RouteId_BeginFrame:
diff --git a/Engine/Source/Developer/TraceServices/Public/TraceServices/Model/AnalysisSession.h b/Engine/Source/Developer/TraceServices/Public/TraceServices/Model/AnalysisSession.h
--- a/Engine/Source/Developer/TraceServices/Public/TraceServices/Model/AnalysisSession.h
+++ b/Engine/Source/Developer/TraceServices/Public/TraceServices/Model/AnalysisSession.h
@@ -89,6 +89,14 @@ public:
 	 */
 	virtual void UpdateDurationSeconds(double Duration) = 0;
 
+	/**
+	 * Getter for the unix timestamp at the start of all analysis.
+	 * This function should be used as an offset for any relative (session) timestamps to calculate
+	 * real UTC timestamps.
+	 */
+	virtual double GetBaseDateTime() const = 0;
+	virtual void SetBaseDateTime(double InBaseDateTime) = 0;
+
 	virtual uint32 GetMetadataCount() const = 0;
 	virtual void EnumerateMetadata(TFunctionRef<void(const FTraceSessionMetadata& Metadata)> Callback) const = 0;
 	virtual void AddMetadata(FName InName, int64 InValue) = 0;
diff --git a/Engine/Source/Runtime/TraceLog/Private/Trace/Writer.cpp b/Engine/Source/Runtime/TraceLog/Private/Trace/Writer.cpp
--- a/Engine/Source/Runtime/TraceLog/Private/Trace/Writer.cpp
+++ b/Engine/Source/Runtime/TraceLog/Private/Trace/Writer.cpp
@@ -13,6 +13,7 @@
 #include "Trace/Trace.inl"
 #include "WriteBufferRedirect.h"
 
+#include <chrono>
 #include <limits.h>
 #include <stdlib.h>
 
@@ -78,6 +79,7 @@ UE_TRACE_MINIMAL_EVENT_BEGIN($Trace, NewTrace, Important | NoSync)
 	UE_TRACE_MINIMAL_EVENT_FIELD(uint64, CycleFrequency)
 	UE_TRACE_MINIMAL_EVENT_FIELD(uint16, Endian)
 	UE_TRACE_MINIMAL_EVENT_FIELD(uint8, PointerSize)
+	UE_TRACE_MINIMAL_EVENT_FIELD(double, StartDateTime)
 UE_TRACE_MINIMAL_EVENT_END()
 
 ////////////////////////////////////////////////////////////////////////////////
@@ -944,8 +946,11 @@ static void Writer_InternalInitializeImpl()
 		return;
 	}
 
+	double StartDateTime = std::chrono::duration_cast<std::chrono::duration<double>>
+	(std::chrono::system_clock::now().time_since_epoch()).count();
+	
 	GStartCycle = TimeGetTimestamp();
-
+	
 	Writer_InitializeSharedBuffers();
 	Writer_InitializePool();
 	Writer_InitializeControl();
@@ -956,7 +961,8 @@ static void Writer_InternalInitializeImpl()
 		<< NewTrace.StartCycle(GStartCycle)
 		<< NewTrace.CycleFrequency(TimeGetFrequency())
 		<< NewTrace.Endian(uint16(0x524d))
-		<< NewTrace.PointerSize(uint8(sizeof(void*)));
+		<< NewTrace.PointerSize(uint8(sizeof(void*)))
+		<< NewTrace.StartDateTime(StartDateTime);
 }
 
 ////////////////////////////////////////////////////////////////////////////////
