Option to adjust the ratio of threads to be used in Aggregation

diff --git a/Engine/Source/Developer/TraceInsights/Private/Insights/InsightsSettings.cpp b/Engine/Source/Developer/TraceInsights/Private/Insights/InsightsSettings.cpp
--- a/Engine/Source/Developer/TraceInsights/Private/Insights/InsightsSettings.cpp
+++ b/Engine/Source/Developer/TraceInsights/Private/Insights/InsightsSettings.cpp
@@ -57,6 +57,7 @@ void FInsightsSettings::ResetToDefaults()
 	AutoScrollFrameAlignment = Defaults.AutoScrollFrameAlignment;
 	AutoScrollViewportOffsetPercent = Defaults.AutoScrollViewportOffsetPercent;
 	AutoScrollMinDelay = Defaults.AutoScrollMinDelay;
+	ThreadsToUseRatio = Defaults.ThreadsToUseRatio;
 }
 
 ////////////////////////////////////////////////////////////////////////////////////////////////////
@@ -74,6 +75,7 @@ void FInsightsSettings::LoadFromConfig()
 	GConfig->GetDouble(TEXT("Insights.TimingProfiler"), TEXT("DefaultZoomLevel"), DefaultZoomLevel, SettingsIni);
 	GConfig->GetBool(TEXT("Insights.TimingProfiler"), TEXT("bAutoHideEmptyTracks"), bAutoHideEmptyTracks, SettingsIni);
 	GConfig->GetBool(TEXT("Insights.TimingProfiler"), TEXT("bAllowPanningOnScreenEdges"), bAllowPanningOnScreenEdges, SettingsIni);
+	GConfig->GetDouble(TEXT("Insights.TimingProfiler"), TEXT("ThreadsToUseRatio"), ThreadsToUseRatio, SettingsIni);
 
 	// Auto-scroll options
 	GConfig->GetBool(TEXT("Insights.TimingProfiler"), TEXT("bAutoScroll"), bAutoScroll, SettingsIni);
@@ -206,6 +208,7 @@ void FInsightsSettings::SaveToConfig()
 	GConfig->SetDouble(TEXT("Insights.TimingProfiler"), TEXT("DefaultZoomLevel"), DefaultZoomLevel, SettingsIni);
 	GConfig->SetBool(TEXT("Insights.TimingProfiler"), TEXT("bAutoHideEmptyTracks"), bAutoHideEmptyTracks, SettingsIni);
 	GConfig->SetBool(TEXT("Insights.TimingProfiler"), TEXT("bAllowPanningOnScreenEdges"), bAllowPanningOnScreenEdges, SettingsIni);
+	GConfig->SetDouble(TEXT("Insights.TimingProfiler"), TEXT("ThreadsToUseRatio"), ThreadsToUseRatio, SettingsIni);
 
 	// Auto-scroll options
 	GConfig->SetBool(TEXT("Insights.TimingProfiler"), TEXT("bAutoScroll"), bAutoScroll, SettingsIni);
diff --git a/Engine/Source/Developer/TraceInsights/Private/Insights/InsightsSettings.h b/Engine/Source/Developer/TraceInsights/Private/Insights/InsightsSettings.h
--- a/Engine/Source/Developer/TraceInsights/Private/Insights/InsightsSettings.h
+++ b/Engine/Source/Developer/TraceInsights/Private/Insights/InsightsSettings.h
@@ -72,6 +72,10 @@ public:
 	void SetAutoScrollMinDelay(double Delay) { AutoScrollMinDelay = Delay; }
 	void SetAndSaveAutoScrollMinDelay(double Delay) { SET_AND_SAVE(AutoScrollMinDelay, Delay); }
 
+	double GetThreadRatio() const { return ThreadsToUseRatio; }
+	void SetThreadRatio(double Ratio) { ThreadsToUseRatio = Ratio; }
+	void SetAndSaveThreadRatio(double Ratio) { SET_AND_SAVE(ThreadsToUseRatio, Ratio); }
+
 	//////////////////////////////////////////////////
 	// [Insights.TimingProfiler.TimingView]
 
@@ -285,6 +289,9 @@ private:
 	/** Minimum time between two auto-scroll updates, in [seconds]. */
 	double AutoScrollMinDelay = 0.3; // [seconds]
 
+	/** The ratio of available threads to use when performing per-frame aggregations. */
+	double ThreadsToUseRatio = 0.75;
+
 	//////////////////////////////////////////////////
 	// [Insights.TimingProfiler.TimingView]
 
diff --git a/Engine/Source/Developer/TraceInsights/Private/Insights/TimingProfiler/Widgets/STimersView.cpp b/Engine/Source/Developer/TraceInsights/Private/Insights/TimingProfiler/Widgets/STimersView.cpp
--- a/Engine/Source/Developer/TraceInsights/Private/Insights/TimingProfiler/Widgets/STimersView.cpp
+++ b/Engine/Source/Developer/TraceInsights/Private/Insights/TimingProfiler/Widgets/STimersView.cpp
@@ -58,6 +58,7 @@
 #include "Insights/Widgets/STimingView.h"
 
 #include <limits>
+#include "Widgets/Input/SSpinBox.h"
 
 #define LOCTEXT_NAMESPACE "UE::Insights::TimingProfiler::STimersView"
 
@@ -312,6 +313,23 @@ void STimersView::Construct(const FArguments& InArgs)
 					GetToggleButtonForTimerType(ETimerNodeType::CpuScope)
 				]
 
+				// Adjust the ratio of threads being used in aggregation calculations
+				+ SHorizontalBox::Slot()
+				.Padding(4.0f, 0.0f, 4.0f, 0.0f)
+				.AutoWidth()
+				.VAlign(VAlign_Center)
+				[
+					SNew(SSpinBox<double>)
+					.ToolTipText(LOCTEXT("ThreadsAggregation_Tooltip", "Ratio of Threads being used for Aggregation.\n" "Using higher values may cause systems to freeze."))
+					.MinValue(this, &STimersView::GetMinimumThreadRatio)
+					.MaxValue(1.0)
+					.MaxFractionalDigits(2)
+					.MinFractionalDigits(2)
+					.Value(this, &STimersView::GetThreadRatio)
+					.OnValueChanged(this, &STimersView::OnThreadRatioChanged)
+					.OnValueCommitted(this, &STimersView::OnThreadRatioCommited)
+				]
+
 				// Filter out timers with zero instance count
 				+ SHorizontalBox::Slot()
 				.Padding(2.0f)
@@ -509,6 +527,15 @@ void STimersView::Construct(const FArguments& InArgs)
 
 	// Update the Session (i.e. when analysis session was already started).
 	InsightsManager_OnSessionChanged();
+	
+	const TraceServices::ITimingProfilerProvider* TimingProfilerProvider = TraceServices::ReadTimingProfilerProvider(*Session.Get());
+	if (TimingProfilerProvider)
+	{
+		TraceServices::FAnalysisSessionEditScope SessionEditScope(*const_cast<TraceServices::IAnalysisSession*>(Session.Get()));
+		
+		FInsightsSettings& Settings = FInsightsManager::Get()->GetSettings();
+		TimingProfilerProvider->SetRatioOfThreadsToUse(Settings.GetThreadRatio());
+	}
 }
 END_SLATE_FUNCTION_BUILD_OPTIMIZATION
 
@@ -1720,6 +1747,58 @@ void STimersView::SearchBox_OnTextChanged(const FText& InFilterText)
 	ApplyFiltering();
 }
 
+////////////////////////////////////////////////////////////////////////////////////////////////////
+// Aggregation Performance
+////////////////////////////////////////////////////////////////////////////////////////////////////
+
+TOptional<double> STimersView::GetMinimumThreadRatio() const
+{
+	return 1.0f / (double)GThreadPool->GetNumThreads();
+}
+
+////////////////////////////////////////////////////////////////////////////////////////////////////
+
+double STimersView::GetThreadRatio() const
+{
+	if (Session.IsValid())
+	{
+		const TraceServices::ITimingProfilerProvider* TimingProfilerProvider = TraceServices::ReadTimingProfilerProvider(*Session.Get());
+		if (TimingProfilerProvider)
+		{
+			TraceServices::FAnalysisSessionReadScope SessionReadScope(*Session.Get());
+			
+			return TimingProfilerProvider->GetRatioOfThreadsToUse();
+		}
+	}
+	return 0.75;
+}
+
+////////////////////////////////////////////////////////////////////////////////////////////////////
+
+void STimersView::OnThreadRatioChanged(double InNewValue)
+{
+	OnThreadRatioCommited(InNewValue, ETextCommit::Default);
+}
+
+////////////////////////////////////////////////////////////////////////////////////////////////////
+
+void STimersView::OnThreadRatioCommited(double InNewValue, ETextCommit::Type CommitType)
+{
+	if (Session.IsValid())
+	{
+		const TraceServices::ITimingProfilerProvider* TimingProfilerProvider = TraceServices::ReadTimingProfilerProvider(*Session.Get());
+		if (TimingProfilerProvider)
+		{
+			TraceServices::FAnalysisSessionEditScope SessionEditScope(*const_cast<TraceServices::IAnalysisSession*>(Session.Get()));
+			
+			TimingProfilerProvider->SetRatioOfThreadsToUse(InNewValue);
+		}
+	}
+	
+	FInsightsSettings& Settings = FInsightsManager::Get()->GetSettings();
+	Settings.SetAndSaveThreadRatio(InNewValue);
+}
+
 ////////////////////////////////////////////////////////////////////////////////////////////////////
 
 bool STimersView::SearchBox_IsEnabled() const
diff --git a/Engine/Source/Developer/TraceInsights/Private/Insights/TimingProfiler/Widgets/STimersView.h b/Engine/Source/Developer/TraceInsights/Private/Insights/TimingProfiler/Widgets/STimersView.h
--- a/Engine/Source/Developer/TraceInsights/Private/Insights/TimingProfiler/Widgets/STimersView.h
+++ b/Engine/Source/Developer/TraceInsights/Private/Insights/TimingProfiler/Widgets/STimersView.h
@@ -262,6 +262,14 @@ private:
 	bool SearchBox_IsEnabled() const;
 	void SearchBox_OnTextChanged(const FText& InFilterText);
 
+	////////////////////////////////////////////////////////////////////////////////////////////////////
+	// Aggregation Performance
+	
+	TOptional<double> GetMinimumThreadRatio() const;
+	double GetThreadRatio() const;
+	void OnThreadRatioChanged(double InNewValue);
+	void OnThreadRatioCommited(double InNewValue, ETextCommit::Type CommitType);
+
 	////////////////////////////////////////////////////////////////////////////////////////////////////
 	// Grouping
 
diff --git a/Engine/Source/Developer/TraceServices/Private/Common/TimelineStatistics.h b/Engine/Source/Developer/TraceServices/Private/Common/TimelineStatistics.h
--- a/Engine/Source/Developer/TraceServices/Private/Common/TimelineStatistics.h
+++ b/Engine/Source/Developer/TraceServices/Private/Common/TimelineStatistics.h
@@ -93,7 +93,8 @@ public:
 											BucketMappingFunc BucketMapper,
 											const TArray<FFrameData>& Frames,
 											TSharedPtr<FCancellationToken> CancellationToken,
-											TMap<BucketKeyType, FAggregatedTimingStats>& Result)
+											TMap<BucketKeyType, FAggregatedTimingStats>& Result,
+											double RatioOfThreadsToUse)
 	{
 		TMap<BucketKeyType, FInternalFrameAggregationEntry> GlobalResult;
 		TMap<BucketKeyType, FAggregatedTimingStats> InitialTimerMap;
@@ -125,8 +126,7 @@ public:
 		}
 
 		TQueue<TSharedPtr<TMap<BucketKeyType, FAggregatedTimingStats>>, EQueueMode::Mpsc> FrameResultsQueue;
-
-		constexpr float RatioOfThreadsToUse = 0.75f;
+		
 		int32 NumTasks = FMath::Max(1, (int32)((float)GThreadPool->GetNumThreads() * RatioOfThreadsToUse));
 		NumTasks = FMath::Min((int32)NumTasks, Frames.Num());
 		int32 FramesPerTask = Frames.Num() / NumTasks;
diff --git a/Engine/Source/Developer/TraceServices/Private/Model/TimingProfiler.cpp b/Engine/Source/Developer/TraceServices/Private/Model/TimingProfiler.cpp
--- a/Engine/Source/Developer/TraceServices/Private/Model/TimingProfiler.cpp
+++ b/Engine/Source/Developer/TraceServices/Private/Model/TimingProfiler.cpp
@@ -773,7 +773,7 @@ ITable<FTimingProfilerAggregatedStats>* FTimingProfilerProvider::CreateAggregati
 		{
 			// Do not allow inf for the last frame end time.
 			Frames[Frames.Num() - 1].EndTime = FMath::Min(Session.GetDurationSeconds(), Frames[Frames.Num() - 1].EndTime);
-			FTimelineStatistics::CreateFrameStatsAggregation(IncludedTimelines, BucketMappingFunc, Frames, Params.CancellationToken, Aggregation);
+			FTimelineStatistics::CreateFrameStatsAggregation(IncludedTimelines, BucketMappingFunc, Frames, Params.CancellationToken, Aggregation, RatioOfThreadsToUse);
 		}
 	}
 
diff --git a/Engine/Source/Developer/TraceServices/Private/Model/TimingProfilerPrivate.h b/Engine/Source/Developer/TraceServices/Private/Model/TimingProfilerPrivate.h
--- a/Engine/Source/Developer/TraceServices/Private/Model/TimingProfilerPrivate.h
+++ b/Engine/Source/Developer/TraceServices/Private/Model/TimingProfilerPrivate.h
@@ -55,6 +55,8 @@ public:
 	virtual const FMetadataSpec* GetMetadataSpec(uint32 MetadataSpecId) const;
 	virtual ITable<FTimingProfilerAggregatedStats>* CreateAggregation(const FCreateAggregationParams& Params) const override;
 	virtual ITimingProfilerButterfly* CreateButterfly(const FCreateButterflyParams& Params) const override;
+	virtual double GetRatioOfThreadsToUse() const override { Session.ReadAccessCheck(); return RatioOfThreadsToUse; }
+	virtual const void SetRatioOfThreadsToUse(double InRatioOfThreadsToUse) const override { Session.WriteAccessCheck(); RatioOfThreadsToUse = InRatioOfThreadsToUse; }
 
 	// ITimingProfilerTimerReader
 	virtual const FTimingProfilerTimer* GetTimer(uint32 TimerId) const override;
@@ -116,6 +118,7 @@ private:
 	TArray<FGpuQueueData> GpuQueueData;
 	TMap<uint32, uint32> GpuQueueIdToQueueIndexMap; // GPU Queue Id --> queue index, in GpuQueues
 	TTableLayout<FTimingProfilerAggregatedStats> AggregatedStatsTableLayout;
+	mutable double RatioOfThreadsToUse = 0.75;
 
 };
 
diff --git a/Engine/Source/Developer/TraceServices/Public/TraceServices/Model/TimingProfiler.h b/Engine/Source/Developer/TraceServices/Public/TraceServices/Model/TimingProfiler.h
--- a/Engine/Source/Developer/TraceServices/Public/TraceServices/Model/TimingProfiler.h
+++ b/Engine/Source/Developer/TraceServices/Public/TraceServices/Model/TimingProfiler.h
@@ -292,6 +292,18 @@ public:
 		Params.CpuThreadFilter = [CpuThreadFilter](uint32 QueueId) { return CpuThreadFilter(QueueId); };
 		return CreateButterfly(Params);
 	}
+
+	/*
+	* Gets the Ratio of Threads to use.
+	* @return The ratio scaled between 0-1.
+	*/
+	virtual double GetRatioOfThreadsToUse() const { return 0.75; }
+
+	/*
+	* Sets the Ratio of Threads to be used.
+	* @param InRatioOfThreadsToUse The new ratio of threads to use.
+	*/
+	virtual void const SetRatioOfThreadsToUse(double InRatioOfThreadsToUse) const {}
 };
 
 /*
