Ability to persistently save CPU Thread Groups / Track visibility states in Timing View

diff --git a/Engine/Source/Developer/TraceInsights/Private/Insights/InsightsSettings.cpp b/Engine/Source/Developer/TraceInsights/Private/Insights/InsightsSettings.cpp
--- a/Engine/Source/Developer/TraceInsights/Private/Insights/InsightsSettings.cpp
+++ b/Engine/Source/Developer/TraceInsights/Private/Insights/InsightsSettings.cpp
@@ -112,6 +112,21 @@ void FInsightsSettings::LoadFromConfig()
 	GConfig->GetBool(TEXT("Insights.TimingProfiler.TimingView"), TEXT("ShowGpuFencesExtendedLines"), bTimingViewShowGpuFencesExtendedLines, SettingsIni);
 	GConfig->GetBool(TEXT("Insights.TimingProfiler.TimingView"), TEXT("ShowGpuFencesRelations"), bTimingViewShowGpuFencesRelations, SettingsIni);
 
+	TArray<FString> VisibleCpuThreadGroups;
+	TArray<FString> HiddenCpuThreadGroups;
+	GConfig->GetArray(TEXT("Insights.TimingProfiler.TimingView.CpuThreadGroups"), TEXT("VisibleCpuThreadGroups"), VisibleCpuThreadGroups, SettingsIni);
+	GConfig->GetArray(TEXT("Insights.TimingProfiler.TimingView.CpuThreadGroups"), TEXT("HiddenCpuThreadGroups"), HiddenCpuThreadGroups, SettingsIni);
+
+	for (const FString& CpuGroup : VisibleCpuThreadGroups)
+	{
+		AddTimingViewCpuThreadVisibility(*CpuGroup, true);
+	}
+
+	for (const FString& CpuGroup : HiddenCpuThreadGroups)
+	{
+		AddTimingViewCpuThreadVisibility(*CpuGroup, false);
+	}
+
 	//////////////////////////////////////////////////
 	// [Insights.TimingProfiler.FramesView]
 
@@ -228,6 +243,16 @@ void FInsightsSettings::SaveToConfig()
 	GConfig->SetBool(TEXT("Insights.TimingProfiler.TimingView"), TEXT("ShowGpuFencesExtendedLines"), bTimingViewShowGpuFencesExtendedLines, SettingsIni);
 	GConfig->SetBool(TEXT("Insights.TimingProfiler.TimingView"), TEXT("ShowGpuFencesRelations"), bTimingViewShowGpuFencesRelations, SettingsIni);
 
+	TArray<FString> VisibleCpuThreadGroups;
+	TArray<FString> HiddenCpuThreadGroups;
+	for (auto& KV : TimingViewShowCpuThreadGroups)
+	{
+		KV.Value ? VisibleCpuThreadGroups.Add(KV.Key) : HiddenCpuThreadGroups.Add(KV.Key);
+	}
+	
+	GConfig->SetArray(TEXT("Insights.TimingProfiler.TimingView.CpuThreadGroups"), TEXT("VisibleCpuThreadGroups"), VisibleCpuThreadGroups, SettingsIni);
+	GConfig->SetArray(TEXT("Insights.TimingProfiler.TimingView.CpuThreadGroups"), TEXT("HiddenCpuThreadGroups"), HiddenCpuThreadGroups, SettingsIni);
+
 	//////////////////////////////////////////////////
 	// [Insights.TimingProfiler.FramesView]
 
diff --git a/Engine/Source/Developer/TraceInsights/Private/Insights/InsightsSettings.h b/Engine/Source/Developer/TraceInsights/Private/Insights/InsightsSettings.h
--- a/Engine/Source/Developer/TraceInsights/Private/Insights/InsightsSettings.h
+++ b/Engine/Source/Developer/TraceInsights/Private/Insights/InsightsSettings.h
@@ -3,6 +3,7 @@
 #pragma once
 
 #include "CoreTypes.h"
+#include "Containers/Map.h"
 #include "Containers/UnrealString.h"
 
 /** Contains all settings for the Unreal Insights, accessible through the main manager. */
@@ -103,6 +104,11 @@ public:
 	void SetTimingViewShowGpuFencesRelations(bool InValue) { bTimingViewShowGpuFencesRelations = InValue; }
 	void SetAndSaveTimingViewShowGpuFencesRelations(bool InValue) { SET_AND_SAVE(bTimingViewShowGpuFencesRelations, InValue); }
 
+	bool GetTimingViewCpuThreadVisibility(const TCHAR* Key) const { return HasTimingViewCpuThreadVisibility(Key) ? *TimingViewShowCpuThreadGroups.Find(Key) : false; }
+	bool HasTimingViewCpuThreadVisibility(const TCHAR* Key) const { return TimingViewShowCpuThreadGroups.Contains(Key); }
+	void AddTimingViewCpuThreadVisibility(const TCHAR* Key, bool InValue) { TimingViewShowCpuThreadGroups.Add(Key, InValue); }
+	void AddAndSaveTimingViewCpuThreadVisibility(const TCHAR* Key, bool InValue) { TimingViewShowCpuThreadGroups.Add(Key, InValue); SaveToConfig(); }
+
 	//////////////////////////////////////////////////
 	// [Insights.TimingProfiler.FramesView]
 
@@ -313,6 +319,9 @@ private:
 	/** If enabled, relations between Signal and Wait fences will be displayed when selecting a Timing Event in a Gpu Queue Track. */
 	bool bTimingViewShowGpuFencesRelations = true;
 
+	/** Shows/hides the Cpu Thread Groups tracks. */
+	TMap<FString, bool> TimingViewShowCpuThreadGroups;
+
 	//////////////////////////////////////////////////
 	// [Insights.TimingProfiler.FramesView]
 
diff --git a/Engine/Source/Developer/TraceInsights/Private/Insights/TimingProfiler/ViewModels/ThreadTimingSharedState.cpp b/Engine/Source/Developer/TraceInsights/Private/Insights/TimingProfiler/ViewModels/ThreadTimingSharedState.cpp
--- a/Engine/Source/Developer/TraceInsights/Private/Insights/TimingProfiler/ViewModels/ThreadTimingSharedState.cpp
+++ b/Engine/Source/Developer/TraceInsights/Private/Insights/TimingProfiler/ViewModels/ThreadTimingSharedState.cpp
@@ -43,6 +43,9 @@ public:
 	virtual void RegisterCommands() override;
 
 public:
+	/** Restores all settings to default state. */
+	TSharedPtr<FUICommandInfo> ResetToDefaultValues;
+	
 	/** Toggles visibility for GPU thread track(s). */
 	TSharedPtr<FUICommandInfo> ShowHideAllGpuTracks;
 
@@ -88,6 +91,12 @@ FThreadTimingViewCommands::FThreadTimingViewCommands()
 UE_DISABLE_OPTIMIZATION_SHIP
 void FThreadTimingViewCommands::RegisterCommands()
 {
+	UI_COMMAND(ResetToDefaultValues,
+		"Reset To Defaults",
+		"Resets all tracks back to their default visibility.",
+		EUserInterfaceActionType::Button,
+		FInputChord());
+	
 	UI_COMMAND(ShowHideAllGpuTracks,
 		"GPU Track(s)",
 		"Shows/hides the GPU track(s).",
@@ -292,6 +301,38 @@ void FThreadTimingSharedState::GetVisibleTimelineIndexes(TSet<uint32>& OutSet) c
 
 ////////////////////////////////////////////////////////////////////////////////////////////////////
 
+void FThreadTimingSharedState::ResetToDefaultValues()
+{
+	if (TimingView &&
+		TimingView->GetName() == FInsightsManagerTabs::TimingProfilerTabId)
+	{
+		bShowHideAllGpuTracks = true;
+		bShowHideAllVerseTracks = true;
+		bShowHideAllCpuTracks = true;
+		ShowAllGpuTracks();
+		ShowAllVerseTracks();
+		ShowAllCpuTracks();
+	}
+	else
+	{
+		bShowHideAllGpuTracks = false;
+		bShowHideAllVerseTracks = false;
+		bShowHideAllCpuTracks = false;
+		HideAllGpuTracks();
+		HideAllVerseTracks();
+		HideAllCpuTracks();
+	}
+
+	Settings->SetTimingViewShowGpuWorkOverlays(true);
+	Settings->SetTimingViewShowGpuWorkExtendedLines(true);
+	SetGpuWorkTracksVisibility(true);
+	Settings->SetTimingViewShowGpuFencesRelations(true);
+	Settings->SetTimingViewShowGpuFencesExtendedLines(true);
+	SetGpuFencesTracksVisibility(true);
+}
+
+////////////////////////////////////////////////////////////////////////////////////////////////////
+
 void FThreadTimingSharedState::OnBeginSession(Timing::ITimingViewSession& InSession)
 {
 	if (&InSession != TimingView)
@@ -490,14 +531,20 @@ void FThreadTimingSharedState::Tick(Timing::ITimingViewSession& InSession, const
 				}
 				if (!ThreadGroups.Contains(GroupName))
 				{
+					if (!Settings->HasTimingViewCpuThreadVisibility(GroupName))
+					{
+						Settings->AddTimingViewCpuThreadVisibility(GroupName, bShowHideAllCpuTracks);
+					}
+					bIsGroupVisible = Settings->GetTimingViewCpuThreadVisibility(GroupName);
+
 					// Note: The GroupName pointer should be valid for the duration of the session.
-					ThreadGroups.Add(GroupName, { GroupName, bIsGroupVisible, 0, CpuTrackOrder });
+					ThreadGroups.Add(GroupName, { GroupName, 0, CpuTrackOrder });
 				}
 				else
 				{
 					FThreadGroup& ThreadGroup = ThreadGroups[GroupName];
-					bIsGroupVisible = ThreadGroup.bIsVisible;
 					ThreadGroup.Order = CpuTrackOrder;
+					bIsGroupVisible = Settings->GetTimingViewCpuThreadVisibility(GroupName);
 				}
 
 				// Check if there is an available Asset Loading track for this thread.
@@ -530,7 +577,6 @@ void FThreadTimingSharedState::Tick(Timing::ITimingViewSession& InSession, const
 							TimingView->GetName() == FInsightsManagerTabs::LoadingProfilerTabId)
 						{
 							Track->SetVisibilityFlag(true);
-							ThreadGroup.bIsVisible = true;
 						}
 						else
 						{
@@ -611,6 +657,12 @@ void FThreadTimingSharedState::ExtendCpuTracksFilterMenu(Timing::ITimingViewSess
 	InOutMenuBuilder.BeginSection("CpuThreadGroups", LOCTEXT("ContextMenu_Section_CpuThreadGroups", "CPU Thread Groups"));
 	CreateThreadGroupsMenu(InOutMenuBuilder);
 	InOutMenuBuilder.EndSection();
+
+	InOutMenuBuilder.BeginSection("Misc", LOCTEXT("ContextMenu_Section_Misc", "Misc Settings"));
+	{
+		InOutMenuBuilder.AddMenuEntry(FThreadTimingViewCommands::Get().ResetToDefaultValues);
+	}
+	InOutMenuBuilder.EndSection();
 }
 
 ////////////////////////////////////////////////////////////////////////////////////////////////////
@@ -627,6 +679,10 @@ void FThreadTimingSharedState::BindCommands()
 	TSharedPtr<FUICommandList> CommandList = TimingView->GetCommandList();
 	ensure(CommandList.IsValid());
 
+	CommandList->MapAction(
+		FThreadTimingViewCommands::Get().ResetToDefaultValues,
+		FExecuteAction::CreateSP(this, &FThreadTimingSharedState::ResetToDefaultValues));
+	
 	CommandList->MapAction(
 		FThreadTimingViewCommands::Get().ShowHideAllGpuTracks,
 		FExecuteAction::CreateSP(this, &FThreadTimingSharedState::ShowHideAllGpuTracks),
@@ -748,7 +804,7 @@ void FThreadTimingSharedState::SetAllCpuTracksToggle(bool bOnOff)
 
 	for (auto& KV : ThreadGroups)
 	{
-		KV.Value.bIsVisible = bShowHideAllCpuTracks;
+		Settings->AddTimingViewCpuThreadVisibility(KV.Value.Name, bShowHideAllCpuTracks);
 	}
 
 	if (TimingView)
@@ -789,8 +845,7 @@ bool FThreadTimingSharedState::ToggleTrackVisibilityByGroup_IsChecked(const TCHA
 {
 	if (ThreadGroups.Contains(InGroupName))
 	{
-		const FThreadGroup& ThreadGroup = ThreadGroups[InGroupName];
-		return ThreadGroup.bIsVisible;
+		return Settings->GetTimingViewCpuThreadVisibility(InGroupName);
 	}
 	return false;
 }
@@ -801,15 +856,15 @@ void FThreadTimingSharedState::ToggleTrackVisibilityByGroup_Execute(const TCHAR*
 {
 	if (ThreadGroups.Contains(InGroupName))
 	{
-		FThreadGroup& ThreadGroup = ThreadGroups[InGroupName];
-		ThreadGroup.bIsVisible = !ThreadGroup.bIsVisible;
+		bool bInvVisibility = !Settings->GetTimingViewCpuThreadVisibility(InGroupName);
+		Settings->AddTimingViewCpuThreadVisibility(InGroupName, bInvVisibility);
 
 		for (const auto& KV : CpuTracks)
 		{
 			FCpuTimingTrack& Track = *KV.Value;
 			if (Track.GetGroupName() == InGroupName)
 			{
-				Track.SetVisibilityFlag(ThreadGroup.bIsVisible);
+				Track.SetVisibilityFlag(bInvVisibility);
 			}
 		}
 
diff --git a/Engine/Source/Developer/TraceInsights/Private/Insights/TimingProfiler/ViewModels/ThreadTimingSharedState.h b/Engine/Source/Developer/TraceInsights/Private/Insights/TimingProfiler/ViewModels/ThreadTimingSharedState.h
--- a/Engine/Source/Developer/TraceInsights/Private/Insights/TimingProfiler/ViewModels/ThreadTimingSharedState.h
+++ b/Engine/Source/Developer/TraceInsights/Private/Insights/TimingProfiler/ViewModels/ThreadTimingSharedState.h
@@ -36,7 +36,6 @@ private:
 	struct FThreadGroup
 	{
 		const TCHAR* Name; /**< The thread group name; pointer to string owned by ThreadProvider. */
-		bool bIsVisible;  /**< Toggle to show/hide all thread timelines associated with this group at once. Used also as default for new thread timelines. */
 		uint32 NumTimelines; /**< Number of thread timelines associated with this group. */
 		int32 Order; //**< Order index used for sorting. Inherited from last thread timeline associated with this group. **/
 
@@ -65,6 +64,10 @@ private:
 
 		virtual bool GetTimingViewShowGpuFencesRelations() const = 0;
 		virtual void SetTimingViewShowGpuFencesRelations(bool InValue) = 0;
+
+		virtual bool GetTimingViewCpuThreadVisibility(const TCHAR* GroupName) const = 0;
+		virtual bool HasTimingViewCpuThreadVisibility(const TCHAR* GroupName) const = 0;
+		virtual void AddTimingViewCpuThreadVisibility(const TCHAR* GroupName, bool InValue) = 0;
 	};
 
 	class FThreadSharedStatePersistentSettings : public IThreadSharedStateSetting
@@ -88,6 +91,10 @@ private:
 		virtual bool GetTimingViewShowGpuFencesRelations() const override { return GetInsightsSettings().GetTimingViewShowGpuFencesRelations(); }
 		virtual void SetTimingViewShowGpuFencesRelations(bool InValue) override { GetInsightsSettings().SetAndSaveTimingViewShowGpuFencesRelations(InValue); };
 
+		virtual bool GetTimingViewCpuThreadVisibility(const TCHAR* GroupName) const override { return GetInsightsSettings().GetTimingViewCpuThreadVisibility(GroupName); }
+		virtual bool HasTimingViewCpuThreadVisibility(const TCHAR* GroupName) const override { return GetInsightsSettings().HasTimingViewCpuThreadVisibility(GroupName); }
+		virtual void AddTimingViewCpuThreadVisibility(const TCHAR* GroupName, bool InValue) override { return GetInsightsSettings().AddAndSaveTimingViewCpuThreadVisibility(GroupName, InValue); }
+		
 		FInsightsSettings& GetInsightsSettings() { return UE::Insights::FInsightsManager::Get()->GetSettings(); }
 		const FInsightsSettings& GetInsightsSettings() const { return UE::Insights::FInsightsManager::Get()->GetSettings(); }
 	};
@@ -111,7 +118,11 @@ private:
 		virtual void SetTimingViewShowGpuFencesExtendedLines(bool InValue) override { bTimingViewShowGpuFencesExtendedLines = InValue; }
 
 		virtual bool GetTimingViewShowGpuFencesRelations() const override { return bTimingViewShowGpuFencesRelations; }
-		virtual void SetTimingViewShowGpuFencesRelations(bool InValue) override { bTimingViewShowGpuFencesRelations = InValue; };
+		virtual void SetTimingViewShowGpuFencesRelations(bool InValue) override { bTimingViewShowGpuFencesRelations = InValue; }
+
+		virtual bool GetTimingViewCpuThreadVisibility(const TCHAR* GroupName) const override { return HasTimingViewCpuThreadVisibility(GroupName) ? *TimingViewShowCpuThreadGroups.Find(GroupName) : false; }
+		virtual bool HasTimingViewCpuThreadVisibility(const TCHAR* GroupName) const override { return TimingViewShowCpuThreadGroups.Contains(GroupName); }
+		virtual void AddTimingViewCpuThreadVisibility(const TCHAR* GroupName, bool InValue) override { TimingViewShowCpuThreadGroups.Add(GroupName, InValue); }
 
 	private:
 		bool bTimingViewShowGpuWorkTracks = true;
@@ -121,6 +132,8 @@ private:
 		bool bTimingViewShowGpuFencesTracks = true;
 		bool bTimingViewShowGpuFencesExtendedLines = true;
 		bool bTimingViewShowGpuFencesRelations = true;
+
+		TMap<FString, bool> TimingViewShowCpuThreadGroups;
 	};
 
 public:
@@ -158,6 +171,8 @@ public:
 
 	void BindCommands();
 
+	void ResetToDefaultValues();
+
 	bool IsAllGpuTracksToggleOn() const { return bShowHideAllGpuTracks; }
 	void SetAllGpuTracksToggle(bool bOnOff);
 	void ShowAllGpuTracks() { SetAllGpuTracksToggle(true); }
