[Insights] Added missing ResetSession() in FInsightsManager::ListenForDirectTrace.

diff --git a/Engine/Source/Developer/TraceInsights/Private/Insights/InsightsManager.cpp b/Engine/Source/Developer/TraceInsights/Private/Insights/InsightsManager.cpp
--- a/Engine/Source/Developer/TraceInsights/Private/Insights/InsightsManager.cpp
+++ b/Engine/Source/Developer/TraceInsights/Private/Insights/InsightsManager.cpp
@@ -843,6 +843,8 @@ uint16 FInsightsManager::ListenForDirectTrace()
 	TUniquePtr<UE::Trace::FDirectSocketStream> TraceData = MakeUnique<UE::Trace::FDirectSocketStream>();
 	const uint16 Port = TraceData->StartListening();
 	
+	ResetSession();
+
 	Session = AnalysisService->StartAnalysis(TraceId, *TraceName, MoveTemp(TraceData));
 
 	if (Session)
