Moved aggregation thread number setting to the timers context menu

diff --git a/Engine/Programs/UnrealInsights/Config/DefaultUnrealInsightsSettings.ini b/Engine/Programs/UnrealInsights/Config/DefaultUnrealInsightsSettings.ini
--- a/Engine/Programs/UnrealInsights/Config/DefaultUnrealInsightsSettings.ini
+++ b/Engine/Programs/UnrealInsights/Config/DefaultUnrealInsightsSettings.ini
@@ -34,6 +34,10 @@ AutoScrollViewportOffsetPercent=0.1
 ; This option can be changed from the "Auto-Scroll Options" menu.
 AutoScrollMinDelay=0.3
 
+; Ratio of available threads to use for per-frame aggregations.
+; This option can be changed from the Settings sub-menu in the Context Menu of any Timer in the Timers Panel.
+ThreadsToUseRatio=0.75
+
 [Insights.TimingProfiler.FramesView]
 ; These options can be changed from the context menu of the Frames view.
 
diff --git a/Engine/Source/Developer/TraceInsights/Private/Insights/TimingProfiler/Widgets/STimersView.cpp b/Engine/Source/Developer/TraceInsights/Private/Insights/TimingProfiler/Widgets/STimersView.cpp
--- a/Engine/Source/Developer/TraceInsights/Private/Insights/TimingProfiler/Widgets/STimersView.cpp
+++ b/Engine/Source/Developer/TraceInsights/Private/Insights/TimingProfiler/Widgets/STimersView.cpp
@@ -313,23 +313,6 @@ void STimersView::Construct(const FArguments& InArgs)
 					GetToggleButtonForTimerType(ETimerNodeType::CpuScope)
 				]
 
-				// Adjust the ratio of threads being used in aggregation calculations
-				+ SHorizontalBox::Slot()
-				.Padding(4.0f, 0.0f, 4.0f, 0.0f)
-				.AutoWidth()
-				.VAlign(VAlign_Center)
-				[
-					SNew(SSpinBox<double>)
-					.ToolTipText(LOCTEXT("ThreadsAggregation_Tooltip", "Ratio of Threads being used for Aggregation.\n" "Using higher values may cause systems to freeze."))
-					.MinValue(this, &STimersView::GetMinimumThreadRatio)
-					.MaxValue(1.0)
-					.MaxFractionalDigits(2)
-					.MinFractionalDigits(2)
-					.Value(this, &STimersView::GetThreadRatio)
-					.OnValueChanged(this, &STimersView::OnThreadRatioChanged)
-					.OnValueCommitted(this, &STimersView::OnThreadRatioCommited)
-				]
-
 				// Filter out timers with zero instance count
 				+ SHorizontalBox::Slot()
 				.Padding(2.0f)
@@ -735,6 +718,15 @@ TSharedPtr<SWidget> STimersView::TreeView_GetMenuContent()
 
 	MenuBuilder.BeginSection("Misc", LOCTEXT("ContextMenu_Section_Misc", "Miscellaneous"));
 	{
+		MenuBuilder.AddSubMenu
+		(
+			LOCTEXT("ContextMenu_Settings_SubMenu", "Settings"),
+			LOCTEXT("ContextMenu_Settings_SubMenu_Desc", "Options to adjust timer-related settings."),
+			FNewMenuDelegate::CreateSP(this, &STimersView::TreeView_BuildSettingsMenu),
+			false,
+			FSlateIcon(FAppStyle::Get().GetStyleSetName(), "Icons.Toolbar.Settings")
+		);
+
 		MenuBuilder.AddMenuEntry
 		(
 			FTimersViewCommands::Get().Command_CopyToClipboard,
@@ -1176,6 +1168,33 @@ void STimersView::TreeView_FindMenu(FMenuBuilder& MenuBuilder)
 
 ////////////////////////////////////////////////////////////////////////////////////////////////////
 
+void STimersView::TreeView_BuildSettingsMenu(FMenuBuilder& MenuBuilder)
+{
+	MenuBuilder.SetSearchable(false);
+
+	MenuBuilder.BeginSection("Settings");
+	{
+		MenuBuilder.AddEditableText
+		(
+			LOCTEXT("ThreadsAggregation", "Number of Threads"),
+			LOCTEXT("ThreadsAggregation_Tooltip",
+				"Number of threads being used for Aggregation.\n"
+				"Using higher values may cause systems to freeze.\n"
+				"This value is persistent to the UnrealInsightsSettings.ini file."),
+			FSlateIcon(),
+			TAttribute<FText>::Create(TAttribute<FText>::FGetter::CreateRaw(this, &STimersView::GetThreadRatio)),
+			FOnTextCommitted::CreateLambda([this](const FText& NewValue, ETextCommit::Type)
+			{
+				double Ratio = FCString::Atof(*NewValue.ToString());
+				OnThreadRatioCommited(Ratio);
+			})
+		);
+	}
+	MenuBuilder.EndSection();
+}
+
+////////////////////////////////////////////////////////////////////////////////////////////////////
+
 void STimersView::InitializeAndShowHeaderColumns()
 {
 	// Create columns.
@@ -1751,14 +1770,14 @@ void STimersView::SearchBox_OnTextChanged(const FText& InFilterText)
 // Aggregation Performance
 ////////////////////////////////////////////////////////////////////////////////////////////////////
 
-TOptional<double> STimersView::GetMinimumThreadRatio() const
+double STimersView::GetMinimumThreadRatio() const
 {
 	return 1.0f / (double)GThreadPool->GetNumThreads();
 }
 
 ////////////////////////////////////////////////////////////////////////////////////////////////////
 
-double STimersView::GetThreadRatio() const
+FText STimersView::GetThreadRatio() const
 {
 	if (Session.IsValid())
 	{
@@ -1767,36 +1786,34 @@ double STimersView::GetThreadRatio() const
 		{
 			TraceServices::FAnalysisSessionReadScope SessionReadScope(*Session.Get());
 			
-			return TimingProfilerProvider->GetRatioOfThreadsToUse();
+			return FText::AsNumber((int)(TimingProfilerProvider->GetRatioOfThreadsToUse() * GThreadPool->GetNumThreads()));
 		}
 	}
-	return 0.75;
+	return FText::AsNumber((int)(0.75 * GThreadPool->GetNumThreads()));
 }
 
 ////////////////////////////////////////////////////////////////////////////////////////////////////
 
-void STimersView::OnThreadRatioChanged(double InNewValue)
-{
-	OnThreadRatioCommited(InNewValue, ETextCommit::Default);
-}
-
-////////////////////////////////////////////////////////////////////////////////////////////////////
-
-void STimersView::OnThreadRatioCommited(double InNewValue, ETextCommit::Type CommitType)
+void STimersView::OnThreadRatioCommited(double InNewValue)
 {
+	double ClampedNewValue = FMath::Clamp(
+		InNewValue / GThreadPool->GetNumThreads(),
+		GetMinimumThreadRatio(),
+		1.0);
+	
 	if (Session.IsValid())
 	{
 		const TraceServices::ITimingProfilerProvider* TimingProfilerProvider = TraceServices::ReadTimingProfilerProvider(*Session.Get());
 		if (TimingProfilerProvider)
 		{
 			TraceServices::FAnalysisSessionEditScope SessionEditScope(*const_cast<TraceServices::IAnalysisSession*>(Session.Get()));
-			
-			TimingProfilerProvider->SetRatioOfThreadsToUse(InNewValue);
+
+			TimingProfilerProvider->SetRatioOfThreadsToUse(ClampedNewValue);
 		}
 	}
 	
 	FInsightsSettings& Settings = FInsightsManager::Get()->GetSettings();
-	Settings.SetAndSaveThreadRatio(InNewValue);
+	Settings.SetAndSaveThreadRatio(ClampedNewValue);
 }
 
 ////////////////////////////////////////////////////////////////////////////////////////////////////
diff --git a/Engine/Source/Developer/TraceInsights/Private/Insights/TimingProfiler/Widgets/STimersView.h b/Engine/Source/Developer/TraceInsights/Private/Insights/TimingProfiler/Widgets/STimersView.h
--- a/Engine/Source/Developer/TraceInsights/Private/Insights/TimingProfiler/Widgets/STimersView.h
+++ b/Engine/Source/Developer/TraceInsights/Private/Insights/TimingProfiler/Widgets/STimersView.h
@@ -199,7 +199,12 @@ private:
 
 	bool ContextMenu_FindInstanceInSelection_CanExecute() const;
 	void ContextMenu_FindInstanceInSelection_Execute(bool bFindMax) const;
+	
+	//////////////////////////////////////////////////
+	// Settings
 
+	void TreeView_BuildSettingsMenu(FMenuBuilder& MenuBuilder);
+	
 	////////////////////////////////////////////////////////////////////////////////////////////////////
 	// Tree View - Columns' Header
 
@@ -265,10 +270,9 @@ private:
 	////////////////////////////////////////////////////////////////////////////////////////////////////
 	// Aggregation Performance
 	
-	TOptional<double> GetMinimumThreadRatio() const;
-	double GetThreadRatio() const;
-	void OnThreadRatioChanged(double InNewValue);
-	void OnThreadRatioCommited(double InNewValue, ETextCommit::Type CommitType);
+	double GetMinimumThreadRatio() const;
+	FText GetThreadRatio() const;
+	void OnThreadRatioCommited(double InNewValue);
 
 	////////////////////////////////////////////////////////////////////////////////////////////////////
 	// Grouping
