Persistent CPU/GPU visibility settings for Timing and Memory Profilers
* GPU/CPU/Verse Tracks have unified saving code
* CPU Tracks will automatically cleanup any redundant visibility flags and will only save groups which have opposite states to the 'ShowAllCpuTracks' state

diff --git a/Engine/Source/Developer/TraceInsights/Private/Insights/InsightsSettings.cpp b/Engine/Source/Developer/TraceInsights/Private/Insights/InsightsSettings.cpp
--- a/Engine/Source/Developer/TraceInsights/Private/Insights/InsightsSettings.cpp
+++ b/Engine/Source/Developer/TraceInsights/Private/Insights/InsightsSettings.cpp
@@ -44,6 +44,7 @@ FInsightsSettings::FInsightsSettings(bool bInIsDefault)
 
 FInsightsSettings::~FInsightsSettings()
 {
+	ProfilerTimingViewTrackSettings.Reset();
 }
 
 ////////////////////////////////////////////////////////////////////////////////////////////////////
@@ -69,6 +70,45 @@ void FInsightsSettings::LoadFromConfig()
 		return;
 	}
 
+	//////////////////////////////////////////////////
+	// [Insights.(AllProfilers).TimingView]
+
+	TArray<FString> PersistentProfilers;
+	GConfig->GetArray(TEXT("Insights.PersistentTimingView"), TEXT("Profiler"), PersistentProfilers, SettingsIni);
+
+	for (const FString& Profiler : PersistentProfilers)
+	{
+		CreateIfNotExistProfilerTimingViewTrackSettings(*Profiler);
+		
+		TArray<FString> EnabledToggleOptions;
+		TArray<FString> DisabledToggleOptions;
+		GConfig->GetArray(*FString::Printf(TEXT("Insights.%s.TimingView"), *Profiler), TEXT("EnabledToggleOptions"), EnabledToggleOptions, SettingsIni);
+		GConfig->GetArray(*FString::Printf(TEXT("Insights.%s.TimingView"), *Profiler), TEXT("DisabledToggleOptions"), DisabledToggleOptions, SettingsIni);
+
+		for (const FString& Option : EnabledToggleOptions)
+		{
+			AddTimingViewToggleOption(*Option, true, *Profiler);
+		}
+		for (const FString& Option : DisabledToggleOptions)
+		{
+			AddTimingViewToggleOption(*Option, false, *Profiler);
+		}
+		
+		TArray<FString> VisibleCpuThreadGroups;
+		TArray<FString> HiddenCpuThreadGroups;
+		GConfig->GetArray(*FString::Printf(TEXT("Insights.%s.TimingView.CpuThreadGroups"), *Profiler), TEXT("VisibleCpuThreadGroups"), VisibleCpuThreadGroups, SettingsIni);
+		GConfig->GetArray(*FString::Printf(TEXT("Insights.%s.TimingView.CpuThreadGroups"), *Profiler), TEXT("HiddenCpuThreadGroups"), HiddenCpuThreadGroups, SettingsIni);
+
+		for (const FString& CpuGroup : VisibleCpuThreadGroups)
+		{
+			AddTimingViewCpuThreadGroupVisibility(*CpuGroup, true, *Profiler);
+		}
+		for (const FString& CpuGroup : HiddenCpuThreadGroups)
+		{
+			AddTimingViewCpuThreadGroupVisibility(*CpuGroup, false, *Profiler);
+		}
+	}
+
 	//////////////////////////////////////////////////
 	// [Insights.TimingProfiler]
 
@@ -101,32 +141,6 @@ void FInsightsSettings::LoadFromConfig()
 	GConfig->GetDouble(TEXT("Insights.TimingProfiler"), TEXT("AutoScrollViewportOffsetPercent"), AutoScrollViewportOffsetPercent, SettingsIni);
 	GConfig->GetDouble(TEXT("Insights.TimingProfiler"), TEXT("AutoScrollMinDelay"), AutoScrollMinDelay, SettingsIni);
 
-	//////////////////////////////////////////////////
-	// [Insights.TimingProfiler.TimingView]
-
-	GConfig->GetBool(TEXT("Insights.TimingProfiler.TimingView"), TEXT("ShowGpuWorkTracks"), bTimingViewShowGpuWorkTracks, SettingsIni);
-	GConfig->GetBool(TEXT("Insights.TimingProfiler.TimingView"), TEXT("ShowGpuWorkOverlays"), bTimingViewShowGpuWorkOverlays, SettingsIni);
-	GConfig->GetBool(TEXT("Insights.TimingProfiler.TimingView"), TEXT("ShowGpuWorkExtendedLines"), bTimingViewShowGpuWorkExtendedLines, SettingsIni);
-
-	GConfig->GetBool(TEXT("Insights.TimingProfiler.TimingView"), TEXT("ShowGpuFencesTracks"), bTimingViewShowGpuFencesTracks, SettingsIni);
-	GConfig->GetBool(TEXT("Insights.TimingProfiler.TimingView"), TEXT("ShowGpuFencesExtendedLines"), bTimingViewShowGpuFencesExtendedLines, SettingsIni);
-	GConfig->GetBool(TEXT("Insights.TimingProfiler.TimingView"), TEXT("ShowGpuFencesRelations"), bTimingViewShowGpuFencesRelations, SettingsIni);
-
-	TArray<FString> VisibleCpuThreadGroups;
-	TArray<FString> HiddenCpuThreadGroups;
-	GConfig->GetArray(TEXT("Insights.TimingProfiler.TimingView.CpuThreadGroups"), TEXT("VisibleCpuThreadGroups"), VisibleCpuThreadGroups, SettingsIni);
-	GConfig->GetArray(TEXT("Insights.TimingProfiler.TimingView.CpuThreadGroups"), TEXT("HiddenCpuThreadGroups"), HiddenCpuThreadGroups, SettingsIni);
-
-	for (const FString& CpuGroup : VisibleCpuThreadGroups)
-	{
-		AddTimingViewCpuThreadVisibility(*CpuGroup, true);
-	}
-
-	for (const FString& CpuGroup : HiddenCpuThreadGroups)
-	{
-		AddTimingViewCpuThreadVisibility(*CpuGroup, false);
-	}
-
 	//////////////////////////////////////////////////
 	// [Insights.TimingProfiler.FramesView]
 
@@ -217,6 +231,39 @@ void FInsightsSettings::LoadFromConfig()
 
 void FInsightsSettings::SaveToConfig()
 {
+	//////////////////////////////////////////////////
+	// [Insights.(AllProfilers).TimingView]
+
+	TArray<FString> ProfilersWithPersistentTimingViews;
+	for (auto& ProfilerSettings : ProfilerTimingViewTrackSettings)
+	{
+		ProfilersWithPersistentTimingViews.Add(ProfilerSettings.Key.ToString());
+	}
+	GConfig->SetArray(TEXT("Insights.PersistentTimingView"), TEXT("Profiler"), ProfilersWithPersistentTimingViews, SettingsIni);
+	
+	for (auto& ProfilerSettings : ProfilerTimingViewTrackSettings)
+	{
+		FString Profiler = ProfilerSettings.Key.ToString();
+
+		TArray<FString> EnabledToggleOptions;
+		TArray<FString> DisabledToggleOptions;
+		for (auto& KV : GetProfilerTimingViewTrackSettings(*Profiler).GetToggleOptions())
+		{
+			KV.Value ? EnabledToggleOptions.Add(*KV.Key) : DisabledToggleOptions.Add(*KV.Key);
+		}
+		GConfig->SetArray(*FString::Printf(TEXT("Insights.%s.TimingView"),*Profiler), TEXT("EnabledToggleOptions"), EnabledToggleOptions, SettingsIni);
+		GConfig->SetArray(*FString::Printf(TEXT("Insights.%s.TimingView"),*Profiler), TEXT("DisabledToggleOptions"), DisabledToggleOptions, SettingsIni);
+			
+		TArray<FString> VisibleCpuThreadGroups;
+		TArray<FString> HiddenCpuThreadGroups;
+		for (auto& KV : GetProfilerTimingViewTrackSettings(*Profiler).GetCpuThreadGroups())
+		{
+			KV.Value ? VisibleCpuThreadGroups.Add(*KV.Key) : HiddenCpuThreadGroups.Add(*KV.Key);
+		}
+		GConfig->SetArray(*FString::Printf(TEXT("Insights.%s.TimingView.CpuThreadGroups"),*Profiler), TEXT("VisibleCpuThreadGroups"), VisibleCpuThreadGroups, SettingsIni);
+		GConfig->SetArray(*FString::Printf(TEXT("Insights.%s.TimingView.CpuThreadGroups"),*Profiler), TEXT("HiddenCpuThreadGroups"), HiddenCpuThreadGroups, SettingsIni);
+	}
+
 	//////////////////////////////////////////////////
 	// [Insights.TimingProfiler]
 
@@ -232,27 +279,6 @@ void FInsightsSettings::SaveToConfig()
 	GConfig->SetDouble(TEXT("Insights.TimingProfiler"), TEXT("AutoScrollViewportOffsetPercent"), AutoScrollViewportOffsetPercent, SettingsIni);
 	GConfig->SetDouble(TEXT("Insights.TimingProfiler"), TEXT("AutoScrollMinDelay"), AutoScrollMinDelay, SettingsIni);
 
-	//////////////////////////////////////////////////
-	// [Insights.TimingProfiler.TimingView]
-
-	GConfig->SetBool(TEXT("Insights.TimingProfiler.TimingView"), TEXT("ShowGpuWorkTracks"), bTimingViewShowGpuWorkTracks, SettingsIni);
-	GConfig->SetBool(TEXT("Insights.TimingProfiler.TimingView"), TEXT("ShowGpuWorkOverlays"), bTimingViewShowGpuWorkOverlays, SettingsIni);
-	GConfig->SetBool(TEXT("Insights.TimingProfiler.TimingView"), TEXT("ShowGpuWorkExtendedLines"), bTimingViewShowGpuWorkExtendedLines, SettingsIni);
-
-	GConfig->SetBool(TEXT("Insights.TimingProfiler.TimingView"), TEXT("ShowGpuFencesTracks"), bTimingViewShowGpuFencesTracks, SettingsIni);
-	GConfig->SetBool(TEXT("Insights.TimingProfiler.TimingView"), TEXT("ShowGpuFencesExtendedLines"), bTimingViewShowGpuFencesExtendedLines, SettingsIni);
-	GConfig->SetBool(TEXT("Insights.TimingProfiler.TimingView"), TEXT("ShowGpuFencesRelations"), bTimingViewShowGpuFencesRelations, SettingsIni);
-
-	TArray<FString> VisibleCpuThreadGroups;
-	TArray<FString> HiddenCpuThreadGroups;
-	for (auto& KV : TimingViewShowCpuThreadGroups)
-	{
-		KV.Value ? VisibleCpuThreadGroups.Add(KV.Key) : HiddenCpuThreadGroups.Add(KV.Key);
-	}
-	
-	GConfig->SetArray(TEXT("Insights.TimingProfiler.TimingView.CpuThreadGroups"), TEXT("VisibleCpuThreadGroups"), VisibleCpuThreadGroups, SettingsIni);
-	GConfig->SetArray(TEXT("Insights.TimingProfiler.TimingView.CpuThreadGroups"), TEXT("HiddenCpuThreadGroups"), HiddenCpuThreadGroups, SettingsIni);
-
 	//////////////////////////////////////////////////
 	// [Insights.TimingProfiler.FramesView]
 
diff --git a/Engine/Source/Developer/TraceInsights/Private/Insights/InsightsSettings.h b/Engine/Source/Developer/TraceInsights/Private/Insights/InsightsSettings.h
--- a/Engine/Source/Developer/TraceInsights/Private/Insights/InsightsSettings.h
+++ b/Engine/Source/Developer/TraceInsights/Private/Insights/InsightsSettings.h
@@ -5,12 +5,125 @@
 #include "CoreTypes.h"
 #include "Containers/Map.h"
 #include "Containers/UnrealString.h"
+#include "UObject/NameTypes.h"
 
 /** Contains all settings for the Unreal Insights, accessible through the main manager. */
 class FInsightsSettings
 {
 	friend class SInsightsSettings;
 
+public:
+	struct FTimingViewTrackSettings
+	{
+		bool FindKey(const TCHAR* Key, const TMap<FString*, bool>& Map) const
+		{
+			FString KeyToCheck = FString(Key);
+			for (const auto& KV : Map)
+			{
+				if (KV.Key && KeyToCheck == *KV.Key)
+				{
+					return true;
+				}
+			}
+			return false;
+		}
+
+		bool GetValueByKey(const TCHAR* Key, const TMap<FString*, bool>& Map) const
+		{
+			FString KeyToCheck = FString(Key);
+			for (const auto& KV : Map)
+			{
+				if (KV.Key && KeyToCheck == *KV.Key)
+				{
+					return KV.Value;
+				}
+			}
+			return false;
+		}
+
+		void RemoveKey(const TCHAR* Key, TMap<FString*, bool>& Map)
+		{
+			FString KeyToCheck = FString(Key);
+			for (auto It = Map.CreateIterator(); It; ++It)
+			{
+				FString* Ptr = It.Key();
+				if (Ptr && KeyToCheck == *Ptr)
+				{
+					delete Ptr;
+					It.RemoveCurrent();
+					return;
+				}
+			}
+		}
+
+		void RemoveAllKeys(TMap<FString*, bool>& Map)
+		{
+			for (auto It = Map.CreateIterator(); It; ++It)
+			{
+				FString* Ptr = It.Key();
+				if (Ptr)
+				{
+					delete Ptr;
+					It.RemoveCurrent();
+				}
+			}
+			Map.Reset();
+		}
+
+		////////////////////////////////////////////////////////////////////////////////////////////////////
+
+		/** Toggle all or no Gpu-related Tracks in the Timing View. */
+		static constexpr const TCHAR* Option_ShowAllGpuTracks = TEXT("ShowAllGpuTracks");
+		/** Toggle all or no Cpu Tracks in the Timing View. */
+		static constexpr const TCHAR* Option_ShowAllCpuTracks = TEXT("ShowAllCpuTracks");
+		/** Toggle all or no Verse Tracks in the Timing View. */
+		static constexpr const TCHAR* Option_ShowAllVerseTracks = TEXT("ShowAllVerseTracks");
+		/** GPU work header tracks. */
+		static constexpr const TCHAR* Option_ShowGpuWorkTracks = TEXT("ShowGpuWorkTracks");
+		/** Extends the visualization of GPU work events over the GPU timing tracks. */
+		static constexpr const TCHAR* Option_ShowGpuWorkOverlays = TEXT("ShowGpuWorkOverlays");
+		/** Extended vertical lines at the edges of each GPU work event. */
+		static constexpr const TCHAR* Option_ShowGpuWorkExtendedLines = TEXT("ShowGpuWorkExtendedLines");
+		/** GPU fences child track. */
+		static constexpr const TCHAR* Option_ShowGpuFencesTracks = TEXT("ShowGpuFencesTracks");
+		/** Relations between Signal and Wait fences will be displayed when selecting a Timing Event in a Gpu Queue Track. */
+		static constexpr const TCHAR* Option_ShowGpuFencesRelations = TEXT("ShowGpuFencesRelations");
+		/** Extended vertical lines at the location of gpu fences. */
+		static constexpr const TCHAR* Option_ShowGpuFencesExtendedLines = TEXT("ShowGpuFencesExtendedLines");
+
+		////////////////////////////////////////////////////////////////////////////////////////////////////
+
+		TMap<FString*, bool>& GetToggleOptions() { return ToggleOptions; }
+		bool IsToggleOptionEnabled(const TCHAR* Key) const { return GetValueByKey(Key, ToggleOptions); }
+		bool HasToggleOption(const TCHAR* Key) const { return FindKey(Key, ToggleOptions); }
+		void AddToggleOption(const TCHAR* Key, bool InValue) { RemoveKey(Key, ToggleOptions); ToggleOptions.Add(new FString(Key), InValue); }
+		
+		TMap<FString*, bool>& GetCpuThreadGroups() { return ShowCpuThreadGroups; }
+		bool IsCpuThreadGroupVisible(const TCHAR* Key) const { return GetValueByKey(Key, ShowCpuThreadGroups);}
+		bool HasCpuThreadGroupVisibility(const TCHAR* Key) const { return FindKey(Key, ShowCpuThreadGroups); }
+		void AddCpuThreadGroupVisibility(const TCHAR* Key, bool InValue) { RemoveKey(Key, ShowCpuThreadGroups); ShowCpuThreadGroups.Add(new FString(Key), InValue); }
+		void RemoveCpuThreadGroup(const TCHAR* Key) { RemoveKey(Key, ShowCpuThreadGroups); }
+
+		void ResetAll() { RemoveAllKeys(ToggleOptions); RemoveAllKeys(ShowCpuThreadGroups); }
+
+		FTimingViewTrackSettings() = default;
+
+		~FTimingViewTrackSettings()
+		{
+			ResetAll();
+		}
+
+	private:
+		/** Holds the members which control parent true/false states for options. */
+		/** Key - The Option_ for true/false toggle commands. */
+		/** Value - The enabled state. */
+		TMap<FString*, bool> ToggleOptions;
+		/** Shows/hides each Cpu Thread Group track. */
+		/** Key - The Name of the Cpu thread group. */
+		/** Value - The visibility state. */
+		TMap<FString*, bool> ShowCpuThreadGroups;
+	};
+
 public:
 	FInsightsSettings(bool bInIsDefault = false);
 	~FInsightsSettings();
@@ -42,6 +155,26 @@ public:
 
 	#define SET_AND_SAVE(Option, Value) { if (Option != Value) { Option = Value; SaveToConfig(); } }
 
+	//////////////////////////////////////////////////
+	// [Insights.(AllProfilers).TimingView]
+
+	FTimingViewTrackSettings& GetProfilerTimingViewTrackSettings(const FName& ProfilerId) {if (ProfilerTimingViewTrackSettings.Contains(ProfilerId)) { return *ProfilerTimingViewTrackSettings.Find(ProfilerId); } return TimingViewTrackSettingsDefaultObject; }
+	const FTimingViewTrackSettings& GetProfilerTimingViewTrackSettings(const FName& ProfilerId) const { if (ProfilerTimingViewTrackSettings.Contains(ProfilerId)) { return *ProfilerTimingViewTrackSettings.Find(ProfilerId); } return TimingViewTrackSettingsDefaultObject; }
+
+	void CreateIfNotExistProfilerTimingViewTrackSettings(const FName& ProfilerId) { if (!ProfilerTimingViewTrackSettings.Find(ProfilerId)) { ProfilerTimingViewTrackSettings.Add(ProfilerId); } }
+	void ResetAllProfilerTimingViewTrackSettings(const FName& ProfilerId) { GetProfilerTimingViewTrackSettings(ProfilerId).ResetAll(); SaveToConfig(); }
+
+	bool IsTimingViewToggleOptionEnabled(const TCHAR* Key, const FName& ProfilerId) const { return GetProfilerTimingViewTrackSettings(ProfilerId).IsToggleOptionEnabled(Key); }
+	bool HasTimingViewToggleOption(const TCHAR* Key, const FName& ProfilerId) const { return GetProfilerTimingViewTrackSettings(ProfilerId).HasToggleOption(Key); }
+	void AddTimingViewToggleOption(const TCHAR* Key, bool InValue, const FName& ProfilerId) { GetProfilerTimingViewTrackSettings(ProfilerId).AddToggleOption(Key, InValue); }
+	void AddAndSaveTimingViewToggleOption(const TCHAR* Key, bool InValue, const FName& ProfilerId) { AddTimingViewToggleOption(Key, InValue, ProfilerId); SaveToConfig(); }
+	
+	bool IsTimingViewCpuThreadGroupVisible(const TCHAR* Key, const FName& ProfilerId) const { return GetProfilerTimingViewTrackSettings(ProfilerId).IsCpuThreadGroupVisible(Key); }
+	bool HasTimingViewCpuThreadGroupVisibility(const TCHAR* Key, const FName& ProfilerId) const { return GetProfilerTimingViewTrackSettings(ProfilerId).HasCpuThreadGroupVisibility(Key); }
+	void AddTimingViewCpuThreadGroupVisibility(const TCHAR* Key, bool InValue, const FName& ProfilerId) { GetProfilerTimingViewTrackSettings(ProfilerId).AddCpuThreadGroupVisibility(Key, InValue); }
+	void AddAndSaveTimingViewCpuThreadGroupVisibility(const TCHAR* Key, bool InValue, const FName& ProfilerId) { AddTimingViewCpuThreadGroupVisibility(Key, InValue, ProfilerId); SaveToConfig(); }
+	void RemoveAndSaveTimingViewCpuThreadGroup(const TCHAR* Key, const FName& ProfilerId) { GetProfilerTimingViewTrackSettings(ProfilerId).RemoveCpuThreadGroup(Key); SaveToConfig(); }
+	
 	//////////////////////////////////////////////////
 	// [Insights.TimingProfiler]
 
@@ -77,38 +210,6 @@ public:
 	void SetThreadRatio(double Ratio) { ThreadsToUseRatio = Ratio; }
 	void SetAndSaveThreadRatio(double Ratio) { SET_AND_SAVE(ThreadsToUseRatio, Ratio); }
 
-	//////////////////////////////////////////////////
-	// [Insights.TimingProfiler.TimingView]
-
-	bool GetTimingViewShowGpuWorkTracks() const { return bTimingViewShowGpuWorkTracks; }
-	void SetTimingViewShowGpuWorkTracks(bool InValue) { bTimingViewShowGpuWorkTracks = InValue; }
-	void SetAndSaveTimingViewShowGpuWorkTracks(bool InValue) { SET_AND_SAVE(bTimingViewShowGpuWorkTracks, InValue); }
-
-	bool GetTimingViewShowGpuWorkOverlays() const { return bTimingViewShowGpuWorkOverlays; }
-	void SetTimingViewShowGpuWorkOverlays(bool InValue) { bTimingViewShowGpuWorkOverlays = InValue; }
-	void SetAndSaveTimingViewShowGpuWorkOverlays(bool InValue) { SET_AND_SAVE(bTimingViewShowGpuWorkOverlays, InValue); }
-
-	bool GetTimingViewShowGpuWorkExtendedLines() const { return bTimingViewShowGpuWorkExtendedLines; }
-	void SetTimingViewShowGpuWorkExtendedLines(bool InValue) { bTimingViewShowGpuWorkExtendedLines = InValue; }
-	void SetAndSaveTimingViewShowGpuWorkExtendedLines(bool InValue) { SET_AND_SAVE(bTimingViewShowGpuWorkExtendedLines, InValue); }
-
-	bool GetTimingViewShowGpuFencesTracks() const { return bTimingViewShowGpuFencesTracks; }
-	void SetTimingViewShowGpuFencesTracks(bool InValue) { bTimingViewShowGpuFencesTracks = InValue; }
-	void SetAndSaveTimingViewShowGpuFencesTracks(bool InValue) { SET_AND_SAVE(bTimingViewShowGpuFencesTracks, InValue); }
-
-	bool GetTimingViewShowGpuFencesExtendedLines() const { return bTimingViewShowGpuFencesExtendedLines; }
-	void SetTimingViewShowGpuFencesExtendedLines(bool InValue) { bTimingViewShowGpuFencesExtendedLines = InValue; }
-	void SetAndSaveTimingViewShowGpuFencesExtendedLines(bool InValue) { SET_AND_SAVE(bTimingViewShowGpuFencesExtendedLines, InValue); }
-
-	bool GetTimingViewShowGpuFencesRelations() const { return bTimingViewShowGpuFencesRelations; }
-	void SetTimingViewShowGpuFencesRelations(bool InValue) { bTimingViewShowGpuFencesRelations = InValue; }
-	void SetAndSaveTimingViewShowGpuFencesRelations(bool InValue) { SET_AND_SAVE(bTimingViewShowGpuFencesRelations, InValue); }
-
-	bool GetTimingViewCpuThreadVisibility(const TCHAR* Key) const { return HasTimingViewCpuThreadVisibility(Key) ? *TimingViewShowCpuThreadGroups.Find(Key) : false; }
-	bool HasTimingViewCpuThreadVisibility(const TCHAR* Key) const { return TimingViewShowCpuThreadGroups.Contains(Key); }
-	void AddTimingViewCpuThreadVisibility(const TCHAR* Key, bool InValue) { TimingViewShowCpuThreadGroups.Add(Key, InValue); }
-	void AddAndSaveTimingViewCpuThreadVisibility(const TCHAR* Key, bool InValue) { TimingViewShowCpuThreadGroups.Add(Key, InValue); SaveToConfig(); }
-
 	//////////////////////////////////////////////////
 	// [Insights.TimingProfiler.FramesView]
 
@@ -267,6 +368,15 @@ private:
 	/** Settings filename ini. */
 	FString SettingsIni;
 
+	//////////////////////////////////////////////////
+	// [Insights.(AllProfilers).TimingView]
+
+	/** Per profiler timing view track settings. */
+	TMap<FName, FTimingViewTrackSettings> ProfilerTimingViewTrackSettings;
+
+	/** Per profiler default object. */
+	inline static FTimingViewTrackSettings TimingViewTrackSettingsDefaultObject = {};
+
 	//////////////////////////////////////////////////
 	// [Insights.TimingProfiler]
 
@@ -298,30 +408,6 @@ private:
 	/** The ratio of available threads to use when performing per-frame aggregations. */
 	double ThreadsToUseRatio = 0.75;
 
-	//////////////////////////////////////////////////
-	// [Insights.TimingProfiler.TimingView]
-
-	/** Toggles visibility for GPU work header tracks. */
-	bool bTimingViewShowGpuWorkTracks = true;
-
-	/** Extends the visualization of GPU work events over the GPU timing tracks. */
-	bool bTimingViewShowGpuWorkOverlays = true;
-
-	/** Shows/hides the extended vertical lines at the edges of each GPU work event. */
-	bool bTimingViewShowGpuWorkExtendedLines = true;
-
-	/** Shows/hides the gpu fences child track. */
-	bool bTimingViewShowGpuFencesTracks = true;
-
-	/** Shows/hides the extended vertical lines at the location of gpu fences. */
-	bool bTimingViewShowGpuFencesExtendedLines = true;
-
-	/** If enabled, relations between Signal and Wait fences will be displayed when selecting a Timing Event in a Gpu Queue Track. */
-	bool bTimingViewShowGpuFencesRelations = true;
-
-	/** Shows/hides the Cpu Thread Groups tracks. */
-	TMap<FString, bool> TimingViewShowCpuThreadGroups;
-
 	//////////////////////////////////////////////////
 	// [Insights.TimingProfiler.FramesView]
 
diff --git a/Engine/Source/Developer/TraceInsights/Private/Insights/TimingProfiler/ViewModels/ThreadTimingSharedState.cpp b/Engine/Source/Developer/TraceInsights/Private/Insights/TimingProfiler/ViewModels/ThreadTimingSharedState.cpp
--- a/Engine/Source/Developer/TraceInsights/Private/Insights/TimingProfiler/ViewModels/ThreadTimingSharedState.cpp
+++ b/Engine/Source/Developer/TraceInsights/Private/Insights/TimingProfiler/ViewModels/ThreadTimingSharedState.cpp
@@ -303,31 +303,29 @@ void FThreadTimingSharedState::GetVisibleTimelineIndexes(TSet<uint32>& OutSet) c
 
 void FThreadTimingSharedState::ResetToDefaultValues()
 {
+	RemoveGpuWorkChildTracks();
+	RemoveGpuFencesChildTracks();
+	Settings->Reset();
+	
 	if (TimingView &&
 		TimingView->GetName() == FInsightsManagerTabs::TimingProfilerTabId)
 	{
-		bShowHideAllGpuTracks = true;
-		bShowHideAllVerseTracks = true;
-		bShowHideAllCpuTracks = true;
 		ShowAllGpuTracks();
 		ShowAllVerseTracks();
 		ShowAllCpuTracks();
 	}
 	else
 	{
-		bShowHideAllGpuTracks = false;
-		bShowHideAllVerseTracks = false;
-		bShowHideAllCpuTracks = false;
 		HideAllGpuTracks();
 		HideAllVerseTracks();
 		HideAllCpuTracks();
 	}
 
-	Settings->SetTimingViewShowGpuWorkOverlays(true);
-	Settings->SetTimingViewShowGpuWorkExtendedLines(true);
+	Settings->AddTimingViewToggleOption(FInsightsSettings::FTimingViewTrackSettings::Option_ShowGpuWorkOverlays, true);
+	Settings->AddTimingViewToggleOption(FInsightsSettings::FTimingViewTrackSettings::Option_ShowGpuWorkExtendedLines, true);
 	SetGpuWorkTracksVisibility(true);
-	Settings->SetTimingViewShowGpuFencesRelations(true);
-	Settings->SetTimingViewShowGpuFencesExtendedLines(true);
+	Settings->AddTimingViewToggleOption(FInsightsSettings::FTimingViewTrackSettings::Option_ShowGpuFencesRelations, true);
+	Settings->AddTimingViewToggleOption(FInsightsSettings::FTimingViewTrackSettings::Option_ShowGpuFencesExtendedLines, true);
 	SetGpuFencesTracksVisibility(true);
 }
 
@@ -343,17 +341,40 @@ void FThreadTimingSharedState::OnBeginSession(Timing::ITimingViewSession& InSess
 	if (TimingView &&
 		TimingView->GetName() == FInsightsManagerTabs::TimingProfilerTabId)
 	{
-		bShowHideAllGpuTracks = true;
-		bShowHideAllVerseTracks = true;
-		bShowHideAllCpuTracks = true;
-
-		Settings = MakeShared<FThreadSharedStatePersistentSettings>();
+		Settings = MakeShared<FThreadSharedStatePersistentSettings>(FInsightsManagerTabs::TimingProfilerTabId);
+
+		if (!Settings->HasTimingViewToggleOption(FInsightsSettings::FTimingViewTrackSettings::Option_ShowAllCpuTracks)) { Settings->AddTimingViewToggleOption(FInsightsSettings::FTimingViewTrackSettings::Option_ShowAllCpuTracks, true); }
+		if (!Settings->HasTimingViewToggleOption(FInsightsSettings::FTimingViewTrackSettings::Option_ShowAllGpuTracks)) { Settings->AddTimingViewToggleOption(FInsightsSettings::FTimingViewTrackSettings::Option_ShowAllGpuTracks, true); }
+		if (!Settings->HasTimingViewToggleOption(FInsightsSettings::FTimingViewTrackSettings::Option_ShowAllVerseTracks)) { Settings->AddTimingViewToggleOption(FInsightsSettings::FTimingViewTrackSettings::Option_ShowAllVerseTracks, true); }
+		
+		if (!Settings->HasTimingViewToggleOption(FInsightsSettings::FTimingViewTrackSettings::Option_ShowGpuWorkTracks)) { Settings->AddTimingViewToggleOption(FInsightsSettings::FTimingViewTrackSettings::Option_ShowGpuWorkTracks, true); }
+		if (!Settings->HasTimingViewToggleOption(FInsightsSettings::FTimingViewTrackSettings::Option_ShowGpuWorkOverlays)) { Settings->AddTimingViewToggleOption(FInsightsSettings::FTimingViewTrackSettings::Option_ShowGpuWorkOverlays, true); }
+		if (!Settings->HasTimingViewToggleOption(FInsightsSettings::FTimingViewTrackSettings::Option_ShowGpuWorkExtendedLines)) { Settings->AddTimingViewToggleOption(FInsightsSettings::FTimingViewTrackSettings::Option_ShowGpuWorkExtendedLines, true); }
+		if (!Settings->HasTimingViewToggleOption(FInsightsSettings::FTimingViewTrackSettings::Option_ShowGpuFencesTracks)) { Settings->AddTimingViewToggleOption(FInsightsSettings::FTimingViewTrackSettings::Option_ShowGpuFencesTracks, true); }
+		if (!Settings->HasTimingViewToggleOption(FInsightsSettings::FTimingViewTrackSettings::Option_ShowGpuFencesExtendedLines)) { Settings->AddTimingViewToggleOption(FInsightsSettings::FTimingViewTrackSettings::Option_ShowGpuFencesExtendedLines, true); }
+		if (!Settings->HasTimingViewToggleOption(FInsightsSettings::FTimingViewTrackSettings::Option_ShowGpuFencesRelations)) { Settings->AddTimingViewToggleOption(FInsightsSettings::FTimingViewTrackSettings::Option_ShowGpuFencesRelations, true); }
+	}
+	else if (TimingView &&
+			 TimingView->GetName() == FInsightsManagerTabs::MemoryProfilerTabId)
+	{
+		Settings = MakeShared<FThreadSharedStatePersistentSettings>(FInsightsManagerTabs::MemoryProfilerTabId);
+		
+		if (!Settings->HasTimingViewToggleOption(FInsightsSettings::FTimingViewTrackSettings::Option_ShowAllCpuTracks)) { Settings->AddTimingViewToggleOption(FInsightsSettings::FTimingViewTrackSettings::Option_ShowAllCpuTracks, false); }
+		if (!Settings->HasTimingViewToggleOption(FInsightsSettings::FTimingViewTrackSettings::Option_ShowAllGpuTracks)) { Settings->AddTimingViewToggleOption(FInsightsSettings::FTimingViewTrackSettings::Option_ShowAllGpuTracks, false); }
+		if (!Settings->HasTimingViewToggleOption(FInsightsSettings::FTimingViewTrackSettings::Option_ShowAllVerseTracks)) { Settings->AddTimingViewToggleOption(FInsightsSettings::FTimingViewTrackSettings::Option_ShowAllVerseTracks, false); }
+
+		if (!Settings->HasTimingViewToggleOption(FInsightsSettings::FTimingViewTrackSettings::Option_ShowGpuWorkTracks)) { Settings->AddTimingViewToggleOption(FInsightsSettings::FTimingViewTrackSettings::Option_ShowGpuWorkTracks, true); }
+		if (!Settings->HasTimingViewToggleOption(FInsightsSettings::FTimingViewTrackSettings::Option_ShowGpuWorkOverlays)) { Settings->AddTimingViewToggleOption(FInsightsSettings::FTimingViewTrackSettings::Option_ShowGpuWorkOverlays, true); }
+		if (!Settings->HasTimingViewToggleOption(FInsightsSettings::FTimingViewTrackSettings::Option_ShowGpuWorkExtendedLines)) { Settings->AddTimingViewToggleOption(FInsightsSettings::FTimingViewTrackSettings::Option_ShowGpuWorkExtendedLines, true); }
+		if (!Settings->HasTimingViewToggleOption(FInsightsSettings::FTimingViewTrackSettings::Option_ShowGpuFencesTracks)) { Settings->AddTimingViewToggleOption(FInsightsSettings::FTimingViewTrackSettings::Option_ShowGpuFencesTracks, true); }
+		if (!Settings->HasTimingViewToggleOption(FInsightsSettings::FTimingViewTrackSettings::Option_ShowGpuFencesExtendedLines)) { Settings->AddTimingViewToggleOption(FInsightsSettings::FTimingViewTrackSettings::Option_ShowGpuFencesExtendedLines, true); }
+		if (!Settings->HasTimingViewToggleOption(FInsightsSettings::FTimingViewTrackSettings::Option_ShowGpuFencesRelations)) { Settings->AddTimingViewToggleOption(FInsightsSettings::FTimingViewTrackSettings::Option_ShowGpuFencesRelations, true); }
 	}
 	else
 	{
-		bShowHideAllGpuTracks = false;
-		bShowHideAllVerseTracks = false;
-		bShowHideAllCpuTracks = false;
+		Settings->AddTimingViewToggleOption(FInsightsSettings::FTimingViewTrackSettings::Option_ShowAllCpuTracks, false);
+		Settings->AddTimingViewToggleOption(FInsightsSettings::FTimingViewTrackSettings::Option_ShowAllGpuTracks, false);
+		Settings->AddTimingViewToggleOption(FInsightsSettings::FTimingViewTrackSettings::Option_ShowAllVerseTracks, false);
 	}
 
 	OldGpu1Track = nullptr;
@@ -380,10 +401,6 @@ void FThreadTimingSharedState::OnEndSession(Timing::ITimingViewSession& InSessio
 		return;
 	}
 
-	bShowHideAllGpuTracks = false;
-	bShowHideAllVerseTracks = false;
-	bShowHideAllCpuTracks = false;
-
 	OldGpu1Track = nullptr;
 	OldGpu2Track = nullptr;
 	GpuTracks.Reset();
@@ -434,7 +451,7 @@ void FThreadTimingSharedState::Tick(Timing::ITimingViewSession& InSession, const
 				{
 					OldGpu1Track = MakeShared<FGpuTimingTrack>(*this, TEXT("GPU"), nullptr, GpuTimelineIndex, FGpuTimingTrack::Gpu1ThreadId);
 					OldGpu1Track->SetOrder(FTimingTrackOrder::Gpu);
-					OldGpu1Track->SetVisibilityFlag(bShowHideAllGpuTracks);
+					OldGpu1Track->SetVisibilityFlag(Settings->IsTimingViewToggleOptionEnabled(FInsightsSettings::FTimingViewTrackSettings::Option_ShowAllGpuTracks));
 					InSession.AddScrollableTrack(OldGpu1Track);
 				}
 			}
@@ -445,7 +462,7 @@ void FThreadTimingSharedState::Tick(Timing::ITimingViewSession& InSession, const
 				{
 					OldGpu2Track = MakeShared<FGpuTimingTrack>(*this, TEXT("GPU2"), nullptr, GpuTimelineIndex, FGpuTimingTrack::Gpu2ThreadId);
 					OldGpu2Track->SetOrder(FTimingTrackOrder::Gpu + 1);
-					OldGpu2Track->SetVisibilityFlag(bShowHideAllGpuTracks);
+					OldGpu2Track->SetVisibilityFlag(Settings->IsTimingViewToggleOptionEnabled(FInsightsSettings::FTimingViewTrackSettings::Option_ShowAllGpuTracks));
 					InSession.AddScrollableTrack(OldGpu2Track);
 				}
 			}
@@ -464,7 +481,7 @@ void FThreadTimingSharedState::Tick(Timing::ITimingViewSession& InSession, const
 					// Create new Timing Events track for the GPU Queue.
 					TSharedPtr<FGpuQueueTimingTrack> Track = MakeShared<FGpuQueueTimingTrack>(*this, QueueInfo.GetDisplayName(), QueueInfo.TimelineIndex, QueueInfo.Id);
 					Track->SetOrder(GpuTrackOrder);
-					Track->SetVisibilityFlag(bShowHideAllGpuTracks);
+					Track->SetVisibilityFlag(Settings->IsTimingViewToggleOptionEnabled(FInsightsSettings::FTimingViewTrackSettings::Option_ShowAllGpuTracks));
 					GpuTracks.Add(QueueInfo.Id, Track);
 					InSession.AddScrollableTrack(Track);
 
@@ -519,7 +536,7 @@ void FThreadTimingSharedState::Tick(Timing::ITimingViewSession& InSession, const
 			ThreadProvider.EnumerateThreads([this, &InSession, &bTracksOrderChanged, &CpuTrackOrder, TimingProfilerProvider, LoadTimeProfilerProvider](const TraceServices::FThreadInfo& ThreadInfo)
 			{
 				// Check if this thread is part of a group?
-				bool bIsGroupVisible = bShowHideAllCpuTracks;
+				bool bIsGroupVisible = Settings->IsTimingViewToggleOptionEnabled(FInsightsSettings::FTimingViewTrackSettings::Option_ShowAllCpuTracks);
 				const TCHAR* GroupName = ThreadInfo.GroupName;
 				if (!GroupName || *GroupName == 0)
 				{
@@ -531,12 +548,6 @@ void FThreadTimingSharedState::Tick(Timing::ITimingViewSession& InSession, const
 				}
 				if (!ThreadGroups.Contains(GroupName))
 				{
-					if (!Settings->HasTimingViewCpuThreadVisibility(GroupName))
-					{
-						Settings->AddTimingViewCpuThreadVisibility(GroupName, bShowHideAllCpuTracks);
-					}
-					bIsGroupVisible = Settings->GetTimingViewCpuThreadVisibility(GroupName);
-
 					// Note: The GroupName pointer should be valid for the duration of the session.
 					ThreadGroups.Add(GroupName, { GroupName, 0, CpuTrackOrder });
 				}
@@ -544,7 +555,11 @@ void FThreadTimingSharedState::Tick(Timing::ITimingViewSession& InSession, const
 				{
 					FThreadGroup& ThreadGroup = ThreadGroups[GroupName];
 					ThreadGroup.Order = CpuTrackOrder;
-					bIsGroupVisible = Settings->GetTimingViewCpuThreadVisibility(GroupName);
+				}
+
+				if (Settings->HasTimingViewCpuThreadGroupVisibility(GroupName))
+				{
+					bIsGroupVisible = Settings->IsTimingViewCpuThreadGroupVisible(GroupName);
 				}
 
 				// Check if there is an available Asset Loading track for this thread.
@@ -777,11 +792,11 @@ void FThreadTimingSharedState::CreateThreadGroupsMenu(FMenuBuilder& InOutMenuBui
 
 void FThreadTimingSharedState::SetAllVerseTracksToggle(bool bOnOff)
 {
-	bShowHideAllVerseTracks = bOnOff;
+	Settings->AddTimingViewToggleOption(FInsightsSettings::FTimingViewTrackSettings::Option_ShowAllVerseTracks, bOnOff);
 
 	if (VerseSamplingTrack.IsValid())
 	{
-		VerseSamplingTrack->SetVisibilityFlag(bShowHideAllVerseTracks);
+		VerseSamplingTrack->SetVisibilityFlag(bOnOff);
 	}
 
 	if (TimingView)
@@ -794,17 +809,23 @@ void FThreadTimingSharedState::SetAllVerseTracksToggle(bool bOnOff)
 
 void FThreadTimingSharedState::SetAllCpuTracksToggle(bool bOnOff)
 {
-	bShowHideAllCpuTracks = bOnOff;
+	Settings->AddTimingViewToggleOption(FInsightsSettings::FTimingViewTrackSettings::Option_ShowAllCpuTracks, bOnOff);
 
 	for (const auto& KV : CpuTracks)
 	{
 		FCpuTimingTrack& Track = *KV.Value;
-		Track.SetVisibilityFlag(bShowHideAllCpuTracks);
+		Track.SetVisibilityFlag(bOnOff);
 	}
 
 	for (auto& KV : ThreadGroups)
 	{
-		Settings->AddTimingViewCpuThreadVisibility(KV.Value.Name, bShowHideAllCpuTracks);
+		if (Settings->HasTimingViewCpuThreadGroupVisibility(KV.Value.Name))
+		{
+			if (Settings->IsTimingViewCpuThreadGroupVisible(KV.Value.Name) == bOnOff)
+			{
+				Settings->RemoveTimingViewCpuThreadGroup(KV.Value.Name);
+			}
+		}
 	}
 
 	if (TimingView)
@@ -817,20 +838,20 @@ void FThreadTimingSharedState::SetAllCpuTracksToggle(bool bOnOff)
 
 void FThreadTimingSharedState::SetAllGpuTracksToggle(bool bOnOff)
 {
-	bShowHideAllGpuTracks = bOnOff;
-
+	Settings->AddTimingViewToggleOption(FInsightsSettings::FTimingViewTrackSettings::Option_ShowAllGpuTracks, bOnOff);
+	
 	if (OldGpu1Track.IsValid())
 	{
-		OldGpu1Track->SetVisibilityFlag(bShowHideAllGpuTracks);
+		OldGpu1Track->SetVisibilityFlag(bOnOff);
 	}
 	if (OldGpu2Track.IsValid())
 	{
-		OldGpu2Track->SetVisibilityFlag(bShowHideAllGpuTracks);
+		OldGpu2Track->SetVisibilityFlag(bOnOff);
 	}
 	for (const auto& KV : GpuTracks)
 	{
 		FGpuQueueTimingTrack& Track = *KV.Value;
-		Track.SetVisibilityFlag(bShowHideAllGpuTracks);
+		Track.SetVisibilityFlag(bOnOff);
 	}
 
 	if (TimingView)
@@ -845,7 +866,9 @@ bool FThreadTimingSharedState::ToggleTrackVisibilityByGroup_IsChecked(const TCHA
 {
 	if (ThreadGroups.Contains(InGroupName))
 	{
-		return Settings->GetTimingViewCpuThreadVisibility(InGroupName);
+		return Settings->HasTimingViewCpuThreadGroupVisibility(InGroupName)
+		? Settings->IsTimingViewCpuThreadGroupVisible(InGroupName)
+		: Settings->IsTimingViewToggleOptionEnabled(FInsightsSettings::FTimingViewTrackSettings::Option_ShowAllCpuTracks);
 	}
 	return false;
 }
@@ -856,15 +879,26 @@ void FThreadTimingSharedState::ToggleTrackVisibilityByGroup_Execute(const TCHAR*
 {
 	if (ThreadGroups.Contains(InGroupName))
 	{
-		bool bInvVisibility = !Settings->GetTimingViewCpuThreadVisibility(InGroupName);
-		Settings->AddTimingViewCpuThreadVisibility(InGroupName, bInvVisibility);
-
+		bool bNewVisibility = !Settings->IsTimingViewToggleOptionEnabled(FInsightsSettings::FTimingViewTrackSettings::Option_ShowAllCpuTracks);
+		if (Settings->HasTimingViewCpuThreadGroupVisibility(InGroupName))
+		{
+			if (!Settings->IsTimingViewCpuThreadGroupVisible(InGroupName) == !bNewVisibility)
+			{
+				Settings->RemoveTimingViewCpuThreadGroup(InGroupName);
+				bNewVisibility = !bNewVisibility;
+			}
+		}
+		else
+		{
+			Settings->AddTimingViewCpuThreadGroupVisibility(InGroupName, bNewVisibility);
+		}
+		
 		for (const auto& KV : CpuTracks)
 		{
 			FCpuTimingTrack& Track = *KV.Value;
 			if (Track.GetGroupName() == InGroupName)
 			{
-				Track.SetVisibilityFlag(bInvVisibility);
+				Track.SetVisibilityFlag(bNewVisibility);
 			}
 		}
 
@@ -877,6 +911,14 @@ void FThreadTimingSharedState::ToggleTrackVisibilityByGroup_Execute(const TCHAR*
 
 ////////////////////////////////////////////////////////////////////////////////////////////////////
 
+void FThreadTimingSharedState::ToggleOptionState(const TCHAR* CommandName)
+{
+	bool bInvVisibility = !Settings->IsTimingViewToggleOptionEnabled(CommandName);
+	Settings->AddTimingViewToggleOption(CommandName, bInvVisibility);
+}
+
+////////////////////////////////////////////////////////////////////////////////////////////////////
+
 TSharedPtr<const ITimingEvent> FThreadTimingSharedState::FindMaxEventInstance(uint32 TimerId, double StartTime, double EndTime)
 {
 	auto CompareAndAssignEvent = [](TSharedPtr<const ITimingEvent>& TimingEvent, TSharedPtr<const ITimingEvent>& TrackEvent)
@@ -989,9 +1031,9 @@ TSharedPtr<const ITimingEvent> FThreadTimingSharedState::FindMinEventInstance(ui
 
 void FThreadTimingSharedState::SetGpuWorkTracksVisibility(bool bOnOff)
 {
-	if (Settings->GetTimingViewShowGpuWorkTracks() != bOnOff)
+	if (Settings->IsTimingViewToggleOptionEnabled(FInsightsSettings::FTimingViewTrackSettings::Option_ShowGpuWorkTracks) != bOnOff)
 	{
-		Settings->SetTimingViewShowGpuWorkTracks(bOnOff);
+		Settings->AddTimingViewToggleOption(FInsightsSettings::FTimingViewTrackSettings::Option_ShowGpuWorkTracks, bOnOff);
 
 		if (bOnOff)
 		{
@@ -1099,9 +1141,9 @@ void FThreadTimingSharedState::RemoveGpuFencesChildTracks()
 
 void FThreadTimingSharedState::SetGpuFencesTracksVisibility(bool bOnOff)
 {
-	if (Settings->GetTimingViewShowGpuFencesTracks() != bOnOff)
+	if (Settings->IsTimingViewToggleOptionEnabled(FInsightsSettings::FTimingViewTrackSettings::Option_ShowGpuFencesTracks) != bOnOff)
 	{
-		Settings->SetTimingViewShowGpuFencesTracks(bOnOff);
+		Settings->AddTimingViewToggleOption(FInsightsSettings::FTimingViewTrackSettings::Option_ShowGpuFencesTracks, bOnOff);
 
 		if (bOnOff)
 		{
@@ -1118,10 +1160,10 @@ void FThreadTimingSharedState::SetGpuFencesTracksVisibility(bool bOnOff)
 
 void FThreadTimingSharedState::Command_ShowGpuFencesRelations_Execute()
 {
-	bool NewValue = !Settings->GetTimingViewShowGpuFencesRelations();
-	Settings->SetTimingViewShowGpuFencesRelations(NewValue);
+	ToggleOptionState(FInsightsSettings::FTimingViewTrackSettings::Option_ShowGpuFencesRelations);
+	bool bVisibility = Settings->IsTimingViewToggleOptionEnabled(FInsightsSettings::FTimingViewTrackSettings::Option_ShowGpuFencesRelations);
 
-	if (NewValue == false && TimingView)
+	if (bVisibility == false && TimingView)
 	{
 		TimingView->EditCurrentRelations().RemoveAll([](TUniquePtr<ITimingEventRelation>& Item)
 			{
@@ -1132,6 +1174,14 @@ void FThreadTimingSharedState::Command_ShowGpuFencesRelations_Execute()
 
 ////////////////////////////////////////////////////////////////////////////////////////////////////
 
+void FThreadTimingSharedState::Command_ShowGpuFencesTracks_Execute()
+{
+	bool bInvVisibility = !Settings->IsTimingViewToggleOptionEnabled(FInsightsSettings::FTimingViewTrackSettings::Option_ShowGpuFencesTracks);
+	SetGpuFencesTracksVisibility(bInvVisibility);
+}
+
+////////////////////////////////////////////////////////////////////////////////////////////////////
+
 void FThreadTimingSharedState::OnTimingEventSelected(TSharedPtr<const ITimingEvent> InSelectedEvent)
 {
 	if (!AreGpuFenceRelationsVisible())
diff --git a/Engine/Source/Developer/TraceInsights/Private/Insights/TimingProfiler/ViewModels/ThreadTimingSharedState.h b/Engine/Source/Developer/TraceInsights/Private/Insights/TimingProfiler/ViewModels/ThreadTimingSharedState.h
--- a/Engine/Source/Developer/TraceInsights/Private/Insights/TimingProfiler/ViewModels/ThreadTimingSharedState.h
+++ b/Engine/Source/Developer/TraceInsights/Private/Insights/TimingProfiler/ViewModels/ThreadTimingSharedState.h
@@ -11,6 +11,7 @@
 
 // TraceInsights
 #include "Insights/InsightsManager.h"
+#include "Insights/InsightsSettings.h"
 #include "Insights/ITimingViewExtender.h"
 
 namespace TraceServices
@@ -46,54 +47,39 @@ private:
 	{
 	public:
 		virtual ~IThreadSharedStateSetting() {}
+		
+		virtual bool IsTimingViewCpuThreadGroupVisible(const TCHAR* GroupName) const = 0;
+		virtual bool HasTimingViewCpuThreadGroupVisibility(const TCHAR* GroupName) const = 0;
+		virtual void AddTimingViewCpuThreadGroupVisibility(const TCHAR* GroupName, bool InValue) = 0;
+		virtual void RemoveTimingViewCpuThreadGroup(const TCHAR* GroupName) = 0;
 
-		virtual bool GetTimingViewShowGpuWorkTracks() const = 0;
-		virtual void SetTimingViewShowGpuWorkTracks(bool InValue) = 0;
-
-		virtual bool GetTimingViewShowGpuWorkOverlays() const = 0;
-		virtual void SetTimingViewShowGpuWorkOverlays(bool InValue) = 0;
-
-		virtual bool GetTimingViewShowGpuWorkExtendedLines() const = 0;
-		virtual void SetTimingViewShowGpuWorkExtendedLines(bool InValue) = 0;
-
-		virtual bool GetTimingViewShowGpuFencesTracks() const = 0;
-		virtual void SetTimingViewShowGpuFencesTracks(bool InValue) = 0;
-
-		virtual bool GetTimingViewShowGpuFencesExtendedLines() const = 0;
-		virtual void SetTimingViewShowGpuFencesExtendedLines(bool InValue) = 0;
-
-		virtual bool GetTimingViewShowGpuFencesRelations() const = 0;
-		virtual void SetTimingViewShowGpuFencesRelations(bool InValue) = 0;
+		virtual bool IsTimingViewToggleOptionEnabled(const TCHAR* Key) const = 0;
+		virtual bool HasTimingViewToggleOption(const TCHAR* Key) const = 0;
+		virtual void AddTimingViewToggleOption(const TCHAR* Key, bool InValue) = 0;
 
-		virtual bool GetTimingViewCpuThreadVisibility(const TCHAR* GroupName) const = 0;
-		virtual bool HasTimingViewCpuThreadVisibility(const TCHAR* GroupName) const = 0;
-		virtual void AddTimingViewCpuThreadVisibility(const TCHAR* GroupName, bool InValue) = 0;
+		virtual void Reset() = 0;
 	};
 
 	class FThreadSharedStatePersistentSettings : public IThreadSharedStateSetting
 	{
 	public:
-		virtual bool GetTimingViewShowGpuWorkTracks() const override { return GetInsightsSettings().GetTimingViewShowGpuWorkTracks(); }
-		virtual void SetTimingViewShowGpuWorkTracks(bool InValue) override { GetInsightsSettings().SetAndSaveTimingViewShowGpuWorkTracks(InValue); }
-
-		virtual bool GetTimingViewShowGpuWorkOverlays() const override { return GetInsightsSettings().GetTimingViewShowGpuWorkOverlays(); }
-		virtual void SetTimingViewShowGpuWorkOverlays(bool InValue) override { GetInsightsSettings().SetAndSaveTimingViewShowGpuWorkOverlays(InValue); }
-
-		virtual bool GetTimingViewShowGpuWorkExtendedLines() const override { return GetInsightsSettings().GetTimingViewShowGpuWorkExtendedLines(); }
-		virtual void SetTimingViewShowGpuWorkExtendedLines(bool InValue) override { GetInsightsSettings().SetAndSaveTimingViewShowGpuWorkExtendedLines(InValue); }
+		FThreadSharedStatePersistentSettings(const FName& InProfiler) : Profiler(InProfiler)
+		{
+			GetInsightsSettings().CreateIfNotExistProfilerTimingViewTrackSettings(Profiler);
+		}
 
-		virtual bool GetTimingViewShowGpuFencesTracks() const override { return  GetInsightsSettings().GetTimingViewShowGpuFencesTracks(); }
-		virtual void SetTimingViewShowGpuFencesTracks(bool InValue) override { GetInsightsSettings().SetAndSaveTimingViewShowGpuFencesTracks(InValue); }
+		FName Profiler;
 
-		virtual bool GetTimingViewShowGpuFencesExtendedLines() const override { return GetInsightsSettings().GetTimingViewShowGpuFencesExtendedLines(); }
-		virtual void SetTimingViewShowGpuFencesExtendedLines(bool InValue) override { GetInsightsSettings().SetAndSaveTimingViewShowGpuFencesExtendedLines(InValue); }
+		virtual bool IsTimingViewCpuThreadGroupVisible(const TCHAR* GroupName) const override { return GetInsightsSettings().IsTimingViewCpuThreadGroupVisible(GroupName, Profiler); }
+		virtual bool HasTimingViewCpuThreadGroupVisibility(const TCHAR* GroupName) const override { return GetInsightsSettings().HasTimingViewCpuThreadGroupVisibility(GroupName, Profiler); }
+		virtual void AddTimingViewCpuThreadGroupVisibility(const TCHAR* GroupName, bool InValue) override { return GetInsightsSettings().AddAndSaveTimingViewCpuThreadGroupVisibility(GroupName, InValue, Profiler); }
+		virtual void RemoveTimingViewCpuThreadGroup(const TCHAR* GroupName) override { return GetInsightsSettings().RemoveAndSaveTimingViewCpuThreadGroup(GroupName, Profiler); }
 
-		virtual bool GetTimingViewShowGpuFencesRelations() const override { return GetInsightsSettings().GetTimingViewShowGpuFencesRelations(); }
-		virtual void SetTimingViewShowGpuFencesRelations(bool InValue) override { GetInsightsSettings().SetAndSaveTimingViewShowGpuFencesRelations(InValue); };
+		virtual bool IsTimingViewToggleOptionEnabled(const TCHAR* Key) const override { return GetInsightsSettings().IsTimingViewToggleOptionEnabled(Key, Profiler); }
+		virtual bool HasTimingViewToggleOption(const TCHAR* Key) const override { return GetInsightsSettings().HasTimingViewToggleOption(Key, Profiler); }
+		virtual void AddTimingViewToggleOption(const TCHAR* Key, bool InValue) override { return GetInsightsSettings().AddAndSaveTimingViewToggleOption(Key, InValue, Profiler); }
 
-		virtual bool GetTimingViewCpuThreadVisibility(const TCHAR* GroupName) const override { return GetInsightsSettings().GetTimingViewCpuThreadVisibility(GroupName); }
-		virtual bool HasTimingViewCpuThreadVisibility(const TCHAR* GroupName) const override { return GetInsightsSettings().HasTimingViewCpuThreadVisibility(GroupName); }
-		virtual void AddTimingViewCpuThreadVisibility(const TCHAR* GroupName, bool InValue) override { return GetInsightsSettings().AddAndSaveTimingViewCpuThreadVisibility(GroupName, InValue); }
+		virtual void Reset() override { GetInsightsSettings().ResetAllProfilerTimingViewTrackSettings(Profiler); }
 		
 		FInsightsSettings& GetInsightsSettings() { return UE::Insights::FInsightsManager::Get()->GetSettings(); }
 		const FInsightsSettings& GetInsightsSettings() const { return UE::Insights::FInsightsManager::Get()->GetSettings(); }
@@ -102,37 +88,19 @@ private:
 	class FThreadSharedStateLocalSettings : public IThreadSharedStateSetting
 	{
 	public:
-		virtual bool GetTimingViewShowGpuWorkTracks() const override { return bTimingViewShowGpuWorkTracks; }
-		virtual void SetTimingViewShowGpuWorkTracks(bool InValue) override { bTimingViewShowGpuWorkTracks = InValue; }
+		virtual bool IsTimingViewCpuThreadGroupVisible(const TCHAR* GroupName) const override { return HasTimingViewCpuThreadGroupVisibility(GroupName) ? *TimingViewShowCpuThreadGroups.Find(GroupName) : false; }
+		virtual bool HasTimingViewCpuThreadGroupVisibility(const TCHAR* GroupName) const override { return TimingViewShowCpuThreadGroups.Contains(GroupName); }
+		virtual void AddTimingViewCpuThreadGroupVisibility(const TCHAR* GroupName, bool InValue) override { TimingViewShowCpuThreadGroups.Add(GroupName, InValue); }
+		virtual void RemoveTimingViewCpuThreadGroup(const TCHAR* GroupName) override { if (HasTimingViewCpuThreadGroupVisibility(GroupName)) { TimingViewShowCpuThreadGroups.Remove(GroupName); } }
 
-		virtual bool GetTimingViewShowGpuWorkOverlays() const override { return bTimingViewShowGpuWorkOverlays; }
-		virtual void SetTimingViewShowGpuWorkOverlays(bool InValue) override { bTimingViewShowGpuWorkOverlays = InValue; }
+		virtual bool IsTimingViewToggleOptionEnabled(const TCHAR* Key) const override { return HasTimingViewToggleOption(Key) ? *TimingViewShowToggleOptions.Find(Key) : false; }
+		virtual bool HasTimingViewToggleOption(const TCHAR* Key) const override { return TimingViewShowToggleOptions.Contains(Key); }
+		virtual void AddTimingViewToggleOption(const TCHAR* Key, bool InValue) override { TimingViewShowToggleOptions.Add(Key, InValue); }
 
-		virtual bool GetTimingViewShowGpuWorkExtendedLines() const override { return bTimingViewShowGpuWorkExtendedLines; }
-		virtual void SetTimingViewShowGpuWorkExtendedLines(bool InValue) override { bTimingViewShowGpuWorkExtendedLines = InValue; }
-
-		virtual bool GetTimingViewShowGpuFencesTracks() const override { return bTimingViewShowGpuFencesTracks; }
-		virtual void SetTimingViewShowGpuFencesTracks(bool InValue) override { bTimingViewShowGpuFencesTracks = InValue; }
-
-		virtual bool GetTimingViewShowGpuFencesExtendedLines() const override { return bTimingViewShowGpuFencesExtendedLines; }
-		virtual void SetTimingViewShowGpuFencesExtendedLines(bool InValue) override { bTimingViewShowGpuFencesExtendedLines = InValue; }
-
-		virtual bool GetTimingViewShowGpuFencesRelations() const override { return bTimingViewShowGpuFencesRelations; }
-		virtual void SetTimingViewShowGpuFencesRelations(bool InValue) override { bTimingViewShowGpuFencesRelations = InValue; }
-
-		virtual bool GetTimingViewCpuThreadVisibility(const TCHAR* GroupName) const override { return HasTimingViewCpuThreadVisibility(GroupName) ? *TimingViewShowCpuThreadGroups.Find(GroupName) : false; }
-		virtual bool HasTimingViewCpuThreadVisibility(const TCHAR* GroupName) const override { return TimingViewShowCpuThreadGroups.Contains(GroupName); }
-		virtual void AddTimingViewCpuThreadVisibility(const TCHAR* GroupName, bool InValue) override { TimingViewShowCpuThreadGroups.Add(GroupName, InValue); }
+		virtual void Reset() override { TimingViewShowToggleOptions.Reset(); TimingViewShowCpuThreadGroups.Reset(); }
 
 	private:
-		bool bTimingViewShowGpuWorkTracks = true;
-		bool bTimingViewShowGpuWorkOverlays = true;
-		bool bTimingViewShowGpuWorkExtendedLines = true;
-
-		bool bTimingViewShowGpuFencesTracks = true;
-		bool bTimingViewShowGpuFencesExtendedLines = true;
-		bool bTimingViewShowGpuFencesRelations = true;
-
+		TMap<FString, bool> TimingViewShowToggleOptions;
 		TMap<FString, bool> TimingViewShowCpuThreadGroups;
 	};
 
@@ -173,35 +141,35 @@ public:
 
 	void ResetToDefaultValues();
 
-	bool IsAllGpuTracksToggleOn() const { return bShowHideAllGpuTracks; }
+	bool IsAllGpuTracksToggleOn() const { return Settings->IsTimingViewToggleOptionEnabled(FInsightsSettings::FTimingViewTrackSettings::Option_ShowAllGpuTracks); }
 	void SetAllGpuTracksToggle(bool bOnOff);
 	void ShowAllGpuTracks() { SetAllGpuTracksToggle(true); }
 	void HideAllGpuTracks() { SetAllGpuTracksToggle(false); }
 	void ShowHideAllGpuTracks() { SetAllGpuTracksToggle(!IsAllGpuTracksToggleOn()); }
 
-	bool IsAllVerseTracksToggleOn() const { return bShowHideAllVerseTracks; }
+	bool IsAllVerseTracksToggleOn() const { return Settings->IsTimingViewToggleOptionEnabled(FInsightsSettings::FTimingViewTrackSettings::Option_ShowAllVerseTracks); }
 	void SetAllVerseTracksToggle(bool bOnOff);
 	void ShowAllVerseTracks() { SetAllVerseTracksToggle(true); }
 	void HideAllVerseTracks() { SetAllVerseTracksToggle(false); }
 	void ShowHideAllVerseTracks() { SetAllVerseTracksToggle(!IsAllVerseTracksToggleOn()); }
 
-	bool IsAllCpuTracksToggleOn() const { return bShowHideAllCpuTracks; }
+	bool IsAllCpuTracksToggleOn() const { return Settings->IsTimingViewToggleOptionEnabled(FInsightsSettings::FTimingViewTrackSettings::Option_ShowAllCpuTracks); }
 	void SetAllCpuTracksToggle(bool bOnOff);
 	void ShowAllCpuTracks() { SetAllCpuTracksToggle(true); }
 	void HideAllCpuTracks() { SetAllCpuTracksToggle(false); }
 	void ShowHideAllCpuTracks() { SetAllCpuTracksToggle(!IsAllCpuTracksToggleOn()); }
 
-	bool AreOverlaysVisibleInGpuQueueTracks() const { return Settings->GetTimingViewShowGpuWorkOverlays(); }
-	bool AreExtendedLinesVisibleInGpuQueueTracks() const { return Settings->GetTimingViewShowGpuWorkExtendedLines(); }
+	bool AreOverlaysVisibleInGpuQueueTracks() const { return Settings->IsTimingViewToggleOptionEnabled(FInsightsSettings::FTimingViewTrackSettings::Option_ShowGpuWorkOverlays); }
+	bool AreExtendedLinesVisibleInGpuQueueTracks() const { return Settings->IsTimingViewToggleOptionEnabled(FInsightsSettings::FTimingViewTrackSettings::Option_ShowGpuWorkExtendedLines); }
 
-	bool AreGpuWorkTracksVisible() const { return Settings->GetTimingViewShowGpuWorkTracks(); }
+	bool AreGpuWorkTracksVisible() const { return Settings->IsTimingViewToggleOptionEnabled(FInsightsSettings::FTimingViewTrackSettings::Option_ShowGpuWorkTracks); }
 	void SetGpuWorkTracksVisibility(bool bOnOff);
 
-	bool AreGpuFencesTracksVisible() const { return Settings->GetTimingViewShowGpuFencesTracks(); }
+	bool AreGpuFencesTracksVisible() const { return Settings->IsTimingViewToggleOptionEnabled(FInsightsSettings::FTimingViewTrackSettings::Option_ShowGpuFencesTracks); }
 	void SetGpuFencesTracksVisibility(bool bOnOff);
 
-	bool AreGpuFencesExtendedLinesVisible() const { return Settings->GetTimingViewShowGpuFencesExtendedLines(); }
-	bool AreGpuFenceRelationsVisible() const { return Settings->GetTimingViewShowGpuFencesRelations(); }
+	bool AreGpuFencesExtendedLinesVisible() const { return Settings->IsTimingViewToggleOptionEnabled(FInsightsSettings::FTimingViewTrackSettings::Option_ShowGpuFencesExtendedLines); }
+	bool AreGpuFenceRelationsVisible() const { return Settings->IsTimingViewToggleOptionEnabled(FInsightsSettings::FTimingViewTrackSettings::Option_ShowGpuFencesRelations); }
 
 	TSharedPtr<const ITimingEvent> FindMaxEventInstance(uint32 TimerId, double StartTime, double EndTime);
 	TSharedPtr<const ITimingEvent> FindMinEventInstance(uint32 TimerId, double StartTime, double EndTime);
@@ -211,20 +179,21 @@ private:
 
 	bool ToggleTrackVisibilityByGroup_IsChecked(const TCHAR* InGroupName) const;
 	void ToggleTrackVisibilityByGroup_Execute(const TCHAR* InGroupName);
+	void ToggleOptionState(const TCHAR* CommandName);
 
-	void Command_ShowGpuWorkTracks_Execute() { SetGpuWorkTracksVisibility(!Settings->GetTimingViewShowGpuWorkTracks()); }
+	void Command_ShowGpuWorkTracks_Execute() { SetGpuWorkTracksVisibility(!Settings->IsTimingViewToggleOptionEnabled(FInsightsSettings::FTimingViewTrackSettings::Option_ShowGpuWorkTracks)); }
 	bool Command_ShowGpuWorkTracks_CanExecute() { return GpuTracks.Num() > 0; }
 	bool Command_ShowGpuWorkTracks_IsChecked() { return AreGpuWorkTracksVisible(); }
 
-	void Command_ShowGpuWorkOverlays_Execute() { Settings->SetTimingViewShowGpuWorkOverlays(!Settings->GetTimingViewShowGpuWorkOverlays()); }
+	void Command_ShowGpuWorkOverlays_Execute() { ToggleOptionState(FInsightsSettings::FTimingViewTrackSettings::Option_ShowGpuWorkOverlays); }
 	bool Command_ShowGpuWorkOverlays_CanExecute() { return AreGpuWorkTracksVisible() && GpuTracks.Num() > 0; }
 	bool Command_ShowGpuWorkOverlays_IsChecked() { return AreOverlaysVisibleInGpuQueueTracks(); }
 
-	void Command_ShowGpuWorkExtendedLines_Execute() { Settings->SetTimingViewShowGpuWorkExtendedLines(!Settings->GetTimingViewShowGpuWorkExtendedLines()); }
+	void Command_ShowGpuWorkExtendedLines_Execute() { ToggleOptionState(FInsightsSettings::FTimingViewTrackSettings::Option_ShowGpuWorkExtendedLines); }
 	bool Command_ShowGpuWorkExtendedLines_CanExecute() { return AreGpuWorkTracksVisible() && GpuTracks.Num() > 0; }
 	bool Command_ShowGpuWorkExtendedLines_IsChecked() { return AreExtendedLinesVisibleInGpuQueueTracks(); }
 
-	void Command_ShowGpuFencesExtendedLines_Execute() { Settings->SetTimingViewShowGpuFencesExtendedLines(!Settings->GetTimingViewShowGpuFencesExtendedLines()); }
+	void Command_ShowGpuFencesExtendedLines_Execute() { ToggleOptionState(FInsightsSettings::FTimingViewTrackSettings::Option_ShowGpuFencesExtendedLines); }
 	bool Command_ShowGpuFencesExtendedLines_CanExecute() { return AreGpuFencesTracksVisible() && GpuTracks.Num() > 0; }
 	bool Command_ShowGpuFencesExtendedLines_IsChecked() { return AreGpuFencesExtendedLinesVisible(); }
 
@@ -232,7 +201,7 @@ private:
 	bool Command_ShowGpuFencesRelations_CanExecute() { return GpuTracks.Num() > 0; }
 	bool Command_ShowGpuFencesRelations_IsChecked() { return AreGpuFenceRelationsVisible(); }
 
-	void Command_ShowGpuFencesTracks_Execute() { SetGpuFencesTracksVisibility(!Settings->GetTimingViewShowGpuFencesTracks()); }
+	void Command_ShowGpuFencesTracks_Execute();
 	bool Command_ShowGpuFencesTracks_CanExecute() { return GpuTracks.Num() > 0; }
 	bool Command_ShowGpuFencesTracks_IsChecked() { return AreGpuFencesTracksVisible(); }
 
@@ -247,10 +216,6 @@ private:
 private:
 	STimingView* TimingView = nullptr;
 
-	bool bShowHideAllGpuTracks = false;
-	bool bShowHideAllVerseTracks = false;
-	bool bShowHideAllCpuTracks = false;
-
 	TSharedPtr<FGpuTimingTrack> OldGpu1Track;
 	TSharedPtr<FGpuTimingTrack> OldGpu2Track;
 
diff --git a/Engine/Source/Developer/TraceInsights/Private/Insights/Widgets/STimingView.cpp b/Engine/Source/Developer/TraceInsights/Private/Insights/Widgets/STimingView.cpp
--- a/Engine/Source/Developer/TraceInsights/Private/Insights/Widgets/STimingView.cpp
+++ b/Engine/Source/Developer/TraceInsights/Private/Insights/Widgets/STimingView.cpp
@@ -338,8 +338,6 @@ void STimingView::Construct(const FArguments& InArgs, FName InViewName)
 void STimingView::HideAllDefaultTracks()
 {
 	FrameSharedState->HideAllFrameTracks();
-	ThreadTimingSharedState->HideAllGpuTracks();
-	ThreadTimingSharedState->HideAllCpuTracks();
 	LoadingSharedState->HideAllLoadingTracks();
 	FileActivitySharedState->HideAllIoTracks();
 }
