[Insights] Refactoring of the Timing View settings implementation. Also added default values in the DefaultUnrealInsightsSettings.ini file.

diff --git a/Engine/Programs/UnrealInsights/Config/DefaultUnrealInsightsSettings.ini b/Engine/Programs/UnrealInsights/Config/DefaultUnrealInsightsSettings.ini
--- a/Engine/Programs/UnrealInsights/Config/DefaultUnrealInsightsSettings.ini
+++ b/Engine/Programs/UnrealInsights/Config/DefaultUnrealInsightsSettings.ini
@@ -1,5 +1,10 @@
 [SectionsToSave]
-bCanSaveAllSections=true
+bCanSaveAllSections=True
+
+[Insights.PersistentTimingView]
++Name=TimingProfiler
++Name=LoadingProfiler
++Name=MemoryProfiler
 
 [Insights.TimingProfiler]
 
@@ -56,6 +61,24 @@ LowerThreshold=60 fps
 ; If enabled, the Timing view will also be zoomed when a frame is selected in the Frames view.
 bAutoZoomOnFrameSelection=False
 
+[Insights.TimingProfiler.TimingView]
+; These options can be changed from the CPU/GPU drop down menus of Timing Insights profiler / Timing view.
+
+ShowAllGpuTracks=True
+ShowGpuWorkTracks=True
+ShowGpuWorkOverlays=True
+ShowGpuWorkExtendedLines=True
+ShowGpuFencesRelations=True
+ShowGpuFencesTracks=True
+ShowGpuFencesExtendedLines=True
+
+ShowAllVerseTracks=True
+
+ShowAllCpuTracks=True
+;+VisibleCpuThreadGroups=GameThread
+;+VisibleCpuThreadGroups=Render
+;+HiddenCpuThreadGroups=Trace
+
 [Insights.TimingProfiler.MainGraph]
 ; These options can be changed from the context menu of the Main Graph track (Timing Insights profiler / Timing view).
 
@@ -109,13 +132,31 @@ Mode=2
 GroupingMode=3
 
 ; Controls the CPU filter button from Timers View.
-ShowCpuTimers=1
+ShowCpuTimers=True
+
+; Controls the Verse filter button from Timers View.
+ShowVerseTimers=True
 
 ; Controls the GPU filter button from Timers View.
-ShowGpuTimers=1
+ShowGpuTimers=True
 
 ; Controls the "Filter Out Zero Count Timers" button from Timers View.
-ShowZeroCountTimers=1
+ShowZeroCountTimers=True
+
+[Insights.AssetLoadingProfiler.TimingView]
+; These options can be changed from the CPU/GPU drop down menus of Asset Loading Insights profiler / Timing view.
+
+ShowAllGpuTracks=False
+ShowGpuWorkTracks=True
+ShowGpuWorkOverlays=True
+ShowGpuWorkExtendedLines=True
+ShowGpuFencesRelations=True
+ShowGpuFencesTracks=True
+ShowGpuFencesExtendedLines=True
+
+ShowAllVerseTracks=False
+
+ShowAllCpuTracks=False
 
 [Insights.MemoryProfiler]
 
@@ -126,3 +167,18 @@ ShowZeroCountTimers=1
 
 ; Memory Budgets (used by the LLM Tags view)
 +BudgetFilePath=(Label="Default",Path="Engine/Programs/UnrealInsights/Config/DefaultMemoryBudgets.xml")
+
+[Insights.MemoryProfiler.TimingView]
+; These options can be changed from the CPU/GPU drop down menus of Memory Insights profiler / Timing view.
+
+ShowAllGpuTracks=False
+ShowGpuWorkTracks=True
+ShowGpuWorkOverlays=True
+ShowGpuWorkExtendedLines=True
+ShowGpuFencesRelations=True
+ShowGpuFencesTracks=True
+ShowGpuFencesExtendedLines=True
+
+ShowAllVerseTracks=False
+
+ShowAllCpuTracks=False
diff --git a/Engine/Source/Developer/TraceInsights/Private/Insights/InsightsSettings.cpp b/Engine/Source/Developer/TraceInsights/Private/Insights/InsightsSettings.cpp
--- a/Engine/Source/Developer/TraceInsights/Private/Insights/InsightsSettings.cpp
+++ b/Engine/Source/Developer/TraceInsights/Private/Insights/InsightsSettings.cpp
@@ -44,7 +44,11 @@ FInsightsSettings::FInsightsSettings(bool bInIsDefault)
 
 FInsightsSettings::~FInsightsSettings()
 {
-	ProfilerTimingViewTrackSettings.Reset();
+	for (auto& KV : AllTimingViewSettings)
+	{
+		delete KV.Value;
+	}
+	AllTimingViewSettings.Reset();
 }
 
 ////////////////////////////////////////////////////////////////////////////////////////////////////
@@ -70,45 +74,46 @@ void FInsightsSettings::LoadFromConfig()
 		return;
 	}
 
-	//////////////////////////////////////////////////
-	// [Insights.(AllProfilers).TimingView]
+	using namespace UE::Insights;
 
-	TArray<FString> PersistentProfilers;
-	GConfig->GetArray(TEXT("Insights.PersistentTimingView"), TEXT("Profiler"), PersistentProfilers, SettingsIni);
-
-	for (const FString& Profiler : PersistentProfilers)
+	//////////////////////////////////////////////////
+	// [Insights.(TimingViewName).TimingView]
 	{
-		CreateIfNotExistProfilerTimingViewTrackSettings(*Profiler);
-		
-		TArray<FString> EnabledToggleOptions;
-		TArray<FString> DisabledToggleOptions;
-		GConfig->GetArray(*FString::Printf(TEXT("Insights.%s.TimingView"), *Profiler), TEXT("EnabledToggleOptions"), EnabledToggleOptions, SettingsIni);
-		GConfig->GetArray(*FString::Printf(TEXT("Insights.%s.TimingView"), *Profiler), TEXT("DisabledToggleOptions"), DisabledToggleOptions, SettingsIni);
-
-		for (const FString& Option : EnabledToggleOptions)
-		{
-			AddTimingViewToggleOption(*Option, true, *Profiler);
-		}
-		for (const FString& Option : DisabledToggleOptions)
-		{
-			AddTimingViewToggleOption(*Option, false, *Profiler);
-		}
-		
-		TArray<FString> VisibleCpuThreadGroups;
-		TArray<FString> HiddenCpuThreadGroups;
-		GConfig->GetArray(*FString::Printf(TEXT("Insights.%s.TimingView.CpuThreadGroups"), *Profiler), TEXT("VisibleCpuThreadGroups"), VisibleCpuThreadGroups, SettingsIni);
-		GConfig->GetArray(*FString::Printf(TEXT("Insights.%s.TimingView.CpuThreadGroups"), *Profiler), TEXT("HiddenCpuThreadGroups"), HiddenCpuThreadGroups, SettingsIni);
+		TArray<FString> PersistentTimingViewSettings;
+		GConfig->GetArray(TEXT("Insights.PersistentTimingView"), TEXT("Name"), PersistentTimingViewSettings, SettingsIni);
 
-		for (const FString& CpuGroup : VisibleCpuThreadGroups)
+		for (const FString& SettingsName : PersistentTimingViewSettings)
 		{
-			AddTimingViewCpuThreadGroupVisibility(*CpuGroup, true, *Profiler);
-		}
-		for (const FString& CpuGroup : HiddenCpuThreadGroups)
-		{
-			AddTimingViewCpuThreadGroupVisibility(*CpuGroup, false, *Profiler);
+			FTimingViewSettings& TimingViewSettings = GetOrCreateTimingViewSettings(*SettingsName);
+
+			const FString Section = FString::Printf(TEXT("Insights.%s.TimingView"), *SettingsName);
+
+			// Toggle Options
+			{
+				TimingViewSettings.EnumerateToggleOptions([this, &Section](FTimingViewSettings::FToggleOption& Option)
+					{
+						GConfig->GetBool(*Section, Option.Name, Option.bValue, SettingsIni);
+					});
+			}
+
+			// CPU Thread Groups
+			{
+				TArray<FString> VisibleCpuThreadGroups;
+				GConfig->GetArray(*Section, TEXT("VisibleCpuThreadGroups"), VisibleCpuThreadGroups, SettingsIni);
+				for (const FString& CpuGroup : VisibleCpuThreadGroups)
+				{
+					TimingViewSettings.AddCpuThreadGroup(*CpuGroup, true);
+				}
+
+				TArray<FString> HiddenCpuThreadGroups;
+				GConfig->GetArray(*Section, TEXT("HiddenCpuThreadGroups"), HiddenCpuThreadGroups, SettingsIni);
+				for (const FString& CpuGroup : HiddenCpuThreadGroups)
+				{
+					TimingViewSettings.AddCpuThreadGroup(*CpuGroup, false);
+				}
+			}
 		}
 	}
-
 	//////////////////////////////////////////////////
 	// [Insights.TimingProfiler]
 
@@ -231,39 +236,58 @@ void FInsightsSettings::LoadFromConfig()
 
 void FInsightsSettings::SaveToConfig()
 {
-	//////////////////////////////////////////////////
-	// [Insights.(AllProfilers).TimingView]
+	using namespace UE::Insights;
 
-	TArray<FString> ProfilersWithPersistentTimingViews;
-	for (auto& ProfilerSettings : ProfilerTimingViewTrackSettings)
-	{
-		ProfilersWithPersistentTimingViews.Add(ProfilerSettings.Key.ToString());
-	}
-	GConfig->SetArray(TEXT("Insights.PersistentTimingView"), TEXT("Profiler"), ProfilersWithPersistentTimingViews, SettingsIni);
-	
-	for (auto& ProfilerSettings : ProfilerTimingViewTrackSettings)
+	//////////////////////////////////////////////////
+	// [Insights.(TimingViewName).TimingView]
 	{
-		FString Profiler = ProfilerSettings.Key.ToString();
-
-		TArray<FString> EnabledToggleOptions;
-		TArray<FString> DisabledToggleOptions;
-		for (auto& KV : GetProfilerTimingViewTrackSettings(*Profiler).GetToggleOptions())
+		TArray<FString> PersistentTimingViewSettings;
+		for (auto& KV : AllTimingViewSettings)
 		{
-			KV.Value ? EnabledToggleOptions.Add(*KV.Key) : DisabledToggleOptions.Add(*KV.Key);
+			PersistentTimingViewSettings.Add(KV.Key.ToString());
 		}
-		GConfig->SetArray(*FString::Printf(TEXT("Insights.%s.TimingView"),*Profiler), TEXT("EnabledToggleOptions"), EnabledToggleOptions, SettingsIni);
-		GConfig->SetArray(*FString::Printf(TEXT("Insights.%s.TimingView"),*Profiler), TEXT("DisabledToggleOptions"), DisabledToggleOptions, SettingsIni);
-			
-		TArray<FString> VisibleCpuThreadGroups;
-		TArray<FString> HiddenCpuThreadGroups;
-		for (auto& KV : GetProfilerTimingViewTrackSettings(*Profiler).GetCpuThreadGroups())
+		GConfig->SetArray(TEXT("Insights.PersistentTimingView"), TEXT("Name"), PersistentTimingViewSettings, SettingsIni);
+
+		for (auto& KV : AllTimingViewSettings)
 		{
-			KV.Value ? VisibleCpuThreadGroups.Add(*KV.Key) : HiddenCpuThreadGroups.Add(*KV.Key);
+			FTimingViewSettings& TimingViewSettings = *KV.Value;
+			FString SettingsName = TimingViewSettings.GetName().ToString();
+
+			const FString Section = FString::Printf(TEXT("Insights.%s.TimingView"), *SettingsName);
+
+			// Toggle Options
+			{
+				TimingViewSettings.EnumerateToggleOptions([this, &Section](const FTimingViewSettings::FToggleOption& Option)
+					{
+						GConfig->SetBool(*Section, Option.Name, Option.bValue, SettingsIni);
+					});
+			}
+
+			// CPU Thread Groups
+			{
+				const bool bShowAllCpuTracks = TimingViewSettings.IsToggleOptionEnabled(FTimingViewSettings::Option_ShowAllCpuTracks);
+
+				TArray<FString> VisibleCpuThreadGroups;
+				TArray<FString> HiddenCpuThreadGroups;
+				TimingViewSettings.EnumerateCpuThreadGroups([this, bShowAllCpuTracks, &VisibleCpuThreadGroups, &HiddenCpuThreadGroups](const FString& GroupName, bool bIsVisible)
+					{
+						if (bIsVisible != bShowAllCpuTracks)
+						{
+							if (bIsVisible)
+							{
+								VisibleCpuThreadGroups.Add(GroupName);
+							}
+							else
+							{
+								HiddenCpuThreadGroups.Add(GroupName);
+							}
+						}
+					});
+				GConfig->SetArray(*Section, TEXT("VisibleCpuThreadGroups"), VisibleCpuThreadGroups, SettingsIni);
+				GConfig->SetArray(*Section, TEXT("HiddenCpuThreadGroups"), HiddenCpuThreadGroups, SettingsIni);
+			}
 		}
-		GConfig->SetArray(*FString::Printf(TEXT("Insights.%s.TimingView.CpuThreadGroups"),*Profiler), TEXT("VisibleCpuThreadGroups"), VisibleCpuThreadGroups, SettingsIni);
-		GConfig->SetArray(*FString::Printf(TEXT("Insights.%s.TimingView.CpuThreadGroups"),*Profiler), TEXT("HiddenCpuThreadGroups"), HiddenCpuThreadGroups, SettingsIni);
 	}
-
 	//////////////////////////////////////////////////
 	// [Insights.TimingProfiler]
 
@@ -331,10 +355,10 @@ void FInsightsSettings::SaveToConfig()
 	GConfig->SetInt(TEXT("Insights.TimingProfiler.TimersView"), TEXT("Mode"), TimersViewMode, SettingsIni);
 
 	GConfig->SetInt(TEXT("Insights.TimingProfiler.TimersView"), TEXT("GroupingMode"), TimersViewGroupingMode, SettingsIni);
-	GConfig->SetInt(TEXT("Insights.TimingProfiler.TimersView"), TEXT("ShowGpuTimers"), bTimersViewShowGpuTimers, SettingsIni);
-	GConfig->SetInt(TEXT("Insights.TimingProfiler.TimersView"), TEXT("ShowVerseTimers"), bTimersViewShowVerseTimers, SettingsIni);
-	GConfig->SetInt(TEXT("Insights.TimingProfiler.TimersView"), TEXT("ShowCpuTimers"), bTimersViewShowCpuTimers, SettingsIni);
-	GConfig->SetInt(TEXT("Insights.TimingProfiler.TimersView"), TEXT("ShowZeroCountTimers"), bTimersViewShowZeroCountTimers, SettingsIni);
+	GConfig->SetBool(TEXT("Insights.TimingProfiler.TimersView"), TEXT("ShowGpuTimers"), bTimersViewShowGpuTimers, SettingsIni);
+	GConfig->SetBool(TEXT("Insights.TimingProfiler.TimersView"), TEXT("ShowVerseTimers"), bTimersViewShowVerseTimers, SettingsIni);
+	GConfig->SetBool(TEXT("Insights.TimingProfiler.TimersView"), TEXT("ShowCpuTimers"), bTimersViewShowCpuTimers, SettingsIni);
+	GConfig->SetBool(TEXT("Insights.TimingProfiler.TimersView"), TEXT("ShowZeroCountTimers"), bTimersViewShowZeroCountTimers, SettingsIni);
 
 	//////////////////////////////////////////////////
 	// [Insights.MemoryProfiler]
diff --git a/Engine/Source/Developer/TraceInsights/Private/Insights/InsightsSettings.h b/Engine/Source/Developer/TraceInsights/Private/Insights/InsightsSettings.h
--- a/Engine/Source/Developer/TraceInsights/Private/Insights/InsightsSettings.h
+++ b/Engine/Source/Developer/TraceInsights/Private/Insights/InsightsSettings.h
@@ -7,122 +7,204 @@
 #include "Containers/UnrealString.h"
 #include "UObject/NameTypes.h"
 
-/** Contains all settings for the Unreal Insights, accessible through the main manager. */
-class FInsightsSettings
+namespace UE::Insights
 {
-	friend class SInsightsSettings;
 
+class FTimingViewSettings
+{
 public:
-	struct FTimingViewTrackSettings
+	struct FToggleOption
+	{
+		const TCHAR* Name = nullptr;
+		bool bDefaultValue = false;
+		mutable bool bValue = false;
+	};
+
+	typedef TMap<const TCHAR*, FToggleOption, FDefaultSetAllocator, TStringPointerMapKeyFuncs_DEPRECATED<const TCHAR*, FToggleOption>> FToggleOptionsMap;
+
+	FTimingViewSettings()
 	{
-		bool FindKey(const TCHAR* Key, const TMap<FString*, bool>& Map) const
+		RegisterDefaultToggleOptions();
+	}
+
+	FTimingViewSettings(const FName& InName)
+		: TimingViewName(InName)
+	{
+		RegisterDefaultToggleOptions();
+	}
+
+	~FTimingViewSettings()
+	{
+	}
+
+	void Reset()
+	{
+		ResetToggleOptions();
+		ResetCpuThreadGroups();
+	}
+
+	const FName& GetName() const
+	{
+		return TimingViewName;
+	}
+
+	//////////////////////////////////////////////////
+	// Toggle Options
+
+	const FToggleOptionsMap& GetToggleOptions() const
+	{
+		return ToggleOptions;
+	}
+	FToggleOptionsMap& GetToggleOptions()
+	{
+		return ToggleOptions;
+	}
+
+	void EnumerateToggleOptions(TFunctionRef<void(const FToggleOption& Option)> Callback) const
+	{
+		for (const auto& KV : ToggleOptions)
 		{
-			FString KeyToCheck = FString(Key);
-			for (const auto& KV : Map)
-			{
-				if (KV.Key && KeyToCheck == *KV.Key)
-				{
-					return true;
-				}
-			}
-			return false;
+			Callback(KV.Value);
 		}
-
-		bool GetValueByKey(const TCHAR* Key, const TMap<FString*, bool>& Map) const
+	}
+	void EnumerateToggleOptions(TFunctionRef<void(FToggleOption& Option)> Callback)
+	{
+		for (auto& KV : ToggleOptions)
 		{
-			FString KeyToCheck = FString(Key);
-			for (const auto& KV : Map)
-			{
-				if (KV.Key && KeyToCheck == *KV.Key)
-				{
-					return KV.Value;
-				}
-			}
-			return false;
+			Callback(KV.Value);
 		}
+	}
+
+	bool IsToggleOptionRegistered(const TCHAR* Key) const
+	{
+		return ToggleOptions.Contains(Key);
+	}
+	void RegisterToggleOptions(const TCHAR* Name, bool bDefaultValue = false)
+	{
+		ToggleOptions.Add(Name, FToggleOption{ Name, bDefaultValue, bDefaultValue });
+	}
 
-		void RemoveKey(const TCHAR* Key, TMap<FString*, bool>& Map)
+	bool GetDefaultToggleOptionValue(const TCHAR* Key) const
+	{
+		const FToggleOption* Found = ToggleOptions.Find(Key);
+		return Found && Found->bDefaultValue;
+	}
+	bool IsToggleOptionEnabled(const TCHAR* Key) const
+	{
+		const FToggleOption* Found = ToggleOptions.Find(Key);
+		return Found && Found->bValue;
+	}
+	void SetToggleOption(const TCHAR* Key, bool bValue)
+	{
+		const FToggleOption* Found = ToggleOptions.Find(Key);
+		if (Found)
 		{
-			FString KeyToCheck = FString(Key);
-			for (auto It = Map.CreateIterator(); It; ++It)
-			{
-				FString* Ptr = It.Key();
-				if (Ptr && KeyToCheck == *Ptr)
-				{
-					delete Ptr;
-					It.RemoveCurrent();
-					return;
-				}
-			}
+			Found->bValue = bValue;
 		}
-
-		void RemoveAllKeys(TMap<FString*, bool>& Map)
+	}
+	void ResetToggleOptions()
+	{
+		for (auto& KV : ToggleOptions)
 		{
-			for (auto It = Map.CreateIterator(); It; ++It)
-			{
-				FString* Ptr = It.Key();
-				if (Ptr)
-				{
-					delete Ptr;
-					It.RemoveCurrent();
-				}
-			}
-			Map.Reset();
+			KV.Value.bValue = KV.Value.bDefaultValue;
 		}
+	}
+
+	//////////////////////////////////////////////////
+	// CPU Thread Groups
 
-		////////////////////////////////////////////////////////////////////////////////////////////////////
-
-		/** Toggle all or no Gpu-related Tracks in the Timing View. */
-		static constexpr const TCHAR* Option_ShowAllGpuTracks = TEXT("ShowAllGpuTracks");
-		/** Toggle all or no Cpu Tracks in the Timing View. */
-		static constexpr const TCHAR* Option_ShowAllCpuTracks = TEXT("ShowAllCpuTracks");
-		/** Toggle all or no Verse Tracks in the Timing View. */
-		static constexpr const TCHAR* Option_ShowAllVerseTracks = TEXT("ShowAllVerseTracks");
-		/** GPU work header tracks. */
-		static constexpr const TCHAR* Option_ShowGpuWorkTracks = TEXT("ShowGpuWorkTracks");
-		/** Extends the visualization of GPU work events over the GPU timing tracks. */
-		static constexpr const TCHAR* Option_ShowGpuWorkOverlays = TEXT("ShowGpuWorkOverlays");
-		/** Extended vertical lines at the edges of each GPU work event. */
-		static constexpr const TCHAR* Option_ShowGpuWorkExtendedLines = TEXT("ShowGpuWorkExtendedLines");
-		/** GPU fences child track. */
-		static constexpr const TCHAR* Option_ShowGpuFencesTracks = TEXT("ShowGpuFencesTracks");
-		/** Relations between Signal and Wait fences will be displayed when selecting a Timing Event in a Gpu Queue Track. */
-		static constexpr const TCHAR* Option_ShowGpuFencesRelations = TEXT("ShowGpuFencesRelations");
-		/** Extended vertical lines at the location of gpu fences. */
-		static constexpr const TCHAR* Option_ShowGpuFencesExtendedLines = TEXT("ShowGpuFencesExtendedLines");
-
-		////////////////////////////////////////////////////////////////////////////////////////////////////
-
-		TMap<FString*, bool>& GetToggleOptions() { return ToggleOptions; }
-		bool IsToggleOptionEnabled(const TCHAR* Key) const { return GetValueByKey(Key, ToggleOptions); }
-		bool HasToggleOption(const TCHAR* Key) const { return FindKey(Key, ToggleOptions); }
-		void AddToggleOption(const TCHAR* Key, bool InValue) { RemoveKey(Key, ToggleOptions); ToggleOptions.Add(new FString(Key), InValue); }
-		
-		TMap<FString*, bool>& GetCpuThreadGroups() { return ShowCpuThreadGroups; }
-		bool IsCpuThreadGroupVisible(const TCHAR* Key) const { return GetValueByKey(Key, ShowCpuThreadGroups);}
-		bool HasCpuThreadGroupVisibility(const TCHAR* Key) const { return FindKey(Key, ShowCpuThreadGroups); }
-		void AddCpuThreadGroupVisibility(const TCHAR* Key, bool InValue) { RemoveKey(Key, ShowCpuThreadGroups); ShowCpuThreadGroups.Add(new FString(Key), InValue); }
-		void RemoveCpuThreadGroup(const TCHAR* Key) { RemoveKey(Key, ShowCpuThreadGroups); }
-
-		void ResetAll() { RemoveAllKeys(ToggleOptions); RemoveAllKeys(ShowCpuThreadGroups); }
-
-		FTimingViewTrackSettings() = default;
-
-		~FTimingViewTrackSettings()
+	void EnumerateCpuThreadGroups(TFunctionRef<void(const FString& Name, bool bIsVisible)> Callback) const
+	{
+		for (const auto& KV : CpuThreadGroups)
 		{
-			ResetAll();
+			Callback(KV.Key, KV.Value);
 		}
+	}
 
-	private:
-		/** Holds the members which control parent true/false states for options. */
-		/** Key - The Option_ for true/false toggle commands. */
-		/** Value - The enabled state. */
-		TMap<FString*, bool> ToggleOptions;
-		/** Shows/hides each Cpu Thread Group track. */
-		/** Key - The Name of the Cpu thread group. */
-		/** Value - The visibility state. */
-		TMap<FString*, bool> ShowCpuThreadGroups;
-	};
+	bool IsCpuThreadGroupVisible(const TCHAR* GroupName) const
+	{
+		const bool* Found = CpuThreadGroups.Find(FString(GroupName));
+		return Found && *Found;
+	}
+	bool HasCpuThreadGroup(const TCHAR* GroupName) const
+	{
+		return CpuThreadGroups.Contains(FString(GroupName));
+	}
+	void AddCpuThreadGroup(const TCHAR* GroupName, bool bIsVisible)
+	{
+		CpuThreadGroups.Add(FString(GroupName), bIsVisible);
+	}
+	void RemoveCpuThreadGroup(const TCHAR* GroupName)
+	{
+		CpuThreadGroups.Remove(FString(GroupName));
+	}
+	void ResetCpuThreadGroups()
+	{
+		CpuThreadGroups.Reset();
+	}
+
+	//////////////////////////////////////////////////
+
+private:
+	void RegisterDefaultToggleOptions()
+	{
+		RegisterToggleOptions(Option_ShowAllGpuTracks, true);
+		RegisterToggleOptions(Option_ShowGpuWorkTracks, true);
+		RegisterToggleOptions(Option_ShowGpuWorkOverlays, true);
+		RegisterToggleOptions(Option_ShowGpuWorkExtendedLines, true);
+		RegisterToggleOptions(Option_ShowGpuFencesRelations, true);
+		RegisterToggleOptions(Option_ShowGpuFencesTracks, true);
+		RegisterToggleOptions(Option_ShowGpuFencesExtendedLines, true);
+
+		RegisterToggleOptions(Option_ShowAllVerseTracks, true);
+
+		RegisterToggleOptions(Option_ShowAllCpuTracks, true);
+	}
+
+public:
+	/** Toggle all or no GPU-related Tracks in the Timing View. */
+	static constexpr const TCHAR* Option_ShowAllGpuTracks = TEXT("ShowAllGpuTracks");
+	/** Visibility of the GPU Work header tracks. */
+	static constexpr const TCHAR* Option_ShowGpuWorkTracks = TEXT("ShowGpuWorkTracks");
+	/** Extends the visualization of GPU work events over the GPU timing tracks. */
+	static constexpr const TCHAR* Option_ShowGpuWorkOverlays = TEXT("ShowGpuWorkOverlays");
+	/** Extended vertical lines at the edges of each GPU work event. */
+	static constexpr const TCHAR* Option_ShowGpuWorkExtendedLines = TEXT("ShowGpuWorkExtendedLines");
+	/** Relations between Signal and Wait fences will be displayed when selecting a Timing Event in a GPU Queue Track. */
+	static constexpr const TCHAR* Option_ShowGpuFencesRelations = TEXT("ShowGpuFencesRelations");
+	/** Visibility of the GPU Fences child tracks. */
+	static constexpr const TCHAR* Option_ShowGpuFencesTracks = TEXT("ShowGpuFencesTracks");
+	/** Extended vertical lines at the location of GPU fences. */
+	static constexpr const TCHAR* Option_ShowGpuFencesExtendedLines = TEXT("ShowGpuFencesExtendedLines");
+
+	/** Toggle all or no Verse Tracks in the Timing View. */
+	static constexpr const TCHAR* Option_ShowAllVerseTracks = TEXT("ShowAllVerseTracks");
+
+	/** Toggle all or no CPU Tracks in the Timing View. */
+	static constexpr const TCHAR* Option_ShowAllCpuTracks = TEXT("ShowAllCpuTracks");
+
+private:
+	// The unique name of the Timing View settings.
+	// When settings are saved it will use a section named [Insights.(TimingViewName).TimingView].
+	FName TimingViewName;
+
+	/** Enabled state for the toggle options. */
+	/** Key - The Option_ for true/false toggle commands. */
+	/** Value - The enabled state. */
+	FToggleOptionsMap ToggleOptions;
+
+	/** Visibility state for the CPU Thread Groups. */
+	/** GroupName - The name of the CPU Thread Group. */
+	/** bIsVisible - The visibility state. */
+	TMap<FString, bool> CpuThreadGroups;
+};
+
+} // namespace UE::Insights
+
+/** Contains all settings for the Unreal Insights, accessible through the main manager. */
+class FInsightsSettings
+{
+	friend class SInsightsSettings;
 
 public:
 	FInsightsSettings(bool bInIsDefault = false);
@@ -156,25 +238,75 @@ public:
 	#define SET_AND_SAVE(Option, Value) { if (Option != Value) { Option = Value; SaveToConfig(); } }
 
 	//////////////////////////////////////////////////
-	// [Insights.(AllProfilers).TimingView]
-
-	FTimingViewTrackSettings& GetProfilerTimingViewTrackSettings(const FName& ProfilerId) {if (ProfilerTimingViewTrackSettings.Contains(ProfilerId)) { return *ProfilerTimingViewTrackSettings.Find(ProfilerId); } return TimingViewTrackSettingsDefaultObject; }
-	const FTimingViewTrackSettings& GetProfilerTimingViewTrackSettings(const FName& ProfilerId) const { if (ProfilerTimingViewTrackSettings.Contains(ProfilerId)) { return *ProfilerTimingViewTrackSettings.Find(ProfilerId); } return TimingViewTrackSettingsDefaultObject; }
-
-	void CreateIfNotExistProfilerTimingViewTrackSettings(const FName& ProfilerId) { if (!ProfilerTimingViewTrackSettings.Find(ProfilerId)) { ProfilerTimingViewTrackSettings.Add(ProfilerId); } }
-	void ResetAllProfilerTimingViewTrackSettings(const FName& ProfilerId) { GetProfilerTimingViewTrackSettings(ProfilerId).ResetAll(); SaveToConfig(); }
-
-	bool IsTimingViewToggleOptionEnabled(const TCHAR* Key, const FName& ProfilerId) const { return GetProfilerTimingViewTrackSettings(ProfilerId).IsToggleOptionEnabled(Key); }
-	bool HasTimingViewToggleOption(const TCHAR* Key, const FName& ProfilerId) const { return GetProfilerTimingViewTrackSettings(ProfilerId).HasToggleOption(Key); }
-	void AddTimingViewToggleOption(const TCHAR* Key, bool InValue, const FName& ProfilerId) { GetProfilerTimingViewTrackSettings(ProfilerId).AddToggleOption(Key, InValue); }
-	void AddAndSaveTimingViewToggleOption(const TCHAR* Key, bool InValue, const FName& ProfilerId) { AddTimingViewToggleOption(Key, InValue, ProfilerId); SaveToConfig(); }
-	
-	bool IsTimingViewCpuThreadGroupVisible(const TCHAR* Key, const FName& ProfilerId) const { return GetProfilerTimingViewTrackSettings(ProfilerId).IsCpuThreadGroupVisible(Key); }
-	bool HasTimingViewCpuThreadGroupVisibility(const TCHAR* Key, const FName& ProfilerId) const { return GetProfilerTimingViewTrackSettings(ProfilerId).HasCpuThreadGroupVisibility(Key); }
-	void AddTimingViewCpuThreadGroupVisibility(const TCHAR* Key, bool InValue, const FName& ProfilerId) { GetProfilerTimingViewTrackSettings(ProfilerId).AddCpuThreadGroupVisibility(Key, InValue); }
-	void AddAndSaveTimingViewCpuThreadGroupVisibility(const TCHAR* Key, bool InValue, const FName& ProfilerId) { AddTimingViewCpuThreadGroupVisibility(Key, InValue, ProfilerId); SaveToConfig(); }
-	void RemoveAndSaveTimingViewCpuThreadGroup(const TCHAR* Key, const FName& ProfilerId) { GetProfilerTimingViewTrackSettings(ProfilerId).RemoveCpuThreadGroup(Key); SaveToConfig(); }
-	
+	// [Insights.(TimingViewName).TimingView]
+
+	UE::Insights::FTimingViewSettings& GetOrCreateTimingViewSettings(const FName& TimingViewName)
+	{
+		UE::Insights::FTimingViewSettings** Found = AllTimingViewSettings.Find(TimingViewName);
+		if (Found)
+		{
+			return **Found;
+		}
+		UE::Insights::FTimingViewSettings* NewTimingViewSettings = new UE::Insights::FTimingViewSettings(TimingViewName);
+		AllTimingViewSettings.Add(TimingViewName, NewTimingViewSettings);
+		return *NewTimingViewSettings;
+	}
+
+	UE::Insights::FTimingViewSettings& GetTimingViewSettings(const FName& TimingViewName)
+	{
+		UE::Insights::FTimingViewSettings** Found = AllTimingViewSettings.Find(TimingViewName);
+		if (Found)
+		{
+			return **Found;
+		}
+		return DefaultTimingViewSettings;
+	}
+	const UE::Insights::FTimingViewSettings& GetTimingViewSettings(const FName& TimingViewName) const
+	{
+		UE::Insights::FTimingViewSettings*const* Found = AllTimingViewSettings.Find(TimingViewName);
+		if (Found)
+		{
+			return **Found;
+		}
+		return DefaultTimingViewSettings;
+	}
+
+	void ResetTimingViewSettings(const FName& TimingViewName)
+	{
+		GetTimingViewSettings(TimingViewName).Reset();
+		SaveToConfig();
+	}
+
+	void SetAndSaveTimingViewToggleOption(const FName& TimingViewName, const TCHAR* Key, bool bNewValue)
+	{
+		UE::Insights::FTimingViewSettings& TimingViewSettings = GetTimingViewSettings(TimingViewName);
+		bool bOldValue = TimingViewSettings.IsToggleOptionEnabled(Key);
+		if (bNewValue != bOldValue)
+		{
+			TimingViewSettings.SetToggleOption(Key, bNewValue);
+			SaveToConfig();
+		}
+	}
+
+	void AddAndSaveTimingViewCpuThreadGroup(const FName& TimingViewName, const TCHAR* GroupName, bool bIsVisible)
+	{
+		UE::Insights::FTimingViewSettings& TimingViewSettings = GetTimingViewSettings(TimingViewName);
+		TimingViewSettings.AddCpuThreadGroup(GroupName, bIsVisible);
+		SaveToConfig();
+	}
+	void RemoveAndSaveTimingViewCpuThreadGroup(const FName& TimingViewName, const TCHAR* GroupName)
+	{
+		UE::Insights::FTimingViewSettings& TimingViewSettings = GetTimingViewSettings(TimingViewName);
+		TimingViewSettings.RemoveCpuThreadGroup(GroupName);
+		SaveToConfig();
+	}
+	void ResetAndSaveTimingViewCpuThreadGroups(const FName& TimingViewName)
+	{
+		UE::Insights::FTimingViewSettings& TimingViewSettings = GetTimingViewSettings(TimingViewName);
+		TimingViewSettings.ResetCpuThreadGroups();
+		SaveToConfig();
+	}
+
 	//////////////////////////////////////////////////
 	// [Insights.TimingProfiler]
 
@@ -369,13 +501,13 @@ private:
 	FString SettingsIni;
 
 	//////////////////////////////////////////////////
-	// [Insights.(AllProfilers).TimingView]
+	// [Insights.(TimingViewName).TimingView]
 
-	/** Per profiler timing view track settings. */
-	TMap<FName, FTimingViewTrackSettings> ProfilerTimingViewTrackSettings;
+	/** Settings for each registered Timing View instance. */
+	TMap<FName, UE::Insights::FTimingViewSettings*> AllTimingViewSettings;
 
-	/** Per profiler default object. */
-	inline static FTimingViewTrackSettings TimingViewTrackSettingsDefaultObject = {};
+	/** The default settings for a Timing View. */
+	inline static UE::Insights::FTimingViewSettings DefaultTimingViewSettings = { };
 
 	//////////////////////////////////////////////////
 	// [Insights.TimingProfiler]
diff --git a/Engine/Source/Developer/TraceInsights/Private/Insights/TimingProfiler/ViewModels/ThreadTimingSharedState.cpp b/Engine/Source/Developer/TraceInsights/Private/Insights/TimingProfiler/ViewModels/ThreadTimingSharedState.cpp
--- a/Engine/Source/Developer/TraceInsights/Private/Insights/TimingProfiler/ViewModels/ThreadTimingSharedState.cpp
+++ b/Engine/Source/Developer/TraceInsights/Private/Insights/TimingProfiler/ViewModels/ThreadTimingSharedState.cpp
@@ -45,7 +45,7 @@ public:
 public:
 	/** Restores all settings to default state. */
 	TSharedPtr<FUICommandInfo> ResetToDefaultValues;
-	
+
 	/** Toggles visibility for GPU thread track(s). */
 	TSharedPtr<FUICommandInfo> ShowHideAllGpuTracks;
 
@@ -92,11 +92,11 @@ UE_DISABLE_OPTIMIZATION_SHIP
 void FThreadTimingViewCommands::RegisterCommands()
 {
 	UI_COMMAND(ResetToDefaultValues,
-		"Reset To Defaults",
-		"Resets all tracks back to their default visibility.",
+		"Reset To Default",
+		"Resets all track settings back to their default state.",
 		EUserInterfaceActionType::Button,
 		FInputChord());
-	
+
 	UI_COMMAND(ShowHideAllGpuTracks,
 		"GPU Track(s)",
 		"Shows/hides the GPU track(s).",
@@ -159,7 +159,7 @@ UE_ENABLE_OPTIMIZATION_SHIP
 
 FThreadTimingSharedState::FThreadTimingSharedState(STimingView* InTimingView)
 	: TimingView(InTimingView)
-	, Settings(MakeShared<FThreadSharedStateLocalSettings>())
+	, TimingViewSettings(MakeShared<FTimingViewLocalSettings>())
 {
 	check(TimingView != nullptr);
 }
@@ -305,8 +305,8 @@ void FThreadTimingSharedState::ResetToDefaultValues()
 {
 	RemoveGpuWorkChildTracks();
 	RemoveGpuFencesChildTracks();
-	Settings->Reset();
-	
+	TimingViewSettings->Reset();
+
 	if (TimingView &&
 		TimingView->GetName() == FInsightsManagerTabs::TimingProfilerTabId)
 	{
@@ -321,11 +321,12 @@ void FThreadTimingSharedState::ResetToDefaultValues()
 		HideAllCpuTracks();
 	}
 
-	Settings->AddTimingViewToggleOption(FInsightsSettings::FTimingViewTrackSettings::Option_ShowGpuWorkOverlays, true);
-	Settings->AddTimingViewToggleOption(FInsightsSettings::FTimingViewTrackSettings::Option_ShowGpuWorkExtendedLines, true);
+	TimingViewSettings->SetToggleOption(FTimingViewSettings::Option_ShowGpuWorkOverlays, true);
+	TimingViewSettings->SetToggleOption(FTimingViewSettings::Option_ShowGpuWorkExtendedLines, true);
 	SetGpuWorkTracksVisibility(true);
-	Settings->AddTimingViewToggleOption(FInsightsSettings::FTimingViewTrackSettings::Option_ShowGpuFencesRelations, true);
-	Settings->AddTimingViewToggleOption(FInsightsSettings::FTimingViewTrackSettings::Option_ShowGpuFencesExtendedLines, true);
+
+	TimingViewSettings->SetToggleOption(FTimingViewSettings::Option_ShowGpuFencesRelations, true);
+	TimingViewSettings->SetToggleOption(FTimingViewSettings::Option_ShowGpuFencesExtendedLines, true);
 	SetGpuFencesTracksVisibility(true);
 }
 
@@ -339,42 +340,19 @@ void FThreadTimingSharedState::OnBeginSession(Timing::ITimingViewSession& InSess
 	}
 
 	if (TimingView &&
-		TimingView->GetName() == FInsightsManagerTabs::TimingProfilerTabId)
+		(TimingView->GetName() == FInsightsManagerTabs::TimingProfilerTabId ||
+		 TimingView->GetName() == FInsightsManagerTabs::LoadingProfilerTabId ||
+		 TimingView->GetName() == FInsightsManagerTabs::MemoryProfilerTabId))
 	{
-		Settings = MakeShared<FThreadSharedStatePersistentSettings>(FInsightsManagerTabs::TimingProfilerTabId);
-
-		if (!Settings->HasTimingViewToggleOption(FInsightsSettings::FTimingViewTrackSettings::Option_ShowAllCpuTracks)) { Settings->AddTimingViewToggleOption(FInsightsSettings::FTimingViewTrackSettings::Option_ShowAllCpuTracks, true); }
-		if (!Settings->HasTimingViewToggleOption(FInsightsSettings::FTimingViewTrackSettings::Option_ShowAllGpuTracks)) { Settings->AddTimingViewToggleOption(FInsightsSettings::FTimingViewTrackSettings::Option_ShowAllGpuTracks, true); }
-		if (!Settings->HasTimingViewToggleOption(FInsightsSettings::FTimingViewTrackSettings::Option_ShowAllVerseTracks)) { Settings->AddTimingViewToggleOption(FInsightsSettings::FTimingViewTrackSettings::Option_ShowAllVerseTracks, true); }
-		
-		if (!Settings->HasTimingViewToggleOption(FInsightsSettings::FTimingViewTrackSettings::Option_ShowGpuWorkTracks)) { Settings->AddTimingViewToggleOption(FInsightsSettings::FTimingViewTrackSettings::Option_ShowGpuWorkTracks, true); }
-		if (!Settings->HasTimingViewToggleOption(FInsightsSettings::FTimingViewTrackSettings::Option_ShowGpuWorkOverlays)) { Settings->AddTimingViewToggleOption(FInsightsSettings::FTimingViewTrackSettings::Option_ShowGpuWorkOverlays, true); }
-		if (!Settings->HasTimingViewToggleOption(FInsightsSettings::FTimingViewTrackSettings::Option_ShowGpuWorkExtendedLines)) { Settings->AddTimingViewToggleOption(FInsightsSettings::FTimingViewTrackSettings::Option_ShowGpuWorkExtendedLines, true); }
-		if (!Settings->HasTimingViewToggleOption(FInsightsSettings::FTimingViewTrackSettings::Option_ShowGpuFencesTracks)) { Settings->AddTimingViewToggleOption(FInsightsSettings::FTimingViewTrackSettings::Option_ShowGpuFencesTracks, true); }
-		if (!Settings->HasTimingViewToggleOption(FInsightsSettings::FTimingViewTrackSettings::Option_ShowGpuFencesExtendedLines)) { Settings->AddTimingViewToggleOption(FInsightsSettings::FTimingViewTrackSettings::Option_ShowGpuFencesExtendedLines, true); }
-		if (!Settings->HasTimingViewToggleOption(FInsightsSettings::FTimingViewTrackSettings::Option_ShowGpuFencesRelations)) { Settings->AddTimingViewToggleOption(FInsightsSettings::FTimingViewTrackSettings::Option_ShowGpuFencesRelations, true); }
-	}
-	else if (TimingView &&
-			 TimingView->GetName() == FInsightsManagerTabs::MemoryProfilerTabId)
-	{
-		Settings = MakeShared<FThreadSharedStatePersistentSettings>(FInsightsManagerTabs::MemoryProfilerTabId);
-		
-		if (!Settings->HasTimingViewToggleOption(FInsightsSettings::FTimingViewTrackSettings::Option_ShowAllCpuTracks)) { Settings->AddTimingViewToggleOption(FInsightsSettings::FTimingViewTrackSettings::Option_ShowAllCpuTracks, false); }
-		if (!Settings->HasTimingViewToggleOption(FInsightsSettings::FTimingViewTrackSettings::Option_ShowAllGpuTracks)) { Settings->AddTimingViewToggleOption(FInsightsSettings::FTimingViewTrackSettings::Option_ShowAllGpuTracks, false); }
-		if (!Settings->HasTimingViewToggleOption(FInsightsSettings::FTimingViewTrackSettings::Option_ShowAllVerseTracks)) { Settings->AddTimingViewToggleOption(FInsightsSettings::FTimingViewTrackSettings::Option_ShowAllVerseTracks, false); }
-
-		if (!Settings->HasTimingViewToggleOption(FInsightsSettings::FTimingViewTrackSettings::Option_ShowGpuWorkTracks)) { Settings->AddTimingViewToggleOption(FInsightsSettings::FTimingViewTrackSettings::Option_ShowGpuWorkTracks, true); }
-		if (!Settings->HasTimingViewToggleOption(FInsightsSettings::FTimingViewTrackSettings::Option_ShowGpuWorkOverlays)) { Settings->AddTimingViewToggleOption(FInsightsSettings::FTimingViewTrackSettings::Option_ShowGpuWorkOverlays, true); }
-		if (!Settings->HasTimingViewToggleOption(FInsightsSettings::FTimingViewTrackSettings::Option_ShowGpuWorkExtendedLines)) { Settings->AddTimingViewToggleOption(FInsightsSettings::FTimingViewTrackSettings::Option_ShowGpuWorkExtendedLines, true); }
-		if (!Settings->HasTimingViewToggleOption(FInsightsSettings::FTimingViewTrackSettings::Option_ShowGpuFencesTracks)) { Settings->AddTimingViewToggleOption(FInsightsSettings::FTimingViewTrackSettings::Option_ShowGpuFencesTracks, true); }
-		if (!Settings->HasTimingViewToggleOption(FInsightsSettings::FTimingViewTrackSettings::Option_ShowGpuFencesExtendedLines)) { Settings->AddTimingViewToggleOption(FInsightsSettings::FTimingViewTrackSettings::Option_ShowGpuFencesExtendedLines, true); }
-		if (!Settings->HasTimingViewToggleOption(FInsightsSettings::FTimingViewTrackSettings::Option_ShowGpuFencesRelations)) { Settings->AddTimingViewToggleOption(FInsightsSettings::FTimingViewTrackSettings::Option_ShowGpuFencesRelations, true); }
+		TimingViewSettings = MakeShared<FTimingViewPersistentSettings>(TimingView->GetName());
 	}
 	else
 	{
-		Settings->AddTimingViewToggleOption(FInsightsSettings::FTimingViewTrackSettings::Option_ShowAllCpuTracks, false);
-		Settings->AddTimingViewToggleOption(FInsightsSettings::FTimingViewTrackSettings::Option_ShowAllGpuTracks, false);
-		Settings->AddTimingViewToggleOption(FInsightsSettings::FTimingViewTrackSettings::Option_ShowAllVerseTracks, false);
+		TimingViewSettings = MakeShared<FTimingViewLocalSettings>();
+
+		TimingViewSettings->SetToggleOption(FTimingViewSettings::Option_ShowAllCpuTracks, false);
+		TimingViewSettings->SetToggleOption(FTimingViewSettings::Option_ShowAllGpuTracks, false);
+		TimingViewSettings->SetToggleOption(FTimingViewSettings::Option_ShowAllVerseTracks, false);
 	}
 
 	OldGpu1Track = nullptr;
@@ -443,6 +421,8 @@ void FThreadTimingSharedState::Tick(Timing::ITimingViewSession& InSession, const
 
 			LLM_SCOPE_BYTAG(Insights);
 
+			const bool bShowAllGpuTracks = TimingViewSettings->IsToggleOptionEnabled(FTimingViewSettings::Option_ShowAllGpuTracks);
+
 			// Check if we have the old GPU timelines.
 			if (!OldGpu1Track.IsValid())
 			{
@@ -451,7 +431,7 @@ void FThreadTimingSharedState::Tick(Timing::ITimingViewSession& InSession, const
 				{
 					OldGpu1Track = MakeShared<FGpuTimingTrack>(*this, TEXT("GPU"), nullptr, GpuTimelineIndex, FGpuTimingTrack::Gpu1ThreadId);
 					OldGpu1Track->SetOrder(FTimingTrackOrder::Gpu);
-					OldGpu1Track->SetVisibilityFlag(Settings->IsTimingViewToggleOptionEnabled(FInsightsSettings::FTimingViewTrackSettings::Option_ShowAllGpuTracks));
+					OldGpu1Track->SetVisibilityFlag(bShowAllGpuTracks);
 					InSession.AddScrollableTrack(OldGpu1Track);
 				}
 			}
@@ -462,7 +442,7 @@ void FThreadTimingSharedState::Tick(Timing::ITimingViewSession& InSession, const
 				{
 					OldGpu2Track = MakeShared<FGpuTimingTrack>(*this, TEXT("GPU2"), nullptr, GpuTimelineIndex, FGpuTimingTrack::Gpu2ThreadId);
 					OldGpu2Track->SetOrder(FTimingTrackOrder::Gpu + 1);
-					OldGpu2Track->SetVisibilityFlag(Settings->IsTimingViewToggleOptionEnabled(FInsightsSettings::FTimingViewTrackSettings::Option_ShowAllGpuTracks));
+					OldGpu2Track->SetVisibilityFlag(bShowAllGpuTracks);
 					InSession.AddScrollableTrack(OldGpu2Track);
 				}
 			}
@@ -472,7 +452,7 @@ void FThreadTimingSharedState::Tick(Timing::ITimingViewSession& InSession, const
 			int32 CpuTrackOrder = FTimingTrackOrder::Cpu;
 
 			// Iterate through GPU queues.
-			TimingProfilerProvider->EnumerateGpuQueues([this, &InSession, &bTracksOrderChanged, &GpuTrackOrder](const TraceServices::FGpuQueueInfo& QueueInfo)
+			TimingProfilerProvider->EnumerateGpuQueues([this, bShowAllGpuTracks , &InSession, &bTracksOrderChanged, &GpuTrackOrder](const TraceServices::FGpuQueueInfo& QueueInfo)
 			{
 				// Check if there is an available GPU track for this queue.
 				TSharedPtr<FGpuQueueTimingTrack>* TrackPtrPtr = GpuTracks.Find(QueueInfo.Id);
@@ -481,7 +461,7 @@ void FThreadTimingSharedState::Tick(Timing::ITimingViewSession& InSession, const
 					// Create new Timing Events track for the GPU Queue.
 					TSharedPtr<FGpuQueueTimingTrack> Track = MakeShared<FGpuQueueTimingTrack>(*this, QueueInfo.GetDisplayName(), QueueInfo.TimelineIndex, QueueInfo.Id);
 					Track->SetOrder(GpuTrackOrder);
-					Track->SetVisibilityFlag(Settings->IsTimingViewToggleOptionEnabled(FInsightsSettings::FTimingViewTrackSettings::Option_ShowAllGpuTracks));
+					Track->SetVisibilityFlag(bShowAllGpuTracks);
 					GpuTracks.Add(QueueInfo.Id, Track);
 					InSession.AddScrollableTrack(Track);
 
@@ -525,7 +505,8 @@ void FThreadTimingSharedState::Tick(Timing::ITimingViewSession& InSession, const
 				{
 					VerseSamplingTrack = MakeShared<FVerseTimingTrack>(*this, TEXT("Verse Sampling"), VerseTimelineIndex);
 					VerseSamplingTrack->SetOrder(FTimingTrackOrder::Cpu - 100);
-					VerseSamplingTrack->SetVisibilityFlag(bShowHideAllVerseTracks);
+					const bool bShowAllVerseTracks = TimingViewSettings->IsToggleOptionEnabled(FTimingViewSettings::Option_ShowAllVerseTracks);
+					VerseSamplingTrack->SetVisibilityFlag(bShowAllVerseTracks);
 					InSession.AddScrollableTrack(VerseSamplingTrack);
 				}
 			}
@@ -536,7 +517,6 @@ void FThreadTimingSharedState::Tick(Timing::ITimingViewSession& InSession, const
 			ThreadProvider.EnumerateThreads([this, &InSession, &bTracksOrderChanged, &CpuTrackOrder, TimingProfilerProvider, LoadTimeProfilerProvider](const TraceServices::FThreadInfo& ThreadInfo)
 			{
 				// Check if this thread is part of a group?
-				bool bIsGroupVisible = Settings->IsTimingViewToggleOptionEnabled(FInsightsSettings::FTimingViewTrackSettings::Option_ShowAllCpuTracks);
 				const TCHAR* GroupName = ThreadInfo.GroupName;
 				if (!GroupName || *GroupName == 0)
 				{
@@ -557,10 +537,9 @@ void FThreadTimingSharedState::Tick(Timing::ITimingViewSession& InSession, const
 					ThreadGroup.Order = CpuTrackOrder;
 				}
 
-				if (Settings->HasTimingViewCpuThreadGroupVisibility(GroupName))
-				{
-					bIsGroupVisible = Settings->IsTimingViewCpuThreadGroupVisible(GroupName);
-				}
+				const bool bIsGroupVisible = TimingViewSettings->HasCpuThreadGroup(GroupName) ?
+					TimingViewSettings->IsCpuThreadGroupVisible(GroupName) :
+					TimingViewSettings->IsToggleOptionEnabled(FTimingViewSettings::Option_ShowAllCpuTracks);
 
 				// Check if there is an available Asset Loading track for this thread.
 				bool bIsLoadingThread = false;
@@ -697,7 +676,7 @@ void FThreadTimingSharedState::BindCommands()
 	CommandList->MapAction(
 		FThreadTimingViewCommands::Get().ResetToDefaultValues,
 		FExecuteAction::CreateSP(this, &FThreadTimingSharedState::ResetToDefaultValues));
-	
+
 	CommandList->MapAction(
 		FThreadTimingViewCommands::Get().ShowHideAllGpuTracks,
 		FExecuteAction::CreateSP(this, &FThreadTimingSharedState::ShowHideAllGpuTracks),
@@ -792,7 +771,7 @@ void FThreadTimingSharedState::CreateThreadGroupsMenu(FMenuBuilder& InOutMenuBui
 
 void FThreadTimingSharedState::SetAllVerseTracksToggle(bool bOnOff)
 {
-	Settings->AddTimingViewToggleOption(FInsightsSettings::FTimingViewTrackSettings::Option_ShowAllVerseTracks, bOnOff);
+	TimingViewSettings->SetToggleOption(FTimingViewSettings::Option_ShowAllVerseTracks, bOnOff);
 
 	if (VerseSamplingTrack.IsValid())
 	{
@@ -809,7 +788,7 @@ void FThreadTimingSharedState::SetAllVerseTracksToggle(bool bOnOff)
 
 void FThreadTimingSharedState::SetAllCpuTracksToggle(bool bOnOff)
 {
-	Settings->AddTimingViewToggleOption(FInsightsSettings::FTimingViewTrackSettings::Option_ShowAllCpuTracks, bOnOff);
+	TimingViewSettings->SetToggleOption(FTimingViewSettings::Option_ShowAllCpuTracks, bOnOff);
 
 	for (const auto& KV : CpuTracks)
 	{
@@ -817,16 +796,7 @@ void FThreadTimingSharedState::SetAllCpuTracksToggle(bool bOnOff)
 		Track.SetVisibilityFlag(bOnOff);
 	}
 
-	for (auto& KV : ThreadGroups)
-	{
-		if (Settings->HasTimingViewCpuThreadGroupVisibility(KV.Value.Name))
-		{
-			if (Settings->IsTimingViewCpuThreadGroupVisible(KV.Value.Name) == bOnOff)
-			{
-				Settings->RemoveTimingViewCpuThreadGroup(KV.Value.Name);
-			}
-		}
-	}
+	TimingViewSettings->ResetCpuThreadGroups();
 
 	if (TimingView)
 	{
@@ -838,8 +808,8 @@ void FThreadTimingSharedState::SetAllCpuTracksToggle(bool bOnOff)
 
 void FThreadTimingSharedState::SetAllGpuTracksToggle(bool bOnOff)
 {
-	Settings->AddTimingViewToggleOption(FInsightsSettings::FTimingViewTrackSettings::Option_ShowAllGpuTracks, bOnOff);
-	
+	TimingViewSettings->SetToggleOption(FTimingViewSettings::Option_ShowAllGpuTracks, bOnOff);
+
 	if (OldGpu1Track.IsValid())
 	{
 		OldGpu1Track->SetVisibilityFlag(bOnOff);
@@ -866,9 +836,9 @@ bool FThreadTimingSharedState::ToggleTrackVisibilityByGroup_IsChecked(const TCHA
 {
 	if (ThreadGroups.Contains(InGroupName))
 	{
-		return Settings->HasTimingViewCpuThreadGroupVisibility(InGroupName)
-		? Settings->IsTimingViewCpuThreadGroupVisible(InGroupName)
-		: Settings->IsTimingViewToggleOptionEnabled(FInsightsSettings::FTimingViewTrackSettings::Option_ShowAllCpuTracks);
+		return TimingViewSettings->HasCpuThreadGroup(InGroupName) ?
+				TimingViewSettings->IsCpuThreadGroupVisible(InGroupName) :
+				TimingViewSettings->IsToggleOptionEnabled(FTimingViewSettings::Option_ShowAllCpuTracks);
 	}
 	return false;
 }
@@ -879,20 +849,20 @@ void FThreadTimingSharedState::ToggleTrackVisibilityByGroup_Execute(const TCHAR*
 {
 	if (ThreadGroups.Contains(InGroupName))
 	{
-		bool bNewVisibility = !Settings->IsTimingViewToggleOptionEnabled(FInsightsSettings::FTimingViewTrackSettings::Option_ShowAllCpuTracks);
-		if (Settings->HasTimingViewCpuThreadGroupVisibility(InGroupName))
+		bool bShowAllCpuTracks = TimingViewSettings->IsToggleOptionEnabled(FTimingViewSettings::Option_ShowAllCpuTracks);
+		bool bNewVisibility = TimingViewSettings->HasCpuThreadGroup(InGroupName) ?
+			!TimingViewSettings->IsCpuThreadGroupVisible(InGroupName) :
+			!bShowAllCpuTracks;
+
+		if (bNewVisibility != bShowAllCpuTracks)
 		{
-			if (!Settings->IsTimingViewCpuThreadGroupVisible(InGroupName) == !bNewVisibility)
-			{
-				Settings->RemoveTimingViewCpuThreadGroup(InGroupName);
-				bNewVisibility = !bNewVisibility;
-			}
+			TimingViewSettings->AddCpuThreadGroup(InGroupName, bNewVisibility);
 		}
 		else
 		{
-			Settings->AddTimingViewCpuThreadGroupVisibility(InGroupName, bNewVisibility);
+			TimingViewSettings->RemoveCpuThreadGroup(InGroupName);
 		}
-		
+
 		for (const auto& KV : CpuTracks)
 		{
 			FCpuTimingTrack& Track = *KV.Value;
@@ -913,8 +883,8 @@ void FThreadTimingSharedState::ToggleTrackVisibilityByGroup_Execute(const TCHAR*
 
 void FThreadTimingSharedState::ToggleOptionState(const TCHAR* CommandName)
 {
-	bool bInvVisibility = !Settings->IsTimingViewToggleOptionEnabled(CommandName);
-	Settings->AddTimingViewToggleOption(CommandName, bInvVisibility);
+	const bool bIsEnabled = TimingViewSettings->IsToggleOptionEnabled(CommandName);
+	TimingViewSettings->SetToggleOption(CommandName, !bIsEnabled);
 }
 
 ////////////////////////////////////////////////////////////////////////////////////////////////////
@@ -1031,9 +1001,9 @@ TSharedPtr<const ITimingEvent> FThreadTimingSharedState::FindMinEventInstance(ui
 
 void FThreadTimingSharedState::SetGpuWorkTracksVisibility(bool bOnOff)
 {
-	if (Settings->IsTimingViewToggleOptionEnabled(FInsightsSettings::FTimingViewTrackSettings::Option_ShowGpuWorkTracks) != bOnOff)
+	if (TimingViewSettings->IsToggleOptionEnabled(FTimingViewSettings::Option_ShowGpuWorkTracks) != bOnOff)
 	{
-		Settings->AddTimingViewToggleOption(FInsightsSettings::FTimingViewTrackSettings::Option_ShowGpuWorkTracks, bOnOff);
+		TimingViewSettings->SetToggleOption(FTimingViewSettings::Option_ShowGpuWorkTracks, bOnOff);
 
 		if (bOnOff)
 		{
@@ -1141,9 +1111,9 @@ void FThreadTimingSharedState::RemoveGpuFencesChildTracks()
 
 void FThreadTimingSharedState::SetGpuFencesTracksVisibility(bool bOnOff)
 {
-	if (Settings->IsTimingViewToggleOptionEnabled(FInsightsSettings::FTimingViewTrackSettings::Option_ShowGpuFencesTracks) != bOnOff)
+	if (TimingViewSettings->IsToggleOptionEnabled(FTimingViewSettings::Option_ShowGpuFencesTracks) != bOnOff)
 	{
-		Settings->AddTimingViewToggleOption(FInsightsSettings::FTimingViewTrackSettings::Option_ShowGpuFencesTracks, bOnOff);
+		TimingViewSettings->SetToggleOption(FTimingViewSettings::Option_ShowGpuFencesTracks, bOnOff);
 
 		if (bOnOff)
 		{
@@ -1160,9 +1130,9 @@ void FThreadTimingSharedState::SetGpuFencesTracksVisibility(bool bOnOff)
 
 void FThreadTimingSharedState::Command_ShowGpuFencesRelations_Execute()
 {
-	ToggleOptionState(FInsightsSettings::FTimingViewTrackSettings::Option_ShowGpuFencesRelations);
-	bool bVisibility = Settings->IsTimingViewToggleOptionEnabled(FInsightsSettings::FTimingViewTrackSettings::Option_ShowGpuFencesRelations);
+	ToggleOptionState(FTimingViewSettings::Option_ShowGpuFencesRelations);
 
+	bool bVisibility = TimingViewSettings->IsToggleOptionEnabled(FTimingViewSettings::Option_ShowGpuFencesRelations);
 	if (bVisibility == false && TimingView)
 	{
 		TimingView->EditCurrentRelations().RemoveAll([](TUniquePtr<ITimingEventRelation>& Item)
@@ -1174,14 +1144,6 @@ void FThreadTimingSharedState::Command_ShowGpuFencesRelations_Execute()
 
 ////////////////////////////////////////////////////////////////////////////////////////////////////
 
-void FThreadTimingSharedState::Command_ShowGpuFencesTracks_Execute()
-{
-	bool bInvVisibility = !Settings->IsTimingViewToggleOptionEnabled(FInsightsSettings::FTimingViewTrackSettings::Option_ShowGpuFencesTracks);
-	SetGpuFencesTracksVisibility(bInvVisibility);
-}
-
-////////////////////////////////////////////////////////////////////////////////////////////////////
-
 void FThreadTimingSharedState::OnTimingEventSelected(TSharedPtr<const ITimingEvent> InSelectedEvent)
 {
 	if (!AreGpuFenceRelationsVisible())
@@ -1291,7 +1253,7 @@ void FThreadTimingSharedState::OnTimingEventSelected(TSharedPtr<const ITimingEve
 
 	TraceServices::FAnalysisSessionReadScope SessionReadScope(*Session);
 
-	TimingProfilerProvider->EnumerateResolvedGpuFences(QueueId, InSelectedEvent->GetStartTime(), InSelectedEvent->GetEndTime(), 
+	TimingProfilerProvider->EnumerateResolvedGpuFences(QueueId, InSelectedEvent->GetStartTime(), InSelectedEvent->GetEndTime(),
 		[AddFenceRelation, TimingProfilerProvider, QueueId](uint32 SignalFenceQueueId, const FGpuSignalFence& SignalFence, uint32 WaitFenceQueueId, const FGpuWaitFence& WaitFence)
 	{
 		AddFenceRelation(SignalFence, WaitFence, WaitFenceQueueId);
diff --git a/Engine/Source/Developer/TraceInsights/Private/Insights/TimingProfiler/ViewModels/ThreadTimingSharedState.h b/Engine/Source/Developer/TraceInsights/Private/Insights/TimingProfiler/ViewModels/ThreadTimingSharedState.h
--- a/Engine/Source/Developer/TraceInsights/Private/Insights/TimingProfiler/ViewModels/ThreadTimingSharedState.h
+++ b/Engine/Source/Developer/TraceInsights/Private/Insights/TimingProfiler/ViewModels/ThreadTimingSharedState.h
@@ -43,65 +43,112 @@ private:
 		int32 GetOrder() const { return Order; }
 	};
 
-	class IThreadSharedStateSetting
+	class ITimingViewSettings
 	{
 	public:
-		virtual ~IThreadSharedStateSetting() {}
-		
-		virtual bool IsTimingViewCpuThreadGroupVisible(const TCHAR* GroupName) const = 0;
-		virtual bool HasTimingViewCpuThreadGroupVisibility(const TCHAR* GroupName) const = 0;
-		virtual void AddTimingViewCpuThreadGroupVisibility(const TCHAR* GroupName, bool InValue) = 0;
-		virtual void RemoveTimingViewCpuThreadGroup(const TCHAR* GroupName) = 0;
+		virtual ~ITimingViewSettings() {}
 
-		virtual bool IsTimingViewToggleOptionEnabled(const TCHAR* Key) const = 0;
-		virtual bool HasTimingViewToggleOption(const TCHAR* Key) const = 0;
-		virtual void AddTimingViewToggleOption(const TCHAR* Key, bool InValue) = 0;
+		virtual bool IsToggleOptionEnabled(const TCHAR* Key) const = 0;
+		virtual void SetToggleOption(const TCHAR* Key, bool Value) = 0;
+
+		virtual bool IsCpuThreadGroupVisible(const TCHAR* GroupName) const = 0;
+		virtual bool HasCpuThreadGroup(const TCHAR* GroupName) const = 0;
+		virtual void AddCpuThreadGroup(const TCHAR* GroupName, bool bIsVisible) = 0;
+		virtual void RemoveCpuThreadGroup(const TCHAR* GroupName) = 0;
+		virtual void ResetCpuThreadGroups() = 0;
 
 		virtual void Reset() = 0;
 	};
 
-	class FThreadSharedStatePersistentSettings : public IThreadSharedStateSetting
+	class FTimingViewPersistentSettings : public ITimingViewSettings
 	{
 	public:
-		FThreadSharedStatePersistentSettings(const FName& InProfiler) : Profiler(InProfiler)
+		FTimingViewPersistentSettings(const FName& InTimingViewName)
+			: Settings(GetInsightsSettings().GetOrCreateTimingViewSettings(InTimingViewName))
 		{
-			GetInsightsSettings().CreateIfNotExistProfilerTimingViewTrackSettings(Profiler);
 		}
 
-		FName Profiler;
+		virtual ~FTimingViewPersistentSettings()
+		{
+		}
 
-		virtual bool IsTimingViewCpuThreadGroupVisible(const TCHAR* GroupName) const override { return GetInsightsSettings().IsTimingViewCpuThreadGroupVisible(GroupName, Profiler); }
-		virtual bool HasTimingViewCpuThreadGroupVisibility(const TCHAR* GroupName) const override { return GetInsightsSettings().HasTimingViewCpuThreadGroupVisibility(GroupName, Profiler); }
-		virtual void AddTimingViewCpuThreadGroupVisibility(const TCHAR* GroupName, bool InValue) override { return GetInsightsSettings().AddAndSaveTimingViewCpuThreadGroupVisibility(GroupName, InValue, Profiler); }
-		virtual void RemoveTimingViewCpuThreadGroup(const TCHAR* GroupName) override { return GetInsightsSettings().RemoveAndSaveTimingViewCpuThreadGroup(GroupName, Profiler); }
+		const FName& GetName() const
+		{
+			return Settings.GetName();
+		}
+
+		//////////////////////////////////////////////////
+
+		virtual bool IsToggleOptionEnabled(const TCHAR* Key) const override
+		{
+			return Settings.IsToggleOptionEnabled(Key);
+		}
+		virtual void SetToggleOption(const TCHAR* Key, bool Value) override
+		{
+			GetInsightsSettings().SetAndSaveTimingViewToggleOption(GetName(), Key, Value);
+		}
 
-		virtual bool IsTimingViewToggleOptionEnabled(const TCHAR* Key) const override { return GetInsightsSettings().IsTimingViewToggleOptionEnabled(Key, Profiler); }
-		virtual bool HasTimingViewToggleOption(const TCHAR* Key) const override { return GetInsightsSettings().HasTimingViewToggleOption(Key, Profiler); }
-		virtual void AddTimingViewToggleOption(const TCHAR* Key, bool InValue) override { return GetInsightsSettings().AddAndSaveTimingViewToggleOption(Key, InValue, Profiler); }
+		//////////////////////////////////////////////////
 
-		virtual void Reset() override { GetInsightsSettings().ResetAllProfilerTimingViewTrackSettings(Profiler); }
-		
-		FInsightsSettings& GetInsightsSettings() { return UE::Insights::FInsightsManager::Get()->GetSettings(); }
-		const FInsightsSettings& GetInsightsSettings() const { return UE::Insights::FInsightsManager::Get()->GetSettings(); }
+		virtual bool IsCpuThreadGroupVisible(const TCHAR* GroupName) const override
+		{
+			return Settings.IsCpuThreadGroupVisible(GroupName);
+		}
+		virtual bool HasCpuThreadGroup(const TCHAR* GroupName) const override
+		{
+			return Settings.HasCpuThreadGroup(GroupName);
+		}
+		virtual void AddCpuThreadGroup(const TCHAR* GroupName, bool Value) override
+		{
+			GetInsightsSettings().AddAndSaveTimingViewCpuThreadGroup(GetName(), GroupName, Value);
+		}
+		virtual void RemoveCpuThreadGroup(const TCHAR* GroupName) override
+		{
+			GetInsightsSettings().RemoveAndSaveTimingViewCpuThreadGroup(GetName(), GroupName);
+		}
+		virtual void ResetCpuThreadGroups() override
+		{
+			GetInsightsSettings().ResetAndSaveTimingViewCpuThreadGroups(GetName());
+		}
+
+		//////////////////////////////////////////////////
+
+		virtual void Reset() override
+		{
+			Settings.Reset();
+		}
+
+	private:
+		FInsightsSettings& GetInsightsSettings()
+		{
+			return UE::Insights::FInsightsManager::Get()->GetSettings();
+		}
+		const FInsightsSettings& GetInsightsSettings() const
+		{
+			return UE::Insights::FInsightsManager::Get()->GetSettings();
+		}
+
+	private:
+		FTimingViewSettings& Settings;
 	};
 
-	class FThreadSharedStateLocalSettings : public IThreadSharedStateSetting
+	class FTimingViewLocalSettings : public ITimingViewSettings
 	{
 	public:
-		virtual bool IsTimingViewCpuThreadGroupVisible(const TCHAR* GroupName) const override { return HasTimingViewCpuThreadGroupVisibility(GroupName) ? *TimingViewShowCpuThreadGroups.Find(GroupName) : false; }
-		virtual bool HasTimingViewCpuThreadGroupVisibility(const TCHAR* GroupName) const override { return TimingViewShowCpuThreadGroups.Contains(GroupName); }
-		virtual void AddTimingViewCpuThreadGroupVisibility(const TCHAR* GroupName, bool InValue) override { TimingViewShowCpuThreadGroups.Add(GroupName, InValue); }
-		virtual void RemoveTimingViewCpuThreadGroup(const TCHAR* GroupName) override { if (HasTimingViewCpuThreadGroupVisibility(GroupName)) { TimingViewShowCpuThreadGroups.Remove(GroupName); } }
-
-		virtual bool IsTimingViewToggleOptionEnabled(const TCHAR* Key) const override { return HasTimingViewToggleOption(Key) ? *TimingViewShowToggleOptions.Find(Key) : false; }
-		virtual bool HasTimingViewToggleOption(const TCHAR* Key) const override { return TimingViewShowToggleOptions.Contains(Key); }
-		virtual void AddTimingViewToggleOption(const TCHAR* Key, bool InValue) override { TimingViewShowToggleOptions.Add(Key, InValue); }
-
-		virtual void Reset() override { TimingViewShowToggleOptions.Reset(); TimingViewShowCpuThreadGroups.Reset(); }
+		FTimingViewLocalSettings() {}
+		virtual ~FTimingViewLocalSettings() {}
+
+		virtual bool IsToggleOptionEnabled(const TCHAR* Key) const override              { return Settings.IsToggleOptionEnabled(Key); }
+		virtual void SetToggleOption(const TCHAR* Key, bool Value) override              { Settings.SetToggleOption(Key, Value); }
+		virtual bool IsCpuThreadGroupVisible(const TCHAR* GroupName) const override      { return Settings.IsCpuThreadGroupVisible(GroupName); }
+		virtual bool HasCpuThreadGroup(const TCHAR* GroupName) const override            { return Settings.HasCpuThreadGroup(GroupName); }
+		virtual void AddCpuThreadGroup(const TCHAR* GroupName, bool bIsVisible) override { Settings.AddCpuThreadGroup(GroupName, bIsVisible); }
+		virtual void RemoveCpuThreadGroup(const TCHAR* GroupName) override               { Settings.RemoveCpuThreadGroup(GroupName); }
+		virtual void ResetCpuThreadGroups() override                                     { Settings.ResetCpuThreadGroups(); }
+		virtual void Reset() override                                                    { Settings.Reset(); }
 
 	private:
-		TMap<FString, bool> TimingViewShowToggleOptions;
-		TMap<FString, bool> TimingViewShowCpuThreadGroups;
+		FTimingViewSettings Settings;
 	};
 
 public:
@@ -141,35 +188,35 @@ public:
 
 	void ResetToDefaultValues();
 
-	bool IsAllGpuTracksToggleOn() const { return Settings->IsTimingViewToggleOptionEnabled(FInsightsSettings::FTimingViewTrackSettings::Option_ShowAllGpuTracks); }
+	bool IsAllGpuTracksToggleOn() const { return TimingViewSettings->IsToggleOptionEnabled(FTimingViewSettings::Option_ShowAllGpuTracks); }
 	void SetAllGpuTracksToggle(bool bOnOff);
 	void ShowAllGpuTracks() { SetAllGpuTracksToggle(true); }
 	void HideAllGpuTracks() { SetAllGpuTracksToggle(false); }
 	void ShowHideAllGpuTracks() { SetAllGpuTracksToggle(!IsAllGpuTracksToggleOn()); }
 
-	bool IsAllVerseTracksToggleOn() const { return Settings->IsTimingViewToggleOptionEnabled(FInsightsSettings::FTimingViewTrackSettings::Option_ShowAllVerseTracks); }
+	bool IsAllVerseTracksToggleOn() const { return TimingViewSettings->IsToggleOptionEnabled(FTimingViewSettings::Option_ShowAllVerseTracks); }
 	void SetAllVerseTracksToggle(bool bOnOff);
 	void ShowAllVerseTracks() { SetAllVerseTracksToggle(true); }
 	void HideAllVerseTracks() { SetAllVerseTracksToggle(false); }
 	void ShowHideAllVerseTracks() { SetAllVerseTracksToggle(!IsAllVerseTracksToggleOn()); }
 
-	bool IsAllCpuTracksToggleOn() const { return Settings->IsTimingViewToggleOptionEnabled(FInsightsSettings::FTimingViewTrackSettings::Option_ShowAllCpuTracks); }
+	bool IsAllCpuTracksToggleOn() const { return TimingViewSettings->IsToggleOptionEnabled(FTimingViewSettings::Option_ShowAllCpuTracks); }
 	void SetAllCpuTracksToggle(bool bOnOff);
 	void ShowAllCpuTracks() { SetAllCpuTracksToggle(true); }
 	void HideAllCpuTracks() { SetAllCpuTracksToggle(false); }
 	void ShowHideAllCpuTracks() { SetAllCpuTracksToggle(!IsAllCpuTracksToggleOn()); }
 
-	bool AreOverlaysVisibleInGpuQueueTracks() const { return Settings->IsTimingViewToggleOptionEnabled(FInsightsSettings::FTimingViewTrackSettings::Option_ShowGpuWorkOverlays); }
-	bool AreExtendedLinesVisibleInGpuQueueTracks() const { return Settings->IsTimingViewToggleOptionEnabled(FInsightsSettings::FTimingViewTrackSettings::Option_ShowGpuWorkExtendedLines); }
+	bool AreOverlaysVisibleInGpuQueueTracks() const { return TimingViewSettings->IsToggleOptionEnabled(FTimingViewSettings::Option_ShowGpuWorkOverlays); }
+	bool AreExtendedLinesVisibleInGpuQueueTracks() const { return TimingViewSettings->IsToggleOptionEnabled(FTimingViewSettings::Option_ShowGpuWorkExtendedLines); }
 
-	bool AreGpuWorkTracksVisible() const { return Settings->IsTimingViewToggleOptionEnabled(FInsightsSettings::FTimingViewTrackSettings::Option_ShowGpuWorkTracks); }
+	bool AreGpuWorkTracksVisible() const { return TimingViewSettings->IsToggleOptionEnabled(FTimingViewSettings::Option_ShowGpuWorkTracks); }
 	void SetGpuWorkTracksVisibility(bool bOnOff);
 
-	bool AreGpuFencesTracksVisible() const { return Settings->IsTimingViewToggleOptionEnabled(FInsightsSettings::FTimingViewTrackSettings::Option_ShowGpuFencesTracks); }
+	bool AreGpuFencesTracksVisible() const { return TimingViewSettings->IsToggleOptionEnabled(FTimingViewSettings::Option_ShowGpuFencesTracks); }
 	void SetGpuFencesTracksVisibility(bool bOnOff);
 
-	bool AreGpuFencesExtendedLinesVisible() const { return Settings->IsTimingViewToggleOptionEnabled(FInsightsSettings::FTimingViewTrackSettings::Option_ShowGpuFencesExtendedLines); }
-	bool AreGpuFenceRelationsVisible() const { return Settings->IsTimingViewToggleOptionEnabled(FInsightsSettings::FTimingViewTrackSettings::Option_ShowGpuFencesRelations); }
+	bool AreGpuFencesExtendedLinesVisible() const { return TimingViewSettings->IsToggleOptionEnabled(FTimingViewSettings::Option_ShowGpuFencesExtendedLines); }
+	bool AreGpuFenceRelationsVisible() const { return TimingViewSettings->IsToggleOptionEnabled(FTimingViewSettings::Option_ShowGpuFencesRelations); }
 
 	TSharedPtr<const ITimingEvent> FindMaxEventInstance(uint32 TimerId, double StartTime, double EndTime);
 	TSharedPtr<const ITimingEvent> FindMinEventInstance(uint32 TimerId, double StartTime, double EndTime);
@@ -181,19 +228,19 @@ private:
 	void ToggleTrackVisibilityByGroup_Execute(const TCHAR* InGroupName);
 	void ToggleOptionState(const TCHAR* CommandName);
 
-	void Command_ShowGpuWorkTracks_Execute() { SetGpuWorkTracksVisibility(!Settings->IsTimingViewToggleOptionEnabled(FInsightsSettings::FTimingViewTrackSettings::Option_ShowGpuWorkTracks)); }
+	void Command_ShowGpuWorkTracks_Execute() { SetGpuWorkTracksVisibility(!AreGpuWorkTracksVisible()); }
 	bool Command_ShowGpuWorkTracks_CanExecute() { return GpuTracks.Num() > 0; }
 	bool Command_ShowGpuWorkTracks_IsChecked() { return AreGpuWorkTracksVisible(); }
 
-	void Command_ShowGpuWorkOverlays_Execute() { ToggleOptionState(FInsightsSettings::FTimingViewTrackSettings::Option_ShowGpuWorkOverlays); }
+	void Command_ShowGpuWorkOverlays_Execute() { ToggleOptionState(FTimingViewSettings::Option_ShowGpuWorkOverlays); }
 	bool Command_ShowGpuWorkOverlays_CanExecute() { return AreGpuWorkTracksVisible() && GpuTracks.Num() > 0; }
 	bool Command_ShowGpuWorkOverlays_IsChecked() { return AreOverlaysVisibleInGpuQueueTracks(); }
 
-	void Command_ShowGpuWorkExtendedLines_Execute() { ToggleOptionState(FInsightsSettings::FTimingViewTrackSettings::Option_ShowGpuWorkExtendedLines); }
+	void Command_ShowGpuWorkExtendedLines_Execute() { ToggleOptionState(FTimingViewSettings::Option_ShowGpuWorkExtendedLines); }
 	bool Command_ShowGpuWorkExtendedLines_CanExecute() { return AreGpuWorkTracksVisible() && GpuTracks.Num() > 0; }
 	bool Command_ShowGpuWorkExtendedLines_IsChecked() { return AreExtendedLinesVisibleInGpuQueueTracks(); }
 
-	void Command_ShowGpuFencesExtendedLines_Execute() { ToggleOptionState(FInsightsSettings::FTimingViewTrackSettings::Option_ShowGpuFencesExtendedLines); }
+	void Command_ShowGpuFencesExtendedLines_Execute() { ToggleOptionState(FTimingViewSettings::Option_ShowGpuFencesExtendedLines); }
 	bool Command_ShowGpuFencesExtendedLines_CanExecute() { return AreGpuFencesTracksVisible() && GpuTracks.Num() > 0; }
 	bool Command_ShowGpuFencesExtendedLines_IsChecked() { return AreGpuFencesExtendedLinesVisible(); }
 
@@ -201,7 +248,7 @@ private:
 	bool Command_ShowGpuFencesRelations_CanExecute() { return GpuTracks.Num() > 0; }
 	bool Command_ShowGpuFencesRelations_IsChecked() { return AreGpuFenceRelationsVisible(); }
 
-	void Command_ShowGpuFencesTracks_Execute();
+	void Command_ShowGpuFencesTracks_Execute() { ToggleOptionState(FTimingViewSettings::Option_ShowGpuFencesTracks); }
 	bool Command_ShowGpuFencesTracks_CanExecute() { return GpuTracks.Num() > 0; }
 	bool Command_ShowGpuFencesTracks_IsChecked() { return AreGpuFencesTracksVisible(); }
 
@@ -233,7 +280,7 @@ private:
 	uint64 TimingProfilerTimelineCount = 0;
 	uint64 LoadTimeProfilerTimelineCount = 0;
 
-	TSharedRef<IThreadSharedStateSetting> Settings;
+	TSharedRef<ITimingViewSettings> TimingViewSettings;
 };
 
 ////////////////////////////////////////////////////////////////////////////////////////////////////
