Add method of creating analysis session using provided data stream

diff --git a/Engine/Source/Developer/TraceInsights/Private/Insights/InsightsManager.cpp b/Engine/Source/Developer/TraceInsights/Private/Insights/InsightsManager.cpp
--- a/Engine/Source/Developer/TraceInsights/Private/Insights/InsightsManager.cpp
+++ b/Engine/Source/Developer/TraceInsights/Private/Insights/InsightsManager.cpp
@@ -836,10 +836,10 @@ void FInsightsManager::LoadTraceFile(const FString& InTraceFilename, bool InAuto
 
 ////////////////////////////////////////////////////////////////////////////////////////////////////
 
-uint16 FInsightsManager::ListenForDirectTrace()
+uint16 FInsightsManager::ListenForDirectTrace(const TCHAR* InStreamName)
 {
 	constexpr uint32 TraceId = 0; 
-	FString TraceName(TEXT("DirectTrace"));
+	FString TraceName = InStreamName ? InStreamName : TEXT("DirectStream");
 	TUniquePtr<UE::Trace::FDirectSocketStream> TraceData = MakeUnique<UE::Trace::FDirectSocketStream>();
 	const uint16 Port = TraceData->StartListening();
 	
@@ -858,6 +858,24 @@ uint16 FInsightsManager::ListenForDirectTrace()
 	return Port;
 }
 
+////////////////////////////////////////////////////////////////////////////////////////////////////
+void FInsightsManager::StartAnalysisWithStream(TUniquePtr<UE::Trace::IInDataStream>&& InStream, const TCHAR* InStreamName)
+{
+	ResetSession(false);
+
+	constexpr uint32 TraceId = 0; 
+	FString TraceName = InStreamName ? InStreamName : TEXT("CustomStream");
+	Session = AnalysisService->StartAnalysis(TraceId, *TraceName, MoveTemp(InStream));
+
+	if (Session)
+	{
+		CurrentTraceId = TraceId;
+		CurrentTraceFilename = TraceName;
+		bIsSessionInfoSet = false;
+		OnSessionChanged();
+	}
+}
+
 ////////////////////////////////////////////////////////////////////////////////////////////////////
 
 bool FInsightsManager::OnDragOver(const FDragDropEvent& DragDropEvent)
diff --git a/Engine/Source/Developer/TraceInsights/Private/Insights/InsightsManager.h b/Engine/Source/Developer/TraceInsights/Private/Insights/InsightsManager.h
--- a/Engine/Source/Developer/TraceInsights/Private/Insights/InsightsManager.h
+++ b/Engine/Source/Developer/TraceInsights/Private/Insights/InsightsManager.h
@@ -228,8 +228,21 @@ public:
 	 */
 	void LoadTraceFile(const FString& TraceFilename, bool InAutoQuit = false, bool InWaitForSymbolResolver = false);
 
+	/**
+	 * Creates a new analysis session instance and opens a listening port to await connections.
+	 * @param InStreamName User facing name of stream
+	 * @return Port number to connect to 
+	 * @return Port used for listening to incoming connections. 0 if failed to open port.
+	 */
+	uint16 ListenForDirectTrace(const TCHAR* InStreamName = nullptr);
+
+	/**
+	 * Creates a new analysis session instance from a provided data stream. Replaces the current analysis session.
+	 * @param InStream Stream to read from
+	 * @param InStreamName User facing name of stream
+	 */
+	void StartAnalysisWithStream(TUniquePtr<UE::Trace::IInDataStream>&& InStream, const TCHAR* InStreamName = nullptr);
 
-	uint16 ListenForDirectTrace();
 
 	bool OnDragOver(const FDragDropEvent& DragDropEvent);
 	bool OnDrop(const FDragDropEvent& DragDropEvent);
diff --git a/Engine/Source/Developer/TraceInsights/Private/Insights/TraceInsightsModule.cpp b/Engine/Source/Developer/TraceInsights/Private/Insights/TraceInsightsModule.cpp
--- a/Engine/Source/Developer/TraceInsights/Private/Insights/TraceInsightsModule.cpp
+++ b/Engine/Source/Developer/TraceInsights/Private/Insights/TraceInsightsModule.cpp
@@ -495,9 +495,16 @@ void FTraceInsightsModule::StartAnalysisForTraceFile(const TCHAR* InTraceFile, b
 
 ////////////////////////////////////////////////////////////////////////////////////////////////////
 
-uint16 FTraceInsightsModule::StartAnalysisWithDirectTrace()
+uint16 FTraceInsightsModule::StartAnalysisWithDirectTrace(const TCHAR* InStreamName)
 {
-	return FInsightsManager::Get()->ListenForDirectTrace();
+	return FInsightsManager::Get()->ListenForDirectTrace(InStreamName);
+}
+
+////////////////////////////////////////////////////////////////////////////////////////////////////
+
+void FTraceInsightsModule::StartAnalysisWithStream(TUniquePtr<UE::Trace::IInDataStream>&& InStream, const TCHAR* InStreamName)
+{
+	FInsightsManager::Get()->StartAnalysisWithStream(MoveTemp(InStream), InStreamName);
 }
 
 ////////////////////////////////////////////////////////////////////////////////////////////////////
diff --git a/Engine/Source/Developer/TraceInsights/Private/Insights/TraceInsightsModule.h b/Engine/Source/Developer/TraceInsights/Private/Insights/TraceInsightsModule.h
--- a/Engine/Source/Developer/TraceInsights/Private/Insights/TraceInsightsModule.h
+++ b/Engine/Source/Developer/TraceInsights/Private/Insights/TraceInsightsModule.h
@@ -58,7 +58,8 @@ public:
 	virtual void StartAnalysisForTrace(uint32 InTraceId, bool InAutoQuit = false, bool InWaitForSymbolResolver = false) override;
 	virtual void StartAnalysisForLastLiveSession(float InRetryTime = 1.0f) override;
 	virtual void StartAnalysisForTraceFile(const TCHAR* InTraceFile, bool InAutoQuit = false, bool InWaitForSymbolResolver = false) override;
-	virtual uint16 StartAnalysisWithDirectTrace() override;
+	virtual uint16 StartAnalysisWithDirectTrace(const TCHAR* InStreamName = nullptr) override;
+	virtual void StartAnalysisWithStream(TUniquePtr<UE::Trace::IInDataStream>&& InStream, const TCHAR* InStreamName = nullptr) override;
 
 	virtual void ShutdownUserInterface() override;
 
diff --git a/Engine/Source/Developer/TraceInsights/Public/Insights/IUnrealInsightsModule.h b/Engine/Source/Developer/TraceInsights/Public/Insights/IUnrealInsightsModule.h
--- a/Engine/Source/Developer/TraceInsights/Public/Insights/IUnrealInsightsModule.h
+++ b/Engine/Source/Developer/TraceInsights/Public/Insights/IUnrealInsightsModule.h
@@ -11,6 +11,7 @@
 #include "Internationalization/Text.h"
 #include "Modules/ModuleInterface.h"
 #include "Templates/SharedPointer.h"
+#include "Trace/DataStream.h"
 #include "UObject/NameTypes.h"
 
 #define UE_API TRACEINSIGHTS_API
@@ -281,9 +282,19 @@ public:
 	 * Starts analysis for a direct connection from the tracing application (might be self). Notice that
 	 * analysis doesn't actually start until a connection has been made.
 	 * 
+	 * @param InStreamName User facing display name for this stream
 	 * @return Port used for listening to incoming connections. 0 if failed to open port.
 	 */
-	virtual uint16 StartAnalysisWithDirectTrace() = 0;
+	virtual uint16 StartAnalysisWithDirectTrace(const TCHAR* InStreamName = nullptr) = 0;
+
+
+	/**
+	 * Start analysis given a data stream.
+	 *
+	 * @param InStream Stream to read from
+	 * @param InStreamName User facing display name for this stream
+	 */
+	virtual void StartAnalysisWithStream(TUniquePtr<UE::Trace::IInDataStream>&& InStream, const TCHAR* InStreamName = nullptr) = 0;
 
 	//////////////////////////////////////////////////
 
