[Insights] Removed unused CYCLE_STATs from TimingProfilerCommon.h.

diff --git a/Engine/Source/Developer/TraceInsights/Public/Insights/IInsightsManager.h b/Engine/Source/Developer/TraceInsights/Public/Insights/IInsightsManager.h
--- a/Engine/Source/Developer/TraceInsights/Public/Insights/IInsightsManager.h
+++ b/Engine/Source/Developer/TraceInsights/Public/Insights/IInsightsManager.h
@@ -2,6 +2,9 @@
 
 #pragma once
 
+#include "Delegates/DelegateCombinations.h"
+
+// TraceInsights
 #include "Insights/IUnrealInsightsModule.h"
 
 namespace Insights
diff --git a/Engine/Source/Developer/TraceInsights/Public/Insights/TimingProfilerCommon.h b/Engine/Source/Developer/TraceInsights/Public/Insights/TimingProfilerCommon.h
--- a/Engine/Source/Developer/TraceInsights/Public/Insights/TimingProfilerCommon.h
+++ b/Engine/Source/Developer/TraceInsights/Public/Insights/TimingProfilerCommon.h
@@ -2,23 +2,7 @@
 
 #pragma once
 
-#include "CoreMinimal.h"
+#include "CoreTypes.h"
 #include "Logging/LogMacros.h"
-#include "Stats/Stats.h"
-#include "Stats/StatsMisc.h"
 
 TRACEINSIGHTS_API DECLARE_LOG_CATEGORY_EXTERN(LogTimingProfiler, Log, All);
-
-DECLARE_STATS_GROUP(TEXT("TimingProfiler"), STATGROUP_TimingProfiler, STATCAT_Advanced);
-
-/** Time spent on Frame Track drawing. */
-DECLARE_CYCLE_STAT_EXTERN(TEXT("FrameTrackOnPaint"), STAT_FT_OnPaint, STATGROUP_TimingProfiler, );
-
-/** Time spent on Graph Track drawing. */
-DECLARE_CYCLE_STAT_EXTERN(TEXT("GraphTrackOnPaint"), STAT_GT_OnPaint, STATGROUP_TimingProfiler, );
-
-/** Time spent on Timing Track drawing. */
-DECLARE_CYCLE_STAT_EXTERN(TEXT("TimingTrackOnPaint"), STAT_TT_OnPaint, STATGROUP_TimingProfiler, );
-
-/** Time spent on ticking profiler manager. */
-DECLARE_CYCLE_STAT_EXTERN(TEXT("TimingProfilerTick"), STAT_TPM_Tick, STATGROUP_TimingProfiler, );
