[Insights]
* Added visibility for the timer types (CPU, GPU, Verse) based on available data in the analyzed session.
* Removed UE_EXPERIMENTAL_VERSE_INSIGHTS_ENABLED compile time switch.

diff --git a/Engine/Source/Developer/TraceInsights/Private/Insights/TimingProfiler/ViewModels/ThreadTimingSharedState.cpp b/Engine/Source/Developer/TraceInsights/Private/Insights/TimingProfiler/ViewModels/ThreadTimingSharedState.cpp
--- a/Engine/Source/Developer/TraceInsights/Private/Insights/TimingProfiler/ViewModels/ThreadTimingSharedState.cpp
+++ b/Engine/Source/Developer/TraceInsights/Private/Insights/TimingProfiler/ViewModels/ThreadTimingSharedState.cpp
@@ -496,7 +496,6 @@ void FThreadTimingSharedState::Tick(Timing::ITimingViewSession& InSession, const
 				GpuTrackOrder += 100;
 			});
 
-#if UE_EXPERIMENTAL_VERSE_INSIGHTS_ENABLED
 			// Check if we have the Verse sampling timeline.
 			if (!VerseSamplingTrack.IsValid())
 			{
@@ -510,7 +509,6 @@ void FThreadTimingSharedState::Tick(Timing::ITimingViewSession& InSession, const
 					InSession.AddScrollableTrack(VerseSamplingTrack);
 				}
 			}
-#endif
 
 			// Iterate through threads.
 			const TraceServices::IThreadProvider& ThreadProvider = TraceServices::ReadThreadProvider(InAnalysisSession);
@@ -634,13 +632,11 @@ void FThreadTimingSharedState::ExtendCpuTracksFilterMenu(Timing::ITimingViewSess
 		return;
 	}
 
-#if UE_EXPERIMENTAL_VERSE_INSIGHTS_ENABLED
 	InOutMenuBuilder.BeginSection("VerseTracks", LOCTEXT("ContextMenu_Section_VerseTracks", "Verse Tracks"));
 	{
 		InOutMenuBuilder.AddMenuEntry(FThreadTimingViewCommands::Get().ShowHideAllVerseTracks);
 	}
 	InOutMenuBuilder.EndSection();
-#endif
 
 	InOutMenuBuilder.BeginSection("CpuTracks", LOCTEXT("ContextMenu_Section_CpuTracks", "CPU Tracks"));
 	{
@@ -719,13 +715,11 @@ void FThreadTimingSharedState::BindCommands()
 		FCanExecuteAction::CreateSP(this, &FThreadTimingSharedState::Command_ShowGpuFencesRelations_CanExecute),
 		FIsActionChecked::CreateSP(this, &FThreadTimingSharedState::Command_ShowGpuFencesRelations_IsChecked));
 
-#if UE_EXPERIMENTAL_VERSE_INSIGHTS_ENABLED
 	CommandList->MapAction(
 		FThreadTimingViewCommands::Get().ShowHideAllVerseTracks,
 		FExecuteAction::CreateSP(this, &FThreadTimingSharedState::ShowHideAllVerseTracks),
 		FCanExecuteAction(),
 		FIsActionChecked::CreateSP(this, &FThreadTimingSharedState::IsAllVerseTracksToggleOn));
-#endif
 
 	CommandList->MapAction(
 		FThreadTimingViewCommands::Get().ShowHideAllCpuTracks,
diff --git a/Engine/Source/Developer/TraceInsights/Private/Insights/TimingProfiler/ViewModels/TimingExporter.cpp b/Engine/Source/Developer/TraceInsights/Private/Insights/TimingProfiler/ViewModels/TimingExporter.cpp
--- a/Engine/Source/Developer/TraceInsights/Private/Insights/TimingProfiler/ViewModels/TimingExporter.cpp
+++ b/Engine/Source/Developer/TraceInsights/Private/Insights/TimingProfiler/ViewModels/TimingExporter.cpp
@@ -282,7 +282,21 @@ int32 FTimingExporter::ExportTimersAsText(const FString& Filename, FExportTimers
 			const TraceServices::FTimingProfilerTimer& Timer = *(TimerReader->GetTimer(TimerIndex));
 			StringBuilder.Appendf(UTF8TEXT("%u"), Timer.Id);
 			Writer.AppendSeparator();
-			StringBuilder.Append(Timer.IsGpuTimer ? UTF8TEXT("GPU") : UTF8TEXT("CPU"));
+			switch (Timer.Type)
+			{
+			case TraceServices::ETimingProfilerTimerType::CPU:
+				StringBuilder.Append(UTF8TEXT("CPU"));
+				break;
+			case TraceServices::ETimingProfilerTimerType::GPU:
+				StringBuilder.Append(UTF8TEXT("GPU"));
+				break;
+			case TraceServices::ETimingProfilerTimerType::Verse:
+				StringBuilder.Append(UTF8TEXT("Verse"));
+				break;
+			default:
+				StringBuilder.Append(UTF8TEXT("Unknown"));
+				break;
+			}
 			Writer.AppendSeparator();
 			Writer.AppendString(Timer.Name);
 			Writer.AppendSeparator();
diff --git a/Engine/Source/Developer/TraceInsights/Private/Insights/TimingProfiler/Widgets/STimerTreeView.cpp b/Engine/Source/Developer/TraceInsights/Private/Insights/TimingProfiler/Widgets/STimerTreeView.cpp
--- a/Engine/Source/Developer/TraceInsights/Private/Insights/TimingProfiler/Widgets/STimerTreeView.cpp
+++ b/Engine/Source/Developer/TraceInsights/Private/Insights/TimingProfiler/Widgets/STimerTreeView.cpp
@@ -1523,9 +1523,11 @@ FTimerNodePtr STimerTreeView::CreateTimerNodeRec(const TraceServices::FTimingPro
 		return nullptr;
 	}
 
-	const ETimerNodeType Type = Node.Timer->IsGpuTimer ?   ETimerNodeType::GpuScope :
-								Node.Timer->IsVerseTimer ? ETimerNodeType::VerseScope :
-														   ETimerNodeType::CpuScope;
+	const ETimerNodeType Type =
+		Node.Timer->Type == TraceServices::ETimingProfilerTimerType::GPU ? ETimerNodeType::GpuScope :
+		Node.Timer->Type == TraceServices::ETimingProfilerTimerType::Verse ? ETimerNodeType::VerseScope :
+		Node.Timer->Type == TraceServices::ETimingProfilerTimerType::CPU ? ETimerNodeType::CpuScope :
+		ETimerNodeType::InvalidOrMax;
 
 	TSharedPtr<STimingProfilerWindow> ProfilerWindow = FTimingProfilerManager::Get()->GetProfilerWindow();
 	const bool bIsTimerAddedToGraphs = ProfilerWindow.IsValid() ? ProfilerWindow->IsTimerAddedToGraphs(Node.Timer->Id) : false;
diff --git a/Engine/Source/Developer/TraceInsights/Private/Insights/TimingProfiler/Widgets/STimersView.cpp b/Engine/Source/Developer/TraceInsights/Private/Insights/TimingProfiler/Widgets/STimersView.cpp
--- a/Engine/Source/Developer/TraceInsights/Private/Insights/TimingProfiler/Widgets/STimersView.cpp
+++ b/Engine/Source/Developer/TraceInsights/Private/Insights/TimingProfiler/Widgets/STimersView.cpp
@@ -293,7 +293,6 @@ void STimersView::Construct(const FArguments& InArgs)
 					GetToggleButtonForTimerType(ETimerNodeType::GpuScope)
 				]
 
-#if UE_EXPERIMENTAL_VERSE_INSIGHTS_ENABLED
 				// Filter by type (Verse)
 				+ SHorizontalBox::Slot()
 				.Padding(4.0f, 0.0f, 4.0f, 0.0f)
@@ -302,7 +301,6 @@ void STimersView::Construct(const FArguments& InArgs)
 				[
 					GetToggleButtonForTimerType(ETimerNodeType::VerseScope)
 				]
-#endif
 
 				// Filter by type (CPU)
 				+ SHorizontalBox::Slot()
@@ -1548,6 +1546,11 @@ void STimersView::HandleItemToStringArray(const FTimerNodePtr& FTimerNodePtr, TA
 TSharedRef<SWidget> STimersView::GetToggleButtonForTimerType(const ETimerNodeType NodeType)
 {
 	return SNew(SCheckBox)
+		.Visibility_Lambda([this, NodeType]()
+			{
+				return static_cast<uint32>(NodeType) < static_cast<uint32>(ETimerNodeType::InvalidOrMax) &&
+					bTimerTypeIsEnabled[static_cast<uint32>(NodeType)] ? EVisibility::Visible : EVisibility::Collapsed;
+			})
 		.Style(FAppStyle::Get(), "ToggleButtonCheckbox")
 		.Padding(FMargin(4.0f, 2.0f, 4.0f, 2.0f))
 		.HAlign(HAlign_Center)
@@ -2579,6 +2582,22 @@ void STimersView::Tick(const FGeometry& AllottedGeometry, const double InCurrent
 	}
 
 	Aggregator->Tick(Session, InCurrentTime, InDeltaTime, [this]() { FinishAggregation(); });
+
+	if (Session.IsValid() && TraceServices::ReadTimingProfilerProvider(*Session.Get()))
+	{
+		TraceServices::FAnalysisSessionReadScope SessionReadScope(*Session.Get());
+		const TraceServices::ITimingProfilerProvider& TimingProfilerProvider = *TraceServices::ReadTimingProfilerProvider(*Session.Get());
+
+		bTimerTypeIsEnabled[static_cast<int>(ETimerNodeType::GpuScope)] = TimingProfilerProvider.HasGpuTiming();
+		bTimerTypeIsEnabled[static_cast<int>(ETimerNodeType::VerseScope)] = TimingProfilerProvider.HasVerseTiming();
+		bTimerTypeIsEnabled[static_cast<int>(ETimerNodeType::CpuScope)] = TimingProfilerProvider.HasCpuTiming();
+	}
+	else
+	{
+		bTimerTypeIsEnabled[static_cast<int>(ETimerNodeType::GpuScope)] = false;
+		bTimerTypeIsEnabled[static_cast<int>(ETimerNodeType::VerseScope)] = false;
+		bTimerTypeIsEnabled[static_cast<int>(ETimerNodeType::CpuScope)] = false;
+	}
 }
 
 ////////////////////////////////////////////////////////////////////////////////////////////////////
@@ -2621,9 +2640,11 @@ void STimersView::RebuildTree(bool bResync)
 			{
 				const TraceServices::FTimingProfilerTimer& Timer = *(TimerReader->GetTimer(TimerIndex));
 				ensure(Timer.Id == TimerIndex);
-				const ETimerNodeType Type = Timer.IsGpuTimer ?	 ETimerNodeType::GpuScope :
-											Timer.IsVerseTimer ? ETimerNodeType::VerseScope :
-																 ETimerNodeType::CpuScope;
+				const ETimerNodeType Type =
+					Timer.Type == TraceServices::ETimingProfilerTimerType::GPU ? ETimerNodeType::GpuScope :
+					Timer.Type == TraceServices::ETimingProfilerTimerType::Verse ? ETimerNodeType::VerseScope :
+					Timer.Type == TraceServices::ETimingProfilerTimerType::CPU ? ETimerNodeType::CpuScope :
+					ETimerNodeType::InvalidOrMax;
 				FTimerNodePtr TimerNodePtr = MakeShared<FTimerNode>(Timer.Id, Timer.Name, Type, false);
 				TimerNodePtr->SetDefaultSortOrder(TimerIndex + 1);
 
diff --git a/Engine/Source/Developer/TraceInsights/Private/Insights/TimingProfiler/Widgets/STimersView.h b/Engine/Source/Developer/TraceInsights/Private/Insights/TimingProfiler/Widgets/STimersView.h
--- a/Engine/Source/Developer/TraceInsights/Private/Insights/TimingProfiler/Widgets/STimersView.h
+++ b/Engine/Source/Developer/TraceInsights/Private/Insights/TimingProfiler/Widgets/STimersView.h
@@ -477,6 +477,9 @@ private:
 	/** The filter collection. */
 	TSharedPtr<FTimerNodeFilterCollection> Filters;
 
+	/** Holds the enabled state of each timer type. */
+	bool bTimerTypeIsEnabled[static_cast<int>(ETimerNodeType::InvalidOrMax)];
+
 	/** Holds the visibility of each timer type. */
 	bool bTimerTypeIsVisible[static_cast<int>(ETimerNodeType::InvalidOrMax)];
 
diff --git a/Engine/Source/Developer/TraceServices/Private/Model/TimingProfiler.cpp b/Engine/Source/Developer/TraceServices/Private/Model/TimingProfiler.cpp
--- a/Engine/Source/Developer/TraceServices/Private/Model/TimingProfiler.cpp
+++ b/Engine/Source/Developer/TraceServices/Private/Model/TimingProfiler.cpp
@@ -46,11 +46,36 @@ FTimingProfilerProvider::~FTimingProfilerProvider()
 {
 }
 
+bool FTimingProfilerProvider::HasCpuTiming() const
+{
+	Session.ReadAccessCheck();
+
+	return bHasCpuTiming;
+}
+
+bool FTimingProfilerProvider::HasGpuTiming() const
+{
+	Session.ReadAccessCheck();
+
+	return bHasGpuTiming ||
+		Timelines[GpuTimelineIndex]->GetEventCount() > 0 ||
+		Timelines[Gpu2TimelineIndex]->GetEventCount() > 0;
+}
+
+bool FTimingProfilerProvider::HasVerseTiming() const
+{
+	Session.ReadAccessCheck();
+
+	return bHasVerseTiming ||
+		Timelines[VerseTimelineIndex]->GetEventCount() > 0;
+}
+
 uint32 FTimingProfilerProvider::AddCpuTimer(FStringView Name, const TCHAR* File, uint32 Line)
 {
 	Session.WriteAccessCheck();
 
-	FTimingProfilerTimer& Timer = AddTimerInternal(Name, File, Line, ETimingProfilerTimerType::Cpu);
+	bHasCpuTiming = true;
+	FTimingProfilerTimer& Timer = AddTimerInternal(Name, File, Line, ETimingProfilerTimerType::CPU);
 	return Timer.Id;
 }
 
@@ -58,7 +83,8 @@ uint32 FTimingProfilerProvider::AddGpuTimer(FStringView Name, const TCHAR* File,
 {
 	Session.WriteAccessCheck();
 
-	FTimingProfilerTimer& Timer = AddTimerInternal(Name, File, Line, ETimingProfilerTimerType::Gpu);
+	bHasGpuTiming = true;
+	FTimingProfilerTimer& Timer = AddTimerInternal(Name, File, Line, ETimingProfilerTimerType::GPU);
 	return Timer.Id;
 }
 
@@ -66,6 +92,7 @@ uint32 FTimingProfilerProvider::AddVerseTimer(FStringView Name, const TCHAR* Fil
 {
 	Session.WriteAccessCheck();
 
+	bHasVerseTiming = true;
 	FTimingProfilerTimer& Timer = AddTimerInternal(Name, File, Line, ETimingProfilerTimerType::Verse);
 	return Timer.Id;
 }
@@ -77,8 +104,7 @@ FTimingProfilerTimer& FTimingProfilerProvider::AddTimerInternal(FStringView Name
 	Timer.Name = Session.StoreString(Name);
 	Timer.File = File ? Session.StoreString(File) : nullptr;
 	Timer.Line = Line;
-	Timer.IsGpuTimer = (TimerType == ETimingProfilerTimerType::Gpu);
-	Timer.IsVerseTimer = (TimerType == ETimingProfilerTimerType::Verse);
+	Timer.Type = TimerType;
 	return Timer;
 }
 
@@ -247,6 +273,7 @@ IEditableTimeline<FTimingProfilerEvent>& FTimingProfilerProvider::GetCpuThreadEd
 		uint32 TimelineIndex = Timelines.Num();
 		CpuThreadTimelineIndexMap.Add(ThreadId, TimelineIndex);
 		Timelines.Add(Timeline);
+		bHasCpuTiming = true;
 		return Timeline.Get();
 	}
 	else
@@ -294,6 +321,7 @@ void FTimingProfilerProvider::AddGpuQueue(uint32 QueueId, uint8 GPU, uint8 Index
 		GpuQueue.WorkTimelineIndex = Timelines.Num();
 		TSharedRef<TimelineInternal> WorkTimeline = MakeShared<TimelineInternal>(Session.GetLinearAllocator());
 		Timelines.Add(WorkTimeline);
+		bHasGpuTiming = true;
 	}
 }
 
@@ -314,6 +342,7 @@ void FTimingProfilerProvider::AddGpuSignalFence(uint32 QueueId, const FGpuSignal
 		}
 
 		Data.SignalFenceArray.Add(Fence);
+		bHasGpuTiming = true;
 	}
 }
 
@@ -334,6 +363,7 @@ void FTimingProfilerProvider::AddGpuWaitFence(uint32 QueueId, const FGpuWaitFenc
 		}
 
 		Data.WaitFenceArray.Add(Fence);
+		bHasGpuTiming = true;
 	}
 }
 
diff --git a/Engine/Source/Developer/TraceServices/Private/Model/TimingProfilerPrivate.h b/Engine/Source/Developer/TraceServices/Private/Model/TimingProfilerPrivate.h
--- a/Engine/Source/Developer/TraceServices/Private/Model/TimingProfilerPrivate.h
+++ b/Engine/Source/Developer/TraceServices/Private/Model/TimingProfilerPrivate.h
@@ -51,6 +51,7 @@ public:
 
 	// GPU
 
+	virtual bool HasGpuTiming() const override;
 	virtual bool GetGpuTimelineIndex(uint32& OutTimelineIndex) const override; // backward compatibility
 	virtual bool GetGpu2TimelineIndex(uint32& OutTimelineIndex) const override; // backward compatibility
 	virtual void EnumerateGpuQueues(TFunctionRef<void(const FGpuQueueInfo&)> Callback) const override;
@@ -62,10 +63,12 @@ public:
 
 	// Verse
 
+	virtual bool HasVerseTiming() const override;
 	virtual bool GetVerseTimelineIndex(uint32& OutTimelineIndex) const override;
 
 	// CPU
 
+	virtual bool HasCpuTiming() const override;
 	virtual bool GetCpuThreadTimelineIndex(uint32 ThreadId, uint32& OutTimelineIndex) const override;
 
 	// Aggregation
@@ -132,13 +135,6 @@ public:
 	TimelineInternal& EditGpu2Timeline();
 
 private:
-	enum class ETimingProfilerTimerType : uint32
-	{
-		Cpu,
-		Gpu,
-		Verse
-	};
-
 	FTimingProfilerTimer& AddTimerInternal(FStringView Name, const TCHAR* File, uint32 Line, ETimingProfilerTimerType TimerType);
 
 	struct FMetadata
@@ -155,6 +151,10 @@ private:
 	TTableLayout<FTimingProfilerAggregatedStats> AggregatedStatsTableLayout;
 	double RatioOfThreadsToUse = 0.75;
 
+	bool bHasGpuTiming = false;
+	bool bHasVerseTiming = false;
+	bool bHasCpuTiming = false;
+
 	// GPU
 
 	static const uint32 GpuTimelineIndex = 0;
diff --git a/Engine/Source/Developer/TraceServices/Private/Modules/TimingProfilerModule.cpp b/Engine/Source/Developer/TraceServices/Private/Modules/TimingProfilerModule.cpp
--- a/Engine/Source/Developer/TraceServices/Private/Modules/TimingProfilerModule.cpp
+++ b/Engine/Source/Developer/TraceServices/Private/Modules/TimingProfilerModule.cpp
@@ -31,9 +31,7 @@ void FTimingProfilerModule::OnAnalysisBegin(IAnalysisSession& InSession)
 	Session.AddProvider(GetTimingProfilerProviderName(), TimingProfilerProvider, TimingProfilerProvider);
 	Session.AddAnalyzer(new FCpuProfilerAnalyzer(Session, *TimingProfilerProvider, EditableThreadProvider));
 	Session.AddAnalyzer(new FGpuProfilerAnalyzer(Session, *TimingProfilerProvider, EditCounterProvider(Session)));
-#if UE_EXPERIMENTAL_VERSE_INSIGHTS_ENABLED
 	Session.AddAnalyzer(new FVerseAnalyzer(Session, *TimingProfilerProvider));
-#endif
 }
 
 void FTimingProfilerModule::GetLoggers(TArray<const TCHAR*>& OutLoggers)
diff --git a/Engine/Source/Developer/TraceServices/Public/TraceServices/Model/TimingProfiler.h b/Engine/Source/Developer/TraceServices/Public/TraceServices/Model/TimingProfiler.h
--- a/Engine/Source/Developer/TraceServices/Public/TraceServices/Model/TimingProfiler.h
+++ b/Engine/Source/Developer/TraceServices/Public/TraceServices/Model/TimingProfiler.h
@@ -14,10 +14,6 @@
 #include "TraceServices/Model/AnalysisSession.h"
 #include "UObject/NameTypes.h"
 
-#ifndef UE_EXPERIMENTAL_VERSE_INSIGHTS_ENABLED
-#define UE_EXPERIMENTAL_VERSE_INSIGHTS_ENABLED 0
-#endif
-
 namespace TraceServices
 {
 template <typename InEventType> class IEditableTimeline;
@@ -31,6 +27,13 @@ struct FMetadataSpec
 	static const uint32 InvalidMetadataSpecId = (uint32)-1;
 };
 
+enum class ETimingProfilerTimerType : uint8
+{
+	CPU,
+	GPU,
+	Verse,
+};
+
 struct FTimingProfilerTimer
 {
 	const TCHAR* Name = nullptr;
@@ -43,10 +46,9 @@ struct FTimingProfilerTimer
 		struct
 		{
 			uint32 Line : 24;
-			uint32 IsGpuTimer : 1;
-			uint32 IsVerseTimer : 1;
+			ETimingProfilerTimerType Type : 8;
 		};
-		uint32 LineAndFlags = 0; // used only to default initialize Line and IsGpuTimer with 0
+		uint32 LineAndType = 0; // used only to default initialize Line and Type with 0
 	};
 
 	bool HasValidMetadataSpecId() const
@@ -268,6 +270,8 @@ public:
 	//////////////////////////////////////////////////
 	// GPU
 
+	virtual bool HasGpuTiming() const { return false; }
+
 	// Only used for backward compatibility with old GPU Insights.
 	virtual bool GetGpuTimelineIndex(uint32& OutTimelineIndex) const { return false; }
 	// Only used for backward compatibility with old GPU Insights.
@@ -284,11 +288,15 @@ public:
 	//////////////////////////////////////////////////
 	// Verse
 
+	virtual bool HasVerseTiming() const { return false; }
+
 	virtual bool GetVerseTimelineIndex(uint32& OutTimelineIndex) const { return false; }
 
 	//////////////////////////////////////////////////
 	// CPU
 
+	virtual bool HasCpuTiming() const { return true; }
+
 	virtual bool GetCpuThreadTimelineIndex(uint32 ThreadId, uint32& OutTimelineIndex) const = 0;
 
 	//////////////////////////////////////////////////
