Timing Insights / (Quick Find / Timers Context): Consistency in event actions for QuickFind and Timers Context Menu
* Added Min/Max functions to Quick Find Toolbar
* Timers Context Menu adds First/Prev/Next/Last functions

diff --git a/Engine/Source/Developer/TraceInsights/Private/Insights/InsightsStyle.cpp b/Engine/Source/Developer/TraceInsights/Private/Insights/InsightsStyle.cpp
--- a/Engine/Source/Developer/TraceInsights/Private/Insights/InsightsStyle.cpp
+++ b/Engine/Source/Developer/TraceInsights/Private/Insights/InsightsStyle.cpp
@@ -170,9 +170,13 @@ void FInsightsStyle::FStyle::Initialize()
 
 	//////////////////////////////////////////////////
 
+	Set("Icons.FindFirst", new IMAGE_BRUSH_SVG("ControlsFirst", Icon16x16));
 	Set("Icons.FindFirst.ToolBar", new IMAGE_BRUSH_SVG("ControlsFirst", Icon20x20));
+	Set("Icons.FindPrevious", new IMAGE_BRUSH_SVG("ControlsPrevious", Icon16x16));
 	Set("Icons.FindPrevious.ToolBar", new IMAGE_BRUSH_SVG("ControlsPrevious", Icon20x20));
+	Set("Icons.FindNext", new IMAGE_BRUSH_SVG("ControlsNext", Icon16x16));
 	Set("Icons.FindNext.ToolBar", new IMAGE_BRUSH_SVG("ControlsNext", Icon20x20));
+	Set("Icons.FindLast", new IMAGE_BRUSH_SVG("ControlsLast", Icon16x16));
 	Set("Icons.FindLast.ToolBar", new IMAGE_BRUSH_SVG("ControlsLast", Icon20x20));
 
 	//////////////////////////////////////////////////
@@ -265,8 +269,10 @@ void FInsightsStyle::FStyle::Initialize()
 	Set("Icons.SelectEventRange", new EDITOR_IMAGE_BRUSH_SVG("Starship/Common/SelectInViewport", Icon16x16));
 
 	Set("Icons.FindInstance", new CORE_IMAGE_BRUSH_SVG("Starship/Common/search", Icon16x16));
-	Set("Icons.FindMinInstance", new EDITOR_IMAGE_BRUSH_SVG("Starship/Common/Profile", Icon16x16));
-	Set("Icons.FindMaxInstance", new EDITOR_IMAGE_BRUSH_SVG("Starship/Common/Profile", Icon16x16));
+	Set("Icons.FindMinInstance", new IMAGE_BRUSH_SVG("SizeSmall", Icon16x16));
+	Set("Icons.FindMinInstance.ToolBar", new IMAGE_BRUSH_SVG("SizeSmall_20", Icon20x20));
+	Set("Icons.FindMaxInstance", new IMAGE_BRUSH_SVG("SizeLarge", Icon16x16));
+	Set("Icons.FindMaxInstance.ToolBar", new IMAGE_BRUSH_SVG("SizeLarge_20", Icon20x20));
 
 	//////////////////////////////////////////////////
 	// Icons for tree/table items
diff --git a/Engine/Source/Developer/TraceInsights/Private/Insights/TimingProfiler/Widgets/STimersView.cpp b/Engine/Source/Developer/TraceInsights/Private/Insights/TimingProfiler/Widgets/STimersView.cpp
--- a/Engine/Source/Developer/TraceInsights/Private/Insights/TimingProfiler/Widgets/STimersView.cpp
+++ b/Engine/Source/Developer/TraceInsights/Private/Insights/TimingProfiler/Widgets/STimersView.cpp
@@ -58,7 +58,6 @@
 #include "Insights/Widgets/STimingView.h"
 
 #include <limits>
-#include "Widgets/Input/SSpinBox.h"
 
 #define LOCTEXT_NAMESPACE "UE::Insights::TimingProfiler::STimersView"
 
@@ -149,6 +148,30 @@ public:
 			EUserInterfaceActionType::Button,
 			FInputChord());
 
+		UI_COMMAND(Command_FindFirstInstance,
+			"First Instance",
+			"Navigates to and selects the first timing event instance for the selected timer.",
+			EUserInterfaceActionType::Button,
+			FInputChord());
+
+		UI_COMMAND(Command_FindPreviousInstance,
+			"Previous Instance",
+			"Navigates to and selects the previous timing event instance for the selected timer.",
+			EUserInterfaceActionType::Button,
+			FInputChord());
+
+		UI_COMMAND(Command_FindNextInstance,
+			"Next Instance",
+			"Navigates to and selects the next timing event instance for the selected timer.",
+			EUserInterfaceActionType::Button,
+			FInputChord());
+
+		UI_COMMAND(Command_FindLastInstance,
+			"Last Instance",
+			"Navigates to and selects the last timing event instance for the selected timer.",
+			EUserInterfaceActionType::Button,
+			FInputChord());
+
 		UI_COMMAND(Command_FindMaxInstance,
 			"Maximum Duration Instance",
 			"Navigates to and selects the timing event instance with the maximum duration, for the selected timer.",
@@ -160,6 +183,30 @@ public:
 			"Navigates to and selects the timing event instance with the minimum duration, for the selected timer.",
 			EUserInterfaceActionType::Button,
 			FInputChord());
+		
+		UI_COMMAND(Command_FindFirstInstanceInSelection,
+			"First Instance in Selection",
+			"Navigates to and selects the first timing event instance for the selected timer, in the selected time range.",
+			EUserInterfaceActionType::Button,
+			FInputChord());
+
+		UI_COMMAND(Command_FindPreviousInstanceInSelection,
+			"Previous Instance in Selection",
+			"Navigates to and selects the previous timing event instance for the selected timer, in the selected time range.",
+			EUserInterfaceActionType::Button,
+			FInputChord());
+
+		UI_COMMAND(Command_FindNextInstanceInSelection,
+			"Next Instance in Selection",
+			"Navigates to and selects the next timing event instance for the selected timer, in the selected time range.",
+			EUserInterfaceActionType::Button,
+			FInputChord());
+
+		UI_COMMAND(Command_FindLastInstanceInSelection,
+			"Last Instance in Selection",
+			"Navigates to and selects the last timing event instance for the selected timer, in the selected time range.",
+			EUserInterfaceActionType::Button,
+			FInputChord());
 
 		UI_COMMAND(Command_FindMaxInstanceInSelection,
 			"Maximum Duration Instance in Selection",
@@ -185,8 +232,16 @@ public:
 	TSharedPtr<FUICommandInfo> Command_OpenSource;
 	TSharedPtr<FUICommandInfo> Command_GenerateNewColor;
 	TSharedPtr<FUICommandInfo> Command_EditColor;
+	TSharedPtr<FUICommandInfo> Command_FindFirstInstance;
+	TSharedPtr<FUICommandInfo> Command_FindPreviousInstance;
+	TSharedPtr<FUICommandInfo> Command_FindNextInstance;
+	TSharedPtr<FUICommandInfo> Command_FindLastInstance;
 	TSharedPtr<FUICommandInfo> Command_FindMaxInstance;
 	TSharedPtr<FUICommandInfo> Command_FindMinInstance;
+	TSharedPtr<FUICommandInfo> Command_FindFirstInstanceInSelection;
+	TSharedPtr<FUICommandInfo> Command_FindPreviousInstanceInSelection;
+	TSharedPtr<FUICommandInfo> Command_FindNextInstanceInSelection;
+	TSharedPtr<FUICommandInfo> Command_FindLastInstanceInSelection;
 	TSharedPtr<FUICommandInfo> Command_FindMaxInstanceInSelection;
 	TSharedPtr<FUICommandInfo> Command_FindMinInstanceInSelection;
 };
@@ -239,10 +294,18 @@ void STimersView::InitCommandList()
 	CommandList->MapAction(FTimersViewCommands::Get().Command_OpenSource, FExecuteAction::CreateSP(this, &STimersView::ContextMenu_OpenSource_Execute), FCanExecuteAction::CreateSP(this, &STimersView::ContextMenu_OpenSource_CanExecute));
 	CommandList->MapAction(FTimersViewCommands::Get().Command_GenerateNewColor, FExecuteAction::CreateSP(this, &STimersView::ContextMenu_GenerateNewColor_Execute), FCanExecuteAction::CreateSP(this, &STimersView::ContextMenu_GenerateNewColor_CanExecute));
 	CommandList->MapAction(FTimersViewCommands::Get().Command_EditColor, FExecuteAction::CreateSP(this, &STimersView::ContextMenu_EditColor_Execute), FCanExecuteAction::CreateSP(this, &STimersView::ContextMenu_EditColor_CanExecute));
-	CommandList->MapAction(FTimersViewCommands::Get().Command_FindMaxInstance, FExecuteAction::CreateSP(this, &STimersView::ContextMenu_FindInstance_Execute, true), FCanExecuteAction::CreateSP(this, &STimersView::ContextMenu_FindInstance_CanExecute));
-	CommandList->MapAction(FTimersViewCommands::Get().Command_FindMinInstance, FExecuteAction::CreateSP(this, &STimersView::ContextMenu_FindInstance_Execute, false), FCanExecuteAction::CreateSP(this, &STimersView::ContextMenu_FindInstance_CanExecute));
-	CommandList->MapAction(FTimersViewCommands::Get().Command_FindMaxInstanceInSelection, FExecuteAction::CreateSP(this, &STimersView::ContextMenu_FindInstanceInSelection_Execute, true), FCanExecuteAction::CreateSP(this, &STimersView::ContextMenu_FindInstanceInSelection_CanExecute));
-	CommandList->MapAction(FTimersViewCommands::Get().Command_FindMinInstanceInSelection, FExecuteAction::CreateSP(this, &STimersView::ContextMenu_FindInstanceInSelection_Execute, false), FCanExecuteAction::CreateSP(this, &STimersView::ContextMenu_FindInstanceInSelection_CanExecute));
+	CommandList->MapAction(FTimersViewCommands::Get().Command_FindFirstInstance, FExecuteAction::CreateSP(this, &STimersView::ContextMenu_FindInstance_Execute, ESelectEventType::First), FCanExecuteAction::CreateSP(this, &STimersView::ContextMenu_FindInstance_CanExecute));
+	CommandList->MapAction(FTimersViewCommands::Get().Command_FindPreviousInstance, FExecuteAction::CreateSP(this, &STimersView::ContextMenu_FindInstance_Execute, ESelectEventType::Previous), FCanExecuteAction::CreateSP(this, &STimersView::ContextMenu_FindInstance_CanExecute));
+	CommandList->MapAction(FTimersViewCommands::Get().Command_FindNextInstance, FExecuteAction::CreateSP(this, &STimersView::ContextMenu_FindInstance_Execute, ESelectEventType::Next), FCanExecuteAction::CreateSP(this, &STimersView::ContextMenu_FindInstance_CanExecute));
+	CommandList->MapAction(FTimersViewCommands::Get().Command_FindLastInstance, FExecuteAction::CreateSP(this, &STimersView::ContextMenu_FindInstance_Execute, ESelectEventType::Last), FCanExecuteAction::CreateSP(this, &STimersView::ContextMenu_FindInstance_CanExecute));
+	CommandList->MapAction(FTimersViewCommands::Get().Command_FindMaxInstance, FExecuteAction::CreateSP(this, &STimersView::ContextMenu_FindInstance_Execute, ESelectEventType::Max), FCanExecuteAction::CreateSP(this, &STimersView::ContextMenu_FindInstance_CanExecute));
+	CommandList->MapAction(FTimersViewCommands::Get().Command_FindMinInstance, FExecuteAction::CreateSP(this, &STimersView::ContextMenu_FindInstance_Execute, ESelectEventType::Min), FCanExecuteAction::CreateSP(this, &STimersView::ContextMenu_FindInstance_CanExecute));
+	CommandList->MapAction(FTimersViewCommands::Get().Command_FindFirstInstanceInSelection, FExecuteAction::CreateSP(this, &STimersView::ContextMenu_FindInstanceInSelection_Execute, ESelectEventType::First), FCanExecuteAction::CreateSP(this, &STimersView::ContextMenu_FindInstanceInSelection_CanExecute));
+	CommandList->MapAction(FTimersViewCommands::Get().Command_FindPreviousInstanceInSelection, FExecuteAction::CreateSP(this, &STimersView::ContextMenu_FindInstanceInSelection_Execute, ESelectEventType::Previous), FCanExecuteAction::CreateSP(this, &STimersView::ContextMenu_FindInstanceInSelection_CanExecute));
+	CommandList->MapAction(FTimersViewCommands::Get().Command_FindNextInstanceInSelection, FExecuteAction::CreateSP(this, &STimersView::ContextMenu_FindInstanceInSelection_Execute, ESelectEventType::Next), FCanExecuteAction::CreateSP(this, &STimersView::ContextMenu_FindInstanceInSelection_CanExecute));
+	CommandList->MapAction(FTimersViewCommands::Get().Command_FindLastInstanceInSelection, FExecuteAction::CreateSP(this, &STimersView::ContextMenu_FindInstanceInSelection_Execute, ESelectEventType::Last), FCanExecuteAction::CreateSP(this, &STimersView::ContextMenu_FindInstanceInSelection_CanExecute));
+	CommandList->MapAction(FTimersViewCommands::Get().Command_FindMaxInstanceInSelection, FExecuteAction::CreateSP(this, &STimersView::ContextMenu_FindInstanceInSelection_Execute, ESelectEventType::Max), FCanExecuteAction::CreateSP(this, &STimersView::ContextMenu_FindInstanceInSelection_CanExecute));
+	CommandList->MapAction(FTimersViewCommands::Get().Command_FindMinInstanceInSelection, FExecuteAction::CreateSP(this, &STimersView::ContextMenu_FindInstanceInSelection_Execute, ESelectEventType::Min), FCanExecuteAction::CreateSP(this, &STimersView::ContextMenu_FindInstanceInSelection_CanExecute));
 }
 
 ////////////////////////////////////////////////////////////////////////////////////////////////////
@@ -1129,6 +1192,42 @@ void STimersView::TreeView_FindMenu(FMenuBuilder& MenuBuilder)
 {
 	MenuBuilder.SetSearchable(false);
 
+	MenuBuilder.AddMenuEntry
+	(
+		FTimersViewCommands::Get().Command_FindFirstInstance,
+		NAME_None,
+		TAttribute<FText>(),
+		TAttribute<FText>(),
+		FSlateIcon(FInsightsStyle::Get().GetStyleSetName(), "Icons.FindFirst")
+	);
+
+	MenuBuilder.AddMenuEntry
+	(
+		FTimersViewCommands::Get().Command_FindPreviousInstance,
+		NAME_None,
+		TAttribute<FText>(),
+		TAttribute<FText>(),
+		FSlateIcon(FInsightsStyle::Get().GetStyleSetName(), "Icons.FindPrevious")
+	);
+
+	MenuBuilder.AddMenuEntry
+	(
+		FTimersViewCommands::Get().Command_FindNextInstance,
+		NAME_None,
+		TAttribute<FText>(),
+		TAttribute<FText>(),
+		FSlateIcon(FInsightsStyle::Get().GetStyleSetName(), "Icons.FindNext")
+	);
+
+	MenuBuilder.AddMenuEntry
+	(
+		FTimersViewCommands::Get().Command_FindLastInstance,
+		NAME_None,
+		TAttribute<FText>(),
+		TAttribute<FText>(),
+		FSlateIcon(FInsightsStyle::Get().GetStyleSetName(), "Icons.FindLast")
+	);
+
 	MenuBuilder.AddMenuEntry
 	(
 		FTimersViewCommands::Get().Command_FindMaxInstance,
@@ -1149,6 +1248,42 @@ void STimersView::TreeView_FindMenu(FMenuBuilder& MenuBuilder)
 
 	MenuBuilder.AddMenuSeparator();
 
+	MenuBuilder.AddMenuEntry
+	(
+		FTimersViewCommands::Get().Command_FindFirstInstanceInSelection,
+		NAME_None,
+		TAttribute<FText>(),
+		TAttribute<FText>(),
+		FSlateIcon(FInsightsStyle::Get().GetStyleSetName(), "Icons.FindFirst")
+	);
+
+	MenuBuilder.AddMenuEntry
+	(
+		FTimersViewCommands::Get().Command_FindPreviousInstanceInSelection,
+		NAME_None,
+		TAttribute<FText>(),
+		TAttribute<FText>(),
+		FSlateIcon(FInsightsStyle::Get().GetStyleSetName(), "Icons.FindPrevious")
+	);
+
+	MenuBuilder.AddMenuEntry
+	(
+		FTimersViewCommands::Get().Command_FindNextInstanceInSelection,
+		NAME_None,
+		TAttribute<FText>(),
+		TAttribute<FText>(),
+		FSlateIcon(FInsightsStyle::Get().GetStyleSetName(), "Icons.FindNext")
+	);
+
+	MenuBuilder.AddMenuEntry
+	(
+		FTimersViewCommands::Get().Command_FindLastInstanceInSelection,
+		NAME_None,
+		TAttribute<FText>(),
+		TAttribute<FText>(),
+		FSlateIcon(FInsightsStyle::Get().GetStyleSetName(), "Icons.FindLast")
+	);
+	
 	MenuBuilder.AddMenuEntry
 	(
 		FTimersViewCommands::Get().Command_FindMaxInstanceInSelection,
@@ -3714,7 +3849,7 @@ bool STimersView::ContextMenu_FindInstance_CanExecute() const
 
 ////////////////////////////////////////////////////////////////////////////////////////////////////
 
-void STimersView::ContextMenu_FindInstance_Execute(bool bFindMax) const
+void STimersView::ContextMenu_FindInstance_Execute(ESelectEventType EventType) const
 {
 	FTimerNodePtr SelectedNode = GetSingleSelectedTimerNode();
 	if (!SelectedNode.IsValid())
@@ -3728,8 +3863,7 @@ void STimersView::ContextMenu_FindInstance_Execute(bool bFindMax) const
 		return;
 	}
 
-	ESelectEventType Type = bFindMax ? ESelectEventType::Max : ESelectEventType::Min;
-	TimingView->SelectEventInstance(SelectedNode->GetTimerId(), Type, false);
+	TimingView->SelectEventInstance(SelectedNode->GetTimerId(), EventType, false);
 }
 
 ////////////////////////////////////////////////////////////////////////////////////////////////////
@@ -3753,7 +3887,7 @@ bool STimersView::ContextMenu_FindInstanceInSelection_CanExecute() const
 
 ////////////////////////////////////////////////////////////////////////////////////////////////////
 
-void STimersView::ContextMenu_FindInstanceInSelection_Execute(bool bFindMax) const
+void STimersView::ContextMenu_FindInstanceInSelection_Execute(ESelectEventType EventType) const
 {
 	FTimerNodePtr SelectedNode = GetSingleSelectedTimerNode();
 	if (!SelectedNode.IsValid() || SelectedNode->GetType() == ETimerNodeType::Group)
@@ -3767,8 +3901,7 @@ void STimersView::ContextMenu_FindInstanceInSelection_Execute(bool bFindMax) con
 		return;
 	}
 
-	ESelectEventType Type = bFindMax ? ESelectEventType::Max : ESelectEventType::Min;
-	TimingView->SelectEventInstance(SelectedNode->GetTimerId(), Type, true);
+	TimingView->SelectEventInstance(SelectedNode->GetTimerId(), EventType, true);
 }
 
 ////////////////////////////////////////////////////////////////////////////////////////////////////
diff --git a/Engine/Source/Developer/TraceInsights/Private/Insights/TimingProfiler/Widgets/STimersView.h b/Engine/Source/Developer/TraceInsights/Private/Insights/TimingProfiler/Widgets/STimersView.h
--- a/Engine/Source/Developer/TraceInsights/Private/Insights/TimingProfiler/Widgets/STimersView.h
+++ b/Engine/Source/Developer/TraceInsights/Private/Insights/TimingProfiler/Widgets/STimersView.h
@@ -38,6 +38,7 @@ class FTimingGraphTrack;
 namespace UE::Insights::TimingProfiler
 {
 
+enum class ESelectEventType : uint32;
 class SFrameTrack;
 class FTimerAggregator;
 class STimingView;
@@ -195,10 +196,10 @@ private:
 	void TreeView_FindMenu(FMenuBuilder& MenuBuilder);
 
 	bool ContextMenu_FindInstance_CanExecute() const;
-	void ContextMenu_FindInstance_Execute(bool bFindMax) const;
+	void ContextMenu_FindInstance_Execute(ESelectEventType EventType) const;
 
 	bool ContextMenu_FindInstanceInSelection_CanExecute() const;
-	void ContextMenu_FindInstanceInSelection_Execute(bool bFindMax) const;
+	void ContextMenu_FindInstanceInSelection_Execute(ESelectEventType EventType) const;
 	
 	//////////////////////////////////////////////////
 	// Settings
diff --git a/Engine/Source/Developer/TraceInsights/Private/Insights/ViewModels/QuickFind.h b/Engine/Source/Developer/TraceInsights/Private/Insights/ViewModels/QuickFind.h
--- a/Engine/Source/Developer/TraceInsights/Private/Insights/ViewModels/QuickFind.h
+++ b/Engine/Source/Developer/TraceInsights/Private/Insights/ViewModels/QuickFind.h
@@ -74,6 +74,26 @@ public:
 private:
 	FOnFindLastEvent OnFindLastEvent;
 
+	////////////////////////////////////////////////////////////////////////////////////////////////////
+	// OnFindMaxEvent
+public:
+	/** The event to execute when the user selects the "Find Max" option. */
+	DECLARE_MULTICAST_DELEGATE(FOnFindMaxEvent);
+	FOnFindMaxEvent& GetOnFindMaxEvent() { return OnFindMaxEvent; }
+
+private:
+	FOnFindMaxEvent OnFindMaxEvent;
+
+	////////////////////////////////////////////////////////////////////////////////////////////////////
+	// OnFindMinEvent
+public:
+	/** The event to execute when the user selects the "Find Min" option. */
+	DECLARE_MULTICAST_DELEGATE(FOnFindMinEvent);
+	FOnFindMinEvent& GetOnFindMinEvent() { return OnFindMinEvent; }
+
+private:
+	FOnFindMinEvent OnFindMinEvent;	
+
 	////////////////////////////////////////////////////////////////////////////////////////////////////
 	// OnFilterAllEvent
 public:
diff --git a/Engine/Source/Developer/TraceInsights/Private/Insights/Widgets/SQuickFind.cpp b/Engine/Source/Developer/TraceInsights/Private/Insights/Widgets/SQuickFind.cpp
--- a/Engine/Source/Developer/TraceInsights/Private/Insights/Widgets/SQuickFind.cpp
+++ b/Engine/Source/Developer/TraceInsights/Private/Insights/Widgets/SQuickFind.cpp
@@ -140,6 +140,40 @@ void SQuickFind::Construct(const FArguments& InArgs, TSharedPtr<FQuickFind> InQu
 				]
 			]
 
+			+ SHorizontalBox::Slot()
+			.HAlign(HAlign_Left)
+			.VAlign(VAlign_Center)
+			.AutoWidth()
+			.Padding(1.0f)
+			[
+				SNew(SButton)
+				.Text(LOCTEXT("FindMax", "FindMax"))
+				.ToolTipText(LOCTEXT("FindMaxDesc", "Find Max\nSelects the timing event with the largest duration that matches the search criteria."))
+				.OnClicked(this, &SQuickFind::FindMax_OnClicked)
+				.Content()
+				[
+					SNew(SImage)
+					.Image(FInsightsStyle::Get().GetBrush("Icons.FindMaxInstance.ToolBar"))
+				]
+			]
+
+			+ SHorizontalBox::Slot()
+			.HAlign(HAlign_Left)
+			.VAlign(VAlign_Center)
+			.AutoWidth()
+			.Padding(1.0f)
+			[
+				SNew(SButton)
+				.Text(LOCTEXT("FindMin", "Find Min"))
+				.ToolTipText(LOCTEXT("FindMinDesc", "Find Min\nSelects the timing event with the shortest duration that matches the search criteria."))
+				.OnClicked(this, &SQuickFind::FindMin_OnClicked)
+				.Content()
+				[
+					SNew(SImage)
+					.Image(FInsightsStyle::Get().GetBrush("Icons.FindMinInstance.ToolBar"))
+				]
+			]
+
 		+ SHorizontalBox::Slot()
 			.HAlign(HAlign_Left)
 			.VAlign(VAlign_Center)
@@ -240,6 +274,27 @@ FReply SQuickFind::FindLast_OnClicked()
 
 	return FReply::Handled();
 }
+
+////////////////////////////////////////////////////////////////////////////////////////////////////
+
+FReply SQuickFind::FindMax_OnClicked()
+{
+	QuickFindViewModel->GetFilterConfigurator()->GetRootNode()->ProcessFilter();
+	QuickFindViewModel->GetOnFindMaxEvent().Broadcast();
+
+	return FReply::Handled();
+}
+
+////////////////////////////////////////////////////////////////////////////////////////////////////
+
+FReply SQuickFind::FindMin_OnClicked()
+{
+	QuickFindViewModel->GetFilterConfigurator()->GetRootNode()->ProcessFilter();
+	QuickFindViewModel->GetOnFindMinEvent().Broadcast();
+
+	return FReply::Handled();
+}
+
 ////////////////////////////////////////////////////////////////////////////////////////////////////
 
 FReply SQuickFind::FilterAll_OnClicked()
diff --git a/Engine/Source/Developer/TraceInsights/Private/Insights/Widgets/SQuickFind.h b/Engine/Source/Developer/TraceInsights/Private/Insights/Widgets/SQuickFind.h
--- a/Engine/Source/Developer/TraceInsights/Private/Insights/Widgets/SQuickFind.h
+++ b/Engine/Source/Developer/TraceInsights/Private/Insights/Widgets/SQuickFind.h
@@ -54,6 +54,8 @@ private:
 	FReply FindPrevious_OnClicked();
 	FReply FindNext_OnClicked();
 	FReply FindLast_OnClicked();
+	FReply FindMax_OnClicked();
+	FReply FindMin_OnClicked();
 
 
 	FReply FilterAll_OnClicked();
diff --git a/Engine/Source/Developer/TraceInsights/Private/Insights/Widgets/STimingView.cpp b/Engine/Source/Developer/TraceInsights/Private/Insights/Widgets/STimingView.cpp
--- a/Engine/Source/Developer/TraceInsights/Private/Insights/Widgets/STimingView.cpp
+++ b/Engine/Source/Developer/TraceInsights/Private/Insights/Widgets/STimingView.cpp
@@ -161,6 +161,7 @@ STimingView::~STimingView()
 	ForegroundTracks.Reset();
 
 	SelectedEvent.Reset();
+	CustomFilterConfigurator.Reset();
 
 	for (Timing::ITimingViewExtender* Extender : GetExtenders())
 	{
@@ -340,6 +341,63 @@ void STimingView::Construct(const FArguments& InArgs, FName InViewName)
 
 ////////////////////////////////////////////////////////////////////////////////////////////////////
 
+TSharedPtr<FFilterConfigurator> STimingView::MakeCustomFilterConfigurator(uint32 TimerId, double IntervalStart, double IntervalEnd)
+{
+	if (CustomFilterConfigurator.IsValid())
+	{
+		CustomFilterConfigurator.Reset();
+	}
+
+	CustomFilterConfigurator = MakeShared<FFilterConfigurator>();
+
+	CustomFilterConfigurator->Add(MakeShared<FFilter>(
+		static_cast<int32>(EFilterField::StartTime),
+		LOCTEXT("StartTime", "Start Time"),
+		LOCTEXT("StartTime", "Start Time"),
+		EFilterDataType::Double,
+		MakeShared<FTimeFilterValueConverter>(),
+		FFilterService::Get()->GetDoubleOperators()));
+
+	CustomFilterConfigurator->Add(MakeShared<FFilter>(
+		static_cast<int32>(EFilterField::EndTime),
+		LOCTEXT("EndTime", "End Time"),
+		LOCTEXT("EndTime", "End Time"),
+		EFilterDataType::Double,
+		MakeShared<FTimeFilterValueConverter>(),
+		FFilterService::Get()->GetDoubleOperators()));  
+
+	CustomFilterConfigurator->Add(MakeShared<FFilter>(
+		static_cast<int32>(EFilterField::TimerId),
+		LOCTEXT("TimerId", "Timer Id"),
+		LOCTEXT("TimerId", "Timer Id"),
+		EFilterDataType::Int64,
+		nullptr,
+		FFilterService::Get()->GetIntegerOperators()));
+
+	FFilterConfiguratorNodePtr StartTimeNode = MakeShared<FFilterConfiguratorNode>(TEXT(""), false);
+	StartTimeNode->SetSelectedFilter(CustomFilterConfigurator->GetRootNode()->GetAvailableFilters()->GetData()[0]);
+	StartTimeNode->SetTextBoxValue(FString::SanitizeFloat(IntervalStart));
+	StartTimeNode->SetSelectedFilterOperator(MakeShared<FFilterOperator<double>>(EFilterOperator::Gte, TEXT(""), std::greater_equal<>()));
+
+	FFilterConfiguratorNodePtr EndTimeNode = MakeShared<FFilterConfiguratorNode>(TEXT(""), false);
+	EndTimeNode->SetSelectedFilter(CustomFilterConfigurator->GetRootNode()->GetAvailableFilters()->GetData()[1]);
+	EndTimeNode->SetTextBoxValue(FString::SanitizeFloat(IntervalEnd));
+
+	EndTimeNode->SetSelectedFilterOperator(MakeShared<FFilterOperator<double>>(EFilterOperator::Lte, TEXT(""), std::less_equal<>()));
+	FFilterConfiguratorNodePtr TimerIdNode = MakeShared<FFilterConfiguratorNode>(TEXT(""), false);
+	TimerIdNode->SetSelectedFilter(CustomFilterConfigurator->GetRootNode()->GetAvailableFilters()->GetData()[2]);
+	TimerIdNode->SetTextBoxValue(FString::FromInt(TimerId));
+
+	CustomFilterConfigurator->GetRootNode()->AddChildAndSetParent(StartTimeNode);
+	CustomFilterConfigurator->GetRootNode()->AddChildAndSetParent(EndTimeNode);
+	CustomFilterConfigurator->GetRootNode()->AddChildAndSetParent(TimerIdNode);
+	CustomFilterConfigurator->GetRootNode()->ProcessFilter();
+		
+	return CustomFilterConfigurator;
+}
+
+////////////////////////////////////////////////////////////////////////////////////////////////////
+
 void STimingView::HideAllDefaultTracks()
 {
 	FrameSharedState->HideAllFrameTracks();
@@ -5207,10 +5265,13 @@ PRAGMA_ENABLE_DEPRECATION_WARNINGS
 #endif // UE_INSIGHTS_BACKWARD_COMPATIBILITY_UE54
 
 		QuickFindVm = MakeShared<::Insights::FQuickFind>(NewFilterConfigurator);
-		QuickFindVm->GetOnFindFirstEvent().AddSP(this, &STimingView::FindFirstEvent);
-		QuickFindVm->GetOnFindPreviousEvent().AddSP(this, &STimingView::FindPrevEvent);
-		QuickFindVm->GetOnFindNextEvent().AddSP(this, &STimingView::FindNextEvent);
-		QuickFindVm->GetOnFindLastEvent().AddSP(this, &STimingView::FindLastEvent);
+		constexpr bool bUseCustomFilter = false;
+		QuickFindVm->GetOnFindFirstEvent().AddSP(this, &STimingView::FindFirstEvent, bUseCustomFilter);
+		QuickFindVm->GetOnFindPreviousEvent().AddSP(this, &STimingView::FindPrevEvent, bUseCustomFilter);
+		QuickFindVm->GetOnFindNextEvent().AddSP(this, &STimingView::FindNextEvent, bUseCustomFilter);
+		QuickFindVm->GetOnFindLastEvent().AddSP(this, &STimingView::FindLastEvent, bUseCustomFilter);
+		QuickFindVm->GetOnFindMaxEvent().AddSP(this, &STimingView::FindMaxEvent);
+		QuickFindVm->GetOnFindMinEvent().AddSP(this, &STimingView::FindMinEvent);
 		QuickFindVm->GetOnFilterAllEvent().AddSP(this, &STimingView::FilterAllTracks);
 		QuickFindVm->GetOnClearFiltersEvent().AddSP(this, &STimingView::ClearFilters);
 	}
@@ -5351,19 +5412,19 @@ void STimingView::ClearRelations()
 
 ////////////////////////////////////////////////////////////////////////////////////////////////////
 
-void STimingView::FindFirstEvent()
+void STimingView::FindFirstEvent(bool bUseCustomFilter)
 {
 	if (SelectedEvent.IsValid())
 	{
 		SelectedEvent.Reset();
 		OnSelectedEventChangedDelegate.Broadcast(SelectedEvent);
 	}
-	FindNextEvent();
+	FindNextEvent(bUseCustomFilter);
 }
 
 ////////////////////////////////////////////////////////////////////////////////////////////////////
 
-void STimingView::FindPrevEvent()
+void STimingView::FindPrevEvent(bool bUseCustomFilter)
 {
 	TSharedPtr<const ITimingEvent> BestMatchEvent;
 	double StartTime = SelectedEvent.IsValid() ? SelectedEvent->GetStartTime() : std::numeric_limits<double>::max();
@@ -5373,11 +5434,11 @@ void STimingView::FindPrevEvent()
 		return EventStartTime < StartTime;
 	};
 	FTimingEventSearchParameters Params(std::numeric_limits<double>::lowest(), StartTime, ETimingEventSearchFlags::StopAtFirstMatch, EventFilter);
-	Params.FilterExecutor = QuickFindVm->GetFilterConfigurator();
+	Params.FilterExecutor = bUseCustomFilter ? CustomFilterConfigurator : QuickFindVm->GetFilterConfigurator();
 	Params.SearchDirection = FTimingEventSearchParameters::ESearchDirection::Backward;
 
 	TSharedPtr<const FBaseTimingTrack> PriorityTrack = SelectedEvent.IsValid() ? SelectedEvent->GetTrack().ToSharedPtr() : nullptr;
-	EnumerateFilteredTracks(QuickFindVm->GetFilterConfigurator(), PriorityTrack, [&Params, &BestMatchEvent](TSharedPtr<const FBaseTimingTrack> Track)
+	EnumerateFilteredTracks(bUseCustomFilter ? CustomFilterConfigurator : QuickFindVm->GetFilterConfigurator(), PriorityTrack, [&Params, &BestMatchEvent](TSharedPtr<const FBaseTimingTrack> Track)
 	{
 		if (!Track->IsVisible())
 		{
@@ -5396,28 +5457,12 @@ void STimingView::FindPrevEvent()
 		}
 	});
 
-	if (BestMatchEvent)
-	{
-		SelectedEvent = BestMatchEvent;
-		BringIntoView(SelectedEvent->GetStartTime(), SelectedEvent->GetEndTime());
-		if (SelectedEvent->GetTrack()->GetLocation() == ETimingTrackLocation::Scrollable)
-		{
-			BringScrollableTrackIntoView(*SelectedEvent->GetTrack());
-		}
-
-		OnSelectedTimingEventChanged();
-	}
-	else
-	{
-		FMessageLog ReportMessageLog(FTimingProfilerManager::Get()->GetLogListingName());
-		ReportMessageLog.Error(LOCTEXT("NoEventFound", "No event found!"));
-		ReportMessageLog.Notify();
-	}
+	PostFindEvent(BestMatchEvent);
 }
 
 ////////////////////////////////////////////////////////////////////////////////////////////////////
 
-void STimingView::FindNextEvent()
+void STimingView::FindNextEvent(bool bUseCustomFilter)
 {
 	TSharedPtr<const ITimingEvent> BestMatchEvent;
 	double StartTime = SelectedEvent.IsValid() ? SelectedEvent->GetStartTime() : std::numeric_limits<double>::lowest();
@@ -5427,10 +5472,10 @@ void STimingView::FindNextEvent()
 		return EventStartTime > StartTime;
 	};
 	FTimingEventSearchParameters Params(StartTime, std::numeric_limits<double>::max(), ETimingEventSearchFlags::StopAtFirstMatch, EventFilter);
-	Params.FilterExecutor = QuickFindVm->GetFilterConfigurator();
+	Params.FilterExecutor = bUseCustomFilter ? CustomFilterConfigurator : QuickFindVm->GetFilterConfigurator();
 
 	TSharedPtr<const FBaseTimingTrack> PriorityTrack = SelectedEvent.IsValid() ? SelectedEvent->GetTrack().ToSharedPtr() : nullptr;
-	EnumerateFilteredTracks(QuickFindVm->GetFilterConfigurator(), PriorityTrack, [&Params, &BestMatchEvent](TSharedPtr<const FBaseTimingTrack> Track)
+	EnumerateFilteredTracks(bUseCustomFilter ? CustomFilterConfigurator : QuickFindVm->GetFilterConfigurator(), PriorityTrack, [&Params, &BestMatchEvent](TSharedPtr<const FBaseTimingTrack> Track)
 	{
 		if (!Track->IsVisible())
 		{
@@ -5450,6 +5495,123 @@ void STimingView::FindNextEvent()
 		}
 	});
 
+	PostFindEvent(BestMatchEvent);
+}
+
+////////////////////////////////////////////////////////////////////////////////////////////////////
+
+void STimingView::FindLastEvent(bool bUseCustomFilter)
+{
+	if (SelectedEvent.IsValid())
+	{
+		SelectedEvent.Reset();
+		OnSelectedEventChangedDelegate.Broadcast(SelectedEvent);
+	}
+	FindPrevEvent(bUseCustomFilter);
+}
+
+////////////////////////////////////////////////////////////////////////////////////////////////////
+
+void STimingView::FindMaxEvent()
+{
+	if (SelectedEvent.IsValid())
+	{
+		SelectedEvent.Reset();
+		OnSelectedEventChangedDelegate.Broadcast(SelectedEvent);
+	}
+
+	TSharedPtr<const ITimingEvent> BestMatchEvent;
+	double CurrentDuration = -1.0;
+
+	auto EventFilter = [&CurrentDuration](double EventStartTime, double EventEndTime, uint32 EventDepth)
+	{
+		const double Duration = EventEndTime - EventStartTime;
+		return Duration > CurrentDuration;
+	};
+
+	auto EventMatched = [&CurrentDuration](double EventStartTime, double EventEndTime, uint32 EventDepth)
+	{
+		const double Duration = EventEndTime - EventStartTime;
+		CurrentDuration = Duration;
+	};
+
+	FTimingEventSearchParameters Params(0.0, std::numeric_limits<double>::max(), ETimingEventSearchFlags::SearchAll, EventFilter, EventMatched);
+	Params.FilterExecutor = QuickFindVm->GetFilterConfigurator();
+	
+	TSharedPtr<const FBaseTimingTrack> PriorityTrack = SelectedEvent.IsValid() ? SelectedEvent->GetTrack().ToSharedPtr() : nullptr;
+	EnumerateFilteredTracks(QuickFindVm->GetFilterConfigurator(), PriorityTrack, [&Params, &BestMatchEvent](TSharedPtr<const FBaseTimingTrack> Track)
+	{
+		if (!Track->IsVisible())
+		{
+			return;
+		}
+
+		TSharedPtr<const ITimingEvent> FoundEvent = Track->SearchEvent(Params);
+		if (FoundEvent.IsValid())
+		{
+			if (!BestMatchEvent.IsValid() || BestMatchEvent->GetDuration() < FoundEvent->GetDuration())
+			{
+				BestMatchEvent = FoundEvent;
+			}
+		}
+	});
+
+	PostFindEvent(BestMatchEvent);
+}
+
+////////////////////////////////////////////////////////////////////////////////////////////////////	
+
+void STimingView::FindMinEvent()
+{
+	if (SelectedEvent.IsValid())
+	{
+		SelectedEvent.Reset();
+		OnSelectedEventChangedDelegate.Broadcast(SelectedEvent);
+	}
+
+	TSharedPtr<const ITimingEvent> BestMatchEvent;
+	double CurrentDuration = std::numeric_limits<double>::max();
+
+	auto EventFilter = [&CurrentDuration](double EventStartTime, double EventEndTime, uint32 EventDepth)
+	{
+		const double Duration = EventEndTime - EventStartTime;
+		return Duration < CurrentDuration;
+	};
+
+	auto EventMatched = [&CurrentDuration](double EventStartTime, double EventEndTime, uint32 EventDepth)
+	{
+		const double Duration = EventEndTime - EventStartTime;
+		CurrentDuration = Duration;
+	};
+
+	FTimingEventSearchParameters Params(0.0, std::numeric_limits<double>::max(), ETimingEventSearchFlags::SearchAll, EventFilter, EventMatched);
+	Params.FilterExecutor = QuickFindVm->GetFilterConfigurator();
+	
+	TSharedPtr<const FBaseTimingTrack> PriorityTrack = SelectedEvent.IsValid() ? SelectedEvent->GetTrack().ToSharedPtr() : nullptr;
+	EnumerateFilteredTracks(QuickFindVm->GetFilterConfigurator(), PriorityTrack, [&Params, &BestMatchEvent](TSharedPtr<const FBaseTimingTrack> Track)
+	{
+		if (!Track->IsVisible())
+		{
+			return;
+		}
+
+		TSharedPtr<const ITimingEvent> FoundEvent = Track->SearchEvent(Params);
+		if (FoundEvent.IsValid())
+		{
+			if (!BestMatchEvent.IsValid() || BestMatchEvent->GetDuration() > FoundEvent->GetDuration())
+			{
+				BestMatchEvent = FoundEvent;
+			}
+		}
+	});
+
+	PostFindEvent(BestMatchEvent);
+}
+
+////////////////////////////////////////////////////////////////////////////////////////////////////
+
+void STimingView::PostFindEvent(TSharedPtr<const ITimingEvent> BestMatchEvent)
+{
 	if (BestMatchEvent)
 	{
 		SelectedEvent = BestMatchEvent;
@@ -5471,18 +5633,6 @@ void STimingView::FindNextEvent()
 
 ////////////////////////////////////////////////////////////////////////////////////////////////////
 
-void STimingView::FindLastEvent()
-{
-	if (SelectedEvent.IsValid())
-	{
-		SelectedEvent.Reset();
-		OnSelectedEventChangedDelegate.Broadcast(SelectedEvent);
-	}
-	FindPrevEvent();
-}
-
-////////////////////////////////////////////////////////////////////////////////////////////////////
-
 void STimingView::FilterAllTracks()
 {
 	LLM_SCOPE_BYTAG(Insights);
@@ -5612,10 +5762,8 @@ ETraceFrameType STimingView::GetFrameTypeToSnapTo()
 
 void STimingView::SelectEventInstance(uint32 TimerId, ESelectEventType Type, bool bUseSelection)
 {
-	SelectTimingEvent(nullptr, false, false);
-
 	double IntervalStart = 0.0f;
-	double IntervalEnd = std::numeric_limits<double>::infinity();
+	double IntervalEnd = std::numeric_limits<double>::max();
 
 	if (bUseSelection && SelectionEndTime > SelectionStartTime)
 	{
@@ -5624,26 +5772,38 @@ void STimingView::SelectEventInstance(uint32 TimerId, ESelectEventType Type, boo
 	}
 
 	TSharedPtr<const ITimingEvent> TimingEvent;
-
-	if (Type == ESelectEventType::Min)
-	{
+	// Create a temporary filter object using the TimerId and search range as filter criteria
+	MakeCustomFilterConfigurator(TimerId, IntervalStart, IntervalEnd);
+
+	switch (Type)
+	{
+	case ESelectEventType::First:
+		FindFirstEvent(true);
+		TimingEvent = SelectedEvent;
+		break;
+	case ESelectEventType::Previous:
+		FindPrevEvent(true);
+		TimingEvent = SelectedEvent;
+		break;
+	case ESelectEventType::Next:
+		FindNextEvent(true);
+		TimingEvent = SelectedEvent;
+		break;
+	case ESelectEventType::Last:
+		FindLastEvent(true);
+		TimingEvent = SelectedEvent;
+		break;
+	case ESelectEventType::Min:
+		SelectTimingEvent(nullptr, false, false);
 		TimingEvent = ThreadTimingSharedState->FindMinEventInstance(TimerId, IntervalStart, IntervalEnd);
-	}
-	else if (Type == ESelectEventType::Max)
-	{
+		break;
+	case ESelectEventType::Max:
+		SelectTimingEvent(nullptr, false, false);
 		TimingEvent = ThreadTimingSharedState->FindMaxEventInstance(TimerId, IntervalStart, IntervalEnd);
+		break;
 	}
 
-	if (TimingEvent.IsValid())
-	{
-		SelectTimingEvent(TimingEvent, true, true);
-	}
-	else
-	{
-		FMessageLog ReportMessageLog(FTimingProfilerManager::Get()->GetLogListingName());
-		ReportMessageLog.Error(LOCTEXT("NoEventInstanceFound", "No event instance found!"));
-		ReportMessageLog.Notify();
-	}
+	PostFindEvent(TimingEvent);
 }
 
 ////////////////////////////////////////////////////////////////////////////////////////////////////
diff --git a/Engine/Source/Developer/TraceInsights/Private/Insights/Widgets/STimingView.h b/Engine/Source/Developer/TraceInsights/Private/Insights/Widgets/STimingView.h
--- a/Engine/Source/Developer/TraceInsights/Private/Insights/Widgets/STimingView.h
+++ b/Engine/Source/Developer/TraceInsights/Private/Insights/Widgets/STimingView.h
@@ -80,8 +80,12 @@ class STimersView;
 
 enum class ESelectEventType : uint32
 {
-	Min = 0,
-	Max = 1
+	First,
+	Previous,
+	Next,
+	Last,
+	Min,
+	Max
 };
 
 /** A custom widget used to display timing events. */
@@ -130,6 +134,8 @@ public:
 	TSharedPtr<LoadingProfiler::FLoadingSharedState> GetLoadingSharedState() const { return LoadingSharedState; }
 	TSharedPtr<FFileActivitySharedState> GetFileActivitySharedState() const { return FileActivitySharedState; }
 
+	TSharedPtr<FFilterConfigurator> MakeCustomFilterConfigurator(uint32 TimerId, double IntervalStart, double IntervalEnd);
+
 	void HideAllDefaultTracks();
 
 	/** Gets the time ruler track. It includes the custom time markers (ones user can drag with mouse). */
@@ -549,10 +555,13 @@ PRAGMA_ENABLE_DEPRECATION_WARNINGS
 
 	void SetTrackPosY(FBaseTimingTrack& Track, float TrackPosY) const;
 
-	void FindFirstEvent();
-	void FindPrevEvent();
-	void FindNextEvent();
-	void FindLastEvent();
+	void FindFirstEvent(bool bUseCustomFilter);
+	void FindPrevEvent(bool bUseCustomFilter);
+	void FindNextEvent(bool bUseCustomFilter);
+	void FindLastEvent(bool bUseCustomFilter);
+	void FindMaxEvent();
+	void FindMinEvent();
+	void PostFindEvent(TSharedPtr<const ITimingEvent> BestMatchEvent);
 	void FilterAllTracks();
 	void ClearFilters();
 
@@ -789,6 +798,8 @@ protected:
 
 	TSharedPtr<::Insights::FQuickFind> QuickFindVm;
 	TSharedPtr<FFilterConfigurator> FilterConfigurator;
+	// Temporary custom filter to run timer search queries from context menus.
+	TSharedPtr<FFilterConfigurator> CustomFilterConfigurator;
 	static uint32 TimingViewId;
 	const FName QuickFindTabId;
 	bool bUpdateFilters = true;
