[Insights] Timing Profiler / Profiler Window: Added Modules Tab to Timing Insights

diff --git a/Engine/Source/Developer/TraceInsights/Private/Insights/TimingProfiler/TimingProfilerCommands.cpp b/Engine/Source/Developer/TraceInsights/Private/Insights/TimingProfiler/TimingProfilerCommands.cpp
--- a/Engine/Source/Developer/TraceInsights/Private/Insights/TimingProfiler/TimingProfilerCommands.cpp
+++ b/Engine/Source/Developer/TraceInsights/Private/Insights/TimingProfiler/TimingProfilerCommands.cpp
@@ -61,6 +61,7 @@ void FTimingProfilerCommands::RegisterCommands()
 	UI_COMMAND(ToggleCalleesTreeViewVisibility, "Callees", "Toggles the visibility of the Callees tree view.", EUserInterfaceActionType::ToggleButton,FInputChord());
 	UI_COMMAND(ToggleStatsCountersViewVisibility, "Counters", "Toggles the visibility of the Counters view.", EUserInterfaceActionType::ToggleButton, FInputChord());
 	UI_COMMAND(ToggleLogViewVisibility, "Log", "Toggles the visibility of the Log view.", EUserInterfaceActionType::ToggleButton,FInputChord());
+	UI_COMMAND(ToggleModulesViewVisibility, "Modules", "Toggles the visibility of the Modules view.", EUserInterfaceActionType::ToggleButton, FInputChord());
 }
 UE_ENABLE_OPTIMIZATION_SHIP
 
@@ -159,6 +160,7 @@ IMPLEMENT_TOGGLE_COMMAND(ToggleCallersTreeViewVisibility, IsCallersTreeViewVisib
 IMPLEMENT_TOGGLE_COMMAND(ToggleCalleesTreeViewVisibility, IsCalleesTreeViewVisible, ShowHideCalleesTreeView)
 IMPLEMENT_TOGGLE_COMMAND(ToggleStatsCountersViewVisibility, IsStatsCountersViewVisible, ShowHideStatsCountersView)
 IMPLEMENT_TOGGLE_COMMAND(ToggleLogViewVisibility, IsLogViewVisible, ShowHideLogView)
+IMPLEMENT_TOGGLE_COMMAND(ToggleModulesViewVisibility, IsModulesViewVisible, ShowHideModulesView)	
 
 #undef IMPLEMENT_TOGGLE_COMMAND
 
diff --git a/Engine/Source/Developer/TraceInsights/Private/Insights/TimingProfiler/TimingProfilerCommands.h b/Engine/Source/Developer/TraceInsights/Private/Insights/TimingProfiler/TimingProfilerCommands.h
--- a/Engine/Source/Developer/TraceInsights/Private/Insights/TimingProfiler/TimingProfilerCommands.h
+++ b/Engine/Source/Developer/TraceInsights/Private/Insights/TimingProfiler/TimingProfilerCommands.h
@@ -51,6 +51,9 @@ public:
 
 	/** Toggles visibility for the Log View. Global and custom command. */
 	TSharedPtr<FUICommandInfo> ToggleLogViewVisibility;
+
+	/** Toggles visibility for the Modules view. Global and custom command. */
+	TSharedPtr<FUICommandInfo> ToggleModulesViewVisibility;
 };
 
 ////////////////////////////////////////////////////////////////////////////////////////////////////
@@ -135,6 +138,7 @@ protected:\
 	DECLARE_TOGGLE_COMMAND(ToggleCalleesTreeViewVisibility)
 	DECLARE_TOGGLE_COMMAND(ToggleStatsCountersViewVisibility)
 	DECLARE_TOGGLE_COMMAND(ToggleLogViewVisibility)
+	DECLARE_TOGGLE_COMMAND(ToggleModulesViewVisibility)
 #undef DECLARE_TOGGLE_COMMAND
 
 	//////////////////////////////////////////////////
diff --git a/Engine/Source/Developer/TraceInsights/Private/Insights/TimingProfiler/TimingProfilerManager.cpp b/Engine/Source/Developer/TraceInsights/Private/Insights/TimingProfiler/TimingProfilerManager.cpp
--- a/Engine/Source/Developer/TraceInsights/Private/Insights/TimingProfiler/TimingProfilerManager.cpp
+++ b/Engine/Source/Developer/TraceInsights/Private/Insights/TimingProfiler/TimingProfilerManager.cpp
@@ -170,6 +170,7 @@ void FTimingProfilerManager::BindCommands()
 	ActionManager.Map_ToggleCalleesTreeViewVisibility_Global();
 	ActionManager.Map_ToggleStatsCountersViewVisibility_Global();
 	ActionManager.Map_ToggleLogViewVisibility_Global();
+	ActionManager.Map_ToggleModulesViewVisibility_Global();
 }
 
 ////////////////////////////////////////////////////////////////////////////////////////////////////
@@ -482,6 +483,19 @@ void FTimingProfilerManager::ShowHideLogView(const bool bIsVisible)
 		Wnd->ShowHideTab(FTimingProfilerTabs::LogViewID, bIsLogViewVisible);
 	}
 }
+	
+////////////////////////////////////////////////////////////////////////////////////////////////////
+
+void FTimingProfilerManager::ShowHideModulesView(const bool bIsVisible)
+{
+	bIsModulesViewVisible = bIsVisible;
+
+	TSharedPtr<STimingProfilerWindow> Wnd = GetProfilerWindow();
+	if (Wnd.IsValid())
+	{
+		Wnd->ShowHideTab(FTimingProfilerTabs::ModulesViewID, bIsModulesViewVisible);
+	}
+}
 
 ////////////////////////////////////////////////////////////////////////////////////////////////////
 
diff --git a/Engine/Source/Developer/TraceInsights/Private/Insights/TimingProfiler/TimingProfilerManager.h b/Engine/Source/Developer/TraceInsights/Private/Insights/TimingProfiler/TimingProfilerManager.h
--- a/Engine/Source/Developer/TraceInsights/Private/Insights/TimingProfiler/TimingProfilerManager.h
+++ b/Engine/Source/Developer/TraceInsights/Private/Insights/TimingProfiler/TimingProfilerManager.h
@@ -133,6 +133,11 @@ public:
 	void SetLogViewVisible(const bool bIsVisible) { bIsLogViewVisible = bIsVisible; }
 	void ShowHideLogView(const bool bIsVisible);
 
+	/** @return true, if the Modules view is visible */
+	const bool IsModulesViewVisible() const { return bIsModulesViewVisible; }
+	void SetModulesViewVisible(const bool bIsVisible) { bIsModulesViewVisible = bIsVisible; }
+	void ShowHideModulesView(const bool bIsVisible);
+
 	////////////////////////////////////////////////////////////////////////////////////////////////////
 
 	void OnSessionChanged();
@@ -237,6 +242,9 @@ private:
 	/** If the Log View is visible or hidden. */
 	bool bIsLogViewVisible;
 
+	/** If the Modules View is visible or hidden. */
+	bool bIsModulesViewVisible;
+
 	//////////////////////////////////////////////////
 
 	double SelectionStartTime;
diff --git a/Engine/Source/Developer/TraceInsights/Private/Insights/TimingProfiler/Widgets/STimingProfilerToolbar.cpp b/Engine/Source/Developer/TraceInsights/Private/Insights/TimingProfiler/Widgets/STimingProfilerToolbar.cpp
--- a/Engine/Source/Developer/TraceInsights/Private/Insights/TimingProfiler/Widgets/STimingProfilerToolbar.cpp
+++ b/Engine/Source/Developer/TraceInsights/Private/Insights/TimingProfiler/Widgets/STimingProfilerToolbar.cpp
@@ -87,6 +87,12 @@ void STimingProfilerToolbar::Construct(const FArguments& InArgs, const FInsights
 						NAME_None, TAttribute<FText>(), TAttribute<FText>(),
 						FSlateIcon(FInsightsStyle::GetStyleSetName(), "Icons.LogView.ToolBar"));
 				}
+				if (Config.ShouldRegisterMinorTab(FTimingProfilerTabs::ModulesViewID))
+				{
+					ToolbarBuilder.AddToolBarButton(FTimingProfilerCommands::Get().ToggleModulesViewVisibility,
+						NAME_None, TAttribute<FText>(), TAttribute<FText>(),
+						FSlateIcon(FInsightsStyle::GetStyleSetName(), "Icons.ModulesView.ToolBar"));
+				}
 			}
 			ToolbarBuilder.EndSection();
 
diff --git a/Engine/Source/Developer/TraceInsights/Private/Insights/TimingProfiler/Widgets/STimingProfilerWindow.cpp b/Engine/Source/Developer/TraceInsights/Private/Insights/TimingProfiler/Widgets/STimingProfilerWindow.cpp
--- a/Engine/Source/Developer/TraceInsights/Private/Insights/TimingProfiler/Widgets/STimingProfilerWindow.cpp
+++ b/Engine/Source/Developer/TraceInsights/Private/Insights/TimingProfiler/Widgets/STimingProfilerWindow.cpp
@@ -19,6 +19,7 @@
 #include "Insights/TimingProfiler/Widgets/STimingProfilerToolbar.h"
 #include "Insights/ViewModels/TimingGraphTrack.h"
 #include "Insights/Widgets/SLogView.h"
+#include "Insights/Widgets/SModulesView.h"
 #include "Insights/Widgets/STimingView.h"
 
 ////////////////////////////////////////////////////////////////////////////////////////////////////
@@ -31,6 +32,7 @@ const FName FTimingProfilerTabs::CallersID(TEXT("Callers"));
 const FName FTimingProfilerTabs::CalleesID(TEXT("Callees"));
 const FName FTimingProfilerTabs::StatsCountersID(TEXT("StasCounters"));
 const FName FTimingProfilerTabs::LogViewID(TEXT("LogView"));
+const FName FTimingProfilerTabs::ModulesViewID(TEXT("Modules"));
 
 ////////////////////////////////////////////////////////////////////////////////////////////////////
 
@@ -54,6 +56,7 @@ STimingProfilerWindow::~STimingProfilerWindow()
 {
 	CloseAllOpenTabs();
 
+	check(ModulesView == nullptr);
 	check(LogView == nullptr);
 	check(StatsView == nullptr);
 	check(CalleesTreeView == nullptr);
@@ -109,6 +112,11 @@ void STimingProfilerWindow::Reset()
 		LogView->Reset();
 	}
 
+	if (ModulesView)
+	{
+		ModulesView->Reset();
+	}
+
 	TimerColors.Reset();
 }
 
@@ -341,6 +349,37 @@ void STimingProfilerWindow::OnLogViewTabClosed(TSharedRef<SDockTab> TabBeingClos
 
 ////////////////////////////////////////////////////////////////////////////////////////////////////
 
+BEGIN_SLATE_FUNCTION_BUILD_OPTIMIZATION
+TSharedRef<SDockTab> STimingProfilerWindow::SpawnTab_ModulesView(const FSpawnTabArgs& Args)
+{
+	FTimingProfilerManager::Get()->SetModulesViewVisible(true);
+
+	const TSharedRef<SDockTab> DockTab = SNew(SDockTab)
+		.ShouldAutosize(false)
+		.TabRole(ETabRole::PanelTab)
+		[
+			SAssignNew(ModulesView, SModulesView)
+		];
+
+	DockTab->SetOnTabClosed(SDockTab::FOnTabClosedCallback::CreateRaw(this, &STimingProfilerWindow::OnModulesViewClosed));
+	AddOpenTab(DockTab);
+
+	return DockTab;
+}
+END_SLATE_FUNCTION_BUILD_OPTIMIZATION
+
+////////////////////////////////////////////////////////////////////////////////////////////////////
+
+void STimingProfilerWindow::OnModulesViewClosed(TSharedRef<SDockTab> TabBeingClosed)
+{
+	FTimingProfilerManager::Get()->SetModulesViewVisible(false);
+	ModulesView = nullptr;
+
+	RemoveOpenTab(TabBeingClosed);
+}
+
+////////////////////////////////////////////////////////////////////////////////////////////////////
+
 void STimingProfilerWindow::Construct(const FArguments& InArgs, const TSharedRef<SDockTab>& ConstructUnderMajorTab, const TSharedPtr<SWindow>& ConstructUnderWindow)
 {
 	TSharedPtr<FTimingProfilerManager> TimingProfilerManager = FTimingProfilerManager::Get();
@@ -426,6 +465,14 @@ void STimingProfilerWindow::RegisterTabSpawners()
 			.SetIcon(FSlateIcon(FInsightsStyle::GetStyleSetName(), "Icons.LogView"))
 			.SetGroup(Group);
 	}
+
+	if (Config.ShouldRegisterMinorTab(FTimingProfilerTabs::ModulesViewID))
+	{
+		TabManagerPtr->RegisterTabSpawner(FTimingProfilerTabs::ModulesViewID, FOnSpawnTab::CreateRaw(this, &STimingProfilerWindow::SpawnTab_ModulesView))
+			.SetDisplayName(LOCTEXT("ModulesViewTabTitle", "Modules"))
+			.SetIcon(FSlateIcon(FInsightsStyle::GetStyleSetName(), "Icons.ModulesView"))
+			.SetGroup(Group);
+	}
 }
 
 ////////////////////////////////////////////////////////////////////////////////////////////////////
@@ -463,6 +510,7 @@ TSharedRef<FTabManager::FLayout> STimingProfilerWindow::CreateDefaultTabLayout()
 					->SetSizeCoefficient(0.2f)
 					->SetHideTabWell(true)
 					->AddTab(FTimingProfilerTabs::LogViewID, ETabState::OpenedTab)
+					->AddTab(FTimingProfilerTabs::ModulesViewID, ETabState::OpenedTab)
 				)
 			)
 			->Split
diff --git a/Engine/Source/Developer/TraceInsights/Private/Insights/TimingProfiler/Widgets/STimingProfilerWindow.h b/Engine/Source/Developer/TraceInsights/Private/Insights/TimingProfiler/Widgets/STimingProfilerWindow.h
--- a/Engine/Source/Developer/TraceInsights/Private/Insights/TimingProfiler/Widgets/STimingProfilerWindow.h
+++ b/Engine/Source/Developer/TraceInsights/Private/Insights/TimingProfiler/Widgets/STimingProfilerWindow.h
@@ -8,7 +8,11 @@
 #include "Insights/ITimingViewSession.h" // for ETimeChangedFlags
 #include "Insights/Widgets/SMajorTabWindow.h"
 
-namespace UE::Insights { class SLogView; }
+namespace UE::Insights
+{
+	class SLogView;
+	class SModulesView;
+}
 
 namespace UE::Insights::TimingProfiler
 {
@@ -82,6 +86,9 @@ private:
 	TSharedRef<SDockTab> SpawnTab_LogView(const FSpawnTabArgs& Args);
 	void OnLogViewTabClosed(TSharedRef<SDockTab> TabBeingClosed);
 
+	TSharedRef<SDockTab> SpawnTab_ModulesView(const FSpawnTabArgs& Args);
+	void OnModulesViewClosed(TSharedRef<SDockTab> TabBeingClosed);
+
 	void OnTimeSelectionChanged(Timing::ETimeChangedFlags InFlags, double InStartTime, double InEndTime);
 
 private:
@@ -106,6 +113,9 @@ private:
 	/** The Log view widget */
 	TSharedPtr<SLogView> LogView;
 
+	/** The Modules view widget. */
+	TSharedPtr<SModulesView> ModulesView;
+
 	/** Color-TimerId Mapping for easy pairing */
 	TMap<uint32, FLinearColor> TimerColors;
 	
diff --git a/Engine/Source/Developer/TraceInsights/Private/Insights/Widgets/SModulesView.cpp b/Engine/Source/Developer/TraceInsights/Private/Insights/Widgets/SModulesView.cpp
--- a/Engine/Source/Developer/TraceInsights/Private/Insights/Widgets/SModulesView.cpp
+++ b/Engine/Source/Developer/TraceInsights/Private/Insights/Widgets/SModulesView.cpp
@@ -804,6 +804,7 @@ void SModulesView::Tick(const FGeometry& AllottedGeometry, const double InCurren
 		if (Modules.Num() != ModuleCount)
 		{
 			UpdateSorting();
+			ListView->RequestListRefresh();
 		}
 	}
 }
diff --git a/Engine/Source/Developer/TraceInsights/Public/Insights/IUnrealInsightsModule.h b/Engine/Source/Developer/TraceInsights/Public/Insights/IUnrealInsightsModule.h
--- a/Engine/Source/Developer/TraceInsights/Public/Insights/IUnrealInsightsModule.h
+++ b/Engine/Source/Developer/TraceInsights/Public/Insights/IUnrealInsightsModule.h
@@ -58,6 +58,7 @@ struct FTimingProfilerTabs
 	static UE_API const FName CalleesID;
 	static UE_API const FName StatsCountersID;
 	static UE_API const FName LogViewID;
+	static UE_API const FName ModulesViewID;
 };
 
 ////////////////////////////////////////////////////////////////////////////////////////////////////
